Move curve editor management to the curve editor integration extension

This offloads code from SSequencer and FSequencer into more focused classes, and removes the new curve editor getter method from the ISequencer interface.

There shouldn't be any behaviour change in the UI.

diff --git a/Engine/Plugins/Animation/ControlRig/Source/ControlRigEditor/Private/Tools/ControlRigTweener.cpp b/Engine/Plugins/Animation/ControlRig/Source/ControlRigEditor/Private/Tools/ControlRigTweener.cpp
--- a/Engine/Plugins/Animation/ControlRig/Source/ControlRigEditor/Private/Tools/ControlRigTweener.cpp
+++ b/Engine/Plugins/Animation/ControlRig/Source/ControlRigEditor/Private/Tools/ControlRigTweener.cpp
@@ -7,6 +7,8 @@
 #include "ISequencer.h"
 #include "MovieSceneSequence.h"
 #include "MovieScene.h"
+#include "MVVM/CurveEditorExtension.h"
+#include "MVVM/ViewModels/EditorViewModel.h"
 #include "Rigs/RigControlHierarchy.h"
 #include "Sequencer/MovieSceneControlRigParameterTrack.h"
 #include "Sequencer/MovieSceneControlRigParameterSection.h"
@@ -88,10 +90,15 @@ static int32 GetIndex(const TArray<FKeyPosition>& Keys, double Time)
 
 bool FAnimSliderKeySelection::Setup(TWeakPtr<ISequencer>& InSequencer)
 {
+	using namespace UE::Sequencer;
+
 	KeyMap.Reset();
 	if (InSequencer.IsValid())
 	{
-		CurveEditor = InSequencer.Pin()->GetCurveEditor();
+		TSharedPtr<FEditorViewModel> SequencerViewModel = InSequencer.Pin()->GetViewModel();
+		FCurveEditorExtension* CurveEditorExtension = SequencerViewModel->CastDynamic<FCurveEditorExtension>();
+		check(CurveEditorExtension);
+		CurveEditor = CurveEditorExtension->GetCurveEditor();
 		GetMapOfContiguousKeys();
 	}
 	return (KeyMap.Num() > 0);
diff --git a/Engine/Source/Editor/Sequencer/Private/MVVM/CurveEditorExtension.cpp b/Engine/Source/Editor/Sequencer/Private/MVVM/CurveEditorExtension.cpp
new file mode 100644
--- /dev/null
+++ b/Engine/Source/Editor/Sequencer/Private/MVVM/CurveEditorExtension.cpp
@@ -0,0 +1,425 @@
+// Copyright Epic Games, Inc. All Rights Reserved.
+
+#include "MVVM/CurveEditorExtension.h"
+
+#include "FrameNumberDetailsCustomization.h"
+#include "Framework/Docking/TabManager.h"
+#include "IPropertyRowGenerator.h"
+#include "MVVM/ViewModels/SequencerEditorViewModel.h"
+#include "SCurveEditorPanel.h"
+#include "SCurveKeyDetailPanel.h"
+#include "SSequencerTreeFilterStatusBar.h"
+#include "Sequencer.h"
+#include "SequencerCommands.h"
+#include "Toolkits/IToolkitHost.h"
+#include "Tree/SCurveEditorTree.h"
+#include "Tree/SCurveEditorTreeFilterStatusBar.h"
+#include "Tree/SCurveEditorTreeTextFilter.h"
+#include "Widgets/CurveEditor/SSequencerCurveEditor.h"
+#include "Widgets/CurveEditor/SequencerCurveEditorTimeSliderController.h"
+#include "Widgets/Docking/SDockTab.h"
+#include "Widgets/Layout/SScrollBorder.h"
+
+#define LOCTEXT_NAMESPACE "SequencerCurveEditorExtension"
+
+namespace UE
+{
+namespace Sequencer
+{
+
+class FSequencerCurveEditor : public FCurveEditor
+{
+public:
+	TWeakPtr<FSequencer> WeakSequencer;
+
+	FSequencerCurveEditor(TWeakPtr<FSequencer> InSequencer)
+		: WeakSequencer(InSequencer)
+	{}
+
+	virtual void GetGridLinesX(TArray<float>& MajorGridLines, TArray<float>& MinorGridLines, TArray<FText>* MajorGridLabels) const override
+	{
+		TSharedPtr<FSequencer> Sequencer = WeakSequencer.Pin();
+		FCurveEditorScreenSpaceH PanelInputSpace = GetPanelInputSpace();
+
+		double MajorGridStep = 0.0;
+		int32  MinorDivisions = 0;
+
+		if (Sequencer.IsValid() && Sequencer->GetGridMetrics(PanelInputSpace.GetPhysicalWidth(), PanelInputSpace.GetInputMin(), PanelInputSpace.GetInputMax(), MajorGridStep, MinorDivisions))
+		{
+			const double FirstMajorLine = FMath::FloorToDouble(PanelInputSpace.GetInputMin() / MajorGridStep) * MajorGridStep;
+			const double LastMajorLine = FMath::CeilToDouble(PanelInputSpace.GetInputMax() / MajorGridStep) * MajorGridStep;
+
+			for (double CurrentMajorLine = FirstMajorLine; CurrentMajorLine < LastMajorLine; CurrentMajorLine += MajorGridStep)
+			{
+				MajorGridLines.Add(PanelInputSpace.SecondsToScreen(CurrentMajorLine));
+
+				for (int32 Step = 1; Step < MinorDivisions; ++Step)
+				{
+					MinorGridLines.Add(PanelInputSpace.SecondsToScreen(CurrentMajorLine + Step * MajorGridStep / MinorDivisions));
+				}
+			}
+		}
+	}
+};
+
+struct FSequencerCurveEditorBounds : ICurveEditorBounds
+{
+	FSequencerCurveEditorBounds(TSharedRef<FSequencer> InSequencer)
+		: WeakSequencer(InSequencer)
+	{
+		TRange<double> Bounds = InSequencer->GetViewRange();
+		InputMin = Bounds.GetLowerBoundValue();
+		InputMax = Bounds.GetUpperBoundValue();
+	}
+
+	virtual void GetInputBounds(double& OutMin, double& OutMax) const override
+	{
+		TSharedPtr<FSequencer> Sequencer = WeakSequencer.Pin();
+		if (Sequencer.IsValid())
+		{
+			const bool bLinkTimeRange = Sequencer->GetSequencerSettings()->GetLinkCurveEditorTimeRange();
+			if (bLinkTimeRange)
+			{
+				TRange<double> Bounds = Sequencer->GetViewRange();
+				OutMin = Bounds.GetLowerBoundValue();
+				OutMax = Bounds.GetUpperBoundValue();
+			}
+			else
+			{
+				// If they don't want to link the time range with Sequencer we return the cached value.
+				OutMin = InputMin;
+				OutMax = InputMax;
+			}
+		}
+	}
+
+	virtual void SetInputBounds(double InMin, double InMax) override
+	{
+		TSharedPtr<FSequencer> Sequencer = WeakSequencer.Pin();
+		if (Sequencer.IsValid())
+		{
+			const bool bLinkTimeRange = Sequencer->GetSequencerSettings()->GetLinkCurveEditorTimeRange();
+			if (bLinkTimeRange)
+			{
+				FFrameRate TickResolution = Sequencer->GetFocusedTickResolution();
+
+				if (InMin * TickResolution > TNumericLimits<int32>::Lowest() && InMax * TickResolution < TNumericLimits<int32>::Max())
+				{
+					Sequencer->SetViewRange(TRange<double>(InMin, InMax), EViewRangeInterpolation::Immediate);
+				}
+			}
+
+			// We update these even if you are linked to the Sequencer Timeline so that when you turn off the link setting
+			// you don't pop to your last values, instead your view stays as is and just stops moving when Sequencer moves.
+			InputMin = InMin;
+			InputMax = InMax;
+		}
+	}
+
+	/** The min/max values for the viewing range. Only used if Curve Editor/Sequencer aren't linked ranges. */
+	double InputMin, InputMax;
+	TWeakPtr<FSequencer> WeakSequencer;
+};
+
+const FName FCurveEditorExtension::CurveEditorTabName = FName(TEXT("SequencerGraphEditor"));
+
+TSharedRef<IPropertyTypeCustomization> CreateFrameNumberCustomization(TWeakPtr<ISequencer> WeakSequencer)
+{
+	TSharedPtr<ISequencer> SequencerPtr = WeakSequencer.Pin();
+	return MakeShared<FFrameNumberDetailsCustomization>(SequencerPtr->GetNumericTypeInterface());
+}
+
+FCurveEditorExtension::FCurveEditorExtension()
+{
+}
+
+void FCurveEditorExtension::OnCreated(TSharedRef<FViewModel> InWeakOwner)
+{
+	ensureMsgf(!WeakOwnerModel.Pin().IsValid(), TEXT("This extension was already created!"));
+	WeakOwnerModel = InWeakOwner->CastThisShared<FSequencerEditorViewModel>();
+}
+
+void FCurveEditorExtension::CreateCurveEditor(const FTimeSliderArgs& TimeSliderArgs)
+{
+	TSharedPtr<FSequencerEditorViewModel> OwnerModel = WeakOwnerModel.Pin();
+	if (!ensure(OwnerModel))
+	{
+		return;
+	}
+
+	TSharedPtr<FSequencer> Sequencer = OwnerModel->GetSequencerImpl();
+	if (!ensure(Sequencer))
+	{
+		return;
+	}
+
+	// If they've said they want to support the curve editor then they need to provide a toolkit host
+	// so that we know where to spawn our tab into.
+	if (!ensure(Sequencer->GetToolkitHost().IsValid()))
+	{
+		return;
+	}
+
+	// Create the curve editor;
+	{
+		FCurveEditorInitParams CurveEditorInitParams;
+
+		USequencerSettings* SequencerSettings = Sequencer->GetSequencerSettings();
+
+		CurveEditorModel = MakeShared<FSequencerCurveEditor>(Sequencer);
+		CurveEditorModel->SetBounds(MakeUnique<FSequencerCurveEditorBounds>(Sequencer.ToSharedRef()));
+		CurveEditorModel->InitCurveEditor(CurveEditorInitParams);
+
+		CurveEditorModel->InputSnapEnabledAttribute = MakeAttributeLambda([SequencerSettings] { return SequencerSettings->GetIsSnapEnabled(); });
+		CurveEditorModel->OnInputSnapEnabledChanged = FOnSetBoolean::CreateLambda([SequencerSettings](bool NewValue) { SequencerSettings->SetIsSnapEnabled(NewValue); });
+
+		CurveEditorModel->OutputSnapEnabledAttribute = MakeAttributeLambda([SequencerSettings] { return SequencerSettings->GetSnapCurveValueToInterval(); });
+		CurveEditorModel->OnOutputSnapEnabledChanged = FOnSetBoolean::CreateLambda([SequencerSettings](bool NewValue) { SequencerSettings->SetSnapCurveValueToInterval(NewValue); });
+
+		CurveEditorModel->FixedGridSpacingAttribute = MakeAttributeLambda([SequencerSettings]() -> TOptional<float> { return SequencerSettings->GetGridSpacing(); });
+		CurveEditorModel->InputSnapRateAttribute = MakeAttributeSP(Sequencer.Get(), &FSequencer::GetFocusedDisplayRate);
+
+		CurveEditorModel->DefaultKeyAttributes = MakeAttributeLambda([this]() { return GetDefaultKeyAttributes(); });
+	}
+
+	// We create a custom Time Slider Controller which is just a wrapper around the actual one, but is 
+	// aware of our custom bounds logic. Currently the range the bar displays is tied to Sequencer 
+	// timeline and not the Bounds, so we need a way of changing it to look at the Bounds but only for 
+	// the Curve Editor time slider controller. We want everything else to just pass through though.
+	TSharedRef<ITimeSliderController> CurveEditorTimeSliderController = MakeShared<FSequencerCurveEditorTimeSliderController>(
+			TimeSliderArgs, Sequencer, CurveEditorModel.ToSharedRef());
+
+	CurveEditorTreeView = SNew(SCurveEditorTree, CurveEditorModel);
+	TSharedRef<SCurveEditorPanel> CurveEditorWidget = SNew(SCurveEditorPanel, CurveEditorModel.ToSharedRef())
+		// Grid lines match the color specified in FSequencerTimeSliderController::OnPaintViewArea
+		.GridLineTint(FLinearColor(0.f, 0.f, 0.f, 0.3f))
+		.ExternalTimeSliderController(CurveEditorTimeSliderController)
+		.MinimumViewPanelHeight(0.f)
+		.TabManager(Sequencer->GetToolkitHost()->GetTabManager())
+		.DisabledTimeSnapTooltip(LOCTEXT("CurveEditorTimeSnapDisabledTooltip", "Time Snapping is currently driven by Sequencer."))
+		.TreeContent()
+		[
+			SNew(SVerticalBox)
+
+			+ SVerticalBox::Slot()
+			.AutoHeight()
+			[
+				SAssignNew(CurveEditorSearchBox, SCurveEditorTreeTextFilter, CurveEditorModel)
+			]
+
+			+ SVerticalBox::Slot()
+			[
+				SNew(SScrollBorder, CurveEditorTreeView.ToSharedRef())
+				[
+					CurveEditorTreeView.ToSharedRef()
+				]
+			]
+
+			+ SVerticalBox::Slot()
+			.AutoHeight()
+			[
+				SNew(SCurveEditorTreeFilterStatusBar, CurveEditorModel)
+			]
+
+			+ SVerticalBox::Slot()
+			.AutoHeight()
+			.HAlign(HAlign_Center)
+			[
+				Sequencer->MakeTransportControls(true)
+			]
+		];
+
+	// Register an instanced custom property type layout to handle converting FFrameNumber from Tick Resolution to Display Rate.
+	TWeakPtr<ISequencer> WeakSequencer(Sequencer);
+	CurveEditorWidget->GetKeyDetailsView()->GetPropertyRowGenerator()->RegisterInstancedCustomPropertyTypeLayout(
+			"FrameNumber", 
+			FOnGetPropertyTypeCustomizationInstance::CreateStatic(CreateFrameNumberCustomization, WeakSequencer));
+
+	// And jump to the Curve Editor tree search if you have the Curve Editor focused
+	CurveEditorModel->GetCommands()->MapAction(
+		FSequencerCommands::Get().QuickTreeSearch,
+		FExecuteAction::CreateLambda([this] { FSlateApplication::Get().SetKeyboardFocus(CurveEditorSearchBox, EFocusCause::SetDirectly); })
+	);
+
+	CurveEditorPanel = SNew(SSequencerCurveEditor, CurveEditorWidget, Sequencer);
+
+	// Check to see if the tab is already opened due to the saved window layout.
+	FTabId TabId = FTabId(FCurveEditorExtension::CurveEditorTabName);
+	TSharedPtr<SDockTab> ExistingCurveEditorTab = Sequencer->GetToolkitHost()->GetTabManager()->FindExistingLiveTab(TabId);
+	if (ExistingCurveEditorTab)
+	{
+		ExistingCurveEditorTab->SetContent(CurveEditorPanel.ToSharedRef());
+	}
+}
+
+void FCurveEditorExtension::OpenCurveEditor()
+{
+	TSharedPtr<FSequencerEditorViewModel> OwnerModel = WeakOwnerModel.Pin();
+	if (!ensure(OwnerModel))
+	{
+		return;
+	}
+
+	TSharedPtr<ISequencer> Sequencer = OwnerModel->GetSequencer();
+	if (!Sequencer)
+	{
+		return;
+	}
+
+	// Request the Tab Manager invoke the tab. This will spawn the tab if needed, otherwise pull it to focus. This assumes
+	// that the Toolkit Host's Tab Manager has already registered a tab with a NullWidget for content.
+	FTabId TabId = FTabId(FCurveEditorExtension::CurveEditorTabName);
+	TSharedPtr<SDockTab> CurveEditorTab = Sequencer->GetToolkitHost()->GetTabManager()->TryInvokeTab(TabId);
+	if (CurveEditorTab.IsValid())
+	{
+		CurveEditorTab->SetContent(CurveEditorPanel.ToSharedRef());
+
+		const FSlateIcon SequencerGraphIcon = FSlateIcon(FAppStyle::GetAppStyleSetName(), "GenericCurveEditor.TabIcon");
+		CurveEditorTab->SetTabIcon(SequencerGraphIcon.GetIcon());
+
+		CurveEditorTab->SetLabel(LOCTEXT("SequencerMainGraphEditorTitle", "Sequencer Curves"));
+
+		CurveEditorModel->ZoomToFit();
+	}
+}
+
+bool FCurveEditorExtension::IsCurveEditorOpen() const
+{
+	TSharedPtr<FSequencerEditorViewModel> OwnerModel = WeakOwnerModel.Pin();
+	if (!ensure(OwnerModel))
+	{
+		return false;
+	}
+
+	TSharedPtr<ISequencer> Sequencer = OwnerModel->GetSequencer();
+	if (!Sequencer)
+	{
+		return false;
+	}
+
+	FTabId TabId = FTabId(FCurveEditorExtension::CurveEditorTabName);
+	return Sequencer->GetToolkitHost()->GetTabManager()->FindExistingLiveTab(TabId).IsValid();
+}
+
+void FCurveEditorExtension::CloseCurveEditor()
+{
+	TSharedPtr<FSequencerEditorViewModel> OwnerModel = WeakOwnerModel.Pin();
+	if (!ensure(OwnerModel))
+	{
+		return;
+	}
+
+	TSharedPtr<ISequencer> Sequencer = OwnerModel->GetSequencer();
+	if (!Sequencer)
+	{
+		return;
+	}
+
+	FTabId TabId = FTabId(FCurveEditorExtension::CurveEditorTabName);
+	TSharedPtr<SDockTab> CurveEditorTab = Sequencer->GetToolkitHost()->GetTabManager()->FindExistingLiveTab(TabId);
+	if (CurveEditorTab)
+	{
+		CurveEditorTab->RequestCloseTab();
+	}
+}
+
+FKeyAttributes FCurveEditorExtension::GetDefaultKeyAttributes() const
+{
+	TSharedPtr<FSequencerEditorViewModel> OwnerModel = WeakOwnerModel.Pin();
+	check(OwnerModel);
+	TSharedPtr<ISequencer> Sequencer = OwnerModel->GetSequencer();
+	check(Sequencer);
+	USequencerSettings* Settings = Sequencer->GetSequencerSettings();
+	check(Settings);
+
+	switch (Settings->GetKeyInterpolation())
+	{
+	case EMovieSceneKeyInterpolation::User:     return FKeyAttributes().SetInterpMode(RCIM_Cubic).SetTangentMode(RCTM_User);
+	case EMovieSceneKeyInterpolation::Break:    return FKeyAttributes().SetInterpMode(RCIM_Cubic).SetTangentMode(RCTM_Break);
+	case EMovieSceneKeyInterpolation::Linear:   return FKeyAttributes().SetInterpMode(RCIM_Linear).SetTangentMode(RCTM_Auto);
+	case EMovieSceneKeyInterpolation::Constant: return FKeyAttributes().SetInterpMode(RCIM_Constant).SetTangentMode(RCTM_Auto);
+	default:                                    return FKeyAttributes().SetInterpMode(RCIM_Cubic).SetTangentMode(RCTM_Auto);
+	}
+}
+
+void FCurveEditorExtension::RequestSyncSelection()
+{
+	// We schedule selection syncing to the next editor tick because we might want to select items that
+	// have just been added to the curve editor tree this tick. If it happened after the Slate update,
+	// these items don't yet have a UI widget, and so selecting them doesn't do anything.
+	//
+	// Note that we capture a weak pointer of our owner model because selection changes can happen
+	// right around the time when we want to unload everything (such as when loading a new map in the
+	// editor). We don't want to extend the lifetime of our stuff in that case.
+	TWeakPtr<FSequencerEditorViewModel> WeakRootViewModel(WeakOwnerModel);
+	GEditor->GetTimerManager()->SetTimerForNextTick([WeakRootViewModel]()
+	{
+		TSharedPtr<FSequencerEditorViewModel> RootViewModel = WeakRootViewModel.Pin();
+		if (!RootViewModel.IsValid())
+		{
+			return;
+		}
+
+		FCurveEditorExtension* This = RootViewModel->CastDynamic<FCurveEditorExtension>();
+		if (This)
+		{
+			This->SyncSelection();
+		}
+	});
+}
+
+void FCurveEditorExtension::SyncSelection()
+{
+	if (!ensure(CurveEditorModel && CurveEditorTreeView))
+	{
+		return;
+	}
+
+	TSharedPtr<FSequencerEditorViewModel> OwnerModel = WeakOwnerModel.Pin();
+	if (!ensure(OwnerModel))
+	{
+		return;
+	}
+
+	TSharedPtr<ISequencer> Sequencer = OwnerModel->GetSequencer();
+	if (!ensure(Sequencer))
+	{
+		return;
+	}
+
+	CurveEditorModel->SuspendBroadcast();
+
+	CurveEditorTreeView->ClearSelection();
+
+	FCurveEditorTreeItemID FirstCurveEditorTreeItemID;
+	const TSet<TWeakPtr<FViewModel>>& SelectedItems = Sequencer->GetSelection().GetSelectedOutlinerItems();
+	for (TWeakPtr<FViewModel> SelectedItem : SelectedItems)
+	{
+		if (ICurveEditorTreeItemExtension* CurveEditorItem = ICastable::CastWeakPtr<ICurveEditorTreeItemExtension>(SelectedItem))
+		{
+			FCurveEditorTreeItemID CurveEditorTreeItem = CurveEditorItem->GetCurveEditorItemID();
+			if (CurveEditorTreeItem != FCurveEditorTreeItemID::Invalid())
+			{
+				if (!CurveEditorTreeView->IsItemSelected(CurveEditorTreeItem))
+				{
+					CurveEditorTreeView->SetItemSelection(CurveEditorTreeItem, true);
+					if (!FirstCurveEditorTreeItemID.IsValid())
+					{
+						FirstCurveEditorTreeItemID = CurveEditorTreeItem;
+					}
+				}
+			}
+		}
+	}
+	if (FirstCurveEditorTreeItemID.IsValid())
+	{
+		CurveEditorTreeView->RequestScrollIntoView(FirstCurveEditorTreeItemID);
+	}
+
+	CurveEditorModel->ResumeBroadcast();
+}
+
+} // namespace Sequencer
+} // namespace UE
+
+#undef LOCTEXT_NAMESPACE
diff --git a/Engine/Source/Editor/Sequencer/Private/MVVM/CurveEditorIntegrationExtension.cpp b/Engine/Source/Editor/Sequencer/Private/MVVM/CurveEditorIntegrationExtension.cpp
--- a/Engine/Source/Editor/Sequencer/Private/MVVM/CurveEditorIntegrationExtension.cpp
+++ b/Engine/Source/Editor/Sequencer/Private/MVVM/CurveEditorIntegrationExtension.cpp
@@ -1,11 +1,13 @@
 // Copyright Epic Games, Inc. All Rights Reserved.
 
 #include "MVVM/CurveEditorIntegrationExtension.h"
+#include "MVVM/CurveEditorExtension.h"
 #include "MVVM/TrackModelStorageExtension.h"
 #include "MVVM/ViewModels/ViewModel.h"
 #include "MVVM/SharedViewModelData.h"
 #include "MVVM/ViewModelPtr.h"
 #include "MVVM/ViewModels/SequenceModel.h"
+#include "MVVM/ViewModels/SequencerEditorViewModel.h"
 #include "MVVM/ViewModels/ViewModelIterators.h"
 #include "MVVM/Extensions/ICurveEditorTreeItemExtension.h"
 #include "CurveEditor.h"
@@ -33,6 +35,18 @@
 	UpdateCurveEditor();
 }
 
+FCurveEditorExtension* FCurveEditorIntegrationExtension::GetCurveEditorExtension()
+{
+	if (TSharedPtr<FSequenceModel> OwnerModel = WeakOwnerModel.Pin())
+	{
+		if (TSharedPtr<FSequencerEditorViewModel> EditorViewModel = OwnerModel->GetEditor())
+		{
+			return EditorViewModel->CastDynamic<FCurveEditorExtension>();
+		}
+	}
+	return nullptr;
+}
+
 void FCurveEditorIntegrationExtension::UpdateCurveEditor()
 {
 	TSharedPtr<FSequenceModel> OwnerModel = WeakOwnerModel.Pin();
@@ -40,15 +54,15 @@
 	{
 		return;
 	}
-
-	ICurveEditorExtension* CurveEditorExtension = OwnerModel->CastThis<ICurveEditorExtension>();
-	if (!CurveEditorExtension)
+	
+	FCurveEditorExtension* CurveEditorExtension = GetCurveEditorExtension();
+	if (!ensure(CurveEditorExtension))
 	{
 		return;
 	}
 
 	TSharedPtr<FCurveEditor> CurveEditor = CurveEditorExtension->GetCurveEditor();
-	if (!CurveEditor)
+	if (!ensure(CurveEditor))
 	{
 		return;
 	}
@@ -121,7 +135,7 @@
 	}
 }
 
-FCurveEditorTreeItemID FCurveEditorIntegrationExtension::AddToCurveEditor(TViewModelPtr<ICurveEditorTreeItemExtension> InViewModel, TSharedPtr<FCurveEditor> CurveEditor)
+FCurveEditorTreeItemID FCurveEditorIntegrationExtension::AddToCurveEditor(TViewModelPtr<ICurveEditorTreeItemExtension> InViewModel, TSharedPtr<FCurveEditor> InCurveEditor)
 {
 	// If the view model doesn't want to be in the curve editor, bail out
 	// Note that this means we will create curve editor items for each parent in the hierarchy up
@@ -135,7 +149,7 @@
 	// Check if we already have a valid curve editor ID
 	if (FCurveEditorTreeItemID* Existing = ViewModelToTreeItemIDMap.Find(InViewModel.AsModel()))
 	{
-		if (CurveEditor->GetTree()->FindItem(*Existing) != nullptr)
+		if (InCurveEditor->GetTree()->FindItem(*Existing) != nullptr)
 		{
 			return *Existing;
 		}
@@ -144,38 +158,32 @@
 	// Recursively create any needed parent curve editor items
 	TSharedPtr<FViewModel> Parent = InViewModel.AsModel()->GetParent();
 
-	FCurveEditorTreeItemID ParentID = Parent.IsValid() ? 
-		AddToCurveEditor(Parent, CurveEditor) : 
+	FCurveEditorTreeItemID ParentID = Parent.IsValid() ?
+		AddToCurveEditor(Parent, InCurveEditor) :
 		FCurveEditorTreeItemID::Invalid();
 
 	// Create the new curve editor item
-	FCurveEditorTreeItem* NewItem = CurveEditor->AddTreeItem(ParentID);
+	FCurveEditorTreeItem* NewItem = InCurveEditor->AddTreeItem(ParentID);
 	TSharedPtr<ICurveEditorTreeItem> CurveEditorTreeItem = InViewModel->GetCurveEditorTreeItem();
 	NewItem->SetWeakItem(CurveEditorTreeItem);
 
 	// Register the new ID in our map and notify the view model
 	ViewModelToTreeItemIDMap.Add(InViewModel.AsModel(), NewItem->GetID());
-	InViewModel->OnAddedToCurveEditor(NewItem->GetID(), CurveEditor);
+	InViewModel->OnAddedToCurveEditor(NewItem->GetID(), InCurveEditor);
 
 	return NewItem->GetID();
 }
 
-void FCurveEditorIntegrationExtension::RecreateCurveEditor()
+void FCurveEditorIntegrationExtension::ResetCurveEditor()
 {
-	TSharedPtr<FSequenceModel> OwnerModel = WeakOwnerModel.Pin();
-	if (!OwnerModel)
-	{
-		return;
-	}
-
-	ICurveEditorExtension* CurveEditorExtension = OwnerModel->CastThis<ICurveEditorExtension>();
-	if (!CurveEditorExtension)
+	FCurveEditorExtension* CurveEditorExtension = GetCurveEditorExtension();
+	if (!ensure(CurveEditorExtension))
 	{
 		return;
 	}
 
 	TSharedPtr<FCurveEditor> CurveEditor = CurveEditorExtension->GetCurveEditor();
-	if (!CurveEditor)
+	if (!ensure(CurveEditor))
 	{
 		return;
 	}
diff --git a/Engine/Source/Editor/Sequencer/Private/MVVM/SequencerViewModelTypeIDs.cpp b/Engine/Source/Editor/Sequencer/Private/MVVM/SequencerViewModelTypeIDs.cpp
--- a/Engine/Source/Editor/Sequencer/Private/MVVM/SequencerViewModelTypeIDs.cpp
+++ b/Engine/Source/Editor/Sequencer/Private/MVVM/SequencerViewModelTypeIDs.cpp
@@ -22,6 +22,7 @@
 #include "MVVM/ViewModels/SequencerEditorViewModel.h"
 #include "MVVM/ViewModels/SequencerTrackAreaViewModel.h"
 
+#include "MVVM/CurveEditorExtension.h"
 #include "MVVM/CurveEditorIntegrationExtension.h"
 #include "MVVM/FolderModelStorageExtension.h"
 #include "MVVM/ObjectBindingModelStorageExtension.h"
@@ -61,13 +62,13 @@ UE_SEQUENCER_DEFINE_CASTABLE(FSequencerOutlinerViewModel);
 UE_SEQUENCER_DEFINE_CASTABLE(FSequencerTrackAreaViewModel);
 
 // Extension types
+UE_SEQUENCER_DEFINE_VIEW_MODEL_TYPE_ID(FCurveEditorExtension);
 UE_SEQUENCER_DEFINE_VIEW_MODEL_TYPE_ID(FCurveEditorIntegrationExtension);
 UE_SEQUENCER_DEFINE_VIEW_MODEL_TYPE_ID(FFolderModelStorageExtension);
 UE_SEQUENCER_DEFINE_VIEW_MODEL_TYPE_ID(FObjectBindingModelStorageExtension);
 UE_SEQUENCER_DEFINE_VIEW_MODEL_TYPE_ID(FSectionModelStorageExtension);
 UE_SEQUENCER_DEFINE_VIEW_MODEL_TYPE_ID(FTrackModelStorageExtension);
 UE_SEQUENCER_DEFINE_VIEW_MODEL_TYPE_ID(FTrackRowModelStorageExtension);
-UE_SEQUENCER_DEFINE_VIEW_MODEL_TYPE_ID(ICurveEditorExtension);
 UE_SEQUENCER_DEFINE_VIEW_MODEL_TYPE_ID(ICurveEditorTreeItemExtension);
 UE_SEQUENCER_DEFINE_VIEW_MODEL_TYPE_ID(IObjectBindingExtension);
 UE_SEQUENCER_DEFINE_VIEW_MODEL_TYPE_ID(ISnappableExtension);
diff --git a/Engine/Source/Editor/Sequencer/Private/MVVM/ViewModels/SequenceModel.cpp b/Engine/Source/Editor/Sequencer/Private/MVVM/ViewModels/SequenceModel.cpp
--- a/Engine/Source/Editor/Sequencer/Private/MVVM/ViewModels/SequenceModel.cpp
+++ b/Engine/Source/Editor/Sequencer/Private/MVVM/ViewModels/SequenceModel.cpp
@@ -139,11 +139,6 @@
 {
 }
 
-TSharedPtr<FCurveEditor> FSequenceModel::GetCurveEditor() const
-{
-	return GetSequencerImpl()->GetCurveEditor();
-}
-
 TOptional<EItemDropZone> FSequenceModel::CanAcceptDrop(const FViewModelPtr& TargetModel, const FDragDropEvent& DragDropEvent, EItemDropZone ItemDropZone)
 {
 	TSharedPtr<FSequencerOutlinerDragDropOp> DragDropOp = DragDropEvent.GetOperationAs<FSequencerOutlinerDragDropOp>();
diff --git a/Engine/Source/Editor/Sequencer/Private/MVVM/ViewModels/SequencerEditorViewModel.cpp b/Engine/Source/Editor/Sequencer/Private/MVVM/ViewModels/SequencerEditorViewModel.cpp
--- a/Engine/Source/Editor/Sequencer/Private/MVVM/ViewModels/SequencerEditorViewModel.cpp
+++ b/Engine/Source/Editor/Sequencer/Private/MVVM/ViewModels/SequencerEditorViewModel.cpp
@@ -3,6 +3,7 @@
 #include "MVVM/ViewModels/SequencerEditorViewModel.h"
 #include "MVVM/ViewModels/SequencerOutlinerViewModel.h"
 #include "MVVM/ViewModels/SequencerTrackAreaViewModel.h"
+#include "MVVM/CurveEditorExtension.h"
 #include "Sequencer.h"
 #include "MovieSceneSequenceID.h"
 
@@ -35,6 +36,11 @@
 	return MakeShared<FSequencerTrackAreaViewModel>(Sequencer.ToSharedRef());
 }
 
+void FSequencerEditorViewModel::InitializeEditorImpl()
+{
+	AddDynamicExtension(FCurveEditorExtension::ID);
+}
+
 TSharedPtr<ISequencer> FSequencerEditorViewModel::GetSequencer() const
 {
 	return WeakSequencer.Pin();
diff --git a/Engine/Source/Editor/Sequencer/Private/SSequencer.cpp b/Engine/Source/Editor/Sequencer/Private/SSequencer.cpp
--- a/Engine/Source/Editor/Sequencer/Private/SSequencer.cpp
+++ b/Engine/Source/Editor/Sequencer/Private/SSequencer.cpp
@@ -3,6 +3,7 @@
 #include "SSequencer.h"
 #include "Engine/Blueprint.h"
 #include "MovieSceneSequence.h"
+#include "MVVM/CurveEditorExtension.h"
 #include "MVVM/ViewModels/CategoryModel.h"
 #include "MVVM/ViewModels/ChannelModel.h"
 #include "MVVM/ViewModels/ViewModelIterators.h"
@@ -115,216 +116,6 @@
 
 #define LOCTEXT_NAMESPACE "Sequencer"
 
-const FName SSequencer::CurveEditorTabName = FName(TEXT("SequencerGraphEditor"));
-
-TSharedRef<IPropertyTypeCustomization> CreateFrameNumberCustomization(TWeakPtr<FSequencer> WeakSequencer)
-{
-	TSharedPtr<ISequencer> SequencerPtr = WeakSequencer.Pin();
-	return MakeShared<FFrameNumberDetailsCustomization>(SequencerPtr->GetNumericTypeInterface());
-}
-
-class SSequencerCurveEditor : public SCompoundWidget
-{
-	SLATE_BEGIN_ARGS(SSequencerCurveEditor)
-	{}
-	SLATE_END_ARGS()
-public:
-	void Construct(const FArguments& InArgs, TSharedRef<SCurveEditorPanel> InEditorPanel, TSharedPtr<FSequencer> InSequencer)
-	{
-		WeakSequencer = InSequencer;
-
-		ChildSlot
-		[
-			SNew(SVerticalBox)
-
-			+ SVerticalBox::Slot()
-			.AutoHeight()
-			[
-				MakeToolbar(InEditorPanel)
-			]
-			+ SVerticalBox::Slot()
-			.FillHeight(1.0f)
-			[
-				InEditorPanel
-			]
-		];
-	}
-
-	TSharedRef<SWidget> MakeToolbar(TSharedRef<SCurveEditorPanel> InEditorPanel)
-	{
-		FSlimHorizontalToolBarBuilder ToolBarBuilder(InEditorPanel->GetCommands(), FMultiBoxCustomization::None, InEditorPanel->GetToolbarExtender(), true);
-		ToolBarBuilder.BeginSection("Asset");
-
-		{
-			TAttribute<FSlateIcon> SaveIcon;
-			SaveIcon.Bind(TAttribute<FSlateIcon>::FGetter::CreateLambda([&] {
-
-				TArray<UMovieScene*> MovieScenesToSave;
-				MovieSceneHelpers::GetDescendantMovieScenes(WeakSequencer.Pin()->GetRootMovieSceneSequence(), MovieScenesToSave);
-				for (UMovieScene* MovieSceneToSave : MovieScenesToSave)
-				{
-					UPackage* MovieScenePackageToSave = MovieSceneToSave->GetOuter()->GetOutermost();
-					if (MovieScenePackageToSave->IsDirty())
-					{
-						return FSlateIcon(FAppStyle::Get().GetStyleSetName(), "Icons.SaveChanged");
-					}
-				}
-
-				return FSlateIcon(FAppStyle::Get().GetStyleSetName(), "Icons.Save");
-				}));
-
-			if (WeakSequencer.Pin()->GetHostCapabilities().bSupportsSaveMovieSceneAsset)
-			{
-				ToolBarBuilder.AddToolBarButton(
-					FUIAction(FExecuteAction::CreateLambda([&] { WeakSequencer.Pin()->SaveCurrentMovieScene(); })),
-					NAME_None,
-					LOCTEXT("SaveDirtyPackages", "Save"),
-					LOCTEXT("SaveDirtyPackagesTooltip", "Saves the current sequence and any subsequences"),
-					SaveIcon
-				);
-			}		
-		}
-		
-		ToolBarBuilder.EndSection();
-		// We just use all of the extenders as our toolbar, we don't have a need to create a separate toolbar.
-		return ToolBarBuilder.MakeWidget();
-	}
-
-private:
-	TWeakPtr<FSequencer> WeakSequencer;
-};
-
-class FSequencerCurveEditorTimeSliderController : public FSequencerTimeSliderController
-{
-public:
-
-	FSequencerCurveEditorTimeSliderController(const FTimeSliderArgs& InArgs, TWeakPtr<FSequencer> InWeakSequencer, TSharedRef<FCurveEditor> InCurveEditor)
-		: FSequencerTimeSliderController(InArgs, InWeakSequencer)
-	{
-		WeakSequencer = InWeakSequencer;
-		WeakCurveEditor = InCurveEditor;
-	}
-
-	virtual void ClampViewRange(double& NewRangeMin, double& NewRangeMax) override
-	{
-		// Since the CurveEditor uses a different view range (potentially) we have to be careful about which one we clamp.
-		TSharedPtr<ISequencer> Sequencer = WeakSequencer.Pin();
-		if (!Sequencer.IsValid())
-		{
-			return;
-		}
-
-		const bool bLinkedTimeRange = Sequencer->GetSequencerSettings()->GetLinkCurveEditorTimeRange();
-		if (bLinkedTimeRange)
-		{
-			return FSequencerTimeSliderController::ClampViewRange(NewRangeMin, NewRangeMax);
-		}
-		else
-		{
-			TSharedPtr<FCurveEditor> CurveEditor = WeakCurveEditor.Pin();
-			if (CurveEditor.IsValid())
-			{
-				double InputMin, InputMax;
-				CurveEditor->GetBounds().GetInputBounds(InputMin, InputMax);
-
-				bool bNeedsClampSet = false;
-				double NewClampRangeMin = InputMin;
-				if (NewRangeMin < InputMin)
-				{
-					NewClampRangeMin = NewRangeMin;
-					bNeedsClampSet = true;
-				}
-
-				double NewClampRangeMax = InputMax;
-				if (NewRangeMax > InputMax)
-				{
-					NewClampRangeMax = NewRangeMax;
-					bNeedsClampSet = true;
-				}
-
-				if (bNeedsClampSet)
-				{
-					CurveEditor->GetBounds().SetInputBounds(NewClampRangeMin, NewClampRangeMax);
-				}
-
-			}
-		}
-	}
-
-	virtual void SetViewRange(double NewRangeMin, double NewRangeMax, EViewRangeInterpolation Interpolation) override
-	{
-		TSharedPtr<ISequencer> Sequencer = WeakSequencer.Pin();
-		if (!Sequencer.IsValid())
-		{
-			return;
-		}
-
-		const bool bLinkedTimeRange = Sequencer->GetSequencerSettings()->GetLinkCurveEditorTimeRange();
-		if (bLinkedTimeRange)
-		{
-			return FSequencerTimeSliderController::SetViewRange(NewRangeMin, NewRangeMax, Interpolation);
-		}
-		else
-		{
-			// Clamp to a minimum size to avoid zero-sized or negative visible ranges
-			double MinVisibleTimeRange = FFrameNumber(1) / GetTickResolution();
-			TRange<double> ExistingViewRange = GetViewRange();
-
-			if (NewRangeMax == ExistingViewRange.GetUpperBoundValue())
-			{
-				if (NewRangeMin > NewRangeMax - MinVisibleTimeRange)
-				{
-					NewRangeMin = NewRangeMax - MinVisibleTimeRange;
-				}
-			}
-			else if (NewRangeMax < NewRangeMin + MinVisibleTimeRange)
-			{
-				NewRangeMax = NewRangeMin + MinVisibleTimeRange;
-			}
-
-			TSharedPtr<FCurveEditor> CurveEditor = WeakCurveEditor.Pin();
-			if (CurveEditor.IsValid())
-			{
-				CurveEditor->GetBounds().SetInputBounds(NewRangeMin, NewRangeMax);
-			}
-		}
-	}
-
-
-
-	virtual FAnimatedRange GetViewRange() const override
-	{ 
-		// If they've linked the Sequencer timerange we can return the internal controller's view range, otherwise we return the bounds (which internally does the same check)
-		TSharedPtr<ISequencer> Sequencer = WeakSequencer.Pin();
-		if (!Sequencer.IsValid())
-		{
-			return FAnimatedRange();
-		}
-		const bool bLinkedTimeRange = Sequencer->GetSequencerSettings()->GetLinkCurveEditorTimeRange();
-		if (bLinkedTimeRange)
-		{
-			return FSequencerTimeSliderController::GetViewRange();
-		}
-		else
-		{
-			TSharedPtr<FCurveEditor> CurveEditor = WeakCurveEditor.Pin();
-			if (CurveEditor.IsValid())
-			{
-				double InputMin, InputMax;
-				CurveEditor->GetBounds().GetInputBounds(InputMin, InputMax);
-
-				return FAnimatedRange(InputMin, InputMax);
-			}
-		}
-
-		return FAnimatedRange();
-	}
-
-private:
-	TWeakPtr<ISequencer> WeakSequencer;
-	TWeakPtr<FCurveEditor> WeakCurveEditor;
-};
-
 /* SSequencer interface
  *****************************************************************************/
 PRAGMA_DISABLE_OPTIMIZATION
@@ -498,66 +289,9 @@ void SSequencer::Construct(const FArguments& InArgs, TSharedRef<FSequencer> InSe
 
 	if (InSequencer->GetHostCapabilities().bSupportsCurveEditor)
 	{
-		// If they've said they want to support the curve editor then they need to provide a toolkit host
-		// so that we know where to spawn our tab into.
-		check(InSequencer->GetToolkitHost().IsValid());
-
-		// We create a custom Time Slider Controller which is just a wrapper around the actual one, but is aware of our custom bounds logic. Currently the range the
-		// bar displays is tied to Sequencer timeline and not the Bounds, so we need a way of changing it to look at the Bounds but only for the Curve Editor time
-		// slider controller. We want everything else to just pass through though.
-		TSharedRef<ITimeSliderController> CurveEditorTimeSliderController = MakeShared<FSequencerCurveEditorTimeSliderController>(TimeSliderArgs, SequencerPtr, InSequencer->GetCurveEditor().ToSharedRef());
-
-		CurveEditorTree = SNew(SCurveEditorTree, InSequencer->GetCurveEditor());
-		TSharedRef<SCurveEditorPanel> CurveEditorWidget = SNew(SCurveEditorPanel, InSequencer->GetCurveEditor().ToSharedRef())
-			// Grid lines match the color specified in FSequencerTimeSliderController::OnPaintViewArea
-			.GridLineTint(FLinearColor(0.f, 0.f, 0.f, 0.3f))
-			.ExternalTimeSliderController(CurveEditorTimeSliderController)
-			.MinimumViewPanelHeight(0.f)
-			.TabManager(InSequencer->GetToolkitHost()->GetTabManager())
-			.DisabledTimeSnapTooltip(LOCTEXT("CurveEditorTimeSnapDisabledTooltip", "Time Snapping is currently driven by Sequencer."))
-			.TreeContent()
-			[
-				SNew(SVerticalBox)
-
-				+ SVerticalBox::Slot()
-				.AutoHeight()
-				[
-					SAssignNew(CurveEditorSearchBox, SCurveEditorTreeTextFilter, InSequencer->GetCurveEditor())
-				]
-
-				+ SVerticalBox::Slot()
-				[
-					SNew(SScrollBorder, CurveEditorTree.ToSharedRef())
-					[
-						CurveEditorTree.ToSharedRef()
-					]
-				]
-
-				+ SVerticalBox::Slot()
-				.AutoHeight()
-				[
-					SNew(SCurveEditorTreeFilterStatusBar, InSequencer->GetCurveEditor())
-				]
-
-				+ SVerticalBox::Slot()
-				.AutoHeight()
-				.HAlign(HAlign_Center)
-				[
-					InSequencer->MakeTransportControls(true)
-				]
-			];
-
-		// Register an instanced custom property type layout to handle converting FFrameNumber from Tick Resolution to Display Rate.
-		CurveEditorWidget->GetKeyDetailsView()->GetPropertyRowGenerator()->RegisterInstancedCustomPropertyTypeLayout("FrameNumber", FOnGetPropertyTypeCustomizationInstance::CreateStatic(CreateFrameNumberCustomization, SequencerPtr));
-	
-		CurveEditorPanel = SNew(SSequencerCurveEditor, CurveEditorWidget, InSequencer);
-
-		// Check to see if the tab is already opened due to the saved window layout.
-		TSharedPtr<SDockTab> ExistingCurveEditorTab = InSequencer->GetToolkitHost()->GetTabManager()->FindExistingLiveTab(FTabId(SSequencer::CurveEditorTabName));
-		if (ExistingCurveEditorTab)
-		{
-			ExistingCurveEditorTab->SetContent(CurveEditorPanel.ToSharedRef());
-		}
+		TSharedPtr<FEditorViewModel> RootModel = InSequencer->GetViewModel();
+		FCurveEditorExtension* CurveEditorExtension = RootModel->CastDynamicChecked<FCurveEditorExtension>();
+		CurveEditorExtension->CreateCurveEditor(TimeSliderArgs);
 	}
 
 
@@ -1153,12 +887,6 @@ void SSequencer::BindCommands(TSharedRef<FUICommandList> SequencerCommandBinding
 		FSequencerCommands::Get().QuickTreeSearch,
 		FExecuteAction::CreateLambda([this] { FSlateApplication::Get().SetKeyboardFocus(SearchBox, EFocusCause::SetDirectly); })
 	);
-
-	// And jump to the Curve Editor tree search if you have the Curve Editor focused
-	SequencerPtr.Pin()->GetCurveEditor()->GetCommands()->MapAction(
-		FSequencerCommands::Get().QuickTreeSearch,
-		FExecuteAction::CreateLambda([this] { FSlateApplication::Get().SetKeyboardFocus(CurveEditorSearchBox, EFocusCause::SetDirectly); })
-	);
 	
 	SequencerCommandBindings->MapAction(
 		FSequencerCommands::Get().ToggleShowStretchBox,
@@ -1301,10 +1029,17 @@ void SSequencer::HandleOutlinerNodeSelectionChanged()
 		return;
 	}
 
+	FViewModelPtr RootModel = Sequencer->GetViewModel();
+	FCurveEditorExtension* CurveEditorExtension = RootModel->CastDynamic<FCurveEditorExtension>();
+	if (!CurveEditorExtension)
+	{
+		return;
+	}
+
 	const TSet<TWeakPtr<FViewModel>>& SelectedDisplayNodes = Sequencer->GetSelection().GetSelectedOutlinerItems();
 
-	TSharedPtr<FCurveEditor> CurveEditor = Sequencer->GetCurveEditor();
-	if (CurveEditor && CurveEditorTree)
+	TSharedPtr<FCurveEditor> CurveEditor = CurveEditorExtension->GetCurveEditor();
+	if (CurveEditor)
 	{
 		const USequencerSettings* SequencerSettings = GetSequencerSettings();
 		// If we're isolating to the selection and there is one, add the filter
@@ -1328,41 +1063,11 @@ void SSequencer::HandleOutlinerNodeSelectionChanged()
 
 		if (GetSequencerSettings()->ShouldSyncCurveEditorSelection())
 		{
-			//since Selection ticks before any filter is applied and the curve editor tree is udpated in FSequencer::Tick we need to sync 
-			//the curve editor on the next tick so the filter is applied first, otherwise the curve tree hierarchy won't be up to date
-			//and the curve may not be found.
-			TWeakPtr<FSequencer> WeakSequencer = SequencerPtr;
-			GEditor->GetTimerManager()->SetTimerForNextTick([WeakSequencer, this]()
-			{
-				TSharedPtr<FSequencer> Sequencer = WeakSequencer.Pin();
-				if (!Sequencer.IsValid())
-				{
-					return;
-				}
-
-				const TSet<TWeakPtr<FViewModel>>& SelectedDisplayNodes = Sequencer->GetSelection().GetSelectedOutlinerItems();
-				TSharedPtr<FCurveEditor> CurveEditor = Sequencer->GetCurveEditor();
-				TSharedRef<FSequencerNodeTree> NodeTree = Sequencer->GetNodeTree();
-				CurveEditor->SuspendBroadcast();
-				// Clear the tree selection
-				CurveEditorTree->ClearSelection();
-				for (TWeakPtr<FViewModel> Node : SelectedDisplayNodes)
-				{
-					if (ICurveEditorTreeItemExtension* CurveEditorItem = ICastable::CastWeakPtr<ICurveEditorTreeItemExtension>(Node))
-					{
-						FCurveEditorTreeItemID CurveEditorTreeItem = CurveEditorItem->GetCurveEditorItemID();
-						if (CurveEditorTreeItem != FCurveEditorTreeItemID::Invalid())
-						{
-							if (!CurveEditorTree->IsItemSelected(CurveEditorTreeItem))
-							{
-								CurveEditorTree->SetItemSelection(CurveEditorTreeItem, true);
-								CurveEditorTree->RequestScrollIntoView(CurveEditorTreeItem);
-							}
-						}
-					}
-				}
-				CurveEditor->ResumeBroadcast();
-			});
+			// We schedule a selection synchronization for the next update. This synchronization must happen
+			// after all filters have been applied, because the items we want to select in the curve editor
+			// might be currently filtered out, but will be visible when filters are re-evaluated. This is
+			// why curve editor integration runs in FSequencerNodeTree after filtering.
+			CurveEditorExtension->RequestSyncSelection();
 		}
 	}
 
@@ -2988,20 +2693,18 @@ TSharedPtr<ITimeSlider> SSequencer::GetTopTimeSliderWidget() const
 
 SSequencer::~SSequencer()
 {
-	USelection::SelectionChangedEvent.RemoveAll(this);
+	using namespace UE::Sequencer;
 
+	USelection::SelectionChangedEvent.RemoveAll(this);
 
 	TSharedPtr<FSequencer> Sequencer = SequencerPtr.Pin();
 	if(Sequencer && Sequencer->GetToolkitHost() && Sequencer->GetToolkitHost()->GetTabManager())
 	{
 		if (Sequencer->GetHostCapabilities().bSupportsCurveEditor)
 		{
-			FTabId TabId = FTabId(SSequencer::CurveEditorTabName);
-			TSharedPtr<SDockTab> CurveEditorTab = Sequencer->GetToolkitHost()->GetTabManager()->FindExistingLiveTab(TabId);
-			if (CurveEditorTab)
-			{
-				CurveEditorTab->RequestCloseTab();
-			}
+			TSharedPtr<FEditorViewModel> RootModel = Sequencer->GetViewModel();
+			FCurveEditorExtension* CurveEditorExtension = RootModel->CastDynamicChecked<FCurveEditorExtension>();
+			CurveEditorExtension->CloseCurveEditor();
 		}
 	}
 
@@ -3902,8 +3605,9 @@ void SSequencer::OnColumnFillCoefficientChanged(float FillCoefficient, int32 Col
 
 void SSequencer::OnCurveEditorVisibilityChanged(bool bShouldBeVisible)
 {
+	using namespace UE::Sequencer;
+
 	TSharedPtr<FSequencer> Sequencer = SequencerPtr.Pin();
-	FTabId TabId = FTabId(SSequencer::CurveEditorTabName);
 
 	// Curve Editor may not be supported
 	if (!Sequencer->GetHostCapabilities().bSupportsCurveEditor)
@@ -3911,30 +3615,16 @@ void SSequencer::OnCurveEditorVisibilityChanged(bool bShouldBeVisible)
 		return;
 	}
 
+	TSharedPtr<FEditorViewModel> RootModel = Sequencer->GetViewModel();
+	FCurveEditorExtension* CurveEditorExtension = RootModel->CastDynamicChecked<FCurveEditorExtension>();
+
 	if (bShouldBeVisible)
 	{
-		// Request the Tab Manager invoke the tab. This will spawn the tab if needed, otherwise pull it to focus. This assumes
-		// that the Toolkit Host's Tab Manager has already registered a tab with a NullWidget for content.
-		TSharedPtr<SDockTab> CurveEditorTab = Sequencer->GetToolkitHost()->GetTabManager()->TryInvokeTab(TabId);
-		if (CurveEditorTab.IsValid())
-		{
-			CurveEditorTab->SetContent(CurveEditorPanel.ToSharedRef());
-
-			const FSlateIcon SequencerGraphIcon = FSlateIcon(FAppStyle::GetAppStyleSetName(), "GenericCurveEditor.TabIcon");
-			CurveEditorTab->SetTabIcon(SequencerGraphIcon.GetIcon());
-
-			CurveEditorTab->SetLabel(LOCTEXT("SequencerMainGraphEditorTitle", "Sequencer Curves"));
-
-			SequencerPtr.Pin()->GetCurveEditor()->ZoomToFit();
-		}
+		CurveEditorExtension->OpenCurveEditor();
 	}
 	else
 	{
-		TSharedPtr<SDockTab> ExistingTab = Sequencer->GetToolkitHost()->GetTabManager()->FindExistingLiveTab(TabId);
-		if (ExistingTab)
-		{
-			ExistingTab->RequestCloseTab();
-		}
+		CurveEditorExtension->CloseCurveEditor();
 	}
 
 	TreeView->UpdateTrackArea();
diff --git a/Engine/Source/Editor/Sequencer/Private/SSequencer.h b/Engine/Source/Editor/Sequencer/Private/SSequencer.h
--- a/Engine/Source/Editor/Sequencer/Private/SSequencer.h
+++ b/Engine/Source/Editor/Sequencer/Private/SSequencer.h
@@ -401,9 +401,6 @@
 	/** Handles changes to the selected outliner nodes. */
 	void HandleOutlinerNodeSelectionChanged();
 
-	/** Syncs the current node selection to the curve editor. */
-	void SyncCurveEditorToSelection();
-
 	/** Empty active timer to ensure Slate ticks during Sequencer playback */
 	EActiveTimerReturnType EnsureSlateTickDuringPlayback(double InCurrentTime, float InDeltaTime);	
 
@@ -663,9 +660,6 @@
 	/** Section area widget for pinned tracks*/
 	TSharedPtr<UE::Sequencer::STrackAreaView> PinnedTrackArea;
 
-	/** Curve editor tree widget */
-	TSharedPtr<SCurveEditorTree> CurveEditorTree;
-
 	/** Curve editor filter that shows only the selected nodes */
 	TSharedPtr<FSequencerSelectionCurveFilter> SequencerSelectionCurveEditorFilter;
 
@@ -675,9 +669,6 @@
 	/** The search box for filtering tracks. */
 	TSharedPtr<SSearchBox> SearchBox;
 
-	/** The search widget for filtering curves in the Curve Editor tree. */
-	TSharedPtr<SWidget> CurveEditorSearchBox;
-
 	/** The current playback time display. */
 	TSharedPtr<STemporarilyFocusedSpinBox<double>> PlayTimeDisplay;
 
@@ -699,9 +690,6 @@
 	/** The top time slider widget */
 	TSharedPtr<ITimeSlider> TopTimeSlider;
 
-	/** The curve editor panel. This is created and updated even if it is not currently visible. */
-	TSharedPtr<SWidget> CurveEditorPanel;
-
 	/** Container for the toolbar, so that we can re-create it as needed. */
 	TSharedPtr<SBox> ToolbarContainer;
 
@@ -797,8 +785,4 @@
 	TWeakPtr<SWindow> WeakNodeGroupWindow;
 
 	TSharedPtr<SSequencerGroupManager> NodeGroupManager;
-
-public:
-	static const FName CurveEditorTabName;
-
 };
diff --git a/Engine/Source/Editor/Sequencer/Private/Sequencer.cpp b/Engine/Source/Editor/Sequencer/Private/Sequencer.cpp
--- a/Engine/Source/Editor/Sequencer/Private/Sequencer.cpp
+++ b/Engine/Source/Editor/Sequencer/Private/Sequencer.cpp
@@ -21,6 +21,7 @@
 #include "MVVM/Views/SOutlinerView.h"
 #include "MVVM/Views/IOutlinerSelectionHandler.h"
 #include "MVVM/Extensions/IDeletableExtension.h"
+#include "MVVM/CurveEditorExtension.h"
 #include "MVVM/CurveEditorIntegrationExtension.h"
 #include "MVVM/ObjectBindingModelStorageExtension.h"
 #include "MVVM/SectionModelStorageExtension.h"
@@ -184,66 +185,6 @@ static TAutoConsoleVariable<float> CVarAutoScrubCurveExponent(
 	2.0f,
 	TEXT("How much to ramp in and out the scrub speed when auto-scrubbing"));
 
-
-struct FSequencerCurveEditorBounds : ICurveEditorBounds
-{
-	FSequencerCurveEditorBounds(TSharedRef<FSequencer> InSequencer)
-		: WeakSequencer(InSequencer)
-	{
-		TRange<double> Bounds = InSequencer->GetViewRange();
-		InputMin = Bounds.GetLowerBoundValue();
-		InputMax = Bounds.GetUpperBoundValue();
-	}
-
-	virtual void GetInputBounds(double& OutMin, double& OutMax) const override
-	{
-		TSharedPtr<FSequencer> Sequencer = WeakSequencer.Pin();
-		if (Sequencer.IsValid())
-		{
-			const bool bLinkTimeRange = Sequencer->GetSequencerSettings()->GetLinkCurveEditorTimeRange();
-			if (bLinkTimeRange)
-			{
-				TRange<double> Bounds = Sequencer->GetViewRange();
-				OutMin = Bounds.GetLowerBoundValue();
-				OutMax = Bounds.GetUpperBoundValue();
-			}
-			else
-			{
-				// If they don't want to link the time range with Sequencer we return the cached value.
-				OutMin = InputMin;
-				OutMax = InputMax;
-			}
-		}
-	}
-
-	virtual void SetInputBounds(double InMin, double InMax) override
-	{
-		TSharedPtr<FSequencer> Sequencer = WeakSequencer.Pin();
-		if (Sequencer.IsValid())
-		{
-			const bool bLinkTimeRange = Sequencer->GetSequencerSettings()->GetLinkCurveEditorTimeRange();
-			if (bLinkTimeRange)
-			{
-				FFrameRate TickResolution = Sequencer->GetFocusedTickResolution();
-
-				if (InMin * TickResolution > TNumericLimits<int32>::Lowest() && InMax * TickResolution < TNumericLimits<int32>::Max())
-				{
-					Sequencer->SetViewRange(TRange<double>(InMin, InMax), EViewRangeInterpolation::Immediate);
-				}
-			}
-			
-			// We update these even if you are linked to the Sequencer Timeline so that when you turn off the link setting
-			// you don't pop to your last values, instead your view stays as is and just stops moving when Sequencer moves.
-			InputMin = InMin;
-			InputMax = InMax;
-		}
-	}
-
-	/** The min/max values for the viewing range. Only used if Curve Editor/Sequencer aren't linked ranges. */
-	double InputMin, InputMax;
-	TWeakPtr<FSequencer> WeakSequencer;
-};
-
 class FSequencerOutlinerSelectionHandler : 
 	public UE::Sequencer::IOutlinerSelectionHandler
 {
@@ -325,41 +266,6 @@ private:
 	TWeakPtr<UE::Sequencer::SOutlinerView> WeakPinnedTreeView;
 };
 
-class FSequencerCurveEditor : public FCurveEditor
-{
-public:
-	TWeakPtr<FSequencer> WeakSequencer;
-
-	FSequencerCurveEditor(TWeakPtr<FSequencer> InSequencer)
-		: WeakSequencer(InSequencer)
-	{}
-
-	virtual void GetGridLinesX(TArray<float>& MajorGridLines, TArray<float>& MinorGridLines, TArray<FText>* MajorGridLabels) const override
-	{
-		TSharedPtr<FSequencer> Sequencer = WeakSequencer.Pin();
-		FCurveEditorScreenSpaceH PanelInputSpace = GetPanelInputSpace();
-
-		double MajorGridStep  = 0.0;
-		int32  MinorDivisions = 0;
-
-		if (Sequencer.IsValid() && Sequencer->GetGridMetrics(PanelInputSpace.GetPhysicalWidth(), PanelInputSpace.GetInputMin(), PanelInputSpace.GetInputMax(), MajorGridStep, MinorDivisions))
-		{
-			const double FirstMajorLine = FMath::FloorToDouble(PanelInputSpace.GetInputMin() / MajorGridStep) * MajorGridStep;
-			const double LastMajorLine  = FMath::CeilToDouble(PanelInputSpace.GetInputMax() / MajorGridStep) * MajorGridStep;
-
-			for (double CurrentMajorLine = FirstMajorLine; CurrentMajorLine < LastMajorLine; CurrentMajorLine += MajorGridStep)
-			{
-				MajorGridLines.Add( PanelInputSpace.SecondsToScreen(CurrentMajorLine) );
- 
-				for (int32 Step = 1; Step < MinorDivisions; ++Step)
-				{
-					MinorGridLines.Add( PanelInputSpace.SecondsToScreen(CurrentMajorLine + Step*MajorGridStep/MinorDivisions) );
-				}
-			}
-		}
-	}
-};
-
 void FSequencer::InitSequencer(const FSequencerInitParams& InitParams, const TSharedRef<ISequencerObjectChangeListener>& InObjectChangeListener, const TArray<FOnCreateTrackEditor>& TrackEditorDelegates, const TArray<FOnCreateEditorObjectBinding>& EditorObjectBindingDelegates)
 {
 	using namespace UE::Sequencer;
@@ -410,30 +316,6 @@ void FSequencer::InitSequencer(const FSequencerInitParams& InitParams, const TSh
 
 	ObjectBindingTagCache = MakeUnique<FObjectBindingTagCache>();
 
-	FCurveEditorInitParams CurveEditorInitParams;
-	{
-	}
-	
-	{
-		CurveEditorModel = MakeShared<FSequencerCurveEditor>(SharedThis(this));
-		CurveEditorModel->SetBounds(MakeUnique<FSequencerCurveEditorBounds>(SharedThis(this)));
-		CurveEditorModel->InitCurveEditor(CurveEditorInitParams);
-
-		CurveEditorModel->InputSnapEnabledAttribute   = MakeAttributeLambda([this]{ return Settings->GetIsSnapEnabled(); });
-		CurveEditorModel->OnInputSnapEnabledChanged   = FOnSetBoolean::CreateLambda([this](bool NewValue){ Settings->SetIsSnapEnabled(NewValue); });
-
-		CurveEditorModel->OutputSnapEnabledAttribute  = MakeAttributeLambda([this]{ return Settings->GetSnapCurveValueToInterval(); });
-		CurveEditorModel->OnOutputSnapEnabledChanged  = FOnSetBoolean::CreateLambda([this](bool NewValue){ Settings->SetSnapCurveValueToInterval(NewValue); });
-
-		CurveEditorModel->FixedGridSpacingAttribute   = MakeAttributeLambda([this]() -> TOptional<float> { return Settings->GetGridSpacing(); });
-		CurveEditorModel->InputSnapRateAttribute      = MakeAttributeSP(this, &FSequencer::GetFocusedDisplayRate);
-
-		CurveEditorModel->DefaultKeyAttributes        = MakeAttributeSP(this, &FSequencer::GetDefaultKeyAttributes);
-		
-		CurveEditorModel->OnCurveArrayChanged.AddRaw(this, &FSequencer::OnCurveModelDisplayChanged);
-
-	}
-
 	{
 		FDelegateHandle OnBlueprintPreCompileHandle = GEditor->OnBlueprintPreCompile().AddLambda([&](UBlueprint* InBlueprint)
 		{
@@ -609,6 +491,14 @@ void FSequencer::InitSequencer(const FSequencerInitParams& InitParams, const TSh
 
 	SelectionHandler->SetTreeViews(SequencerWidget->GetTreeView(), SequencerWidget->GetPinnedTreeView());
 
+	if (GetHostCapabilities().bSupportsCurveEditor)
+	{
+		FCurveEditorExtension* CurveEditorExtension = ViewModel->CastDynamicChecked<FCurveEditorExtension>();
+		TSharedPtr<FCurveEditor> CurveEditor = CurveEditorExtension->GetCurveEditor();
+
+		CurveEditor->OnCurveArrayChanged.AddRaw(this, &FSequencer::OnCurveModelDisplayChanged);
+	}
+
 	// When undo occurs, get a notification so we can make sure our view is up to date
 	GEditor->RegisterForUndo(this);
 
@@ -2512,10 +2402,10 @@ void FSequencer::RecreateCurveEditor()
 {
 	using namespace UE::Sequencer;
 
-	FCurveEditorIntegrationExtension* CurveEditorIntegration = ViewModel->GetRootModel()->CastDynamic<FCurveEditorIntegrationExtension>();
+	FCurveEditorIntegrationExtension* CurveEditorIntegration= ViewModel->GetRootModel()->CastDynamic<FCurveEditorIntegrationExtension>();
 	if (ensure(CurveEditorIntegration))
 	{
-		CurveEditorIntegration->RecreateCurveEditor();
+		CurveEditorIntegration->ResetCurveEditor();
 	}
 }
 
@@ -6054,6 +5944,8 @@ void FSequencer::SetShowCurveEditor(bool bInShowCurveEditor)
 
 bool FSequencer::GetCurveEditorIsVisible() const
 {
+	using namespace UE::Sequencer;
+
 	// Some Sequencer usages don't support the Curve Editor
 	if (!GetHostCapabilities().bSupportsCurveEditor)
 	{
@@ -6062,7 +5954,13 @@ bool FSequencer::GetCurveEditorIsVisible() const
 
 	// We always want to retrieve this directly from the UI instead of mirroring it to a local bool as there are
 	// a lot of ways the UI could get out of sync with a local bool (such as previously restored tab layouts)
-	return GetToolkitHost()->GetTabManager()->FindExistingLiveTab(FTabId(SSequencer::CurveEditorTabName)).IsValid();
+	FCurveEditorExtension* CurveEditorExtension = ViewModel->CastDynamic<FCurveEditorExtension>();
+	if (ensure(CurveEditorExtension))
+	{
+		return CurveEditorExtension->IsCurveEditorOpen();
+	}
+
+	return false;
 }
 
 void FSequencer::SaveCurrentMovieScene()
@@ -13354,37 +13252,43 @@ void FSequencer::BindCommands()
 		Commands.RefreshUI,
 		FExecuteAction::CreateSP( this, &FSequencer::RefreshUI));
 
-	// We want a subset of the commands to work in the Curve Editor too, but bound to our functions. This minimizes code duplication
-	// while also freeing us up from issues that result from Sequencer already using two lists (for which our commands might be spread
-	// across both lists which makes a direct copy like it already uses difficult).
-	CurveEditorSharedBindings->MapAction(Commands.TogglePlay,			*SequencerCommandBindings->GetActionForCommand(Commands.TogglePlay));
-	CurveEditorSharedBindings->MapAction(Commands.TogglePlayViewport,	*SequencerCommandBindings->GetActionForCommand(Commands.TogglePlayViewport));
-	CurveEditorSharedBindings->MapAction(Commands.PlayForward,			*SequencerCommandBindings->GetActionForCommand(Commands.PlayForward));
-	CurveEditorSharedBindings->MapAction(Commands.JumpToStart,			*SequencerCommandBindings->GetActionForCommand(Commands.JumpToStart));
-	CurveEditorSharedBindings->MapAction(Commands.JumpToEnd,			*SequencerCommandBindings->GetActionForCommand(Commands.JumpToEnd));
-	CurveEditorSharedBindings->MapAction(Commands.ShuttleBackward,		*SequencerCommandBindings->GetActionForCommand(Commands.ShuttleBackward));
-	CurveEditorSharedBindings->MapAction(Commands.ShuttleForward,		*SequencerCommandBindings->GetActionForCommand(Commands.ShuttleForward));
-	CurveEditorSharedBindings->MapAction(Commands.Pause,				*SequencerCommandBindings->GetActionForCommand(Commands.Pause));
-	CurveEditorSharedBindings->MapAction(Commands.StepForward,			*SequencerCommandBindings->GetActionForCommand(Commands.StepForward));
-	CurveEditorSharedBindings->MapAction(Commands.StepBackward,			*SequencerCommandBindings->GetActionForCommand(Commands.StepBackward));
-	CurveEditorSharedBindings->MapAction(Commands.StepForwardViewport,  *SequencerCommandBindings->GetActionForCommand(Commands.StepForwardViewport));
-	CurveEditorSharedBindings->MapAction(Commands.StepBackwardViewport, *SequencerCommandBindings->GetActionForCommand(Commands.StepBackwardViewport));
-	CurveEditorSharedBindings->MapAction(Commands.JumpForward,          *SequencerCommandBindings->GetActionForCommand(Commands.JumpForward));
-	CurveEditorSharedBindings->MapAction(Commands.JumpBackward,         *SequencerCommandBindings->GetActionForCommand(Commands.JumpBackward));
-	CurveEditorSharedBindings->MapAction(Commands.StepToNextKey,		*SequencerCommandBindings->GetActionForCommand(Commands.StepToNextKey));
-	CurveEditorSharedBindings->MapAction(Commands.StepToPreviousKey, 	*SequencerCommandBindings->GetActionForCommand(Commands.StepToPreviousKey));
-	CurveEditorSharedBindings->MapAction(Commands.StepToNextMark,       *SequencerCommandBindings->GetActionForCommand(Commands.StepToNextMark));
-	CurveEditorSharedBindings->MapAction(Commands.StepToPreviousMark,   *SequencerCommandBindings->GetActionForCommand(Commands.StepToPreviousMark));
-	CurveEditorSharedBindings->MapAction(Commands.ToggleMarkAtPlayPosition, *SequencerCommandBindings->GetActionForCommand(Commands.ToggleMarkAtPlayPosition));
-
-	CurveEditorSharedBindings->MapAction(Commands.AddTransformKey, *SequencerCommandBindings->GetActionForCommand(Commands.AddTransformKey));
-	CurveEditorSharedBindings->MapAction(Commands.AddTranslationKey, *SequencerCommandBindings->GetActionForCommand(Commands.AddTranslationKey));
-	CurveEditorSharedBindings->MapAction(Commands.AddRotationKey, *SequencerCommandBindings->GetActionForCommand(Commands.AddRotationKey));
-	CurveEditorSharedBindings->MapAction(Commands.AddScaleKey, *SequencerCommandBindings->GetActionForCommand(Commands.AddScaleKey));
-
-	CurveEditorSharedBindings->MapAction(Commands.ResetFilters, *SequencerCommandBindings->GetActionForCommand(Commands.ResetFilters));
-
-	GetCurveEditor()->GetCommands()->Append(CurveEditorSharedBindings);
+	FCurveEditorExtension* CurveEditorExtension = ViewModel->CastDynamicChecked<FCurveEditorExtension>();
+	if (ensure(CurveEditorExtension))
+	{
+		// We want a subset of the commands to work in the Curve Editor too, but bound to our functions. This minimizes code duplication
+		// while also freeing us up from issues that result from Sequencer already using two lists (for which our commands might be spread
+		// across both lists which makes a direct copy like it already uses difficult).
+		CurveEditorSharedBindings->MapAction(Commands.TogglePlay, *SequencerCommandBindings->GetActionForCommand(Commands.TogglePlay));
+		CurveEditorSharedBindings->MapAction(Commands.TogglePlayViewport, *SequencerCommandBindings->GetActionForCommand(Commands.TogglePlayViewport));
+		CurveEditorSharedBindings->MapAction(Commands.PlayForward, *SequencerCommandBindings->GetActionForCommand(Commands.PlayForward));
+		CurveEditorSharedBindings->MapAction(Commands.JumpToStart, *SequencerCommandBindings->GetActionForCommand(Commands.JumpToStart));
+		CurveEditorSharedBindings->MapAction(Commands.JumpToEnd, *SequencerCommandBindings->GetActionForCommand(Commands.JumpToEnd));
+		CurveEditorSharedBindings->MapAction(Commands.ShuttleBackward, *SequencerCommandBindings->GetActionForCommand(Commands.ShuttleBackward));
+		CurveEditorSharedBindings->MapAction(Commands.ShuttleForward, *SequencerCommandBindings->GetActionForCommand(Commands.ShuttleForward));
+		CurveEditorSharedBindings->MapAction(Commands.Pause, *SequencerCommandBindings->GetActionForCommand(Commands.Pause));
+		CurveEditorSharedBindings->MapAction(Commands.StepForward, *SequencerCommandBindings->GetActionForCommand(Commands.StepForward));
+		CurveEditorSharedBindings->MapAction(Commands.StepBackward, *SequencerCommandBindings->GetActionForCommand(Commands.StepBackward));
+		CurveEditorSharedBindings->MapAction(Commands.StepForwardViewport, *SequencerCommandBindings->GetActionForCommand(Commands.StepForwardViewport));
+		CurveEditorSharedBindings->MapAction(Commands.StepBackwardViewport, *SequencerCommandBindings->GetActionForCommand(Commands.StepBackwardViewport));
+		CurveEditorSharedBindings->MapAction(Commands.JumpForward, *SequencerCommandBindings->GetActionForCommand(Commands.JumpForward));
+		CurveEditorSharedBindings->MapAction(Commands.JumpBackward, *SequencerCommandBindings->GetActionForCommand(Commands.JumpBackward));
+		CurveEditorSharedBindings->MapAction(Commands.StepToNextKey, *SequencerCommandBindings->GetActionForCommand(Commands.StepToNextKey));
+		CurveEditorSharedBindings->MapAction(Commands.StepToPreviousKey, *SequencerCommandBindings->GetActionForCommand(Commands.StepToPreviousKey));
+		CurveEditorSharedBindings->MapAction(Commands.StepToNextMark, *SequencerCommandBindings->GetActionForCommand(Commands.StepToNextMark));
+		CurveEditorSharedBindings->MapAction(Commands.StepToPreviousMark, *SequencerCommandBindings->GetActionForCommand(Commands.StepToPreviousMark));
+		CurveEditorSharedBindings->MapAction(Commands.ToggleMarkAtPlayPosition, *SequencerCommandBindings->GetActionForCommand(Commands.ToggleMarkAtPlayPosition));
+
+		CurveEditorSharedBindings->MapAction(Commands.AddTransformKey, *SequencerCommandBindings->GetActionForCommand(Commands.AddTransformKey));
+		CurveEditorSharedBindings->MapAction(Commands.AddTranslationKey, *SequencerCommandBindings->GetActionForCommand(Commands.AddTranslationKey));
+		CurveEditorSharedBindings->MapAction(Commands.AddRotationKey, *SequencerCommandBindings->GetActionForCommand(Commands.AddRotationKey));
+		CurveEditorSharedBindings->MapAction(Commands.AddScaleKey, *SequencerCommandBindings->GetActionForCommand(Commands.AddScaleKey));
+
+		CurveEditorSharedBindings->MapAction(Commands.ResetFilters, *SequencerCommandBindings->GetActionForCommand(Commands.ResetFilters));
+
+		TSharedPtr<FCurveEditor> CurveEditor = CurveEditorExtension->GetCurveEditor();
+		check(CurveEditor);
+		CurveEditor->GetCommands()->Append(CurveEditorSharedBindings);
+	}
 
 	// bind widget specific commands
 	SequencerWidget->BindCommands(SequencerCommandBindings, CurveEditorSharedBindings);
@@ -13625,18 +13529,6 @@ void FSequencer::BuildCustomContextMenuForGuid(FMenuBuilder& MenuBuilder, FGuid
 	OutlinerViewModel->BuildCustomContextMenuForGuid(MenuBuilder, ObjectBinding);
 }
 
-FKeyAttributes FSequencer::GetDefaultKeyAttributes() const
-{
-	switch (Settings->GetKeyInterpolation())
-	{
-	case EMovieSceneKeyInterpolation::User:     return FKeyAttributes().SetInterpMode(RCIM_Cubic).SetTangentMode(RCTM_User);
-	case EMovieSceneKeyInterpolation::Break:    return FKeyAttributes().SetInterpMode(RCIM_Cubic).SetTangentMode(RCTM_Break);
-	case EMovieSceneKeyInterpolation::Linear:   return FKeyAttributes().SetInterpMode(RCIM_Linear).SetTangentMode(RCTM_Auto);
-	case EMovieSceneKeyInterpolation::Constant: return FKeyAttributes().SetInterpMode(RCIM_Constant).SetTangentMode(RCTM_Auto);
-	default:                                    return FKeyAttributes().SetInterpMode(RCIM_Cubic).SetTangentMode(RCTM_Auto);
-	}
-}
-
 bool FSequencer::GetGridMetrics(const float PhysicalWidth, const double InViewStart, const double InViewEnd, double& OutMajorInterval, int32& OutMinorDivisions) const
 {
 	FSlateFontInfo SmallLayoutFont = FCoreStyle::GetDefaultFontStyle("Regular", 8);
diff --git a/Engine/Source/Editor/Sequencer/Private/Sequencer.h b/Engine/Source/Editor/Sequencer/Private/Sequencer.h
--- a/Engine/Source/Editor/Sequencer/Private/Sequencer.h
+++ b/Engine/Source/Editor/Sequencer/Private/Sequencer.h
@@ -617,11 +617,6 @@ public:
 
 	void OnSelectedNodesOnlyChanged();
 
-	virtual TSharedPtr<FCurveEditor> GetCurveEditor() const override
-	{
-		return CurveEditorModel;
-	}
-
 	/** Will create a custom menu if FSequencerViewParams::OnBuildCustomContextMenuForGuid is specified. */
 	void BuildCustomContextMenuForGuid(FMenuBuilder& MenuBuilder, FGuid ObjectBinding);
 
@@ -1141,9 +1136,6 @@ private:
 	/** Check whether we're viewing the master sequence or not */
 	bool IsViewingMasterSequence() const { return ActiveTemplateIDs.Num() == 1; }
 
-	/** Get the default key attributes to apply to newly created keys on the curve editor */
-	FKeyAttributes GetDefaultKeyAttributes() const;
-
 	/** Recompile any dirty director blueprints in the sequence hierarchy */
 	void RecompileDirtyDirectors();
 
@@ -1450,8 +1442,6 @@ private:
 	/** Transient collection of keys that is used for jumping between keys contained within the current selection */
 	TUniquePtr<FSequencerKeyCollection> SelectedKeyCollection;
 
-	TSharedPtr<FCurveEditor> CurveEditorModel;
-
 	TSharedPtr<UE::MovieScene::FInitialValueCache> InitialValueCache;
 
 	/** A signature that will suppress auto evaluation when it is the only change dirtying the template. */
diff --git a/Engine/Source/Editor/Sequencer/Private/Widgets/CurveEditor/SSequencerCurveEditor.cpp b/Engine/Source/Editor/Sequencer/Private/Widgets/CurveEditor/SSequencerCurveEditor.cpp
new file mode 100644
--- /dev/null
+++ b/Engine/Source/Editor/Sequencer/Private/Widgets/CurveEditor/SSequencerCurveEditor.cpp
@@ -0,0 +1,72 @@
+// Copyright Epic Games, Inc. All Rights Reserved.
+
+#include "SSequencerCurveEditor.h"
+#include "Framework/MultiBox/MultiBoxBuilder.h"
+#include "MovieSceneCommonHelpers.h"
+#include "SCurveEditorPanel.h"
+#include "Sequencer.h"
+
+#define LOCTEXT_NAMESPACE "SSequencerCurveEditor"
+
+void SSequencerCurveEditor::Construct(const FArguments& InArgs, TSharedRef<SCurveEditorPanel> InEditorPanel, TSharedPtr<FSequencer> InSequencer)
+{
+	WeakSequencer = InSequencer;
+
+	ChildSlot
+	[
+		SNew(SVerticalBox)
+
+		+ SVerticalBox::Slot()
+		.AutoHeight()
+		[
+			MakeToolbar(InEditorPanel)
+		]
+		+ SVerticalBox::Slot()
+		.FillHeight(1.0f)
+		[
+			InEditorPanel
+		]
+	];
+}
+
+TSharedRef<SWidget> SSequencerCurveEditor::MakeToolbar(TSharedRef<SCurveEditorPanel> InEditorPanel)
+{
+	FSlimHorizontalToolBarBuilder ToolBarBuilder(InEditorPanel->GetCommands(), FMultiBoxCustomization::None, InEditorPanel->GetToolbarExtender(), true);
+	ToolBarBuilder.BeginSection("Asset");
+
+	{
+		TAttribute<FSlateIcon> SaveIcon;
+		SaveIcon.Bind(TAttribute<FSlateIcon>::FGetter::CreateLambda([&] {
+
+			TArray<UMovieScene*> MovieScenesToSave;
+			MovieSceneHelpers::GetDescendantMovieScenes(WeakSequencer.Pin()->GetRootMovieSceneSequence(), MovieScenesToSave);
+			for (UMovieScene* MovieSceneToSave : MovieScenesToSave)
+			{
+				UPackage* MovieScenePackageToSave = MovieSceneToSave->GetOuter()->GetOutermost();
+				if (MovieScenePackageToSave->IsDirty())
+				{
+					return FSlateIcon(FAppStyle::Get().GetStyleSetName(), "Icons.SaveChanged");
+				}
+			}
+
+			return FSlateIcon(FAppStyle::Get().GetStyleSetName(), "Icons.Save");
+			}));
+
+		if (WeakSequencer.Pin()->GetHostCapabilities().bSupportsSaveMovieSceneAsset)
+		{
+			ToolBarBuilder.AddToolBarButton(
+				FUIAction(FExecuteAction::CreateLambda([&] { WeakSequencer.Pin()->SaveCurrentMovieScene(); })),
+				NAME_None,
+				LOCTEXT("SaveDirtyPackages", "Save"),
+				LOCTEXT("SaveDirtyPackagesTooltip", "Saves the current sequence and any subsequences"),
+				SaveIcon
+			);
+		}		
+	}
+	
+	ToolBarBuilder.EndSection();
+	// We just use all of the extenders as our toolbar, we don't have a need to create a separate toolbar.
+	return ToolBarBuilder.MakeWidget();
+}
+
+#undef LOCTEXT_NAMESPACE
diff --git a/Engine/Source/Editor/Sequencer/Private/Widgets/CurveEditor/SSequencerCurveEditor.h b/Engine/Source/Editor/Sequencer/Private/Widgets/CurveEditor/SSequencerCurveEditor.h
new file mode 100644
--- /dev/null
+++ b/Engine/Source/Editor/Sequencer/Private/Widgets/CurveEditor/SSequencerCurveEditor.h
@@ -0,0 +1,23 @@
+// Copyright Epic Games, Inc. All Rights Reserved.
+
+#pragma once
+
+#include "Widgets/SCompoundWidget.h"
+
+class FSequencer;
+class SCurveEditorPanel;
+
+class SSequencerCurveEditor : public SCompoundWidget
+{
+	SLATE_BEGIN_ARGS(SSequencerCurveEditor)
+	{}
+	SLATE_END_ARGS()
+public:
+	void Construct(const FArguments& InArgs, TSharedRef<SCurveEditorPanel> InEditorPanel, TSharedPtr<FSequencer> InSequencer);
+
+	TSharedRef<SWidget> MakeToolbar(TSharedRef<SCurveEditorPanel> InEditorPanel);
+
+private:
+	TWeakPtr<FSequencer> WeakSequencer;
+};
+
diff --git a/Engine/Source/Editor/Sequencer/Private/Widgets/CurveEditor/SequencerCurveEditorTimeSliderController.cpp b/Engine/Source/Editor/Sequencer/Private/Widgets/CurveEditor/SequencerCurveEditorTimeSliderController.cpp
new file mode 100644
--- /dev/null
+++ b/Engine/Source/Editor/Sequencer/Private/Widgets/CurveEditor/SequencerCurveEditorTimeSliderController.cpp
@@ -0,0 +1,127 @@
+// Copyright Epic Games, Inc. All Rights Reserved.
+
+#include "Widgets/CurveEditor/SequencerCurveEditorTimeSliderController.h"
+#include "CurveEditor.h"
+#include "Sequencer.h"
+#include "SequencerSettings.h"
+
+FSequencerCurveEditorTimeSliderController::FSequencerCurveEditorTimeSliderController(const FTimeSliderArgs& InArgs, TWeakPtr<FSequencer> InWeakSequencer, TSharedRef<FCurveEditor> InCurveEditor)
+	: FSequencerTimeSliderController(InArgs, InWeakSequencer)
+{
+	WeakSequencer = InWeakSequencer;
+	WeakCurveEditor = InCurveEditor;
+}
+
+
+void FSequencerCurveEditorTimeSliderController::ClampViewRange(double& NewRangeMin, double& NewRangeMax)
+{
+	// Since the CurveEditor uses a different view range (potentially) we have to be careful about which one we clamp.
+	TSharedPtr<ISequencer> Sequencer = WeakSequencer.Pin();
+	if (!Sequencer.IsValid())
+	{
+		return;
+	}
+
+	const bool bLinkedTimeRange = Sequencer->GetSequencerSettings()->GetLinkCurveEditorTimeRange();
+	if (bLinkedTimeRange)
+	{
+		return FSequencerTimeSliderController::ClampViewRange(NewRangeMin, NewRangeMax);
+	}
+	else
+	{
+		TSharedPtr<FCurveEditor> CurveEditor = WeakCurveEditor.Pin();
+		if (CurveEditor.IsValid())
+		{
+			double InputMin, InputMax;
+			CurveEditor->GetBounds().GetInputBounds(InputMin, InputMax);
+
+			bool bNeedsClampSet = false;
+			double NewClampRangeMin = InputMin;
+			if (NewRangeMin < InputMin)
+			{
+				NewClampRangeMin = NewRangeMin;
+				bNeedsClampSet = true;
+			}
+
+			double NewClampRangeMax = InputMax;
+			if (NewRangeMax > InputMax)
+			{
+				NewClampRangeMax = NewRangeMax;
+				bNeedsClampSet = true;
+			}
+
+			if (bNeedsClampSet)
+			{
+				CurveEditor->GetBounds().SetInputBounds(NewClampRangeMin, NewClampRangeMax);
+			}
+
+		}
+	}
+}
+
+void FSequencerCurveEditorTimeSliderController::SetViewRange(double NewRangeMin, double NewRangeMax, EViewRangeInterpolation Interpolation)
+{
+	TSharedPtr<ISequencer> Sequencer = WeakSequencer.Pin();
+	if (!Sequencer.IsValid())
+	{
+		return;
+	}
+
+	const bool bLinkedTimeRange = Sequencer->GetSequencerSettings()->GetLinkCurveEditorTimeRange();
+	if (bLinkedTimeRange)
+	{
+		return FSequencerTimeSliderController::SetViewRange(NewRangeMin, NewRangeMax, Interpolation);
+	}
+	else
+	{
+		// Clamp to a minimum size to avoid zero-sized or negative visible ranges
+		double MinVisibleTimeRange = FFrameNumber(1) / GetTickResolution();
+		TRange<double> ExistingViewRange = GetViewRange();
+
+		if (NewRangeMax == ExistingViewRange.GetUpperBoundValue())
+		{
+			if (NewRangeMin > NewRangeMax - MinVisibleTimeRange)
+			{
+				NewRangeMin = NewRangeMax - MinVisibleTimeRange;
+			}
+		}
+		else if (NewRangeMax < NewRangeMin + MinVisibleTimeRange)
+		{
+			NewRangeMax = NewRangeMin + MinVisibleTimeRange;
+		}
+
+		TSharedPtr<FCurveEditor> CurveEditor = WeakCurveEditor.Pin();
+		if (CurveEditor.IsValid())
+		{
+			CurveEditor->GetBounds().SetInputBounds(NewRangeMin, NewRangeMax);
+		}
+	}
+}
+
+FAnimatedRange FSequencerCurveEditorTimeSliderController::GetViewRange() const
+{ 
+	// If they've linked the Sequencer timerange we can return the internal controller's view range, otherwise we return the bounds (which internally does the same check)
+	TSharedPtr<ISequencer> Sequencer = WeakSequencer.Pin();
+	if (!Sequencer.IsValid())
+	{
+		return FAnimatedRange();
+	}
+	const bool bLinkedTimeRange = Sequencer->GetSequencerSettings()->GetLinkCurveEditorTimeRange();
+	if (bLinkedTimeRange)
+	{
+		return FSequencerTimeSliderController::GetViewRange();
+	}
+	else
+	{
+		TSharedPtr<FCurveEditor> CurveEditor = WeakCurveEditor.Pin();
+		if (CurveEditor.IsValid())
+		{
+			double InputMin, InputMax;
+			CurveEditor->GetBounds().GetInputBounds(InputMin, InputMax);
+
+			return FAnimatedRange(InputMin, InputMax);
+		}
+	}
+
+	return FAnimatedRange();
+}
diff --git a/Engine/Source/Editor/Sequencer/Private/Widgets/CurveEditor/SequencerCurveEditorTimeSliderController.h b/Engine/Source/Editor/Sequencer/Private/Widgets/CurveEditor/SequencerCurveEditorTimeSliderController.h
new file mode 100644
--- /dev/null
+++ b/Engine/Source/Editor/Sequencer/Private/Widgets/CurveEditor/SequencerCurveEditorTimeSliderController.h
@@ -0,0 +1,25 @@
+// Copyright Epic Games, Inc. All Rights Reserved.
+
+#pragma once
+
+#include "CoreMinimal.h"
+#include "SequencerTimeSliderController.h"
+
+class ISequencer;
+class FSequencer;
+
+class FSequencerCurveEditorTimeSliderController : public FSequencerTimeSliderController
+{
+public:
+
+	FSequencerCurveEditorTimeSliderController(const FTimeSliderArgs& InArgs, TWeakPtr<FSequencer> InWeakSequencer, TSharedRef<FCurveEditor> InCurveEditor);
+
+	virtual void ClampViewRange(double& NewRangeMin, double& NewRangeMax) override;
+	virtual void SetViewRange(double NewRangeMin, double NewRangeMax, EViewRangeInterpolation Interpolation) override;
+	virtual FAnimatedRange GetViewRange() const override;
+
+private:
+	TWeakPtr<ISequencer> WeakSequencer;
+	TWeakPtr<FCurveEditor> WeakCurveEditor;
+};
+
diff --git a/Engine/Source/Editor/Sequencer/Public/ISequencer.h b/Engine/Source/Editor/Sequencer/Public/ISequencer.h
--- a/Engine/Source/Editor/Sequencer/Public/ISequencer.h
+++ b/Engine/Source/Editor/Sequencer/Public/ISequencer.h
@@ -210,9 +210,6 @@ public:
 	/** @return a multicast delegate which is executed when sequencer closes. */
 	virtual FOnCloseEvent& OnCloseEvent() = 0;
 
-	/** @return associated curve editor*/
-	virtual TSharedPtr<FCurveEditor> GetCurveEditor() const = 0;
-
 	/** @return Widget used to display the sequencer */
 	virtual TSharedRef<SWidget> GetSequencerWidget() const = 0;
 
diff --git a/Engine/Source/Editor/Sequencer/Public/MVVM/CurveEditorExtension.h b/Engine/Source/Editor/Sequencer/Public/MVVM/CurveEditorExtension.h
new file mode 100644
--- /dev/null
+++ b/Engine/Source/Editor/Sequencer/Public/MVVM/CurveEditorExtension.h
@@ -0,0 +1,84 @@
+// Copyright Epic Games, Inc. All Rights Reserved.
+
+#pragma once
+
+#include "MVVM/ViewModels/ViewModel.h"
+#include "MVVM/ViewModelPtr.h"
+#include "MVVM/Extensions/DynamicExtensionContainer.h"
+#include "MVVM/Extensions/ICurveEditorTreeItemExtension.h"
+#include "CurveDataAbstraction.h"
+#include "CurveEditorTypes.h"
+#include "Tree/ICurveEditorTreeItem.h"
+
+class ISequencer;
+class IPropertyTypeCustomization;
+class SCurveEditorTree;
+struct FTimeSliderArgs;
+
+namespace UE
+{
+namespace Sequencer
+{
+
+class FSequencerEditorViewModel;
+
+class SEQUENCER_API FCurveEditorExtension : public IDynamicExtension
+{
+public:
+
+	UE_SEQUENCER_DECLARE_VIEW_MODEL_TYPE_ID(FCurveEditorExtension)
+
+	FCurveEditorExtension();
+
+	virtual void OnCreated(TSharedRef<FViewModel> InWeakOwner) override;
+
+	/** Creates the curve editor view-model and widget */
+	void CreateCurveEditor(const FTimeSliderArgs& TimeSliderArgs);
+
+	/** Gets the curve editor view-model */
+	TSharedPtr<FCurveEditor> GetCurveEditor() const { return CurveEditorModel; }
+
+	/** Opens the curve editor */
+	void OpenCurveEditor();
+	/** Returns whether the curve editor is open */
+	bool IsCurveEditorOpen() const;
+	/** Closes the curve editor */
+	void CloseCurveEditor();
+
+	/**
+	 * Synchronize curve editor selection with sequencer outliner selection on the next update.
+	 */
+	void RequestSyncSelection();
+
+public:
+
+	static const FName CurveEditorTabName;
+
+private:
+
+	/** Get the default key attributes to apply to newly created keys on the curve editor */
+	FKeyAttributes GetDefaultKeyAttributes() const;
+
+	/** Syncs the selection of the curve editor with that of the sequencer */
+	void SyncSelection();
+
+private:
+
+	/** The sequencer editor we are extending with a curve editor */
+	TWeakPtr<FSequencerEditorViewModel> WeakOwnerModel;
+
+	/** Curve editor */
+	TSharedPtr<FCurveEditor> CurveEditorModel;
+	/** Curve editor tree widget */
+	TSharedPtr<SCurveEditorTree> CurveEditorTreeView;
+	/** The search widget for filtering curves in the Curve Editor tree. */
+	TSharedPtr<SWidget> CurveEditorSearchBox;
+	/** The curve editor panel. This is created and updated even if it is not currently visible. */
+	TSharedPtr<SWidget> CurveEditorPanel;
+
+	friend class FCurveEditorIntegrationExtension;
+};
+
+} // namespace Sequencer
+} // namespace UE
+
diff --git a/Engine/Source/Editor/Sequencer/Public/MVVM/CurveEditorIntegrationExtension.h b/Engine/Source/Editor/Sequencer/Public/MVVM/CurveEditorIntegrationExtension.h
--- a/Engine/Source/Editor/Sequencer/Public/MVVM/CurveEditorIntegrationExtension.h
+++ b/Engine/Source/Editor/Sequencer/Public/MVVM/CurveEditorIntegrationExtension.h
@@ -15,6 +15,7 @@
 {
 
 class FSequenceModel;
+class FCurveEditorExtension;
 
 /**
  * Extension for managing integration between outliner items and the curve editor.
@@ -27,7 +28,7 @@
  * - Outliner items implementing ICurveEditorTreeItemExtension (or its default shim) if they 
  *   want to show up in the curve editor.
  */
-class FCurveEditorIntegrationExtension : public IDynamicExtension
+class SEQUENCER_API FCurveEditorIntegrationExtension : public IDynamicExtension
 {
 public:
 
@@ -37,14 +38,27 @@
 
 	virtual void OnCreated(TSharedRef<FViewModel> InWeakOwner) override;
 
+	/** 
+	 * Keeps the curve editor items up-to-date with the sequencer outliner by adding/removing 
+	 * entries as needed.
+	 */
 	void UpdateCurveEditor();
-	void RecreateCurveEditor();
+
+	/**
+	 * Clears the curve editor of all contents.
+	 */
+	void ResetCurveEditor();
 
 private:
 
+	/** Update curve editor items when sequencer outliner items change */
 	void OnHierarchyChanged();
 
-	FCurveEditorTreeItemID AddToCurveEditor(TViewModelPtr<ICurveEditorTreeItemExtension> InViewModel, TSharedPtr<FCurveEditor> CurveEditor);
+	/** Adds the given view-model to the curve editor */
+	FCurveEditorTreeItemID AddToCurveEditor(TViewModelPtr<ICurveEditorTreeItemExtension> InViewModel, TSharedPtr<FCurveEditor> InCurveEditor);
+
+	/** Finds the curve editor extension on the top-level sequencer editor view-model */
+	FCurveEditorExtension* GetCurveEditorExtension();
 
 private:
 
diff --git a/Engine/Source/Editor/Sequencer/Public/MVVM/Extensions/ICurveEditorTreeItemExtension.h b/Engine/Source/Editor/Sequencer/Public/MVVM/Extensions/ICurveEditorTreeItemExtension.h
--- a/Engine/Source/Editor/Sequencer/Public/MVVM/Extensions/ICurveEditorTreeItemExtension.h
+++ b/Engine/Source/Editor/Sequencer/Public/MVVM/Extensions/ICurveEditorTreeItemExtension.h
@@ -15,19 +15,6 @@
 {
 
 /**
- * Extension interface for getting access to the curve editor.
- */
-class SEQUENCER_API ICurveEditorExtension
-{
-public:
-
-	UE_SEQUENCER_DECLARE_VIEW_MODEL_TYPE_ID(ICurveEditorExtension)
-
-	/** Gets the curve editor */
-	virtual TSharedPtr<FCurveEditor> GetCurveEditor() const = 0;
-};
-
-/**
  * Extension interface for view-models that want to show up in the curve editor.
  */
 class SEQUENCER_API ICurveEditorTreeItemExtension
diff --git a/Engine/Source/Editor/Sequencer/Public/MVVM/ViewModels/SequenceModel.h b/Engine/Source/Editor/Sequencer/Public/MVVM/ViewModels/SequenceModel.h
--- a/Engine/Source/Editor/Sequencer/Public/MVVM/ViewModels/SequenceModel.h
+++ b/Engine/Source/Editor/Sequencer/Public/MVVM/ViewModels/SequenceModel.h
@@ -38,13 +38,12 @@
 class SEQUENCER_API FSequenceModel
 	: public FViewModel
 	, public ISortableExtension
-	, public ICurveEditorExtension
 	, public IOutlinerDropTargetOutlinerExtension
 	, public UE::MovieScene::ISignedObjectEventHandler
 {
 public:
 
-	UE_SEQUENCER_DECLARE_CASTABLE(FSequenceModel, FViewModel, ISortableExtension, ICurveEditorExtension, IOutlinerDropTargetOutlinerExtension);
+	UE_SEQUENCER_DECLARE_CASTABLE(FSequenceModel, FViewModel, ISortableExtension, IOutlinerDropTargetOutlinerExtension);
 
 	static FOnInitializeSequenceModel CreateExtensionsEvent;
 
@@ -71,9 +70,6 @@
 	FSortingKey GetSortingKey() const override;
 	void SetCustomOrder(int32 InCustomOrder) override;
 
-	/*~ ICurveEditorExtension */
-	TSharedPtr<FCurveEditor> GetCurveEditor() const override;
-
 	/*~ IOutlinerDropTargetOutlinerExtension */
 	TOptional<EItemDropZone> CanAcceptDrop(const FViewModelPtr& TargetModel, const FDragDropEvent& DragDropEvent, EItemDropZone InItemDropZone) override;
 	void PerformDrop(const FViewModelPtr& TargetModel, const FDragDropEvent& DragDropEvent, EItemDropZone InItemDropZone) override;
diff --git a/Engine/Source/Editor/Sequencer/Public/MVVM/ViewModels/SequencerEditorViewModel.h b/Engine/Source/Editor/Sequencer/Public/MVVM/ViewModels/SequencerEditorViewModel.h
--- a/Engine/Source/Editor/Sequencer/Public/MVVM/ViewModels/SequencerEditorViewModel.h
+++ b/Engine/Source/Editor/Sequencer/Public/MVVM/ViewModels/SequencerEditorViewModel.h
@@ -41,6 +41,7 @@
 	virtual TSharedPtr<FViewModel> CreateRootModelImpl() override;
 	virtual TSharedPtr<FOutlinerViewModel> CreateOutlinerImpl() override;
 	virtual TSharedPtr<FTrackAreaViewModel> CreateTrackAreaImpl() override;
+	virtual void InitializeEditorImpl() override;
 
 protected:
 
