Fix crash when opening sequencers that don't support curve editors

diff --git a/Engine/Source/Editor/Sequencer/Private/MVVM/CurveEditorIntegrationExtension.cpp b/Engine/Source/Editor/Sequencer/Private/MVVM/CurveEditorIntegrationExtension.cpp
--- a/Engine/Source/Editor/Sequencer/Private/MVVM/CurveEditorIntegrationExtension.cpp
+++ b/Engine/Source/Editor/Sequencer/Private/MVVM/CurveEditorIntegrationExtension.cpp
@@ -56,13 +56,13 @@ void FCurveEditorIntegrationExtension::UpdateCurveEditor()
 	}
 	
 	FCurveEditorExtension* CurveEditorExtension = GetCurveEditorExtension();
-	if (!ensure(CurveEditorExtension))
+	if (!CurveEditorExtension)
 	{
 		return;
 	}
 
 	TSharedPtr<FCurveEditor> CurveEditor = CurveEditorExtension->GetCurveEditor();
-	if (!ensure(CurveEditor))
+	if (!CurveEditor)
 	{
 		return;
 	}
diff --git a/Engine/Source/Editor/Sequencer/Private/Sequencer.cpp b/Engine/Source/Editor/Sequencer/Private/Sequencer.cpp
--- a/Engine/Source/Editor/Sequencer/Private/Sequencer.cpp
+++ b/Engine/Source/Editor/Sequencer/Private/Sequencer.cpp
@@ -13253,7 +13253,7 @@ void FSequencer::BindCommands()
 		FExecuteAction::CreateSP( this, &FSequencer::RefreshUI));
 
 	FCurveEditorExtension* CurveEditorExtension = ViewModel->CastDynamicChecked<FCurveEditorExtension>();
-	if (ensure(CurveEditorExtension))
+	if (CurveEditorExtension && CurveEditorExtension->GetCurveEditor())
 	{
 		// We want a subset of the commands to work in the Curve Editor too, but bound to our functions. This minimizes code duplication
 		// while also freeing us up from issues that result from Sequencer already using two lists (for which our commands might be spread
@@ -13286,7 +13286,6 @@ void FSequencer::BindCommands()
 		CurveEditorSharedBindings->MapAction(Commands.ResetFilters, *SequencerCommandBindings->GetActionForCommand(Commands.ResetFilters));
 
 		TSharedPtr<FCurveEditor> CurveEditor = CurveEditorExtension->GetCurveEditor();
-		check(CurveEditor);
 		CurveEditor->GetCommands()->Append(CurveEditorSharedBindings);
 	}
 
