Only create curve-editor extensions when the host sequencer supports a curve editor.

This lets us not even create a curve editor view model anymore if we don't need to.

diff --git a/Engine/Source/Editor/Sequencer/Private/MVVM/ViewModels/SequenceModel.cpp b/Engine/Source/Editor/Sequencer/Private/MVVM/ViewModels/SequenceModel.cpp
--- a/Engine/Source/Editor/Sequencer/Private/MVVM/ViewModels/SequenceModel.cpp
+++ b/Engine/Source/Editor/Sequencer/Private/MVVM/ViewModels/SequenceModel.cpp
@@ -48,7 +48,9 @@ void FSequenceModel::InitializeExtensions()
 	GetChildrenForList(&RootOutlinerItems).AddChild(BottomSpacer);
 
 	// Create extensions
-	CreateExtensionsEvent.Broadcast(this);
+	TSharedPtr<FSequencerEditorViewModel> Editor = WeakEditor.Pin();
+	check(Editor);
+	CreateExtensionsEvent.Broadcast(Editor, SharedThis(this));
 
 	// Post-initialize extensions - this does the bulk of the population work for the sequence model
 	PostInitializeExtensions();
diff --git a/Engine/Source/Editor/Sequencer/Private/MVVM/ViewModels/SequencerEditorViewModel.cpp b/Engine/Source/Editor/Sequencer/Private/MVVM/ViewModels/SequencerEditorViewModel.cpp
--- a/Engine/Source/Editor/Sequencer/Private/MVVM/ViewModels/SequencerEditorViewModel.cpp
+++ b/Engine/Source/Editor/Sequencer/Private/MVVM/ViewModels/SequencerEditorViewModel.cpp
@@ -4,6 +4,7 @@
 #include "MVVM/ViewModels/SequencerOutlinerViewModel.h"
 #include "MVVM/ViewModels/SequencerTrackAreaViewModel.h"
 #include "MVVM/CurveEditorExtension.h"
+#include "ISequencerModule.h"
 #include "Sequencer.h"
 #include "MovieSceneSequenceID.h"
 
@@ -12,11 +13,20 @@ namespace UE
 namespace Sequencer
 {
 
-FSequencerEditorViewModel::FSequencerEditorViewModel(TSharedRef<ISequencer> InSequencer)
+FSequencerEditorViewModel::FSequencerEditorViewModel(TSharedRef<ISequencer> InSequencer, const FSequencerHostCapabilities& InHostCapabilities)
 	: WeakSequencer(InSequencer)
+	, bSupportsCurveEditor(InHostCapabilities.bSupportsCurveEditor)
 {
 }
 
+void FSequencerEditorViewModel::PreInitializeEditorImpl()
+{
+	if (bSupportsCurveEditor)
+	{
+		AddDynamicExtension(FCurveEditorExtension::ID);
+	}
+}
+
 TSharedPtr<FViewModel> FSequencerEditorViewModel::CreateRootModelImpl()
 {
 	TSharedPtr<FSequenceModel> RootSequenceModel = MakeShared<FSequenceModel>(SharedThis(this));
@@ -36,11 +46,6 @@ TSharedPtr<FTrackAreaViewModel> FSequencerEditorViewModel::CreateTrackAreaImpl()
 	return MakeShared<FSequencerTrackAreaViewModel>(Sequencer.ToSharedRef());
 }
 
-void FSequencerEditorViewModel::InitializeEditorImpl()
-{
-	AddDynamicExtension(FCurveEditorExtension::ID);
-}
-
 TSharedPtr<ISequencer> FSequencerEditorViewModel::GetSequencer() const
 {
 	return WeakSequencer.Pin();
diff --git a/Engine/Source/Editor/Sequencer/Private/Sequencer.cpp b/Engine/Source/Editor/Sequencer/Private/Sequencer.cpp
--- a/Engine/Source/Editor/Sequencer/Private/Sequencer.cpp
+++ b/Engine/Source/Editor/Sequencer/Private/Sequencer.cpp
@@ -426,7 +426,7 @@ void FSequencer::InitSequencer(const FSequencerInitParams& InitParams, const TSh
 		}
 	}
 
-	ViewModel = MakeShared<FSequencerEditorViewModel>(SharedThis(this));
+	ViewModel = MakeShared<FSequencerEditorViewModel>(SharedThis(this), GetHostCapabilities());
 	ViewModel->InitializeEditor();
 	ViewModel->SetSequence(InitParams.RootSequence);
 
@@ -493,7 +493,8 @@ void FSequencer::InitSequencer(const FSequencerInitParams& InitParams, const TSh
 
 	if (GetHostCapabilities().bSupportsCurveEditor)
 	{
-		FCurveEditorExtension* CurveEditorExtension = ViewModel->CastDynamicChecked<FCurveEditorExtension>();
+		FCurveEditorExtension* CurveEditorExtension = ViewModel->CastDynamic<FCurveEditorExtension>();
+		check(CurveEditorExtension);
 		TSharedPtr<FCurveEditor> CurveEditor = CurveEditorExtension->GetCurveEditor();
 
 		CurveEditor->OnCurveArrayChanged.AddRaw(this, &FSequencer::OnCurveModelDisplayChanged);
@@ -13252,8 +13253,9 @@ void FSequencer::BindCommands()
 		Commands.RefreshUI,
 		FExecuteAction::CreateSP( this, &FSequencer::RefreshUI));
 
-	FCurveEditorExtension* CurveEditorExtension = ViewModel->CastDynamicChecked<FCurveEditorExtension>();
-	if (CurveEditorExtension && CurveEditorExtension->GetCurveEditor())
+	// If this sequencer supports a curve editor, let's add bindings for it.
+	FCurveEditorExtension* CurveEditorExtension = ViewModel->CastDynamic<FCurveEditorExtension>();
+	if (CurveEditorExtension && ensure(CurveEditorExtension->GetCurveEditor()))
 	{
 		// We want a subset of the commands to work in the Curve Editor too, but bound to our functions. This minimizes code duplication
 		// while also freeing us up from issues that result from Sequencer already using two lists (for which our commands might be spread
diff --git a/Engine/Source/Editor/Sequencer/Private/SequencerModule.cpp b/Engine/Source/Editor/Sequencer/Private/SequencerModule.cpp
--- a/Engine/Source/Editor/Sequencer/Private/SequencerModule.cpp
+++ b/Engine/Source/Editor/Sequencer/Private/SequencerModule.cpp
@@ -21,6 +21,7 @@
 #include "AnimatedPropertyKey.h"
 #include "MovieSceneSignedObject.h"
 
+#include "MVVM/CurveEditorExtension.h"
 #include "MVVM/CurveEditorIntegrationExtension.h"
 #include "MVVM/FolderModelStorageExtension.h"
 #include "MVVM/ObjectBindingModelStorageExtension.h"
@@ -443,14 +444,20 @@ public:
 		}
 
 		FSequenceModel::CreateExtensionsEvent.AddLambda(
-			[&](FSequenceModel* InModel)
+			[&](TSharedPtr<FEditorViewModel> InEditor, TSharedPtr<FSequenceModel> InModel)
 			{
-				InModel->AddDynamicExtension(FCurveEditorIntegrationExtension::ID);
 				InModel->AddDynamicExtension(FFolderModelStorageExtension::ID);
 				InModel->AddDynamicExtension(FObjectBindingModelStorageExtension::ID);
 				InModel->AddDynamicExtension(FTrackModelStorageExtension::ID, TrackModelDelegates);
 				InModel->AddDynamicExtension(FTrackRowModelStorageExtension::ID);
 				InModel->AddDynamicExtension(FSectionModelStorageExtension::ID);
+
+				// If the editor supports a curve editor, add an integration extension to
+				// sync view-model hierarchies between the outliner and curve editor.
+				if (InEditor->CastDynamic<FCurveEditorExtension>())
+				{
+					InModel->AddDynamicExtension(FCurveEditorIntegrationExtension::ID);
+				}
 			}
 		);
 
diff --git a/Engine/Source/Editor/Sequencer/Public/MVVM/ViewModels/SequenceModel.h b/Engine/Source/Editor/Sequencer/Public/MVVM/ViewModels/SequenceModel.h
--- a/Engine/Source/Editor/Sequencer/Public/MVVM/ViewModels/SequenceModel.h
+++ b/Engine/Source/Editor/Sequencer/Public/MVVM/ViewModels/SequenceModel.h
@@ -30,10 +30,11 @@ namespace UE
 namespace Sequencer
 {
 
+class FEditorViewModel;
 class FSequenceModel;
 class FSequencerEditorViewModel;
 
-DECLARE_MULTICAST_DELEGATE_OneParam(FOnInitializeSequenceModel, FSequenceModel*);
+DECLARE_MULTICAST_DELEGATE_TwoParams(FOnInitializeSequenceModel, TSharedPtr<FEditorViewModel>, TSharedPtr<FSequenceModel>);
 
 class SEQUENCER_API FSequenceModel
 	: public FViewModel
diff --git a/Engine/Source/Editor/Sequencer/Public/MVVM/ViewModels/SequencerEditorViewModel.h b/Engine/Source/Editor/Sequencer/Public/MVVM/ViewModels/SequencerEditorViewModel.h
--- a/Engine/Source/Editor/Sequencer/Public/MVVM/ViewModels/SequencerEditorViewModel.h
+++ b/Engine/Source/Editor/Sequencer/Public/MVVM/ViewModels/SequencerEditorViewModel.h
@@ -9,6 +9,7 @@ struct FMovieSceneSequenceID;
 class ISequencer;
 class FSequencer;
 class UMovieSceneSequence;
+struct FSequencerHostCapabilities;
 
 namespace UE
 {
@@ -26,7 +27,7 @@ public:
 
 	UE_SEQUENCER_DECLARE_CASTABLE(FSequencerEditorViewModel, FEditorViewModel);
 
-	FSequencerEditorViewModel(TSharedRef<ISequencer> InSequencer);
+	FSequencerEditorViewModel(TSharedRef<ISequencer> InSequencer, const FSequencerHostCapabilities& InHostCapabilities);
 
 	// @todo_sequencer_mvvm remove this later
 	TSharedPtr<ISequencer> GetSequencer() const;
@@ -38,14 +39,15 @@ public:
 
 protected:
 
+	virtual void PreInitializeEditorImpl() override;
 	virtual TSharedPtr<FViewModel> CreateRootModelImpl() override;
 	virtual TSharedPtr<FOutlinerViewModel> CreateOutlinerImpl() override;
 	virtual TSharedPtr<FTrackAreaViewModel> CreateTrackAreaImpl() override;
-	virtual void InitializeEditorImpl() override;
 
 protected:
 
 	TWeakPtr<ISequencer> WeakSequencer;
+	bool bSupportsCurveEditor;
 };
 
 } // namespace Sequencer
diff --git a/Engine/Source/Editor/SequencerCore/Private/MVVM/ViewModels/EditorViewModel.cpp b/Engine/Source/Editor/SequencerCore/Private/MVVM/ViewModels/EditorViewModel.cpp
--- a/Engine/Source/Editor/SequencerCore/Private/MVVM/ViewModels/EditorViewModel.cpp
+++ b/Engine/Source/Editor/SequencerCore/Private/MVVM/ViewModels/EditorViewModel.cpp
@@ -29,6 +29,8 @@
 {
 	ensureMsgf(!Outliner && !TrackArea, TEXT("This editor view-model has already been initialized"));
 
+	PreInitializeEditorImpl();
+
 	FViewModelChildren PanelChildren = GetEditorPanels();
 
 	Outliner = CreateOutlinerImpl();
diff --git a/Engine/Source/Editor/SequencerCore/Public/MVVM/ViewModels/EditorViewModel.h b/Engine/Source/Editor/SequencerCore/Public/MVVM/ViewModels/EditorViewModel.h
--- a/Engine/Source/Editor/SequencerCore/Public/MVVM/ViewModels/EditorViewModel.h
+++ b/Engine/Source/Editor/SequencerCore/Public/MVVM/ViewModels/EditorViewModel.h
@@ -65,6 +65,8 @@
 
 protected:
 
+	/** Initializes this view model before panels and root models are created */
+	virtual void PreInitializeEditorImpl() {}
 	/** Creates the root data model for this editor. */
 	virtual TSharedPtr<FViewModel> CreateRootModelImpl();
 	/** Creates the outliner for this editor. */
