Sequencer: Add blueprint functions to LevelSequenceEditorSubsystem:

ConvertToPossessable/ConvertToSpawnable

Copy Tracks/Paste Tracks
Copy Sections/Paste Sections
Copy Bindings/Paste Bindings
Copy Folders/Paste Folders

AddActors
CreateCamera

Moved a lot of functions from FSequencer to FSequencerUtilities so that they can be shared with LevelSequenceEditorSubsystem.
These functions are now local to FSequencerUtilities: NewCameraAdded, AddSpawnable, ExpandMultiplePossessableBindings, UpdateFolderBindingID, UpdateBindingIDs
Removed a few calls to NotifyMovieSceneDataChanged as they seem to be unnecessary now.

diff --git a/Engine/Plugins/MovieScene/LevelSequenceEditor/Source/LevelSequenceEditor/LevelSequenceEditor.Build.cs b/Engine/Plugins/MovieScene/LevelSequenceEditor/Source/LevelSequenceEditor/LevelSequenceEditor.Build.cs
--- a/Engine/Plugins/MovieScene/LevelSequenceEditor/Source/LevelSequenceEditor/LevelSequenceEditor.Build.cs
+++ b/Engine/Plugins/MovieScene/LevelSequenceEditor/Source/LevelSequenceEditor/LevelSequenceEditor.Build.cs
@@ -21,6 +21,7 @@ public class LevelSequenceEditor : ModuleRules
         
         PrivateDependencyModuleNames.AddRange(
 			new string[] {
+				"ApplicationCore",
 				"AppFramework",
                 "LevelSequence",
 				"BlueprintGraph",
diff --git a/Engine/Plugins/MovieScene/LevelSequenceEditor/Source/LevelSequenceEditor/Private/LevelSequenceEditorSubsystem.cpp b/Engine/Plugins/MovieScene/LevelSequenceEditor/Source/LevelSequenceEditor/Private/LevelSequenceEditorSubsystem.cpp
--- a/Engine/Plugins/MovieScene/LevelSequenceEditor/Source/LevelSequenceEditor/Private/LevelSequenceEditorSubsystem.cpp
+++ b/Engine/Plugins/MovieScene/LevelSequenceEditor/Source/LevelSequenceEditor/Private/LevelSequenceEditorSubsystem.cpp
@@ -5,6 +5,7 @@
 #include "Compilation/MovieSceneCompiledDataManager.h"
 #include "ISequencer.h"
 #include "ISequencerModule.h"
+#include "LevelSequence.h"
 #include "LevelSequenceEditorCommands.h"
 #include "MovieScene.h"
 #include "MovieSceneCommonHelpers.h"
@@ -13,7 +14,7 @@
 #include "MovieSceneTimeHelpers.h"
 #include "SequencerSettings.h"
 #include "SequencerUtilities.h"
-#include "SequencerBindingProxy.h"
+#include "MovieSceneBindingProxy.h"
 #include "SequenceTimeUnit.h"
 #include "Sections/MovieScene3DConstraintSection.h"
 #include "Sections/MovieScene3DTransformSection.h"
@@ -30,6 +31,7 @@
 #include "Framework/MultiBox/MultiBoxBuilder.h"
 #include "Framework/Notifications/NotificationManager.h"
 #include "GameFramework/Actor.h"
+#include "HAL/PlatformApplicationMisc.h"
 #include "IContentBrowserSingleton.h"
 #include "LevelEditor.h"
 #include "Modules/ModuleManager.h"
@@ -38,6 +40,7 @@
 #include "ScopedTransaction.h"
 #include "Widgets/Notifications/SNotificationList.h"
 #include "Framework/MultiBox/MultiBoxBuilder.h"
+
 DEFINE_LOG_CATEGORY(LogLevelSequenceEditor);
 
 #define LOCTEXT_NAMESPACE "LevelSequenceEditor"
@@ -176,13 +179,239 @@ TSharedPtr<ISequencer> ULevelSequenceEditorSubsystem::GetActiveSequencer()
 	{
 		if (Ptr.IsValid())
 		{
-			return Ptr.Pin();
+			UMovieSceneSequence* Sequence = Ptr.Pin()->GetFocusedMovieSceneSequence();
+			if (Sequence && Sequence->IsA<ULevelSequence>())
+			{
+				return Ptr.Pin();
+			}
 		}
 	}
 
 	return nullptr;
 }
 
+TArray<FMovieSceneBindingProxy> ULevelSequenceEditorSubsystem::AddActors(const TArray<AActor*>& InActors)
+{
+	TArray<FMovieSceneBindingProxy> BindingProxies;
+
+	TSharedPtr<ISequencer> Sequencer = GetActiveSequencer();
+	if (Sequencer == nullptr)
+	{
+		return BindingProxies;
+	}
+
+	UMovieSceneSequence* Sequence = Sequencer->GetFocusedMovieSceneSequence();
+	if (!Sequence)
+	{
+		return BindingProxies;
+	}
+
+	TArray<TWeakObjectPtr<AActor> > Actors;
+	for (AActor* Actor : InActors)
+	{
+		Actors.Add(Actor);
+	}
+
+	TArray<FGuid> Guids = FSequencerUtilities::AddActors(Sequencer.ToSharedRef(), Actors);
+	
+	for (const FGuid& Guid : Guids)
+	{
+		BindingProxies.Add(FMovieSceneBindingProxy(Guid, Sequence));
+	}
+
+	return BindingProxies;
+}
+
+FMovieSceneBindingProxy ULevelSequenceEditorSubsystem::CreateCamera(bool bSpawnable, ACineCameraActor*& OutActor)
+{
+	TSharedPtr<ISequencer> Sequencer = GetActiveSequencer();
+	if (Sequencer == nullptr)
+	{
+		return FMovieSceneBindingProxy();
+	}
+
+	UMovieSceneSequence* Sequence = Sequencer->GetFocusedMovieSceneSequence();
+	if (!Sequence)
+	{
+		return FMovieSceneBindingProxy();
+	}
+
+	FGuid Guid = FSequencerUtilities::CreateCamera(Sequencer.ToSharedRef(), bSpawnable, OutActor);
+
+	return FMovieSceneBindingProxy(Guid, Sequence);
+}
+
+TArray<FMovieSceneBindingProxy> ULevelSequenceEditorSubsystem::ConvertToSpawnable(const FMovieSceneBindingProxy& ObjectBinding)
+{
+	TArray<FMovieSceneBindingProxy> SpawnableProxies;
+
+	TSharedPtr<ISequencer> Sequencer = GetActiveSequencer();
+	if (Sequencer == nullptr)
+	{
+		return SpawnableProxies;
+	}
+
+	UMovieSceneSequence* Sequence = Sequencer->GetFocusedMovieSceneSequence();
+	if (!Sequence)
+	{
+		return SpawnableProxies;
+	}
+
+	TArray<FMovieSceneSpawnable*> Spawnables = FSequencerUtilities::ConvertToSpawnable(Sequencer.ToSharedRef(), ObjectBinding.BindingID);
+	for (FMovieSceneSpawnable* Spawnable : Spawnables)
+	{
+		if (Spawnable)
+		{
+			SpawnableProxies.Add(FMovieSceneBindingProxy(Spawnable->GetGuid(), Sequence));
+		}
+	}
+
+	return SpawnableProxies;
+}
+
+FMovieSceneBindingProxy ULevelSequenceEditorSubsystem::ConvertToPossessable(const FMovieSceneBindingProxy& ObjectBinding)
+{
+	FMovieSceneBindingProxy PossessableProxy;
+
+	TSharedPtr<ISequencer> Sequencer = GetActiveSequencer();
+	if (Sequencer == nullptr)
+	{
+		return PossessableProxy;
+	}
+
+	UMovieSceneSequence* Sequence = Sequencer->GetFocusedMovieSceneSequence();
+	if (!Sequence)
+	{
+		return PossessableProxy;
+	}
+
+	if (FMovieScenePossessable* Possessable = FSequencerUtilities::ConvertToPossessable(Sequencer.ToSharedRef(), ObjectBinding.BindingID))
+	{
+		PossessableProxy = FMovieSceneBindingProxy(Possessable->GetGuid(), Sequence);
+	}
+
+	return PossessableProxy;
+}
+
+void ULevelSequenceEditorSubsystem::CopyFolders(const TArray<UMovieSceneFolder*>& Folders, FString& ExportedText)
+{
+	FSequencerUtilities::CopyFolders(Folders, ExportedText);
+
+	FPlatformApplicationMisc::ClipboardCopy(*ExportedText);
+}
+
+bool ULevelSequenceEditorSubsystem::PasteFolders(const FString& InTextToImport, FMovieScenePasteFoldersParams PasteFoldersParams, TArray<UMovieSceneFolder*>& OutFolders)
+{
+	FString TextToImport = InTextToImport;
+	if (TextToImport.IsEmpty())
+	{
+		FPlatformApplicationMisc::ClipboardPaste(TextToImport);
+	}
+
+	TArray<FNotificationInfo> PasteErrors;
+	if (!FSequencerUtilities::PasteFolders(TextToImport, PasteFoldersParams, OutFolders, PasteErrors))
+	{
+		for (FNotificationInfo PasteError : PasteErrors)
+		{
+			UE_LOG(LogLevelSequenceEditor, Error, TEXT("%s"), *PasteError.Text.Get().ToString());
+		}
+		return false;
+	}
+
+	return true;
+}
+
+void ULevelSequenceEditorSubsystem::CopySections(const TArray<UMovieSceneSection*>& Sections, FString& ExportedText)
+{
+	FSequencerUtilities::CopySections(Sections, ExportedText);
+
+	FPlatformApplicationMisc::ClipboardCopy(*ExportedText);
+}
+
+bool ULevelSequenceEditorSubsystem::PasteSections(const FString& InTextToImport, FMovieScenePasteSectionsParams PasteSectionsParams, TArray<UMovieSceneSection*>& OutSections)
+{
+	FString TextToImport = InTextToImport;
+	if (TextToImport.IsEmpty())
+	{
+		FPlatformApplicationMisc::ClipboardPaste(TextToImport);
+	}
+
+	TArray<FNotificationInfo> PasteErrors;
+	if (!FSequencerUtilities::PasteSections(TextToImport, PasteSectionsParams, OutSections, PasteErrors))
+	{
+		for (FNotificationInfo PasteError : PasteErrors)
+		{
+			UE_LOG(LogLevelSequenceEditor, Error, TEXT("%s"), *PasteError.Text.Get().ToString());
+		}
+		return false;
+	}
+
+	return true;
+}
+
+void ULevelSequenceEditorSubsystem::CopyTracks(const TArray<UMovieSceneTrack*>& Tracks, FString& ExportedText)
+{
+	TArray<UMovieSceneFolder*> Folders;
+	FSequencerUtilities::CopyTracks(Tracks, Folders, ExportedText);
+	FPlatformApplicationMisc::ClipboardCopy(*ExportedText);
+}
+
+bool ULevelSequenceEditorSubsystem::PasteTracks(const FString& InTextToImport, FMovieScenePasteTracksParams PasteTracksParams, TArray<UMovieSceneTrack*>& OutTracks)
+{
+	FString TextToImport = InTextToImport;
+	if (TextToImport.IsEmpty())
+	{
+		FPlatformApplicationMisc::ClipboardPaste(TextToImport);
+	}
+
+	TArray<FNotificationInfo> PasteErrors;
+	if (!FSequencerUtilities::PasteTracks(TextToImport, PasteTracksParams, OutTracks, PasteErrors))
+	{
+		for (FNotificationInfo PasteError : PasteErrors)
+		{
+			UE_LOG(LogLevelSequenceEditor, Error, TEXT("%s"), *PasteError.Text.Get().ToString());
+		}
+		return false;
+	}
+
+	return true;
+}
+
+void ULevelSequenceEditorSubsystem::CopyBindings(const TArray<FMovieSceneBindingProxy>& Bindings, FString& ExportedText)
+{
+	TArray<UMovieSceneFolder*> Folders;
+	FSequencerUtilities::CopyBindings(Bindings, Folders, ExportedText);
+	FPlatformApplicationMisc::ClipboardCopy(*ExportedText);
+}
+
+bool ULevelSequenceEditorSubsystem::PasteBindings(const FString& InTextToImport, FMovieScenePasteBindingsParams PasteBindingsParams, TArray<FMovieSceneBindingProxy>& OutObjectBindings)
+{
+	TSharedPtr<ISequencer> Sequencer = GetActiveSequencer();
+	if (Sequencer == nullptr)
+	{
+		return false;
+	}
+
+	FString TextToImport = InTextToImport;
+	if (TextToImport.IsEmpty())
+	{
+		FPlatformApplicationMisc::ClipboardPaste(TextToImport);
+	}
+
+	TArray<FNotificationInfo> PasteErrors;
+	TArray<FMovieSceneBindingProxy> OutBindings;
+	if (!FSequencerUtilities::PasteBindings(TextToImport, Sequencer.ToSharedRef(), PasteBindingsParams, OutBindings, PasteErrors))
+	{
+		for (FNotificationInfo PasteError : PasteErrors)
+		{
+			UE_LOG(LogLevelSequenceEditor, Error, TEXT("%s"), *PasteError.Text.Get().ToString());
+		}
+		return false;
+	}
+
+	return true;
+}
+
 void ULevelSequenceEditorSubsystem::SnapSectionsToTimelineUsingSourceTimecodeInternal()
 {
 	const TSharedPtr<ISequencer> Sequencer = GetActiveSequencer();
@@ -395,10 +624,10 @@ void ULevelSequenceEditorSubsystem::BakeTransformInternal()
 	TArray<FGuid> ObjectBindings;
 	Sequencer->GetSelectedObjects(ObjectBindings);
 
-	TArray<FSequencerBindingProxy> BindingProxies;
+	TArray<FMovieSceneBindingProxy> BindingProxies;
 	for (FGuid Guid : ObjectBindings)
 	{
-		BindingProxies.Add(FSequencerBindingProxy(Guid, Sequencer->GetFocusedMovieSceneSequence()));
+		BindingProxies.Add(FMovieSceneBindingProxy(Guid, Sequencer->GetFocusedMovieSceneSequence()));
 	}
 
 	FMovieSceneScriptingParams Params;
@@ -406,7 +635,7 @@ void ULevelSequenceEditorSubsystem::BakeTransformInternal()
 	BakeTransform(BindingProxies, FFrameTime(InFrame), FFrameTime(OutFrame), ConvertFrameTime(1, DisplayRate, TickResolution), Params);
 }
 
-void ULevelSequenceEditorSubsystem::BakeTransform(const TArray<FSequencerBindingProxy>& ObjectBindings, const FFrameTime& BakeInTime, const FFrameTime& BakeOutTime, const FFrameTime& BakeInterval, const FMovieSceneScriptingParams& Params)
+void ULevelSequenceEditorSubsystem::BakeTransform(const TArray<FMovieSceneBindingProxy>& ObjectBindings, const FFrameTime& BakeInTime, const FFrameTime& BakeOutTime, const FFrameTime& BakeInterval, const FMovieSceneScriptingParams& Params)
 {
 	TSharedPtr<ISequencer> Sequencer = GetActiveSequencer();
 	if (Sequencer == nullptr)
@@ -460,7 +689,7 @@ void ULevelSequenceEditorSubsystem::BakeTransform(const TArray<FSequencerBinding
 	};
 
 	TMap<FGuid, FBakeData> BakeDataMap;
-	for (const FSequencerBindingProxy& ObjectBinding : ObjectBindings)
+	for (const FMovieSceneBindingProxy& ObjectBinding : ObjectBindings)
 	{
 		BakeDataMap.Add(ObjectBinding.BindingID);
 	}
@@ -474,7 +703,7 @@ void ULevelSequenceEditorSubsystem::BakeTransform(const TArray<FSequencerBinding
 
 		Sequencer->SetLocalTimeDirectly(Range.GetTime());
 
-		for (const FSequencerBindingProxy& ObjectBinding : ObjectBindings)
+		for (const FMovieSceneBindingProxy& ObjectBinding : ObjectBindings)
 		{
 			FGuid Guid = ObjectBinding.BindingID;
 
@@ -810,7 +1039,7 @@ void ULevelSequenceEditorSubsystem::FixActorReferences()
 			// In that case, make sure we also replace the existing guid with the new guid 
 			FGuid ExistingGuid = Sequencer->FindObjectId(*ActorPtr, Sequencer->GetFocusedTemplateID());
 
-			FGuid NewGuid = FSequencerUtilities::DoAssignActor(Sequencer.Get(), ActorPtr, ActorPossessableToFix.GetGuid());
+			FGuid NewGuid = FSequencerUtilities::AssignActor(Sequencer.ToSharedRef(), ActorPtr, ActorPossessableToFix.GetGuid());
 
 			OldGuidToNewGuidMap.Add(OldGuid, NewGuid);
 
@@ -821,15 +1050,9 @@ void ULevelSequenceEditorSubsystem::FixActorReferences()
 		}
 	}
 
-	UMovieSceneCompiledDataManager* CompiledDataManager = FindObject<UMovieSceneCompiledDataManager>(GetTransientPackage(), TEXT("SequencerCompiledDataManager"));
-	if (!CompiledDataManager)
-	{
-		CompiledDataManager = NewObject<UMovieSceneCompiledDataManager>(GetTransientPackage(), "SequencerCompiledDataManager");
-	}
-
 	for (TPair<FGuid, FGuid> GuidPair : OldGuidToNewGuidMap)
 	{
-		FSequencerUtilities::UpdateBindingIDs(Sequencer.Get(), CompiledDataManager, GuidPair.Key, GuidPair.Value);
+		FSequencerUtilities::UpdateBindingIDs(Sequencer.ToSharedRef(), GuidPair.Key, GuidPair.Value);
 	}
 }
 
@@ -848,109 +1071,22 @@ void ULevelSequenceEditorSubsystem::AddActorsToBindingInternal()
 		return;
 	}
 
-	FSequencerBindingProxy BindingProxy(ObjectBindings[0], Sequencer->GetFocusedMovieSceneSequence());
+	FMovieSceneBindingProxy BindingProxy(ObjectBindings[0], Sequencer->GetFocusedMovieSceneSequence());
 
 	TArray<AActor*> SelectedActors;
 	GEditor->GetSelectedActors()->GetSelectedObjects<AActor>(SelectedActors); 
 	AddActorsToBinding(SelectedActors, BindingProxy);
 }
 
-void ULevelSequenceEditorSubsystem::AddActorsToBinding(const TArray<AActor*>& Actors, const FSequencerBindingProxy& ObjectBinding)
+void ULevelSequenceEditorSubsystem::AddActorsToBinding(const TArray<AActor*>& Actors, const FMovieSceneBindingProxy& ObjectBinding)
 {
-	if (!Actors.Num())
-	{
-		return;
-	}
-
 	TSharedPtr<ISequencer> Sequencer = GetActiveSequencer();
 	if (Sequencer == nullptr)
 	{
 		return;
 	}
 
-	UMovieSceneSequence* Sequence = Sequencer->GetFocusedMovieSceneSequence();
-	if (!Sequence)
-	{
-		return;
-	}
-
-	UMovieScene* MovieScene = Sequence->GetMovieScene();
-	if (!MovieScene)
-	{
-		return;
-	}
-
-	UClass* ActorClass = nullptr;
-	int32 NumRuntimeObjects = 0;
-
-	FGuid Guid = ObjectBinding.BindingID;
-	TArrayView<TWeakObjectPtr<>> ObjectsInCurrentSequence = Sequencer->FindObjectsInCurrentSequence(Guid);
-
-	for (TWeakObjectPtr<> Ptr : ObjectsInCurrentSequence)
-	{
-		if (const AActor* Actor = Cast<AActor>(Ptr.Get()))
-		{
-			ActorClass = Actor->GetClass();
-			++NumRuntimeObjects;
-		}
-	}
-
-	FScopedTransaction AddSelectedToBinding(LOCTEXT("AddSelectedToBinding", "Add Selected to Binding"));
-
-	Sequence->Modify();
-	MovieScene->Modify();
-
-	// Bind objects
-	int32 NumObjectsAdded = 0;
-	for (AActor* ActorToAdd : Actors)
-	{
-		if (!ObjectsInCurrentSequence.Contains(ActorToAdd))
-		{
-			if (ActorClass == nullptr || UClass::FindCommonBase(ActorToAdd->GetClass(), ActorClass) != nullptr)
-			{
-				if (ActorClass == nullptr)
-				{
-					ActorClass = ActorToAdd->GetClass();
-				}
-
-				ActorToAdd->Modify();
-				Sequence->BindPossessableObject(Guid, *ActorToAdd, Sequencer->GetPlaybackContext());
-				++NumObjectsAdded;
-			}
-			else
-			{
-				const FText NotificationText = FText::Format(LOCTEXT("UnableToAssignObject", "Cannot assign object {0}. Expected class {1}"), FText::FromString(ActorToAdd->GetName()), FText::FromString(ActorClass->GetName()));
-				FNotificationInfo Info(NotificationText);
-				Info.ExpireDuration = 3.f;
-				Info.bUseLargeFont = false;
-				FSlateNotificationManager::Get().AddNotification(Info);
-			}
-		}
-	}
-
-	// Update label
-	if (NumRuntimeObjects + NumObjectsAdded > 0)
-	{
-		FMovieScenePossessable* Possessable = MovieScene->FindPossessable(Guid);
-		if (Possessable && ActorClass != nullptr)
-		{
-			if (NumRuntimeObjects + NumObjectsAdded > 1)
-			{
-				FString NewLabel = ActorClass->GetName() + FString::Printf(TEXT(" (%d)"), NumRuntimeObjects + NumObjectsAdded);
-				Possessable->SetName(NewLabel);
-			}
-			else if (NumObjectsAdded > 0 && Actors.Num() > 0)
-			{
-				Possessable->SetName(Actors[0]->GetActorLabel());
-			}
-
-			Possessable->SetPossessedObjectClass(ActorClass);
-		}
-	}
-
-	Sequencer->RestorePreAnimatedState();
-
-	Sequencer->NotifyMovieSceneDataChanged(EMovieSceneDataChangeType::MovieSceneStructureItemsChanged);
+	FSequencerUtilities::AddActorsToBinding(Sequencer.ToSharedRef(), Actors, ObjectBinding);
 }
 
 void ULevelSequenceEditorSubsystem::ReplaceBindingWithActorsInternal()
@@ -968,14 +1104,14 @@ void ULevelSequenceEditorSubsystem::ReplaceBindingWithActorsInternal()
 		return;
 	}
 
-	FSequencerBindingProxy BindingProxy(ObjectBindings[0], Sequencer->GetFocusedMovieSceneSequence());
+	FMovieSceneBindingProxy BindingProxy(ObjectBindings[0], Sequencer->GetFocusedMovieSceneSequence());
 
 	TArray<AActor*> SelectedActors;
 	GEditor->GetSelectedActors()->GetSelectedObjects<AActor>(SelectedActors); 
 	ReplaceBindingWithActors(SelectedActors, BindingProxy);
 }
 
-void ULevelSequenceEditorSubsystem::ReplaceBindingWithActors(const TArray<AActor*>& Actors, const FSequencerBindingProxy& ObjectBinding)
+void ULevelSequenceEditorSubsystem::ReplaceBindingWithActors(const TArray<AActor*>& Actors, const FMovieSceneBindingProxy& ObjectBinding)
 {
 	TSharedPtr<ISequencer> Sequencer = GetActiveSequencer();
 	if (Sequencer == nullptr)
@@ -983,33 +1119,7 @@ void ULevelSequenceEditorSubsystem::ReplaceBindingWithActors(const TArray<AActor
 		return;
 	}
 
-	FScopedTransaction ReplaceBindingWithActors(LOCTEXT("ReplaceBindingWithActors", "Replace Binding with Actors"));
-
-	FGuid Guid = ObjectBinding.BindingID;
-	TArray<AActor*> ExistingActors;
-	for (TWeakObjectPtr<> Ptr : Sequencer->FindObjectsInCurrentSequence(Guid))
-	{
-		if (AActor* Actor = Cast<AActor>(Ptr.Get()))
-		{
-			if (!Actors.Contains(Actor))
-			{
-				ExistingActors.Add(Actor);
-			}
-		}
-	}
-
-	RemoveActorsFromBinding(ExistingActors, ObjectBinding);
-
-	TArray<AActor*> NewActors;
-	for (AActor* NewActor : Actors)
-	{
-		if (!ExistingActors.Contains(NewActor))
-		{
-			NewActors.Add(NewActor);
-		}
-	}
-
-	AddActorsToBinding(NewActors, ObjectBinding);
+	FSequencerUtilities::ReplaceBindingWithActors(Sequencer.ToSharedRef(), Actors, ObjectBinding);
 }
 
 void ULevelSequenceEditorSubsystem::RemoveActorsFromBindingInternal()
@@ -1027,99 +1137,22 @@ void ULevelSequenceEditorSubsystem::RemoveActorsFromBindingInternal()
 		return;
 	}
 
-	FSequencerBindingProxy BindingProxy(ObjectBindings[0], Sequencer->GetFocusedMovieSceneSequence());
+	FMovieSceneBindingProxy BindingProxy(ObjectBindings[0], Sequencer->GetFocusedMovieSceneSequence());
 
 	TArray<AActor*> SelectedActors;
 	GEditor->GetSelectedActors()->GetSelectedObjects<AActor>(SelectedActors); 
 	RemoveActorsFromBinding(SelectedActors, BindingProxy);
 }
 
-void ULevelSequenceEditorSubsystem::RemoveActorsFromBinding(const TArray<AActor*>& Actors, const FSequencerBindingProxy& ObjectBinding)
+void ULevelSequenceEditorSubsystem::RemoveActorsFromBinding(const TArray<AActor*>& Actors, const FMovieSceneBindingProxy& ObjectBinding)
 {
-	if (!Actors.Num())
-	{
-		return;
-	}
-
 	TSharedPtr<ISequencer> Sequencer = GetActiveSequencer();
 	if (Sequencer == nullptr)
 	{
 		return;
 	}
 
-	UMovieSceneSequence* Sequence = Sequencer->GetFocusedMovieSceneSequence();
-	if (!Sequence)
-	{
-		return;
-	}
-
-	UMovieScene* MovieScene = Sequence->GetMovieScene();
-	if (!MovieScene)
-	{
-		return;
-	}
-
-	UClass* ActorClass = nullptr;
-	int32 NumRuntimeObjects = 0;
-
-	FGuid Guid = ObjectBinding.BindingID;
-	for (TWeakObjectPtr<> Ptr : Sequencer->FindObjectsInCurrentSequence(Guid))
-	{
-		if (const AActor* Actor = Cast<AActor>(Ptr.Get()))
-		{
-			ActorClass = Actor->GetClass();
-			++NumRuntimeObjects;
-		}
-	}
-
-	FScopedTransaction RemoveSelectedFromBinding(LOCTEXT("RemoveSelectedFromBinding", "Remove Selected from Binding"));
-
-	TArray<UObject*> ObjectsToRemove;
-	for (AActor* ActorToRemove : Actors)
-	{
-		// Restore state on any components
-		for (UActorComponent* Component : TInlineComponentArray<UActorComponent*>(ActorToRemove))
-		{
-			if (Component)
-			{
-				Sequencer->PreAnimatedState.RestorePreAnimatedState(*Component);
-			}
-		}
-
-		// Restore state on the object itself
-		Sequencer->PreAnimatedState.RestorePreAnimatedState(*ActorToRemove);
-
-		ActorToRemove->Modify();
-
-		ObjectsToRemove.Add(ActorToRemove);
-	}
-	
-	Sequence->Modify();
-	MovieScene->Modify();
-
-	// Unbind objects
-	Sequence->UnbindObjects(Guid, ObjectsToRemove, Sequencer->GetPlaybackContext());
-
-	// Update label
-	if (NumRuntimeObjects - ObjectsToRemove.Num() > 0)
-	{
-		FMovieScenePossessable* Possessable = MovieScene->FindPossessable(Guid);
-		if (Possessable && ActorClass != nullptr)
-		{
-			if (NumRuntimeObjects - ObjectsToRemove.Num() > 1)
-			{
-				FString NewLabel = ActorClass->GetName() + FString::Printf(TEXT(" (%d)"), NumRuntimeObjects - ObjectsToRemove.Num());
-
-				Possessable->SetName(NewLabel);
-			}
-			else if (ObjectsToRemove.Num() > 0 && Actors.Num() > 0)
-			{
-				Possessable->SetName(Actors[0]->GetActorLabel());
-			}
-		}
-	}
-
-	Sequencer->NotifyMovieSceneDataChanged(EMovieSceneDataChangeType::MovieSceneStructureItemsChanged);
+	FSequencerUtilities::RemoveActorsFromBinding(Sequencer.ToSharedRef(), Actors, ObjectBinding);
 }
 
 void ULevelSequenceEditorSubsystem::RemoveAllBindingsInternal()
@@ -1137,12 +1170,12 @@ void ULevelSequenceEditorSubsystem::RemoveAllBindingsInternal()
 		return;
 	}
 
-	FSequencerBindingProxy BindingProxy(ObjectBindings[0], Sequencer->GetFocusedMovieSceneSequence());
+	FMovieSceneBindingProxy BindingProxy(ObjectBindings[0], Sequencer->GetFocusedMovieSceneSequence());
 
 	RemoveAllBindings(BindingProxy);
 }
 
-void ULevelSequenceEditorSubsystem::RemoveAllBindings(const FSequencerBindingProxy& ObjectBinding)
+void ULevelSequenceEditorSubsystem::RemoveAllBindings(const FMovieSceneBindingProxy& ObjectBinding)
 {
 	TSharedPtr<ISequencer> Sequencer = GetActiveSequencer();
 	if (Sequencer == nullptr)
@@ -1191,12 +1224,12 @@ void ULevelSequenceEditorSubsystem::RemoveInvalidBindingsInternal()
 		return;
 	}
 
-	FSequencerBindingProxy BindingProxy(ObjectBindings[0], Sequencer->GetFocusedMovieSceneSequence());
+	FMovieSceneBindingProxy BindingProxy(ObjectBindings[0], Sequencer->GetFocusedMovieSceneSequence());
 
 	RemoveInvalidBindings(BindingProxy);
 }
 
-void ULevelSequenceEditorSubsystem::RemoveInvalidBindings(const FSequencerBindingProxy& ObjectBinding)
+void ULevelSequenceEditorSubsystem::RemoveInvalidBindings(const FMovieSceneBindingProxy& ObjectBinding)
 {
 	TSharedPtr<ISequencer> Sequencer = GetActiveSequencer();
 	if (Sequencer == nullptr)
@@ -1279,7 +1312,7 @@ void ULevelSequenceEditorSubsystem::AddAssignActorMenu(FMenuBuilder& MenuBuilder
 		return;
 	}
 
-	FSequencerBindingProxy BindingProxy(ObjectBindings[0], Sequencer->GetFocusedMovieSceneSequence());
+	FMovieSceneBindingProxy BindingProxy(ObjectBindings[0], Sequencer->GetFocusedMovieSceneSequence());
 
 	TSet<const AActor*> BoundObjects;
 	{
@@ -1323,7 +1356,7 @@ void ULevelSequenceEditorSubsystem::AddAssignActorMenu(FMenuBuilder& MenuBuilder
 				FOnActorPicked::CreateLambda([=](AActor* Actor){
 					// Create a new binding for this actor
 					FSlateApplication::Get().DismissAllMenus();
-					FSequencerUtilities::DoAssignActor(Sequencer.Get(), Actor, ObjectBindings[0]);
+					FSequencerUtilities::AssignActor(Sequencer.ToSharedRef(), Actor, ObjectBindings[0]);
 				})
 			)
 		];
@@ -1425,17 +1458,17 @@ void ULevelSequenceEditorSubsystem::RebindComponentInternal(const FName& Compone
 		return;
 	}
 
-	TArray<FSequencerBindingProxy> BindingProxies;
+	TArray<FMovieSceneBindingProxy> BindingProxies;
 	for (const FGuid& ObjectBinding : ObjectBindings)
 	{
-		FSequencerBindingProxy BindingProxy(ObjectBinding, Sequencer->GetFocusedMovieSceneSequence());
+		FMovieSceneBindingProxy BindingProxy(ObjectBinding, Sequencer->GetFocusedMovieSceneSequence());
 		BindingProxies.Add(BindingProxy);
 	}
 
 	RebindComponent(BindingProxies, ComponentName);
 }
 
-void ULevelSequenceEditorSubsystem::RebindComponent(const TArray<FSequencerBindingProxy>& PossessableBindings, const FName& ComponentName)
+void ULevelSequenceEditorSubsystem::RebindComponent(const TArray<FMovieSceneBindingProxy>& PossessableBindings, const FName& ComponentName)
 {
 	TSharedPtr<ISequencer> Sequencer = GetActiveSequencer();
 	if (Sequencer == nullptr)
@@ -1461,7 +1494,7 @@ void ULevelSequenceEditorSubsystem::RebindComponent(const TArray<FSequencerBindi
 	MovieScene->Modify();
 
 	bool bAnythingChanged = false;
-	for (const FSequencerBindingProxy& PossessableBinding : PossessableBindings)
+	for (const FMovieSceneBindingProxy& PossessableBinding : PossessableBindings)
 	{
 		FMovieScenePossessable* ComponentPossessable = MovieScene->FindPossessable(PossessableBinding.BindingID);
 
diff --git a/Engine/Plugins/MovieScene/LevelSequenceEditor/Source/LevelSequenceEditor/Public/LevelSequenceEditorSubsystem.h b/Engine/Plugins/MovieScene/LevelSequenceEditor/Source/LevelSequenceEditor/Public/LevelSequenceEditorSubsystem.h
--- a/Engine/Plugins/MovieScene/LevelSequenceEditor/Source/LevelSequenceEditor/Public/LevelSequenceEditorSubsystem.h
+++ b/Engine/Plugins/MovieScene/LevelSequenceEditor/Source/LevelSequenceEditor/Public/LevelSequenceEditorSubsystem.h
@@ -8,15 +8,18 @@
 #include "ISequencer.h"
 
 #include "SequenceTimeUnit.h"
-#include "SequencerBindingProxy.h"
+#include "MovieSceneBindingProxy.h"
+#include "SequencerUtilities.h"
 
 #include "LevelSequenceEditorSubsystem.generated.h"
 
 DECLARE_LOG_CATEGORY_EXTERN(LogLevelSequenceEditor, Log, All);
 
+class ACineCameraActor;
 class FExtender;
 class FMenuBuilder;
 class UMovieSceneCompiledDataManager;
+class UMovieSceneFolder;
 class UMovieSceneSection;
 
 USTRUCT(BlueprintType)
@@ -26,6 +29,7 @@ struct FMovieSceneScriptingParams
 		
 	FMovieSceneScriptingParams() {}
 
+	UPROPERTY(BlueprintReadWrite, Category = "Movie Scene")
 	ESequenceTimeUnit TimeUnit = ESequenceTimeUnit::DisplayRate;
 };
 
@@ -44,7 +48,91 @@ public:
 	virtual void Deinitialize() override;
 
 	void OnSequencerCreated(TSharedRef<ISequencer> InSequencer);
-	
+
+	/** Add existing actors to Sequencer. Tracks will be automatically added based on default track settings. */
+	UFUNCTION(BlueprintCallable, Category = "Level Sequence Editor")
+	TArray<FMovieSceneBindingProxy> AddActors(const TArray<AActor*>& Actors);
+
+	/** Create a cine camera actor and add it to Sequencer */
+	UFUNCTION(BlueprintCallable, Category = "Level Sequence Editor")
+	FMovieSceneBindingProxy CreateCamera(bool bSpawnable, ACineCameraActor*& OutActor);
+
+	/** Convert to spawnable. If there are multiple objects assigned to the possessable, multiple spawnables will be created. */
+	UFUNCTION(BlueprintCallable, Category = "Level Sequence Editor")
+	TArray<FMovieSceneBindingProxy> ConvertToSpawnable(const FMovieSceneBindingProxy& ObjectBinding);
+
+	/** Convert to possessable */
+	UFUNCTION(BlueprintCallable, Category = "Level Sequence Editor")
+	FMovieSceneBindingProxy ConvertToPossessable(const FMovieSceneBindingProxy& ObjectBinding);
+
+	/** 
+	 * Copy folders 
+	 * The copied folders will be saved to the clipboard as well as assigned to the ExportedText string. 
+	 * The ExportedTest string can be used in conjunction with PasteFolders if, for example, pasting copy/pasting multiple 
+	 * folders without relying on a single clipboard. 
+	 */
+	UFUNCTION(BlueprintCallable, Category = "Level Sequence Editor")
+	void CopyFolders(const TArray<UMovieSceneFolder*>& Folders, FString& ExportedText);
+
+	/** 
+	 * Paste folders 
+	 * Paste folders from the given TextToImport string (used in conjunction with CopyFolders). 
+	 * If TextToImport is empty, the contents of the clipboard will be used.
+	 */
+	UFUNCTION(BlueprintCallable, Category = "Level Sequence Editor")
+	bool PasteFolders(const FString& TextToImport, FMovieScenePasteFoldersParams PasteFoldersParams, TArray<UMovieSceneFolder*>& OutFolders);
+
+	/**
+	 * Copy sections
+	 * The copied sections will be saved to the clipboard as well as assigned to the ExportedText string.
+	 * The ExportedTest string can be used in conjunction with PasteSections if, for example, pasting copy/pasting multiple
+	 * sections without relying on a single clipboard.
+	 */
+	UFUNCTION(BlueprintCallable, Category = "Level Sequence Editor")
+	void CopySections(const TArray<UMovieSceneSection*>& Sections, FString& ExportedText);
+
+	/**
+	 * Paste sections
+	 * Paste sections from the given TextToImport string (used in conjunction with CopySections).
+	 * If TextToImport is empty, the contents of the clipboard will be used.
+	 */
+	UFUNCTION(BlueprintCallable, Category = "Level Sequence Editor")
+	bool PasteSections(const FString& TextToImport, FMovieScenePasteSectionsParams PasteSectionsParams, TArray<UMovieSceneSection*>& OutSections);
+
+	/**
+	 * Copy tracks
+	 * The copied tracks will be saved to the clipboard as well as assigned to the ExportedText string.
+	 * The ExportedTest string can be used in conjunction with PasteTracks if, for example, pasting copy/pasting multiple
+	 * tracks without relying on a single clipboard.
+	 */
+	UFUNCTION(BlueprintCallable, Category = "Level Sequence Editor")
+	void CopyTracks(const TArray<UMovieSceneTrack*>& Tracks, FString& ExportedText);
+
+	/**
+	 * Paste tracks
+	 * Paste tracks from the given TextToImport string (used in conjunction with CopyTracks).
+	 * If TextToImport is empty, the contents of the clipboard will be used.
+	 */
+	UFUNCTION(BlueprintCallable, Category = "Level Sequence Editor")
+	bool PasteTracks(const FString& TextToImport, FMovieScenePasteTracksParams PasteTracksParams, TArray<UMovieSceneTrack*>& OutTracks);
+
+	/**
+	 * Copy bindings
+	 * The copied bindings will be saved to the clipboard as well as assigned to the ExportedText string.
+	 * The ExportedTest string can be used in conjunction with PasteBindings if, for example, pasting copy/pasting multiple
+	 * bindings without relying on a single clipboard.
+	 */
+	UFUNCTION(BlueprintCallable, Category = "Level Sequence Editor")
+	void CopyBindings(const TArray<FMovieSceneBindingProxy>& Bindings, FString& ExportedText);
+
+	/**
+	 * Paste bindings
+	 * Paste bindings from the given TextToImport string (used in conjunction with CopyBindings).
+	 * If TextToImport is empty, the contents of the clipboard will be used.
+	 */
+	UFUNCTION(BlueprintCallable, Category = "Level Sequence Editor")
+	bool PasteBindings(const FString& TextToImport, FMovieScenePasteBindingsParams PasteBindingsParams, TArray<FMovieSceneBindingProxy>& OutObjectBindings);
+
 	/** Snap sections to timeline using source timecode */
 	UFUNCTION(BlueprintCallable, Category = "Level Sequence Editor")
 	void SnapSectionsToTimelineUsingSourceTimecode(const TArray<UMovieSceneSection*>& Sections);
@@ -55,7 +143,7 @@ public:
 
 	/** Bake transform */
 	UFUNCTION(BlueprintCallable, Category = "Level Sequence Editor")
-	void BakeTransform(const TArray<FSequencerBindingProxy>& ObjectBindings, const FFrameTime& BakeInTime, const FFrameTime& BakeOutTime, const FFrameTime& BakeInterval, const FMovieSceneScriptingParams& Params = FMovieSceneScriptingParams());
+	void BakeTransform(const TArray<FMovieSceneBindingProxy>& ObjectBindings, const FFrameTime& BakeInTime, const FFrameTime& BakeOutTime, const FFrameTime& BakeInterval, const FMovieSceneScriptingParams& Params = FMovieSceneScriptingParams());
 
 	/** Attempts to automatically fix up broken actor references in the current scene */
 	UFUNCTION(BlueprintCallable, Category = "Level Sequence Editor")
@@ -63,27 +151,27 @@ public:
 
 	/** Assigns the given actors to the binding */
 	UFUNCTION(BlueprintCallable, Category = "Level Sequence Editor")
-	void AddActorsToBinding(const TArray<AActor*>& Actors, const FSequencerBindingProxy& ObjectBinding);
+	void AddActorsToBinding(const TArray<AActor*>& Actors, const FMovieSceneBindingProxy& ObjectBinding);
 
 	/** Replaces the binding with the given actors */
 	UFUNCTION(BlueprintCallable, Category = "Level Sequence Editor")
-	void ReplaceBindingWithActors(const TArray<AActor*>& Actors, const FSequencerBindingProxy& ObjectBinding);
+	void ReplaceBindingWithActors(const TArray<AActor*>& Actors, const FMovieSceneBindingProxy& ObjectBinding);
 
 	/** Removes the given actors from the binding */
 	UFUNCTION(BlueprintCallable, Category = "Level Sequence Editor")
-	void RemoveActorsFromBinding(const TArray<AActor*>& Actors, const FSequencerBindingProxy& ObjectBinding);
+	void RemoveActorsFromBinding(const TArray<AActor*>& Actors, const FMovieSceneBindingProxy& ObjectBinding);
 
 	/** Remove all bound actors from this track */
 	UFUNCTION(BlueprintCallable, Category = "Level Sequence Editor")
-	void RemoveAllBindings(const FSequencerBindingProxy& ObjectBinding);
+	void RemoveAllBindings(const FMovieSceneBindingProxy& ObjectBinding);
 
 	/** Remove missing objects bound to this track */
 	UFUNCTION(BlueprintCallable, Category = "Level Sequence Editor")
-	void RemoveInvalidBindings(const FSequencerBindingProxy& ObjectBinding);
+	void RemoveInvalidBindings(const FMovieSceneBindingProxy& ObjectBinding);
 
 	/** Rebind the component binding to the requested component */
 	UFUNCTION(BlueprintCallable, Category = "Level Sequence Editor")
-	void RebindComponent(const TArray<FSequencerBindingProxy>& ComponentBindings, const FName& ComponentName);
+	void RebindComponent(const TArray<FMovieSceneBindingProxy>& ComponentBindings, const FName& ComponentName);
 
 private:
 
diff --git a/Engine/Source/Editor/Sequencer/Private/SSequencer.cpp b/Engine/Source/Editor/Sequencer/Private/SSequencer.cpp
--- a/Engine/Source/Editor/Sequencer/Private/SSequencer.cpp
+++ b/Engine/Source/Editor/Sequencer/Private/SSequencer.cpp
@@ -3137,21 +3137,7 @@ void SSequencer::OnAssetsDropped( const FAssetDragDropOp& DragDropOp )
 					continue;
 				}
 
-				FGuid NewGuid = SequencerRef.MakeNewSpawnable(*CurObject, DragDropOp.GetActorFactory());
-
-				UMovieScene* MovieScene = SequencerRef.GetFocusedMovieSceneSequence()->GetMovieScene();
-				if (MovieScene)
-				{
-					FMovieSceneSpawnable* Spawnable = MovieScene->FindSpawnable(NewGuid);
-
-					if (Spawnable)
-					{
-						if (ACameraActor* CameraActor = Cast<ACameraActor>(Spawnable->GetObjectTemplate()))
-						{
-							SequencerRef.NewCameraAdded(CameraActor, NewGuid);
-						}
-					}
-				}
+				SequencerRef.MakeNewSpawnable(*CurObject, DragDropOp.GetActorFactory());
 			}
 
 			DropResult = ESequencerDropResult::DropHandled;
diff --git a/Engine/Source/Editor/Sequencer/Private/Sequencer.cpp b/Engine/Source/Editor/Sequencer/Private/Sequencer.cpp
--- a/Engine/Source/Editor/Sequencer/Private/Sequencer.cpp
+++ b/Engine/Source/Editor/Sequencer/Private/Sequencer.cpp
@@ -27,7 +27,6 @@
 #include "MVVM/SectionModelStorageExtension.h"
 #include "Camera/PlayerCameraManager.h"
 #include "Misc/MessageDialog.h"
-#include "Containers/ArrayBuilder.h"
 #include "Misc/FeedbackContext.h"
 #include "Misc/ScopedSlowTask.h"
 #include "Modules/ModuleManager.h"
@@ -75,6 +74,7 @@
 #include "SequencerSettings.h"
 #include "SequencerLog.h"
 #include "SequencerEdMode.h"
+#include "MovieSceneBindingProxy.h"
 #include "MovieSceneSequence.h"
 #include "MovieSceneFolder.h"
 #include "PropertyEditorModule.h"
@@ -83,27 +83,20 @@
 #include "EditorSupportDelegates.h"
 #include "MVVM/Views/SOutlinerView.h"
 #include "ScopedTransaction.h"
-#include "Tracks/MovieScene3DTransformTrack.h"
-#include "Tracks/MovieScene3DAttachTrack.h"
-#include "Tracks/MovieSceneCameraAnimTrack.h"
-#include "Tracks/MovieSceneCameraShakeTrack.h"
-#include "Tracks/MovieSceneCameraCutTrack.h"
 #include "ISequencerTrackEditor.h"
 #include "MovieSceneToolHelpers.h"
-#include "Sections/MovieScene3DAttachSection.h"
 #include "Sections/MovieSceneBoolSection.h"
-#include "Sections/MovieSceneCameraCutSection.h"
 #include "Sections/MovieScene3DTransformSection.h"
-#include "Sections/MovieSceneSpawnSection.h"
 #include "Sections/MovieSceneSubSection.h"
+#include "Tracks/MovieScene3DTransformTrack.h"
+#include "Tracks/MovieSceneCinematicShotTrack.h"
+#include "Tracks/MovieScenePropertyTrack.h"
 #include "Tracks/MovieSceneSubTrack.h"
 #include "Sections/MovieSceneCinematicShotSection.h"
 #include "MovieSceneObjectBindingIDCustomization.h"
 #include "MovieSceneObjectBindingID.h"
 #include "ISettingsModule.h"
 #include "Framework/Commands/GenericCommands.h"
-#include "Tracks/MovieSceneSpawnTrack.h"
-#include "Tracks/MovieScenePropertyTrack.h"
 #include "Framework/Notifications/NotificationManager.h"
 #include "Widgets/Notifications/SNotificationList.h"
 #include "Widgets/Input/STextEntryPopup.h"
@@ -117,11 +110,9 @@
 #include "ContentBrowserModule.h"
 #include "PackageTools.h"
 #include "SequencerUtilities.h"
-#include "Tracks/MovieSceneCinematicShotTrack.h"
 #include "CineCameraActor.h"
 #include "CameraRig_Rail.h"
 #include "CameraRig_Crane.h"
-#include "Components/SplineComponent.h"
 #include "DesktopPlatformModule.h"
 #include "Factories.h"
 #include "FbxExporter.h"
@@ -133,8 +124,6 @@
 #include "IVREditorModule.h"
 #include "HAL/PlatformApplicationMisc.h"
 #include "SequencerKeyActor.h"
-#include "MovieSceneCopyableBinding.h"
-#include "MovieSceneCopyableTrack.h"
 #include "ISequencerChannelInterface.h"
 #include "IMovieRendererInterface.h"
 #include "SequencerKeyCollection.h"
@@ -1151,70 +1140,9 @@ void FSequencer::SuppressAutoEvaluation(UMovieSceneSequence* Sequence, const FGu
 
 FGuid FSequencer::CreateBinding(UObject& InObject, const FString& InName)
 {
-	const FScopedTransaction Transaction(LOCTEXT("CreateBinding", "Create New Binding"));
-
-	UMovieSceneSequence* OwnerSequence = GetFocusedMovieSceneSequence();
-	UMovieScene* OwnerMovieScene = OwnerSequence->GetMovieScene();
-
-	OwnerSequence->Modify();
-	OwnerMovieScene->Modify();
-		
-	const FGuid PossessableGuid = OwnerMovieScene->AddPossessable(InName, InObject.GetClass());
-
-	// Attempt to use the parent as a context if necessary
-	UObject* ParentObject = OwnerSequence->GetParentObject(&InObject);
-	UObject* BindingContext = GetPlaybackContext(); //UWorld
-
-	AActor* ParentActorAdded = nullptr;
-	FGuid ParentGuid;
-
-	if (ParentObject)
-	{
-		// Ensure we have possessed the outer object, if necessary
-		ParentGuid = GetHandleToObject(ParentObject, false);
-		if (!ParentGuid.IsValid())
-		{
-			ParentGuid = GetHandleToObject(ParentObject);
-			ParentActorAdded = Cast<AActor>(ParentObject);
-		}
-
-		if (OwnerSequence->AreParentContextsSignificant())
-		{
-			BindingContext = ParentObject;
-		}
-
-		// Set up parent/child guids for possessables within spawnables
-		if (ParentGuid.IsValid())
-		{
-			FMovieScenePossessable* ChildPossessable = OwnerMovieScene->FindPossessable(PossessableGuid);
-			if (ensure(ChildPossessable))
-			{
-				ChildPossessable->SetParent(ParentGuid, OwnerMovieScene);
-			}
-
-			FMovieSceneSpawnable* ParentSpawnable = OwnerMovieScene->FindSpawnable(ParentGuid);
-			if (ParentSpawnable)
-			{
-				ParentSpawnable->AddChildPossessable(PossessableGuid);
-			}
-		}
-	}
-
-	if (!OwnerMovieScene->FindPossessable(PossessableGuid)->BindSpawnableObject(GetFocusedTemplateID(), &InObject, this))
-	{
-		OwnerSequence->BindPossessableObject(PossessableGuid, InObject, BindingContext);
-	}
-		
-	// Broadcast if a parent actor was added as a result of adding this object
-	if (ParentActorAdded && ParentGuid.IsValid())
-	{
-		OnActorAddedToSequencerEvent.Broadcast(ParentActorAdded, ParentGuid);
-	}
-
-	return PossessableGuid;
+    return FSequencerUtilities::CreateBinding(AsShared(), InObject, InName);
 }
 
-
 UObject* FSequencer::GetPlaybackContext() const
 {
 	return CachedPlaybackContext.Get();
@@ -3859,7 +3787,7 @@ FGuid FSequencer::GetHandleToObject( UObject* Object, bool bCreateHandleIfMissin
 	{
 		AActor* PossessedActor = Cast<AActor>(Object);
 
-		ObjectGuid = CreateBinding(*Object, PossessedActor != nullptr ? PossessedActor->GetActorLabel() : Object->GetName());
+		ObjectGuid = FSequencerUtilities::CreateBinding(AsShared(), *Object, PossessedActor != nullptr ? PossessedActor->GetActorLabel() : Object->GetName());
 
 		AActor* OwningActor = PossessedActor;
 		FGuid OwningObjectGuid = ObjectGuid;
@@ -5441,82 +5369,21 @@ bool FSequencer::IsReadOnly() const
 	return bReadOnly || (GetFocusedMovieSceneSequence() && GetFocusedMovieSceneSequence()->GetMovieScene()->IsReadOnly());
 }
 
-FGuid FSequencer::AddSpawnable(UObject& Object, UActorFactory* ActorFactory)
-{
-	UMovieSceneSequence* Sequence = GetFocusedMovieSceneSequence();
-	if (!Sequence->AllowsSpawnableObjects())
-	{
-		return FGuid();
-	}
-
-	// Grab the MovieScene that is currently focused.  We'll add our Blueprint as an inner of the
-	// MovieScene asset.
-	UMovieScene* OwnerMovieScene = Sequence->GetMovieScene();
-
-	TValueOrError<FNewSpawnable, FText> Result = SpawnRegister->CreateNewSpawnableType(Object, *OwnerMovieScene, ActorFactory);
-	if (!Result.IsValid())
-	{
-		FNotificationInfo Info(Result.GetError());
-		Info.ExpireDuration = 3.0f;
-		FSlateNotificationManager::Get().AddNotification(Info);
-		return FGuid();
-	}
-
-	FNewSpawnable& NewSpawnable = Result.GetValue();
-	
-	NewSpawnable.Name = MovieSceneHelpers::MakeUniqueSpawnableName(OwnerMovieScene, NewSpawnable.Name);
-	
-	FGuid NewGuid = OwnerMovieScene->AddSpawnable(NewSpawnable.Name, *NewSpawnable.ObjectTemplate);
-
-	ForceEvaluate();
-
-	return NewGuid;
-}
-
-FGuid FSequencer::MakeNewSpawnable( UObject& Object, UActorFactory* ActorFactory, bool bSetupDefaults )
+FGuid FSequencer::MakeNewSpawnable(UObject& Object, UActorFactory* ActorFactory, bool bSetupDefaults)
 {
-	UMovieSceneSequence* Sequence = GetFocusedMovieSceneSequence();
-	UMovieScene* MovieScene = Sequence->GetMovieScene();
-
-	if (MovieScene->IsReadOnly())
-	{
-		FSequencerUtilities::ShowReadOnlyError();
-		return FGuid();
-	}
-
-	// @todo sequencer: Undo doesn't seem to be working at all
-	const FScopedTransaction Transaction( LOCTEXT("UndoAddingObject", "Add Object to MovieScene") );
-
-	FGuid NewGuid = AddSpawnable(Object, ActorFactory);
-	if (!NewGuid.IsValid())
-	{
-		return FGuid();
-	}
+	const FScopedTransaction Transaction(LOCTEXT("UndoAddingObject", "Add Object to MovieScene"));
 
 	TArray<UMovieSceneFolder*> SelectedParentFolders;
 	FString NewNodePath;
 	CalculateSelectedFolderAndPath(SelectedParentFolders, NewNodePath);
 
+	FGuid NewGuid = FSequencerUtilities::MakeNewSpawnable(AsShared(), Object, ActorFactory, bSetupDefaults);
+	
 	if (SelectedParentFolders.Num() > 0)
 	{
 		SelectedParentFolders[0]->AddChildObjectBinding(NewGuid);
 	}
 
-	FMovieSceneSpawnable* Spawnable = GetFocusedMovieSceneSequence()->GetMovieScene()->FindSpawnable(NewGuid);
-	if (!Spawnable)
-	{
-		return FGuid();
-	}
-
-	// Spawn the object so we can position it correctly, it's going to get spawned anyway since things default to spawned.
-	UObject* SpawnedObject = SpawnRegister->SpawnObject(NewGuid, *MovieScene, ActiveTemplateIDs.Top(), *this);
-
-	if (bSetupDefaults)
-	{
-		FTransformData TransformData;
-		SpawnRegister->SetupDefaultsForSpawnable(SpawnedObject, Spawnable->GetGuid(), TransformData, AsShared(), Settings);
-	}
-
 	return NewGuid;
 }
 
@@ -5683,21 +5550,14 @@ void FSequencer::OnNewActorsDropped(const TArray<UObject*>& DroppedObjects, cons
 		for ( AActor* Actor : DroppedActors )
 		{
 			AActor* NewActor = Actor;
-			bool bCreateAndAttachCamera = false;
-			if (NewActor->GetClass() == ACameraRig_Rail::StaticClass() ||
-				NewActor->GetClass() == ACameraRig_Crane::StaticClass())
-			{
-				bCreateAndAttachCamera = true;
-			}
-
-			FGuid PossessableGuid = CreateBinding(*NewActor, NewActor->GetActorLabel());
+			FGuid PossessableGuid = FSequencerUtilities::CreateBinding(AsShared(), *NewActor, NewActor->GetActorLabel());
 			FGuid NewGuid = PossessableGuid;
 
 			OnActorAddedToSequencerEvent.Broadcast(NewActor, PossessableGuid);
 
 			if (bAddSpawnable)
 			{
-				TArray< FMovieSceneSpawnable*> Spawnables = ConvertToSpawnableInternal(PossessableGuid);
+				TArray<FMovieSceneSpawnable*> Spawnables = FSequencerUtilities::ConvertToSpawnable(AsShared(), PossessableGuid);
 				if (Spawnables.Num() > 0)
 				{
 					for (TWeakObjectPtr<> WeakObject : FindBoundObjects(Spawnables[0]->GetGuid(), ActiveTemplateIDs.Top()))
@@ -5713,97 +5573,11 @@ void FSequencer::OnNewActorsDropped(const TArray<UObject*>& DroppedObjects, cons
 				}
 			}
 
-			if (bCreateAndAttachCamera)
+			if (NewActor->GetClass() == ACameraRig_Rail::StaticClass() ||
+				NewActor->GetClass() == ACameraRig_Crane::StaticClass())
 			{
-				ACameraRig_Rail* RailActor = nullptr;
-				if (Actor->GetClass() == ACameraRig_Rail::StaticClass())
-				{
-					RailActor = Cast<ACameraRig_Rail>(NewActor);
-				}
-
-				// Create a cine camera actor
-				UWorld* PlaybackContext = Cast<UWorld>(GetPlaybackContext());
-				ACineCameraActor* NewCamera = PlaybackContext->SpawnActor<ACineCameraActor>();
-				FGuid NewCameraGuid = CreateBinding(*NewCamera, NewCamera->GetActorLabel());
-
-				if (RailActor)
-				{
-					NewCamera->SetActorRotation(FRotator(0.f, -90.f, 0.f));
-				}
-
-				OnActorAddedToSequencerEvent.Broadcast(NewCamera, NewCameraGuid);
-
-				if (bAddSpawnable)
-				{
-					FString NewCameraName = MovieSceneHelpers::MakeUniqueSpawnableName(OwnerMovieScene, FName::NameToDisplayString(ACineCameraActor::StaticClass()->GetFName().ToString(), false));
-										
-					FMovieSceneSpawnable* Spawnable = ConvertToSpawnableInternal(NewCameraGuid)[0];
-					Spawnable->SetName(NewCameraName);
-
-					for (TWeakObjectPtr<> WeakObject : FindBoundObjects(Spawnable->GetGuid(), ActiveTemplateIDs.Top()))
-					{
-						NewCamera = Cast<ACineCameraActor>(WeakObject.Get());
-						if (NewCamera)
-						{
-							break;
-						}
-					}
-
-					NewCamera->SetActorLabel(NewCameraName, false);
-
-					NewCameraGuid = Spawnable->GetGuid();
-
-					// Create an attach track
-					UMovieScene3DAttachTrack* AttachTrack = Cast<UMovieScene3DAttachTrack>(OwnerMovieScene->AddTrack(UMovieScene3DAttachTrack::StaticClass(), NewCameraGuid));
-
-					FMovieSceneObjectBindingID AttachBindingID = UE::MovieScene::FRelativeObjectBindingID(NewGuid);
-					FFrameNumber StartTime = UE::MovieScene::DiscreteInclusiveLower(GetPlaybackRange());
-					FFrameNumber Duration  = UE::MovieScene::DiscreteSize(GetPlaybackRange());
-
-					AttachTrack->AddConstraint(StartTime, Duration.Value, NAME_None, NAME_None, AttachBindingID);
-				}
-				else
-				{
-					// Parent it
-					NewCamera->AttachToActor(NewActor, FAttachmentTransformRules::KeepRelativeTransform);
-				}
-
-				if (RailActor)
-				{
-					// Extend the rail a bit
-					if (RailActor->GetRailSplineComponent()->GetNumberOfSplinePoints() == 2)
-					{
-						FVector SplinePoint1 = RailActor->GetRailSplineComponent()->GetLocationAtSplinePoint(0, ESplineCoordinateSpace::Local);
-						FVector SplinePoint2 = RailActor->GetRailSplineComponent()->GetLocationAtSplinePoint(1, ESplineCoordinateSpace::Local);
-						FVector SplineDirection = SplinePoint2 - SplinePoint1;
-						SplineDirection.Normalize();
-
-						float DefaultRailDistance = 650.f;
-						SplinePoint2 = SplinePoint1 + SplineDirection* DefaultRailDistance;
-						RailActor->GetRailSplineComponent()->SetLocationAtSplinePoint(1, SplinePoint2, ESplineCoordinateSpace::Local);
-						RailActor->GetRailSplineComponent()->bSplineHasBeenEdited = true;
-					}
-
-					// Create a track for the CurrentPositionOnRail
-					FPropertyPath PropertyPath;
-					PropertyPath.AddProperty(FPropertyInfo(RailActor->GetClass()->FindPropertyByName(TEXT("CurrentPositionOnRail"))));
-
-					FKeyPropertyParams KeyPropertyParams(TArrayBuilder<UObject*>().Add(RailActor), PropertyPath, ESequencerKeyMode::ManualKeyForced);
-
-					FFrameTime OriginalTime = GetLocalTime().Time;
-
-					SetLocalTimeDirectly(UE::MovieScene::DiscreteInclusiveLower(GetPlaybackRange()));
-					RailActor->CurrentPositionOnRail = 0.f;
-					KeyProperty(KeyPropertyParams);
-
-					SetLocalTimeDirectly(UE::MovieScene::DiscreteExclusiveUpper(GetPlaybackRange())-1);
-					RailActor->CurrentPositionOnRail = 1.f;
-					KeyProperty(KeyPropertyParams);
-
-					SetLocalTimeDirectly(OriginalTime);
-				}
-
-				NewCameraAdded(NewCamera, NewCameraGuid);
+				ACineCameraActor* OutActor;
+				FSequencerUtilities::CreateCameraWithRig(AsShared(), NewActor, bAddSpawnable, OutActor);
 			}
 		}
 
@@ -6066,40 +5840,11 @@ TArray<FGuid> FSequencer::AddActors(const TArray<TWeakObjectPtr<AActor> >& InAct
 {
 	using namespace UE::Sequencer;
 
-	TArray<FGuid> PossessableGuids;
-
-	if (GetFocusedMovieSceneSequence()->GetMovieScene()->IsReadOnly())
-	{
-		FSequencerUtilities::ShowReadOnlyError();
-		return PossessableGuids;
-	}
-
 	const FScopedTransaction Transaction(LOCTEXT("UndoPossessingObject", "Possess Object in Sequencer"));
-	GetFocusedMovieSceneSequence()->Modify();
-
-	bool bPossessableAdded = false;
-	for (TWeakObjectPtr<AActor> WeakActor : InActors)
-	{
-		if (AActor* Actor = WeakActor.Get())
-		{
-			FGuid ExistingGuid = FindObjectId(*Actor, ActiveTemplateIDs.Top());
-			if (!ExistingGuid.IsValid())
-			{
-				FGuid PossessableGuid = CreateBinding(*Actor, Actor->GetActorLabel());
-				PossessableGuids.Add(PossessableGuid);
-
-				if (ACameraActor* CameraActor = Cast<ACameraActor>(Actor))
-				{
-					NewCameraAdded(CameraActor, PossessableGuid);
-				}
 
-				OnActorAddedToSequencerEvent.Broadcast(Actor, PossessableGuid);
-			}
-			bPossessableAdded = true;
-		}
-	}
+	TArray<FGuid> PossessableGuids = FSequencerUtilities::AddActors(AsShared(), InActors);
 
-	if (bPossessableAdded)
+	if (PossessableGuids.Num())
 	{
 		// Check if a folder is selected so we can add the actors to the selected folder.
 		TArray<UMovieSceneFolder*> SelectedParentFolders;
@@ -6174,7 +5919,7 @@ TArray<FGuid> FSequencer::AddActors(const TArray<TWeakObjectPtr<AActor> >& InAct
 			}
 		}
 	}
-		
+
 	RefreshTree();
 
 	SynchronizeSequencerSelectionWithExternalSelection();
@@ -6182,7 +5927,6 @@ TArray<FGuid> FSequencer::AddActors(const TArray<TWeakObjectPtr<AActor> >& InAct
 	return PossessableGuids;
 }
 
-
 void FSequencer::OnSelectedOutlinerNodesChanged()
 {
 	HandleSelectedOutlinerNodesChanged();
@@ -7460,193 +7204,6 @@ void FSequencer::DeleteSelectedItems()
 }
 
 
-void FSequencer::AddActorsToBinding(FGuid InObjectBinding, const TArray<AActor*>& InActors)
-{
-	if (!InActors.Num())
-	{
-		return;
-	}
-
-	UClass* ActorClass = nullptr;
-	int32 NumRuntimeObjects = 0;
-
-	TArrayView<TWeakObjectPtr<>> ObjectsInCurrentSequence = FindObjectsInCurrentSequence(InObjectBinding);
-
-	for (TWeakObjectPtr<> Ptr : ObjectsInCurrentSequence)
-	{
-		if (const AActor* Actor = Cast<AActor>(Ptr.Get()))
-		{
-			ActorClass = Actor->GetClass();
-			++NumRuntimeObjects;
-		}
-	}
-
-	FScopedTransaction AddSelectedToBinding(NSLOCTEXT("Sequencer", "AddSelectedToBinding", "Add Selected to Binding"));
-
-	UMovieSceneSequence* OwnerSequence = GetFocusedMovieSceneSequence();
-	UMovieScene* OwnerMovieScene = OwnerSequence->GetMovieScene();
-
-	OwnerSequence->Modify();
-	OwnerMovieScene->Modify();
-
-	// Bind objects
-	int32 NumObjectsAdded = 0;
-	for (AActor* ActorToAdd : InActors)
-	{
-		if (!ObjectsInCurrentSequence.Contains(ActorToAdd))
-		{
-			if (ActorClass == nullptr || UClass::FindCommonBase(ActorToAdd->GetClass(), ActorClass) != nullptr)
-			{
-				if (ActorClass == nullptr)
-				{
-					ActorClass = ActorToAdd->GetClass();
-				}
-
-				ActorToAdd->Modify();
-				if (!OwnerMovieScene->FindPossessable(InObjectBinding)->BindSpawnableObject(GetFocusedTemplateID(), ActorToAdd, this))
-				{
-					OwnerSequence->BindPossessableObject(InObjectBinding, *ActorToAdd, GetPlaybackContext());
-				}
-				++NumObjectsAdded;
-			}
-			else
-			{
-				const FText NotificationText = FText::Format(LOCTEXT("UnableToAssignObject", "Cannot assign object {0}. Expected class {1}"), FText::FromString(ActorToAdd->GetName()), FText::FromString(ActorClass->GetName()));
-				FNotificationInfo Info(NotificationText);
-				Info.ExpireDuration = 3.f;
-				Info.bUseLargeFont = false;
-				FSlateNotificationManager::Get().AddNotification(Info);
-			}
-		}
-	}
-
-	// Update label
-	if (NumRuntimeObjects + NumObjectsAdded > 0)
-	{
-		FMovieScenePossessable* Possessable = OwnerMovieScene->FindPossessable(InObjectBinding);
-		if (Possessable && ActorClass != nullptr)
-		{
-			if (NumRuntimeObjects + NumObjectsAdded > 1)
-			{
-				FString NewLabel = ActorClass->GetName() + FString::Printf(TEXT(" (%d)"), NumRuntimeObjects + NumObjectsAdded);
-				Possessable->SetName(NewLabel);
-			}
-			else if (NumObjectsAdded > 0 && InActors.Num() > 0)
-			{
-				Possessable->SetName(InActors[0]->GetActorLabel());
-			}
-
-			Possessable->SetPossessedObjectClass(ActorClass);
-		}
-	}
-
-	RestorePreAnimatedState();
-
-	NotifyMovieSceneDataChanged(EMovieSceneDataChangeType::MovieSceneStructureItemsChanged);
-}
-
-void FSequencer::ReplaceBindingWithActors(FGuid InObjectBinding, const TArray<AActor*>& InActors)
-{
-	FScopedTransaction ReplaceBindingWithActors(NSLOCTEXT("Sequencer", "ReplaceBindingWithActors", "Replace Binding with Actors"));
-
-	TArray<AActor*> ExistingActors;
-	for (TWeakObjectPtr<> Ptr : FindObjectsInCurrentSequence(InObjectBinding))
-	{
-		if (AActor* Actor = Cast<AActor>(Ptr.Get()))
-		{
-			if (!InActors.Contains(Actor))
-			{
-				ExistingActors.Add(Actor);
-			}
-		}
-	}
-
-	RemoveActorsFromBinding(InObjectBinding, ExistingActors);
-
-	TArray<AActor*> NewActors;
-	for (AActor* NewActor : InActors)
-	{
-		if (!ExistingActors.Contains(NewActor))
-		{
-			NewActors.Add(NewActor);
-		}
-	}
-
-	AddActorsToBinding(InObjectBinding, NewActors);
-}
-
-void FSequencer::RemoveActorsFromBinding(FGuid InObjectBinding, const TArray<AActor*>& InActors)
-{
-	if (!InActors.Num())
-	{
-		return;
-	}
-
-	UClass* ActorClass = nullptr;
-	int32 NumRuntimeObjects = 0;
-
-	for (TWeakObjectPtr<> Ptr : FindObjectsInCurrentSequence(InObjectBinding))
-	{
-		if (const AActor* Actor = Cast<AActor>(Ptr.Get()))
-		{
-			ActorClass = Actor->GetClass();
-			++NumRuntimeObjects;
-		}
-	}
-
-	FScopedTransaction RemoveSelectedFromBinding(NSLOCTEXT("Sequencer", "RemoveSelectedFromBinding", "Remove Selected from Binding"));
-
-	UMovieSceneSequence* OwnerSequence = GetFocusedMovieSceneSequence();
-	UMovieScene* OwnerMovieScene = OwnerSequence->GetMovieScene();
-
-	TArray<UObject*> ObjectsToRemove;
-	for (AActor* ActorToRemove : InActors)
-	{
-		// Restore state on any components
-		for (UActorComponent* Component : TInlineComponentArray<UActorComponent*>(ActorToRemove))
-		{
-			if (Component)
-			{
-				PreAnimatedState.RestorePreAnimatedState(*Component);
-			}
-		}
-
-		// Restore state on the object itself
-		PreAnimatedState.RestorePreAnimatedState(*ActorToRemove);
-
-		ActorToRemove->Modify();
-
-		ObjectsToRemove.Add(ActorToRemove);
-	}
-	OwnerSequence->Modify();
-	OwnerMovieScene->Modify();
-
-
-	// Unbind objects
-	OwnerSequence->UnbindObjects(InObjectBinding, ObjectsToRemove, GetPlaybackContext());
-
-	// Update label
-	if (NumRuntimeObjects - ObjectsToRemove.Num() > 0)
-	{
-		FMovieScenePossessable* Possessable = OwnerMovieScene->FindPossessable(InObjectBinding);
-		if (Possessable && ActorClass != nullptr)
-		{
-			if (NumRuntimeObjects - ObjectsToRemove.Num() > 1)
-			{
-				FString NewLabel = ActorClass->GetName() + FString::Printf(TEXT(" (%d)"), NumRuntimeObjects - ObjectsToRemove.Num());
-
-				Possessable->SetName(NewLabel);
-			}
-			else if (ObjectsToRemove.Num() > 0 && InActors.Num() > 0)
-			{
-				Possessable->SetName(InActors[0]->GetActorLabel());
-			}
-		}
-	}
-
-	NotifyMovieSceneDataChanged(EMovieSceneDataChangeType::MovieSceneStructureItemsChanged);
-}
-
 void FSequencer::DeleteNode(TSharedRef<FViewModel> NodeToBeDeleted, const bool bKeepState)
 {
 	// If this node is selected, delete all selected nodes
@@ -8250,220 +7807,43 @@ void FSequencer::RemoveSelectedNodesFromFolders()
 	NotifyMovieSceneDataChanged(EMovieSceneDataChangeType::MovieSceneStructureItemsChanged);
 }
 
-void ExportObjectBindingsToText(const TArray<UMovieSceneCopyableBinding*>& ObjectsToExport, FString& ExportedText)
-{
-	// Clear the mark state for saving.
-	UnMarkAllObjects(EObjectMark(OBJECTMARK_TagExp | OBJECTMARK_TagImp));
-
-	FStringOutputDevice Archive;
-	const FExportObjectInnerContext Context;
+TArray<TSharedPtr<FMovieSceneClipboard>> GClipboardStack;
 
-	// Export each of the selected nodes
-	UObject* LastOuter = nullptr;
+void FSequencer::CopySelectedObjects(TArray<TSharedPtr<UE::Sequencer::FObjectBindingModel>>& ObjectNodes, const TArray<UMovieSceneFolder*>& Folders, FString& ExportedText)
+{
+	using namespace UE::Sequencer;
 
-	for (UMovieSceneCopyableBinding* ObjectToExport : ObjectsToExport)
+	// Gather guids for the object nodes and any child object nodes
+	TArray<FMovieSceneBindingProxy> Bindings;
+	for (TSharedPtr<FObjectBindingModel> ObjectNode : ObjectNodes)
 	{
-		// The nodes should all be from the same scope
-		UObject* ThisOuter = ObjectToExport->GetOuter();
-		check((LastOuter == ThisOuter) || (LastOuter == nullptr));
-		LastOuter = ThisOuter;
+		Bindings.Add(FMovieSceneBindingProxy(ObjectNode->GetObjectGuid(), GetFocusedMovieSceneSequence()));
 
-		// We can't use TextExportTransient on USTRUCTS (which our object contains) so we're going to manually null out some references before serializing them. These references are
-		// serialized manually into the archive, as the auto-serialization will only store a reference (to a privately owned object) which creates issues on deserialization. Attempting 
-		// to deserialize these private objects throws a superflous error in the console that makes it look like things went wrong when they're actually OK and expected.
-		TArray<UMovieSceneTrack*> OldTracks = ObjectToExport->Binding.StealTracks(nullptr);
-		UObject* OldSpawnableTemplate = ObjectToExport->Spawnable.GetObjectTemplate();
-		ObjectToExport->Spawnable.SetObjectTemplate(nullptr);
-
-		UExporter::ExportToOutputDevice(&Context, ObjectToExport, nullptr, Archive, TEXT("copy"), 0, PPF_ExportsNotFullyQualified | PPF_Copy | PPF_Delimited, false, ThisOuter);
+		TSet<TSharedRef<FViewModel> > DescendantNodes;
 
-		// Restore the references (as we don't want to modify the original in the event of a copy operation!)
-		ObjectToExport->Binding.SetTracks(MoveTemp(OldTracks), nullptr);
-		ObjectToExport->Spawnable.SetObjectTemplate(OldSpawnableTemplate);
+		SequencerHelpers::GetDescendantNodes(ObjectNode.ToSharedRef(), DescendantNodes);
 
-		// We manually export the object template for the same private-ownership reason as above. Templates need to be re-created anyways as each Spawnable contains its own copy of the template.
-		if (ObjectToExport->SpawnableObjectTemplate)
-		{
-			UExporter::ExportToOutputDevice(&Context, ObjectToExport->SpawnableObjectTemplate, nullptr, Archive, TEXT("copy"), 0, PPF_ExportsNotFullyQualified | PPF_Copy | PPF_Delimited);
-		}
-	}
-
-	ExportedText = Archive;
-}
-
-class FObjectBindingTextFactory : public FCustomizableTextObjectFactory
-{
-public:
-	FObjectBindingTextFactory(FSequencer& InSequencer)
-		: FCustomizableTextObjectFactory(GWarn)
-		, SequencerPtr(&InSequencer)
-	{
-	}
-
-	// FCustomizableTextObjectFactory implementation
-	virtual bool CanCreateClass(UClass* InObjectClass, bool& bOmitSubObjs) const override
-	{
-		if (InObjectClass->IsChildOf<UMovieSceneCopyableBinding>())
-		{
-			return true;
-		}
-
-		return SequencerPtr->GetSpawnRegister().CanSpawnObject(InObjectClass);
-	}
-	
-
-	virtual void ProcessConstructedObject(UObject* NewObject) override
-	{
-		check(NewObject);
-
-		if (NewObject->IsA<UMovieSceneCopyableBinding>())
-		{
-			UMovieSceneCopyableBinding* CopyableBinding = Cast<UMovieSceneCopyableBinding>(NewObject);
-			NewCopyableBindings.Add(CopyableBinding);
-		}
-		else
-		{
-			NewSpawnableObjectTemplates.Add(NewObject);
-		}
-	}
-
-public:
-	TArray<UMovieSceneCopyableBinding*> NewCopyableBindings;
-	TArray<UObject*> NewSpawnableObjectTemplates;
-
-private:
-	FSequencer* SequencerPtr;
-};
-
-
-void FSequencer::ImportObjectBindingsFromText(const FString& TextToImport, /*out*/ TArray<UMovieSceneCopyableBinding*>& ImportedObjects)
-{
-	UPackage* TempPackage = NewObject<UPackage>(nullptr, TEXT("/Engine/Sequencer/Editor/Transient"), RF_Transient);
-	TempPackage->AddToRoot();
-
-	// Turn the text buffer into objects
-	FObjectBindingTextFactory Factory(*this);
-	Factory.ProcessBuffer(TempPackage, RF_Transactional, TextToImport);
-	ImportedObjects = Factory.NewCopyableBindings;
-
-	// We had to explicitly serialize object templates due to them being a reference to a privately owned object. We now deserialize these object template copies
-	// and match them up with their MovieSceneCopyableBinding again.
-	
-	int32 SpawnableObjectTemplateIndex = 0;
-	for (auto ImportedObject : ImportedObjects)
-	{
-		if (ImportedObject->Spawnable.GetGuid().IsValid() && SpawnableObjectTemplateIndex < Factory.NewSpawnableObjectTemplates.Num())
-		{
-			// This Spawnable Object Template is owned by our transient package, so you'll need to change the owner if you want to keep it later.
-			ImportedObject->SpawnableObjectTemplate = Factory.NewSpawnableObjectTemplates[SpawnableObjectTemplateIndex++];
-		}
-	}
-
-	// Remove the temp package from the root now that it has served its purpose
-	TempPackage->RemoveFromRoot();
-}
-
-TArray<TSharedPtr<FMovieSceneClipboard>> GClipboardStack;
-
-void FSequencer::CopySelectedObjects(TArray<TSharedPtr<UE::Sequencer::FObjectBindingModel>>& ObjectNodes, const TArray<UMovieSceneFolder*>& Folders, FString& ExportedText)
-{
-	using namespace UE::Sequencer;
-
-	UMovieScene* MovieScene = GetFocusedMovieSceneSequence()->GetMovieScene();
-
-	// Gather guids for the object nodes and any child object nodes
-	TSet<FGuid> GuidsToCopy;
-	TMap<FGuid, UMovieSceneFolder*> GuidToFolder;
-	for (TSharedPtr<FObjectBindingModel> ObjectNode : ObjectNodes)
-	{
-		GuidsToCopy.Add(ObjectNode->GetObjectGuid());
-
-		TSharedPtr<FFolderModel> FolderNode = ObjectNode->FindAncestorOfType<FFolderModel>();
-		if (FolderNode.IsValid() && Folders.Contains(FolderNode->GetFolder()))
-		{
-			GuidToFolder.Add(ObjectNode->GetObjectGuid(), FolderNode->GetFolder());
-		}
-
-		TSet<TSharedRef<FViewModel> > DescendantNodes;
-
-		SequencerHelpers::GetDescendantNodes(ObjectNode.ToSharedRef(), DescendantNodes);
-
-		for (auto DescendantNode : DescendantNodes)
+		for (auto DescendantNode : DescendantNodes)
 		{
 			if (FObjectBindingModel* DescendantObjectNode = DescendantNode->CastThis<FObjectBindingModel>())
 			{
-				GuidsToCopy.Add(DescendantObjectNode->GetObjectGuid());
-
-				TSharedPtr<FFolderModel> DescendantFolderNode = DescendantObjectNode->FindAncestorOfType<FFolderModel>();
-				if (DescendantFolderNode.IsValid())
-				{
-					GuidToFolder.Add(DescendantObjectNode->GetObjectGuid(), DescendantFolderNode->GetFolder());
-				}
-			}
-		}
-	}
-
-	// Export each of the bindings
-	TArray<UMovieSceneCopyableBinding*> CopyableBindings;
-
-	for (auto ObjectBinding : GuidsToCopy)
-	{
-		UMovieSceneCopyableBinding *CopyableBinding = NewObject<UMovieSceneCopyableBinding>(GetTransientPackage(), UMovieSceneCopyableBinding::StaticClass(), NAME_None, RF_Transient);
-		CopyableBindings.Add(CopyableBinding);
-
-		FMovieScenePossessable* Possessable = MovieScene->FindPossessable(ObjectBinding);
-		if (Possessable)
-		{
-			CopyableBinding->Possessable = *Possessable;
-		}
-		else
-		{
-			FMovieSceneSpawnable* Spawnable = MovieScene->FindSpawnable(ObjectBinding);
-			if (Spawnable)
-			{
-				CopyableBinding->Spawnable = *Spawnable;
-				
-				// We manually serialize the spawnable object template so that it's not a reference to a privately owned object. Spawnables all have unique copies of their template objects anyways.
-				// Object Templates are re-created on paste (based on these templates) with the correct ownership set up.
-				CopyableBinding->SpawnableObjectTemplate = Spawnable->GetObjectTemplate();
-			}
-		}
-
-		const FMovieSceneBinding* Binding = MovieScene->FindBinding(ObjectBinding);
-		if (Binding)
-		{
-			CopyableBinding->Binding = *Binding;
-			for (auto Track : Binding->GetTracks())
-			{
-				// Tracks suffer from the same issues as Spawnable's Object Templates (reference to a privately owned object). We'll manually serialize the tracks to copy them,
-				// and then restore them on paste.
-				UMovieSceneTrack* DuplicatedTrack = Cast<UMovieSceneTrack>(StaticDuplicateObject(Track, CopyableBinding));
-
-				CopyableBinding->Tracks.Add(DuplicatedTrack);
+				Bindings.Add(FMovieSceneBindingProxy(DescendantObjectNode->GetObjectGuid(), GetFocusedMovieSceneSequence()));
 			}
 		}
-
-		if (GuidToFolder.Contains(ObjectBinding))
-		{
-			UMovieSceneFolder::CalculateFolderPath(GuidToFolder[ObjectBinding], Folders, CopyableBinding->FolderPath);
-		}
 	}
-	if (CopyableBindings.Num() > 0)
-	{
-		ExportObjectBindingsToText(CopyableBindings, /*out*/ ExportedText);
 
-		// Make sure to clear the clipboard for the keys
-		GClipboardStack.Empty();
-	}
+	FSequencerUtilities::CopyBindings(Bindings, Folders, ExportedText);
+		
+	// Make sure to clear the clipboard for the keys
+	GClipboardStack.Empty();
 }
 
 void FSequencer::CopySelectedTracks(TArray<TSharedPtr<FViewModel>>& TrackNodes, const TArray<UMovieSceneFolder*>& Folders, FString& ExportedText)
 {
 	using namespace UE::Sequencer;
 
-	UMovieScene* MovieScene = GetFocusedMovieSceneSequence()->GetMovieScene();
-
-	TArray<UObject*> CopyableObjects;
+	TArray<UMovieSceneTrack*> CopyableTracks;
+	TArray<UMovieSceneSection*> CopyableSections;
 	for (TSharedPtr<FViewModel> TrackNode : TrackNodes)
 	{
 		bool bIsParentSelected = false;
@@ -8487,83 +7867,40 @@ void FSequencer::CopySelectedTracks(TArray<TSharedPtr<FViewModel>>& TrackNodes,
 				{
 					if (Section && Section->GetRowIndex() == TrackRowNode->GetRowIndex())
 					{
-						CopyableObjects.Add(Section);
+						CopyableSections.Add(Section);
 					}
 				}
 			}
-			else
-			{
-				UMovieSceneCopyableTrack* CopyableTrack = NewObject<UMovieSceneCopyableTrack>(GetTransientPackage(), UMovieSceneCopyableTrack::StaticClass(), NAME_None, RF_Transient);
-				CopyableObjects.Add(CopyableTrack);
-
-				UMovieSceneTrack* DuplicatedTrack = Cast<UMovieSceneTrack>(StaticDuplicateObject(TrackNode->CastThisChecked<ITrackExtension>()->GetTrack(), CopyableTrack));
-				CopyableTrack->Track = DuplicatedTrack;
-				CopyableTrack->bIsAMasterTrack = MovieScene->IsAMasterTrack(*TrackNode->CastThisChecked<ITrackExtension>()->GetTrack());
-				CopyableTrack->bIsACameraCutTrack = TrackNode->CastThisChecked<ITrackExtension>()->GetTrack()->IsA<UMovieSceneCameraCutTrack>();
-
-				TSharedPtr<FFolderModel> FolderNode = TrackNode->FindAncestorOfType<FFolderModel>();
-				if (FolderNode.IsValid() && Folders.Contains(FolderNode->GetFolder()))
-				{
-					UMovieSceneFolder::CalculateFolderPath(FolderNode->GetFolder(), Folders, CopyableTrack->FolderPath);
-				}
+			else if (UMovieSceneTrack* Track = TrackNode->CastThisChecked<ITrackExtension>()->GetTrack())
+			{	
+				CopyableTracks.Add(Track);
 			}
 		}
 	}
 
-	if (CopyableObjects.Num())
+	if (CopyableSections.Num() + CopyableTracks.Num() > 0)
 	{
-		ExportObjectsToText(CopyableObjects, /*out*/ ExportedText);
-
-		// Make sure to clear the clipboard for the keys
-		GClipboardStack.Empty();
-	}
-}
+		FString SectionsExportedText;
+		FSequencerUtilities::CopySections(CopyableSections, SectionsExportedText);
 
+		FString TracksExportedText;
+		FSequencerUtilities::CopyTracks(CopyableTracks, Folders, TracksExportedText);
 
-void FSequencer::CopySelectedFolders(const TArray<UMovieSceneFolder*>& Folders, FString& ExportedText)
-{
-	if (Folders.Num() > 0)
-	{
-		TArray<UObject*> Objects;
-		for (UMovieSceneFolder* Folder : Folders)
-		{
-			Objects.Add(Folder);
-		}
-
-		ExportObjectsToText(Objects, /*out*/ ExportedText);
+		ExportedText.Empty();
+		ExportedText += SectionsExportedText;
+		ExportedText += TracksExportedText;
 
 		// Make sure to clear the clipboard for the keys
 		GClipboardStack.Empty();
 	}
 }
 
-
-void FSequencer::ExportObjectsToText(const TArray<UObject*>& ObjectsToExport, FString& ExportedText)
+void FSequencer::CopySelectedFolders(const TArray<UMovieSceneFolder*>& Folders, FString& ExportedText)
 {
-	// Clear the mark state for saving.
-	UnMarkAllObjects(EObjectMark(OBJECTMARK_TagExp | OBJECTMARK_TagImp));
-
-	FStringOutputDevice Archive;
-	const FExportObjectInnerContext Context;
-
-	// Export each of the selected nodes
-	UObject* LastOuter = nullptr;
-
-	for (UObject* ObjectToExport : ObjectsToExport)
-	{
-		// The nodes should all be from the same scope
-		UObject* ThisOuter = ObjectToExport->GetOuter();
-		if (LastOuter != nullptr && ThisOuter != LastOuter)
-		{
-			UE_LOG(LogSequencer, Warning, TEXT("Cannot copy objects from different outers. Only copying from %s"), *LastOuter->GetName());
-			continue;
-		}
-		LastOuter = ThisOuter;
-
-		UExporter::ExportToOutputDevice(&Context, ObjectToExport, nullptr, Archive, TEXT("copy"), 0, PPF_ExportsNotFullyQualified | PPF_Copy | PPF_Delimited, false, ThisOuter);
-	}
+	FSequencerUtilities::CopyFolders(Folders, ExportedText);
 
-	ExportedText = Archive;
+	// Make sure to clear the clipboard for the keys
+	GClipboardStack.Empty();
 }
 
 bool FSequencer::DoPaste(bool bClearSelection)
@@ -8604,7 +7941,7 @@ bool FSequencer::DoPaste(bool bClearSelection)
 	TArray<UMovieSceneFolder*> PastedFolders;
 	if (EnumHasAnyFlags(PasteSupport, ESequencerPasteSupport::Folders)) 
 	{
-		bAnythingPasted |= PasteFolders(TextToImport, ParentFolder, PastedFolders, PasteErrors);
+		bAnythingPasted |= FSequencerUtilities::PasteFolders(TextToImport, FMovieScenePasteFoldersParams(GetFocusedMovieSceneSequence(), ParentFolder), PastedFolders, PasteErrors);
 	}
 	if (EnumHasAnyFlags(PasteSupport, ESequencerPasteSupport::ObjectBindings)) 
 	{
@@ -8635,943 +7972,170 @@ bool FSequencer::DoPaste(bool bClearSelection)
 	return bAnythingPasted;
 }
 
-bool FSequencer::PasteFolders(const FString& TextToImport, UMovieSceneFolder* InParentFolder, TArray<UMovieSceneFolder*>& OutFolders, TArray<FNotificationInfo>& PasteErrors)
+bool FSequencer::PasteObjectBindings(const FString& TextToImport, UMovieSceneFolder* InParentFolder, const TArray<UMovieSceneFolder*>& InFolders, TArray<FNotificationInfo>& PasteErrors, bool bClearSelection)
 {
-	TArray<UMovieSceneFolder*> ImportedFolders;
-	ImportFoldersFromText(TextToImport, ImportedFolders);
-
-	if (ImportedFolders.Num() == 0)
-	{
-		return false;
-	}
-
-	UMovieSceneSequence* OwnerSequence = GetFocusedMovieSceneSequence();
-	UMovieScene* MovieScene = GetFocusedMovieSceneSequence()->GetMovieScene();
+	using namespace UE::Sequencer;
 
-	MovieScene->Modify();
+	TArray<TSharedRef<FViewModel>> SelectedNodes;
+	Selection.GetSelectedOutlinerItems(SelectedNodes);
 
-	for (UMovieSceneFolder* CopiedFolder : ImportedFolders)
+	TArray<FMovieSceneBindingProxy> TargetBindings;
+	if (!bClearSelection)
 	{
-		CopiedFolder->Rename(nullptr, MovieScene);
-
-		OutFolders.Add(CopiedFolder);
-
-		// Clear the folder contents, those relationships will be made when the tracks are pasted
-		CopiedFolder->ClearChildMasterTracks();
-		CopiedFolder->ClearChildObjectBindings();
-
-		bool bHasParent = false;
-		for (UMovieSceneFolder* ImportedParentFolder : ImportedFolders)
-		{
-			if (ImportedParentFolder != CopiedFolder)
-			{
-				if (ImportedParentFolder->GetChildFolders().Contains(CopiedFolder))
-				{
-					bHasParent = true;
-					break;
-				}
-			}
-		}
-
-		if (!bHasParent)
+		for (TSharedRef<FViewModel> Node : SelectedNodes)
 		{
-			if (InParentFolder)
-			{
-				InParentFolder->AddChildFolder(CopiedFolder);
-			}
-			else
+			if (FObjectBindingModel* ObjectNode = Node->CastThis<FObjectBindingModel>())
 			{
-				MovieScene->AddRootFolder(CopiedFolder);
+				TargetBindings.Add(FMovieSceneBindingProxy(ObjectNode->GetObjectGuid(), GetFocusedMovieSceneSequence()));
 			}
 		}
 	}
 
-	NotifyMovieSceneDataChanged(EMovieSceneDataChangeType::MovieSceneStructureItemsChanged);
-
-	return true;
+	TArray<FMovieSceneBindingProxy> OutBindings;
+	return FSequencerUtilities::PasteBindings(TextToImport, AsShared(), FMovieScenePasteBindingsParams(TargetBindings, InParentFolder, InFolders), OutBindings, PasteErrors);
 }
 
-bool FSequencer::PasteObjectBindings(const FString& TextToImport, UMovieSceneFolder* InParentFolder, const TArray<UMovieSceneFolder*>& InFolders, TArray<FNotificationInfo>& PasteErrors, bool bClearSelection)
+bool FSequencer::PasteTracks(const FString& TextToImport, UMovieSceneFolder* InParentFolder, const TArray<UMovieSceneFolder*>& InFolders, TArray<FNotificationInfo>& PasteErrors, bool bClearSelection)
 {
 	using namespace UE::Sequencer;
 
-	UWorld* World = Cast<UWorld>(GetPlaybackContext());
-
-	UMovieSceneSequence* OwnerSequence = GetFocusedMovieSceneSequence();
-	UObject* BindingContext = GetPlaybackContext();
-
-	UMovieScene* MovieScene = GetFocusedMovieSceneSequence()->GetMovieScene();
-	TMap<FGuid, FGuid> OldToNewGuidMap;
-	TArray<FGuid> PossessableGuids;
-	TArray<FGuid> SpawnableGuids;
-	TMap<FGuid, UMovieSceneFolder*> GuidToFolderMap;
-
-	TArray<FMovieSceneBinding> BindingsPasted;
-
 	TArray<TSharedRef<FViewModel>> SelectedNodes;
 	Selection.GetSelectedOutlinerItems(SelectedNodes);
 
-	TArray<FGuid> SelectedParentGuids;
+	TArray<FMovieSceneBindingProxy> TargetBindings;
 	if (!bClearSelection)
 	{
 		for (TSharedRef<FViewModel> Node : SelectedNodes)
 		{
-			if (FObjectBindingModel* ObjectNode = Node->CastThis<FObjectBindingModel>())
+			if (TSharedPtr<FObjectBindingModel> ObjectNode = Node->CastThisShared<FObjectBindingModel>())
 			{
-				SelectedParentGuids.Add(ObjectNode->GetObjectGuid());
+				TargetBindings.Add(FMovieSceneBindingProxy(ObjectNode->GetObjectGuid(), GetFocusedMovieSceneSequence()));
 			}
 		}
 	}
 
-	int32 NumTargets = 1;
-	if (SelectedParentGuids.Num() > 1)
-	{
-		NumTargets = SelectedParentGuids.Num();
-	}
-
-	for (int32 TargetIndex = 0; TargetIndex < NumTargets; ++TargetIndex)
-	{
-		TArray<UMovieSceneCopyableBinding*> ImportedBindings;
-		ImportObjectBindingsFromText(TextToImport, ImportedBindings);
-
-		if (ImportedBindings.Num() == 0)
-		{
-			return false;
-		}
-
-		for (UMovieSceneCopyableBinding* CopyableBinding : ImportedBindings)
-		{
-			// Clear transient flags on the imported tracks
-			for (UMovieSceneTrack* CopiedTrack : CopyableBinding->Tracks)
-			{
-				CopiedTrack->ClearFlags(RF_Transient);
-				TArray<UObject*> Subobjects;
-				GetObjectsWithOuter(CopiedTrack, Subobjects);
-				for (UObject* Subobject : Subobjects)
-				{
-					Subobject->ClearFlags(RF_Transient);
-				}
-			}
-
-			UMovieSceneFolder* ParentFolder = InParentFolder;
-
-			if (CopyableBinding->FolderPath.Num() > 0)
-			{
-				ParentFolder = UMovieSceneFolder::GetFolderWithPath(CopyableBinding->FolderPath, InFolders, ParentFolder ? ParentFolder->GetChildFolders() : MovieScene->GetRootFolders());
-			}
-
-			if (CopyableBinding->Possessable.GetGuid().IsValid())
-			{
-				FGuid NewGuid = FGuid::NewGuid();
-
-				FMovieSceneBinding NewBinding(NewGuid, CopyableBinding->Binding.GetName(), CopyableBinding->Tracks);
-
-				FMovieScenePossessable NewPossessable = CopyableBinding->Possessable;
-				NewPossessable.SetGuid(NewGuid);
-
-				MovieScene->AddPossessable(NewPossessable, NewBinding);
-
-				OldToNewGuidMap.Add(CopyableBinding->Possessable.GetGuid(), NewGuid);
-
-				BindingsPasted.Add(NewBinding);
-
-				PossessableGuids.Add(NewGuid);
-
-				if (ParentFolder)
-				{
-					GuidToFolderMap.Add(NewGuid, ParentFolder);
-				}
-
-				if (FMovieScenePossessable* Possessable = MovieScene->FindPossessable(NewGuid))
-				{
-					if (TargetIndex < SelectedParentGuids.Num())
-					{
-						Possessable->SetParent(SelectedParentGuids[TargetIndex], MovieScene);
-					}
-				}
-
-				TArray<AActor*> ActorsToDuplicate;
-				for (auto RuntimeObject : FindBoundObjects(CopyableBinding->Possessable.GetGuid(), ActiveTemplateIDs.Top()))
-				{
-					AActor* Actor = Cast<AActor>(RuntimeObject.Get());
-					if (Actor)
-					{
-						ActorsToDuplicate.Add(Actor);
-					}
-				}
-
-				TArray<AActor*> DuplicatedActors;
-				if (ActorsToDuplicate.Num() != 0)
-				{
-					GEditor->SelectNone(false, true);
-					for (AActor* ActorToDuplicate : ActorsToDuplicate)
-					{
-						GEditor->SelectActor(ActorToDuplicate, true, false, false);
-					}
+	TArray<UMovieSceneTrack*> OutTracks;
+	return FSequencerUtilities::PasteTracks(TextToImport, FMovieScenePasteTracksParams(GetFocusedMovieSceneSequence(), TargetBindings, InParentFolder, InFolders), OutTracks, PasteErrors);
+}
 
-					// Duplicate the bound actors
-					GEditor->edactDuplicateSelected(World->GetCurrentLevel(), false);
+bool FSequencer::PasteSections(const FString& TextToImport, TArray<FNotificationInfo>& PasteErrors)
+{
+	using namespace UE::Sequencer;
 
-					// Duplicating the bound actor through GEditor, edits the copy/paste clipboard. This is not desired from the user's 
-					// point of view since the user didn't explicitly invoke the copy operation. Instead, restore the copied contents
-					// of the clipboard after duplicating the actor
-					FPlatformApplicationMisc::ClipboardCopy(*TextToImport);
+	FScopedTransaction Transaction(FGenericCommands::Get().Paste->GetDescription());
 
-					USelection* ActorSelection = GEditor->GetSelectedActors();
-					TArray<TWeakObjectPtr<AActor> > SelectedActors;
-					for (FSelectionIterator Iter(*ActorSelection); Iter; ++Iter)
-					{
-						AActor* Actor = Cast<AActor>(*Iter);
-						if (Actor)
-						{
-							DuplicatedActors.Add(Actor);
-						}
-					}
+	TArray<TSharedRef<FViewModel>> SelectedNodes;
+	Selection.GetSelectedOutlinerItems(SelectedNodes);
 
-					// Bind the duplicated actors
-					if (DuplicatedActors.Num())
-					{
-						ReplaceBindingWithActors(NewGuid, DuplicatedActors);
-					}
-				}
-			}
-			else if (CopyableBinding->Spawnable.GetGuid().IsValid())
+	if (SelectedNodes.Num() == 0)
+	{
+		for (const TViewModelPtr<IOutlinerExtension>& RootNode : NodeTree->GetRootNodes())
+		{
+			TSet<TWeakObjectPtr<UMovieSceneSection>> Sections;
+			SequencerHelpers::GetAllSections(RootNode.AsModel(), Sections);
+			for (TWeakObjectPtr<UMovieSceneSection> Section : Sections)
 			{
-				// We need to let the sequence create the spawnable so that it has everything set up properly internally.
-				// This is required to get spawnables with the correct references to object templates, object templates with
-				// correct owners, etc. However, making a new spawnable also creates the binding for us - this is a problem
-				// because we need to use our binding (which has tracks associated with it). To solve this, we let it create
-				// an object template based off of our (transient package owned) template, then find the newly created binding
-				// and update it.
-				FGuid NewGuid = MakeNewSpawnable(*CopyableBinding->SpawnableObjectTemplate, nullptr, false);
-				FMovieSceneBinding NewBinding(NewGuid, CopyableBinding->Binding.GetName(), CopyableBinding->Tracks);
-				FMovieSceneSpawnable* Spawnable = MovieScene->FindSpawnable(NewGuid);
-
-				// Copy the name of the original spawnable too.
-				Spawnable->SetName(CopyableBinding->Spawnable.GetName());
-
-				// Clear the transient flags on the copyable binding before assigning to the new spawnable
-				for (auto Track : NewBinding.GetTracks())
-				{
-					Track->ClearFlags(RF_Transient);
-					for (auto Section : Track->GetAllSections())
-					{
-						Section->ClearFlags(RF_Transient);
-					}
-				}
-
-				// Replace the auto-generated binding with our deserialized bindings (which has our tracks)
-				MovieScene->ReplaceBinding(NewGuid, NewBinding);
-
-				OldToNewGuidMap.Add(CopyableBinding->Spawnable.GetGuid(), NewGuid);
-
-				BindingsPasted.Add(NewBinding);
-
-				SpawnableGuids.Add(NewGuid);
-
-				if (ParentFolder)
+				if (Selection.GetSelectedSections().Contains(Section.Get()))
 				{
-					GuidToFolderMap.Add(NewGuid, ParentFolder);
+					SelectedNodes.Add(RootNode.AsModel().ToSharedRef());
+					break;
 				}
-			}
-		}
-	}
-
-	// Fix up parent guids
-	for (auto PossessableGuid : PossessableGuids)
-	{
-		FMovieScenePossessable* Possessable = MovieScene->FindPossessable(PossessableGuid);
-		if (Possessable && OldToNewGuidMap.Contains(Possessable->GetParent()) && PossessableGuid != OldToNewGuidMap[Possessable->GetParent()])
-		{
-			Possessable->SetParent(OldToNewGuidMap[Possessable->GetParent()], MovieScene);
-		}
-	}
-
-	// Fix possessable actor bindings
-	for (int32 PossessableGuidIndex = 0; PossessableGuidIndex < PossessableGuids.Num(); ++PossessableGuidIndex)
-	{
-		FMovieScenePossessable* Possessable = MovieScene->FindPossessable(PossessableGuids[PossessableGuidIndex]);
-		UWorld* PlaybackContext = Cast<UWorld>(GetPlaybackContext());
-		if (Possessable && PlaybackContext)
-		{
-			for (TActorIterator<AActor> ActorItr(PlaybackContext); ActorItr; ++ActorItr)
-			{
-				AActor *Actor = *ActorItr;
-				if (Actor && Actor->GetActorLabel() == *Possessable->GetName())
-				{
-					FGuid ExistingGuid = FindObjectId(*Actor, ActiveTemplateIDs.Top());
-
-					if (!ExistingGuid.IsValid())
-					{
-						FGuid NewGuid = FSequencerUtilities::DoAssignActor(this, Actor, Possessable->GetGuid());
-
-						// If assigning produces a new guid, update the possessable guids and the bindings pasted data
-						if (NewGuid.IsValid())
-						{
-							for (auto BindingPasted : BindingsPasted)
-							{
-								if (BindingPasted.GetObjectGuid() == PossessableGuids[PossessableGuidIndex])
-								{
-									BindingPasted.SetObjectGuid(NewGuid);
-								}
-							}
-
-							if (GuidToFolderMap.Contains(PossessableGuids[PossessableGuidIndex]))
-							{
-								GuidToFolderMap.Add(NewGuid, GuidToFolderMap[PossessableGuids[PossessableGuidIndex]]);
-								GuidToFolderMap.Remove(PossessableGuids[PossessableGuidIndex]);
-							}
-
-							PossessableGuids[PossessableGuidIndex] = NewGuid;
-						}
-					}
-				}
-			}
-		}
-	}
-
-	// Set up folders
-	for (auto PossessableGuid : PossessableGuids)
-	{
-		FMovieScenePossessable* Possessable = MovieScene->FindPossessable(PossessableGuid);
-		if (Possessable && !Possessable->GetParent().IsValid())
-		{
-			if (GuidToFolderMap.Contains(PossessableGuid))
-			{
-				GuidToFolderMap[PossessableGuid]->AddChildObjectBinding(PossessableGuid);
-			}
-		}
-	}
-	for (auto SpawnableGuid : SpawnableGuids)
-	{
-		FMovieSceneSpawnable* Spawnable = MovieScene->FindSpawnable(SpawnableGuid);
-		if (Spawnable)
-		{
-			if (GuidToFolderMap.Contains(SpawnableGuid))
-			{
-				GuidToFolderMap[SpawnableGuid]->AddChildObjectBinding(SpawnableGuid);
-			}
-		}
-	}
-
-	OnMovieSceneBindingsPastedDelegate.Broadcast(BindingsPasted);
-
-	// Refresh all immediately so that spawned actors will be generated immediately
-	ForceEvaluate();
-
-	// Fix possessable component bindings
-	for (auto PossessableGuid : PossessableGuids)
-	{
-		// If a possessable guid does not have any bound objects, they might be 
-		// possessable components for spawnables, so they need to be remapped
-		if (FindBoundObjects(PossessableGuid, ActiveTemplateIDs.Top()).Num() == 0)
-		{
-			FMovieScenePossessable* Possessable = MovieScene->FindPossessable(PossessableGuid);
-			if (Possessable)
-			{
-				FGuid ParentGuid = Possessable->GetParent();
-				for (TWeakObjectPtr<> WeakObject : FindBoundObjects(ParentGuid, ActiveTemplateIDs.Top()))
-				{
-					if (AActor* SpawnedActor = Cast<AActor>(WeakObject.Get()))
-					{
-						for (UActorComponent* Component : SpawnedActor->GetComponents())
-						{
-							if (Component->GetName() == Possessable->GetName())
-							{
-								OwnerSequence->BindPossessableObject(PossessableGuid, *Component, SpawnedActor);
-								break;
-							}
-						}
-					}
-				}
-
-				// If the parent doesn't actually exist, clear it.
-				FMovieScenePossessable* PossessableParent = MovieScene->FindPossessable(ParentGuid);
-				FMovieSceneSpawnable* SpawnableParent = MovieScene->FindSpawnable(ParentGuid);
-				if (!PossessableParent && !SpawnableParent)
-				{
-					Possessable->SetParent(FGuid(), MovieScene);
-				}
-				else if (SpawnableParent)
-				{
-					SpawnableParent->AddChildPossessable(PossessableGuid);
-				}
-			}
-		}
-	}
-
-	// Remap bindings in sections (ie. attach tracks)
-	for (TPair<FGuid, FGuid> GuidPair : OldToNewGuidMap)
-	{
-		FSequencerUtilities::UpdateBindingIDs(this, CompiledDataManager, GuidPair.Key, GuidPair.Value);
-	}
-
-	NotifyMovieSceneDataChanged(EMovieSceneDataChangeType::MovieSceneStructureItemsChanged);
-
-	return true;
-}
-
-bool FSequencer::PasteTracks(const FString& TextToImport, UMovieSceneFolder* InParentFolder, const TArray<UMovieSceneFolder*>& InFolders, TArray<FNotificationInfo>& PasteErrors, bool bClearSelection)
-{
-	using namespace UE::Sequencer;
-
-	TArray<UMovieSceneCopyableTrack*> ImportedTracks;
-	FSequencer::ImportTracksFromText(TextToImport, ImportedTracks);
-
-	if (ImportedTracks.Num() == 0)
-	{
-		return false;
-	}
-
-	UMovieScene* MovieScene = GetFocusedMovieSceneSequence()->GetMovieScene();
-	UObject* BindingContext = GetPlaybackContext();
-
-	TArray<TSharedRef<FViewModel>> SelectedNodes;
-	Selection.GetSelectedOutlinerItems(SelectedNodes);
-
-	TArray<TSharedPtr<FObjectBindingModel>> ObjectNodes;
-
-	if (!bClearSelection)
-	{
-		for (TSharedRef<FViewModel> Node : SelectedNodes)
-		{
-			if (TSharedPtr<FObjectBindingModel> ObjectNode = Node->CastThisShared<FObjectBindingModel>())
-			{
-				ObjectNodes.Add(ObjectNode);
-			}
-		}
-	}
-
-	int32 NumMasterOrCameraCutTracks = 0;
-	int32 NumTracks = 0;
-
-	for (UMovieSceneCopyableTrack* CopyableTrack : ImportedTracks)
-	{
-		if (CopyableTrack->bIsAMasterTrack || CopyableTrack->bIsACameraCutTrack)
-		{
-			++NumMasterOrCameraCutTracks;
-		}
-		else
-		{
-			++NumTracks;
-		}
-	}
-
-	int32 NumTracksPasted = 0;
-	int32 NumMasterOrCameraCutTracksPasted = 0;
-	if (ObjectNodes.Num())
-	{
-		for (TSharedPtr<FObjectBindingModel> ObjectNode : ObjectNodes)
-		{
-			FGuid ObjectGuid = ObjectNode->GetObjectGuid();
-
-			TArray<UMovieSceneCopyableTrack*> NewTracks;
-			FSequencer::ImportTracksFromText(TextToImport, NewTracks);
-
-			for (UMovieSceneCopyableTrack* CopyableTrack : NewTracks)
-			{
-				if (!CopyableTrack->bIsAMasterTrack && !CopyableTrack->bIsACameraCutTrack)
-				{
-					UMovieSceneTrack* NewTrack = CopyableTrack->Track;
-					NewTrack->ClearFlags(RF_Transient);
-					TArray<UObject*> Subobjects;
-					GetObjectsWithOuter(NewTrack, Subobjects);
-					for (UObject* Subobject : Subobjects)
-					{
-						Subobject->ClearFlags(RF_Transient);
-					}
-
-					// Remove tracks with the same name before adding
-					for (const FMovieSceneBinding& Binding : MovieScene->GetBindings())
-					{
-						if (Binding.GetObjectGuid() == ObjectGuid)
-						{
-							// Tracks of the same class should be unique per name.
-							for (UMovieSceneTrack* Track : Binding.GetTracks())
-							{
-								if (Track->GetClass() == NewTrack->GetClass() && Track->GetTrackName() == NewTrack->GetTrackName() && Track->GetDisplayName().IdenticalTo(NewTrack->GetDisplayName()))
-								{
-									// If a track of the same class and name exists, remove it so the new track replaces it
-									MovieScene->RemoveTrack(*Track);
-									break;
-								}
-							}
-						}
-					}
-
-					if (!MovieScene->AddGivenTrack(NewTrack, ObjectGuid))
-					{
-						continue;
-					}
-					else
-					{
-						++NumTracksPasted;
-					}
-				}
-			}
-		}
-	}
-
-	// Add as master track or set camera cut track
-	for (UMovieSceneCopyableTrack* CopyableTrack : ImportedTracks)
-	{
-		if (CopyableTrack->bIsAMasterTrack || CopyableTrack->bIsACameraCutTrack)
-		{
-			UMovieSceneTrack* NewTrack = CopyableTrack->Track;
-			NewTrack->ClearFlags(RF_Transient);
-			TArray<UObject*> Subobjects;
-			GetObjectsWithOuter(NewTrack, Subobjects);
-			for (UObject* Subobject : Subobjects)
-			{
-				Subobject->ClearFlags(RF_Transient);
-			}
-
-			UMovieSceneFolder* ParentFolder = InParentFolder;
-
-			if (CopyableTrack->FolderPath.Num() > 0)
-			{
-				ParentFolder = UMovieSceneFolder::GetFolderWithPath(CopyableTrack->FolderPath, InFolders, ParentFolder ? ParentFolder->GetChildFolders() : MovieScene->GetRootFolders());
-			}
-
-			if (NewTrack->IsA(UMovieSceneCameraCutTrack::StaticClass()))
-			{
-				MovieScene->SetCameraCutTrack(NewTrack);
-				if (ParentFolder != nullptr)
-				{
-					ParentFolder->AddChildMasterTrack(NewTrack);
-				}
-
-				++NumMasterOrCameraCutTracksPasted;
-			}
-			else
-			{
-				if (MovieScene->AddGivenMasterTrack(NewTrack))
-				{
-					if (ParentFolder != nullptr)
-					{
-						ParentFolder->AddChildMasterTrack(NewTrack);
-					}
-				}
-
-				++NumMasterOrCameraCutTracksPasted;
-			}
-		}
-	}
-
-	if (NumMasterOrCameraCutTracksPasted < NumMasterOrCameraCutTracks)
-	{
-		FNotificationInfo Info(LOCTEXT("PasteTracks_NoMasterTracks", "Can't paste track. Master track could not be pasted"));
-		PasteErrors.Add(Info);
-	}
-
-	if (NumTracksPasted < NumTracks)
-	{
-		FNotificationInfo Info(LOCTEXT("PasteSections_NoSelectedObjects", "Can't paste track. No selected objects to paste tracks onto"));
-		PasteErrors.Add(Info);
-	}
-
-	if ((NumMasterOrCameraCutTracksPasted + NumTracksPasted) > 0)
-	{
-		NotifyMovieSceneDataChanged(EMovieSceneDataChangeType::MovieSceneStructureItemsChanged);
-
-		return true;
-	}
-
-	return false;
-}
-
-void GetSupportedTracks(TSharedRef<UE::Sequencer::FViewModel> DisplayNode, const TArray<UMovieSceneSection*>& ImportedSections, TArray<TSharedRef<UE::Sequencer::FViewModel>>& TracksToPasteOnto)
-{
-	using namespace UE::Sequencer;
-
-	ITrackExtension* TrackNode = DisplayNode->CastThis<ITrackExtension>();
-	if (TrackNode == nullptr)
-	{
-		return;
-	}
-
-	if (TracksToPasteOnto.Contains(DisplayNode))
-	{
-		return;
-	}
-
-	UMovieSceneTrack* Track = TrackNode->GetTrack();
-	if (Track)
-	{
-		for (int32 SectionIndex = 0; SectionIndex < ImportedSections.Num(); ++SectionIndex)
-		{
-			UMovieSceneSection* Section = ImportedSections[SectionIndex];
-
-			if (Track->SupportsType(Section->GetClass()))
-			{
-				TracksToPasteOnto.Add(DisplayNode);
-				return;
-			}
-		}
-	}
-}
-
-
-bool FSequencer::PasteSections(const FString& TextToImport, TArray<FNotificationInfo>& PasteErrors)
-{
-	using namespace UE::Sequencer;
-
-	FScopedTransaction Transaction(FGenericCommands::Get().Paste->GetDescription());
-
-	// First import as a track and extract sections to allow for copying track contents to another track
-	TArray<UMovieSceneCopyableTrack*> ImportedTracks;
-	FSequencer::ImportTracksFromText(TextToImport, ImportedTracks);
-
-	TArray<UMovieSceneSection*> ImportedSections;
-	for (UMovieSceneCopyableTrack* CopyableTrack : ImportedTracks)
-	{
-		for (UMovieSceneSection* CopyableSection : CopyableTrack->Track->GetAllSections())
-		{
-			ImportedSections.Add(CopyableSection);
-		}
-	}
-
-	// Otherwise, import as sections
-	if (ImportedSections.Num() == 0)
-	{
-		FSequencer::ImportSectionsFromText(TextToImport, ImportedSections);
-	}
-
-	if (ImportedSections.Num() == 0)
-	{
-		return false;
-	}
-
-	TArray<TSharedRef<FViewModel>> SelectedNodes;
-	Selection.GetSelectedOutlinerItems(SelectedNodes);
-
-	if (SelectedNodes.Num() == 0)
-	{
-		for (const TViewModelPtr<IOutlinerExtension>& RootNode : NodeTree->GetRootNodes())
-		{
-			TSet<TWeakObjectPtr<UMovieSceneSection>> Sections;
-			SequencerHelpers::GetAllSections(RootNode.AsModel(), Sections);
-			for (TWeakObjectPtr<UMovieSceneSection> Section : Sections)
-			{
-				if (Selection.GetSelectedSections().Contains(Section.Get()))
-				{
-					SelectedNodes.Add(RootNode.AsModel().ToSharedRef());
-					break;
-				}
-			}
-		}
-	}
-
-	if (SelectedNodes.Num() == 0)
-	{
-		FNotificationInfo Info(LOCTEXT("PasteSections_NoSelectedTracks", "Can't paste section. No selected tracks to paste sections onto"));
-		PasteErrors.Add(Info);
-		return false;
-	}
-
-	FFrameNumber LocalTime = GetLocalTime().Time.GetFrame();
-
-	TOptional<FFrameNumber> FirstFrame;
-	for (UMovieSceneSection* Section : ImportedSections)
-	{
-		if (Section->HasStartFrame())
-		{
-			if (FirstFrame.IsSet())
-			{
-				if (FirstFrame.GetValue() > Section->GetInclusiveStartFrame())
-				{
-					FirstFrame = Section->GetInclusiveStartFrame();
-				}
-			}
-			else
-			{
-				FirstFrame = Section->GetInclusiveStartFrame();
-			}
-		}
-	}
-
-	// Check if any of the selected nodes supports pasting this type of section
-	TArray<TSharedRef<FViewModel>> TracksToPasteOnto;
-	for (TSharedRef<FViewModel> Node : SelectedNodes)
-	{
-		GetSupportedTracks(Node, ImportedSections, TracksToPasteOnto);
-	}
-
-	// Otherwise, look at all child nodes for supported tracks
-	if (TracksToPasteOnto.Num() == 0)
-	{
-		for (TSharedRef<FViewModel> Node : SelectedNodes)
-		{
-			TSet<TSharedRef<FViewModel> > DescendantNodes;
-			SequencerHelpers::GetDescendantNodes(Node, DescendantNodes);
-
-			for (TSharedRef<FViewModel> DescendantNode : DescendantNodes)
-			{
-				// Don't automatically paste onto subtracks because that would lead to multiple paste destinations
-				if (DescendantNode->IsA<FTrackRowModel>())
-				{
-					continue;
-				}
-				GetSupportedTracks(DescendantNode, ImportedSections, TracksToPasteOnto);
-			}
-		}
-	}
-
-	TArray<UMovieSceneSection*> NewSections;
-	TArray<int32> SectionIndicesImported;
-
-	for (TSharedRef<FViewModel> TrackNode : TracksToPasteOnto)
-	{
-		UMovieSceneTrack* Track = TrackNode->CastThisChecked<ITrackExtension>()->GetTrack();
-		for (int32 SectionIndex = 0; SectionIndex < ImportedSections.Num(); ++SectionIndex)
-		{
-			UMovieSceneSection* Section = ImportedSections[SectionIndex];
-			if (!Track->SupportsType(Section->GetClass()))
-			{
-				continue;
-			}
-
-			SectionIndicesImported.AddUnique(SectionIndex);
-
-			Track->Modify();
-
-			Section->ClearFlags(RF_Transient);
-			Section->Rename(nullptr, Track);
-			Track->AddSection(*Section);
-			if (Section->HasStartFrame())
-			{
-				FFrameNumber NewStartFrame = LocalTime + (Section->GetInclusiveStartFrame() - FirstFrame.GetValue());
-				Section->MoveSection(NewStartFrame - Section->GetInclusiveStartFrame());
-			}
-
-			if (Track->SupportsMultipleRows())
-			{
-				if (FTrackRowModel* TrackRowNode = TrackNode->CastThis<FTrackRowModel>())
-				{
-					Section->SetRowIndex(TrackRowNode->GetRowIndex());
-				}
-			}
-			NewSections.Add(Section);
-		}
-
-		// Fix up rows after sections are in place
-		if (Track->SupportsMultipleRows())
-		{
-			// If any newly created section overlaps the previous sections, put all the sections on the max available row
-			// Find the  this section overlaps any previous sections, 
-			int32 MaxAvailableRowIndex = -1;
-			for (UMovieSceneSection* Section : NewSections)
-			{
-				if (MovieSceneToolHelpers::OverlapsSection(Track, Section, NewSections))
-				{
-					int32 AvailableRowIndex = MovieSceneToolHelpers::FindAvailableRowIndex(Track, Section, NewSections);
-					MaxAvailableRowIndex = FMath::Max(AvailableRowIndex, MaxAvailableRowIndex);
-				}
-			}
-
-			if (MaxAvailableRowIndex != -1)
-			{
-				for (UMovieSceneSection* Section : NewSections)
-				{
-					Section->SetRowIndex(MaxAvailableRowIndex);					
-				}
-			}
-		}
-
-		// Regenerate for pasting onto the next track 
-		ImportedSections.Empty();
-		ImportedTracks.Empty();
-
-		FSequencer::ImportTracksFromText(TextToImport, ImportedTracks);
-
-		for (UMovieSceneCopyableTrack* CopyableTrack : ImportedTracks)
-		{
-			for (UMovieSceneSection* CopyableSection : CopyableTrack->Track->GetAllSections())
-			{
-				ImportedSections.Add(CopyableSection);
-			}
-		}
-
-		if (ImportedSections.Num() == 0)
-		{
-			FSequencer::ImportSectionsFromText(TextToImport, ImportedSections);
-		}
-	}
-
-	for (int32 SectionIndex = 0; SectionIndex < ImportedSections.Num(); ++SectionIndex)
-	{
-		if (!SectionIndicesImported.Contains(SectionIndex))
-		{
-			UE_LOG(LogSequencer, Display, TEXT("Could not paste section of type %s"), *ImportedSections[SectionIndex]->GetClass()->GetName());
-		}
-	}
-
-	if (SectionIndicesImported.Num() == 0)
-	{
-		FNotificationInfo Info(LOCTEXT("PasteSections_NothingPasted", "Can't paste section. No matching section types found."));
-		PasteErrors.Add(Info);
-		return false;
-	}
-
-	NotifyMovieSceneDataChanged(EMovieSceneDataChangeType::MovieSceneStructureItemAdded);
-
-	FSectionModelStorageExtension* SectionModelStorage = ViewModel->GetRootModel()->CastDynamic<FSectionModelStorageExtension>();
-	{
-		UE::MovieScene::FScopedSignedObjectModifyDefer ForceFlush(true);
-		check(SectionModelStorage);
-	}
-
-	EmptySelection();
-
-	Selection.SuspendBroadcast();
-	for (UMovieSceneSection* NewSection : NewSections)
-	{
-		TSharedPtr<FSectionModel> SectionModel = SectionModelStorage->FindModelForSection(NewSection);
-		if (ensure(SectionModel))
-		{
-			Selection.AddToSelection(SectionModel);
-		}
-	}
-	Selection.ResumeBroadcast();
-
-	ThrobSectionSelection();
-
-	return true;
-}
-
-class FTrackObjectTextFactory : public FCustomizableTextObjectFactory
-{
-public:
-	FTrackObjectTextFactory()
-		: FCustomizableTextObjectFactory(GWarn)
-	{
-	}
-
-	// FCustomizableTextObjectFactory implementation
-	virtual bool CanCreateClass(UClass* InObjectClass, bool& bOmitSubObjs) const override
-	{
-		if (InObjectClass->IsChildOf(UMovieSceneCopyableTrack::StaticClass()))
-		{
-			return true;
-		}
-		return false;
-	}
-
-
-	virtual void ProcessConstructedObject(UObject* NewObject) override
-	{
-		check(NewObject);
-
-		NewTracks.Add(Cast<UMovieSceneCopyableTrack>(NewObject));
-	}
-
-public:
-	TArray<UMovieSceneCopyableTrack*> NewTracks;
-};
-
-
-class FSectionObjectTextFactory : public FCustomizableTextObjectFactory
-{
-public:
-	FSectionObjectTextFactory()
-		: FCustomizableTextObjectFactory(GWarn)
-	{
-	}
-
-	// FCustomizableTextObjectFactory implementation
-	virtual bool CanCreateClass(UClass* InObjectClass, bool& bOmitSubObjs) const override
-	{
-		if (InObjectClass->IsChildOf(UMovieSceneSection::StaticClass()))
-		{
-			return true;
-		}
-		return false;
-	}
-
-
-	virtual void ProcessConstructedObject(UObject* NewObject) override
-	{
-		check(NewObject);
-
-		NewSections.Add(Cast<UMovieSceneSection>(NewObject));
+			}
+		}
 	}
 
-public:
-	TArray<UMovieSceneSection*> NewSections;
-};
-
-class FFolderObjectTextFactory : public FCustomizableTextObjectFactory
-{
-public:
-	FFolderObjectTextFactory()
-		: FCustomizableTextObjectFactory(GWarn)
+	if (SelectedNodes.Num() == 0)
 	{
+		FNotificationInfo Info(LOCTEXT("PasteSections_NoSelectedTracks", "Can't paste section. No selected tracks to paste sections onto"));
+		PasteErrors.Add(Info);
+		return false;
 	}
 
-	// FCustomizableTextObjectFactory implementation
-	virtual bool CanCreateClass(UClass* InObjectClass, bool& bOmitSubObjs) const override
+	TArray<UMovieSceneTrack*> Tracks;
+	TArray<int32> TrackRowIndices;
+	for (TSharedRef<FViewModel> Node : SelectedNodes)
 	{
-		if (InObjectClass->IsChildOf(UMovieSceneFolder::StaticClass()))
+		if (ITrackExtension* TrackNode = Node->CastThis<ITrackExtension>())
 		{
-			return true;
+			UMovieSceneTrack* TrackForNode = TrackNode->GetTrack();
+			if (TrackForNode)
+			{
+				Tracks.Add(TrackForNode);
+				if (FTrackRowModel* TrackRowNode = Node->CastThis<FTrackRowModel>())
+				{
+					TrackRowIndices.Add(TrackRowNode->GetRowIndex());
+				}
+			}
 		}
-		return false;
 	}
 
-
-	virtual void ProcessConstructedObject(UObject* NewObject) override
+	// Otherwise, look at all child nodes
+	if (Tracks.Num() == 0)
 	{
-		check(NewObject);
+		for (TSharedRef<FViewModel> Node : SelectedNodes)
+		{
+			TSet<TSharedRef<FViewModel> > DescendantNodes;
+			SequencerHelpers::GetDescendantNodes(Node, DescendantNodes);
+
+			for (TSharedRef<FViewModel> DescendantNode : DescendantNodes)
+			{
+				// Don't automatically paste onto subtracks because that would lead to multiple paste destinations
+				if (DescendantNode->IsA<FTrackRowModel>())
+				{
+					continue;
+				}
 
-		NewFolders.Add(Cast<UMovieSceneFolder>(NewObject));
+				if (ITrackExtension* TrackNode = DescendantNode->CastThis<ITrackExtension>())
+				{
+					UMovieSceneTrack* TrackForNode = TrackNode->GetTrack();
+					if (TrackForNode)
+					{
+						Tracks.Add(TrackForNode);
+						if (FTrackRowModel* TrackRowNode = Node->CastThis<FTrackRowModel>())
+						{
+							TrackRowIndices.Add(TrackRowNode->GetRowIndex());
+						}
+					}
+				}
+			}
+		}
 	}
 
-public:
-	TArray<UMovieSceneFolder*> NewFolders;
-};
+	TArray<UMovieSceneSection*> NewSections;
+	FSequencerUtilities::PasteSections(TextToImport, FMovieScenePasteSectionsParams(Tracks, TrackRowIndices, GetLocalTime().Time), NewSections, PasteErrors);
 
-bool FSequencer::CanPaste(const FString& TextToImport)
-{
-	FObjectBindingTextFactory ObjectBindingFactory(*this);
-	if (ObjectBindingFactory.CanCreateObjectsFromText(TextToImport))
-	{
-		return true;
-	}
-		
-	FTrackObjectTextFactory TrackFactory;
-	if (TrackFactory.CanCreateObjectsFromText(TextToImport))
-	{
-		return true;
-	}
+	NotifyMovieSceneDataChanged(EMovieSceneDataChangeType::MovieSceneStructureItemAdded);
 
-	FSectionObjectTextFactory SectionFactory;
-	if (SectionFactory.CanCreateObjectsFromText(TextToImport))
+	FSectionModelStorageExtension* SectionModelStorage = ViewModel->GetRootModel()->CastDynamic<FSectionModelStorageExtension>();
 	{
-		return true;
+		UE::MovieScene::FScopedSignedObjectModifyDefer ForceFlush(true);
+		check(SectionModelStorage);
 	}
 
-	FFolderObjectTextFactory FolderFactory;
-	if (FolderFactory.CanCreateObjectsFromText(TextToImport))
+	EmptySelection();
+
+	Selection.SuspendBroadcast();
+	for (UMovieSceneSection* NewSection : NewSections)
 	{
-		return true;
+		TSharedPtr<FSectionModel> SectionModel = SectionModelStorage->FindModelForSection(NewSection);
+		if (ensure(SectionModel))
+		{
+			Selection.AddToSelection(SectionModel);
+		}
 	}
+	Selection.ResumeBroadcast();
 
-	return false;
+	ThrobSectionSelection();
+
+	return true;
 }
 
-void FSequencer::ImportTracksFromText(const FString& TextToImport, /*out*/ TArray<UMovieSceneCopyableTrack*>& ImportedTracks)
+bool FSequencer::CanPaste(const FString& TextToImport)
 {
-	UPackage* TempPackage = NewObject<UPackage>(nullptr, TEXT("/Engine/Sequencer/Editor/Transient"), RF_Transient);
-	TempPackage->AddToRoot();
-
-	// Turn the text buffer into objects
-	FTrackObjectTextFactory Factory;
-	Factory.ProcessBuffer(TempPackage, RF_Transactional, TextToImport);
-
-	ImportedTracks = Factory.NewTracks;
-
-	// Remove the temp package from the root now that it has served its purpose
-	TempPackage->RemoveFromRoot();
+	return (FSequencerUtilities::CanPasteBindings(AsShared(), TextToImport) ||
+			FSequencerUtilities::CanPasteTracks(TextToImport) ||
+			FSequencerUtilities::CanPasteSections(TextToImport) ||
+			FSequencerUtilities::CanPasteFolders(TextToImport));
 }
 
 void FSequencer::ObjectImplicitlyAdded(UObject* InObject) const
@@ -9605,36 +8169,6 @@ TArray<FText> FSequencer::GetTrackFilterNames() const
 	return SequencerWidget->GetTrackFilterNames();
 }
 
-void FSequencer::ImportSectionsFromText(const FString& TextToImport, /*out*/ TArray<UMovieSceneSection*>& ImportedSections)
-{
-	UPackage* TempPackage = NewObject<UPackage>(nullptr, TEXT("/Engine/Sequencer/Editor/Transient"), RF_Transient);
-	TempPackage->AddToRoot();
-
-	// Turn the text buffer into objects
-	FSectionObjectTextFactory Factory;
-	Factory.ProcessBuffer(TempPackage, RF_Transactional, TextToImport);
-
-	ImportedSections = Factory.NewSections;
-
-	// Remove the temp package from the root now that it has served its purpose
-	TempPackage->RemoveFromRoot();
-}
-
-void FSequencer::ImportFoldersFromText(const FString& TextToImport, /*out*/ TArray<UMovieSceneFolder*>& ImportedFolders)
-{
-	UPackage* TempPackage = NewObject<UPackage>(nullptr, TEXT("/Engine/Sequencer/Editor/Transient"), RF_Transient);
-	TempPackage->AddToRoot();
-
-	// Turn the text buffer into objects
-	FFolderObjectTextFactory Factory;
-	Factory.ProcessBuffer(TempPackage, RF_Transactional, TextToImport);
-
-	ImportedFolders = Factory.NewFolders;
-
-	// Remove the temp package from the root now that it has served its purpose
-	TempPackage->RemoveFromRoot();
-}
-
 void FSequencer::ToggleNodeActive()
 {
 	bool bIsActive = !IsNodeActive();
@@ -9896,13 +8430,13 @@ void FSequencer::ConvertToSpawnable(TSharedRef<UE::Sequencer::FObjectBindingMode
 	FMovieScenePossessable* Possessable = GetFocusedMovieSceneSequence()->GetMovieScene()->FindPossessable(NodeToBeConverted->GetObjectGuid());
 	if (Possessable)
 	{
-		ConvertToSpawnableInternal(Possessable->GetGuid());
+		FSequencerUtilities::ConvertToSpawnable(AsShared(), Possessable->GetGuid());
 		NotifyMovieSceneDataChanged( EMovieSceneDataChangeType::MovieSceneStructureItemsChanged );
 	}
 }
 TArray<FGuid> FSequencer::ConvertToSpawnable(FGuid Guid)
 {
-	TArray< FMovieSceneSpawnable*> Spawnables = ConvertToSpawnableInternal(Guid);
+	TArray< FMovieSceneSpawnable*> Spawnables = FSequencerUtilities::ConvertToSpawnable(AsShared(), Guid);
 	TArray<FGuid> SpawnableGuids;
 	if (Spawnables.Num() > 0)
 	{
@@ -9924,301 +8458,83 @@ void FSequencer::ConvertSelectedNodesToSpawnables()
 
 	if (MovieScene->IsReadOnly())
 	{
-		FSequencerUtilities::ShowReadOnlyError();
-		return;
-	}
-
-	// @todo sequencer: Undo doesn't seem to be working at all
-	const FScopedTransaction Transaction( LOCTEXT("ConvertSelectedNodesSpawnable", "Convert Selected Nodes to Spawnables") );
-
-	// Ensure we're in a non-possessed state
-	TGuardValue<bool> Guard(bUpdatingExternalSelection, true);
-	RestorePreAnimatedState();
-	MovieScene->Modify();
-
-	TArray<IObjectBindingExtension*> ObjectBindingNodes;
-
-	for (const TWeakPtr<FViewModel>& Node : Selection.GetSelectedOutlinerItems())
-	{
-		if (IObjectBindingExtension* ObjectBindingNode = ICastable::CastWeakPtr<IObjectBindingExtension>(Node))
-		{
-			// If we have a possessable for this node, and it has no parent, we can convert it to a spawnable
-			FMovieScenePossessable* Possessable = MovieScene->FindPossessable(ObjectBindingNode->GetObjectGuid());
-			if (Possessable && !Possessable->GetParent().IsValid())
-			{
-				ObjectBindingNodes.Add(ObjectBindingNode);
-			}
-		}
-	}
-
-	FScopedSlowTask SlowTask(ObjectBindingNodes.Num(), LOCTEXT("ConvertSpawnableProgress", "Converting Selected Possessable Nodes to Spawnables"));
-	SlowTask.MakeDialog(true);
-
-	TArray<AActor*> SpawnedActors;
-	for (IObjectBindingExtension* ObjectBindingNode : ObjectBindingNodes)
-	{
-		SlowTask.EnterProgressFrame();
-	
-		FMovieScenePossessable* Possessable = MovieScene->FindPossessable(ObjectBindingNode->GetObjectGuid());
-		if (Possessable)
-		{
-			TArray<FMovieSceneSpawnable*> Spawnables = ConvertToSpawnableInternal(Possessable->GetGuid());
-
-			for (FMovieSceneSpawnable* Spawnable : Spawnables)
-			{
-				for (TWeakObjectPtr<> WeakObject : FindBoundObjects(Spawnable->GetGuid(), ActiveTemplateIDs.Top()))
-				{
-					if (AActor* SpawnedActor = Cast<AActor>(WeakObject.Get()))
-					{
-						SpawnedActors.Add(SpawnedActor);
-					}
-				}
-			}
-		}
-
-		if (GWarn->ReceivedUserCancel())
-		{
-			break;
-		}
-	}
-
-	if (SpawnedActors.Num())
-	{
-		const bool bNotifySelectionChanged = true;
-		const bool bDeselectBSP = true;
-		const bool bWarnAboutTooManyActors = false;
-		const bool bSelectEvenIfHidden = false;
-
-		GEditor->GetSelectedActors()->Modify();
-		GEditor->GetSelectedActors()->BeginBatchSelectOperation();
-		GEditor->SelectNone(bNotifySelectionChanged, bDeselectBSP, bWarnAboutTooManyActors);
-		for (auto SpawnedActor : SpawnedActors)
-		{
-			GEditor->SelectActor(SpawnedActor, true, bNotifySelectionChanged, bSelectEvenIfHidden);
-		}
-		GEditor->GetSelectedActors()->EndBatchSelectOperation();
-		GEditor->NoteSelectionChange();
-	}
-
-	NotifyMovieSceneDataChanged( EMovieSceneDataChangeType::MovieSceneStructureItemsChanged );
-}
-
-TArray<FGuid> FSequencer::ExpandMultiplePossessableBindings(FGuid PossessableGuid)
-{
-	using namespace UE::Sequencer;
-
-	UMovieSceneSequence* Sequence = GetFocusedMovieSceneSequence();
-	UMovieScene* MovieScene = Sequence->GetMovieScene();
-
-	TArray<FGuid> NewPossessableGuids;
-
-	if (MovieScene->IsReadOnly())
-	{
-		FSequencerUtilities::ShowReadOnlyError();
-		return TArray<FGuid>();
-	}
-
-	
-	// Create a copy of the TArrayView of bound objects, as the underlying array will get destroyed
-	TArray<TWeakObjectPtr<>> FoundObjects;
-	for (TWeakObjectPtr<> BoundObject : FindBoundObjects(PossessableGuid, ActiveTemplateIDs.Top()))
-	{
-		FoundObjects.Insert(BoundObject,0);
-	}
-
-	if (FoundObjects.Num() < 2)
-	{
-		// If less than two objects, nothing to do, return the same Guid
-		NewPossessableGuids.Add(PossessableGuid);
-		return NewPossessableGuids;
-	}
-
-	Sequence->Modify();
-	MovieScene->Modify();
-
-	FMovieSceneBinding* PossessableBinding = (FMovieSceneBinding*)MovieScene->GetBindings().FindByPredicate([&](FMovieSceneBinding& Binding) { return Binding.GetObjectGuid() == PossessableGuid; });
-
-	// First gather the children
-	TArray<FGuid> ChildPossessableGuids;
-	for (int32 Index = 0; Index < MovieScene->GetPossessableCount(); ++Index)
-	{
-		FMovieScenePossessable& Possessable = MovieScene->GetPossessable(Index);
-		if (Possessable.GetParent() == PossessableGuid)
-		{
-			ChildPossessableGuids.Add(Possessable.GetGuid());
-		}
-	}
-
-	TArray<UMovieSceneTrack* > Tracks = PossessableBinding->StealTracks(MovieScene);
-
-	// Remove binding to stop any children from claiming the old guid as their parent
-	if (MovieScene->RemovePossessable(PossessableGuid))
-	{
-		Sequence->UnbindPossessableObjects(PossessableGuid);
-	}
-
-	for (TWeakObjectPtr<> FoundObjectPtr : FoundObjects)
-	{
-		UObject* FoundObject = FoundObjectPtr.Get();
-		if (!FoundObject)
-		{
-			continue;
-		}
-
-		FoundObject->Modify();
-
-		UObject* BindingContext = GetPlaybackContext();
-
-		// Find this object's parent object, if it has one.
-		UObject* ParentObject = Sequence->GetParentObject(FoundObject);
-		if (ParentObject)
-		{
-			BindingContext = ParentObject;
-		}
-
-		// Create a new Possessable for this object
-		AActor* PossessedActor = Cast<AActor>(FoundObject);
-		const FGuid NewPossessableGuid = MovieScene->AddPossessable(PossessedActor != nullptr ? PossessedActor->GetActorLabel() : FoundObject->GetName(), FoundObject->GetClass());
-		FMovieScenePossessable* NewPossessable = MovieScene->FindPossessable(NewPossessableGuid);
-		if (NewPossessable)
-		{
-			FMovieSceneBinding* NewPossessableBinding = (FMovieSceneBinding*)MovieScene->GetBindings().FindByPredicate([&](FMovieSceneBinding& Binding) { return Binding.GetObjectGuid() == NewPossessableGuid; });
-			
-			if (ParentObject)
-			{
-				FGuid ParentGuid = FindObjectId(*ParentObject, ActiveTemplateIDs.Top());
-				NewPossessable->SetParent(ParentGuid, MovieScene);
-			}
-
-			if (!NewPossessable->BindSpawnableObject(GetFocusedTemplateID(), FoundObject, this))
-			{
-				Sequence->BindPossessableObject(NewPossessableGuid, *FoundObject, BindingContext);
-			}
-
-			NewPossessableGuids.Add(NewPossessableGuid);
-
-			// Create copies of the tracks
-			for (UMovieSceneTrack* Track : Tracks)
-			{
-				UMovieSceneTrack* DuplicatedTrack = Cast<UMovieSceneTrack>(StaticDuplicateObject(Track, MovieScene));
-				NewPossessableBinding->AddTrack(*DuplicatedTrack, MovieScene);
-			}
-		}
-	}
-
-	// Finally, recurse in to any children
-	for (FGuid ChildPossessableGuid : ChildPossessableGuids)
-	{
-		ExpandMultiplePossessableBindings(ChildPossessableGuid);
-	}
-
-	NotifyMovieSceneDataChanged(EMovieSceneDataChangeType::MovieSceneStructureItemsChanged);
-
-	return NewPossessableGuids;
-}
-
-TArray<FMovieSceneSpawnable*> FSequencer::ConvertToSpawnableInternal(FGuid PossessableGuid)
-{
-	using namespace UE::Sequencer;
-
-	UMovieSceneSequence* Sequence = GetFocusedMovieSceneSequence();
-	UMovieScene* MovieScene = Sequence->GetMovieScene();
-
-	if (MovieScene->IsReadOnly() || !Sequence->AllowsSpawnableObjects())
-	{
-		FSequencerUtilities::ShowReadOnlyError();
-		return TArray<FMovieSceneSpawnable*>();
-	}
-
-	TArrayView<TWeakObjectPtr<>> FoundObjects = FindBoundObjects(PossessableGuid, ActiveTemplateIDs.Top());
-
-	TArray<FMovieSceneSpawnable*> CreatedSpawnables;
-
-	if (FoundObjects.Num() == 0)
-	{
-		FMovieScenePossessable* Possessable = MovieScene->FindPossessable(PossessableGuid);
-
-		UE_LOG(LogSequencer, Error, TEXT("Failed to convert %s to spawnable because there are no objects bound to it"), Possessable ? *Possessable->GetName() : TEXT(""));
-	}
-	else if (FoundObjects.Num() > 1)
-	{
-		// Expand to individual possessables for each bound object, then convert each one individually
-		TArray<FGuid> ExpandedPossessableGuids = ExpandMultiplePossessableBindings(PossessableGuid);
-		for (FGuid NewPossessableGuid : ExpandedPossessableGuids)
-		{
-			CreatedSpawnables.Append(ConvertToSpawnableInternal(NewPossessableGuid));
-		}
-
-		ForceEvaluate();
-	}
-	else
-	{
-		UObject* FoundObject = FoundObjects[0].Get();
-		if (!FoundObject)
-		{
-			return TArray<FMovieSceneSpawnable*>();
-		}
+		FSequencerUtilities::ShowReadOnlyError();
+		return;
+	}
 
-		Sequence->Modify();
-		MovieScene->Modify();
+	// @todo sequencer: Undo doesn't seem to be working at all
+	const FScopedTransaction Transaction( LOCTEXT("ConvertSelectedNodesSpawnable", "Convert Selected Nodes to Spawnables") );
+
+	// Ensure we're in a non-possessed state
+	TGuardValue<bool> Guard(bUpdatingExternalSelection, true);
+	RestorePreAnimatedState();
+	MovieScene->Modify();
+
+	TArray<IObjectBindingExtension*> ObjectBindingNodes;
 
-		// Locate the folder containing the original possessable
-		UMovieSceneFolder* ParentFolder;
-		for (UMovieSceneFolder* Folder : MovieScene->GetRootFolders())
+	for (const TWeakPtr<FViewModel>& Node : Selection.GetSelectedOutlinerItems())
+	{
+		if (IObjectBindingExtension* ObjectBindingNode = ICastable::CastWeakPtr<IObjectBindingExtension>(Node))
 		{
-			ParentFolder = Folder->FindFolderContaining(PossessableGuid);
-			if (ParentFolder != nullptr)
+			// If we have a possessable for this node, and it has no parent, we can convert it to a spawnable
+			FMovieScenePossessable* Possessable = MovieScene->FindPossessable(ObjectBindingNode->GetObjectGuid());
+			if (Possessable && !Possessable->GetParent().IsValid())
 			{
-				break;
+				ObjectBindingNodes.Add(ObjectBindingNode);
 			}
 		}
+	}
 
-		FMovieSceneSpawnable* Spawnable = MovieScene->FindSpawnable(AddSpawnable(*FoundObject));
-		if (Spawnable)
-		{
-			CreatedSpawnables.Add(Spawnable);
-			FGuid SpawnableGuid = Spawnable->GetGuid();
-
-			// Remap all the spawnable's tracks and child bindings onto the new possessable
-			MovieScene->MoveBindingContents(PossessableGuid, SpawnableGuid);
+	FScopedSlowTask SlowTask(ObjectBindingNodes.Num(), LOCTEXT("ConvertSpawnableProgress", "Converting Selected Possessable Nodes to Spawnables"));
+	SlowTask.MakeDialog(true);
 
-			FMovieSceneBinding* PossessableBinding = (FMovieSceneBinding*)MovieScene->GetBindings().FindByPredicate([&](FMovieSceneBinding& Binding) { return Binding.GetObjectGuid() == PossessableGuid; });
-			check(PossessableBinding);
+	TArray<AActor*> SpawnedActors;
+	for (IObjectBindingExtension* ObjectBindingNode : ObjectBindingNodes)
+	{
+		SlowTask.EnterProgressFrame();
+	
+		FMovieScenePossessable* Possessable = MovieScene->FindPossessable(ObjectBindingNode->GetObjectGuid());
+		if (Possessable)
+		{
+			TArray<FMovieSceneSpawnable*> Spawnables = FSequencerUtilities::ConvertToSpawnable(AsShared(), Possessable->GetGuid());
 
-			for (UMovieSceneFolder* Folder : MovieScene->GetRootFolders())
+			for (FMovieSceneSpawnable* Spawnable : Spawnables)
 			{
-				if (ReplaceFolderBindingGUID(Folder, PossessableGuid, SpawnableGuid))
+				for (TWeakObjectPtr<> WeakObject : FindBoundObjects(Spawnable->GetGuid(), ActiveTemplateIDs.Top()))
 				{
-					break;
+					if (AActor* SpawnedActor = Cast<AActor>(WeakObject.Get()))
+					{
+						SpawnedActors.Add(SpawnedActor);
+					}
 				}
 			}
+		}
 
-			int32 SortingOrder = PossessableBinding->GetSortingOrder();
-
-			if (MovieScene->RemovePossessable(PossessableGuid))
-			{
-				Sequence->UnbindPossessableObjects(PossessableGuid);
-
-				FMovieSceneBinding* SpawnableBinding = (FMovieSceneBinding*)MovieScene->GetBindings().FindByPredicate([&](FMovieSceneBinding& Binding) { return Binding.GetObjectGuid() == SpawnableGuid; });
-				check(SpawnableBinding);
-
-				SpawnableBinding->SetSortingOrder(SortingOrder);
-
-			}
-
-			TOptional<FTransformData> TransformData;
-			SpawnRegister->HandleConvertPossessableToSpawnable(FoundObject, *this, TransformData);
-			SpawnRegister->SetupDefaultsForSpawnable(nullptr, Spawnable->GetGuid(), TransformData, AsShared(), Settings);
-
-			FSequencerUtilities::UpdateBindingIDs(this, CompiledDataManager, PossessableGuid, Spawnable->GetGuid());
+		if (GWarn->ReceivedUserCancel())
+		{
+			break;
+		}
+	}
 
-			ForceEvaluate();
+	if (SpawnedActors.Num())
+	{
+		const bool bNotifySelectionChanged = true;
+		const bool bDeselectBSP = true;
+		const bool bWarnAboutTooManyActors = false;
+		const bool bSelectEvenIfHidden = false;
 
-			NotifyMovieSceneDataChanged(EMovieSceneDataChangeType::MovieSceneStructureItemsChanged);
+		GEditor->GetSelectedActors()->Modify();
+		GEditor->GetSelectedActors()->BeginBatchSelectOperation();
+		GEditor->SelectNone(bNotifySelectionChanged, bDeselectBSP, bWarnAboutTooManyActors);
+		for (auto SpawnedActor : SpawnedActors)
+		{
+			GEditor->SelectActor(SpawnedActor, true, bNotifySelectionChanged, bSelectEvenIfHidden);
 		}
+		GEditor->GetSelectedActors()->EndBatchSelectOperation();
+		GEditor->NoteSelectionChange();
 	}
 
-	return CreatedSpawnables;
+	NotifyMovieSceneDataChanged( EMovieSceneDataChangeType::MovieSceneStructureItemsChanged );
 }
 
 void FSequencer::ConvertToPossessable(TSharedRef<UE::Sequencer::FObjectBindingModel> NodeToBeConverted)
@@ -10240,7 +8556,7 @@ void FSequencer::ConvertToPossessable(TSharedRef<UE::Sequencer::FObjectBindingMo
 	FMovieSceneSpawnable* Spawnable = GetFocusedMovieSceneSequence()->GetMovieScene()->FindSpawnable(NodeToBeConverted->GetObjectGuid());
 	if (Spawnable)
 	{
-		ConvertToPossessableInternal(Spawnable->GetGuid());
+		FSequencerUtilities::ConvertToPossessable(AsShared(), Spawnable->GetGuid());
 		NotifyMovieSceneDataChanged(EMovieSceneDataChangeType::MovieSceneStructureItemsChanged);
 	}
 }
@@ -10287,7 +8603,7 @@ void FSequencer::ConvertSelectedNodesToPossessables()
 			FMovieSceneSpawnable* Spawnable = MovieScene->FindSpawnable(ObjectBindingNode->GetObjectGuid());
 			if (Spawnable)
 			{
-				FMovieScenePossessable* Possessable = ConvertToPossessableInternal(Spawnable->GetGuid());
+				FMovieScenePossessable* Possessable = FSequencerUtilities::ConvertToPossessable(AsShared(), Spawnable->GetGuid());
 
 				ForceEvaluate();
 
@@ -10328,159 +8644,6 @@ void FSequencer::ConvertSelectedNodesToPossessables()
 	}
 }
 
-FMovieScenePossessable* FSequencer::ConvertToPossessableInternal(FGuid SpawnableGuid)
-{
-	using namespace UE::Sequencer;
-
-	UMovieSceneSequence* Sequence = GetFocusedMovieSceneSequence();
-	UMovieScene* MovieScene = Sequence->GetMovieScene();
-
-	if (MovieScene->IsReadOnly())
-	{
-		FSequencerUtilities::ShowReadOnlyError();
-		return nullptr;
-	}
-
-	// Find the object in the environment
-	UMovieSceneSequence* FocusedSequence = GetFocusedMovieSceneSequence();
-	FMovieSceneSpawnable* Spawnable = MovieScene->FindSpawnable(SpawnableGuid);
-	if (!Spawnable || !Spawnable->GetObjectTemplate())
-	{
-		return nullptr;
-	}
-
-	AActor* SpawnableActorTemplate = Cast<AActor>(Spawnable->GetObjectTemplate());
-	if (!SpawnableActorTemplate)
-	{
-		return nullptr;
-	}
-
-	TMap<TWeakObjectPtr<AActor>, FTransform> AttachedChildTransforms;
-	FTransform DefaultTransform = SpawnableActorTemplate->GetActorTransform();
-	for (TWeakObjectPtr<> RuntimeObject : FindBoundObjects(SpawnableGuid, ActiveTemplateIDs.Top()) )
-	{
-		// Prefer the transform at the current time over the spawnable actor template's transform because that's most likely 0. 
-		// This makes it so that the object will return to the current position on restore state.
-		AActor* Actor = Cast<AActor>(RuntimeObject.Get());
-		if (Actor)
-		{
-			if (Actor->GetRootComponent())
-			{
-				DefaultTransform = Actor->GetRootComponent()->GetRelativeTransform();
-			}
-
-			// Removing a parent will compensate the children at their world transform. We don't want that since we'll be replacing that parent right away.
-			// To negate that, we store the relative transform of these children and reset it after the parent is replaced with the new possessable.
-			TArray<AActor*> AttachedActors;
-			Actor->GetAttachedActors(AttachedActors);
-			for (AActor* ChildActor : AttachedActors)
-			{
-				if (ChildActor && ChildActor->GetRootComponent())
-				{	
-					// Only do this for child actors that Sequencer is controlling
-					FGuid ExistingID;
-					for (FMovieSceneSequenceID SequenceID : ActiveTemplateIDs)
-					{
-						ExistingID = FindObjectId(*ChildActor, SequenceID);
-						if (ExistingID.IsValid())
-						{
-							AttachedChildTransforms.Add(ChildActor);
-							AttachedChildTransforms[ChildActor] = ChildActor->GetRootComponent()->GetRelativeTransform();
-							break;
-						}
-					}
-				}
-			}
-			break;
-		}
-	}
-
-	Sequence->Modify();
-	MovieScene->Modify();
-
-	// Delete the spawn track
-	UMovieSceneSpawnTrack* SpawnTrack = Cast<UMovieSceneSpawnTrack>(MovieScene->FindTrack(UMovieSceneSpawnTrack::StaticClass(), SpawnableGuid, NAME_None));
-	if (SpawnTrack)
-	{
-		MovieScene->RemoveTrack(*SpawnTrack);
-	}
-
-	FActorSpawnParameters SpawnInfo;
-	SpawnInfo.bDeferConstruction = true;
-	SpawnInfo.Template = SpawnableActorTemplate;
-
-	UWorld* PlaybackContext = Cast<UWorld>(GetPlaybackContext());
-	AActor* PossessedActor = PlaybackContext->SpawnActor(Spawnable->GetObjectTemplate()->GetClass(), &DefaultTransform, SpawnInfo);
-
-	if (!PossessedActor)
-	{
-		return nullptr;
-	}
-
-	PossessedActor->SetActorLabel(Spawnable->GetName());
-
-	const bool bIsDefaultTransform = true;
-	PossessedActor->FinishSpawning(DefaultTransform, bIsDefaultTransform);
-
-	// The transform needs to be set again for deferred construction and dynamic root components. Until the fix for: UE-67537
-	PossessedActor->SetActorTransform(DefaultTransform);
-
-	const FGuid NewPossessableGuid = CreateBinding(*PossessedActor, PossessedActor->GetActorLabel());
-	const FGuid OldSpawnableGuid = Spawnable->GetGuid();
-
-	FMovieScenePossessable* Possessable = MovieScene->FindPossessable(NewPossessableGuid);
-	if (Possessable)
-	{
-		// Remap all the spawnable's tracks and child bindings onto the new possessable
-		MovieScene->MoveBindingContents(OldSpawnableGuid, NewPossessableGuid);
-
-		FMovieSceneBinding* SpawnableBinding = MovieScene->FindBinding(OldSpawnableGuid);
-		check(SpawnableBinding);
-
-		for (UMovieSceneFolder* Folder : MovieScene->GetRootFolders())
-		{
-			if (ReplaceFolderBindingGUID(Folder, Spawnable->GetGuid(), Possessable->GetGuid()))
-			{
-				break;
-			}
-		}
-
-		int32 SortingOrder = SpawnableBinding->GetSortingOrder();
-
-		// Remove the spawnable and all it's sub tracks
-		if (MovieScene->RemoveSpawnable(OldSpawnableGuid))
-		{
-			SpawnRegister->DestroySpawnedObject(OldSpawnableGuid, ActiveTemplateIDs.Top(), *this);
-
-			FMovieSceneBinding* PossessableBinding = MovieScene->FindBinding(NewPossessableGuid);
-			check(PossessableBinding);
-			
-			PossessableBinding->SetSortingOrder(SortingOrder);
-		}
-
-		static const FName SequencerActorTag(TEXT("SequencerActor"));
-		PossessedActor->Tags.Remove(SequencerActorTag);
-
-		FSequencerUtilities::UpdateBindingIDs(this, CompiledDataManager, OldSpawnableGuid, NewPossessableGuid);
-
-		GEditor->SelectActor(PossessedActor, false, true);
-
-		for (TPair<TWeakObjectPtr<AActor>, FTransform> AttachedChildTransform : AttachedChildTransforms)
-		{
-			if (AActor* AttachedChild = AttachedChildTransform.Key.Get())
-			{
-				if (AttachedChild->GetRootComponent())
-				{
-					AttachedChild->GetRootComponent()->SetRelativeTransform(AttachedChildTransform.Value);
-				}
-			}
-		}
-
-		ForceEvaluate();
-	}
-
-	return Possessable;
-}
 
 void FSequencer::OnLoadRecordedData()
 {
@@ -10544,37 +8707,6 @@ void FSequencer::OnLoadRecordedData()
 
 }
 
-bool FSequencer::ReplaceFolderBindingGUID(UMovieSceneFolder* Folder, FGuid Original, FGuid Converted)
-{
-	UMovieScene* MovieScene = GetFocusedMovieSceneSequence()->GetMovieScene();
-
-	if (MovieScene->IsReadOnly())
-	{
-		FSequencerUtilities::ShowReadOnlyError();
-		return true;
-	}
-
-	for (FGuid ChildGuid : Folder->GetChildObjectBindings())
-	{
-		if (ChildGuid == Original)
-		{
-			Folder->AddChildObjectBinding(Converted);
-			Folder->RemoveChildObjectBinding(Original);
-			return true;
-		}
-	}
-
-	for (UMovieSceneFolder* ChildFolder : Folder->GetChildFolders())
-	{
-		if (ReplaceFolderBindingGUID(ChildFolder, Original, Converted))
-		{
-			return true;
-		}
-	}
-
-	return false;
-}
-
 void FSequencer::OnAddFolder()
 {
 	using namespace UE::Sequencer;
@@ -11799,7 +9931,7 @@ void FSequencer::CutSelectedKeys()
 
 void FSequencer::CopySelectedSections()
 {
-	TArray<UObject*> SelectedSections;
+	TArray<UMovieSceneSection*> SelectedSections;
 	for (TWeakObjectPtr<UMovieSceneSection> SelectedSectionPtr : Selection.GetSelectedSections())
 	{
 		if (SelectedSectionPtr.IsValid())
@@ -11809,7 +9941,8 @@ void FSequencer::CopySelectedSections()
 	}
 
 	FString ExportedText;
-	FSequencer::ExportObjectsToText(SelectedSections, /*out*/ ExportedText);
+	FSequencerUtilities::CopySections(SelectedSections, ExportedText);
+
 	FPlatformApplicationMisc::ClipboardCopy(*ExportedText);
 
 	// Make sure to clear the clipboard for the keys
@@ -11840,113 +9973,6 @@ void FSequencer::OnClipboardUsed(TSharedPtr<FMovieSceneClipboard> Clipboard)
 	});
 }
 
-void FSequencer::CreateCamera()
-{
-	UMovieScene* FocusedMovieScene = GetFocusedMovieSceneSequence()->GetMovieScene();
-	if (!FocusedMovieScene)
-	{
-		return;
-	}
-
-	if (FocusedMovieScene->IsReadOnly())
-	{
-		FSequencerUtilities::ShowReadOnlyError();
-		return;
-	}
-
-	UWorld* World = GCurrentLevelEditingViewportClient ? GCurrentLevelEditingViewportClient->GetWorld() : nullptr;
-	if (!World)
-	{
-		return;
-	}
-
-	const FScopedTransaction Transaction(NSLOCTEXT("Sequencer", "CreateCameraHere", "Create Camera Here"));
-
-	const bool bCreateAsSpawnable = Settings->GetCreateSpawnableCameras() && GetFocusedMovieSceneSequence()->AllowsSpawnableObjects();
-
-	FActorSpawnParameters SpawnParams;
-	if (bCreateAsSpawnable)
-	{
-		// Don't bother transacting this object if we're creating a spawnable since it's temporary
-		SpawnParams.ObjectFlags &= ~RF_Transactional;
-	}
-
-	// Set new camera to match viewport
-	ACineCameraActor* NewCamera = World->SpawnActor<ACineCameraActor>(SpawnParams);
-	if (!NewCamera)
-	{
-		return;
-	}
-
-	FGuid CameraGuid;
-
-	FMovieSceneSpawnable* Spawnable = nullptr;
-
-	if (bCreateAsSpawnable)
-	{
-		FString NewName = MovieSceneHelpers::MakeUniqueSpawnableName(FocusedMovieScene, FName::NameToDisplayString(ACineCameraActor::StaticClass()->GetFName().ToString(), false));
-
-		CameraGuid = MakeNewSpawnable(*NewCamera);
-		Spawnable = FocusedMovieScene->FindSpawnable(CameraGuid);
-
-		if (ensure(Spawnable))
-		{
-			Spawnable->SetName(NewName);			
-		}
-
-		// Destroy the old actor
-		World->EditorDestroyActor(NewCamera, false);
-
-		for (TWeakObjectPtr<UObject>& Object : FindBoundObjects(CameraGuid, ActiveTemplateIDs.Top()))
-		{
-			NewCamera = Cast<ACineCameraActor>(Object.Get());
-			if (NewCamera)
-			{
-				break;
-			}
-		}
-		ensure(NewCamera);
-
-		NewCamera->SetActorLabel(NewName, false);
-	}
-	else
-	{
-		CameraGuid = CreateBinding(*NewCamera, NewCamera->GetActorLabel());
-	}
-	
-	if (!CameraGuid.IsValid())
-	{
-		return;
-	}
-	
-	NewCamera->SetActorLocation( GCurrentLevelEditingViewportClient->GetViewLocation(), false );
-	NewCamera->SetActorRotation( GCurrentLevelEditingViewportClient->GetViewRotation() );
-	//pNewCamera->CameraComponent->FieldOfView = ViewportClient->ViewFOV; //@todo set the focal length from this field of view
-
-	OnActorAddedToSequencerEvent.Broadcast(NewCamera, CameraGuid);
-
-	NewCameraAdded(NewCamera, CameraGuid);
-
-	NotifyMovieSceneDataChanged(EMovieSceneDataChangeType::MovieSceneStructureItemAdded);
-}
-
-void FSequencer::NewCameraAdded(ACameraActor* NewCamera, FGuid CameraGuid)
-{
-	if (OnCameraAddedToSequencer().IsBound() && !OnCameraAddedToSequencer().Execute(NewCamera, CameraGuid))
-	{
-		return;
-	}
-
-	MovieSceneToolHelpers::LockCameraActorToViewport(SharedThis(this), NewCamera);
-
-	UMovieSceneSequence* Sequence = GetFocusedMovieSceneSequence();
-	if (Sequence->IsTrackSupported(UMovieSceneCameraCutTrack::StaticClass()) == ETrackSupport::Supported)
-	{
-		MovieSceneToolHelpers::CreateCameraCutSectionForCamera(Sequence->GetMovieScene(), CameraGuid, GetLocalTime().Time.FloorToFrame());
-	}
-}
-
-
 void FSequencer::FixPossessableObjectClassInternal(UMovieSceneSequence* Sequence, FMovieSceneSequenceIDRef SequenceID)
 {
 	UMovieScene* MovieScene = Sequence->GetMovieScene();
@@ -13001,7 +11027,11 @@ void FSequencer::BindCommands()
 
 	SequencerCommandBindings->MapAction(
 		Commands.CreateCamera,
-		FExecuteAction::CreateSP(this, &FSequencer::CreateCamera),
+		FExecuteAction::CreateLambda([this]{ 
+			const bool bSpawnable = Settings->GetCreateSpawnableCameras() && GetFocusedMovieSceneSequence()->AllowsSpawnableObjects();
+			ACineCameraActor* OutActor;
+			FSequencerUtilities::CreateCamera(AsShared(), bSpawnable, OutActor); 
+		} ),
 		FCanExecuteAction(),
 		FIsActionChecked(),
 		FIsActionButtonVisible::CreateLambda([this] { return ExactCast<ULevelSequence>(GetFocusedMovieSceneSequence()) != nullptr && IVREditorModule::Get().IsVREditorModeActive() == false; }) //@todo VREditor: Creating a camera while in VR mode disrupts the hmd. This is a temporary fix by hiding the button when in VR mode.
@@ -13096,7 +11126,6 @@ void FSequencer::BindCommands()
 		FExecuteAction::CreateSP(this, &FSequencer::StepForward),
 		EUIActionRepeatMode::RepeatEnabled);
 
-
 	SequencerCommandBindings->MapAction(
 		Commands.StepBackward,
 		FExecuteAction::CreateSP(this, &FSequencer::StepBackward),
diff --git a/Engine/Source/Editor/Sequencer/Private/Sequencer.h b/Engine/Source/Editor/Sequencer/Private/Sequencer.h
--- a/Engine/Source/Editor/Sequencer/Private/Sequencer.h
+++ b/Engine/Source/Editor/Sequencer/Private/Sequencer.h
@@ -357,16 +357,6 @@ public:
 
 protected:
 
-	/**
-	 * Attempts to add a new spawnable to the MovieScene for the specified asset, class, or actor
-	 *
-	 * @param	Object	The asset, class, or actor to add a spawnable for
-	 * @param	ActorFactory	Optional actor factory to use to create spawnable type
-	 *
-	 * @return	The spawnable ID, or invalid ID on failure
-	 */
-	FGuid AddSpawnable( UObject& Object, UActorFactory* ActorFactory = nullptr );
-
 	/**
 	 * Save default spawnable state for the currently selected objects
 	 */
@@ -528,11 +518,6 @@ public:
 	FText GetNavigateForwardTooltip() const;
 	FText GetNavigateBackwardTooltip() const;
 
-	/** Called when a user executes the assign selected to track menu item */
-	void AddActorsToBinding(FGuid ObjectBinding, const TArray<AActor*>& InActors);
-	void ReplaceBindingWithActors(FGuid ObjectBinding, const TArray<AActor*>& InActors);
-	void RemoveActorsFromBinding(FGuid ObjectBinding, const TArray<AActor*>& InActors);
-
 	/** Called when a user executes the delete node menu item */
 	void DeleteNode(TSharedRef<FViewModel> NodeToBeDeleted, const bool bKeepState);
 	void DeleteSelectedNodes(const bool bKeepState);
@@ -551,10 +536,10 @@ public:
 	void CopySelectedObjects(TArray<TSharedPtr<UE::Sequencer::FObjectBindingModel>>& ObjectNodes, const TArray<UMovieSceneFolder*>& Folders, /*out*/ FString& ExportedText);
 	void CopySelectedTracks(TArray<TSharedPtr<UE::Sequencer::FViewModel>>& TrackNodes, const TArray<UMovieSceneFolder*>& Folders, /*out*/ FString& ExportedText);
 	void CopySelectedFolders(const TArray<UMovieSceneFolder*>& Folders, /*out*/ FString& ExportedText);
-	void ExportObjectsToText(const TArray<UObject*>& ObjectsToExport, /*out*/ FString& ExportedText);
 
 	/** Called when a user executes the paste track menu item */
 	bool CanPaste(const FString& TextToImport);
+
 	/**
 	 * Attempts to paste from the clipboard
 	 * @return Whether the paste event was handled
@@ -562,14 +547,8 @@ public:
 	bool DoPaste(bool bClearSelection = false);
 	bool PasteTracks(const FString& TextToImport, UMovieSceneFolder* ParentFolder, const TArray<UMovieSceneFolder*>& InFolders, TArray<FNotificationInfo>& PasteErrors, bool bClearSelection = false);
 	bool PasteSections(const FString& TextToImport, TArray<FNotificationInfo>& PasteErrors);
-	bool PasteFolders(const FString& TextToImport, UMovieSceneFolder* ParentFolder, TArray<UMovieSceneFolder*>& OutFolders, TArray<FNotificationInfo>& PasteErrors);
 	bool PasteObjectBindings(const FString& TextToImport, UMovieSceneFolder* ParentFolder, const TArray<UMovieSceneFolder*>& InFolders, TArray<FNotificationInfo>& PasteErrors, bool bClearSelection = false);
 
-	void ImportTracksFromText(const FString& TextToImport, /*out*/ TArray<UMovieSceneCopyableTrack*>& ImportedTracks);
-	void ImportSectionsFromText(const FString& TextToImport, /*out*/ TArray<UMovieSceneSection*>& ImportedSections);
-	void ImportFoldersFromText(const FString& TextToImport, /*out*/ TArray<UMovieSceneFolder*>& ImportedFolders);
-	void ImportObjectBindingsFromText(const FString& TextToImport, /*out*/ TArray<UMovieSceneCopyableBinding*>& ImportedObjects);
-
 	/** Called when a user executes the active node menu item */
 	void ToggleNodeActive();
 	bool IsNodeActive() const;
@@ -649,12 +628,6 @@ public:
 	/** Promote a clipboard to the top of the clipboard stack, and update its timestamp */
 	void OnClipboardUsed(TSharedPtr<FMovieSceneClipboard> Clipboard);
 
-	/** Create camera and set it as the current camera cut. */
-	void CreateCamera();
-
-	/** Called when a new camera is added. Locks the viewport to the NewCamera. */
-	void NewCameraAdded(ACameraActor* NewCamera, FGuid CameraGuid);
-
 	/** Attempts to automatically fix up possessables whose object class don't match the object class of their currently bound objects */
 	void FixPossessableObjectClass();
 	void FixPossessableObjectClassInternal(UMovieSceneSequence* Sequence, FMovieSceneSequenceIDRef SequenceID);
@@ -1067,18 +1040,6 @@ private:
 	/** Updates viewport clients' actor locks if they relate to sequencer cameras */
 	void UpdateLevelViewportClientsActorLocks();
 
-	/** Expands Possessables with multiple bindings into indidual Possessables for each binding */
-	TArray<FGuid> ExpandMultiplePossessableBindings(FGuid PossessableGuid);
-
-	/** Internal conversion function that doesn't perform expensive reset/update tasks */
-	TArray<FMovieSceneSpawnable*> ConvertToSpawnableInternal(FGuid PossessableGuid);
-
-	/** Internal conversion function that doesn't perform expensive reset/update tasks */
-	FMovieScenePossessable* ConvertToPossessableInternal(FGuid SpawnableGuid);
-
-	/** Recurses through a folder to replace converted GUID with new GUID */
-	bool ReplaceFolderBindingGUID(UMovieSceneFolder *Folder, FGuid Original, FGuid Converted);
-
 	/** Internal function to render movie for a given start/end time */
 	void RenderMovieInternal(TRange<FFrameNumber> Range, bool bSetFrameOverrides = false) const;
 
diff --git a/Engine/Source/Editor/Sequencer/Private/SequencerUtilities.cpp b/Engine/Source/Editor/Sequencer/Private/SequencerUtilities.cpp
--- a/Engine/Source/Editor/Sequencer/Private/SequencerUtilities.cpp
+++ b/Engine/Source/Editor/Sequencer/Private/SequencerUtilities.cpp
@@ -1,9 +1,19 @@
 // Copyright Epic Games, Inc. All Rights Reserved.
 
 #include "SequencerUtilities.h"
+#include "CineCameraActor.h"
+#include "CameraRig_Rail.h"
+#include "CameraRig_Crane.h"
+#include "Components/SplineComponent.h"
+#include "Containers/ArrayBuilder.h"
+#include "Engine/Selection.h"
+#include "EngineUtils.h"
+#include "Exporters/Exporter.h"
+#include "Factories.h"
 #include "Misc/Attribute.h"
 #include "Misc/Paths.h"
 #include "Layout/Margin.h"
+#include "LevelEditorViewport.h"
 #include "Fonts/SlateFontInfo.h"
 #include "Widgets/DeclarativeSyntaxSupport.h"
 #include "Widgets/SBoxPanel.h"
@@ -13,20 +23,32 @@
 #include "Widgets/Input/SButton.h"
 #include "Styling/CoreStyle.h"
 #include "Styling/AppStyle.h"
-#include "MovieSceneTrack.h"
+#include "MovieSceneCopyableBinding.h"
+#include "MovieSceneCopyableTrack.h"
+#include "MovieSceneFolder.h"
 #include "MovieSceneSection.h"
+#include "MovieSceneSpawnRegister.h"
 #include "MovieSceneTimeHelpers.h"
-#include "MovieSceneFolder.h"
+#include "MovieSceneToolHelpers.h"
+#include "MovieSceneTrack.h"
+#include "Tracks/MovieScene3DAttachTrack.h"
+#include "Tracks/MovieSceneCameraCutTrack.h"
+#include "Tracks/MovieSceneSpawnTrack.h"
 #include "Compilation/MovieSceneCompiledDataManager.h"
 #include "Framework/MultiBox/MultiBoxBuilder.h"
 #include "ISequencerTrackEditor.h"
 #include "ISequencer.h"
 #include "Sequencer.h"
+#include "SequencerLog.h"
 #include "SequencerNodeTree.h"
+#include "MovieSceneBindingProxy.h"
 #include "ScopedTransaction.h"
+#include "UnrealEdGlobals.h"
+#include "UnrealExporter.h"
 #include "UObject/Package.h"
 #include "AssetRegistry/AssetRegistryModule.h"
 #include "FileHelpers.h"
+#include "HAL/PlatformApplicationMisc.h"
 #include "LevelSequence.h"
 #include "MVVM/Extensions/IObjectBindingExtension.h"
 #include "MVVM/ViewModels/TrackModel.h"
@@ -34,11 +56,8 @@
 #include "Framework/Notifications/NotificationManager.h"
 #include "Widgets/Notifications/SNotificationList.h"
 
-PRAGMA_DISABLE_OPTIMIZATION
-
 #define LOCTEXT_NAMESPACE "FSequencerUtilities"
 
-
 static EVisibility GetRolloverVisibility(TAttribute<bool> HoverState, TWeakPtr<SButton> WeakButton)
 {
 	TSharedPtr<SButton> Button = WeakButton.Pin();
@@ -404,19 +423,2018 @@
 	return AssociatedMaps;
 }
 
-FGuid FSequencerUtilities::DoAssignActor(ISequencer * InSequencerPtr, AActor* Actor, FGuid InObjectBinding)
-{
+/** Recurses through a folder to replace converted GUID with new GUID */
+bool UpdateFolderBindingID(UMovieSceneFolder* Folder, FGuid OldGuid, FGuid NewGuid)
+{
+	for (FGuid ChildGuid : Folder->GetChildObjectBindings())
+	{
+		if (ChildGuid == OldGuid)
+		{
+			Folder->AddChildObjectBinding(NewGuid);
+			Folder->RemoveChildObjectBinding(OldGuid);
+			return true;
+		}
+	}
+
+	for (UMovieSceneFolder* ChildFolder : Folder->GetChildFolders())
+	{
+		if (UpdateFolderBindingID(ChildFolder, OldGuid, NewGuid))
+		{
+			return true;
+		}
+	}
+
+	return false;
+}
+
+/** Expands Possessables with multiple bindings into individual Possessables for each binding */
+TArray<FGuid> ExpandMultiplePossessableBindings(TSharedRef<ISequencer> Sequencer, FGuid PossessableGuid)
+{
+	TArray<FGuid> NewPossessableGuids;
+
+	UMovieSceneSequence* Sequence = Sequencer->GetFocusedMovieSceneSequence();
+	if (!Sequence)
+	{
+		return NewPossessableGuids;
+	}
+
+	UMovieScene* MovieScene = Sequence->GetMovieScene();
+	if (!MovieScene)
+	{
+		return NewPossessableGuids;
+	}
+
+	// Create a copy of the TArrayView of bound objects, as the underlying array will get destroyed
+	TArray<TWeakObjectPtr<>> FoundObjects;
+	for (TWeakObjectPtr<> BoundObject : Sequencer->FindBoundObjects(PossessableGuid, Sequencer->GetFocusedTemplateID()))
+	{
+		FoundObjects.Insert(BoundObject, 0);
+	}
+
+	if (FoundObjects.Num() < 2)
+	{
+		// If less than two objects, nothing to do, return the same Guid
+		NewPossessableGuids.Add(PossessableGuid);
+		return NewPossessableGuids;
+	}
+
+	Sequence->Modify();
+	MovieScene->Modify();
+
+	FMovieSceneBinding* PossessableBinding = (FMovieSceneBinding*)MovieScene->GetBindings().FindByPredicate([&](FMovieSceneBinding& Binding) { return Binding.GetObjectGuid() == PossessableGuid; });
+
+	// First gather the children
+	TArray<FGuid> ChildPossessableGuids;
+	for (int32 Index = 0; Index < MovieScene->GetPossessableCount(); ++Index)
+	{
+		FMovieScenePossessable& Possessable = MovieScene->GetPossessable(Index);
+		if (Possessable.GetParent() == PossessableGuid)
+		{
+			ChildPossessableGuids.Add(Possessable.GetGuid());
+		}
+	}
+
+	TArray<UMovieSceneTrack* > Tracks = PossessableBinding->StealTracks(MovieScene);
+
+	// Remove binding to stop any children from claiming the old guid as their parent
+	if (MovieScene->RemovePossessable(PossessableGuid))
+	{
+		Sequence->UnbindPossessableObjects(PossessableGuid);
+	}
+
+	for (TWeakObjectPtr<> FoundObjectPtr : FoundObjects)
+	{
+		UObject* FoundObject = FoundObjectPtr.Get();
+		if (!FoundObject)
+		{
+			continue;
+		}
+
+		FoundObject->Modify();
+
+		UObject* BindingContext = Sequencer->GetPlaybackContext();
+
+		// Find this object's parent object, if it has one.
+		UObject* ParentObject = Sequence->GetParentObject(FoundObject);
+		if (ParentObject)
+		{
+			BindingContext = ParentObject;
+		}
+
+		// Create a new Possessable for this object
+		AActor* PossessedActor = Cast<AActor>(FoundObject);
+		const FGuid NewPossessableGuid = MovieScene->AddPossessable(PossessedActor != nullptr ? PossessedActor->GetActorLabel() : FoundObject->GetName(), FoundObject->GetClass());
+		FMovieScenePossessable* NewPossessable = MovieScene->FindPossessable(NewPossessableGuid);
+		if (NewPossessable)
+		{
+			FMovieSceneBinding* NewPossessableBinding = (FMovieSceneBinding*)MovieScene->GetBindings().FindByPredicate([&](FMovieSceneBinding& Binding) { return Binding.GetObjectGuid() == NewPossessableGuid; });
+
+			if (ParentObject)
+			{
+				FGuid ParentGuid = Sequencer->FindObjectId(*ParentObject, Sequencer->GetFocusedTemplateID());
+				NewPossessable->SetParent(ParentGuid, MovieScene);
+			}
+
+			if (!NewPossessable->BindSpawnableObject(Sequencer->GetFocusedTemplateID(), FoundObject, &Sequencer.Get()))
+			{
+				Sequence->BindPossessableObject(NewPossessableGuid, *FoundObject, BindingContext);
+			}
+
+			NewPossessableGuids.Add(NewPossessableGuid);
+
+			// Create copies of the tracks
+			for (UMovieSceneTrack* Track : Tracks)
+			{
+				UMovieSceneTrack* DuplicatedTrack = Cast<UMovieSceneTrack>(StaticDuplicateObject(Track, MovieScene));
+				NewPossessableBinding->AddTrack(*DuplicatedTrack, MovieScene);
+			}
+		}
+	}
+
+	// Finally, recurse in to any children
+	for (FGuid ChildPossessableGuid : ChildPossessableGuids)
+	{
+		ExpandMultiplePossessableBindings(Sequencer, ChildPossessableGuid);
+	}
+
+	Sequencer->NotifyMovieSceneDataChanged(EMovieSceneDataChangeType::MovieSceneStructureItemsChanged);
+
+	return NewPossessableGuids;
+}
+
+void NewCameraAdded(TSharedRef<ISequencer> Sequencer, ACameraActor* NewCamera, FGuid CameraGuid)
+{
+	if (Sequencer->OnCameraAddedToSequencer().IsBound() && !Sequencer->OnCameraAddedToSequencer().Execute(NewCamera, CameraGuid))
+	{
+		return;
+	}
+
+	MovieSceneToolHelpers::LockCameraActorToViewport(Sequencer, NewCamera);
+
+	UMovieSceneSequence* Sequence = Sequencer->GetFocusedMovieSceneSequence();
+	if (Sequence && Sequence->IsTrackSupported(UMovieSceneCameraCutTrack::StaticClass()) == ETrackSupport::Supported)
+	{
+		MovieSceneToolHelpers::CreateCameraCutSectionForCamera(Sequence->GetMovieScene(), CameraGuid, Sequencer->GetLocalTime().Time.FloorToFrame());
+	}
+}
+
+FGuid AddSpawnable(TSharedRef<ISequencer> Sequencer, UObject& Object, UActorFactory* ActorFactory = nullptr, FName SpawnableName = NAME_None)
+{
+	UMovieSceneSequence* Sequence = Sequencer->GetFocusedMovieSceneSequence();
+	if (!Sequence->AllowsSpawnableObjects())
+	{
+		return FGuid();
+	}
+
+	// Grab the MovieScene that is currently focused.  We'll add our Blueprint as an inner of the
+	// MovieScene asset.
+	UMovieScene* OwnerMovieScene = Sequence->GetMovieScene();
+
+	TValueOrError<FNewSpawnable, FText> Result = Sequencer->GetSpawnRegister().CreateNewSpawnableType(Object, *OwnerMovieScene, ActorFactory);
+	if (!Result.IsValid())
+	{
+		FNotificationInfo Info(Result.GetError());
+		Info.ExpireDuration = 3.0f;
+		FSlateNotificationManager::Get().AddNotification(Info);
+		return FGuid();
+	}
+
+	FNewSpawnable& NewSpawnable = Result.GetValue();
+
+	NewSpawnable.Name = SpawnableName == NAME_None ? MovieSceneHelpers::MakeUniqueSpawnableName(OwnerMovieScene, NewSpawnable.Name) : SpawnableName.ToString();
+
+	FGuid NewGuid = OwnerMovieScene->AddSpawnable(NewSpawnable.Name, *NewSpawnable.ObjectTemplate);
+
+	Sequencer->ForceEvaluate();
+
+	return NewGuid;
+}
+
+FGuid FSequencerUtilities::MakeNewSpawnable(TSharedRef<ISequencer> Sequencer, UObject& Object, UActorFactory* ActorFactory, bool bSetupDefaults, FName SpawnableName)
+{
+	UMovieSceneSequence* Sequence = Sequencer->GetFocusedMovieSceneSequence();
+	if (!Sequence)
+	{
+		return FGuid();
+	}
+
+	UMovieScene* MovieScene = Sequence->GetMovieScene();
+	if (!MovieScene)
+	{
+		return FGuid();
+	}
+
+	if (MovieScene->IsReadOnly() || !Sequence->AllowsSpawnableObjects())
+	{
+		ShowReadOnlyError();
+		return FGuid();
+	}
+
+	FGuid NewGuid = AddSpawnable(Sequencer, Object, ActorFactory, SpawnableName);
+	if (!NewGuid.IsValid())
+	{
+		return FGuid();
+	}
+
+	FMovieSceneSpawnable* Spawnable = MovieScene->FindSpawnable(NewGuid);
+	if (!Spawnable)
+	{
+		return FGuid();
+	}
+
+	// Spawn the object so we can position it correctly, it's going to get spawned anyway since things default to spawned.
+	UObject* SpawnedObject = Sequencer->GetSpawnRegister().SpawnObject(NewGuid, *MovieScene, Sequencer->GetFocusedTemplateID(), Sequencer.Get());
+
+	if (bSetupDefaults)
+	{
+		FTransformData TransformData;
+		Sequencer->GetSpawnRegister().SetupDefaultsForSpawnable(SpawnedObject, Spawnable->GetGuid(), TransformData, Sequencer, Sequencer->GetSequencerSettings());
+	}
+
+	if (ACameraActor* NewCamera = Cast<ACameraActor>(SpawnedObject))
+	{
+		NewCameraAdded(Sequencer, NewCamera, NewGuid);
+	}
+
+	return NewGuid;
+}
+
+FGuid FSequencerUtilities::CreateCamera(TSharedRef<ISequencer> Sequencer, const bool bSpawnable, ACineCameraActor*& OutActor)
+{
+	FGuid CameraGuid;
+
+	UMovieSceneSequence* Sequence = Sequencer->GetFocusedMovieSceneSequence();
+	if (!Sequence)
+	{
+		return CameraGuid;
+	}
+
+	UMovieScene* MovieScene = Sequence->GetMovieScene();
+	if (!MovieScene)
+	{
+		return CameraGuid;
+	}
+
+	if (MovieScene->IsReadOnly() || !Sequence->AllowsSpawnableObjects())
+	{
+		ShowReadOnlyError();
+		return CameraGuid;
+	}
+
+	UWorld* World = GCurrentLevelEditingViewportClient ? GCurrentLevelEditingViewportClient->GetWorld() : nullptr;
+	if (!World)
+	{
+		return CameraGuid;
+	}
+
+	const FScopedTransaction Transaction(LOCTEXT("CreateCamera", "Create Camera"));
+
+	FActorSpawnParameters SpawnParams;
+	if (bSpawnable)
+	{
+		// Don't bother transacting this object if we're creating a spawnable since it's temporary
+		SpawnParams.ObjectFlags &= ~RF_Transactional;
+	}
+
+	// Set new camera to match viewport
+	OutActor = World->SpawnActor<ACineCameraActor>(SpawnParams);
+	if (!OutActor)
+	{
+		return CameraGuid;
+	}
+
+	OutActor->SetActorLocation(GCurrentLevelEditingViewportClient->GetViewLocation(), false);
+	OutActor->SetActorRotation(GCurrentLevelEditingViewportClient->GetViewRotation());
+	//OutActor->CameraComponent->FieldOfView = ViewportClient->ViewFOV; //@todo set the focal length from this field of view
+
+	FMovieSceneSpawnable* Spawnable = nullptr;
+
+	if (bSpawnable)
+	{
+		FString NewName = MovieSceneHelpers::MakeUniqueSpawnableName(MovieScene, FName::NameToDisplayString(ACineCameraActor::StaticClass()->GetFName().ToString(), false));
+
+		CameraGuid = MakeNewSpawnable(Sequencer, *OutActor);
+		Spawnable = MovieScene->FindSpawnable(CameraGuid);
+
+		if (ensure(Spawnable))
+		{
+			Spawnable->SetName(NewName);
+		}
+
+		// Destroy the old actor
+		World->EditorDestroyActor(OutActor, false);
+
+		for (TWeakObjectPtr<UObject>& Object : Sequencer->FindBoundObjects(CameraGuid, Sequencer->GetFocusedTemplateID()))
+		{
+			OutActor = Cast<ACineCameraActor>(Object.Get());
+			if (OutActor)
+			{
+				break;
+			}
+		}
+		ensure(OutActor);
+
+		OutActor->SetActorLabel(NewName, false);
+	}
+	else
+	{
+		CameraGuid = CreateBinding(Sequencer, *OutActor, OutActor->GetActorLabel());
+	}
+
+	if (!CameraGuid.IsValid())
+	{
+		return CameraGuid;
+	}
+
+	Sequencer->OnActorAddedToSequencer().Broadcast(OutActor, CameraGuid);
+
+	NewCameraAdded(Sequencer, OutActor, CameraGuid);
+
+	return CameraGuid;
+}
+
+FGuid FSequencerUtilities::CreateCameraWithRig(TSharedRef<ISequencer> Sequencer, AActor* Actor, const bool bSpawnable, ACineCameraActor*& OutActor)
+{
+	FGuid CameraGuid;
+
+	UMovieSceneSequence* Sequence = Sequencer->GetFocusedMovieSceneSequence();
+	if (!Sequence)
+	{
+		return CameraGuid;
+	}
+
+	UMovieScene* MovieScene = Sequence->GetMovieScene();
+	if (!MovieScene)
+	{
+		return CameraGuid;
+	}
+
+	if (MovieScene->IsReadOnly() || !Sequence->AllowsSpawnableObjects())
+	{
+		ShowReadOnlyError();
+		return CameraGuid;
+	}
+
+	const FScopedTransaction Transaction(LOCTEXT("CreateCameraWithRig", "Create Camera with Rig"));
+
+	ACameraRig_Rail* RailActor = nullptr;
+	if (Actor->GetClass() == ACameraRig_Rail::StaticClass())
+	{
+		RailActor = Cast<ACameraRig_Rail>(Actor);
+	}
+
+	// Create a cine camera actor
+	UWorld* PlaybackContext = Cast<UWorld>(Sequencer->GetPlaybackContext());
+	OutActor = PlaybackContext->SpawnActor<ACineCameraActor>();
+	CameraGuid = CreateBinding(Sequencer, *OutActor, OutActor->GetActorLabel());
+
+	if (RailActor)
+	{
+		OutActor->SetActorRotation(FRotator(0.f, -90.f, 0.f));
+	}
+
+	Sequencer->OnActorAddedToSequencer().Broadcast(OutActor, CameraGuid);
+
+	TRange<FFrameNumber> PlaybackRange = MovieScene->GetPlaybackRange();
+
+	if (bSpawnable)
+	{
+		FString NewCameraName = MovieSceneHelpers::MakeUniqueSpawnableName(MovieScene, FName::NameToDisplayString(ACineCameraActor::StaticClass()->GetFName().ToString(), false));
+
+		FMovieSceneSpawnable* Spawnable = ConvertToSpawnable(Sequencer, CameraGuid)[0];
+		Spawnable->SetName(NewCameraName);
+
+		for (TWeakObjectPtr<> WeakObject : Sequencer->FindBoundObjects(Spawnable->GetGuid(), Sequencer->GetFocusedTemplateID()))
+		{
+			OutActor = Cast<ACineCameraActor>(WeakObject.Get());
+			if (OutActor)
+			{
+				break;
+			}
+		}
+
+		OutActor->SetActorLabel(NewCameraName, false);
+
+		CameraGuid = Spawnable->GetGuid();
+
+		// Create an attach track
+		UMovieScene3DAttachTrack* AttachTrack = Cast<UMovieScene3DAttachTrack>(MovieScene->AddTrack(UMovieScene3DAttachTrack::StaticClass(), CameraGuid));
+
+		FGuid NewGuid = Sequencer->FindObjectId(*Actor, Sequencer->GetFocusedTemplateID());
+		FMovieSceneObjectBindingID AttachBindingID = UE::MovieScene::FRelativeObjectBindingID(NewGuid);
+		FFrameNumber StartTime = UE::MovieScene::DiscreteInclusiveLower(PlaybackRange);
+		FFrameNumber Duration = UE::MovieScene::DiscreteSize(PlaybackRange);
+
+		AttachTrack->AddConstraint(StartTime, Duration.Value, NAME_None, NAME_None, AttachBindingID);
+	}
+	else
+	{
+		// Parent it
+		OutActor->AttachToActor(Actor, FAttachmentTransformRules::KeepRelativeTransform);
+	}
+
+	if (RailActor)
+	{
+		// Extend the rail a bit
+		if (RailActor->GetRailSplineComponent()->GetNumberOfSplinePoints() == 2)
+		{
+			FVector SplinePoint1 = RailActor->GetRailSplineComponent()->GetLocationAtSplinePoint(0, ESplineCoordinateSpace::Local);
+			FVector SplinePoint2 = RailActor->GetRailSplineComponent()->GetLocationAtSplinePoint(1, ESplineCoordinateSpace::Local);
+			FVector SplineDirection = SplinePoint2 - SplinePoint1;
+			SplineDirection.Normalize();
+
+			float DefaultRailDistance = 650.f;
+			SplinePoint2 = SplinePoint1 + SplineDirection * DefaultRailDistance;
+			RailActor->GetRailSplineComponent()->SetLocationAtSplinePoint(1, SplinePoint2, ESplineCoordinateSpace::Local);
+			RailActor->GetRailSplineComponent()->bSplineHasBeenEdited = true;
+		}
+
+		// Create a track for the CurrentPositionOnRail
+		FPropertyPath PropertyPath;
+		PropertyPath.AddProperty(FPropertyInfo(RailActor->GetClass()->FindPropertyByName(TEXT("CurrentPositionOnRail"))));
+
+		FKeyPropertyParams KeyPropertyParams(TArrayBuilder<UObject*>().Add(RailActor), PropertyPath, ESequencerKeyMode::ManualKeyForced);
+
+		FFrameTime OriginalTime = Sequencer->GetLocalTime().Time;
+
+		Sequencer->SetLocalTimeDirectly(UE::MovieScene::DiscreteInclusiveLower(PlaybackRange));
+		RailActor->CurrentPositionOnRail = 0.f;
+		Sequencer->KeyProperty(KeyPropertyParams);
+
+		Sequencer->SetLocalTimeDirectly(UE::MovieScene::DiscreteExclusiveUpper(PlaybackRange) - 1);
+		RailActor->CurrentPositionOnRail = 1.f;
+		Sequencer->KeyProperty(KeyPropertyParams);
+
+		Sequencer->SetLocalTimeDirectly(OriginalTime);
+	}
+
+	NewCameraAdded(Sequencer, OutActor, CameraGuid);
+
+	return CameraGuid;
+}
+
+TArray<FGuid> FSequencerUtilities::AddActors(TSharedRef<ISequencer> Sequencer, const TArray<TWeakObjectPtr<AActor> >& InActors)
+{
+	TArray<FGuid> PossessableGuids;
+
+	UMovieSceneSequence* Sequence = Sequencer->GetFocusedMovieSceneSequence();
+	if (!Sequence)
+	{
+		return PossessableGuids;
+	}
+
+	UMovieScene* MovieScene = Sequence->GetMovieScene();
+	if (!MovieScene)
+	{
+		return PossessableGuids;
+	}
+
+	if (MovieScene->IsReadOnly() || !Sequence->AllowsSpawnableObjects())
+	{
+		ShowReadOnlyError();
+		return PossessableGuids;
+	}
+	
+	const FScopedTransaction Transaction(LOCTEXT("AddActors", "Add Actors"));
+	Sequence->Modify();
+
+	for (TWeakObjectPtr<AActor> WeakActor : InActors)
+	{
+		if (AActor* Actor = WeakActor.Get())
+		{
+			FGuid ExistingGuid = Sequencer->FindObjectId(*Actor, Sequencer->GetFocusedTemplateID());
+			if (!ExistingGuid.IsValid())
+			{
+				FGuid PossessableGuid = CreateBinding(Sequencer, *Actor, Actor->GetActorLabel());
+				PossessableGuids.Add(PossessableGuid);
+
+				if (ACameraActor* CameraActor = Cast<ACameraActor>(Actor))
+				{
+					NewCameraAdded(Sequencer, CameraActor, PossessableGuid);
+				}
+
+				Sequencer->OnActorAddedToSequencer().Broadcast(Actor, PossessableGuid);
+			}
+		}
+	}
+
+	return PossessableGuids;
+}
+
+TArray<FMovieSceneSpawnable*> FSequencerUtilities::ConvertToSpawnable(TSharedRef<ISequencer> Sequencer, FGuid PossessableGuid)
+{
+	TArray<FMovieSceneSpawnable*> CreatedSpawnables;
+
+	UMovieSceneSequence* Sequence = Sequencer->GetFocusedMovieSceneSequence();
+	if (!Sequence)
+	{
+		return CreatedSpawnables;
+	}
+
+	UMovieScene* MovieScene = Sequence->GetMovieScene();
+	if (!MovieScene)
+	{
+		return CreatedSpawnables;
+	}
+
+	if (MovieScene->IsReadOnly() || !Sequence->AllowsSpawnableObjects())
+	{
+		ShowReadOnlyError();
+		return CreatedSpawnables;
+	}
+
+	TArrayView<TWeakObjectPtr<>> FoundObjects = Sequencer->FindBoundObjects(PossessableGuid, Sequencer->GetFocusedTemplateID());
+
+	if (FoundObjects.Num() == 0)
+	{
+		FMovieScenePossessable* Possessable = MovieScene->FindPossessable(PossessableGuid);
+
+		UE_LOG(LogSequencer, Error, TEXT("Failed to convert %s to spawnable because there are no objects bound to it"), Possessable ? *Possessable->GetName() : TEXT(""));
+	}
+	else if (FoundObjects.Num() > 1)
+	{
+		// Expand to individual possessables for each bound object, then convert each one individually
+		TArray<FGuid> ExpandedPossessableGuids = ExpandMultiplePossessableBindings(Sequencer, PossessableGuid);
+		for (FGuid NewPossessableGuid : ExpandedPossessableGuids)
+		{
+			CreatedSpawnables.Append(ConvertToSpawnable(Sequencer, NewPossessableGuid));
+		}
+
+		Sequencer->ForceEvaluate();
+	}
+	else
+	{
+		UObject* FoundObject = FoundObjects[0].Get();
+		if (!FoundObject)
+		{
+			return CreatedSpawnables;
+		}
+
+		Sequence->Modify();
+		MovieScene->Modify();
+
+		// Locate the folder containing the original possessable
+		UMovieSceneFolder* ParentFolder = nullptr;
+		for (UMovieSceneFolder* Folder : MovieScene->GetRootFolders())
+		{
+			ParentFolder = Folder->FindFolderContaining(PossessableGuid);
+			if (ParentFolder != nullptr)
+			{
+				break;
+			}
+		}
+
+		FMovieSceneSpawnable* Spawnable = MovieScene->FindSpawnable(AddSpawnable(Sequencer, *FoundObject));
+		if (Spawnable)
+		{
+			FGuid SpawnableGuid = Spawnable->GetGuid();
+			CreatedSpawnables.Add(Spawnable);
+
+			// Remap all the spawnable's tracks and child bindings onto the new possessable
+			MovieScene->MoveBindingContents(PossessableGuid, SpawnableGuid);
+
+			FMovieSceneBinding* PossessableBinding = (FMovieSceneBinding*)MovieScene->GetBindings().FindByPredicate([&](FMovieSceneBinding& Binding) { return Binding.GetObjectGuid() == PossessableGuid; });
+			check(PossessableBinding);
+
+			for (UMovieSceneFolder* Folder : MovieScene->GetRootFolders())
+			{
+				if (UpdateFolderBindingID(Folder, PossessableGuid, SpawnableGuid))
+				{
+					break;
+				}
+			}
+
+			int32 SortingOrder = PossessableBinding->GetSortingOrder();
+
+			if (MovieScene->RemovePossessable(PossessableGuid))
+			{
+				Sequence->UnbindPossessableObjects(PossessableGuid);
+
+				FMovieSceneBinding* SpawnableBinding = (FMovieSceneBinding*)MovieScene->GetBindings().FindByPredicate([&](FMovieSceneBinding& Binding) { return Binding.GetObjectGuid() == SpawnableGuid; });
+				check(SpawnableBinding);
+
+				SpawnableBinding->SetSortingOrder(SortingOrder);
+			}
+
+			TOptional<FTransformData> TransformData;
+			Sequencer->GetSpawnRegister().HandleConvertPossessableToSpawnable(FoundObject, *Sequencer, TransformData);
+			Sequencer->GetSpawnRegister().SetupDefaultsForSpawnable(nullptr, Spawnable->GetGuid(), TransformData, Sequencer, Sequencer->GetSequencerSettings());
+
+			UpdateBindingIDs(Sequencer, PossessableGuid, Spawnable->GetGuid());
+
+			Sequencer->ForceEvaluate();
+		}
+	}
+
+	return CreatedSpawnables;
+}
+
+FMovieScenePossessable* FSequencerUtilities::ConvertToPossessable(TSharedRef<ISequencer> Sequencer, FGuid SpawnableGuid)
+{
+	FMovieScenePossessable* CreatedPossessable = nullptr;
+
+	UMovieSceneSequence* Sequence = Sequencer->GetFocusedMovieSceneSequence();
+	if (!Sequence)
+	{
+		return CreatedPossessable;
+	}
+
+	UMovieScene* MovieScene = Sequence->GetMovieScene();
+	if (!MovieScene)
+	{
+		return CreatedPossessable;
+	}
+
+	if (MovieScene->IsReadOnly() || !Sequence->AllowsSpawnableObjects())
+	{
+		ShowReadOnlyError();
+		return CreatedPossessable;
+	}
+
+	// Find the object in the environment
+	FMovieSceneSpawnable* Spawnable = MovieScene->FindSpawnable(SpawnableGuid);
+	if (!Spawnable || !Spawnable->GetObjectTemplate())
+	{
+		return CreatedPossessable;
+	}
+
+	AActor* SpawnableActorTemplate = Cast<AActor>(Spawnable->GetObjectTemplate());
+	if (!SpawnableActorTemplate)
+	{
+		return CreatedPossessable;
+	}
+
+	TMap<TWeakObjectPtr<AActor>, FTransform> AttachedChildTransforms;
+	FTransform DefaultTransform = SpawnableActorTemplate->GetActorTransform();
+	for (TWeakObjectPtr<> RuntimeObject : Sequencer->FindBoundObjects(SpawnableGuid, Sequencer->GetFocusedTemplateID()))
+	{
+		// Prefer the transform at the current time over the spawnable actor template's transform because that's most likely 0. 
+		// This makes it so that the object will return to the current position on restore state.
+		AActor* Actor = Cast<AActor>(RuntimeObject.Get());
+		if (Actor)
+		{
+			if (Actor->GetRootComponent())
+			{
+				DefaultTransform = Actor->GetRootComponent()->GetRelativeTransform();
+			}
+
+			// Removing a parent will compensate the children at their world transform. We don't want that since we'll be replacing that parent right away.
+			// To negate that, we store the relative transform of these children and reset it after the parent is replaced with the new possessable.
+			TArray<AActor*> AttachedActors;
+			Actor->GetAttachedActors(AttachedActors);
+			for (AActor* ChildActor : AttachedActors)
+			{
+				if (ChildActor && ChildActor->GetRootComponent())
+				{
+					// Only do this for child actors that Sequencer is controlling
+					FGuid ExistingID = Sequencer->FindObjectId(*ChildActor, Sequencer->GetFocusedTemplateID());
+					if (ExistingID.IsValid())
+					{
+						AttachedChildTransforms.Add(ChildActor);
+						AttachedChildTransforms[ChildActor] = ChildActor->GetRootComponent()->GetRelativeTransform();
+					}
+				}
+			}
+			break;
+		}
+	}
+
+	Sequence->Modify();
+	MovieScene->Modify();
+
+	// Delete the spawn track
+	UMovieSceneSpawnTrack* SpawnTrack = Cast<UMovieSceneSpawnTrack>(MovieScene->FindTrack(UMovieSceneSpawnTrack::StaticClass(), SpawnableGuid, NAME_None));
+	if (SpawnTrack)
+	{
+		MovieScene->RemoveTrack(*SpawnTrack);
+	}
+
+	FActorSpawnParameters SpawnInfo;
+	SpawnInfo.bDeferConstruction = true;
+	SpawnInfo.Template = SpawnableActorTemplate;
+
+	UWorld* PlaybackContext = Cast<UWorld>(Sequencer->GetPlaybackContext());
+	AActor* PossessedActor = PlaybackContext->SpawnActor(Spawnable->GetObjectTemplate()->GetClass(), &DefaultTransform, SpawnInfo);
+
+	if (!PossessedActor)
+	{
+		return nullptr;
+	}
+
+	PossessedActor->SetActorLabel(Spawnable->GetName());
+
+	const bool bIsDefaultTransform = true;
+	PossessedActor->FinishSpawning(DefaultTransform, bIsDefaultTransform);
+
+	// The transform needs to be set again for deferred construction and dynamic root components. Until the fix for: UE-67537
+	PossessedActor->SetActorTransform(DefaultTransform);
+
+	const FGuid NewPossessableGuid = CreateBinding(Sequencer, *PossessedActor, PossessedActor->GetActorLabel());
+	const FGuid OldSpawnableGuid = Spawnable->GetGuid();
+
+	FMovieScenePossessable* Possessable = MovieScene->FindPossessable(NewPossessableGuid);
+	if (Possessable)
+	{
+		// Remap all the spawnable's tracks and child bindings onto the new possessable
+		MovieScene->MoveBindingContents(OldSpawnableGuid, NewPossessableGuid);
+
+		FMovieSceneBinding* SpawnableBinding = MovieScene->FindBinding(OldSpawnableGuid);
+		check(SpawnableBinding);
+
+		for (UMovieSceneFolder* Folder : MovieScene->GetRootFolders())
+		{
+			if (UpdateFolderBindingID(Folder, Spawnable->GetGuid(), Possessable->GetGuid()))
+			{
+				break;
+			}
+		}
+
+		int32 SortingOrder = SpawnableBinding->GetSortingOrder();
+
+		// Remove the spawnable and all it's sub tracks
+		if (MovieScene->RemoveSpawnable(OldSpawnableGuid))
+		{
+			Sequencer->GetSpawnRegister().DestroySpawnedObject(OldSpawnableGuid, Sequencer->GetFocusedTemplateID(), Sequencer.Get());
+
+			FMovieSceneBinding* PossessableBinding = MovieScene->FindBinding(NewPossessableGuid);
+			check(PossessableBinding);
+
+			PossessableBinding->SetSortingOrder(SortingOrder);
+		}
+
+		static const FName SequencerActorTag(TEXT("SequencerActor"));
+		PossessedActor->Tags.Remove(SequencerActorTag);
+
+		UpdateBindingIDs(Sequencer, OldSpawnableGuid, NewPossessableGuid);
+
+		GEditor->SelectActor(PossessedActor, false, true);
+
+		for (TPair<TWeakObjectPtr<AActor>, FTransform> AttachedChildTransform : AttachedChildTransforms)
+		{
+			if (AActor* AttachedChild = AttachedChildTransform.Key.Get())
+			{
+				if (AttachedChild->GetRootComponent())
+				{
+					AttachedChild->GetRootComponent()->SetRelativeTransform(AttachedChildTransform.Value);
+				}
+			}
+		}
+
+		Sequencer->ForceEvaluate();
+	}
+
+	return Possessable;
+}
+
+void ExportObjectsToText(const TArray<UObject*>& ObjectsToExport, FString& ExportedText)
+{
+	if (ObjectsToExport.Num() == 0)
+	{
+		return;
+	}
+
+	// Clear the mark state for saving.
+	UnMarkAllObjects(EObjectMark(OBJECTMARK_TagExp | OBJECTMARK_TagImp));
+
+	FStringOutputDevice Archive;
+	const FExportObjectInnerContext Context;
+
+	// Export each of the selected nodes
+	UObject* LastOuter = nullptr;
+
+	for (UObject* ObjectToExport : ObjectsToExport)
+	{
+		// The nodes should all be from the same scope
+		UObject* ThisOuter = ObjectToExport->GetOuter();
+		if (LastOuter != nullptr && ThisOuter != LastOuter)
+		{
+			UE_LOG(LogSequencer, Warning, TEXT("Cannot copy objects from different outers. Only copying from %s"), *LastOuter->GetName());
+			continue;
+		}
+		LastOuter = ThisOuter;
+
+		UExporter::ExportToOutputDevice(&Context, ObjectToExport, nullptr, Archive, TEXT("copy"), 0, PPF_ExportsNotFullyQualified | PPF_Copy | PPF_Delimited, false, ThisOuter);
+	}
+
+	ExportedText = Archive;
+}
+
+/**
+ *
+ * Copy/paste folders
+ *
+ */
+
+void FSequencerUtilities::CopyFolders(const TArray<UMovieSceneFolder*>& Folders, FString& ExportedText)
+{
+	TArray<UObject*> Objects;
+	for (UMovieSceneFolder* Folder : Folders)
+	{
+		Objects.Add(Folder);
+	}
+
+	ExportObjectsToText(Objects, /*out*/ ExportedText);
+}
+
+class FFolderObjectTextFactory : public FCustomizableTextObjectFactory
+{
+public:
+	FFolderObjectTextFactory()
+		: FCustomizableTextObjectFactory(GWarn)
+	{
+	}
+
+	// FCustomizableTextObjectFactory implementation
+	virtual bool CanCreateClass(UClass* InObjectClass, bool& bOmitSubObjs) const override
+	{
+		if (InObjectClass->IsChildOf(UMovieSceneFolder::StaticClass()))
+		{
+			return true;
+		}
+		return false;
+	}
+
+
+	virtual void ProcessConstructedObject(UObject* NewObject) override
+	{
+		check(NewObject);
+
+		NewFolders.Add(Cast<UMovieSceneFolder>(NewObject));
+	}
+
+public:
+	TArray<UMovieSceneFolder*> NewFolders;
+};
+
+void ImportFoldersFromText(const FString& TextToImport, /*out*/ TArray<UMovieSceneFolder*>& ImportedFolders)
+{
+	UPackage* TempPackage = NewObject<UPackage>(nullptr, TEXT("/Engine/Sequencer/Editor/Transient"), RF_Transient);
+	TempPackage->AddToRoot();
+
+	// Turn the text buffer into objects
+	FFolderObjectTextFactory Factory;
+	Factory.ProcessBuffer(TempPackage, RF_Transactional, TextToImport);
+
+	ImportedFolders = Factory.NewFolders;
+
+	// Remove the temp package from the root now that it has served its purpose
+	TempPackage->RemoveFromRoot();
+}
+
+bool FSequencerUtilities::PasteFolders(const FString& TextToImport, FMovieScenePasteFoldersParams PasteFoldersParams, TArray<UMovieSceneFolder*>& OutFolders, TArray<FNotificationInfo>& OutErrors)
+{
+	if (!PasteFoldersParams.Sequence)
+	{
+		return false;
+	}
+
+	UMovieScene* MovieScene = PasteFoldersParams.Sequence->GetMovieScene();
+	if (!MovieScene)
+	{
+		return false;
+	}
+
+	TArray<UMovieSceneFolder*> ImportedFolders;
+	ImportFoldersFromText(TextToImport, ImportedFolders);
+
+	if (ImportedFolders.Num() == 0)
+	{
+		return false;
+	}
+
+	const FScopedTransaction Transaction(LOCTEXT("PasteFolders", "Paste Folders"));
+
+	MovieScene->Modify();
+
+	for (UMovieSceneFolder* CopiedFolder : ImportedFolders)
+	{
+		CopiedFolder->Rename(nullptr, MovieScene);
+
+		OutFolders.Add(CopiedFolder);
+
+		// Clear the folder contents, those relationships will be made when the tracks are pasted
+		CopiedFolder->ClearChildMasterTracks();
+		CopiedFolder->ClearChildObjectBindings();
+
+		bool bHasParent = false;
+		for (UMovieSceneFolder* ImportedParentFolder : ImportedFolders)
+		{
+			if (ImportedParentFolder != CopiedFolder)
+			{
+				if (ImportedParentFolder->GetChildFolders().Contains(CopiedFolder))
+				{
+					bHasParent = true;
+					break;
+				}
+			}
+		}
+
+		if (!bHasParent)
+		{
+			if (PasteFoldersParams.ParentFolder)
+			{
+				PasteFoldersParams.ParentFolder->AddChildFolder(CopiedFolder);
+			}
+			else
+			{
+				MovieScene->AddRootFolder(CopiedFolder);
+			}
+		}
+	}
+
+	return true;
+}
+
+bool FSequencerUtilities::CanPasteFolders(const FString& TextToImport)
+{
+	FFolderObjectTextFactory FolderFactory;
+	return FolderFactory.CanCreateObjectsFromText(TextToImport);
+}
+
+/**
+ *
+ * Copy/paste tracks
+ *
+ */
+
+void FSequencerUtilities::CopyTracks(const TArray<UMovieSceneTrack*>& Tracks, const TArray<UMovieSceneFolder*>& Folders, FString& ExportedText)
+{
+	TArray<UObject*> Objects;
+	for (UMovieSceneTrack* Track : Tracks)
+	{
+		UMovieScene* MovieScene = Track->GetTypedOuter<UMovieScene>();
+
+		UMovieSceneCopyableTrack* CopyableTrack = NewObject<UMovieSceneCopyableTrack>(GetTransientPackage(), UMovieSceneCopyableTrack::StaticClass(), NAME_None, RF_Transient);
+		Objects.Add(CopyableTrack);
+
+		UMovieSceneTrack* DuplicatedTrack = Cast<UMovieSceneTrack>(StaticDuplicateObject(Track, CopyableTrack));
+		CopyableTrack->Track = DuplicatedTrack;
+		CopyableTrack->bIsAMasterTrack = MovieScene->IsAMasterTrack(*Track);
+		CopyableTrack->bIsACameraCutTrack = Track->IsA<UMovieSceneCameraCutTrack>();
+
+		UMovieSceneFolder* Folder = nullptr;
+		for (UMovieSceneFolder* RootFolder : MovieScene->GetRootFolders())
+		{
+			Folder = RootFolder->FindFolderContaining(Track);
+			if (Folder && Folders.Contains(Folder))
+			{
+				UMovieSceneFolder::CalculateFolderPath(Folder, Folders, CopyableTrack->FolderPath);
+				break;
+			}
+		}
+	}
+
+	ExportObjectsToText(Objects, /*out*/ ExportedText);
+}
+
+class FTrackObjectTextFactory : public FCustomizableTextObjectFactory
+{
+public:
+	FTrackObjectTextFactory()
+		: FCustomizableTextObjectFactory(GWarn)
+	{
+	}
+
+	// FCustomizableTextObjectFactory implementation
+	virtual bool CanCreateClass(UClass* InObjectClass, bool& bOmitSubObjs) const override
+	{
+		if (InObjectClass->IsChildOf(UMovieSceneCopyableTrack::StaticClass()))
+		{
+			return true;
+		}
+		return false;
+	}
+
+
+	virtual void ProcessConstructedObject(UObject* NewObject) override
+	{
+		check(NewObject);
+
+		NewTracks.Add(Cast<UMovieSceneCopyableTrack>(NewObject));
+	}
+
+public:
+	TArray<UMovieSceneCopyableTrack*> NewTracks;
+};
+
+void ImportTracksFromText(const FString& TextToImport, /*out*/ TArray<UMovieSceneCopyableTrack*>& ImportedTracks)
+{
+	UPackage* TempPackage = NewObject<UPackage>(nullptr, TEXT("/Engine/Sequencer/Editor/Transient"), RF_Transient);
+	TempPackage->AddToRoot();
+
+	// Turn the text buffer into objects
+	FTrackObjectTextFactory Factory;
+	Factory.ProcessBuffer(TempPackage, RF_Transactional, TextToImport);
+
+	ImportedTracks = Factory.NewTracks;
+
+	// Remove the temp package from the root now that it has served its purpose
+	TempPackage->RemoveFromRoot();
+}
+
+bool FSequencerUtilities::PasteTracks(const FString& TextToImport, FMovieScenePasteTracksParams PasteTracksParams, TArray<UMovieSceneTrack*>& OutTracks, TArray<FNotificationInfo>& OutErrors)
+{
+	TArray<UMovieSceneCopyableTrack*> ImportedTracks;
+	ImportTracksFromText(TextToImport, ImportedTracks);
+
+	if (ImportedTracks.Num() == 0)
+	{
+		return false;
+	}
+
+	const FScopedTransaction Transaction(LOCTEXT("PasteTracks", "Paste Tracks"));
+
+	int32 NumMasterOrCameraCutTracks = 0;
+	int32 NumTracks = 0;
+
+	for (UMovieSceneCopyableTrack* CopyableTrack : ImportedTracks)
+	{
+		if (CopyableTrack->bIsAMasterTrack || CopyableTrack->bIsACameraCutTrack)
+		{
+			++NumMasterOrCameraCutTracks;
+		}
+		else
+		{
+			++NumTracks;
+		}
+	}
+
+	int32 NumTracksPasted = 0;
+	int32 NumMasterOrCameraCutTracksPasted = 0;
+
+	for (const FMovieSceneBindingProxy& ObjectBinding : PasteTracksParams.Bindings)
+	{
+		TArray<UMovieSceneCopyableTrack*> NewTracks;
+		ImportTracksFromText(TextToImport, NewTracks);
+
+		UMovieScene* MovieScene = ObjectBinding.GetMovieScene();
+		if (!MovieScene)
+		{
+			continue;
+		}
+
+		for (UMovieSceneCopyableTrack* CopyableTrack : NewTracks)
+		{
+			if (!CopyableTrack->bIsAMasterTrack && !CopyableTrack->bIsACameraCutTrack)
+			{
+				UMovieSceneTrack* NewTrack = CopyableTrack->Track;
+				NewTrack->ClearFlags(RF_Transient);
+				TArray<UObject*> Subobjects;
+				GetObjectsWithOuter(NewTrack, Subobjects);
+				for (UObject* Subobject : Subobjects)
+				{
+					Subobject->ClearFlags(RF_Transient);
+				}
+
+				// Remove tracks with the same name before adding
+				for (const FMovieSceneBinding& Binding : MovieScene->GetBindings())
+				{
+					if (Binding.GetObjectGuid() == ObjectBinding.BindingID)
+					{
+						// Tracks of the same class should be unique per name.
+						for (UMovieSceneTrack* Track : Binding.GetTracks())
+						{
+							if (Track->GetClass() == NewTrack->GetClass() && Track->GetTrackName() == NewTrack->GetTrackName() && Track->GetDisplayName().IdenticalTo(NewTrack->GetDisplayName()))
+							{
+								// If a track of the same class and name exists, remove it so the new track replaces it
+								MovieScene->RemoveTrack(*Track);
+								break;
+							}
+						}
+					}
+				}
+
+				if (!MovieScene->AddGivenTrack(NewTrack, ObjectBinding.BindingID))
+				{
+					continue;
+				}
+				else
+				{
+					OutTracks.Add(NewTrack);
+					++NumTracksPasted;
+				}
+			}
+		}
+	}
+
+	UMovieScene* MovieScene = PasteTracksParams.Sequence ? PasteTracksParams.Sequence->GetMovieScene() : nullptr;
+	if (MovieScene)
+	{
+		// Add as master track or set camera cut track
+		for (UMovieSceneCopyableTrack* CopyableTrack : ImportedTracks)
+		{
+			if (CopyableTrack->bIsAMasterTrack || CopyableTrack->bIsACameraCutTrack)
+			{
+				UMovieSceneTrack* NewTrack = CopyableTrack->Track;
+				NewTrack->ClearFlags(RF_Transient);
+				TArray<UObject*> Subobjects;
+				GetObjectsWithOuter(NewTrack, Subobjects);
+				for (UObject* Subobject : Subobjects)
+				{
+					Subobject->ClearFlags(RF_Transient);
+				}
+
+				UMovieSceneFolder* ParentFolder = PasteTracksParams.ParentFolder;
+
+				if (CopyableTrack->FolderPath.Num() > 0)
+				{
+					ParentFolder = UMovieSceneFolder::GetFolderWithPath(CopyableTrack->FolderPath, PasteTracksParams.Folders, ParentFolder ? ParentFolder->GetChildFolders() : MovieScene->GetRootFolders());
+				}
+
+				if (NewTrack->IsA(UMovieSceneCameraCutTrack::StaticClass()))
+				{
+					MovieScene->SetCameraCutTrack(NewTrack);
+					if (ParentFolder != nullptr)
+					{
+						ParentFolder->AddChildMasterTrack(NewTrack);
+					}
+
+					++NumMasterOrCameraCutTracksPasted;
+				}
+				else
+				{
+					if (MovieScene->AddGivenMasterTrack(NewTrack))
+					{
+						if (ParentFolder != nullptr)
+						{
+							ParentFolder->AddChildMasterTrack(NewTrack);
+						}
+					}
+
+					++NumMasterOrCameraCutTracksPasted;
+				}
+
+				OutTracks.Add(NewTrack);
+			}
+		}
+	}
+	
+	if (NumMasterOrCameraCutTracksPasted < NumMasterOrCameraCutTracks)
+	{
+		FNotificationInfo Info(LOCTEXT("PasteTracks_NoMasterTracks", "Can't paste track. Master track could not be pasted"));
+		OutErrors.Add(Info);
+	}
+
+	if (NumTracksPasted < NumTracks)
+	{
+		FNotificationInfo Info(LOCTEXT("PasteTracks_NoSelectedObjects", "Can't paste track. No selected objects to paste tracks onto"));
+		OutErrors.Add(Info);
+	}
+
+	return (NumMasterOrCameraCutTracksPasted + NumTracksPasted) > 0;
+}
+
+bool FSequencerUtilities::CanPasteTracks(const FString& TextToImport)
+{
+	FTrackObjectTextFactory TrackFactory;
+	return TrackFactory.CanCreateObjectsFromText(TextToImport);
+}
+
+/**
+ *
+ * Copy/paste sections
+ *
+ */
+
+void FSequencerUtilities::CopySections(const TArray<UMovieSceneSection*>& Sections, FString& ExportedText)
+{
+	TArray<UObject*> Objects;
+	for (UMovieSceneSection* Section : Sections)
+	{
+		Objects.Add(Section);
+	}
+
+	ExportObjectsToText(Objects, /*out*/ ExportedText);
+}
+
+class FSectionObjectTextFactory : public FCustomizableTextObjectFactory
+{
+public:
+	FSectionObjectTextFactory()
+		: FCustomizableTextObjectFactory(GWarn)
+	{
+	}
+
+	// FCustomizableTextObjectFactory implementation
+	virtual bool CanCreateClass(UClass* InObjectClass, bool& bOmitSubObjs) const override
+	{
+		if (InObjectClass->IsChildOf(UMovieSceneSection::StaticClass()))
+		{
+			return true;
+		}
+		return false;
+	}
+
+
+	virtual void ProcessConstructedObject(UObject* NewObject) override
+	{
+		check(NewObject);
+
+		NewSections.Add(Cast<UMovieSceneSection>(NewObject));
+	}
+
+public:
+	TArray<UMovieSceneSection*> NewSections;
+};
+
+void ImportSectionsFromText(const FString& TextToImport, /*out*/ TArray<UMovieSceneSection*>& ImportedSections)
+{
+	UPackage* TempPackage = NewObject<UPackage>(nullptr, TEXT("/Engine/Sequencer/Editor/Transient"), RF_Transient);
+	TempPackage->AddToRoot();
+
+	// Turn the text buffer into objects
+	FSectionObjectTextFactory Factory;
+	Factory.ProcessBuffer(TempPackage, RF_Transactional, TextToImport);
+
+	ImportedSections = Factory.NewSections;
+
+	// Remove the temp package from the root now that it has served its purpose
+	TempPackage->RemoveFromRoot();
+}
+
+bool FSequencerUtilities::PasteSections(const FString& TextToImport, FMovieScenePasteSectionsParams PasteSectionsParams, TArray<UMovieSceneSection*>& OutSections, TArray<FNotificationInfo>& OutErrors)
+{
+	// First import as a track and extract sections to allow for copying track contents to another track
+	TArray<UMovieSceneCopyableTrack*> ImportedTracks;
+	ImportTracksFromText(TextToImport, ImportedTracks);
+
+	TArray<UMovieSceneSection*> ImportedSections;
+	for (UMovieSceneCopyableTrack* CopyableTrack : ImportedTracks)
+	{
+		for (UMovieSceneSection* CopyableSection : CopyableTrack->Track->GetAllSections())
+		{
+			ImportedSections.Add(CopyableSection);
+		}
+	}
+
+	// Otherwise, import as sections
+	if (ImportedSections.Num() == 0)
+	{
+		ImportSectionsFromText(TextToImport, ImportedSections);
+	}
+
+	if (ImportedSections.Num() == 0)
+	{
+		return false;
+	}
+
+	const FScopedTransaction Transaction(LOCTEXT("PasteSections", "Paste Sections"));
+
+	TOptional<FFrameNumber> FirstFrame;
+	for (UMovieSceneSection* Section : ImportedSections)
+	{
+		if (Section->HasStartFrame())
+		{
+			if (FirstFrame.IsSet())
+			{
+				if (FirstFrame.GetValue() > Section->GetInclusiveStartFrame())
+				{
+					FirstFrame = Section->GetInclusiveStartFrame();
+				}
+			}
+			else
+			{
+				FirstFrame = Section->GetInclusiveStartFrame();
+			}
+		}
+	}
+
+	TArray<int32> SectionIndicesImported;
+
+	for (int32 Index = 0; Index < PasteSectionsParams.Tracks.Num(); ++Index)
+	{
+		UMovieSceneTrack* Track = PasteSectionsParams.Tracks[Index];
+		int32 RowIndex = Index < PasteSectionsParams.TrackRowIndices.Num() ? PasteSectionsParams.TrackRowIndices[Index] : 0;
+
+		for (int32 SectionIndex = 0; SectionIndex < ImportedSections.Num(); ++SectionIndex)
+		{
+			UMovieSceneSection* Section = ImportedSections[SectionIndex];
+			if (!Track->SupportsType(Section->GetClass()))
+			{
+				continue;
+			}
+
+			SectionIndicesImported.AddUnique(SectionIndex);
+
+			Track->Modify();
+
+			Section->ClearFlags(RF_Transient);
+			Section->Rename(nullptr, Track);
+			Track->AddSection(*Section);
+			if (Section->HasStartFrame())
+			{
+				FFrameNumber NewStartFrame = PasteSectionsParams.Time.FrameNumber + (Section->GetInclusiveStartFrame() - FirstFrame.GetValue());
+				Section->MoveSection(NewStartFrame - Section->GetInclusiveStartFrame());
+			}
+
+			if (Track->SupportsMultipleRows())
+			{
+				Section->SetRowIndex(RowIndex);
+			}
+			OutSections.Add(Section);
+		}
+
+		// Fix up rows after sections are in place
+		if (Track->SupportsMultipleRows())
+		{
+			// If any newly created section overlaps the previous sections, put all the sections on the max available row
+			// Find the  this section overlaps any previous sections, 
+			int32 MaxAvailableRowIndex = -1;
+			for (UMovieSceneSection* Section : OutSections)
+			{
+				if (MovieSceneToolHelpers::OverlapsSection(Track, Section, OutSections))
+				{
+					int32 AvailableRowIndex = MovieSceneToolHelpers::FindAvailableRowIndex(Track, Section, OutSections);
+					MaxAvailableRowIndex = FMath::Max(AvailableRowIndex, MaxAvailableRowIndex);
+				}
+			}
+
+			if (MaxAvailableRowIndex != -1)
+			{
+				for (UMovieSceneSection* Section : OutSections)
+				{
+					Section->SetRowIndex(MaxAvailableRowIndex);
+				}
+			}
+		}
+
+		// Regenerate for pasting onto the next track 
+		ImportedSections.Empty();
+		ImportedTracks.Empty();
+
+		ImportTracksFromText(TextToImport, ImportedTracks);
+
+		for (UMovieSceneCopyableTrack* CopyableTrack : ImportedTracks)
+		{
+			for (UMovieSceneSection* CopyableSection : CopyableTrack->Track->GetAllSections())
+			{
+				ImportedSections.Add(CopyableSection);
+			}
+		}
+
+		if (ImportedSections.Num() == 0)
+		{
+			ImportSectionsFromText(TextToImport, ImportedSections);
+		}
+	}
+
+	for (int32 SectionIndex = 0; SectionIndex < ImportedSections.Num(); ++SectionIndex)
+	{
+		if (!SectionIndicesImported.Contains(SectionIndex))
+		{
+			UE_LOG(LogSequencer, Display, TEXT("Could not paste section of type %s"), *ImportedSections[SectionIndex]->GetClass()->GetName());
+		}
+	}
+
+	if (SectionIndicesImported.Num() == 0)
+	{
+		FNotificationInfo Info(LOCTEXT("PasteSections_NothingPasted", "Can't paste section. No matching section types found."));
+		OutErrors.Add(Info);
+		return false;
+	}
+
+	return true;
+}
+
+bool FSequencerUtilities::CanPasteSections(const FString& TextToImport)
+{
+	FSectionObjectTextFactory SectionFactory;
+	return SectionFactory.CanCreateObjectsFromText(TextToImport);
+}
+
+/**
+ *
+ * Copy/paste object bindings
+ *
+ */
+
+void ExportObjectBindingsToText(const TArray<UMovieSceneCopyableBinding*>& ObjectsToExport, FString& ExportedText)
+{
+	// Clear the mark state for saving.
+	UnMarkAllObjects(EObjectMark(OBJECTMARK_TagExp | OBJECTMARK_TagImp));
+
+	FStringOutputDevice Archive;
+	const FExportObjectInnerContext Context;
+
+	// Export each of the selected nodes
+	UObject* LastOuter = nullptr;
+
+	for (UMovieSceneCopyableBinding* ObjectToExport : ObjectsToExport)
+	{
+		// The nodes should all be from the same scope
+		UObject* ThisOuter = ObjectToExport->GetOuter();
+		check((LastOuter == ThisOuter) || (LastOuter == nullptr));
+		LastOuter = ThisOuter;
+
+		// We can't use TextExportTransient on USTRUCTS (which our object contains) so we're going to manually null out some references before serializing them. These references are
+		// serialized manually into the archive, as the auto-serialization will only store a reference (to a privately owned object) which creates issues on deserialization. Attempting 
+		// to deserialize these private objects throws a superflous error in the console that makes it look like things went wrong when they're actually OK and expected.
+		TArray<UMovieSceneTrack*> OldTracks = ObjectToExport->Binding.StealTracks(nullptr);
+		UObject* OldSpawnableTemplate = ObjectToExport->Spawnable.GetObjectTemplate();
+		ObjectToExport->Spawnable.SetObjectTemplate(nullptr);
+
+		UExporter::ExportToOutputDevice(&Context, ObjectToExport, nullptr, Archive, TEXT("copy"), 0, PPF_ExportsNotFullyQualified | PPF_Copy | PPF_Delimited, false, ThisOuter);
+
+		// Restore the references (as we don't want to modify the original in the event of a copy operation!)
+		ObjectToExport->Binding.SetTracks(MoveTemp(OldTracks), nullptr);
+		ObjectToExport->Spawnable.SetObjectTemplate(OldSpawnableTemplate);
+
+		// We manually export the object template for the same private-ownership reason as above. Templates need to be re-created anyways as each Spawnable contains its own copy of the template.
+		if (ObjectToExport->SpawnableObjectTemplate)
+		{
+			UExporter::ExportToOutputDevice(&Context, ObjectToExport->SpawnableObjectTemplate, nullptr, Archive, TEXT("copy"), 0, PPF_ExportsNotFullyQualified | PPF_Copy | PPF_Delimited);
+		}
+	}
+
+	ExportedText = Archive;
+}
+
+void FSequencerUtilities::CopyBindings(const TArray<FMovieSceneBindingProxy>& Bindings, const TArray<UMovieSceneFolder*>& InFolders, FString& ExportedText)
+{
+	TArray<UMovieSceneCopyableBinding*> Objects;
+	for (const FMovieSceneBindingProxy& ObjectBinding : Bindings)
+	{
+		UMovieSceneCopyableBinding* CopyableBinding = NewObject<UMovieSceneCopyableBinding>(GetTransientPackage(), UMovieSceneCopyableBinding::StaticClass(), NAME_None, RF_Transient);
+		Objects.Add(CopyableBinding);
+
+		UMovieScene* MovieScene = ObjectBinding.GetMovieScene();
+		if (!MovieScene)
+		{
+			continue;
+		}
+
+		FMovieScenePossessable* Possessable = MovieScene->FindPossessable(ObjectBinding.BindingID);
+		if (Possessable)
+		{
+			CopyableBinding->Possessable = *Possessable;
+		}
+		else
+		{
+			FMovieSceneSpawnable* Spawnable = MovieScene->FindSpawnable(ObjectBinding.BindingID);
+			if (Spawnable)
+			{
+				CopyableBinding->Spawnable = *Spawnable;
+
+				// We manually serialize the spawnable object template so that it's not a reference to a privately owned object. Spawnables all have unique copies of their template objects anyways.
+				// Object Templates are re-created on paste (based on these templates) with the correct ownership set up.
+				CopyableBinding->SpawnableObjectTemplate = Spawnable->GetObjectTemplate();
+			}
+		}
+
+		const FMovieSceneBinding* Binding = MovieScene->FindBinding(ObjectBinding.BindingID);
+		if (Binding)
+		{
+			CopyableBinding->Binding = *Binding;
+			for (UMovieSceneTrack* Track : Binding->GetTracks())
+			{
+				// Tracks suffer from the same issues as Spawnable's Object Templates (reference to a privately owned object). We'll manually serialize the tracks to copy them,
+				// and then restore them on paste.
+				UMovieSceneTrack* DuplicatedTrack = Cast<UMovieSceneTrack>(StaticDuplicateObject(Track, CopyableBinding));
+
+				CopyableBinding->Tracks.Add(DuplicatedTrack);
+			}
+		}
+
+		UMovieSceneFolder* Folder = nullptr;
+		for (UMovieSceneFolder* RootFolder : MovieScene->GetRootFolders())
+		{
+			Folder = RootFolder->FindFolderContaining(ObjectBinding.BindingID);
+			if (Folder && InFolders.Contains(Folder))
+			{
+				UMovieSceneFolder::CalculateFolderPath(Folder, InFolders, CopyableBinding->FolderPath);
+				break;
+			}
+		}
+	}
+
+	ExportObjectBindingsToText(Objects, /*out*/ ExportedText);
+}
+
+class FObjectBindingTextFactory : public FCustomizableTextObjectFactory
+{
+public:
+	FObjectBindingTextFactory(ISequencer& InSequencer)
+		: FCustomizableTextObjectFactory(GWarn)
+		, Sequencer(&InSequencer)
+	{
+	}
+
+	// FCustomizableTextObjectFactory implementation
+	virtual bool CanCreateClass(UClass* InObjectClass, bool& bOmitSubObjs) const override
+	{
+		if (InObjectClass->IsChildOf<UMovieSceneCopyableBinding>())
+		{
+			return true;
+		}
+
+		return Sequencer->GetSpawnRegister().CanSpawnObject(InObjectClass);
+	}
+
+
+	virtual void ProcessConstructedObject(UObject* NewObject) override
+	{
+		check(NewObject);
+
+		if (NewObject->IsA<UMovieSceneCopyableBinding>())
+		{
+			UMovieSceneCopyableBinding* CopyableBinding = Cast<UMovieSceneCopyableBinding>(NewObject);
+			NewCopyableBindings.Add(CopyableBinding);
+		}
+		else
+		{
+			NewSpawnableObjectTemplates.Add(NewObject);
+		}
+	}
+
+public:
+	TArray<UMovieSceneCopyableBinding*> NewCopyableBindings;
+	TArray<UObject*> NewSpawnableObjectTemplates;
+
+private:
+	ISequencer* Sequencer;
+};
+
+void ImportObjectBindingsFromText(ISequencer& InSequencer, const FString& TextToImport, /*out*/ TArray<UMovieSceneCopyableBinding*>& ImportedObjects)
+{
+	UPackage* TempPackage = NewObject<UPackage>(nullptr, TEXT("/Engine/Sequencer/Editor/Transient"), RF_Transient);
+	TempPackage->AddToRoot();
+
+	// Turn the text buffer into objects
+	FObjectBindingTextFactory Factory(InSequencer);
+	Factory.ProcessBuffer(TempPackage, RF_Transactional, TextToImport);
+	ImportedObjects = Factory.NewCopyableBindings;
+
+	// We had to explicitly serialize object templates due to them being a reference to a privately owned object. We now deserialize these object template copies
+	// and match them up with their MovieSceneCopyableBinding again.
+
+	int32 SpawnableObjectTemplateIndex = 0;
+	for (auto ImportedObject : ImportedObjects)
+	{
+		if (ImportedObject->Spawnable.GetGuid().IsValid() && SpawnableObjectTemplateIndex < Factory.NewSpawnableObjectTemplates.Num())
+		{
+			// This Spawnable Object Template is owned by our transient package, so you'll need to change the owner if you want to keep it later.
+			ImportedObject->SpawnableObjectTemplate = Factory.NewSpawnableObjectTemplates[SpawnableObjectTemplateIndex++];
+		}
+	}
+
+	// Remove the temp package from the root now that it has served its purpose
+	TempPackage->RemoveFromRoot();
+}
+
+bool FSequencerUtilities::PasteBindings(const FString& TextToImport, TSharedRef<ISequencer> Sequencer, FMovieScenePasteBindingsParams PasteBindingsParams, TArray<FMovieSceneBindingProxy>& OutBindings, TArray<FNotificationInfo>& OutErrors)
+{
+	UMovieSceneSequence* Sequence = Sequencer->GetFocusedMovieSceneSequence();
+	if (!Sequence)
+	{
+		return false;
+	}
+
+	UMovieScene* MovieScene = Sequence->GetMovieScene();
+	if (!MovieScene)
+	{
+		return false;
+	}
+
+	UWorld* World = Cast<UWorld>(Sequencer->GetPlaybackContext());
+	if (!World)
+	{
+		return false;
+	}
+
+	const FScopedTransaction Transaction(LOCTEXT("PasteBindings", "Paste Bindings"));
+
+	TMap<FGuid, FGuid> OldToNewGuidMap;
+	TArray<FGuid> PossessableGuids;
+	TArray<FGuid> SpawnableGuids;
+	TMap<FGuid, UMovieSceneFolder*> GuidToFolderMap;
+
+	TArray<FMovieSceneBinding> BindingsPasted;
+
+	const int NumTargets = FMath::Max(1, PasteBindingsParams.Bindings.Num());
+
+	for (int32 TargetIndex = 0; TargetIndex < NumTargets; ++TargetIndex)
+	{
+		TArray<UMovieSceneCopyableBinding*> ImportedBindings;
+		ImportObjectBindingsFromText(Sequencer.Get(), TextToImport, ImportedBindings);
+
+		if (ImportedBindings.Num() == 0)
+		{
+			return false;
+		}
+
+		for (UMovieSceneCopyableBinding* CopyableBinding : ImportedBindings)
+		{
+			// Clear transient flags on the imported tracks
+			for (UMovieSceneTrack* CopiedTrack : CopyableBinding->Tracks)
+			{
+				CopiedTrack->ClearFlags(RF_Transient);
+				TArray<UObject*> Subobjects;
+				GetObjectsWithOuter(CopiedTrack, Subobjects);
+				for (UObject* Subobject : Subobjects)
+				{
+					Subobject->ClearFlags(RF_Transient);
+				}
+			}
+
+			UMovieSceneFolder* ParentFolder = PasteBindingsParams.ParentFolder;
+
+			if (CopyableBinding->FolderPath.Num() > 0)
+			{
+				ParentFolder = UMovieSceneFolder::GetFolderWithPath(CopyableBinding->FolderPath, PasteBindingsParams.Folders, ParentFolder ? ParentFolder->GetChildFolders() : MovieScene->GetRootFolders());
+			}
+
+			if (CopyableBinding->Possessable.GetGuid().IsValid())
+			{
+				FGuid NewGuid = FGuid::NewGuid();
+
+				FMovieSceneBinding NewBinding(NewGuid, CopyableBinding->Binding.GetName(), CopyableBinding->Tracks);
+
+				FMovieScenePossessable NewPossessable = CopyableBinding->Possessable;
+				NewPossessable.SetGuid(NewGuid);
+
+				MovieScene->AddPossessable(NewPossessable, NewBinding);
+
+				OldToNewGuidMap.Add(CopyableBinding->Possessable.GetGuid(), NewGuid);
+
+				BindingsPasted.Add(NewBinding);
+
+				PossessableGuids.Add(NewGuid);
+
+				if (ParentFolder)
+				{
+					GuidToFolderMap.Add(NewGuid, ParentFolder);
+				}
+
+				if (FMovieScenePossessable* Possessable = MovieScene->FindPossessable(NewGuid))
+				{
+					if (TargetIndex < PasteBindingsParams.Bindings.Num())
+					{
+						Possessable->SetParent(PasteBindingsParams.Bindings[TargetIndex].BindingID, MovieScene);
+					}
+				}
+
+				TArray<AActor*> ActorsToDuplicate;
+				for (auto RuntimeObject : Sequencer->FindBoundObjects(CopyableBinding->Possessable.GetGuid(), Sequencer->GetFocusedTemplateID()))
+				{
+					AActor* Actor = Cast<AActor>(RuntimeObject.Get());
+					if (Actor)
+					{
+						ActorsToDuplicate.Add(Actor);
+					}
+				}
+
+				TArray<AActor*> DuplicatedActors;
+				if (ActorsToDuplicate.Num() != 0)
+				{
+					GEditor->SelectNone(false, true);
+					for (AActor* ActorToDuplicate : ActorsToDuplicate)
+					{
+						GEditor->SelectActor(ActorToDuplicate, true, false, false);
+					}
+
+					// Duplicate the bound actors
+					GEditor->edactDuplicateSelected(World->GetCurrentLevel(), false);
+
+					// Duplicating the bound actor through GEditor, edits the copy/paste clipboard. This is not desired from the user's 
+					// point of view since the user didn't explicitly invoke the copy operation. Instead, restore the copied contents
+					// of the clipboard after duplicating the actor
+					FPlatformApplicationMisc::ClipboardCopy(*TextToImport);
+
+					USelection* ActorSelection = GEditor->GetSelectedActors();
+					TArray<TWeakObjectPtr<AActor> > SelectedActors;
+					for (FSelectionIterator Iter(*ActorSelection); Iter; ++Iter)
+					{
+						AActor* Actor = Cast<AActor>(*Iter);
+						if (Actor)
+						{
+							DuplicatedActors.Add(Actor);
+						}
+					}
+
+					// Bind the duplicated actors
+					if (DuplicatedActors.Num())
+					{
+						ReplaceBindingWithActors(Sequencer, DuplicatedActors, FMovieSceneBindingProxy(NewGuid, Sequence));
+					}
+				}
+			}
+			else if (CopyableBinding->Spawnable.GetGuid().IsValid())
+			{
+				// We need to let the sequence create the spawnable so that it has everything set up properly internally.
+				// This is required to get spawnables with the correct references to object templates, object templates with
+				// correct owners, etc. However, making a new spawnable also creates the binding for us - this is a problem
+				// because we need to use our binding (which has tracks associated with it). To solve this, we let it create
+				// an object template based off of our (transient package owned) template, then find the newly created binding
+				// and update it.
+
+				FGuid NewGuid = MakeNewSpawnable(Sequencer, *CopyableBinding->SpawnableObjectTemplate, nullptr, false, FName(*CopyableBinding->Spawnable.GetName()));
+
+				FMovieSceneBinding NewBinding(NewGuid, CopyableBinding->Binding.GetName(), CopyableBinding->Tracks);
+				FMovieSceneSpawnable* Spawnable = MovieScene->FindSpawnable(NewGuid);
+
+				// Copy the name of the original spawnable too.
+				Spawnable->SetName(CopyableBinding->Spawnable.GetName());
+
+				// Clear the transient flags on the copyable binding before assigning to the new spawnable
+				for (auto Track : NewBinding.GetTracks())
+				{
+					Track->ClearFlags(RF_Transient);
+					for (auto Section : Track->GetAllSections())
+					{
+						Section->ClearFlags(RF_Transient);
+					}
+				}
+
+				// Replace the auto-generated binding with our deserialized bindings (which has our tracks)
+				MovieScene->ReplaceBinding(NewGuid, NewBinding);
+
+				OldToNewGuidMap.Add(CopyableBinding->Spawnable.GetGuid(), NewGuid);
+
+				BindingsPasted.Add(NewBinding);
+
+				SpawnableGuids.Add(NewGuid);
+
+				if (ParentFolder)
+				{
+					GuidToFolderMap.Add(NewGuid, ParentFolder);
+				}
+			}
+		}
+	}
+
+	// Fix up parent guids
+	for (auto PossessableGuid : PossessableGuids)
+	{
+		FMovieScenePossessable* Possessable = MovieScene->FindPossessable(PossessableGuid);
+		if (Possessable && OldToNewGuidMap.Contains(Possessable->GetParent()) && PossessableGuid != OldToNewGuidMap[Possessable->GetParent()])
+		{
+			Possessable->SetParent(OldToNewGuidMap[Possessable->GetParent()], MovieScene);
+		}
+	}
+
+	// Fix possessable actor bindings
+	for (int32 PossessableGuidIndex = 0; PossessableGuidIndex < PossessableGuids.Num(); ++PossessableGuidIndex)
+	{
+		FMovieScenePossessable* Possessable = MovieScene->FindPossessable(PossessableGuids[PossessableGuidIndex]);
+		UWorld* PlaybackContext = Cast<UWorld>(Sequencer->GetPlaybackContext());
+		if (Possessable && PlaybackContext)
+		{
+			for (TActorIterator<AActor> ActorItr(PlaybackContext); ActorItr; ++ActorItr)
+			{
+				AActor* Actor = *ActorItr;
+				if (Actor && Actor->GetActorLabel() == *Possessable->GetName())
+				{
+					FGuid ExistingGuid = Sequencer->FindObjectId(*Actor, Sequencer->GetFocusedTemplateID());
+
+					if (!ExistingGuid.IsValid())
+					{
+						FGuid NewGuid = AssignActor(Sequencer, Actor, Possessable->GetGuid());
+
+						// If assigning produces a new guid, update the possessable guids and the bindings pasted data
+						if (NewGuid.IsValid())
+						{
+							for (auto BindingPasted : BindingsPasted)
+							{
+								if (BindingPasted.GetObjectGuid() == PossessableGuids[PossessableGuidIndex])
+								{
+									BindingPasted.SetObjectGuid(NewGuid);
+								}
+							}
+
+							if (GuidToFolderMap.Contains(PossessableGuids[PossessableGuidIndex]))
+							{
+								GuidToFolderMap.Add(NewGuid, GuidToFolderMap[PossessableGuids[PossessableGuidIndex]]);
+								GuidToFolderMap.Remove(PossessableGuids[PossessableGuidIndex]);
+							}
+
+							PossessableGuids[PossessableGuidIndex] = NewGuid;
+						}
+					}
+				}
+			}
+		}
+	}
+	
+	// Set up folders
+	for (auto PossessableGuid : PossessableGuids)
+	{
+		FMovieScenePossessable* Possessable = MovieScene->FindPossessable(PossessableGuid);
+		if (Possessable && !Possessable->GetParent().IsValid())
+		{
+			if (GuidToFolderMap.Contains(PossessableGuid))
+			{
+				GuidToFolderMap[PossessableGuid]->AddChildObjectBinding(PossessableGuid);
+			}
+		}
+	}
+	for (auto SpawnableGuid : SpawnableGuids)
+	{
+		FMovieSceneSpawnable* Spawnable = MovieScene->FindSpawnable(SpawnableGuid);
+		if (Spawnable)
+		{
+			if (GuidToFolderMap.Contains(SpawnableGuid))
+			{
+				GuidToFolderMap[SpawnableGuid]->AddChildObjectBinding(SpawnableGuid);
+			}
+		}
+	}
+
+	Sequencer->OnMovieSceneBindingsPasted().Broadcast(BindingsPasted);
+
+	// Refresh all immediately so that spawned actors will be generated immediately
+	Sequencer->ForceEvaluate();
+
+	// Fix possessable component bindings
+	for (auto PossessableGuid : PossessableGuids)
+	{
+		// If a possessable guid does not have any bound objects, they might be 
+		// possessable components for spawnables, so they need to be remapped
+		if (Sequencer->FindBoundObjects(PossessableGuid, Sequencer->GetFocusedTemplateID()).Num() == 0)
+		{
+			FMovieScenePossessable* Possessable = MovieScene->FindPossessable(PossessableGuid);
+			if (Possessable)
+			{
+				FGuid ParentGuid = Possessable->GetParent();
+				for (TWeakObjectPtr<> WeakObject : Sequencer->FindBoundObjects(ParentGuid, Sequencer->GetFocusedTemplateID()))
+				{
+					if (AActor* SpawnedActor = Cast<AActor>(WeakObject.Get()))
+					{
+						for (UActorComponent* Component : SpawnedActor->GetComponents())
+						{
+							if (Component->GetName() == Possessable->GetName())
+							{
+								Sequence->BindPossessableObject(PossessableGuid, *Component, SpawnedActor);
+								break;
+							}
+						}
+					}
+				}
+
+				// If the parent doesn't actually exist, clear it.
+				FMovieScenePossessable* PossessableParent = MovieScene->FindPossessable(ParentGuid);
+				FMovieSceneSpawnable* SpawnableParent = MovieScene->FindSpawnable(ParentGuid);
+				if (!PossessableParent && !SpawnableParent)
+				{
+					Possessable->SetParent(FGuid(), MovieScene);
+				}
+				else if (SpawnableParent)
+				{
+					SpawnableParent->AddChildPossessable(PossessableGuid);
+				}
+			}
+		}
+	}
+
+	// Remap bindings in sections (ie. attach tracks)
+	for (TPair<FGuid, FGuid> GuidPair : OldToNewGuidMap)
+	{
+		UpdateBindingIDs(Sequencer, GuidPair.Key, GuidPair.Value);
+	}
+
+	for (auto BindingPasted : BindingsPasted)
+	{
+		OutBindings.Add(FMovieSceneBindingProxy(BindingPasted.GetObjectGuid(), Sequence));
+	}
+
+	return true; 
+}
+
+bool FSequencerUtilities::CanPasteBindings(TSharedRef<ISequencer> Sequencer, const FString& TextToImport)
+{
+	FObjectBindingTextFactory ObjectBindingFactory(Sequencer.Get());
+	return ObjectBindingFactory.CanCreateObjectsFromText(TextToImport);
+}
+
+FGuid FSequencerUtilities::CreateBinding(TSharedRef<ISequencer> Sequencer, UObject& InObject, const FString& InName)
+{
+	const FScopedTransaction Transaction(LOCTEXT("CreateBinding", "Create New Binding"));
+
+	UMovieSceneSequence* OwnerSequence = Sequencer->GetFocusedMovieSceneSequence();
+	UMovieScene* OwnerMovieScene = OwnerSequence->GetMovieScene();
+
+	OwnerSequence->Modify();
+	OwnerMovieScene->Modify();
+
+	const FGuid PossessableGuid = OwnerMovieScene->AddPossessable(InName, InObject.GetClass());
+
+	// Attempt to use the parent as a context if necessary
+	UObject* ParentObject = OwnerSequence->GetParentObject(&InObject);
+	UObject* BindingContext = Sequencer->GetPlaybackContext();
+
+	AActor* ParentActorAdded = nullptr;
+	FGuid ParentGuid;
+
+	if (ParentObject)
+	{
+		// Ensure we have possessed the outer object, if necessary
+		ParentGuid = Sequencer->GetHandleToObject(ParentObject, false);
+		if (!ParentGuid.IsValid())
+		{
+			ParentGuid = Sequencer->GetHandleToObject(ParentObject);
+			ParentActorAdded = Cast<AActor>(ParentObject);
+		}
+
+		if (OwnerSequence->AreParentContextsSignificant())
+		{
+			BindingContext = ParentObject;
+		}
+
+		// Set up parent/child guids for possessables within spawnables
+		if (ParentGuid.IsValid())
+		{
+			FMovieScenePossessable* ChildPossessable = OwnerMovieScene->FindPossessable(PossessableGuid);
+			if (ensure(ChildPossessable))
+			{
+				ChildPossessable->SetParent(ParentGuid, OwnerMovieScene);
+			}
+
+			FMovieSceneSpawnable* ParentSpawnable = OwnerMovieScene->FindSpawnable(ParentGuid);
+			if (ParentSpawnable)
+			{
+				ParentSpawnable->AddChildPossessable(PossessableGuid);
+			}
+		}
+	}
+
+	if (!OwnerMovieScene->FindPossessable(PossessableGuid)->BindSpawnableObject(Sequencer->GetFocusedTemplateID(), &InObject, &Sequencer.Get()))
+	{
+		OwnerSequence->BindPossessableObject(PossessableGuid, InObject, BindingContext);
+	}
+
+	// Broadcast if a parent actor was added as a result of adding this object
+	if (ParentActorAdded && ParentGuid.IsValid())
+	{
+		Sequencer->OnActorAddedToSequencer().Broadcast(ParentActorAdded, ParentGuid);
+	}
+
+	return PossessableGuid;
+}
+
+void FSequencerUtilities::UpdateBindingIDs(TSharedRef<ISequencer> Sequencer, FGuid OldGuid, FGuid NewGuid)
+{
+	UMovieSceneCompiledDataManager* CompiledDataManager = FindObject<UMovieSceneCompiledDataManager>(GetTransientPackage(), TEXT("SequencerCompiledDataManager"));
+	if (!CompiledDataManager)
+	{
+		CompiledDataManager = NewObject<UMovieSceneCompiledDataManager>(GetTransientPackage(), "SequencerCompiledDataManager");
+	}
+
+	if (!CompiledDataManager)
+	{
+		return;
+	}
+
+	const FMovieSceneSequenceHierarchy* Hierarchy = CompiledDataManager->FindHierarchy(Sequencer->GetEvaluationTemplate().GetCompiledDataID());
+
+	FMovieSceneSequenceIDRef FocusedGuid = Sequencer->GetFocusedTemplateID();
+
+	TMap<UE::MovieScene::FFixedObjectBindingID, UE::MovieScene::FFixedObjectBindingID> OldFixedToNewFixedMap;
+	OldFixedToNewFixedMap.Add(UE::MovieScene::FFixedObjectBindingID(OldGuid, FocusedGuid), UE::MovieScene::FFixedObjectBindingID(NewGuid, FocusedGuid));
+
+	if (UMovieScene* MovieScene = Sequencer->GetRootMovieSceneSequence()->GetMovieScene())
+	{
+		for (UMovieSceneSection* Section : MovieScene->GetAllSections())
+		{
+			if (Section)
+			{
+				Section->OnBindingIDsUpdated(OldFixedToNewFixedMap, Sequencer->GetRootTemplateID(), Hierarchy, *Sequencer);
+			}
+		}
+	}
+
+	if (Hierarchy)
+	{
+		for (const TTuple<FMovieSceneSequenceID, FMovieSceneSubSequenceData>& Pair : Hierarchy->AllSubSequenceData())
+		{
+			if (UMovieSceneSequence* Sequence = Pair.Value.GetSequence())
+			{
+				if (UMovieScene* MovieScene = Sequence->GetMovieScene())
+				{
+					for (UMovieSceneSection* Section : MovieScene->GetAllSections())
+					{
+						if (Section)
+						{
+							Section->OnBindingIDsUpdated(OldFixedToNewFixedMap, Pair.Key, Hierarchy, *Sequencer);
+						}
+					}
+				}
+			}
+		}
+	}
+}
+
+FGuid FSequencerUtilities::AssignActor(TSharedRef<ISequencer> Sequencer, AActor* Actor, FGuid InObjectBinding)
+{
 	if (Actor == nullptr)
 	{
 		return FGuid();
 	}
 
-	UMovieSceneSequence* OwnerSequence = InSequencerPtr->GetFocusedMovieSceneSequence();
+	UMovieSceneSequence* OwnerSequence = Sequencer->GetFocusedMovieSceneSequence();
 	UMovieScene* OwnerMovieScene = OwnerSequence->GetMovieScene();
 
 	if (OwnerMovieScene->IsReadOnly())
 	{
-		FSequencerUtilities::ShowReadOnlyError();
+		ShowReadOnlyError();
 		return FGuid();
 	}
 
@@ -426,7 +2444,7 @@
 	OwnerSequence->Modify();
 	OwnerMovieScene->Modify();
 
-	TArrayView<TWeakObjectPtr<>> RuntimeObjects = InSequencerPtr->FindObjectsInCurrentSequence(InObjectBinding);
+	TArrayView<TWeakObjectPtr<>> RuntimeObjects = Sequencer->FindObjectsInCurrentSequence(InObjectBinding);
 
 	UObject* RuntimeObject = RuntimeObjects.Num() ? RuntimeObjects[0].Get() : nullptr;
 
@@ -435,7 +2453,7 @@
 	FGuid NewGuid;
 	{
 		// Get the object guid to assign, remove the binding if it already exists
-		FGuid ParentGuid = InSequencerPtr->FindObjectId(*Actor, InSequencerPtr->GetFocusedTemplateID());
+		FGuid ParentGuid = Sequencer->FindObjectId(*Actor, Sequencer->GetFocusedTemplateID());
 		FString NewActorLabel = Actor->GetActorLabel();
 		if (ParentGuid.IsValid())
 		{
@@ -446,9 +2464,9 @@
 		// Add this object
 		NewPossessableActor = FMovieScenePossessable(NewActorLabel, Actor->GetClass());
 		NewGuid = NewPossessableActor.GetGuid();
-		if (!NewPossessableActor.BindSpawnableObject(InSequencerPtr->GetFocusedTemplateID(), Actor, InSequencerPtr))
+		if (!NewPossessableActor.BindSpawnableObject(Sequencer->GetFocusedTemplateID(), Actor, &Sequencer.Get()))
 		{
-			OwnerSequence->BindPossessableObject(NewPossessableActor.GetGuid(), *Actor, InSequencerPtr->GetPlaybackContext());
+			OwnerSequence->BindPossessableObject(NewPossessableActor.GetGuid(), *Actor, Sequencer->GetPlaybackContext());
 		}
 
 		// Defer replacing this object until the components have been updated
@@ -456,10 +2474,10 @@
 
 	auto UpdateComponent = [&](FGuid OldComponentGuid, UActorComponent* NewComponent)
 	{
-		FMovieSceneSequenceIDRef FocusedGuid = InSequencerPtr->GetFocusedTemplateID();
+		FMovieSceneSequenceIDRef FocusedGuid = Sequencer->GetFocusedTemplateID();
 
 		// Get the object guid to assign, remove the binding if it already exists
-		FGuid NewComponentGuid = InSequencerPtr->FindObjectId(*NewComponent, FocusedGuid);
+		FGuid NewComponentGuid = Sequencer->FindObjectId(*NewComponent, FocusedGuid);
 		if (NewComponentGuid.IsValid())
 		{
 			OwnerMovieScene->RemovePossessable(NewComponentGuid);
@@ -473,8 +2491,8 @@
 		// Replace
 		OwnerMovieScene->ReplacePossessable(OldComponentGuid, NewPossessable);
 		OwnerSequence->UnbindPossessableObjects(OldComponentGuid);
-		InSequencerPtr->State.Invalidate(OldComponentGuid, FocusedGuid);
-		InSequencerPtr->State.Invalidate(NewPossessable.GetGuid(), FocusedGuid);
+		Sequencer->State.Invalidate(OldComponentGuid, FocusedGuid);
+		Sequencer->State.Invalidate(NewPossessable.GetGuid(), FocusedGuid);
 
 		FMovieScenePossessable* ThisPossessable = OwnerMovieScene->FindPossessable(NewPossessable.GetGuid());
 		if (ensure(ThisPossessable))
@@ -491,7 +2509,7 @@
 		{
 			if (ComponentToReplace != nullptr)
 			{
-				FGuid ComponentGuid = InSequencerPtr->FindObjectId(*ComponentToReplace, InSequencerPtr->GetFocusedTemplateID());
+				FGuid ComponentGuid = Sequencer->FindObjectId(*ComponentToReplace, Sequencer->GetFocusedTemplateID());
 				if (ComponentGuid.IsValid())
 				{
 					bool bComponentWasUpdated = false;
@@ -539,12 +2557,12 @@
 	OwnerMovieScene->ReplacePossessable(InObjectBinding, NewPossessableActor);
 	OwnerSequence->UnbindPossessableObjects(InObjectBinding);
 
-	InSequencerPtr->State.Invalidate(InObjectBinding, InSequencerPtr->GetFocusedTemplateID());
-	InSequencerPtr->State.Invalidate(NewPossessableActor.GetGuid(), InSequencerPtr->GetFocusedTemplateID());
+	Sequencer->State.Invalidate(InObjectBinding, Sequencer->GetFocusedTemplateID());
+	Sequencer->State.Invalidate(NewPossessableActor.GetGuid(), Sequencer->GetFocusedTemplateID());
 
 	// Try to fix up folders
 	TArray<UMovieSceneFolder*> FoldersToCheck;
-	for (UMovieSceneFolder* Folder : InSequencerPtr->GetFocusedMovieSceneSequence()->GetMovieScene()->GetRootFolders())
+	for (UMovieSceneFolder* Folder : Sequencer->GetFocusedMovieSceneSequence()->GetMovieScene()->GetRootFolders())
 	{
 		FoldersToCheck.Add(Folder);
 	}
@@ -566,54 +2584,218 @@
 		}
 	}
 
-	InSequencerPtr->RestorePreAnimatedState();
+	Sequencer->RestorePreAnimatedState();
 
-	InSequencerPtr->NotifyMovieSceneDataChanged(EMovieSceneDataChangeType::MovieSceneStructureItemsChanged);
+	Sequencer->NotifyMovieSceneDataChanged(EMovieSceneDataChangeType::MovieSceneStructureItemsChanged);
 
 	return NewGuid;
 }
 
-void FSequencerUtilities::UpdateBindingIDs(ISequencer* InSequencerPtr, UMovieSceneCompiledDataManager* InCompiledDataManagerPtr, FGuid OldGuid, FGuid NewGuid)
+void FSequencerUtilities::AddActorsToBinding(TSharedRef<ISequencer> Sequencer, const TArray<AActor*>& Actors, const FMovieSceneBindingProxy& ObjectBinding)
 {
-	FMovieSceneSequenceIDRef FocusedGuid = InSequencerPtr->GetFocusedTemplateID();
+	if (!Actors.Num())
+	{
+		return;
+	}
 
-	TMap<UE::MovieScene::FFixedObjectBindingID, UE::MovieScene::FFixedObjectBindingID> OldFixedToNewFixedMap;
-	OldFixedToNewFixedMap.Add(UE::MovieScene::FFixedObjectBindingID(OldGuid, FocusedGuid), UE::MovieScene::FFixedObjectBindingID(NewGuid, FocusedGuid));
+	UMovieSceneSequence* Sequence = Sequencer->GetFocusedMovieSceneSequence();
+	if (!Sequence)
+	{
+		return;
+	}
 
-	const FMovieSceneSequenceHierarchy* Hierarchy = InCompiledDataManagerPtr->FindHierarchy(InSequencerPtr->GetEvaluationTemplate().GetCompiledDataID());
+	UMovieScene* MovieScene = Sequence->GetMovieScene();
+	if (!MovieScene)
+	{
+		return;
+	}
 
-	if (UMovieScene* MovieScene = InSequencerPtr->GetRootMovieSceneSequence()->GetMovieScene())
+	UClass* ActorClass = nullptr;
+	int32 NumRuntimeObjects = 0;
+
+	FGuid Guid = ObjectBinding.BindingID;
+	TArrayView<TWeakObjectPtr<>> ObjectsInCurrentSequence = Sequencer->FindObjectsInCurrentSequence(Guid);
+
+	for (TWeakObjectPtr<> Ptr : ObjectsInCurrentSequence)
 	{
-		for (UMovieSceneSection* Section : MovieScene->GetAllSections())
+		if (const AActor* Actor = Cast<AActor>(Ptr.Get()))
 		{
-			if (Section)
+			ActorClass = Actor->GetClass();
+			++NumRuntimeObjects;
+		}
+	}
+
+	FScopedTransaction AddSelectedToBinding(LOCTEXT("AddSelectedToBinding", "Add Selected to Binding"));
+
+	Sequence->Modify();
+	MovieScene->Modify();
+
+	// Bind objects
+	int32 NumObjectsAdded = 0;
+	for (AActor* ActorToAdd : Actors)
+	{
+		if (!ObjectsInCurrentSequence.Contains(ActorToAdd))
+		{
+			if (ActorClass == nullptr || UClass::FindCommonBase(ActorToAdd->GetClass(), ActorClass) != nullptr)
 			{
-				Section->OnBindingIDsUpdated(OldFixedToNewFixedMap, InSequencerPtr->GetRootTemplateID(), Hierarchy, *InSequencerPtr);
+				if (ActorClass == nullptr)
+				{
+					ActorClass = ActorToAdd->GetClass();
+				}
+
+				ActorToAdd->Modify();
+				Sequence->BindPossessableObject(Guid, *ActorToAdd, Sequencer->GetPlaybackContext());
+				++NumObjectsAdded;
+			}
+			else
+			{
+				const FText NotificationText = FText::Format(LOCTEXT("UnableToAssignObject", "Cannot assign object {0}. Expected class {1}"), FText::FromString(ActorToAdd->GetName()), FText::FromString(ActorClass->GetName()));
+				FNotificationInfo Info(NotificationText);
+				Info.ExpireDuration = 3.f;
+				Info.bUseLargeFont = false;
+				FSlateNotificationManager::Get().AddNotification(Info);
 			}
 		}
 	}
 
-	if (Hierarchy)
+	// Update label
+	if (NumRuntimeObjects + NumObjectsAdded > 0)
 	{
-		for (const TTuple<FMovieSceneSequenceID, FMovieSceneSubSequenceData>& Pair : Hierarchy->AllSubSequenceData())
+		FMovieScenePossessable* Possessable = MovieScene->FindPossessable(Guid);
+		if (Possessable && ActorClass != nullptr)
 		{
-			if (UMovieSceneSequence* Sequence = Pair.Value.GetSequence())
+			if (NumRuntimeObjects + NumObjectsAdded > 1)
 			{
-				if (UMovieScene* MovieScene = Sequence->GetMovieScene())
-				{
-					for (UMovieSceneSection* Section : MovieScene->GetAllSections())
-					{
-						if (Section)
-						{
-							Section->OnBindingIDsUpdated(OldFixedToNewFixedMap, Pair.Key, Hierarchy, *InSequencerPtr);
-						}
-					}
-				}
+				FString NewLabel = ActorClass->GetName() + FString::Printf(TEXT(" (%d)"), NumRuntimeObjects + NumObjectsAdded);
+				Possessable->SetName(NewLabel);
+			}
+			else if (NumObjectsAdded > 0 && Actors.Num() > 0)
+			{
+				Possessable->SetName(Actors[0]->GetActorLabel());
+			}
+
+			Possessable->SetPossessedObjectClass(ActorClass);
+		}
+	}
+
+	Sequencer->RestorePreAnimatedState();
+
+	Sequencer->NotifyMovieSceneDataChanged(EMovieSceneDataChangeType::MovieSceneStructureItemsChanged);
+}
+
+void FSequencerUtilities::ReplaceBindingWithActors(TSharedRef<ISequencer> Sequencer, const TArray<AActor*>& Actors, const FMovieSceneBindingProxy& ObjectBinding)
+{
+	FScopedTransaction ReplaceBindingWithActors(LOCTEXT("ReplaceBindingWithActors", "Replace Binding with Actors"));
+
+	FGuid Guid = ObjectBinding.BindingID;
+	TArray<AActor*> ExistingActors;
+	for (TWeakObjectPtr<> Ptr : Sequencer->FindObjectsInCurrentSequence(Guid))
+	{
+		if (AActor* Actor = Cast<AActor>(Ptr.Get()))
+		{
+			if (!Actors.Contains(Actor))
+			{
+				ExistingActors.Add(Actor);
 			}
 		}
 	}
+
+	RemoveActorsFromBinding(Sequencer, ExistingActors, ObjectBinding);
+
+	TArray<AActor*> NewActors;
+	for (AActor* NewActor : Actors)
+	{
+		if (!ExistingActors.Contains(NewActor))
+		{
+			NewActors.Add(NewActor);
+		}
+	}
+
+	AddActorsToBinding(Sequencer, NewActors, ObjectBinding);
 }
 
+void FSequencerUtilities::RemoveActorsFromBinding(TSharedRef<ISequencer> Sequencer, const TArray<AActor*>& Actors, const FMovieSceneBindingProxy& ObjectBinding)
+{
+	if (!Actors.Num())
+	{
+		return;
+	}
+
+	UMovieSceneSequence* Sequence = Sequencer->GetFocusedMovieSceneSequence();
+	if (!Sequence)
+	{
+		return;
+	}
+
+	UMovieScene* MovieScene = Sequence->GetMovieScene();
+	if (!MovieScene)
+	{
+		return;
+	}
+
+	UClass* ActorClass = nullptr;
+	int32 NumRuntimeObjects = 0;
+
+	FGuid Guid = ObjectBinding.BindingID;
+	for (TWeakObjectPtr<> Ptr : Sequencer->FindObjectsInCurrentSequence(Guid))
+	{
+		if (const AActor* Actor = Cast<AActor>(Ptr.Get()))
+		{
+			ActorClass = Actor->GetClass();
+			++NumRuntimeObjects;
+		}
+	}
+
+	FScopedTransaction RemoveSelectedFromBinding(LOCTEXT("RemoveSelectedFromBinding", "Remove Selected from Binding"));
+
+	TArray<UObject*> ObjectsToRemove;
+	for (AActor* ActorToRemove : Actors)
+	{
+		// Restore state on any components
+		for (UActorComponent* Component : TInlineComponentArray<UActorComponent*>(ActorToRemove))
+		{
+			if (Component)
+			{
+				Sequencer->PreAnimatedState.RestorePreAnimatedState(*Component);
+			}
+		}
+
+		// Restore state on the object itself
+		Sequencer->PreAnimatedState.RestorePreAnimatedState(*ActorToRemove);
+
+		ActorToRemove->Modify();
+
+		ObjectsToRemove.Add(ActorToRemove);
+	}
+
+	Sequence->Modify();
+	MovieScene->Modify();
+
+	// Unbind objects
+	Sequence->UnbindObjects(Guid, ObjectsToRemove, Sequencer->GetPlaybackContext());
+
+	// Update label
+	if (NumRuntimeObjects - ObjectsToRemove.Num() > 0)
+	{
+		FMovieScenePossessable* Possessable = MovieScene->FindPossessable(Guid);
+		if (Possessable && ActorClass != nullptr)
+		{
+			if (NumRuntimeObjects - ObjectsToRemove.Num() > 1)
+			{
+				FString NewLabel = ActorClass->GetName() + FString::Printf(TEXT(" (%d)"), NumRuntimeObjects - ObjectsToRemove.Num());
+
+				Possessable->SetName(NewLabel);
+			}
+			else if (ObjectsToRemove.Num() > 0 && Actors.Num() > 0)
+			{
+				Possessable->SetName(Actors[0]->GetActorLabel());
+			}
+		}
+	}
+
+	Sequencer->NotifyMovieSceneDataChanged(EMovieSceneDataChangeType::MovieSceneStructureItemsChanged);
+}
+
 void FSequencerUtilities::ShowReadOnlyError()
 {
 	FNotificationInfo Info(LOCTEXT("SequenceReadOnly", "Sequence is read only."));
@@ -621,6 +2803,4 @@
 	FSlateNotificationManager::Get().AddNotification(Info)->SetCompletionState(SNotificationItem::CS_Fail);
 }
 
-
 #undef LOCTEXT_NAMESPACE
-PRAGMA_ENABLE_OPTIMIZATION
diff --git a/Engine/Source/Editor/Sequencer/Public/SequencerUtilities.h b/Engine/Source/Editor/Sequencer/Public/SequencerUtilities.h
--- a/Engine/Source/Editor/Sequencer/Public/SequencerUtilities.h
+++ b/Engine/Source/Editor/Sequencer/Public/SequencerUtilities.h
@@ -5,20 +5,117 @@
 #include "CoreMinimal.h"
 #include "Widgets/SWidget.h"
 #include "Framework/SlateDelegates.h"
+#include "MovieSceneBindingProxy.h"
+#include "SequencerUtilities.generated.h"
 
+class ACineCameraActor;
 template< typename ObjectType > class TAttribute;
+class UActorFactory;
+class UMovieSceneFolder;
 class UMovieSceneTrack;
 class UMovieSceneSection;
+class UMovieSceneSequence;
 class ISequencer;
 class FMenuBuilder;
+struct FMovieSceneBinding;
+struct FMovieScenePossessable;
+struct FMovieSceneSpawnable;
+struct FNotificationInfo;
 class ULevelSequence;
-class UMovieSceneCompiledDataManager;
 enum class EMovieSceneBlendType : uint8;
 
+/* Paste folders params */
+USTRUCT(BlueprintType)
+struct FMovieScenePasteFoldersParams
+{
+	GENERATED_BODY()
+
+	FMovieScenePasteFoldersParams() {}
+	FMovieScenePasteFoldersParams(UMovieSceneSequence* InSequence, UMovieSceneFolder* InParentFolder = nullptr)
+		: Sequence(InSequence)
+		, ParentFolder(InParentFolder) {}
+
+	UPROPERTY(BlueprintReadWrite, Category = "Movie Scene")
+	TObjectPtr<UMovieSceneSequence> Sequence;
+	
+	UPROPERTY(BlueprintReadWrite, Category = "Movie Scene")
+	TObjectPtr<UMovieSceneFolder> ParentFolder;
+};
+
+/* Paste sections params */
+USTRUCT(BlueprintType)
+struct FMovieScenePasteSectionsParams
+{
+	GENERATED_BODY()
+
+	FMovieScenePasteSectionsParams() {}
+	FMovieScenePasteSectionsParams(const TArray<UMovieSceneTrack*>& InTracks, const TArray<int32>& InTrackRowIndices, FFrameTime InTime)
+		: Tracks(InTracks)
+		, TrackRowIndices(InTrackRowIndices)
+		, Time(InTime) {}
+
+	UPROPERTY(BlueprintReadWrite, Category = "Movie Scene")
+	TArray<TObjectPtr<UMovieSceneTrack>> Tracks;
+
+	UPROPERTY(BlueprintReadWrite, Category = "Movie Scene")
+	TArray<int32> TrackRowIndices;
+
+	UPROPERTY(BlueprintReadWrite, Category = "Movie Scene")
+	FFrameTime Time;
+};
+
+/* Paste tracks params */
+USTRUCT(BlueprintType)
+struct FMovieScenePasteTracksParams
+{
+	GENERATED_BODY()
+
+	FMovieScenePasteTracksParams() {}
+	FMovieScenePasteTracksParams(UMovieSceneSequence* InSequence, const TArray<FMovieSceneBindingProxy>& InBindings = TArray<FMovieSceneBindingProxy>(), UMovieSceneFolder* InParentFolder = nullptr, const TArray<UMovieSceneFolder*>& InFolders = TArray<UMovieSceneFolder*>())
+		: Sequence(InSequence)
+		, Bindings(InBindings)
+		, ParentFolder(InParentFolder)
+		, Folders(InFolders) {}
+
+	UPROPERTY(BlueprintReadWrite, Category = "Movie Scene")
+	TObjectPtr<UMovieSceneSequence> Sequence;
+
+	UPROPERTY(BlueprintReadWrite, Category = "Movie Scene")
+	TArray<FMovieSceneBindingProxy> Bindings;
+	
+	UPROPERTY(BlueprintReadWrite, Category = "Movie Scene")
+	TObjectPtr<UMovieSceneFolder> ParentFolder;
+	
+	UPROPERTY(BlueprintReadWrite, Category = "Movie Scene")
+	TArray<TObjectPtr<UMovieSceneFolder>> Folders;
+};
+
+/* Paste bindings params */
+USTRUCT(BlueprintType)
+struct FMovieScenePasteBindingsParams
+{
+	GENERATED_BODY()
+
+	FMovieScenePasteBindingsParams(const TArray<FMovieSceneBindingProxy>& InBindings = TArray<FMovieSceneBindingProxy>(), UMovieSceneFolder* InParentFolder = nullptr, const TArray<UMovieSceneFolder*>& InFolders = TArray<UMovieSceneFolder*>())
+		: Bindings(InBindings)
+		, ParentFolder(InParentFolder)
+		, Folders(InFolders) {}
+
+	UPROPERTY(BlueprintReadWrite, Category = "Movie Scene")
+	TArray<FMovieSceneBindingProxy> Bindings;
+	
+	UPROPERTY(BlueprintReadWrite, Category = "Movie Scene")
+	TObjectPtr<UMovieSceneFolder> ParentFolder;
+
+	UPROPERTY(BlueprintReadWrite, Category = "Movie Scene")
+	TArray<TObjectPtr<UMovieSceneFolder>> Folders;
+};
+
 struct SEQUENCER_API FSequencerUtilities
 {
 	/* Creates a button (used for +Section) that opens a ComboButton with a user-defined sub-menu content. */
 	static TSharedRef<SWidget> MakeAddButton(FText HoverText, FOnGetContent MenuContent, const TAttribute<bool>& HoverState, TWeakPtr<ISequencer> InSequencer);
+	
 	/* Creates a button (used for +Section) that fires a user-defined OnClick response with no sub-menu. */
 	static TSharedRef<SWidget> MakeAddButton(FText HoverText, FOnClicked OnClicked, const TAttribute<bool>& HoverState, TWeakPtr<ISequencer> InSequencer);
 
@@ -38,9 +135,51 @@ struct SEQUENCER_API FSequencerUtilities
 	 */
 	static FName GetUniqueName(FName CandidateName, const TArray<FName>& ExistingNames);
 
-	static FGuid DoAssignActor(ISequencer* InSequencerPtr, AActor* Actor, FGuid InObjectBinding);
+	/** Add existing actors to Sequencer */
+	static TArray<FGuid> AddActors(TSharedRef<ISequencer> Sequencer, const TArray<TWeakObjectPtr<AActor> >& InActors);
+
+	/** Create a new camera actor and add it to Sequencer */
+	static FGuid CreateCamera(TSharedRef<ISequencer> Sequencer, const bool bSpawnable, ACineCameraActor*& OutActor);
+
+	/** Create a new camera from a rig and add it to Sequencer */
+	static FGuid CreateCameraWithRig(TSharedRef<ISequencer> Sequencer, AActor* Actor, const bool bSpawnable, ACineCameraActor*& OutActor);
+
+	static FGuid MakeNewSpawnable(TSharedRef<ISequencer> Sequencer, UObject& SourceObject, UActorFactory* ActorFactory = nullptr, bool bSetupDefaults = true, FName SpawnableName = NAME_None);
+
+	/** Convert the requested possessable to spawnable. If there are multiple objects assigned to the possessable, multiple spawnables will be created */
+	static TArray<FMovieSceneSpawnable*> ConvertToSpawnable(TSharedRef<ISequencer> Sequencer, FGuid PossessableGuid);
+	
+	/** Convert the requested spawnable to possessable */
+	static FMovieScenePossessable* ConvertToPossessable(TSharedRef<ISequencer> Sequencer, FGuid SpawnableGuid);
+
+	/** Copy/paste folders */
+	static void CopyFolders(const TArray<UMovieSceneFolder*>& Folders, FString& ExportedText);
+	static bool PasteFolders(const FString& TextToImport, FMovieScenePasteFoldersParams PasteFoldersParams, TArray<UMovieSceneFolder*>& OutFolders, TArray<FNotificationInfo>& OutErrors);
+	static bool CanPasteFolders(const FString& TextToImport);
+
+	/** Copy/paste tracks */
+	static void CopyTracks(const TArray<UMovieSceneTrack*>& Tracks, const TArray<UMovieSceneFolder*>& InFolders, FString& ExportedText);
+	static bool PasteTracks(const FString& TextToImport, FMovieScenePasteTracksParams PasteTracksParams, TArray<UMovieSceneTrack*>& OutTracks, TArray<FNotificationInfo>& OutErrors);
+	static bool CanPasteTracks(const FString& TextToImport);
+
+	/** Copy/paste sections */
+	static void CopySections(const TArray<UMovieSceneSection*>& Sections, FString& ExportedText);
+	static bool PasteSections(const FString& TextToImport, FMovieScenePasteSectionsParams PasteSectionsParams, TArray<UMovieSceneSection*>& OutSections, TArray<FNotificationInfo>& OutErrors);
+	static bool CanPasteSections(const FString& TextToImport);
+
+	/** Copy/paste object bindings */
+	static void CopyBindings(const TArray<FMovieSceneBindingProxy>& Bindings, const TArray<UMovieSceneFolder*>& InFolders, FString& ExportedText);
+	static bool PasteBindings(const FString& TextToImport, TSharedRef<ISequencer> Sequencer, FMovieScenePasteBindingsParams PasteBindingsParams, TArray<FMovieSceneBindingProxy>& OutBindings, TArray<FNotificationInfo>& OutErrors);
+	static bool CanPasteBindings(TSharedRef<ISequencer> Sequencer, const FString& TextToImport);
 
-	static void UpdateBindingIDs(ISequencer* InSequencerPtr, UMovieSceneCompiledDataManager* InCompiledDataManagerPtr, FGuid OldGuid, FGuid NewGuid);
+	/** Utility functions for managing bindings */
+	static FGuid CreateBinding(TSharedRef<ISequencer> Sequencer, UObject& InObject, const FString& InName);
+	static void UpdateBindingIDs(TSharedRef<ISequencer> Sequencer, FGuid OldGuid, FGuid NewGuid);
+	static FGuid AssignActor(TSharedRef<ISequencer> Sequencer, AActor* Actor, FGuid InObjectBinding);
+	static void AddActorsToBinding(TSharedRef<ISequencer> Sequencer, const TArray<AActor*>& Actors, const FMovieSceneBindingProxy& ObjectBinding);
+	static void ReplaceBindingWithActors(TSharedRef<ISequencer> Sequencer, const TArray<AActor*>& Actors, const FMovieSceneBindingProxy& ObjectBinding);
+	static void RemoveActorsFromBinding(TSharedRef<ISequencer> Sequencer, const TArray<AActor*>& Actors, const FMovieSceneBindingProxy& ObjectBinding);
 
+	/** Show a read only error of the movie scene is locked */
 	static void ShowReadOnlyError();
 };
diff --git a/Engine/Source/Runtime/MovieScene/Private/MovieSceneFolder.cpp b/Engine/Source/Runtime/MovieScene/Private/MovieSceneFolder.cpp
--- a/Engine/Source/Runtime/MovieScene/Private/MovieSceneFolder.cpp
+++ b/Engine/Source/Runtime/MovieScene/Private/MovieSceneFolder.cpp
@@ -277,6 +277,25 @@
 	return nullptr;
 }
 
+UMovieSceneFolder* UMovieSceneFolder::FindFolderContaining(const UMovieSceneTrack* InTrack)
+{
+	if (ChildMasterTracks.Contains(InTrack))
+	{
+		return this;
+	}
+
+	for (UMovieSceneFolder* ChildFolder : GetChildFolders())
+	{
+		UMovieSceneFolder* Folder = ChildFolder->FindFolderContaining(InTrack);
+		if (Folder != nullptr)
+		{
+			return Folder;
+		}
+	}
+
+	return nullptr;
+}
+
 void
 TraverseFolder(UMovieSceneFolder* Folder, TMap<UMovieSceneFolder*, UMovieSceneFolder*>& ChildToParentMap)
 {
diff --git a/Engine/Source/Runtime/MovieScene/Public/MovieSceneFolder.h b/Engine/Source/Runtime/MovieScene/Public/MovieSceneFolder.h
--- a/Engine/Source/Runtime/MovieScene/Public/MovieSceneFolder.h
+++ b/Engine/Source/Runtime/MovieScene/Public/MovieSceneFolder.h
@@ -64,6 +64,9 @@
 	/** Searches for a guid in this folder and its child folders, if found returns the folder containing the guid. */
 	UMovieSceneFolder* FindFolderContaining(const FGuid& InObjectBinding);
 
+	/** Searches for a track in this folder and its child folders, if found returns the folder containing the track. */
+	UMovieSceneFolder* FindFolderContaining(const UMovieSceneTrack* InTrack);
+
 	/** Get the folder path for this folder, stopping at the given root folders */
 	static void CalculateFolderPath(UMovieSceneFolder* InFolder, TArrayView<UMovieSceneFolder* const> RootFolders, TArray<FName>& FolderPath);
 	
