Sequencer: CIS fixes

diff --git a/Engine/Source/Editor/Sequencer/Private/Sequencer.cpp b/Engine/Source/Editor/Sequencer/Private/Sequencer.cpp
--- a/Engine/Source/Editor/Sequencer/Private/Sequencer.cpp
+++ b/Engine/Source/Editor/Sequencer/Private/Sequencer.cpp
@@ -5549,35 +5549,37 @@ void FSequencer::OnNewActorsDropped(const TArray<UObject*>& DroppedObjects, cons
 
 		for ( AActor* Actor : DroppedActors )
 		{
-			AActor* NewActor = Actor;
-			FGuid PossessableGuid = FSequencerUtilities::CreateBinding(AsShared(), *NewActor, NewActor->GetActorLabel());
-			FGuid NewGuid = PossessableGuid;
+			if (AActor* NewActor = Actor)
+			{
+				FGuid PossessableGuid = FSequencerUtilities::CreateBinding(AsShared(), *NewActor, NewActor->GetActorLabel());
+				FGuid NewGuid = PossessableGuid;
 
-			OnActorAddedToSequencerEvent.Broadcast(NewActor, PossessableGuid);
+				OnActorAddedToSequencerEvent.Broadcast(NewActor, PossessableGuid);
 
-			if (bAddSpawnable)
-			{
-				TArray<FMovieSceneSpawnable*> Spawnables = FSequencerUtilities::ConvertToSpawnable(AsShared(), PossessableGuid);
-				if (Spawnables.Num() > 0)
+				if (bAddSpawnable)
 				{
-					for (TWeakObjectPtr<> WeakObject : FindBoundObjects(Spawnables[0]->GetGuid(), ActiveTemplateIDs.Top()))
+					TArray<FMovieSceneSpawnable*> Spawnables = FSequencerUtilities::ConvertToSpawnable(AsShared(), PossessableGuid);
+					if (Spawnables.Num() > 0)
 					{
-						AActor* SpawnedActor = Cast<AActor>(WeakObject.Get());
-						if (SpawnedActor)
+						for (TWeakObjectPtr<> WeakObject : FindBoundObjects(Spawnables[0]->GetGuid(), ActiveTemplateIDs.Top()))
 						{
-							SpawnedActors.Add(SpawnedActor);
-							NewActor = SpawnedActor;
+							AActor* SpawnedActor = Cast<AActor>(WeakObject.Get());
+							if (SpawnedActor)
+							{
+								SpawnedActors.Add(SpawnedActor);
+								NewActor = SpawnedActor;
+							}
 						}
+						NewGuid = Spawnables[0]->GetGuid();
 					}
-					NewGuid = Spawnables[0]->GetGuid();
 				}
-			}
 
-			if (NewActor->GetClass() == ACameraRig_Rail::StaticClass() ||
-				NewActor->GetClass() == ACameraRig_Crane::StaticClass())
-			{
-				ACineCameraActor* OutActor;
-				FSequencerUtilities::CreateCameraWithRig(AsShared(), NewActor, bAddSpawnable, OutActor);
+				if (NewActor->GetClass() == ACameraRig_Rail::StaticClass() ||
+					NewActor->GetClass() == ACameraRig_Crane::StaticClass())
+				{
+					ACineCameraActor* OutActor;
+					FSequencerUtilities::CreateCameraWithRig(AsShared(), NewActor, bAddSpawnable, OutActor);
+				}
 			}
 		}
 
diff --git a/Engine/Source/Editor/Sequencer/Private/SequencerUtilities.cpp b/Engine/Source/Editor/Sequencer/Private/SequencerUtilities.cpp
--- a/Engine/Source/Editor/Sequencer/Private/SequencerUtilities.cpp
+++ b/Engine/Source/Editor/Sequencer/Private/SequencerUtilities.cpp
@@ -601,7 +601,14 @@ FGuid AddSpawnable(TSharedRef<ISequencer> Sequencer, UObject& Object, UActorFact
 
 	FNewSpawnable& NewSpawnable = Result.GetValue();
 
-	NewSpawnable.Name = SpawnableName == NAME_None ? MovieSceneHelpers::MakeUniqueSpawnableName(OwnerMovieScene, NewSpawnable.Name) : SpawnableName.ToString();
+	if (SpawnableName == NAME_None)
+	{
+		NewSpawnable.Name = MovieSceneHelpers::MakeUniqueSpawnableName(OwnerMovieScene, NewSpawnable.Name);
+	}
+	else
+	{
+		NewSpawnable.Name = SpawnableName.ToString();
+	}
 
 	FGuid NewGuid = OwnerMovieScene->AddSpawnable(NewSpawnable.Name, *NewSpawnable.ObjectTemplate);
 
diff --git a/Engine/Source/Editor/Sequencer/Public/SequencerUtilities.h b/Engine/Source/Editor/Sequencer/Public/SequencerUtilities.h
--- a/Engine/Source/Editor/Sequencer/Public/SequencerUtilities.h
+++ b/Engine/Source/Editor/Sequencer/Public/SequencerUtilities.h
@@ -5,6 +5,7 @@
 #include "CoreMinimal.h"
 #include "Widgets/SWidget.h"
 #include "Framework/SlateDelegates.h"
+#include "Misc/FrameTime.h"
 #include "MovieSceneBindingProxy.h"
 #include "SequencerUtilities.generated.h"
 
