Sequencer: Fix missing transaction for adding keys to selected objects

diff --git a/Engine/Source/Editor/Sequencer/Private/Sequencer.cpp b/Engine/Source/Editor/Sequencer/Private/Sequencer.cpp
--- a/Engine/Source/Editor/Sequencer/Private/Sequencer.cpp
+++ b/Engine/Source/Editor/Sequencer/Private/Sequencer.cpp
@@ -2118,6 +2118,8 @@ bool FSequencer::CanAddTransformKeysForSelectedObjects() const
 
 void FSequencer::OnAddTransformKeysForSelectedObjects(EMovieSceneTransformChannel Channel)
 {
+	FScopedTransaction SetKeyTransaction(NSLOCTEXT("Sequencer", "SetTransformKey_Transaction", "Set Transform Key"));
+
 	TArray<TSharedPtr<ISequencerTrackEditor>> PossibleTrackEditors;
 	bool AtLeastOneHasPriority = false;
 	for (int32 i = 0; i < TrackEditors.Num(); ++i)
