Sequencer: Add camera cut tracks to the Cameras filter.

diff --git a/Engine/Source/Editor/Sequencer/Private/SequencerTrackFilters.h b/Engine/Source/Editor/Sequencer/Private/SequencerTrackFilters.h
--- a/Engine/Source/Editor/Sequencer/Private/SequencerTrackFilters.h
+++ b/Engine/Source/Editor/Sequencer/Private/SequencerTrackFilters.h
@@ -10,6 +10,7 @@
 #include "MovieSceneSequence.h"
 
 #include "Tracks/MovieSceneAudioTrack.h"
+#include "Tracks/MovieSceneCameraCutTrack.h"
 #include "Tracks/MovieSceneEventTrack.h"
 #include "Tracks/MovieSceneLevelVisibilityTrack.h"
 #include "Tracks/MovieSceneParticleTrack.h"
@@ -277,7 +278,7 @@ class FSequencerTrackFilter_SkeletalMeshObjects : public FSequencerTrackFilter_C
 {
 	virtual FString GetName() const override { return TEXT("SequencerSkeletalMeshObjectsFilter"); }
 	virtual FText GetDisplayName() const override { return LOCTEXT("SequencerTrackFilter_SkeletalMeshObjects", "Skeletal Mesh"); }
-	virtual FText GetToolTipText() const override { return LOCTEXT("SequencerTrackFilter_SkeletalMeshObjectsToolTip", "Show only SkeletalMesh objects."); }
+	virtual FText GetToolTipText() const override { return LOCTEXT("SequencerTrackFilter_SkeletalMeshObjectsToolTip", "Show only Skeletal Mesh objects."); }
 	virtual FSlateIcon GetIcon() const override { return FSlateIconFinder::FindIconForClass(USkeletalMeshComponent::StaticClass()); }
 };
 
@@ -285,8 +286,18 @@ class FSequencerTrackFilter_CameraObjects : public FSequencerTrackFilter_Compone
 {
 	virtual FString GetName() const override { return TEXT("SequencerCameraObjectsFilter"); }
 	virtual FText GetDisplayName() const override { return LOCTEXT("SequencerTrackFilter_CameraObjects", "Cameras"); }
-	virtual FText GetToolTipText() const override { return LOCTEXT("SequencerTrackFilter_CameraObjectsToolTip", "Show only Camera objects."); }
+	virtual FText GetToolTipText() const override { return LOCTEXT("SequencerTrackFilter_CameraObjectsToolTip", "Show only Camera objects and Camera Cut tracks."); }
 	virtual FSlateIcon GetIcon() const override { return FSlateIconFinder::FindIconForClass(UCameraComponent::StaticClass()); }
+
+	// IFilter implementation
+	virtual bool PassesFilter(FTrackFilterType InItem) const override
+	{
+		if (InItem->IsA(UMovieSceneCameraCutTrack::StaticClass()))
+		{
+			return true;
+		}
+		return FSequencerTrackFilter_ComponentType< UCameraComponent >::PassesFilter(InItem);
+	}
 };
 
 class FSequencerTrackFilter_LightObjects : public FSequencerTrackFilter_ComponentType< ULightComponentBase >
diff --git a/Engine/Source/Editor/Sequencer/Public/SequencerTrackFilterBase.h b/Engine/Source/Editor/Sequencer/Public/SequencerTrackFilterBase.h
--- a/Engine/Source/Editor/Sequencer/Public/SequencerTrackFilterBase.h
+++ b/Engine/Source/Editor/Sequencer/Public/SequencerTrackFilterBase.h
@@ -66,6 +66,7 @@
 template<typename TrackType>
 class FSequencerTrackFilter_ClassType : public FSequencerTrackFilter
 {
+public:
 	// IFilter implementation
 	virtual bool PassesFilter(FTrackFilterType InItem) const override
 	{
@@ -81,6 +82,7 @@
 template<typename ComponentType>
 class FSequencerTrackFilter_ComponentType : public FSequencerTrackFilter
 {
+public:
 	// IFilter implementation
 	virtual bool PassesFilter(FTrackFilterType InItem) const override
 	{
