Sequencer: CIS fixes

diff --git a/Engine/Source/Editor/Sequencer/Private/Sequencer.cpp b/Engine/Source/Editor/Sequencer/Private/Sequencer.cpp
--- a/Engine/Source/Editor/Sequencer/Private/Sequencer.cpp
+++ b/Engine/Source/Editor/Sequencer/Private/Sequencer.cpp
@@ -5576,8 +5576,7 @@ void FSequencer::OnNewActorsDropped(const TArray<UObject*>& DroppedObjects, cons
 					}
 				}
 
-				if (NewActor->GetClass() == ACameraRig_Rail::StaticClass() ||
-					NewActor->GetClass() == ACameraRig_Crane::StaticClass())
+				if (NewActor && (NewActor->GetClass() == ACameraRig_Rail::StaticClass() || NewActor->GetClass() == ACameraRig_Crane::StaticClass()))
 				{
 					ACineCameraActor* OutActor;
 					FSequencerUtilities::CreateCameraWithRig(AsShared(), NewActor, bAddSpawnable, OutActor);
