Sequencer: Set the filter button to show a layered image (red dot) when there are active filters

diff --git a/Engine/Source/Editor/Sequencer/Private/SSequencer.cpp b/Engine/Source/Editor/Sequencer/Private/SSequencer.cpp
--- a/Engine/Source/Editor/Sequencer/Private/SSequencer.cpp
+++ b/Engine/Source/Editor/Sequencer/Private/SSequencer.cpp
@@ -28,6 +28,7 @@
 #include "Framework/Application/SlateApplication.h"
 #include "Widgets/Layout/SSpacer.h"
 #include "Widgets/Images/SImage.h"
+#include "Widgets/Images/SLayeredImage.h"
 #include "Widgets/Layout/SBox.h"
 #include "Widgets/Layout/SWrapBox.h"
 #include "Widgets/Layout/SGridPanel.h"
@@ -1093,9 +1094,27 @@ TSharedRef<SWidget> SSequencer::MakeAddButton()
 
 TSharedRef<SWidget> SSequencer::MakeFilterButton()
 {
+	TSharedPtr<SLayeredImage> FilterImage = SNew(SLayeredImage)
+		.Image(FAppStyle::Get().GetBrush("Icons.Filter"))
+		.ColorAndOpacity(FSlateColor::UseForeground());
+
+	const TAttribute<const FSlateBrush*> ModifiedIcon = TAttribute<const FSlateBrush*>::CreateLambda([this]()
+		{
+			for (TSharedRef<FSequencerTrackFilter> TrackFilter : AllTrackFilters)
+			{
+				if (IsTrackFilterActive(TrackFilter))
+				{
+					return FAppStyle::Get().GetBrush("Icons.BadgeModified");
+				}
+			}
+			return (const FSlateBrush*)nullptr;
+		});
+
+	FilterImage->AddLayer(ModifiedIcon);
+
 	return SNew(SComboButton)
-	.ButtonStyle( FAppStyle::Get(), "SimpleButton" )
-	.ForegroundColor(FSlateColor::UseForeground())
+	.ComboButtonStyle(&FAppStyle::Get().GetWidgetStyle<FComboButtonStyle>("SimpleComboButtonWithIcon"))
+	.ForegroundColor(FSlateColor::UseStyle())
 	.ContentPadding(0)
 	.ToolTipText(LOCTEXT("AddTrackFilterToolTip", "Add a track filter."))
 	.OnGetMenuContent(this, &SSequencer::MakeFilterMenu)
@@ -1103,9 +1122,7 @@ TSharedRef<SWidget> SSequencer::MakeFilterButton()
 	.HasDownArrow(false)
 	.ButtonContent()
 	[
-		SNew(SImage)
-		.ColorAndOpacity(FSlateColor::UseForeground())
-		.Image(FAppStyle::Get().GetBrush("Icons.Filter"))
+		FilterImage.ToSharedRef()
 	];
 }
 
