Sequencer: Follow up to 20820007 to remove AllowsSpawnableObjects where not necessary

diff --git a/Engine/Source/Editor/Sequencer/Private/SequencerUtilities.cpp b/Engine/Source/Editor/Sequencer/Private/SequencerUtilities.cpp
--- a/Engine/Source/Editor/Sequencer/Private/SequencerUtilities.cpp
+++ b/Engine/Source/Editor/Sequencer/Private/SequencerUtilities.cpp
@@ -682,7 +682,7 @@ FGuid FSequencerUtilities::CreateCamera(TSharedRef<ISequencer> Sequencer, const
 		return CameraGuid;
 	}
 
-	if (MovieScene->IsReadOnly() || !Sequence->AllowsSpawnableObjects())
+	if (MovieScene->IsReadOnly())
 	{
 		ShowReadOnlyError();
 		return CameraGuid;
@@ -776,7 +776,7 @@ FGuid FSequencerUtilities::CreateCameraWithRig(TSharedRef<ISequencer> Sequencer,
 		return CameraGuid;
 	}
 
-	if (MovieScene->IsReadOnly() || !Sequence->AllowsSpawnableObjects())
+	if (MovieScene->IsReadOnly())
 	{
 		ShowReadOnlyError();
 		return CameraGuid;
@@ -896,7 +896,7 @@ TArray<FGuid> FSequencerUtilities::AddActors(TSharedRef<ISequencer> Sequencer, c
 		return PossessableGuids;
 	}
 
-	if (MovieScene->IsReadOnly() || !Sequence->AllowsSpawnableObjects())
+	if (MovieScene->IsReadOnly())
 	{
 		ShowReadOnlyError();
 		return PossessableGuids;
@@ -1052,7 +1052,7 @@ FMovieScenePossessable* FSequencerUtilities::ConvertToPossessable(TSharedRef<ISe
 		return CreatedPossessable;
 	}
 
-	if (MovieScene->IsReadOnly() || !Sequence->AllowsSpawnableObjects())
+	if (MovieScene->IsReadOnly())
 	{
 		ShowReadOnlyError();
 		return CreatedPossessable;
