Sequencer: Only add selected folder once

diff --git a/Engine/Source/Editor/Sequencer/Private/Sequencer.cpp b/Engine/Source/Editor/Sequencer/Private/Sequencer.cpp
--- a/Engine/Source/Editor/Sequencer/Private/Sequencer.cpp
+++ b/Engine/Source/Editor/Sequencer/Private/Sequencer.cpp
@@ -8836,7 +8836,11 @@ void FSequencer::CalculateSelectedFolderAndPath(TArray<UMovieSceneFolder*>& OutS
 			}
 			if (CurrentNode.IsValid())
 			{
-				OutSelectedParentFolders.Add(CurrentNode->CastThisChecked<FFolderModel>()->GetFolder());
+				UMovieSceneFolder* Folder = CurrentNode->CastThisChecked<FFolderModel>()->GetFolder();
+				if (Folder && !OutSelectedParentFolders.Contains(Folder))
+				{
+					OutSelectedParentFolders.Add(Folder);
+				}
 
 				// The first valid folder we find will be used to put the new folder into, so it's the node that we
 				// want to know the path from.
