Sequencer: Follow up to refactor, PasteSections should return false if nothing to paste

diff --git a/Engine/Source/Editor/Sequencer/Private/Sequencer.cpp b/Engine/Source/Editor/Sequencer/Private/Sequencer.cpp
--- a/Engine/Source/Editor/Sequencer/Private/Sequencer.cpp
+++ b/Engine/Source/Editor/Sequencer/Private/Sequencer.cpp
@@ -8105,7 +8105,10 @@ bool FSequencer::PasteSections(const FString& TextToImport, TArray<FNotification
 	}
 
 	TArray<UMovieSceneSection*> NewSections;
-	FSequencerUtilities::PasteSections(TextToImport, FMovieScenePasteSectionsParams(Tracks, TrackRowIndices, GetLocalTime().Time), NewSections, PasteErrors);
+	if (!FSequencerUtilities::PasteSections(TextToImport, FMovieScenePasteSectionsParams(Tracks, TrackRowIndices, GetLocalTime().Time), NewSections, PasteErrors))
+	{
+		return false;
+	}
 
 	NotifyMovieSceneDataChanged(EMovieSceneDataChangeType::MovieSceneStructureItemAdded);
 
