Sequencer: Fixed track models remaining in an expanded state when there were no longer any children

diff --git a/Engine/Source/Editor/Sequencer/Private/MVVM/ViewModels/TrackModel.cpp b/Engine/Source/Editor/Sequencer/Private/MVVM/ViewModels/TrackModel.cpp
--- a/Engine/Source/Editor/Sequencer/Private/MVVM/ViewModels/TrackModel.cpp
+++ b/Engine/Source/Editor/Sequencer/Private/MVVM/ViewModels/TrackModel.cpp
@@ -185,12 +185,18 @@
 
 	if (NumRows == 0)
 	{
+		// Reset expansion state if this track can no longer be expanded
+		SetExpansion(false);
+
 		// Clear any left-over row models, layout models, or section models.
 		OutlinerChildren.Empty();
 		SectionChildren.Empty();
 	}
 	else if (NumRows == 1)
 	{
+		// Reset expansion state if this track can no longer be expanded
+		SetExpansion(false);
+
 		// Keep sections alive by retaining the previous list temporarily
 		TSharedPtr<FViewModel> SectionsTail;
 
