Sequencer: Broadcast OnStop delegate in SetLocalTimeLooped so that if playback stops naturally in FSequencer, the delegate will be emitted. Otherwise, it was only emited when the user invoked the stop button. This makes it similar to MovieSceneSequencePlayer which emits OnStop when playback stops automatically.

diff --git a/Engine/Source/Editor/Sequencer/Private/Sequencer.cpp b/Engine/Source/Editor/Sequencer/Private/Sequencer.cpp
--- a/Engine/Source/Editor/Sequencer/Private/Sequencer.cpp
+++ b/Engine/Source/Editor/Sequencer/Private/Sequencer.cpp
@@ -4898,6 +4898,11 @@ void FSequencer::SetLocalTimeLooped(FFrameTime NewLocalTime)
 		SetPlaybackStatus(NewPlaybackStatus.GetValue());
 		// Evaluate the sequence with the new status
 		EvaluateInternal(EvalRange);
+
+		if (NewPlaybackStatus.GetValue() == EMovieScenePlayerStatus::Stopped)
+		{
+			OnStopDelegate.Broadcast();
+		}
 	}
 }
 
