Sequencer: Fixed sections getting collapsed when they are modified

diff --git a/Engine/Source/Editor/Sequencer/Private/MVVM/ViewModels/TrackModel.cpp b/Engine/Source/Editor/Sequencer/Private/MVVM/ViewModels/TrackModel.cpp
--- a/Engine/Source/Editor/Sequencer/Private/MVVM/ViewModels/TrackModel.cpp
+++ b/Engine/Source/Editor/Sequencer/Private/MVVM/ViewModels/TrackModel.cpp
@@ -194,9 +194,6 @@ void FTrackModel::ForceUpdate()
 	}
 	else if (NumRows == 1)
 	{
-		// Reset expansion state if this track can no longer be expanded
-		SetExpansion(false);
-
 		// Keep sections alive by retaining the previous list temporarily
 		TSharedPtr<FViewModel> SectionsTail;
 
@@ -231,6 +228,12 @@ void FTrackModel::ForceUpdate()
 		{
 			LayoutBuilder.RefreshLayout(Section);
 		}
+
+		if (OutlinerChildren.IsEmpty())
+		{
+			// Reset expansion state if this track can no longer be expanded
+			SetExpansion(false);
+		}
 	}
 	else
 	{
