Sequencer: Fix paste destinations for multiple rows

Renamed GetSections to GetSectionModels and made it virtual.
Added GetSections to ITrackExtension to return the UMovieSceneSections of the SectionModels.

diff --git a/Engine/Source/Editor/Sequencer/Private/MVVM/ViewModels/TrackModel.cpp b/Engine/Source/Editor/Sequencer/Private/MVVM/ViewModels/TrackModel.cpp
--- a/Engine/Source/Editor/Sequencer/Private/MVVM/ViewModels/TrackModel.cpp
+++ b/Engine/Source/Editor/Sequencer/Private/MVVM/ViewModels/TrackModel.cpp
@@ -73,7 +73,7 @@ EViewModelListType FTrackModel::GetTopLevelChannelGroupType()
 	return TopLevelChannelGroup;
 }
 
-FViewModelChildren FTrackModel::GetSections()
+FViewModelChildren FTrackModel::GetSectionModels()
 {
 	return GetChildrenForList(&SectionList);
 }
@@ -287,7 +287,7 @@ void FTrackModel::ForceUpdate()
 
 				// Keep sections on rows alive as well
 				RowModels[RowIndex].OldSections = MakeUnique<FScopedViewModelListHead>(TrackRowModel, EViewModelListType::Recycled);
-				TrackRowModel->GetSections().MoveChildrenTo(RowModels[RowIndex].OldSections->GetChildren());
+				TrackRowModel->GetSectionModels().MoveChildrenTo(RowModels[RowIndex].OldSections->GetChildren());
 			}
 		}
 
@@ -303,7 +303,7 @@ void FTrackModel::ForceUpdate()
 				SectionModel = SectionModelStorage->CreateModelForSection(Section, SectionInterface);
 			}
 
-			RowModels[RowIndex].Row->GetSections().InsertChild(SectionModel, RowModels[RowIndex].SectionsTail);
+			RowModels[RowIndex].Row->GetSectionModels().InsertChild(SectionModel, RowModels[RowIndex].SectionsTail);
 			RowModels[RowIndex].SectionsTail = SectionModel;
 		}
 
diff --git a/Engine/Source/Editor/Sequencer/Private/MVVM/ViewModels/TrackRowModel.cpp b/Engine/Source/Editor/Sequencer/Private/MVVM/ViewModels/TrackRowModel.cpp
--- a/Engine/Source/Editor/Sequencer/Private/MVVM/ViewModels/TrackRowModel.cpp
+++ b/Engine/Source/Editor/Sequencer/Private/MVVM/ViewModels/TrackRowModel.cpp
@@ -79,7 +79,7 @@
 	return RowIndex;
 }
 
-FViewModelChildren FTrackRowModel::GetSections()
+FViewModelChildren FTrackRowModel::GetSectionModels()
 {
 	return GetChildrenForList(&SectionList);
 }
diff --git a/Engine/Source/Editor/Sequencer/Private/SequencerContextMenus.cpp b/Engine/Source/Editor/Sequencer/Private/SequencerContextMenus.cpp
--- a/Engine/Source/Editor/Sequencer/Private/SequencerContextMenus.cpp
+++ b/Engine/Source/Editor/Sequencer/Private/SequencerContextMenus.cpp
@@ -1,7 +1,6 @@
 // Copyright Epic Games, Inc. All Rights Reserved.
 
 #include "SequencerContextMenus.h"
-#include "MVVM/ViewModels/ViewModel.h"
 #include "Modules/ModuleManager.h"
 #include "Styling/AppStyle.h"
 #include "SequencerCommonHelpers.h"
@@ -11,7 +10,6 @@
 #include "SSequencerSection.h"
 #include "SequencerSettings.h"
 #include "MVVM/Views/ITrackAreaHotspot.h"
-#include "MVVM/ViewModels/ViewModel.h"
 #include "SequencerHotspots.h"
 #include "ScopedTransaction.h"
 #include "MovieSceneToolHelpers.h"
@@ -42,9 +40,11 @@
 #include "MovieSceneSequence.h"
 #include "MovieScene.h"
 #include "Channels/MovieSceneChannel.h"
+#include "MVVM/Extensions/ITrackExtension.h"
 #include "MVVM/ViewModels/ViewModelIterators.h"
 #include "MVVM/ViewModels/SectionModel.h"
 #include "MVVM/ViewModels/TrackModel.h"
+#include "MVVM/ViewModels/ViewModel.h"
 #include "MVVM/ViewModels/SequencerEditorViewModel.h"
 #include "Tracks/MovieScenePropertyTrack.h"
 #include "Algo/AnyOf.h"
@@ -1365,40 +1365,6 @@
 }
 
 
-TSharedPtr<UE::Sequencer::FTrackModel> GetTrackFromNode(const UE::Sequencer::TViewModelPtr<UE::Sequencer::IOutlinerExtension>& InNode, FString& Scope)
-{
-	using namespace UE::Sequencer;
-
-	if (TSharedPtr<FTrackModel> TrackModel = InNode.ImplicitCast())
-	{
-		return TrackModel;
-	}
-	if (TSharedPtr<FObjectBindingModel> ObjectBindingModel = InNode.ImplicitCast())
-	{
-		return nullptr;
-	}
-
-	TViewModelPtr<IOutlinerExtension> Parent = InNode.AsModel()->FindAncestorOfType<IOutlinerExtension>();
-	if (Parent)
-	{
-		TSharedPtr<FTrackModel> Track = GetTrackFromNode(Parent, Scope);
-		if (Track.IsValid())
-		{
-			FString ThisScope = Track->GetLabel().ToString();
-			if (!Scope.IsEmpty())
-			{
-				ThisScope.AppendChar('.');
-				ThisScope.Append(Scope);
-				Scope = MoveTemp(ThisScope);
-			}
-			return Track;
-		}
-	}
-
-	return nullptr;
-}
-
-
 void FPasteContextMenu::Setup()
 {
 	using namespace UE::Sequencer;
@@ -1425,26 +1391,19 @@
 		{
 			for (const TViewModelPtr<IOutlinerExtension>& Node : Args.DestinationNodes)
 			{
-				if (!Node.AsModel()->IsA<FChannelModel>() && !Node.AsModel()->IsA<FChannelGroupModel>())
+				TViewModelPtr<ITrackExtension> TrackNode = Node.AsModel()->FindAncestorOfType<ITrackExtension>(true);
+				if (!TrackNode)
 				{
 					continue;
 				}
 
-				FString Scope;
-				TSharedPtr<FTrackModel> TrackNode = GetTrackFromNode(Node, Scope);
-				if (!TrackNode.IsValid())
-				{
-					continue;
-				}
-
 				FPasteDestination& Destination = PasteDestinations[PasteDestinations.AddDefaulted()];
 
-				TArray<UMovieSceneSection*> Sections;
-				for (const TViewModelPtr<FSectionModel>& Section : TrackNode->GetSections().IterateSubList<FSectionModel>())
+				for (UMovieSceneSection* Section : TrackNode->GetSections())
 				{
-					if (Section->GetSection())
+					if (Section)
 					{
-						GatherPasteDestinationsForNode(Node, Section->GetSection(), NAME_None, Destination.Reconcilers);
+						GatherPasteDestinationsForNode(Node, Section, NAME_None, Destination.Reconcilers);
 					}
 				}
 
@@ -1485,23 +1444,13 @@
 		// Build a list of sections based on selected tracks
 		for (const TViewModelPtr<IOutlinerExtension>& Node : Args.DestinationNodes)
 		{
-			FString Scope;
-			TSharedPtr<FTrackModel> TrackNode = GetTrackFromNode(Node, Scope);
-			if (!TrackNode.IsValid())
+			TViewModelPtr<ITrackExtension> TrackNode = Node.AsModel()->FindAncestorOfType<ITrackExtension>(true);
+			if (!TrackNode)
 			{
 				continue;
 			}
 
-			TArray<UMovieSceneSection*> Sections;
-			for (const TViewModelPtr<FSectionModel>& Section : TrackNode->GetSections().IterateSubList<FSectionModel>())
-			{
-				if (Section->GetSection())
-				{
-					Sections.Add(Section->GetSection());
-				}
-			}
-
-			UMovieSceneSection* Section = MovieSceneHelpers::FindNearestSectionAtTime(Sections, Args.PasteAtTime);
+			UMovieSceneSection* Section = MovieSceneHelpers::FindNearestSectionAtTime(TrackNode->GetSections(), Args.PasteAtTime);
 			TSharedPtr<FSectionModel> SectionModel = Sequencer->GetNodeTree()->GetSectionModel(Section);
 			if (SectionModel)
 			{
diff --git a/Engine/Source/Editor/Sequencer/Public/MVVM/Extensions/ITrackExtension.h b/Engine/Source/Editor/Sequencer/Public/MVVM/Extensions/ITrackExtension.h
--- a/Engine/Source/Editor/Sequencer/Public/MVVM/Extensions/ITrackExtension.h
+++ b/Engine/Source/Editor/Sequencer/Public/MVVM/Extensions/ITrackExtension.h
@@ -1,31 +1,52 @@
-// Copyright Epic Games, Inc. All Rights Reserved.
-
-#pragma once
-
-#include "MVVM/ViewModelTypeID.h"
-
-class UMovieSceneTrack;
-class ISequencerTrackEditor;
-
-namespace UE
-{
-namespace Sequencer
-{
-
-class SEQUENCER_API ITrackExtension
-{
-public:
-
-	UE_SEQUENCER_DECLARE_VIEW_MODEL_TYPE_ID(ITrackExtension)
-
-	virtual ~ITrackExtension(){}
-
-	virtual UMovieSceneTrack* GetTrack() const = 0;
-	virtual int32 GetRowIndex() const = 0;
-
-	virtual TSharedPtr<ISequencerTrackEditor> GetTrackEditor() const = 0;
-};
-
-} // namespace Sequencer
-} // namespace UE
-
+// Copyright Epic Games, Inc. All Rights Reserved.
+
+#pragma once
+
+#include "MVVM/ViewModelTypeID.h"
+#include "MVVM/ViewModels/SectionModel.h"
+#include "MVVM/ViewModels/ViewModel.h"
+#include "MVVM/ViewModels/ViewModelHierarchy.h"
+#include "MVVM/ViewModels/ViewModelIterators.h"
+
+class UMovieSceneSection;
+class UMovieSceneTrack;
+class ISequencerTrackEditor;
+struct FViewModelChildren;
+
+namespace UE
+{
+namespace Sequencer
+{
+
+class SEQUENCER_API ITrackExtension
+{
+public:
+
+	UE_SEQUENCER_DECLARE_VIEW_MODEL_TYPE_ID(ITrackExtension)
+
+	virtual ~ITrackExtension(){}
+
+	virtual UMovieSceneTrack* GetTrack() const = 0;
+	virtual int32 GetRowIndex() const = 0;
+	virtual FViewModelChildren GetSectionModels() = 0;
+
+	TArray<UMovieSceneSection*> GetSections() 
+	{ 
+		TArray<UMovieSceneSection*> Sections; 
+		for (const TViewModelPtr<FSectionModel>& Item : GetSectionModels().IterateSubList<FSectionModel>())
+		{
+			if (UMovieSceneSection* Section = Item->GetSection())
+			{
+				Sections.Add(Section);
+			}
+		}
+
+		return Sections;
+	}
+
+	virtual TSharedPtr<ISequencerTrackEditor> GetTrackEditor() const = 0;
+};
+
+} // namespace Sequencer
+} // namespace UE
+
diff --git a/Engine/Source/Editor/Sequencer/Public/MVVM/ViewModels/TrackModel.h b/Engine/Source/Editor/Sequencer/Public/MVVM/ViewModels/TrackModel.h
--- a/Engine/Source/Editor/Sequencer/Public/MVVM/ViewModels/TrackModel.h
+++ b/Engine/Source/Editor/Sequencer/Public/MVVM/ViewModels/TrackModel.h
@@ -63,7 +63,6 @@
 	explicit FTrackModel(UMovieSceneTrack* Track);
 	~FTrackModel();
 
-	FViewModelChildren GetSections();
 	FViewModelChildren GetTopLevelChannels();
 
 public:
@@ -74,6 +73,8 @@
 	/*~ ITrackExtension */
 	UMovieSceneTrack* GetTrack() const override;
 	int32 GetRowIndex() const override;
+	FViewModelChildren GetSectionModels() override;
+
 	TSharedPtr<ISequencerTrackEditor> GetTrackEditor() const override;
 
 	/*~ ISignedObjectEventHandler */
diff --git a/Engine/Source/Editor/Sequencer/Public/MVVM/ViewModels/TrackRowModel.h b/Engine/Source/Editor/Sequencer/Public/MVVM/ViewModels/TrackRowModel.h
--- a/Engine/Source/Editor/Sequencer/Public/MVVM/ViewModels/TrackRowModel.h
+++ b/Engine/Source/Editor/Sequencer/Public/MVVM/ViewModels/TrackRowModel.h
@@ -39,8 +39,6 @@
 
 	void Initialize();
 
-	FViewModelChildren GetSections();
-
 	TSharedPtr<ISequencerTrackEditor> GetTrackEditor() const { return TrackEditor; }
 
 	/*~ FOutlinerItemModel */
@@ -62,6 +60,7 @@
 	/*~ ITrackExtension */
 	UMovieSceneTrack* GetTrack() const override;
 	int32 GetRowIndex() const override;
+	FViewModelChildren GetSectionModels() override;
 
 	/*~ ITrackAreaExtension */
 	FTrackAreaParameters GetTrackAreaParameters() const override;
