Sequencer: Fixed global pre-animated-state not being captured after an explicit call to FSequencer::ResetToNewRootSequence

diff --git a/Engine/Source/Editor/Sequencer/Private/Sequencer.cpp b/Engine/Source/Editor/Sequencer/Private/Sequencer.cpp
--- a/Engine/Source/Editor/Sequencer/Private/Sequencer.cpp
+++ b/Engine/Source/Editor/Sequencer/Private/Sequencer.cpp
@@ -985,6 +985,7 @@ void FSequencer::ResetToNewRootSequence(UMovieSceneSequence& NewSequence)
 	ActiveTemplateStates.Add(true);
 
 	RootTemplateInstance.Initialize(NewSequence, *this, CompiledDataManager);
+	RootTemplateInstance.EnableGlobalPreAnimatedStateCapture();
 
 	RootToLocalTransform = FMovieSceneSequenceTransform();
 	RootToLocalLoopCounter = FMovieSceneWarpCounter();
