Sequencer: Fix crash dereferencing null section

diff --git a/Engine/Source/Editor/Sequencer/Private/SequencerHotspots.cpp b/Engine/Source/Editor/Sequencer/Private/SequencerHotspots.cpp
--- a/Engine/Source/Editor/Sequencer/Private/SequencerHotspots.cpp
+++ b/Engine/Source/Editor/Sequencer/Private/SequencerHotspots.cpp
@@ -394,7 +394,7 @@
 
 bool FSectionEasingAreaHotspot::Contains(UMovieSceneSection* InSection) const
 {
-	return Easings.ContainsByPredicate([=](const FEasingAreaHandle& InHandle){ return InHandle.WeakSectionModel.Pin()->GetSection() == InSection; });
+	return Easings.ContainsByPredicate([=](const FEasingAreaHandle& InHandle){ return InHandle.WeakSectionModel.IsValid() && InHandle.WeakSectionModel.Pin()->GetSection() == InSection; });
 }
 
 } // namespace Sequencer
