Sequencer: Avoid removing RF_TRansient from RigVM Memory

diff --git a/Engine/Source/Editor/Sequencer/Private/SequencerUtilities.cpp b/Engine/Source/Editor/Sequencer/Private/SequencerUtilities.cpp
--- a/Engine/Source/Editor/Sequencer/Private/SequencerUtilities.cpp
+++ b/Engine/Source/Editor/Sequencer/Private/SequencerUtilities.cpp
@@ -2035,12 +2035,10 @@ bool FSequencerUtilities::PasteBindings(const FString& TextToImport, TSharedRef<
 			for (UMovieSceneTrack* CopiedTrack : CopyableBinding->Tracks)
 			{
 				CopiedTrack->ClearFlags(RF_Transient);
-				TArray<UObject*> Subobjects;
-				GetObjectsWithOuter(CopiedTrack, Subobjects);
-				for (UObject* Subobject : Subobjects)
-				{
-					Subobject->ClearFlags(RF_Transient);
-				}
+				for (UMovieSceneSection* Section : CopiedTrack->GetAllSections())
+                {
+					Section->ClearFlags(RF_Transient);
+                }
 			}
 
 			UMovieSceneFolder* ParentFolder = PasteBindingsParams.ParentFolder;
