Sequencer: Clear RF_Transient on MovieSceneSection on paste

diff --git a/Engine/Source/Editor/Sequencer/Private/SequencerUtilities.cpp b/Engine/Source/Editor/Sequencer/Private/SequencerUtilities.cpp
--- a/Engine/Source/Editor/Sequencer/Private/SequencerUtilities.cpp
+++ b/Engine/Source/Editor/Sequencer/Private/SequencerUtilities.cpp
@@ -1486,11 +1486,9 @@ bool FSequencerUtilities::PasteTracks(const FString& TextToImport, FMovieScenePa
 			{
 				UMovieSceneTrack* NewTrack = CopyableTrack->Track;
 				NewTrack->ClearFlags(RF_Transient);
-				TArray<UObject*> Subobjects;
-				GetObjectsWithOuter(NewTrack, Subobjects);
-				for (UObject* Subobject : Subobjects)
+				for (UMovieSceneSection* Section : NewTrack->GetAllSections())
 				{
-					Subobject->ClearFlags(RF_Transient);
+					Section->ClearFlags(RF_Transient);
 				}
 
 				// Remove tracks with the same name before adding
@@ -1534,11 +1532,9 @@ bool FSequencerUtilities::PasteTracks(const FString& TextToImport, FMovieScenePa
 			{
 				UMovieSceneTrack* NewTrack = CopyableTrack->Track;
 				NewTrack->ClearFlags(RF_Transient);
-				TArray<UObject*> Subobjects;
-				GetObjectsWithOuter(NewTrack, Subobjects);
-				for (UObject* Subobject : Subobjects)
+				for (UMovieSceneSection* Section : NewTrack->GetAllSections())
 				{
-					Subobject->ClearFlags(RF_Transient);
+					Section->ClearFlags(RF_Transient);
 				}
 
 				UMovieSceneFolder* ParentFolder = PasteTracksParams.ParentFolder;
