Sequencer: Add a filter for subsequences and cinematic shots

diff --git a/Engine/Source/Editor/Sequencer/Private/SSequencer.cpp b/Engine/Source/Editor/Sequencer/Private/SSequencer.cpp
--- a/Engine/Source/Editor/Sequencer/Private/SSequencer.cpp
+++ b/Engine/Source/Editor/Sequencer/Private/SSequencer.cpp
@@ -987,6 +987,8 @@ void SSequencer::InitializeTrackFilters()
 	AllTrackFilters.Add(MakeShared<FSequencerTrackFilter_EventTracks>());
 	AllTrackFilters.Add(MakeShared<FSequencerTrackFilter_LevelVisibilityTracks>());
 	AllTrackFilters.Add(MakeShared<FSequencerTrackFilter_ParticleTracks>());
+	AllTrackFilters.Add(MakeShared<FSequencerTrackFilter_CinematicShotTracks>());
+	AllTrackFilters.Add(MakeShared<FSequencerTrackFilter_SubTracks>());
 
 	AllTrackFilters.Add(MakeShared<FSequencerTrackFilter_CameraObjects>());
 	AllTrackFilters.Add(MakeShared<FSequencerTrackFilter_LightObjects>());
diff --git a/Engine/Source/Editor/Sequencer/Private/SequencerNodeTree.cpp b/Engine/Source/Editor/Sequencer/Private/SequencerNodeTree.cpp
--- a/Engine/Source/Editor/Sequencer/Private/SequencerNodeTree.cpp
+++ b/Engine/Source/Editor/Sequencer/Private/SequencerNodeTree.cpp
@@ -28,8 +28,6 @@
 #include "IKeyArea.h"
 #include "ISequencerSection.h"
 #include "MovieSceneSequence.h"
-#include "Tracks/MovieSceneCinematicShotTrack.h"
-#include "Tracks/MovieSceneSubTrack.h"
 #include "Sequencer.h"
 #include "MovieSceneFolder.h"
 #include "ISequencerTrackEditor.h"
@@ -686,7 +684,6 @@ static bool FilterNodesRecursive(
 
 	bool bPassedAnyFilters = false;
 	bool bIsTrackOrObjectBinding = false;
-	bool bIsSubTrack = false;
 
 	TSharedPtr<IPinnableExtension> Pinnable = StartNode.ImplicitCast();
 	const bool bIsPinned = Pinnable && Pinnable->IsPinned();
@@ -697,14 +694,6 @@ static bool FilterNodesRecursive(
 
 		UMovieSceneTrack* Track = TrackModel->GetTrack();
 
-		// Always show subsequence tracks so that the user can navigate up and down the hierarchy with the filter
-		if (Track && Track->IsA<UMovieSceneSubTrack>())
-		{
-			bPassedAnyFilters = true;
-			bIsSubTrack = true;
-			goto NextStep;
-		}
-
 		for (TSharedPtr<FChannelGroupModel> ChannelGroupModel : StartNode.AsModel()->GetDescendantsOfType<FChannelGroupModel>())
 		{
 			for (const TWeakViewModelPtr<FChannelModel>& WeakChannel : ChannelGroupModel->GetChannels())
@@ -902,9 +891,7 @@ static bool FilterNodesRecursive(
 		}
 	}
 
-NextStep:
-
-	if (bPassedAnyFilters && !bIsSubTrack)
+	if (bPassedAnyFilters)
 	{
 		// If filtering on selection set is enabled, we need to run another pass to verify we're in an enabled node group
 		UMovieSceneNodeGroupCollection& NodeGroups = Sequencer.GetFocusedMovieSceneSequence()->GetMovieScene()->GetNodeGroups();
diff --git a/Engine/Source/Editor/Sequencer/Private/SequencerTrackFilters.h b/Engine/Source/Editor/Sequencer/Private/SequencerTrackFilters.h
--- a/Engine/Source/Editor/Sequencer/Private/SequencerTrackFilters.h
+++ b/Engine/Source/Editor/Sequencer/Private/SequencerTrackFilters.h
@@ -11,9 +11,11 @@
 
 #include "Tracks/MovieSceneAudioTrack.h"
 #include "Tracks/MovieSceneCameraCutTrack.h"
+#include "Tracks/MovieSceneCinematicShotTrack.h"
 #include "Tracks/MovieSceneEventTrack.h"
 #include "Tracks/MovieSceneLevelVisibilityTrack.h"
 #include "Tracks/MovieSceneParticleTrack.h"
+#include "Tracks/MovieSceneSubTrack.h"
 
 #include "Camera/CameraComponent.h"
 #include "Components/LightComponentBase.h"
@@ -274,6 +276,28 @@ class FSequencerTrackFilter_ParticleTracks : public FSequencerTrackFilter_ClassT
 	virtual FSlateIcon GetIcon() const override { return FSlateIconFinder::FindIconForClass(UParticleSystem::StaticClass()); }
 };
 
+class FSequencerTrackFilter_CinematicShotTracks : public FSequencerTrackFilter_ClassType< UMovieSceneCinematicShotTrack >
+{
+	virtual FString GetName() const override { return TEXT("SequencerCinematicShotTracksFilter"); }
+	virtual FText GetDisplayName() const override { return LOCTEXT("SequencerTrackFilter_CinematicShotTracks", "Shots"); }
+	virtual FText GetToolTipText() const override { return LOCTEXT("SequencerTrackFilter_CinematicShotTracksToolTip", "Show only Shot tracks."); }
+	virtual FSlateIcon GetIcon() const override { return FSlateIcon(FAppStyle::GetAppStyleSetName(), "Sequencer.Tracks.CinematicShot"); }
+};
+
+class FSequencerTrackFilter_SubTracks : public FSequencerTrackFilter_ClassType< UMovieSceneSubTrack >
+{
+	virtual FString GetName() const override { return TEXT("SequencerSubTracksFilter"); }
+	virtual FText GetDisplayName() const override { return LOCTEXT("SequencerTrackFilter_SubTracks", "Subsequences"); }
+	virtual FText GetToolTipText() const override { return LOCTEXT("SequencerTrackFilter_SubTracksToolTip", "Show only Subsequence tracks."); }
+	virtual FSlateIcon GetIcon() const override { return FSlateIcon(FAppStyle::GetAppStyleSetName(), "Sequencer.Tracks.Sub"); }
+
+	// IFilter implementation
+	virtual bool PassesFilter(FTrackFilterType InItem) const override
+	{
+		return InItem->IsA(UMovieSceneSubTrack::StaticClass()) && !InItem->IsA(UMovieSceneCinematicShotTrack::StaticClass());
+	}
+};
+
 class FSequencerTrackFilter_SkeletalMeshObjects : public FSequencerTrackFilter_ComponentType< USkeletalMeshComponent >
 {
 	virtual FString GetName() const override { return TEXT("SequencerSkeletalMeshObjectsFilter"); }
