Sequencer: Check for null object template when attempting to paste a spawnable

diff --git a/Engine/Source/Editor/Sequencer/Private/SequencerUtilities.cpp b/Engine/Source/Editor/Sequencer/Private/SequencerUtilities.cpp
--- a/Engine/Source/Editor/Sequencer/Private/SequencerUtilities.cpp
+++ b/Engine/Source/Editor/Sequencer/Private/SequencerUtilities.cpp
@@ -2128,7 +2128,21 @@ bool FSequencerUtilities::PasteBindings(const FString& TextToImport, TSharedRef<
 				// an object template based off of our (transient package owned) template, then find the newly created binding
 				// and update it.
 
-				FGuid NewGuid = MakeNewSpawnable(Sequencer, *CopyableBinding->SpawnableObjectTemplate, nullptr, false, FName(*CopyableBinding->Spawnable.GetName()));
+				FGuid NewGuid;
+				if (CopyableBinding->SpawnableObjectTemplate)
+				{
+					NewGuid = MakeNewSpawnable(Sequencer, *CopyableBinding->SpawnableObjectTemplate, nullptr, false, FName(*CopyableBinding->Spawnable.GetName()));
+				}
+				else
+				{
+					FMovieSceneSpawnable NewSpawnable{};
+					NewSpawnable.SetGuid(FGuid::NewGuid());
+					NewSpawnable.SetName(CopyableBinding->Spawnable.GetName());
+
+					MovieScene->AddSpawnable(NewSpawnable, FMovieSceneBinding(NewSpawnable.GetGuid(), NewSpawnable.GetName()));
+
+					NewGuid = NewSpawnable.GetGuid();
+				}
 
 				FMovieSceneBinding NewBinding(NewGuid, CopyableBinding->Binding.GetName(), CopyableBinding->Tracks);
 				FMovieSceneSpawnable* Spawnable = MovieScene->FindSpawnable(NewGuid);
