Sequencer: Add support for multiple possessable object paste.

The previous incarnation relied on looking for the actor with the same name as the binding, which obviously isn't always the case, but was deemed good enough back then. This now will bind the pasted binding to all of the actors that the copied binding was bound to by actor name. For the situation that the actor is duplicated, it will also duplicate all of the bound actors instead of just the one.

diff --git a/Engine/Plugins/MovieScene/LevelSequenceEditor/Source/LevelSequenceEditor/Private/LevelSequenceEditorSubsystem.cpp b/Engine/Plugins/MovieScene/LevelSequenceEditor/Source/LevelSequenceEditor/Private/LevelSequenceEditorSubsystem.cpp
--- a/Engine/Plugins/MovieScene/LevelSequenceEditor/Source/LevelSequenceEditor/Private/LevelSequenceEditorSubsystem.cpp
+++ b/Engine/Plugins/MovieScene/LevelSequenceEditor/Source/LevelSequenceEditor/Private/LevelSequenceEditorSubsystem.cpp
@@ -402,8 +402,14 @@ bool ULevelSequenceEditorSubsystem::PasteTracks(const FString& InTextToImport, F
 
 void ULevelSequenceEditorSubsystem::CopyBindings(const TArray<FMovieSceneBindingProxy>& Bindings, FString& ExportedText)
 {
+	TSharedPtr<ISequencer> Sequencer = GetActiveSequencer();
+	if (Sequencer == nullptr)
+	{
+		return;
+	}
+
 	TArray<UMovieSceneFolder*> Folders;
-	FSequencerUtilities::CopyBindings(Bindings, Folders, ExportedText);
+	FSequencerUtilities::CopyBindings(Sequencer.ToSharedRef(), Bindings, Folders, ExportedText);
 	FPlatformApplicationMisc::ClipboardCopy(*ExportedText);
 }
 
diff --git a/Engine/Source/Editor/Sequencer/Private/MovieSceneCopyableBinding.h b/Engine/Source/Editor/Sequencer/Private/MovieSceneCopyableBinding.h
--- a/Engine/Source/Editor/Sequencer/Private/MovieSceneCopyableBinding.h
+++ b/Engine/Source/Editor/Sequencer/Private/MovieSceneCopyableBinding.h
@@ -43,6 +43,9 @@ public:
 	UPROPERTY()
 	FMovieScenePossessable Possessable;
 
+	UPROPERTY()
+	TArray<FString> BoundObjectNames;
+
 	UPROPERTY()
 	TArray<FName> FolderPath;
 };
diff --git a/Engine/Source/Editor/Sequencer/Private/Sequencer.cpp b/Engine/Source/Editor/Sequencer/Private/Sequencer.cpp
--- a/Engine/Source/Editor/Sequencer/Private/Sequencer.cpp
+++ b/Engine/Source/Editor/Sequencer/Private/Sequencer.cpp
@@ -7842,7 +7842,7 @@ void FSequencer::CopySelectedObjects(TArray<TSharedPtr<UE::Sequencer::FObjectBin
 		}
 	}
 
-	FSequencerUtilities::CopyBindings(Bindings, Folders, ExportedText);
+	FSequencerUtilities::CopyBindings(AsShared(), Bindings, Folders, ExportedText);
 		
 	// Make sure to clear the clipboard for the keys
 	GClipboardStack.Empty();
diff --git a/Engine/Source/Editor/Sequencer/Private/SequencerUtilities.cpp b/Engine/Source/Editor/Sequencer/Private/SequencerUtilities.cpp
--- a/Engine/Source/Editor/Sequencer/Private/SequencerUtilities.cpp
+++ b/Engine/Source/Editor/Sequencer/Private/SequencerUtilities.cpp
@@ -1852,7 +1852,7 @@ void ExportObjectBindingsToText(const TArray<UMovieSceneCopyableBinding*>& Objec
 	ExportedText = Archive;
 }
 
-void FSequencerUtilities::CopyBindings(const TArray<FMovieSceneBindingProxy>& Bindings, const TArray<UMovieSceneFolder*>& InFolders, FString& ExportedText)
+void FSequencerUtilities::CopyBindings(TSharedRef<ISequencer> Sequencer, const TArray<FMovieSceneBindingProxy>& Bindings, const TArray<UMovieSceneFolder*>& InFolders, FString& ExportedText)
 {
 	TArray<UMovieSceneCopyableBinding*> Objects;
 	for (const FMovieSceneBindingProxy& ObjectBinding : Bindings)
@@ -1870,6 +1870,15 @@ void FSequencerUtilities::CopyBindings(const TArray<FMovieSceneBindingProxy>& Bi
 		if (Possessable)
 		{
 			CopyableBinding->Possessable = *Possessable;
+
+			// Store the names of the bound objects so that they can be found on paste
+			for (TWeakObjectPtr<> RuntimeObject : Sequencer->FindBoundObjects(CopyableBinding->Possessable.GetGuid(), Sequencer->GetFocusedTemplateID()))
+			{
+				if (AActor* Actor = Cast<AActor>(RuntimeObject.Get()))
+				{
+					CopyableBinding->BoundObjectNames.Add(Actor->GetName());
+				}
+			}
 		}
 		else
 		{
@@ -2008,6 +2017,7 @@ bool FSequencerUtilities::PasteBindings(const FString& TextToImport, TSharedRef<
 
 	TMap<FGuid, FGuid> OldToNewGuidMap;
 	TArray<FGuid> PossessableGuids;
+	TArray<TArray<FString> > PossessableObjectNames;
 	TArray<FGuid> SpawnableGuids;
 	TMap<FGuid, UMovieSceneFolder*> GuidToFolderMap;
 
@@ -2075,12 +2085,13 @@ bool FSequencerUtilities::PasteBindings(const FString& TextToImport, TSharedRef<
 				}
 
 				TArray<AActor*> ActorsToDuplicate;
-				for (auto RuntimeObject : Sequencer->FindBoundObjects(CopyableBinding->Possessable.GetGuid(), Sequencer->GetFocusedTemplateID()))
+				for (TActorIterator<AActor> ActorItr(World); ActorItr; ++ActorItr)
 				{
-					AActor* Actor = Cast<AActor>(RuntimeObject.Get());
-					if (Actor)
+					AActor* Actor = *ActorItr;
+					if (Actor && CopyableBinding->BoundObjectNames.Contains(Actor->GetName()))
 					{
 						ActorsToDuplicate.Add(Actor);
+						CopyableBinding->BoundObjectNames.Remove(Actor->GetName());
 					}
 				}
 
@@ -2109,15 +2120,19 @@ bool FSequencerUtilities::PasteBindings(const FString& TextToImport, TSharedRef<
 						if (Actor)
 						{
 							DuplicatedActors.Add(Actor);
+
+							CopyableBinding->BoundObjectNames.Add(Actor->GetName());
 						}
 					}
 
 					// Bind the duplicated actors
 					if (DuplicatedActors.Num())
 					{
-						ReplaceBindingWithActors(Sequencer, DuplicatedActors, FMovieSceneBindingProxy(NewGuid, Sequence));
+						AddActorsToBinding(Sequencer, DuplicatedActors, FMovieSceneBindingProxy(NewGuid, Sequence));
 					}
 				}
+
+				PossessableObjectNames.Add(CopyableBinding->BoundObjectNames);
 			}
 			else if (CopyableBinding->Spawnable.GetGuid().IsValid())
 			{
@@ -2197,7 +2212,7 @@ bool FSequencerUtilities::PasteBindings(const FString& TextToImport, TSharedRef<
 			for (TActorIterator<AActor> ActorItr(PlaybackContext); ActorItr; ++ActorItr)
 			{
 				AActor* Actor = *ActorItr;
-				if (Actor && Actor->GetActorLabel() == *Possessable->GetName())
+				if (Actor && PossessableGuidIndex < PossessableObjectNames.Num() && PossessableObjectNames[PossessableGuidIndex].Contains(Actor->GetName()))
 				{
 					FGuid ExistingGuid = Sequencer->FindObjectId(*Actor, Sequencer->GetFocusedTemplateID());
 
diff --git a/Engine/Source/Editor/Sequencer/Public/SequencerUtilities.h b/Engine/Source/Editor/Sequencer/Public/SequencerUtilities.h
--- a/Engine/Source/Editor/Sequencer/Public/SequencerUtilities.h
+++ b/Engine/Source/Editor/Sequencer/Public/SequencerUtilities.h
@@ -169,7 +169,7 @@ struct SEQUENCER_API FSequencerUtilities
 	static bool CanPasteSections(const FString& TextToImport);
 
 	/** Copy/paste object bindings */
-	static void CopyBindings(const TArray<FMovieSceneBindingProxy>& Bindings, const TArray<UMovieSceneFolder*>& InFolders, FString& ExportedText);
+	static void CopyBindings(TSharedRef<ISequencer> Sequencer, const TArray<FMovieSceneBindingProxy>& Bindings, const TArray<UMovieSceneFolder*>& InFolders, FString& ExportedText);
 	static bool PasteBindings(const FString& TextToImport, TSharedRef<ISequencer> Sequencer, FMovieScenePasteBindingsParams PasteBindingsParams, TArray<FMovieSceneBindingProxy>& OutBindings, TArray<FNotificationInfo>& OutErrors);
 	static bool CanPasteBindings(TSharedRef<ISequencer> Sequencer, const FString& TextToImport);
 
