avoid n^2 loop when iterating thru items with nothing selected

diff --git a/Engine/Source/Editor/Sequencer/Private/Sequencer.cpp b/Engine/Source/Editor/Sequencer/Private/Sequencer.cpp
--- a/Engine/Source/Editor/Sequencer/Private/Sequencer.cpp
+++ b/Engine/Source/Editor/Sequencer/Private/Sequencer.cpp
@@ -1193,7 +1193,7 @@ void FSequencer::GetKeysFromSelection(TUniquePtr<FSequencerKeyCollection>& KeyCo
 	KeyCollection->Update(FSequencerKeyCollectionSignature::FromNodesRecursive(SelectedItems, ThresholdFrames));
 }
 
-void FSequencer::GetAllKeys(TUniquePtr<FSequencerKeyCollection>& KeyCollection, float DuplicateThresholdSeconds)
+void FSequencer::GetAllKeys(TUniquePtr<FSequencerKeyCollection>& KeyCollection, float DuplicateThresholdSeconds) const
 {
 	using namespace UE::Sequencer;
 
@@ -1202,11 +1202,8 @@ void FSequencer::GetAllKeys(TUniquePtr<FSequencerKeyCollection>& KeyCollection,
 		KeyCollection.Reset(new FSequencerKeyCollection);
 	}
 
-	TArray<TSharedRef<FViewModel>> AllNodes;
-	NodeTree->GetAllNodes(AllNodes);
-
-	FFrameNumber ThresholdFrames = (DuplicateThresholdSeconds * GetFocusedTickResolution()).FloorToFrame();
-	KeyCollection->Update(FSequencerKeyCollectionSignature::FromNodesRecursive(AllNodes, ThresholdFrames));
+	const FFrameNumber ThresholdFrames = (DuplicateThresholdSeconds * GetFocusedTickResolution()).FloorToFrame();
+	KeyCollection->Update(FSequencerKeyCollectionSignature::FromNodesRecursive({NodeTree->GetRootNode()->AsShared()}, ThresholdFrames));
 }
 
 
diff --git a/Engine/Source/Editor/Sequencer/Private/Sequencer.h b/Engine/Source/Editor/Sequencer/Private/Sequencer.h
--- a/Engine/Source/Editor/Sequencer/Private/Sequencer.h
+++ b/Engine/Source/Editor/Sequencer/Private/Sequencer.h
@@ -943,7 +943,7 @@ protected:
 	/** Get all the keys for the current sequencer selection */
 	virtual void GetKeysFromSelection(TUniquePtr<FSequencerKeyCollection>& KeyCollection, float DuplicateThresholdSeconds) override;
 
-	void GetAllKeys(TUniquePtr<FSequencerKeyCollection>& KeyCollection, float DuplicateThresoldSeconds);
+	void GetAllKeys(TUniquePtr<FSequencerKeyCollection>& KeyCollection, float DuplicateThresoldSeconds) const;
 
 	UMovieSceneSection* FindNextOrPreviousShot(UMovieSceneSequence* Sequence, FFrameNumber SearchFromTime, const bool bNext) const;
 
