21514370 Sequencer: UX Refresh for 5.1 - Nested tracklanes, key bars, snapping polish

  - Nested track lanes are now supported, allowing sequencer channels to be added as fully fledged child widgets of a parent, while still being arranged automatically and have their lifetime managed by their respective outliner item (therefore, they are now fully virtualized as you might expect)

  - Key bars have been added. This is exposed through a new ConnectionStyle member of FKeyDrawParams. Their rendering is also cached and batched within FKeyRenderer

  - Most relevant parts have been moved to SequencerCore now. Channel and category views remain in the Sequencer library due to their reliance on internal classes there, but they should be trivially moved to SequencerCore as a future task.

  - Track area now has a shared FTimeToPixel that can be shared amonst any widgets that need to convert times to track-area-pixel-space.

  - Added methods to ISequencerSection to allow client classes to create custom channel and category model types

  - Layer bar improvements:
    - Layer bars now encompass finite section ranges.
    - Resize handles are now only visible on hover as per design guidelines.

  - Inline curve display improvements:
    - Rendering is now part of each channel's view
    - Added API for determining curve extents
    - Curves are now rendered using FCurveModels

  - Snapping fixes and improvements
    - There were several places where snapping seemed not to work at all, or would snap inconsistently with dense keyframe data.
    - I've added a Weighting to the snap points so we can specify preferred snap points (section bounds, time cursor etc), and accumulate weights for specific frames where many keys exist. This helps to snap to times that are actually relevant, rather than always snapping to the nearest snap point, even if that snap point only has a single key on that time, and there are others within the snap threshold that have hundreds.
    - Converted the snapping utilities to operate on FFrameTimes rather than FFrameNumbers in order to support custom sequencer UIs that might not operate on integer time bases (it remains fully frame accurate with integers however)

  - Allow 'Key this section' to be unchecked

diff --git a/Engine/Plugins/Animation/ControlRig/Source/ControlRigEditor/Private/Sequencer/ControlRigSpaceChannelEditors.cpp b/Engine/Plugins/Animation/ControlRig/Source/ControlRigEditor/Private/Sequencer/ControlRigSpaceChannelEditors.cpp
--- a/Engine/Plugins/Animation/ControlRig/Source/ControlRigEditor/Private/Sequencer/ControlRigSpaceChannelEditors.cpp
+++ b/Engine/Plugins/Animation/ControlRig/Source/ControlRigEditor/Private/Sequencer/ControlRigSpaceChannelEditors.cpp
@@ -16,7 +16,7 @@
 #include "GameFramework/Actor.h"
 #include "Styling/AppStyle.h"
 #include "Styling/CoreStyle.h"
-#include "KeyDrawParams.h"
+#include "MVVM/Views/KeyDrawParams.h"
 #include "Framework/MultiBox/MultiBoxBuilder.h"
 #include "Channels/MovieSceneChannelProxy.h"
 #include "Channels/MovieSceneChannelEditorData.h"
diff --git a/Engine/Source/Editor/EditorStyle/Private/StarshipStyle.cpp b/Engine/Source/Editor/EditorStyle/Private/StarshipStyle.cpp
--- a/Engine/Source/Editor/EditorStyle/Private/StarshipStyle.cpp
+++ b/Engine/Source/Editor/EditorStyle/Private/StarshipStyle.cpp
@@ -1946,7 +1946,7 @@ void FStarshipEditorStyle::FStyle::SetupSequencerStyles()
 
 		Set("SequencerToolBar", SequencerToolbar);
 
-		Set("Sequencer.IconKeyAuto", new IMAGE_BRUSH("Sequencer/IconKeyAuto", Icon12x12));
+		Set("Sequencer.IconKeyAuto", new IMAGE_BRUSH_SVG("Sequencer/Keys/Key_CubicAuto", Icon12x12));
 		Set("Sequencer.IconKeyBreak", new IMAGE_BRUSH("Sequencer/IconKeyBreak", Icon12x12));
 		Set("Sequencer.IconKeyConstant", new IMAGE_BRUSH("Sequencer/IconKeyConstant", Icon12x12));
 		Set("Sequencer.IconKeyLinear", new IMAGE_BRUSH("Sequencer/IconKeyLinear", Icon12x12));
@@ -1965,6 +1965,10 @@ void FStarshipEditorStyle::FStyle::SetupSequencerStyles()
 		Set("Sequencer.TangentHandle", new IMAGE_BRUSH("Sequencer/TangentHandle", FVector2D(7, 7)));
 		Set("Sequencer.GenericDivider", new IMAGE_BRUSH("Sequencer/GenericDivider", FVector2D(2.f, 2.f), FLinearColor::White, ESlateBrushTileType::Vertical));
 
+		Set("Sequencer.KeyBar.Dotted", new BORDER_BRUSH(TEXT("Sequencer/Keys/KeyBar_Dotted"), FMargin(0.f, 4.f, 0.f, 0.f)));
+		Set("Sequencer.KeyBar.Dashed", new BORDER_BRUSH(TEXT("Sequencer/Keys/KeyBar_Dashed"), FMargin(0.f, 4.f, 0.f, 0.f)));
+		Set("Sequencer.KeyBar.Solid", new BORDER_BRUSH(TEXT("Sequencer/Keys/KeyBar_Solid"), FMargin(0.f, 4.f, 0.f, 0.f)));
+
 		Set("Sequencer.Timeline.ScrubHandleDown", new BOX_BRUSH("Sequencer/ScrubHandleDown", FMargin(6.f / 13.f, 5 / 12.f, 6 / 13.f, 8 / 12.f)));
 		Set("Sequencer.Timeline.ScrubHandleUp", new BOX_BRUSH("Sequencer/ScrubHandleUp", FMargin(6.f / 13.f, 8 / 12.f, 6 / 13.f, 5 / 12.f)));
 		Set("Sequencer.Timeline.ScrubFill", new BOX_BRUSH("Sequencer/ScrubFill", FMargin(2.f / 4.f, 0.f)));
@@ -2023,16 +2027,20 @@ void FStarshipEditorStyle::FStyle::SetupSequencerStyles()
 
 		Set("Sequencer.Section.Background", new BORDER_BRUSH(TEXT("Sequencer/SectionBackground"), FMargin(4.f / 16.f)));
 		Set("Sequencer.Section.BackgroundTint", new BOX_BRUSH(TEXT("Sequencer/SectionBackgroundTint"), FMargin(4 / 16.f)));
-		Set("Sequencer.Section.SelectedSectionOverlay", new IMAGE_BRUSH(TEXT("Sequencer/SelectedSectionOverlay"), Icon16x16, FLinearColor::White, ESlateBrushTileType::Both));
+		Set("Sequencer.Section.CollapsedSelectedSectionOverlay", new IMAGE_BRUSH(TEXT("Sequencer/Section/CollapsedSelectedSectionOverlay"), Icon16x16, FLinearColor::White, ESlateBrushTileType::Both));
+		Set("Sequencer.Section.SectionHeaderSelectedSectionOverlay", new IMAGE_BRUSH(TEXT("Sequencer/Section/SectionHeaderSelectedSectionOverlay"), Icon16x16, FLinearColor::White, ESlateBrushTileType::Both));
 		Set("Sequencer.Section.SelectedTrackTint", new BOX_BRUSH(TEXT("Sequencer/SelectedTrackTint"), FMargin(0.f, 0.5f)));
 		Set("Sequencer.Section.SelectionBorder", new BORDER_BRUSH(TEXT("Sequencer/SectionHighlight"), FMargin(7.f / 16.f)));
 		Set("Sequencer.Section.LockedBorder", new BORDER_BRUSH(TEXT("Sequencer/SectionLocked"), FMargin(7.f / 16.f)));
-		Set("Sequencer.Section.SelectedSectionOverlay", new IMAGE_BRUSH(TEXT("Sequencer/SelectedSectionOverlay"), Icon16x16, FLinearColor::White, ESlateBrushTileType::Both));
 		Set("Sequencer.Section.FilmBorder", new IMAGE_BRUSH(TEXT("Sequencer/SectionFilmBorder"), FVector2D(10, 7), FLinearColor::White, ESlateBrushTileType::Horizontal));
-		Set("Sequencer.Section.GripLeft", new BOX_BRUSH("Sequencer/SectionGripLeft", FMargin(5.f / 16.f)));
-		Set("Sequencer.Section.GripRight", new BOX_BRUSH("Sequencer/SectionGripRight", FMargin(5.f / 16.f)));
+		Set("Sequencer.Section.GripLeft", new FSlateRoundedBoxBrush(FStyleColors::White, FVector4(4, 0, 0, 4)));
+		Set("Sequencer.Section.GripRight", new FSlateRoundedBoxBrush(FStyleColors::White, FVector4(0, 4, 4, 0)));
 		Set("Sequencer.Section.EasingHandle", new IMAGE_BRUSH("Sequencer/EasingHandle", FVector2D(10.f, 10.f)));
 
+		Set("Sequencer.Section.Background_Collapsed", new FSlateRoundedBoxBrush(FLinearColor::White, FVector4(4, 4, 4, 4)));
+		Set("Sequencer.Section.Background_Header", new FSlateRoundedBoxBrush(FLinearColor::White, FVector4(4, 4, 0, 0)));
+		Set("Sequencer.Section.Background_Contents", new FSlateRoundedBoxBrush(FLinearColor::White, FVector4(0, 0, 4, 4)));
+
 		Set("Sequencer.Section.PreRoll", new BORDER_BRUSH(TEXT("Sequencer/PreRoll"), FMargin(0.f, .5f, 0.f, .5f)));
 
 		Set("Sequencer.Section.PinCusion", new IMAGE_BRUSH(TEXT("Sequencer/PinCusion"), Icon16x16, FLinearColor::White, ESlateBrushTileType::Both));
@@ -2042,9 +2050,8 @@ void FStarshipEditorStyle::FStyle::SetupSequencerStyles()
 		Set("Sequencer.Section.EmptySpace", new BOX_BRUSH(TEXT("Sequencer/EmptySpace"), FMargin(0.f, 7.f / 14.f)));
 
 		Set("Sequencer.LayerBar.Background",  new BOX_BRUSH(TEXT("Sequencer/LayerBarBackground"), FMargin(4.f / 16.f)));
-		Set("Sequencer.LayerBar.Border",      new BORDER_BRUSH(TEXT("Sequencer/LayerBarBorder"), FMargin(4.f / 16.f)));
-		Set("Sequencer.LayerBar.HandleLeft",  new BOX_BRUSH(TEXT("Sequencer/LayerBarHandleLeft"), FMargin(4.f / 16.f)));
-		Set("Sequencer.LayerBar.HandleRight", new BOX_BRUSH(TEXT("Sequencer/LayerBarHandleRight"), FMargin(4.f / 16.f)));
+		Set("Sequencer.LayerBar.HandleLeft",  new FSlateColorBrush(FStyleColors::White)); // Maybe use a rounded box if we decide to round-out the layer bars
+		Set("Sequencer.LayerBar.HandleRight", new FSlateColorBrush(FStyleColors::White)); // Maybe use a rounded box in we decide to round-out the layer bars
 
 		Set("Sequencer.ExposedNamePill_BG", new BOX_BRUSH(TEXT("Sequencer/ExposedNamePill_BG"), FMargin(14.f / 30.f), FLinearColor(1.f, 1.f, 1.f, .8f)));
 		Set("Sequencer.ExposedNamePill", FButtonStyle()
diff --git a/Engine/Source/Editor/MovieSceneTools/Private/Channels/BuiltInChannelEditors.cpp b/Engine/Source/Editor/MovieSceneTools/Private/Channels/BuiltInChannelEditors.cpp
--- a/Engine/Source/Editor/MovieSceneTools/Private/Channels/BuiltInChannelEditors.cpp
+++ b/Engine/Source/Editor/MovieSceneTools/Private/Channels/BuiltInChannelEditors.cpp
@@ -20,7 +20,7 @@
 #include "CurveKeyEditors/SStringCurveKeyEditor.h"
 #include "CurveKeyEditors/SEnumKeyEditor.h"
 #include "UObject/StructOnScope.h"
-#include "KeyDrawParams.h"
+#include "MVVM/Views/KeyDrawParams.h"
 #include "Framework/MultiBox/MultiBoxBuilder.h"
 #include "Channels/MovieSceneChannelProxy.h"
 #include "Channels/MovieSceneChannelEditorData.h"
@@ -784,6 +784,7 @@ void DrawKeysImpl(ChannelType* Channel, TArrayView<const FKeyHandle> InKeyHandle
 
 	FKeyDrawParams TempParams;
 	TempParams.BorderBrush = TempParams.FillBrush = DiamondKeyBrush;
+	TempParams.ConnectionStyle = EKeyConnectionStyle::Solid;
 
 	for (int32 Index = 0; Index < InKeyHandles.Num(); ++Index)
 	{
@@ -795,6 +796,7 @@ void DrawKeysImpl(ChannelType* Channel, TArrayView<const FKeyHandle> InKeyHandle
 		ERichCurveTangentMode TangentMode = KeyIndex == INDEX_NONE ? RCTM_None : Values[KeyIndex].TangentMode.GetValue();
 
 		TempParams.FillOffset = FVector2D(0.f, 0.f);
+		TempParams.ConnectionStyle = EKeyConnectionStyle::Solid;
 
 		switch (InterpMode)
 		{
@@ -807,6 +809,7 @@ void DrawKeysImpl(ChannelType* Channel, TArrayView<const FKeyHandle> InKeyHandle
 		case RCIM_Constant:
 			TempParams.BorderBrush = TempParams.FillBrush = SquareKeyBrush;
 			TempParams.FillTint = FLinearColor(0.0f, 0.445f, 0.695f, 1.0f); // blue
+			TempParams.ConnectionStyle = EKeyConnectionStyle::Dashed;
 			break;
 
 		case RCIM_Cubic:
@@ -1177,11 +1180,11 @@ struct TCurveChannelSectionMenuExtension : TSharedFromThis<TCurveChannelSectionM
 		};
 		auto GetKeyAreaCurveNormalized = [=](FString KeyAreaName) { return !Settings->HasKeyAreaCurveExtents(KeyAreaName); };
 
-		auto OnKeyAreaCurveMinChanged = [=](float NewValue, FString KeyAreaName) { float CurveMin = 0.f; float CurveMax = 0.f; Settings->GetKeyAreaCurveExtents(KeyAreaName, CurveMin, CurveMax); Settings->SetKeyAreaCurveExtents(KeyAreaName, NewValue, CurveMax); };
-		auto GetKeyAreaCurveMin = [=](FString KeyAreaName) { float CurveMin = 0.f; float CurveMax = 0.f; Settings->GetKeyAreaCurveExtents(KeyAreaName, CurveMin, CurveMax); return CurveMin; };
+		auto OnKeyAreaCurveMinChanged = [=](double NewValue, FString KeyAreaName) { double CurveMin = 0.f; double CurveMax = 0.f; Settings->GetKeyAreaCurveExtents(KeyAreaName, CurveMin, CurveMax); Settings->SetKeyAreaCurveExtents(KeyAreaName, NewValue, CurveMax); };
+		auto GetKeyAreaCurveMin = [=](FString KeyAreaName) { double CurveMin = 0.f; double CurveMax = 0.f; Settings->GetKeyAreaCurveExtents(KeyAreaName, CurveMin, CurveMax); return CurveMin; };
 
-		auto OnKeyAreaCurveMaxChanged = [=](float NewValue, FString KeyAreaName) { float CurveMin = 0.f; float CurveMax = 0.f; Settings->GetKeyAreaCurveExtents(KeyAreaName, CurveMin, CurveMax); Settings->SetKeyAreaCurveExtents(KeyAreaName, CurveMin, NewValue); };
-		auto GetKeyAreaCurveMax = [=](FString KeyAreaName) { float CurveMin = 0.f; float CurveMax = 0.f; Settings->GetKeyAreaCurveExtents(KeyAreaName, CurveMin, CurveMax); return CurveMax; };
+		auto OnKeyAreaCurveMaxChanged = [=](double NewValue, FString KeyAreaName) { double CurveMin = 0.f; double CurveMax = 0.f; Settings->GetKeyAreaCurveExtents(KeyAreaName, CurveMin, CurveMax); Settings->SetKeyAreaCurveExtents(KeyAreaName, CurveMin, NewValue); };
+		auto GetKeyAreaCurveMax = [=](FString KeyAreaName) { double CurveMin = 0.f; double CurveMax = 0.f; Settings->GetKeyAreaCurveExtents(KeyAreaName, CurveMin, CurveMax); return CurveMax; };
 
 		MenuBuilder.AddMenuEntry(
 			LOCTEXT("ToggleShowCurve", "Show Curve"),
@@ -1234,11 +1237,11 @@ struct TCurveChannelSectionMenuExtension : TSharedFromThis<TCurveChannelSectionM
 					.WidthOverride(50.f)
 					.IsEnabled_Lambda([=]() { return SharedThis->IsShowCurve() && Settings->HasKeyAreaCurveExtents(KeyAreaName); })
 					[
-						SNew(SSpinBox<float>)
+						SNew(SSpinBox<double>)
 						.Style(&FAppStyle::GetWidgetStyle<FSpinBoxStyle>("Sequencer.HyperlinkSpinBox"))
-						.OnValueCommitted_Lambda([=](float NewValue, ETextCommit::Type CommitType) { OnKeyAreaCurveMinChanged(NewValue, KeyAreaName); })
-						.OnValueChanged_Lambda([=](float NewValue) { OnKeyAreaCurveMinChanged(NewValue, KeyAreaName); })
-						.Value_Lambda([=]() -> float { return GetKeyAreaCurveMin(KeyAreaName); })
+						.OnValueCommitted_Lambda([=](double NewValue, ETextCommit::Type CommitType) { OnKeyAreaCurveMinChanged(NewValue, KeyAreaName); })
+						.OnValueChanged_Lambda([=](double NewValue) { OnKeyAreaCurveMinChanged(NewValue, KeyAreaName); })
+						.Value_Lambda([=]() -> double { return GetKeyAreaCurveMin(KeyAreaName); })
 					]
 				]
 			+ SHorizontalBox::Slot()
@@ -1248,11 +1251,11 @@ struct TCurveChannelSectionMenuExtension : TSharedFromThis<TCurveChannelSectionM
 					.WidthOverride(50.f)
 					.IsEnabled_Lambda([=]() { return SharedThis->IsShowCurve() && Settings->HasKeyAreaCurveExtents(KeyAreaName); })
 					[
-						SNew(SSpinBox<float>)
+						SNew(SSpinBox<double>)
 						.Style(&FAppStyle::GetWidgetStyle<FSpinBoxStyle>("Sequencer.HyperlinkSpinBox"))
-						.OnValueCommitted_Lambda([=](float NewValue, ETextCommit::Type CommitType) { OnKeyAreaCurveMaxChanged(NewValue, KeyAreaName); })
-						.OnValueChanged_Lambda([=](float NewValue) { OnKeyAreaCurveMaxChanged(NewValue, KeyAreaName); })
-						.Value_Lambda([=]() -> float { return GetKeyAreaCurveMax(KeyAreaName); })
+						.OnValueCommitted_Lambda([=](double NewValue, ETextCommit::Type CommitType) { OnKeyAreaCurveMaxChanged(NewValue, KeyAreaName); })
+						.OnValueChanged_Lambda([=](double NewValue) { OnKeyAreaCurveMaxChanged(NewValue, KeyAreaName); })
+						.Value_Lambda([=]() -> double { return GetKeyAreaCurveMax(KeyAreaName); })
 					]
 				],
 				LOCTEXT("KeyAreaCurveRangeText", "Key Area Curve Range")
@@ -1565,5 +1568,13 @@ TUniquePtr<FCurveModel> CreateCurveEditorModel(const TMovieSceneChannelHandle<FM
 	return MakeUnique<FEventChannelCurveModel>(EventChannel, OwningSection, InSequencer);
 }
 
+bool ShouldShowCurve(const FMovieSceneFloatChannel* Channel, UMovieSceneSection* InSection)
+{
+	return Channel->GetShowCurve();
+}
+bool ShouldShowCurve(const FMovieSceneDoubleChannel* Channel, UMovieSceneSection* InSection)
+{
+	return Channel->GetShowCurve();
+}
 
 #undef LOCTEXT_NAMESPACE
diff --git a/Engine/Source/Editor/MovieSceneTools/Private/Channels/BuiltInChannelEditors.h b/Engine/Source/Editor/MovieSceneTools/Private/Channels/BuiltInChannelEditors.h
--- a/Engine/Source/Editor/MovieSceneTools/Private/Channels/BuiltInChannelEditors.h
+++ b/Engine/Source/Editor/MovieSceneTools/Private/Channels/BuiltInChannelEditors.h
@@ -86,3 +86,6 @@ TUniquePtr<FCurveModel> CreateCurveEditorModel(const TMovieSceneChannelHandle<FM
 TUniquePtr<FCurveModel> CreateCurveEditorModel(const TMovieSceneChannelHandle<FMovieSceneIntegerChannel>& IntegerChannel, UMovieSceneSection* OwningSection, TSharedRef<ISequencer> InSequencer);
 TUniquePtr<FCurveModel> CreateCurveEditorModel(const TMovieSceneChannelHandle<FMovieSceneBoolChannel>& BoolChannel, UMovieSceneSection* OwningSection, TSharedRef<ISequencer> InSequencer);
 TUniquePtr<FCurveModel> CreateCurveEditorModel(const TMovieSceneChannelHandle<FMovieSceneEventChannel>& EventChannel, UMovieSceneSection* OwningSection, TSharedRef<ISequencer> InSequencer);
+
+bool ShouldShowCurve(const FMovieSceneFloatChannel* Channel, UMovieSceneSection* InSection);
+bool ShouldShowCurve(const FMovieSceneDoubleChannel* Channel, UMovieSceneSection* InSection);
\ No newline at end of file
diff --git a/Engine/Source/Editor/MovieSceneTools/Private/Sections/ParameterSection.cpp b/Engine/Source/Editor/MovieSceneTools/Private/Sections/ParameterSection.cpp
--- a/Engine/Source/Editor/MovieSceneTools/Private/Sections/ParameterSection.cpp
+++ b/Engine/Source/Editor/MovieSceneTools/Private/Sections/ParameterSection.cpp
@@ -1,7 +1,7 @@
 // Copyright Epic Games, Inc. All Rights Reserved.
 
 #include "Sections/ParameterSection.h"
-#include "ISectionLayoutBuilder.h"
+#include "MVVM/ViewModels/CategoryModel.h"
 #include "ScopedTransaction.h"
 #include "Sections/MovieSceneParameterSection.h"
 #include "SequencerSectionPainter.h"
@@ -16,35 +16,39 @@
 
 #define LOCTEXT_NAMESPACE "ParameterSection"
 
-FReply FParameterSection::OnKeyDoubleClicked(const TArray<FKeyHandle> &KeyHandles )
+namespace UE::Sequencer
 {
-	UMovieSceneParameterSection* ParameterSection = Cast<UMovieSceneParameterSection>( WeakSection.Get() );
-	if (!ParameterSection)
+	class SColorStripView : public SLeafWidget, public ITrackLaneWidget
 	{
-		return FReply::Handled();
-	}
+		SLATE_BEGIN_ARGS(SColorStripView){}
+		SLATE_END_ARGS()
 
-	for (FColorParameterNameAndCurves& NameAndCurve : ParameterSection->GetColorParameterNamesAndCurves())
-	{
-		FMovieSceneKeyColorPicker KeyColorPicker(ParameterSection, &NameAndCurve.RedCurve, &NameAndCurve.GreenCurve, &NameAndCurve.BlueCurve, &NameAndCurve.AlphaCurve, KeyHandles);
-	}
+		void Construct(const FArguments& InArgs, TViewModelPtr<ITrackLaneExtension> InModel)
+		{
+			WeakModel = InModel;
+		}
 
-	return FReply::Handled();
-}
+		TSharedRef<const SWidget> AsWidget() const override
+		{
+			return AsShared();
+		}
 
-int32 FParameterSection::OnPaintSection(FSequencerSectionPainter& Painter) const
-{
-	int32 LayerId = FSequencerSection::OnPaintSection(Painter);
+		FTrackLaneScreenAlignment GetAlignment(const FTimeToPixel& InTimeToPixel, const FGeometry& InParentGeometry) const override
+		{
+			TSharedPtr<ITrackLaneExtension> TrackLaneExtension = WeakModel.ImplicitPin();
+			if (TrackLaneExtension)
+			{
+				FTrackLaneVirtualAlignment VirtualAlignment = TrackLaneExtension->ArrangeVirtualTrackLaneView();
+				FTrackLaneScreenAlignment  ScreenAlignment  = VirtualAlignment.ToScreen(InTimeToPixel, InParentGeometry);
 
-	for (FKeyAreaElement KeyAreaElement : Painter.KeyAreaElements)
-	{
-		FMovieSceneFloatChannel* RChannel = nullptr;
-		FMovieSceneFloatChannel* GChannel = nullptr;
-		FMovieSceneFloatChannel* BChannel = nullptr;
-		FMovieSceneFloatChannel* AChannel = nullptr;
-		
-		for (TSharedRef<IKeyArea> KeyArea : KeyAreaElement.KeyAreas)
+				return ScreenAlignment;
+			}
+			return FTrackLaneScreenAlignment();
+		}
+
+		int32 OnPaint( const FPaintArgs& Args, const FGeometry& AllottedGeometry, const FSlateRect& MyCullingRect, FSlateWindowElementList& OutDrawElements, int32 LayerId, const FWidgetStyle& InWidgetStyle, bool bParentEnabled ) const override
 		{
+			#if 0
 			const FMovieSceneChannelMetaData* MetaData = KeyArea->GetChannel().GetMetaData();
 			if (MetaData)
 			{
@@ -65,80 +69,151 @@ int32 FParameterSection::OnPaintSection(FSequencerSectionPainter& Painter) const
 					AChannel = static_cast<FMovieSceneFloatChannel*>(KeyArea->ResolveChannel());
 				}
 			}
-		}
 
-		if (RChannel && GChannel && BChannel && AChannel)
-		{
-			const ESlateDrawEffect DrawEffects = Painter.bParentEnabled ? ESlateDrawEffect::None : ESlateDrawEffect::DisabledEffect;
+			if (RChannel && GChannel && BChannel && AChannel)
+			{
+				const ESlateDrawEffect DrawEffects = Painter.bParentEnabled ? ESlateDrawEffect::None : ESlateDrawEffect::DisabledEffect;
 
-			const FTimeToPixel& TimeConverter = Painter.GetTimeConverter();
+				const FTimeToPixel& TimeConverter = Painter.GetTimeConverter();
 
-			FGeometry KeyAreaGeometry = KeyAreaElement.KeyAreaGeometry;
+				FGeometry KeyAreaGeometry = KeyAreaElement.KeyAreaGeometry;
 
-			const float StartTime       = TimeConverter.PixelToSeconds(0.f);
-			const float EndTime         = TimeConverter.PixelToSeconds(KeyAreaGeometry.GetLocalSize().X);
-			const float SectionDuration = EndTime - StartTime;
+				const float StartTime       = TimeConverter.PixelToSeconds(0.f);
+				const float EndTime         = TimeConverter.PixelToSeconds(KeyAreaGeometry.GetLocalSize().X);
+				const float SectionDuration = EndTime - StartTime;
 
-			const float KeyAreaHeight = KeyAreaElement.Type == FKeyAreaElement::Group ? KeyAreaGeometry.Size.Y / 3.f - 4.0f : KeyAreaGeometry.Size.Y - 3.0f;
+				const float KeyAreaHeight = KeyAreaElement.Type == FKeyAreaElement::Group ? KeyAreaGeometry.Size.Y / 3.f - 4.0f : KeyAreaGeometry.Size.Y - 3.0f;
 
-			FVector2D GradientSize = FVector2D( KeyAreaGeometry.Size.X - 2.f, KeyAreaHeight );
-			if ( GradientSize.X >= 1.f )
-			{
-				FPaintGeometry PaintGeometry = KeyAreaGeometry.ToPaintGeometry( FVector2D( 1.f, 1.f ), GradientSize );
-
-				// If we are showing a background pattern and the colors is transparent, draw a checker pattern
-				FSlateDrawElement::MakeBox(
-					Painter.DrawElements,
-					LayerId,
-					PaintGeometry,
-					FAppStyle::GetBrush( "Checker" ),
-					DrawEffects);
-
-				FLinearColor DefaultColor = FLinearColor::Black;
-				DefaultColor.R = RChannel->GetDefault().IsSet() ? RChannel->GetDefault().GetValue() : 0.f;
-				DefaultColor.G = GChannel->GetDefault().IsSet() ? GChannel->GetDefault().GetValue() : 0.f;
-				DefaultColor.B = BChannel->GetDefault().IsSet() ? BChannel->GetDefault().GetValue() : 0.f;
-				DefaultColor.A = AChannel->GetDefault().IsSet() ? AChannel->GetDefault().GetValue() : 0.f;
-					
-				TArray<FMovieSceneFloatChannel*> ColorChannels;
-				ColorChannels.Add(RChannel);
-				ColorChannels.Add(GChannel);
-				ColorChannels.Add(BChannel);
-				ColorChannels.Add(AChannel);
-
-				TArray< TTuple<float, FLinearColor> > ColorKeys;
-				MovieSceneSectionHelpers::ConsolidateColorCurves( ColorKeys, DefaultColor, ColorChannels, TimeConverter );
-
-				TArray<FSlateGradientStop> GradientStops;
-
-				for (const TTuple<float, FLinearColor>& ColorStop : ColorKeys)
+				FVector2D GradientSize = FVector2D( KeyAreaGeometry.Size.X - 2.f, KeyAreaHeight );
+				if ( GradientSize.X >= 1.f )
 				{
-					const float Time = ColorStop.Get<0>();
-
-					// HACK: The color is converted to SRgb and then reinterpreted as linear here because gradients are converted to FColor
-					// without the SRgb conversion before being passed to the renderer for some reason.
-					const FLinearColor Color = ColorStop.Get<1>().ToFColor( true ).ReinterpretAsLinear();
+					FPaintGeometry PaintGeometry = KeyAreaGeometry.ToPaintGeometry( FVector2D( 1.f, 1.f ), GradientSize );
 
-					float TimeFraction = (Time - StartTime) / SectionDuration;
-					GradientStops.Add( FSlateGradientStop( FVector2D( TimeFraction * KeyAreaGeometry.Size.X, 0 ), Color ) );
-				}
-
-				if ( GradientStops.Num() > 0 )
-				{
-					FSlateDrawElement::MakeGradient(
+					// If we are showing a background pattern and the colors is transparent, draw a checker pattern
+					FSlateDrawElement::MakeBox(
 						Painter.DrawElements,
-						Painter.LayerId + 1,
+						LayerId,
 						PaintGeometry,
-						GradientStops,
-						Orient_Vertical,
-						DrawEffects
-						);
+						FAppStyle::GetBrush( "Checker" ),
+						DrawEffects);
+
+					FLinearColor DefaultColor = FLinearColor::Black;
+					DefaultColor.R = RChannel->GetDefault().IsSet() ? RChannel->GetDefault().GetValue() : 0.f;
+					DefaultColor.G = GChannel->GetDefault().IsSet() ? GChannel->GetDefault().GetValue() : 0.f;
+					DefaultColor.B = BChannel->GetDefault().IsSet() ? BChannel->GetDefault().GetValue() : 0.f;
+					DefaultColor.A = AChannel->GetDefault().IsSet() ? AChannel->GetDefault().GetValue() : 0.f;
+						
+					TArray<FMovieSceneFloatChannel*> ColorChannels;
+					ColorChannels.Add(RChannel);
+					ColorChannels.Add(GChannel);
+					ColorChannels.Add(BChannel);
+					ColorChannels.Add(AChannel);
+
+					TArray< TTuple<float, FLinearColor> > ColorKeys;
+					MovieSceneSectionHelpers::ConsolidateColorCurves( ColorKeys, DefaultColor, ColorChannels, TimeConverter );
+
+					TArray<FSlateGradientStop> GradientStops;
+
+					for (const TTuple<float, FLinearColor>& ColorStop : ColorKeys)
+					{
+						const float Time = ColorStop.Get<0>();
+
+						// HACK: The color is converted to SRgb and then reinterpreted as linear here because gradients are converted to FColor
+						// without the SRgb conversion before being passed to the renderer for some reason.
+						const FLinearColor Color = ColorStop.Get<1>().ToFColor( true ).ReinterpretAsLinear();
+
+						float TimeFraction = (Time - StartTime) / SectionDuration;
+						GradientStops.Add( FSlateGradientStop( FVector2D( TimeFraction * KeyAreaGeometry.Size.X, 0 ), Color ) );
+					}
+
+					if ( GradientStops.Num() > 0 )
+					{
+						FSlateDrawElement::MakeGradient(
+							Painter.DrawElements,
+							Painter.LayerId + 1,
+							PaintGeometry,
+							GradientStops,
+							Orient_Vertical,
+							DrawEffects
+							);
+					}
 				}
 			}
+			#endif
+			return LayerId;
+		}
+		virtual FVector2D ComputeDesiredSize(float) const override
+		{
+			return FVector2D();
 		}
+
+	private:
+		TWeakViewModelPtr<ITrackLaneExtension> WeakModel;
+	};
+}
+
+FReply FParameterSection::OnKeyDoubleClicked(const TArray<FKeyHandle> &KeyHandles )
+{
+	UMovieSceneParameterSection* ParameterSection = Cast<UMovieSceneParameterSection>( WeakSection.Get() );
+	if (!ParameterSection)
+	{
+		return FReply::Handled();
+	}
+
+	for (FColorParameterNameAndCurves& NameAndCurve : ParameterSection->GetColorParameterNamesAndCurves())
+	{
+		FMovieSceneKeyColorPicker KeyColorPicker(ParameterSection, &NameAndCurve.RedCurve, &NameAndCurve.GreenCurve, &NameAndCurve.BlueCurve, &NameAndCurve.AlphaCurve, KeyHandles);
+	}
+
+	return FReply::Handled();
+}
+
+TSharedPtr<UE::Sequencer::FCategoryModel> FParameterSection::ConstructCategoryModel(FName InCategoryName, const FText& InDisplayText, TArrayView<const FChannelData> Channels) const
+{
+	using namespace UE::Sequencer;
+
+	// Only construct the color category if it has all the channels
+	uint8 ColorChannelMask = 0;
+
+	for (const FChannelData& Channel : Channels)
+	{
+		if (Channel.MetaData.DisplayText.EqualTo(FCommonChannelData::ChannelR))
+		{
+			ColorChannelMask |= 1 << 0;
+		}
+		else if (Channel.MetaData.DisplayText.EqualTo(FCommonChannelData::ChannelG))
+		{
+			ColorChannelMask |= 1 << 1;
+		}
+		else if (Channel.MetaData.DisplayText.EqualTo(FCommonChannelData::ChannelB))
+		{
+			ColorChannelMask |= 1 << 2;
+		}
+		else if (Channel.MetaData.DisplayText.EqualTo(FCommonChannelData::ChannelA))
+		{
+			ColorChannelMask |= 1 << 3;
+		}
+	}
+
+	if (ColorChannelMask == 0b1111)
+	{
+		class FColorCategory : public FCategoryModel
+		{
+		public:
+			FColorCategory(FName InCategoryName)
+				: FCategoryModel(InCategoryName)
+			{}
+
+			TSharedPtr<ITrackLaneWidget> CreateTrackLaneView(const FCreateTrackLaneViewParams& InParams) override
+			{
+				return SNew(SColorStripView, SharedThis(this));
+			}
+		};
+
+		return MakeShared<FColorCategory>(InCategoryName);
 	}
 
-	return LayerId;
+	return nullptr;
 }
 
 bool FParameterSection::RequestDeleteCategory( const TArray<FName>& CategoryNamePath )
diff --git a/Engine/Source/Editor/MovieSceneTools/Private/Sections/ParameterSection.h b/Engine/Source/Editor/MovieSceneTools/Private/Sections/ParameterSection.h
--- a/Engine/Source/Editor/MovieSceneTools/Private/Sections/ParameterSection.h
+++ b/Engine/Source/Editor/MovieSceneTools/Private/Sections/ParameterSection.h
@@ -20,7 +20,8 @@ public:
 
 	//~ ISequencerSection interface
 	virtual FReply OnKeyDoubleClicked(const TArray<FKeyHandle>& KeyHandles) override;
-	virtual int32 OnPaintSection(FSequencerSectionPainter& InPainter) const override;
 	virtual bool RequestDeleteCategory(const TArray<FName>& CategoryNamePath) override;
 	virtual bool RequestDeleteKeyArea(const TArray<FName>& KeyAreaNamePath) override;
+
+	virtual TSharedPtr<UE::Sequencer::FCategoryModel> ConstructCategoryModel(FName InCategoryName, const FText& InDisplayText, TArrayView<const FChannelData> Channels) const;
 };
diff --git a/Engine/Source/Editor/Sequencer/Private/DisplayNodes/SequencerKeyTimeCache.cpp b/Engine/Source/Editor/Sequencer/Private/DisplayNodes/SequencerKeyTimeCache.cpp
deleted file mode 100644
diff --git a/Engine/Source/Editor/Sequencer/Private/DisplayNodes/SequencerKeyTimeCache.h b/Engine/Source/Editor/Sequencer/Private/DisplayNodes/SequencerKeyTimeCache.h
deleted file mode 100644
diff --git a/Engine/Source/Editor/Sequencer/Private/IKeyArea.cpp b/Engine/Source/Editor/Sequencer/Private/IKeyArea.cpp
--- a/Engine/Source/Editor/Sequencer/Private/IKeyArea.cpp
+++ b/Engine/Source/Editor/Sequencer/Private/IKeyArea.cpp
@@ -330,3 +330,16 @@ TUniquePtr<FCurveModel> IKeyArea::CreateCurveEditorModel(TSharedRef<ISequencer>
 
 	return nullptr;
 }
+
+bool IKeyArea::ShouldShowCurve() const
+{
+	ISequencerChannelInterface* EditorInterface = FindChannelEditorInterface();
+	FMovieSceneChannel* Channel = ChannelHandle.Get();
+	UMovieSceneSection* OwningSection = GetOwningSection();
+	if (EditorInterface && Channel && OwningSection)
+	{
+		return EditorInterface->ShouldShowCurve_Raw(Channel, OwningSection);
+	}
+
+	return false;
+}
\ No newline at end of file
diff --git a/Engine/Source/Editor/Sequencer/Private/ISequencerSection.cpp b/Engine/Source/Editor/Sequencer/Private/ISequencerSection.cpp
--- a/Engine/Source/Editor/Sequencer/Private/ISequencerSection.cpp
+++ b/Engine/Source/Editor/Sequencer/Private/ISequencerSection.cpp
@@ -7,16 +7,6 @@
 #include "IKeyArea.h"
 #include "Channels/MovieSceneChannelProxy.h"
 
-/** Structure used during key area creation to group channels by their group name */
-struct FChannelData
-{
-	/** Handle to the channel */
-	FMovieSceneChannelHandle Channel;
-
-	/** The channel's editor meta data */
-	const FMovieSceneChannelMetaData& MetaData;
-};
-
 /** Data pertaining to a group of channels */
 struct FGroupData
 {
@@ -25,7 +15,7 @@ struct FGroupData
 		, SortOrder(-1)
 	{}
 
-	void AddChannel(FChannelData&& InChannel)
+	void AddChannel(ISequencerSection::FChannelData&& InChannel)
 	{
 		if (InChannel.MetaData.SortOrder < SortOrder)
 		{
@@ -42,7 +32,7 @@ struct FGroupData
 	uint32 SortOrder;
 
 	/** Array of channels within this group */
-	TArray<FChannelData, TInlineAllocator<4>> Channels;
+	TArray<ISequencerSection::FChannelData, TInlineAllocator<4>> Channels;
 };
 
 void ISequencerSection::GenerateSectionLayout( ISectionLayoutBuilder& LayoutBuilder )
@@ -90,20 +80,25 @@ void ISequencerSection::GenerateSectionLayout( ISectionLayoutBuilder& LayoutBuil
 		return;
 	}
 
+	auto ChannelFactory = [this](FName InChannelName, const FMovieSceneChannelHandle& InChannel)
+	{
+		return this->ConstructChannelModel(InChannelName, InChannel);
+	};
+
 	// Collapse single channels to the top level track node if allowed
 	if (GroupToChannelsMap.Num() == 1)
 	{
 		const TTuple<FName, FGroupData>& Pair = *GroupToChannelsMap.CreateIterator();
 		if (Pair.Value.Channels.Num() == 1 && Pair.Value.Channels[0].MetaData.bCanCollapseToTrack)
 		{
-			LayoutBuilder.SetTopLevelChannel(Pair.Value.Channels[0].Channel);
+			LayoutBuilder.SetTopLevelChannel(Pair.Value.Channels[0].Channel, ChannelFactory);
 			return;
 		}
 	}
 
 	// Sort the channels in each group by its sort order and name
 	TArray<FName, TInlineAllocator<6>> SortedGroupNames;
-	for (auto& Pair : GroupToChannelsMap)
+	for (TPair<FName, FGroupData>& Pair : GroupToChannelsMap)
 	{
 		SortedGroupNames.Add(Pair.Key);
 
@@ -138,16 +133,21 @@ void ISequencerSection::GenerateSectionLayout( ISectionLayoutBuilder& LayoutBuil
 	// Create key areas for each group name
 	for (FName GroupName : SortedGroupNames)
 	{
-		auto& ChannelData = GroupToChannelsMap.FindChecked(GroupName);
+		FGroupData& ChannelData = GroupToChannelsMap.FindChecked(GroupName);
 
 		if (!GroupName.IsNone())
 		{
-			LayoutBuilder.PushCategory(GroupName, ChannelData.GroupText);
+			auto Factory = [this, &ChannelData](FName InCategoryName, const FText& InDisplayText)
+			{
+				return this->ConstructCategoryModel(InCategoryName, InDisplayText, ChannelData.Channels);
+			};
+
+			LayoutBuilder.PushCategory(GroupName, ChannelData.GroupText, Factory);
 		}
 
 		for (const FChannelData& ChannelAndData : ChannelData.Channels)
 		{
-			LayoutBuilder.AddChannel(ChannelAndData.Channel);
+			LayoutBuilder.AddChannel(ChannelAndData.Channel, ChannelFactory);
 		}
 
 		if (!GroupName.IsNone())
diff --git a/Engine/Source/Editor/Sequencer/Private/MVVM/ViewModels/CategoryModel.cpp b/Engine/Source/Editor/Sequencer/Private/MVVM/ViewModels/CategoryModel.cpp
--- a/Engine/Source/Editor/Sequencer/Private/MVVM/ViewModels/CategoryModel.cpp
+++ b/Engine/Source/Editor/Sequencer/Private/MVVM/ViewModels/CategoryModel.cpp
@@ -6,8 +6,11 @@
 #include "MVVM/ViewModels/SequencerModelUtils.h"
 #include "MVVM/Views/SOutlinerItemViewBase.h"
 #include "MVVM/Views/SSequencerKeyNavigationButtons.h"
+#include "MVVM/Views/SChannelView.h"
 #include "MVVM/ViewModels/SequencerEditorViewModel.h"
 
+#include "SequencerSectionPainter.h"
+
 #include "Styling/AppStyle.h"
 
 #include "ISequencerSection.h"
@@ -46,7 +49,8 @@ FOutlinerSizing FCategoryModel::GetDesiredSizing() const
 
 TSharedPtr<ITrackLaneWidget> FCategoryModel::CreateTrackLaneView(const FCreateTrackLaneViewParams& InParams)
 {
-	return nullptr;
+	return SNew(SChannelView, SharedThis(this), InParams.TimeToPixel, InParams.Editor->GetTrackArea())
+		.KeyBarColor(this, &FCategoryModel::GetKeyBarColor);
 }
 
 FTrackLaneVirtualAlignment FCategoryModel::ArrangeVirtualTrackLaneView() const
@@ -55,6 +59,18 @@ FTrackLaneVirtualAlignment FCategoryModel::ArrangeVirtualTrackLaneView() const
 	return Parent ? Parent->ArrangeVirtualTrackLaneView() : FTrackLaneVirtualAlignment();
 }
 
+FLinearColor FCategoryModel::GetKeyBarColor() const
+{
+	TViewModelPtr<ITrackExtension> Track = FindAncestorOfType<ITrackExtension>();
+	UMovieSceneTrack* TrackObject = Track ? Track->GetTrack() : nullptr;
+	if (TrackObject)
+	{
+		FLinearColor Tint = FSequencerSectionPainter::BlendColor(TrackObject->GetColorTint());
+		return Tint.CopyWithNewOpacity(1.f);
+	}
+	return FColor(160, 160, 160);
+}
+
 FCategoryGroupModel::FCategoryGroupModel(FName InCategoryName, const FText& InDisplayText)
 	: CategoryName(InCategoryName)
 	, DisplayText(InDisplayText)
@@ -164,7 +180,7 @@ TSharedRef<SWidget> FCategoryGroupModel::CreateOutlinerView(const FCreateOutline
 FTrackAreaParameters FCategoryGroupModel::GetTrackAreaParameters() const
 {
 	FTrackAreaParameters Parameters;
-	Parameters.LaneType = ETrackAreaLaneType::None;
+	Parameters.LaneType = ETrackAreaLaneType::Inline;
 	return Parameters;
 }
 
diff --git a/Engine/Source/Editor/Sequencer/Private/MVVM/ViewModels/ChannelModel.cpp b/Engine/Source/Editor/Sequencer/Private/MVVM/ViewModels/ChannelModel.cpp
--- a/Engine/Source/Editor/Sequencer/Private/MVVM/ViewModels/ChannelModel.cpp
+++ b/Engine/Source/Editor/Sequencer/Private/MVVM/ViewModels/ChannelModel.cpp
@@ -8,11 +8,16 @@
 #include "MVVM/ViewModels/SequencerEditorViewModel.h"
 #include "MVVM/Views/SOutlinerItemViewBase.h"
 #include "MVVM/Views/SSequencerKeyNavigationButtons.h"
+#include "MVVM/Views/KeyRenderer.h"
+#include "MVVM/Views/SChannelView.h"
 #include "IKeyArea.h"
 #include "ISequencerSection.h"
 #include "ISequencerChannelInterface.h"
+#include "SequencerHotspots.h"
 #include "SequencerNodeTree.h"
 #include "SKeyAreaEditorSwitcher.h"
+#include "SequencerSectionPainter.h"
+#include "SSequencerSection.h"
 #include "Channels/MovieSceneChannel.h"
 #include "CurveModel.h"
 #include "Styling/AppStyle.h"
@@ -69,9 +74,154 @@ UMovieSceneSection* FChannelModel::GetSection() const
 
 FOutlinerSizing FChannelModel::GetDesiredSizing() const
 {
+	if (KeyArea->ShouldShowCurve())
+	{
+		TViewModelPtr<FSequenceModel> Sequence = FindAncestorOfType<FSequenceModel>();
+		if (Sequence)
+		{
+			return Sequence->GetSequencer()->GetSequencerSettings()->GetKeyAreaHeightWithCurves();
+		}
+	}
 	return FOutlinerSizing(15.f);
 }
 
+TSharedPtr<ITrackLaneWidget> FChannelModel::CreateTrackLaneView(const FCreateTrackLaneViewParams& InParams)
+{
+	return SNew(SChannelView, SharedThis(this), InParams.TimeToPixel, InParams.Editor->GetTrackArea())
+		.KeyBarColor(this, &FChannelModel::GetKeyBarColor);
+}
+
+FTrackLaneVirtualAlignment FChannelModel::ArrangeVirtualTrackLaneView() const
+{
+	TSharedPtr<FSectionModel> Section = FindAncestorOfType<FSectionModel>();
+	if (Section)
+	{
+		return FTrackLaneVirtualAlignment::Proportional(Section->GetRange(), 1.f);
+	}
+	return FTrackLaneVirtualAlignment::Proportional(TRange<FFrameNumber>(), 1.f);
+}
+
+bool FChannelModel::UpdateCachedKeys(TSharedPtr<FCachedKeys>& OutCachedKeys) const
+{
+	struct FSequencerCachedKeys : FCachedKeys
+	{
+		FSequencerCachedKeys(const FChannelModel* InChannel)
+		{
+			Update(InChannel);
+		}
+
+		bool Update(const FChannelModel* InChannel)
+		{
+			UMovieSceneSection* Section = InChannel->GetSection();
+			if (!Section || !CachedSignature.IsValid() || Section->GetSignature() != CachedSignature)
+			{
+				CachedSignature = Section ? Section->GetSignature() : FGuid();
+
+				KeyTimes.Reset();
+				KeyHandles.Reset();
+
+				TArray<FFrameNumber> KeyFrames;
+				InChannel->GetKeyArea()->GetKeyInfo(&KeyHandles, &KeyFrames);
+
+				KeyTimes.SetNumUninitialized(KeyFrames.Num());
+				for (int32 Index = 0; Index < KeyFrames.Num(); ++Index)
+				{
+					KeyTimes[Index] = KeyFrames[Index];
+				}
+				return true;
+			}
+
+			return false;
+		}
+
+		/** The guid with which the above array was generated */
+		FGuid CachedSignature;
+	};
+
+	if (OutCachedKeys)
+	{
+		return StaticCastSharedPtr<FSequencerCachedKeys>(OutCachedKeys)->Update(this);
+	}
+	else
+	{
+		OutCachedKeys = MakeShared<FSequencerCachedKeys>(this);
+		return true;
+	}
+
+	return false;
+}
+
+bool FChannelModel::GetFixedExtents(double& OutFixedMin, double& OutFixedMax) const
+{
+	TViewModelPtr<FSequenceModel> SequenceModel = FindAncestorOfType<FSequenceModel>();
+	if (!SequenceModel)
+	{
+		return false;
+	}
+
+	FString KeyAreaName = KeyArea->GetName().ToString();
+	if (SequenceModel->GetSequencer()->GetSequencerSettings()->HasKeyAreaCurveExtents(KeyAreaName))
+	{
+		SequenceModel->GetSequencer()->GetSequencerSettings()->GetKeyAreaCurveExtents(KeyAreaName, OutFixedMin, OutFixedMax);
+		return true;
+	}
+
+	return false;
+}
+
+int32 FChannelModel::CustomPaint(const FGeometry& KeyGeometry, int32 LayerId) const
+{
+	return LayerId;
+}
+
+void FChannelModel::DrawKeys(TArrayView<const FKeyHandle> InKeyHandles, TArrayView<FKeyDrawParams> OutKeyDrawParams)
+{
+	KeyArea->DrawKeys(InKeyHandles, OutKeyDrawParams);
+}
+
+TUniquePtr<FCurveModel> FChannelModel::CreateCurveModel()
+{
+	if (TViewModelPtr<FSequenceModel> SequenceModel = FindAncestorOfType<FSequenceModel>())
+	{
+		return KeyArea->CreateCurveEditorModel(SequenceModel->GetSequencer().ToSharedRef());
+	}
+
+	return nullptr;
+}
+
+FLinearColor FChannelModel::GetKeyBarColor() const
+{
+	if (TViewModelPtr<FSequenceModel> SequenceModel = FindAncestorOfType<FSequenceModel>())
+	{
+		if (SequenceModel->GetSequencer()->GetSequencerSettings()->GetShowChannelColors())
+		{
+			TOptional<FLinearColor> ChannelColor = KeyArea->GetColor();
+			if (ChannelColor)
+			{
+				return ChannelColor.GetValue();
+			}
+		}
+	}
+
+	TViewModelPtr<ITrackExtension> Track = FindAncestorOfType<ITrackExtension>();
+	UMovieSceneTrack* TrackObject = Track ? Track->GetTrack() : nullptr;
+	if (TrackObject)
+	{
+		FLinearColor Tint = FSequencerSectionPainter::BlendColor(TrackObject->GetColorTint()).LinearRGBToHSV();
+
+		// If this is a top level chanel, draw using the fill color
+		FViewModelPtr OutlinerItem = GetLinkedOutlinerItem();
+		if (OutlinerItem && !OutlinerItem->IsA<FChannelGroupModel>())
+		{
+			Tint.G *= .5f;
+			Tint.B = FMath::Max(.03f, Tint.B*.1f);
+		}
+
+		return Tint.HSVToLinearRGB().CopyWithNewOpacity(1.f);
+	}
+	return FColor(160, 160, 160);
+}
+
 FChannelGroupModel::FChannelGroupModel(FName InChannelName, const FText& InDisplayText)
 	: ChannelName(InChannelName)
 	, DisplayText(InDisplayText)
@@ -161,7 +311,7 @@ TArray<TSharedRef<IKeyArea>> FChannelGroupModel::GetAllKeyAreas() const
 FTrackAreaParameters FChannelGroupModel::GetTrackAreaParameters() const
 {
 	FTrackAreaParameters Parameters;
-	Parameters.LaneType = ETrackAreaLaneType::None;
+	Parameters.LaneType = ETrackAreaLaneType::Inline;
 	return Parameters;
 }
 
diff --git a/Engine/Source/Editor/Sequencer/Private/MVVM/ViewModels/LayerBarModel.cpp b/Engine/Source/Editor/Sequencer/Private/MVVM/ViewModels/LayerBarModel.cpp
--- a/Engine/Source/Editor/Sequencer/Private/MVVM/ViewModels/LayerBarModel.cpp
+++ b/Engine/Source/Editor/Sequencer/Private/MVVM/ViewModels/LayerBarModel.cpp
@@ -95,8 +95,8 @@ void FLayerBarModel::AddToSnapField(const ISnapCandidate& Candidate, ISnapField&
 	using namespace UE::MovieScene;
 
 	TRange<FFrameNumber> Range = ComputeRange();
-	SnapField.AddSnapPoint(FSnapPoint{ FSnapPoint::CustomSection, DiscreteInclusiveLower(Range) });
-	SnapField.AddSnapPoint(FSnapPoint{ FSnapPoint::CustomSection, DiscreteExclusiveUpper(Range) });
+	SnapField.AddSnapPoint(FSnapPoint{ FSnapPoint::CustomSection, DiscreteInclusiveLower(Range), 1.f });
+	SnapField.AddSnapPoint(FSnapPoint{ FSnapPoint::CustomSection, DiscreteExclusiveUpper(Range), 1.f });
 }
 
 bool FLayerBarModel::CanDrag() const
diff --git a/Engine/Source/Editor/Sequencer/Private/MVVM/ViewModels/SectionModel.cpp b/Engine/Source/Editor/Sequencer/Private/MVVM/ViewModels/SectionModel.cpp
--- a/Engine/Source/Editor/Sequencer/Private/MVVM/ViewModels/SectionModel.cpp
+++ b/Engine/Source/Editor/Sequencer/Private/MVVM/ViewModels/SectionModel.cpp
@@ -79,6 +79,18 @@ TViewModelPtr<ITrackExtension> FSectionModel::GetParentTrackExtension() const
 	return FindAncestorOfType<ITrackExtension>();
 }
 
+int32 FSectionModel::GetPreRollFrames() const
+{
+	UMovieSceneSection* Section = GetSection();
+	return Section ? Section->GetPreRollFrames() : 0;
+}
+
+int32 FSectionModel::GetPostRollFrames() const
+{
+	UMovieSceneSection* Section = GetSection();
+	return Section ? Section->GetPostRollFrames() : 0;
+}
+
 TSharedPtr<ITrackLaneWidget> FSectionModel::CreateTrackLaneView(const FCreateTrackLaneViewParams& InParams)
 {
 	UMovieSceneSection* Section = WeakSection.Get();
@@ -88,7 +100,7 @@ TSharedPtr<ITrackLaneWidget> FSectionModel::CreateTrackLaneView(const FCreateTra
 	}
 
 	TSharedPtr<FSequencer> Sequencer = InParams.Editor->CastThisChecked<FSequencerEditorViewModel>()->GetSequencerImpl();
-	return SNew(SSequencerSection, Sequencer, SharedThis(this));
+	return SNew(SSequencerSection, Sequencer, SharedThis(this), InParams.OwningTrackLane);
 }
 
 FTrackLaneVirtualAlignment FSectionModel::ArrangeVirtualTrackLaneView() const
@@ -135,6 +147,19 @@ void FSectionModel::UpdateCachedData()
 	// Compute the layer bar range from this section's effective key range
 	TRange<FFrameNumber> KeyRangeHull = TRange<FFrameNumber>::Empty();
 
+	if (SectionRange.GetLowerBound().IsClosed() && SectionRange.GetUpperBound().IsClosed())
+	{
+		KeyRangeHull = SectionRange;
+	}
+	else if (SectionRange.GetLowerBound().IsClosed())
+	{
+		KeyRangeHull = TRange<FFrameNumber>::Inclusive(SectionRange.GetLowerBoundValue(), SectionRange.GetLowerBoundValue());
+	}
+	else if (SectionRange.GetUpperBound().IsClosed())
+	{
+		KeyRangeHull = TRange<FFrameNumber>::Inclusive(SectionRange.GetUpperBoundValue(), SectionRange.GetUpperBoundValue());
+	}
+
 	// Find the first key time and use that
 	const FMovieSceneChannelProxy& Proxy = Section->GetChannelProxy();
 	for (const FMovieSceneChannelEntry& Entry : Proxy.GetAllEntries())
@@ -149,7 +174,7 @@ void FSectionModel::UpdateCachedData()
 		}
 	}
 
-	if (!KeyRangeHull.IsEmpty())
+	if (!KeyRangeHull.IsEmpty() && KeyRangeHull.GetLowerBound().IsClosed() && KeyRangeHull.GetUpperBound().IsClosed())
 	{
 		LayerBarRange = TRange<FFrameNumber>::Intersection(KeyRangeHull, SectionRange);
 	}
diff --git a/Engine/Source/Editor/Sequencer/Private/MVVM/ViewModels/TrackModel.cpp b/Engine/Source/Editor/Sequencer/Private/MVVM/ViewModels/TrackModel.cpp
--- a/Engine/Source/Editor/Sequencer/Private/MVVM/ViewModels/TrackModel.cpp
+++ b/Engine/Source/Editor/Sequencer/Private/MVVM/ViewModels/TrackModel.cpp
@@ -344,6 +344,7 @@ FTrackAreaParameters FTrackModel::GetTrackAreaParameters() const
 {
 	FTrackAreaParameters Params;
 	Params.LaneType = ETrackAreaLaneType::Nested;
+	Params.TrackLanePadding.Bottom = 1.f;
 	return Params;
 }
 
@@ -352,6 +353,11 @@ FViewModelVariantIterator FTrackModel::GetTrackAreaModelList() const
 	return &SectionList;
 }
 
+FViewModelVariantIterator FTrackModel::GetTopLevelChildTrackAreaModels() const
+{
+	return &TopLevelChannelList;
+}
+
 bool FTrackModel::CanRename() const
 {
 	UMovieSceneNameableTrack* NameableTrack = Cast<UMovieSceneNameableTrack>(GetTrack());
diff --git a/Engine/Source/Editor/Sequencer/Private/MVVM/ViewModels/TrackModelLayoutBuilder.cpp b/Engine/Source/Editor/Sequencer/Private/MVVM/ViewModels/TrackModelLayoutBuilder.cpp
--- a/Engine/Source/Editor/Sequencer/Private/MVVM/ViewModels/TrackModelLayoutBuilder.cpp
+++ b/Engine/Source/Editor/Sequencer/Private/MVVM/ViewModels/TrackModelLayoutBuilder.cpp
@@ -149,7 +149,7 @@ void FTrackModelLayoutBuilder::RefreshLayout(TSharedPtr<FSectionModel> InSection
 	TrackAreaList = FHierarchicalModelListRefresher();
 }
 
-void FTrackModelLayoutBuilder::PushCategory(FName CategoryName, const FText& DisplayLabel)
+void FTrackModelLayoutBuilder::PushCategory(FName CategoryName, const FText& DisplayLabel, TFunction<TSharedPtr<FCategoryModel>(FName, const FText&)> OptionalFactory)
 {
 	check(TrackAreaList.IsValid() && SequencerSection);
 
@@ -175,7 +175,15 @@ void FTrackModelLayoutBuilder::PushCategory(FName CategoryName, const FText& Dis
 	TSharedPtr<FCategoryModel> TrackAreaModel = TrackAreaList.FindItem<FCategoryModel>(CategoryNamePredicate);
 	if (!TrackAreaModel)
 	{
-		TrackAreaModel = MakeShared<FCategoryModel>(CategoryName);
+		if (OptionalFactory)
+		{
+			TrackAreaModel = OptionalFactory(CategoryName, DisplayLabel);
+		}
+
+		if (!TrackAreaModel)
+		{
+			TrackAreaModel = MakeShared<FCategoryModel>(CategoryName);
+		}
 	}
 
 	OutlinerModel->AddCategory(TrackAreaModel);
@@ -193,7 +201,7 @@ void FTrackModelLayoutBuilder::PopCategory()
 	OutlinerList.Pop();
 }
 
-void FTrackModelLayoutBuilder::SetTopLevelChannel(const FMovieSceneChannelHandle& Channel)
+void FTrackModelLayoutBuilder::SetTopLevelChannel(const FMovieSceneChannelHandle& Channel, TFunction<TSharedPtr<UE::Sequencer::FChannelModel>(FName, const FMovieSceneChannelHandle&)> OptionalFactory)
 {
 	ensureAlwaysMsgf(
 			OutlinerList.GetCurrentParent() == Root,
@@ -204,10 +212,10 @@ void FTrackModelLayoutBuilder::SetTopLevelChannel(const FMovieSceneChannelHandle
 	TSharedPtr<FViewModel> CurrentModel = OutlinerList.GetCurrentParent();
 	OutlinerList.RecurseInto(CurrentModel, TopLevelGroup);
 
-	AddChannel(Channel, true);
+	AddChannel(Channel, true, OptionalFactory);
 }
 
-void FTrackModelLayoutBuilder::AddChannel(const FMovieSceneChannelHandle& Channel)
+void FTrackModelLayoutBuilder::AddChannel(const FMovieSceneChannelHandle& Channel, TFunction<TSharedPtr<UE::Sequencer::FChannelModel>(FName, const FMovieSceneChannelHandle&)> OptionalFactory)
 {
 	check(TrackAreaList.IsValid() && SequencerSection);
 
@@ -218,10 +226,10 @@ void FTrackModelLayoutBuilder::AddChannel(const FMovieSceneChannelHandle& Channe
 		TrackAreaList.RecurseInto(CurrentModel, CurrentModel->GetChildList(EViewModelListType::Generic));
 	}
 
-	AddChannel(Channel, false);
+	AddChannel(Channel, false, OptionalFactory);
 }
 
-void FTrackModelLayoutBuilder::AddChannel(const FMovieSceneChannelHandle& Channel, bool bIsTopLevel)
+void FTrackModelLayoutBuilder::AddChannel(const FMovieSceneChannelHandle& Channel, bool bIsTopLevel, TFunction<TSharedPtr<UE::Sequencer::FChannelModel>(FName, const FMovieSceneChannelHandle&)> OptionalFactory)
 {
 	check(TrackAreaList.IsValid() && SequencerSection);
 
@@ -243,7 +251,7 @@ void FTrackModelLayoutBuilder::AddChannel(const FMovieSceneChannelHandle& Channe
 	//       of channels. We have to be careful to ensure that we reuse existing items that may have been added
 	//       to the outliner from a previously built section.
 
-	TSharedPtr<FChannelGroupModel> OutlinerModel;
+	TViewModelPtr<FChannelGroupModel> OutlinerModel;
 	if (bIsTopLevel)
 	{
 		OutlinerModel = OutlinerList.FindItem<FChannelGroupModel>(ChannelNamePredicate);
@@ -280,7 +288,15 @@ void FTrackModelLayoutBuilder::AddChannel(const FMovieSceneChannelHandle& Channe
 		TrackAreaModel = TrackAreaList.FindRecycledItem<FChannelModel>(ChannelNamePredicate);
 		if (!TrackAreaModel)
 		{
-			TrackAreaModel = MakeShared<FChannelModel>(ChannelName, SequencerSection, Channel);
+			if (OptionalFactory)
+			{
+				TrackAreaModel = OptionalFactory(ChannelName, Channel);
+			}
+
+			if (!TrackAreaModel)
+			{
+				TrackAreaModel = MakeShared<FChannelModel>(ChannelName, SequencerSection, Channel);
+			}
 		}
 
 		OutlinerModel->AddChannel(TrackAreaModel);
@@ -290,7 +306,7 @@ void FTrackModelLayoutBuilder::AddChannel(const FMovieSceneChannelHandle& Channe
 		}
 		else
 		{
-			TrackAreaModel->SetLinkedOutlinerItem(OutlinerModel);
+			TrackAreaModel->SetLinkedOutlinerItem(OutlinerModel.ImplicitCast());
 		}
 		TrackAreaModel->Initialize(SequencerSection, Channel);
 
diff --git a/Engine/Source/Editor/Sequencer/Private/MVVM/ViewModels/TrackModelLayoutBuilder.h b/Engine/Source/Editor/Sequencer/Private/MVVM/ViewModels/TrackModelLayoutBuilder.h
--- a/Engine/Source/Editor/Sequencer/Private/MVVM/ViewModels/TrackModelLayoutBuilder.h
+++ b/Engine/Source/Editor/Sequencer/Private/MVVM/ViewModels/TrackModelLayoutBuilder.h
@@ -103,14 +103,14 @@ public:
 
 	// ISectionLayoutBuilder interface
 
-	virtual void PushCategory( FName CategoryName, const FText& DisplayLabel ) override;
-	virtual void SetTopLevelChannel( const FMovieSceneChannelHandle& Channel ) override;
-	virtual void AddChannel( const FMovieSceneChannelHandle& Channel ) override;
+	virtual void PushCategory( FName CategoryName, const FText& DisplayLabel, TFunction<TSharedPtr<UE::Sequencer::FCategoryModel>(FName, const FText&)> OptionalFactory ) override;
+	virtual void SetTopLevelChannel( const FMovieSceneChannelHandle& Channel, TFunction<TSharedPtr<UE::Sequencer::FChannelModel>(FName, const FMovieSceneChannelHandle&)> OptionalFactory ) override;
+	virtual void AddChannel( const FMovieSceneChannelHandle& Channel, TFunction<TSharedPtr<UE::Sequencer::FChannelModel>(FName, const FMovieSceneChannelHandle&)> OptionalFactory ) override;
 	virtual void PopCategory() override;
 
 private:
 
-	void AddChannel( const FMovieSceneChannelHandle& Channel, bool bIsTopLevel );
+	void AddChannel( const FMovieSceneChannelHandle& Channel, bool bIsTopLevel, TFunction<TSharedPtr<UE::Sequencer::FChannelModel>(FName, const FMovieSceneChannelHandle&)> OptionalFactory );
 
 private:
 
diff --git a/Engine/Source/Editor/Sequencer/Private/MVVM/ViewModels/TrackRowModel.cpp b/Engine/Source/Editor/Sequencer/Private/MVVM/ViewModels/TrackRowModel.cpp
--- a/Engine/Source/Editor/Sequencer/Private/MVVM/ViewModels/TrackRowModel.cpp
+++ b/Engine/Source/Editor/Sequencer/Private/MVVM/ViewModels/TrackRowModel.cpp
@@ -99,6 +99,7 @@ FTrackAreaParameters FTrackRowModel::GetTrackAreaParameters() const
 {
 	FTrackAreaParameters Params;
 	Params.LaneType = ETrackAreaLaneType::Nested;
+	Params.TrackLanePadding.Bottom = 1.f;
 	return Params;
 }
 
@@ -107,6 +108,11 @@ FViewModelVariantIterator FTrackRowModel::GetTrackAreaModelList() const
 	return &SectionList;
 }
 
+FViewModelVariantIterator FTrackRowModel::GetTopLevelChildTrackAreaModels() const
+{
+	return &TopLevelChannelList;
+}
+
 void FTrackRowModel::CreateCurveModels(TArray<TUniquePtr<FCurveModel>>& OutCurveModels)
 {
 	TViewModelPtr<FChannelGroupModel> ChannelGroup = TopLevelChannelList.GetHead().ImplicitCast();
diff --git a/Engine/Source/Editor/Sequencer/Private/MVVM/Views/SChannelView.cpp b/Engine/Source/Editor/Sequencer/Private/MVVM/Views/SChannelView.cpp
new file mode 100644
--- /dev/null
+++ b/Engine/Source/Editor/Sequencer/Private/MVVM/Views/SChannelView.cpp
@@ -0,0 +1,723 @@
+// Copyright Epic Games, Inc. All Rights Reserved.
+
+#include "MVVM/Views/SChannelView.h"
+#include "MVVM/ViewModels/ChannelModel.h"
+#include "MVVM/ViewModels/SectionModel.h"
+#include "MVVM/ViewModels/SequenceModel.h"
+#include "MVVM/ViewModels/TrackAreaViewModel.h"
+
+#include "IKeyArea.h"
+#include "Sequencer.h"
+#include "SequencerHotspots.h"
+#include "ISequencerEditTool.h"
+#include "ISequencerSection.h"
+#include "SequencerAddKeyOperation.h"
+#include "SSequencerSection.h"
+
+#include "ScopedTransaction.h"
+
+#include "Algo/Transform.h"
+
+#define LOCTEXT_NAMESPACE "SequencerChannelView"
+
+namespace UE::Sequencer
+{
+
+/** Convert a SequencerCore FKeyRenderer hittest query to FSequencerSelectedKeys */
+TArray<FSequencerSelectedKey> KeyRendererResultToSelectedKeys(TArrayView<const FKeyRenderer::FKeysForModel> InKeys, UMovieSceneSection* SectionObject)
+{
+	TArray<FSequencerSelectedKey> SelectedKeys;
+	for (const FKeyRenderer::FKeysForModel& Result : InKeys)
+	{
+		TViewModelPtr<FChannelModel> Channel = Result.Model.ImplicitCast();
+		if (Channel && Result.Keys.Num() > 0)
+		{
+			SelectedKeys.Reserve(SelectedKeys.Num() + Result.Keys.Num());
+			for (FKeyHandle Key : Result.Keys)
+			{
+				SelectedKeys.Add(FSequencerSelectedKey(*SectionObject, Channel, Key));
+			}
+		}
+	}
+	return SelectedKeys;
+}
+
+
+/** Key renderer interface for defining selection and hover states. Should be removed in favor of SequencerCore knowing about selection in the future */
+struct FSequencerKeyRendererInterface : IKeyRendererInterface
+{
+	UMovieSceneSection* SectionObject;
+	const TSet<FSequencerSelectedKey>* SelectedKeys;
+	const TMap<FSequencerSelectedKey, ESelectionPreviewState>* SelectionPreview;
+	TArrayView<const FSequencerSelectedKey> HoveredKeys;
+
+	FSequencerKeyRendererInterface(UMovieSceneSection* InSectionObject, FSequencer* Sequencer, TSharedPtr<ITrackAreaHotspot> Hotspot)
+	{
+		SectionObject    = InSectionObject;
+		SelectedKeys     = &Sequencer->GetSelection().GetSelectedKeys();
+		SelectionPreview = &Sequencer->GetSelectionPreview().GetDefinedKeyStates();
+		if (TSharedPtr<FKeyHotspot> KeyHotspot = HotspotCast<FKeyHotspot>(Hotspot))
+		{
+			HoveredKeys = KeyHotspot->Keys;
+		}
+	}
+
+	bool HasAnySelectedKeys() const
+	{
+		return SelectedKeys->Num() != 0;
+	}
+	bool HasAnyPreviewSelectedKeys() const
+	{
+		return SelectionPreview->Num() != 0;
+	}
+	bool HasAnyHoveredKeys() const
+	{
+		return HoveredKeys.Num() != 0;
+	}
+
+	bool IsKeySelected(const TViewModelPtr<IKeyExtension>& InOwner, FKeyHandle InKey) const
+	{
+		TSharedPtr<FChannelModel> Channel = InOwner.ImplicitCast();
+		FSequencerSelectedKey SelectedKey(*SectionObject, Channel, InKey);
+		return SelectedKeys->Contains(SelectedKey);
+	}
+	bool IsKeyHovered(const TViewModelPtr<IKeyExtension>& InOwner, FKeyHandle InKey) const
+	{
+		TSharedPtr<FChannelModel> Channel = InOwner.ImplicitCast();
+		FSequencerSelectedKey SelectedKey(*SectionObject, Channel, InKey);
+		return HoveredKeys.Contains(SelectedKey);
+	}
+	EKeySelectionPreviewState GetPreviewSelectionState(const TViewModelPtr<IKeyExtension>& InOwner, FKeyHandle InKey) const
+	{
+		TSharedPtr<FChannelModel> Channel = InOwner.ImplicitCast();
+		FSequencerSelectedKey SelectedKey(*SectionObject, Channel, InKey);
+
+		if (const ESelectionPreviewState* State = SelectionPreview->Find(SelectedKey))
+		{
+			switch(*State)
+			{
+			case ESelectionPreviewState::Undefined:   return EKeySelectionPreviewState::Undefined;
+			case ESelectionPreviewState::Selected:    return EKeySelectionPreviewState::Selected;
+			case ESelectionPreviewState::NotSelected: return EKeySelectionPreviewState::NotSelected;
+			}
+		}
+
+		return EKeySelectionPreviewState::Undefined;
+	}
+};
+
+
+FChannelViewKeyCachedState::FChannelViewKeyCachedState()
+	: ValidPlayRangeMin(TNumericLimits<int32>::Lowest())
+	, ValidPlayRangeMax(TNumericLimits<int32>::Max())
+	, VisibleRange(0, 1)
+	, KeySizePx(12, 12)
+	, SelectionSerial(0)
+	, SelectionPreviewHash(0)
+	, bShowCurve(false)
+	, bShowKeyBars(true)
+	, bCollapseChildren(false)
+{}
+
+FChannelViewKeyCachedState::FChannelViewKeyCachedState(TRange<FFrameTime> InVisibleRange, TSharedPtr<ITrackAreaHotspot> Hotspot, FViewModelPtr Model, FSequencer* Sequencer)
+{
+	// Gather keys for a region larger than the view range to ensure we draw keys that are only just offscreen.
+	// Compute visible range taking into account a half-frame offset for keys, plus half a key width for keys that are partially offscreen
+	TRange<FFrameNumber> ValidKeyRange = Sequencer->GetSubSequenceRange().Get(Sequencer->GetFocusedMovieSceneSequence()->GetMovieScene()->GetPlaybackRange());
+
+	TViewModelPtr<FChannelModel> Channel = Model.ImplicitCast();
+	TViewModelPtr<FLinkedOutlinerExtension> Outliner = Model.ImplicitCast();
+
+	ValidPlayRangeMin = UE::MovieScene::DiscreteInclusiveLower(ValidKeyRange);
+	ValidPlayRangeMax = UE::MovieScene::DiscreteExclusiveUpper(ValidKeyRange);
+	VisibleRange = InVisibleRange;
+	KeySizePx = FVector2D(12, 12);
+	SelectionSerial = Sequencer->GetSelection().GetSerialNumber();
+	SelectionPreviewHash = Sequencer->GetSelectionPreview().GetSelectionHash();
+	bShowCurve = Channel && Channel->GetKeyArea()->ShouldShowCurve();
+	bShowKeyBars = Sequencer->GetSequencerSettings()->GetShowKeyBars();
+	bCollapseChildren = Outliner && Outliner->GetLinkedOutlinerItem() && Outliner->GetLinkedOutlinerItem()->IsExpanded() == false;
+	if (!Hotspot || Hotspot->CastThis<FKeyHotspot>() || Hotspot->CastThis<FKeyBarHotspot>())
+	{
+		WeakHotspot = Hotspot;
+		RawHotspotDoNotUse = Hotspot.Get();
+	}
+}
+
+EViewDependentCacheFlags FChannelViewKeyCachedState::CompareTo(const FChannelViewKeyCachedState& Other) const
+{
+	EViewDependentCacheFlags Flags = EViewDependentCacheFlags::None;
+
+	// Check if either hotspot expired
+	if ( (WeakHotspot.Pin() == nullptr && RawHotspotDoNotUse != nullptr) ||
+		 (Other.WeakHotspot.Pin() == nullptr && Other.RawHotspotDoNotUse != nullptr) )
+	{
+		Flags |= EViewDependentCacheFlags::KeyStateChanged;
+	}
+
+	// Check if the hotspot has changed (this is necessary because theoretically both hotspots
+	// could still be _alive_ but Other.WeakHotspot is the active one)
+	if (WeakHotspot != Other.WeakHotspot)
+	{
+		// Hotspot has changed - maybe we (un)hovered some keys or a key bar
+		Flags |= EViewDependentCacheFlags::KeyStateChanged;
+	}
+
+	if (ValidPlayRangeMin != Other.ValidPlayRangeMin || ValidPlayRangeMax != Other.ValidPlayRangeMax)
+	{
+		// The valid key ranges for the data has changed
+		Flags |= EViewDependentCacheFlags::KeyStateChanged;
+	}
+
+	if (SelectionSerial != Other.SelectionSerial || SelectionPreviewHash != Other.SelectionPreviewHash)
+	{
+		// Selection states have changed
+		Flags |= EViewDependentCacheFlags::KeyStateChanged;
+	}
+
+	if (VisibleRange != Other.VisibleRange)
+	{
+		Flags |= EViewDependentCacheFlags::ViewChanged;
+
+		const double RangeSize = VisibleRange.Size<FFrameTime>().AsDecimal();
+		const double OtherRangeSize = Other.VisibleRange.Size<FFrameTime>().AsDecimal();
+
+		if (!FMath::IsNearlyEqual(RangeSize, OtherRangeSize, RangeSize * 0.001))
+		{
+			Flags |= EViewDependentCacheFlags::ViewZoomed;
+		}
+	}
+
+	if (bShowCurve != Other.bShowCurve)
+	{
+		Flags |= EViewDependentCacheFlags::DataChanged;
+	}
+
+	if (bShowKeyBars != Other.bShowKeyBars)
+	{
+		Flags |= EViewDependentCacheFlags::DataChanged;
+	}
+
+	if (bCollapseChildren != Other.bCollapseChildren)
+	{
+		Flags |= EViewDependentCacheFlags::DataChanged;
+	}
+
+	if (KeySizePx != Other.KeySizePx)
+	{
+		Flags |= EViewDependentCacheFlags::ViewChanged | EViewDependentCacheFlags::ViewZoomed;
+	}
+
+	return Flags;
+}
+
+void SChannelView::Construct(const FArguments& InArgs, const FViewModelPtr& InViewModel, TSharedPtr<FTimeToPixel> InTimeToPixel, TSharedPtr<FTrackAreaViewModel> TrackArea)
+{
+	WeakModel = InViewModel;
+	WeakTrackArea = TrackArea;
+	TrackAreaTimeToPixel = InTimeToPixel;
+
+	KeyBarColor = InArgs._KeyBarColor;
+
+	KeyRenderer.Initialize(InViewModel);
+}
+
+FReply SChannelView::OnMouseButtonDown(const FGeometry& MyGeometry, const FPointerEvent& MouseEvent)
+{
+	TViewModelPtr<FTrackAreaViewModel> TrackArea = WeakTrackArea.Pin();
+	if (!TrackArea)
+	{
+		return FReply::Unhandled();
+	}
+
+	TArrayView<const FSequencerSelectedKey> HoveredKeys;
+
+	// The hovered key is defined from the sequencer hotspot
+	if (TSharedPtr<FKeyHotspot> KeyHotspot = HotspotCast<FKeyHotspot>(TrackArea->GetHotspot()))
+	{
+		HoveredKeys = KeyHotspot->Keys;
+	}
+
+	if (MouseEvent.GetEffectingButton() == EKeys::MiddleMouseButton)
+	{
+		GEditor->BeginTransaction(LOCTEXT("CreateKey_Transaction", "Create Key"));
+
+		// Generate a key and set it as the PressedKey
+		TArray<FSequencerSelectedKey> NewKeys;
+		CreateKeysUnderMouse(MouseEvent.GetScreenSpacePosition(), MyGeometry, HoveredKeys, NewKeys);
+
+		if (NewKeys.Num())
+		{
+			TSharedPtr<FSequencer> Sequencer = LegacyGetSequencer();
+			Sequencer->GetSelection().EmptySelectedKeys();
+			for (const FSequencerSelectedKey& NewKey : NewKeys)
+			{
+				Sequencer->GetSelection().AddToSelection(NewKey);
+			}
+
+			// Pass the event to the tool to copy the hovered key and move it
+			TrackArea->SetHotspot(MakeShared<FKeyHotspot>(MoveTemp(NewKeys), Sequencer));
+
+			// Return unhandled so that the EditTool can handle the mouse down based on the newly created keyframe and prepare to move it
+			return FReply::Unhandled();
+		}
+	}
+
+	return FReply::Unhandled();
+}
+
+FReply SChannelView::OnMouseMove(const FGeometry& MyGeometry, const FPointerEvent& MouseEvent)
+{
+	TViewModelPtr<FTrackAreaViewModel> TrackArea = WeakTrackArea.Pin();
+	if (!TrackArea)
+	{
+		return FReply::Unhandled();
+	}
+
+	TSharedPtr<FSequencer> Sequencer = LegacyGetSequencer();
+
+	// Don't fiddle with hotspots if there is an active drag
+	const ISequencerEditTool* EditTool = TrackArea->GetEditTool();
+	if (EditTool && EditTool->GetDragHotspot())
+	{
+		return FReply::Unhandled();
+	}
+
+	// Checked for hovered key
+	TArray<FSequencerSelectedKey> KeysUnderMouse;
+	GetKeysUnderMouse(MouseEvent.GetScreenSpacePosition(), MyGeometry, KeysUnderMouse);
+	if (KeysUnderMouse.Num())
+	{
+		TSharedPtr<FKeyHotspot> KeyHotspot = HotspotCast<FKeyHotspot>(TrackArea->GetHotspot());
+		if (!KeyHotspot || KeyHotspot->Keys != KeysUnderMouse)
+		{
+			TrackArea->SetHotspot(MakeShared<FKeyHotspot>(MoveTemp(KeysUnderMouse), Sequencer));
+		}
+
+		return FReply::Handled();
+	}
+
+	{
+		TViewModelPtr<FSectionModel> Section = GetSection();
+		if (Section)
+		{
+			const float HalfHeightPx = MyGeometry.GetLocalSize().Y*.5f;
+			const FVector2D MousePixel = MyGeometry.AbsoluteToLocal(MouseEvent.GetScreenSpacePosition());
+
+			// Keys always render in the middle of the geometry so do a quick vertical test first
+			// We actually hittest a larger vertical area than the bars render to make it easier for user
+			// to hit key bars
+			if (FMath::Abs(MousePixel.Y - HalfHeightPx) <= 4.f)
+			{
+				TRange<FFrameNumber> SectionRange = Section->GetRange();
+
+				// Make the time <-> pixel converter relative to the section range rather than the view range
+				FTimeToPixel RelativeTimeToPixel = *TrackAreaTimeToPixel;
+				if (SectionRange.GetLowerBound().IsClosed())
+				{
+					RelativeTimeToPixel = TrackAreaTimeToPixel->RelativeTo(SectionRange.GetLowerBoundValue());
+				}
+
+				TArray<FKeyRenderer::FKeysForModel> LeadingKeys, TrailingKeys;
+				if (KeyRenderer.HitTestKeyBar(RelativeTimeToPixel.PixelToFrame(MousePixel.X), LeadingKeys, TrailingKeys))
+				{
+					UMovieSceneSection* SectionObject = Section->GetSection();
+					TrackArea->SetHotspot(MakeShared<FKeyBarHotspot>(
+						KeyRendererResultToSelectedKeys(LeadingKeys, SectionObject),
+						KeyRendererResultToSelectedKeys(TrailingKeys, SectionObject),
+						Sequencer)
+					);
+					return FReply::Handled();
+				}
+			}
+		}
+	}
+
+	return FReply::Unhandled();
+}
+
+FReply SChannelView::OnMouseButtonUp(const FGeometry& MyGeometry, const FPointerEvent& MouseEvent)
+{
+	if (MouseEvent.GetEffectingButton() == EKeys::MiddleMouseButton)
+	{
+		TSharedPtr<FSequencer> Sequencer = LegacyGetSequencer();
+
+		// Snap keys on mouse up since we want to create keys at the exact mouse position (ie. to keep the newly created keys under the mouse 
+		// while dragging) but obey snapping rules if necessary
+		if (Sequencer->GetSequencerSettings()->GetIsSnapEnabled() && Sequencer->GetSequencerSettings()->GetSnapKeyTimesToInterval())
+		{
+			Sequencer->SnapToFrame();
+
+			for (const FSequencerSelectedKey& SelectedKey : Sequencer->GetSelection().GetSelectedKeys())
+			{
+				const FFrameNumber CurrentTime = SelectedKey.WeakChannel.Pin()->GetKeyArea()->GetKeyTime(SelectedKey.KeyHandle);
+				Sequencer->SetLocalTime(CurrentTime, ESnapTimeMode::STM_Interval);
+				break;
+			}
+		}
+		GEditor->EndTransaction();
+
+		// Return unhandled so that the EditTool can handle the mouse up based on the newly created keyframe and finish moving it
+		return FReply::Unhandled();
+	}
+	return FReply::Unhandled();
+}
+
+void SChannelView::OnMouseLeave(const FPointerEvent& MouseEvent)
+{
+	SCompoundWidget::OnMouseLeave(MouseEvent);
+
+	TViewModelPtr<FTrackAreaViewModel> TrackArea = WeakTrackArea.Pin();
+	if (TrackArea)
+	{
+		TrackArea->SetHotspot(nullptr);
+	}
+}
+
+FReply SChannelView::OnMouseButtonDoubleClick(const FGeometry& MyGeometry, const FPointerEvent& MouseEvent)
+{
+	TSharedPtr<FSectionModel> Section = GetSection();
+	if (!Section)
+	{
+		return FReply::Handled();
+	}
+
+	if(MouseEvent.GetEffectingButton() == EKeys::LeftMouseButton)
+	{
+		TArray<FSequencerSelectedKey> Keys;
+		GetKeysUnderMouse(MouseEvent.GetScreenSpacePosition(), MyGeometry, Keys);
+		TArray<FKeyHandle> KeyHandles;
+		for (FSequencerSelectedKey Key : Keys)
+		{
+			KeyHandles.Add(Key.KeyHandle);
+		}
+		if (KeyHandles.Num() > 0)
+		{
+			FReply Reply = Section->GetSectionInterface()->OnKeyDoubleClicked(KeyHandles);
+			return Reply.Handled();
+		}
+	}
+
+	return FReply::Unhandled();
+}
+
+int32 SChannelView::OnPaint(const FPaintArgs& Args, const FGeometry& AllottedGeometry, const FSlateRect& MyCullingRect, FSlateWindowElementList& OutDrawElements, int32 LayerId, const FWidgetStyle& InWidgetStyle, bool bParentEnabled) const
+{
+	static const FName SelectionColorName("SelectionColor");
+
+	LayerId = DrawLane(Args, AllottedGeometry, MyCullingRect, OutDrawElements, LayerId, InWidgetStyle, bParentEnabled);
+
+	TSharedPtr<FSequencer> Sequencer = LegacyGetSequencer();
+	FViewModelPtr Model = WeakModel.Pin();
+	if (!Model)
+	{
+		return LayerId;
+	}
+
+	TViewModelPtr<FTrackAreaViewModel> TrackArea = WeakTrackArea.Pin();
+	TViewModelPtr<FSectionModel>  Section  = Model->FindAncestorOfType<FSectionModel>();
+	TViewModelPtr<FSequenceModel> Sequence = Model->FindAncestorOfType<FSequenceModel>();
+
+	UMovieSceneSection* SectionObject = Section->GetSection();
+	if (!Section || !Sequence || !SectionObject || !TrackArea)
+	{
+		return LayerId;
+	}
+
+	TRange<FFrameNumber> SectionRange = Section->GetRange();
+
+	// Make the time <-> pixel converter relative to the section range rather than the view range
+	FTimeToPixel RelativeTimeToPixel = *TrackAreaTimeToPixel;
+	if (SectionRange.GetLowerBound().IsClosed())
+	{
+		RelativeTimeToPixel = TrackAreaTimeToPixel->RelativeTo(SectionRange.GetLowerBoundValue());
+	}
+
+	FVector2D TopLeft = AllottedGeometry.AbsoluteToLocal(MyCullingRect.GetTopLeft());
+	FVector2D BottomRight = AllottedGeometry.AbsoluteToLocal(MyCullingRect.GetBottomRight());
+
+	TRange<FFrameTime> VisibleRange = TRange<FFrameTime>(
+		RelativeTimeToPixel.PixelToFrame(TopLeft.X),
+		RelativeTimeToPixel.PixelToFrame(BottomRight.X)
+	);
+
+	FSequencerKeyRendererInterface KeyRenderInterface(SectionObject, Sequencer.Get(), TrackArea->GetHotspot());
+	FChannelViewKeyCachedState NewCachedState(VisibleRange, TrackArea->GetHotspot(), Model, Sequencer.Get());
+
+	FKeyBatchParameters Params(RelativeTimeToPixel);
+	Params.CacheState = KeyRendererCache.IsSet() ? NewCachedState.CompareTo(KeyRendererCache.GetValue()) : EViewDependentCacheFlags::All;
+	Params.VisibleRange = VisibleRange;
+	Params.ValidPlayRangeMin = NewCachedState.ValidPlayRangeMin;
+	Params.ValidPlayRangeMax = NewCachedState.ValidPlayRangeMax;
+	Params.bShowCurve = NewCachedState.bShowCurve;
+	Params.bShowKeyBars = NewCachedState.bShowKeyBars;
+	Params.bCollapseChildren = NewCachedState.bCollapseChildren;
+	Params.ClientInterface = &KeyRenderInterface;
+
+	FKeyRendererPaintArgs PaintArgs;
+	PaintArgs.DrawElements = &OutDrawElements;
+	PaintArgs.KeyThrobValue = SSequencerSection::GetKeySelectionThrobValue();
+	PaintArgs.DrawEffects = bParentEnabled ? ESlateDrawEffect::None : ESlateDrawEffect::DisabledEffect;
+	PaintArgs.SelectionColor = FAppStyle::GetSlateColor(SelectionColorName).GetColor(InWidgetStyle);
+	if (KeyBarColor.IsSet())
+	{
+		PaintArgs.KeyBarColor = KeyBarColor.Get();
+	}
+
+	TViewModelPtr<FChannelModel> Channel = Model.ImplicitCast();
+	if (Sequencer->GetSequencerSettings()->GetShowChannelColors() && Channel)
+	{
+		if (TOptional<FLinearColor> ChannelColor = Channel->GetKeyArea()->GetColor())
+		{
+			PaintArgs.CurveColor = ChannelColor.GetValue();
+		}
+	}
+
+
+	KeyRenderer.Update(Params, AllottedGeometry);
+	LayerId = KeyRenderer.DrawCurve(Params, AllottedGeometry, MyCullingRect, PaintArgs, LayerId);
+	LayerId = KeyRenderer.Draw(Params, AllottedGeometry, MyCullingRect, PaintArgs, LayerId);
+
+	KeyRendererCache = NewCachedState;
+
+	return LayerId;
+}
+
+int32 SChannelView::DrawLane(const FPaintArgs& Args, const FGeometry& AllottedGeometry, const FSlateRect& MyCullingRect, FSlateWindowElementList& OutDrawElements, int32 LayerId, const FWidgetStyle& InWidgetStyle, bool bParentEnabled) const
+{
+	FViewModelPtr Model = WeakModel.Pin();
+	TViewModelPtr<FSequenceModel> Sequence = Model ? Model->FindAncestorOfType<FSequenceModel>() : nullptr;
+	TViewModelPtr<FLinkedOutlinerExtension> Outliner = WeakModel.ImplicitPin();
+	if (!Sequence || !Outliner)
+	{
+		return LayerId;
+	}
+
+	TViewModelPtr<IOutlinerExtension> OutlinerItem = Outliner->GetLinkedOutlinerItem();
+	if (!OutlinerItem)
+	{
+		return LayerId;
+	}
+
+	FSequencer* Sequencer = Sequence->GetSequencerImpl().Get();
+	FSequencerSelection& Selection = Sequencer->GetSelection();
+
+	const ESlateDrawEffect DrawEffects = bParentEnabled
+		? ESlateDrawEffect::None
+		: ESlateDrawEffect::DisabledEffect;
+
+	FLinearColor HighlightColor;
+	bool bDrawHighlight = false;
+	if (Sequencer->GetSelection().NodeHasSelectedKeysOrSections(OutlinerItem.AsModel()))
+	{
+		bDrawHighlight = true;
+		HighlightColor = FLinearColor(1.0f, 1.0f, 1.0f, 0.15f);
+	}
+	else if (TViewModelPtr<IHoveredExtension> Hovered = OutlinerItem.ImplicitCast())
+	{
+		if (Hovered->IsHovered())
+		{
+			bDrawHighlight = true;
+			HighlightColor = FLinearColor(1.0f, 1.0f, 1.0f, 0.05f);
+		}
+	}
+
+	// --------------------------------------------
+	// Draw hover or selection highlight
+	if (bDrawHighlight)
+	{
+		static const FName HighlightBrushName("Sequencer.AnimationOutliner.DefaultBorder");
+
+		FSlateDrawElement::MakeBox(
+			OutDrawElements,
+			LayerId++,
+			AllottedGeometry.ToPaintGeometry(),
+			FAppStyle::GetBrush(HighlightBrushName),
+			DrawEffects,
+			HighlightColor
+		);
+	}
+
+	// --------------------------------------------
+	// Draw outliner selection tint
+	if (Selection.IsSelected(OutlinerItem.AsModel()))
+	{
+		static const FName SelectionColorName("SelectionColor");
+		static const FName SelectedTrackTintBrushName("Sequencer.Section.SelectedTrackTint");
+
+		FSlateDrawElement::MakeBox(
+			OutDrawElements,
+			LayerId++,
+			AllottedGeometry.ToPaintGeometry(),
+			FAppStyle::GetBrush(SelectedTrackTintBrushName),
+			DrawEffects,
+			FAppStyle::GetSlateColor(SelectionColorName).GetColor(InWidgetStyle)
+		);
+	}
+
+	return LayerId;
+}
+
+TSharedRef<const SWidget> SChannelView::AsWidget() const
+{
+	return AsShared();
+}
+
+FTrackLaneScreenAlignment SChannelView::GetAlignment(const FTimeToPixel& InTimeToPixel, const FGeometry& InParentGeometry) const
+{
+	TSharedPtr<ITrackLaneExtension> TrackLaneExtension = WeakModel.ImplicitPin();
+	if (TrackLaneExtension)
+	{
+		FTrackLaneVirtualAlignment VirtualAlignment = TrackLaneExtension->ArrangeVirtualTrackLaneView();
+		FTrackLaneScreenAlignment  ScreenAlignment  = VirtualAlignment.ToScreen(InTimeToPixel, InParentGeometry);
+
+		return ScreenAlignment;
+	}
+	return FTrackLaneScreenAlignment();
+}
+
+TViewModelPtr<FSectionModel> SChannelView::GetSection() const
+{
+	FViewModelPtr Model = WeakModel.Pin();
+	if (Model)
+	{
+		return Model->FindAncestorOfType<FSectionModel>();
+	}
+	return nullptr;
+}
+
+TSharedPtr<FSequencer> SChannelView::LegacyGetSequencer() const
+{
+	FViewModelPtr Model = WeakModel.Pin();
+	if (Model)
+	{
+		TViewModelPtr<FSequenceModel> Sequence = Model->FindAncestorOfType<FSequenceModel>();
+		if (Sequence)
+		{
+			return Sequence->GetSequencerImpl();
+		}
+	}
+	ensureMsgf(false, TEXT("No valid Sequencer exists"));
+	return nullptr;
+}
+
+void SChannelView::GetKeysUnderMouse(const FVector2D& MousePosition, const FGeometry& AllottedGeometry, TArray<FSequencerSelectedKey>& OutKeys) const
+{
+	TViewModelPtr<FSectionModel> Section = GetSection();
+	if (!Section || !KeyRendererCache.IsSet())
+	{
+		return;
+	}
+
+	const float HalfHeightPx = AllottedGeometry.GetLocalSize().Y*.5f;
+	const FVector2D MousePixel = AllottedGeometry.AbsoluteToLocal(MousePosition);
+
+	// Keys always render in the middle of the geometry so do a quick vertical test first
+	if (FMath::Abs(MousePixel.Y - HalfHeightPx) <= KeyRendererCache->KeySizePx.Y*.5f)
+	{
+		GetKeysAtPixelX(MousePixel.X, OutKeys);
+	}
+}
+
+void SChannelView::GetKeysAtPixelX(float LocalMousePixelX, TArray<FSequencerSelectedKey>& OutKeys) const
+{
+	TViewModelPtr<FSectionModel> Section = GetSection();
+	if (!Section)
+	{
+		return;
+	}
+
+	UMovieSceneSection* SectionObject = Section->GetSection();
+	if (!SectionObject)
+	{
+		return;
+	}
+
+	TRange<FFrameNumber> SectionRange = Section->GetRange();
+
+	// Make the time <-> pixel converter relative to the section range rather than the view range
+	FTimeToPixel RelativeTimeToPixel = *TrackAreaTimeToPixel;
+	if (SectionRange.GetLowerBound().IsClosed())
+	{
+		RelativeTimeToPixel = TrackAreaTimeToPixel->RelativeTo(SectionRange.GetLowerBoundValue());
+	}
+
+	TArray<FKeyRenderer::FKeysForModel> Results;
+	KeyRenderer.HitTestKeys(RelativeTimeToPixel.PixelToFrame(LocalMousePixelX), Results);
+	OutKeys = KeyRendererResultToSelectedKeys(Results, SectionObject);
+}
+
+void SChannelView::CreateKeysUnderMouse(const FVector2D& MousePosition, const FGeometry& AllottedGeometry, TArrayView<const FSequencerSelectedKey> InPressedKeys, TArray<FSequencerSelectedKey>& OutKeys)
+{
+	FViewModelPtr Model = WeakModel.Pin();
+	TViewModelPtr<FSectionModel> Section = GetSection();
+	if (!Section || !Model)
+	{
+		return;
+	}
+
+	UMovieSceneSection* SectionObject = Section->GetSection();
+	if (!SectionObject)
+	{
+		return;
+	}
+
+	if (SectionObject->IsReadOnly())
+	{
+		return;
+	}
+
+	TRange<FFrameNumber> SectionRange = Section->GetRange();
+
+	// Make the time <-> pixel converter relative to the section range rather than the view range
+	FTimeToPixel RelativeTimeToPixel = *TrackAreaTimeToPixel;
+	if (SectionRange.GetLowerBound().IsClosed())
+	{
+		RelativeTimeToPixel = TrackAreaTimeToPixel->RelativeTo(SectionRange.GetLowerBoundValue());
+	}
+
+	// If the pressed key exists, offset the new key and look for it in the newly laid out key areas
+	if (InPressedKeys.Num())
+	{
+		SectionObject->Modify();
+
+		// Offset by 1 pixel worth of time if possible
+		const FFrameTime TimeFuzz = RelativeTimeToPixel.PixelDeltaToFrame(1.f);
+
+		for (const FSequencerSelectedKey& PressedKey : InPressedKeys)
+		{
+			const FFrameNumber CurrentTime = PressedKey.WeakChannel.Pin()->GetKeyArea()->GetKeyTime(PressedKey.KeyHandle);
+			const FKeyHandle NewHandle = PressedKey.WeakChannel.Pin()->GetKeyArea()->DuplicateKey(PressedKey.KeyHandle);
+
+			PressedKey.WeakChannel.Pin()->GetKeyArea()->SetKeyTime(NewHandle, CurrentTime + TimeFuzz.FrameNumber);
+			OutKeys.Add(FSequencerSelectedKey(*SectionObject, PressedKey.WeakChannel, NewHandle));
+		}
+	}
+	else
+	{
+		TSharedPtr<FTrackModel>             TrackModel             = Section->FindAncestorOfType<FTrackModel>();
+		TSharedPtr<IObjectBindingExtension> ObjectBindingExtension = Section->FindAncestorOfType<IObjectBindingExtension>();
+
+		FGuid ObjectBinding = ObjectBindingExtension ? ObjectBindingExtension->GetObjectGuid() : FGuid();
+
+		FVector2D LocalSpaceMousePosition = AllottedGeometry.AbsoluteToLocal(MousePosition);
+		const FFrameTime CurrentTime = RelativeTimeToPixel.PixelToFrame(LocalSpaceMousePosition.X);
+
+		TArray<TSharedRef<IKeyArea>> ValidKeyAreasUnderCursor;
+		for (TViewModelPtr<FChannelModel> Channel : Model->GetDescendantsOfType<FChannelModel>(true))
+		{
+			ValidKeyAreasUnderCursor.Add(Channel->GetKeyArea().ToSharedRef());
+		}
+
+		FScopedTransaction Transaction(LOCTEXT("CreateKeysUnderMouse", "Create keys under mouse"));
+		FAddKeyOperation::FromKeyAreas(TrackModel->GetTrackEditor().Get(), ValidKeyAreasUnderCursor).Commit(CurrentTime.FrameNumber, *LegacyGetSequencer());
+
+		// Get the keys under the mouse as the newly created keys.
+		GetKeysAtPixelX(LocalSpaceMousePosition.X, OutKeys);
+	}
+}
+
+} // namespace UE::Sequencer
+
+#undef LOCTEXT_NAMESPACE
+
diff --git a/Engine/Source/Editor/Sequencer/Private/MVVM/Views/SSequencerTrackAreaView.cpp b/Engine/Source/Editor/Sequencer/Private/MVVM/Views/SSequencerTrackAreaView.cpp
--- a/Engine/Source/Editor/Sequencer/Private/MVVM/Views/SSequencerTrackAreaView.cpp
+++ b/Engine/Source/Editor/Sequencer/Private/MVVM/Views/SSequencerTrackAreaView.cpp
@@ -128,9 +128,8 @@ FReply SSequencerTrackAreaView::OnDragOver(const FGeometry& MyGeometry, const FD
 			// give track editors a chance to accept the drag event
 			const TArray<TSharedPtr<ISequencerTrackEditor>>& TrackEditors = Sequencer->GetTrackEditors();
 
-			FTimeToPixel TimeToPixel(MyGeometry, Sequencer->GetViewRange(), Sequencer->GetFocusedTickResolution());
 			FVector2D LocalPos = MyGeometry.AbsoluteToLocal(DragDropEvent.GetScreenSpacePosition());
-			FFrameNumber DropFrameNumber = TimeToPixel.PixelToFrame(LocalPos.X).FrameNumber;
+			FFrameNumber DropFrameNumber = TimeToPixel->PixelToFrame(LocalPos.X).FrameNumber;
 			if (Sequencer->GetSequencerSettings()->GetIsSnapEnabled() && Sequencer->GetSequencerSettings()->GetSnapPlayTimeToInterval())
 			{
 				DropFrameNumber = FFrameRate::Snap(DropFrameNumber, Sequencer->GetFocusedTickResolution(), Sequencer->GetFocusedDisplayRate()).FrameNumber;
@@ -182,9 +181,8 @@ FReply SSequencerTrackAreaView::OnDrop(const FGeometry& MyGeometry, const FDragD
 			// give track editors a chance to process the drag event
 			const TArray<TSharedPtr<ISequencerTrackEditor>>& TrackEditors = Sequencer->GetTrackEditors();
 
-			FTimeToPixel TimeToPixel(MyGeometry, Sequencer->GetViewRange(), Sequencer->GetFocusedTickResolution());
 			FVector2D LocalPos = MyGeometry.AbsoluteToLocal(DragDropEvent.GetScreenSpacePosition());
-			FFrameNumber DropFrameNumber = TimeToPixel.PixelToFrame(LocalPos.X).FrameNumber;
+			FFrameNumber DropFrameNumber = TimeToPixel->PixelToFrame(LocalPos.X).FrameNumber;
 			if (Sequencer->GetSequencerSettings()->GetIsSnapEnabled() && Sequencer->GetSequencerSettings()->GetSnapPlayTimeToInterval())
 			{
 				DropFrameNumber = FFrameRate::Snap(DropFrameNumber, Sequencer->GetFocusedTickResolution(), Sequencer->GetFocusedDisplayRate()).FrameNumber;
diff --git a/Engine/Source/Editor/Sequencer/Private/SSequencer.cpp b/Engine/Source/Editor/Sequencer/Private/SSequencer.cpp
--- a/Engine/Source/Editor/Sequencer/Private/SSequencer.cpp
+++ b/Engine/Source/Editor/Sequencer/Private/SSequencer.cpp
@@ -2054,6 +2054,7 @@ TSharedRef<SWidget> SSequencer::MakeViewMenu()
 	MenuBuilder.AddMenuEntry( FSequencerCommands::Get().ToggleShowRangeSlider );
 	MenuBuilder.AddMenuEntry( FSequencerCommands::Get().ToggleShowSelectedNodesOnly );
 	MenuBuilder.AddMenuEntry( FSequencerCommands::Get().ToggleLayerBars );
+	MenuBuilder.AddMenuEntry( FSequencerCommands::Get().ToggleKeyBars );
 	MenuBuilder.AddMenuEntry( FSequencerCommands::Get().ToggleChannelColors );
 	MenuBuilder.AddMenuEntry( FSequencerCommands::Get().ToggleShowStatusBar );
 	MenuBuilder.AddMenuEntry( FSequencerCommands::Get().ToggleShowPreAndPostRoll );
diff --git a/Engine/Source/Editor/Sequencer/Private/SSequencerSection.cpp b/Engine/Source/Editor/Sequencer/Private/SSequencerSection.cpp
--- a/Engine/Source/Editor/Sequencer/Private/SSequencerSection.cpp
+++ b/Engine/Source/Editor/Sequencer/Private/SSequencerSection.cpp
@@ -8,9 +8,12 @@
 #include "MVVM/ViewModels/SequencerEditorViewModel.h"
 #include "MVVM/Views/ITrackAreaHotspot.h"
 #include "MVVM/Views/STrackAreaView.h"
+#include "MVVM/Views/SCompoundTrackLaneView.h"
+#include "MVVM/Views/STrackLane.h"
 #include "MVVM/Extensions/IObjectBindingExtension.h"
 #include "Rendering/DrawElements.h"
 #include "Styling/AppStyle.h"
+#include "Styling/StyleColors.h"
 #include "SequencerSelectionPreview.h"
 #include "SequencerSettings.h"
 #include "Editor.h"
@@ -23,15 +26,15 @@
 #include "ISequencerSection.h"
 #include "SequencerHotspots.h"
 #include "Widgets/SOverlay.h"
-#include "SequencerAddKeyOperation.h"
 #include "MovieScene.h"
 #include "Fonts/FontCache.h"
 #include "Framework/Application/SlateApplication.h"
-#include "KeyDrawParams.h"
 #include "MovieSceneTimeHelpers.h"
 #include "Tracks/MovieScenePropertyTrack.h"
 #include "Sections/MovieSceneSubSection.h"
 #include "Generators/MovieSceneEasingFunction.h"
+#include "Widgets/SWeakWidget.h"
+#include "Algo/Transform.h"
 
 namespace UE
 {
@@ -72,6 +75,7 @@ struct FSequencerSectionPainterImpl : FSequencerSectionPainter
 		, Sequencer(InSequencer)
 		, SectionWidget(InSectionWidget)
 		, TimeToPixelConverter(ConstructTimeConverterForSection(SectionGeometry, *InSection->GetSection(), Sequencer))
+		, bClipRectEnabled(false)
 	{
 		CalculateSelectionColor();
 
@@ -99,23 +103,28 @@ struct FSequencerSectionPainterImpl : FSequencerSectionPainter
 
 	virtual int32 PaintSectionBackground(const FLinearColor& Tint) override
 	{
-		UMovieSceneSection* SectionObject = SectionModel->GetSection();
+		using namespace UE::MovieScene;
+
+		TRange<FFrameNumber> SectionRange = SectionModel->GetRange();
 
 		const ESlateDrawEffect DrawEffects = bParentEnabled
 			? ESlateDrawEffect::None
 			: ESlateDrawEffect::DisabledEffect;
 
-		static const FSlateBrush* SectionBackgroundBrush = FAppStyle::GetBrush("Sequencer.Section.Background");
-		static const FSlateBrush* SectionBackgroundTintBrush = FAppStyle::GetBrush("Sequencer.Section.BackgroundTint");
-		static const FSlateBrush* SelectedSectionOverlay = FAppStyle::GetBrush("Sequencer.Section.SelectedSectionOverlay");
+		static const FSlateBrush* CollapsedSectionBackgroundBrush = FAppStyle::GetBrush("Sequencer.Section.Background_Collapsed");
+		static const FSlateBrush* SectionHeaderBackgroundBrush = FAppStyle::GetBrush("Sequencer.Section.Background_Header");
+		static const FSlateBrush* SectionContentsBackgroundBrush = FAppStyle::GetBrush("Sequencer.Section.Background_Contents");
+
+		static const FSlateBrush* CollapsedSelectedSectionOverlay = FAppStyle::GetBrush("Sequencer.Section.CollapsedSelectedSectionOverlay");
+		static const FSlateBrush* SectionHeaderSelectedSectionOverlay = FAppStyle::GetBrush("Sequencer.Section.SectionHeaderSelectedSectionOverlay");
 
 		FLinearColor FinalTint = GetFinalTintColor(Tint);
 
 		// Offset lower bounds and size for infinte sections so we don't draw the rounded border on the visible area
-		const float InfiniteLowerOffset = SectionObject->HasStartFrame() ? 0.f : 100.f;
-		const float InfiniteSizeOffset  = InfiniteLowerOffset + (SectionObject->HasEndFrame() ? 0.f : 100.f);
+		const float InfiniteLowerOffset = SectionRange.GetLowerBound().IsClosed() ? 0.f : 100.f;
+		const float InfiniteSizeOffset  = InfiniteLowerOffset + (SectionRange.GetUpperBound().IsClosed() ? 0.f : 100.f);
 
-		FPaintGeometry PaintGeometry = SectionGeometry.ToPaintGeometry(
+		FGeometry ExpandedSectionGeometry = SectionGeometry.MakeChild(
 			SectionGeometry.GetLocalSize() + FVector2D(InfiniteSizeOffset, 0.f),
 			FSlateLayoutTransform(FVector2D(-InfiniteLowerOffset, 0.f))
 			);
@@ -123,15 +132,18 @@ struct FSequencerSectionPainterImpl : FSequencerSectionPainter
 		if (Sequencer.GetSequencerSettings()->ShouldShowPrePostRoll())
 		{
 			TOptional<FSlateClippingState> PreviousClipState = DrawElements.GetClippingState();
-			DrawElements.PopClip();
+			if (PreviousClipState.IsSet())
+			{
+				DrawElements.PopClip();
+			}
 
 			static const FSlateBrush* PreRollBrush = FAppStyle::GetBrush("Sequencer.Section.PreRoll");
 			float BrushHeight = 16.f, BrushWidth = 10.f;
 
-			if (SectionObject->HasStartFrame())
+			if (SectionRange.GetLowerBound().IsClosed())
 			{
-				FFrameNumber SectionStartTime = SectionObject->GetInclusiveStartFrame();
-				FFrameNumber PreRollStartTime = SectionStartTime - SectionObject->GetPreRollFrames();
+				FFrameNumber SectionStartTime = DiscreteInclusiveLower(SectionRange);
+				FFrameNumber PreRollStartTime = SectionStartTime - SectionModel->GetPreRollFrames();
 
 				const float PreRollPx = TimeToPixelConverter.FrameToPixel(SectionStartTime) - TimeToPixelConverter.FrameToPixel(PreRollStartTime);
 				if (PreRollPx > 0)
@@ -154,10 +166,10 @@ struct FSequencerSectionPainterImpl : FSequencerSectionPainter
 				}
 			}
 
-			if (SectionObject->HasEndFrame())
+			if (SectionRange.GetUpperBound().IsClosed())
 			{
-				FFrameNumber SectionEndTime  = SectionObject->GetExclusiveEndFrame();
-				FFrameNumber PostRollEndTime = SectionEndTime + SectionObject->GetPostRollFrames();
+				FFrameNumber SectionEndTime  = DiscreteExclusiveUpper(SectionRange.GetUpperBound());
+				FFrameNumber PostRollEndTime = SectionEndTime + SectionModel->GetPostRollFrames();
 
 				const float PostRollPx = TimeToPixelConverter.FrameToPixel(PostRollEndTime) - TimeToPixelConverter.FrameToPixel(SectionEndTime);
 				if (PostRollPx > 0)
@@ -187,36 +199,95 @@ struct FSequencerSectionPainterImpl : FSequencerSectionPainter
 			}
 		}
 
+		// If this section has any children, we draw the section header on the top row, and a dimmed fill for all children
+
+		TViewModelPtr<IOutlinerExtension> Outliner = SectionModel->FindAncestorOfType<IOutlinerExtension>();
+		const bool bHasChildren = Outliner && Outliner.AsModel()->GetChildren(EViewModelListType::Outliner);
+		const bool bIsExpanded  = Outliner && Outliner->IsExpanded();
+
+		FLinearColor BlendedTint = BlendColor(Tint).CopyWithNewOpacity(1.f);
 
+		if (!bHasChildren || !bIsExpanded)
 		{
-			TOptional<FSlateClippingState> PreviousClipState = DrawElements.GetClippingState();
-			DrawElements.PopClip();
+			FSlateDrawElement::MakeBox(
+				DrawElements,
+				LayerId,
+				ExpandedSectionGeometry.ToPaintGeometry(),
+				CollapsedSectionBackgroundBrush,
+				DrawEffects,
+				BlendedTint
+			);
+
+			// Draw the selection hash
+			if (SelectionColor.IsSet())
+			{
+				FSlateDrawElement::MakeBox(
+					DrawElements,
+					LayerId,
+					ExpandedSectionGeometry.ToPaintGeometry(FVector2D(1, 1), ExpandedSectionGeometry.GetLocalSize() - FVector2D(2,2)),
+					CollapsedSelectedSectionOverlay,
+					DrawEffects,
+					SelectionColor.GetValue().CopyWithNewOpacity(0.8f)
+				);
+			}
+		}
+		else
+		{
+			TSharedPtr<IGeometryExtension> Geometry = Outliner.ImplicitCast();
+			const float HeaderHeight = Geometry ? Geometry->GetVirtualGeometry().GetHeight() : 10.f;
+
+			FGeometry HeaderGeometry = ExpandedSectionGeometry.MakeChild(
+				FVector2D(ExpandedSectionGeometry.GetLocalSize().X, HeaderHeight),
+				FSlateLayoutTransform()
+			);
+			FPaintGeometry ContentsGeometry = ExpandedSectionGeometry.ToPaintGeometry(
+				FVector2D(ExpandedSectionGeometry.GetLocalSize().X, ExpandedSectionGeometry.GetLocalSize().Y - HeaderHeight),
+				FSlateLayoutTransform(FVector2D(0.f, HeaderHeight)) 
+			);
 
-			// Draw the section background
 			FSlateDrawElement::MakeBox(
 				DrawElements,
 				LayerId,
-				PaintGeometry,
-				SectionBackgroundBrush,
-				DrawEffects
+				HeaderGeometry.ToPaintGeometry(),
+				SectionHeaderBackgroundBrush,
+				DrawEffects,
+				BlendedTint
 			);
-			++LayerId;
 
-			if (PreviousClipState.IsSet())
+			// Draw the selection hash
+			if (SelectionColor.IsSet())
 			{
-				DrawElements.GetClippingManager().PushClippingState(PreviousClipState.GetValue());
+				FSlateDrawElement::MakeBox(
+					DrawElements,
+					LayerId,
+					HeaderGeometry.ToPaintGeometry(FVector2D(1, 1), HeaderGeometry.GetLocalSize() - FVector2D(2,2)),
+					SectionHeaderSelectedSectionOverlay,
+					DrawEffects,
+					SelectionColor.GetValue().CopyWithNewOpacity(0.8f)
+				);
 			}
+
+			FLinearColor FillTint = BlendedTint.LinearRGBToHSV();
+			FillTint.G *= .5f;
+			FillTint.B = FMath::Max(.03f, Tint.B*.1f);
+			FSlateDrawElement::MakeBox(
+				DrawElements,
+				LayerId,
+				ContentsGeometry,
+				SectionContentsBackgroundBrush,
+				DrawEffects,
+				FillTint.HSVToLinearRGB()
+			);
+		}
+
+		if (!ensure(!bClipRectEnabled))
+		{
+			FSlateClippingZone ClippingZone(SectionClippingRect);
+			DrawElements.PushClip(ClippingZone);
+
+			bClipRectEnabled = true;
 		}
 
-		// Draw the section background tint over the background
-		FSlateDrawElement::MakeBox(
-			DrawElements,
-			LayerId,
-			PaintGeometry,
-			SectionBackgroundTintBrush,
-			DrawEffects,
-			FinalTint
-		);
 		++LayerId;
 
 		// Draw underlapping sections
@@ -228,22 +299,12 @@ struct FSequencerSectionPainterImpl : FSequencerSectionPainter
 		// Draw the blend type text
 		DrawBlendType();
 
+		// Draw the locked / key border
+		DrawBorder(FinalTint);
+
 		// Draw easing curves
 		DrawEasing(FinalTint);
 
-		// Draw the selection hash
-		if (SelectionColor.IsSet())
-		{
-			FSlateDrawElement::MakeBox(
-				DrawElements,
-				LayerId,
-				SectionGeometry.ToPaintGeometry(FVector2D(1, 1), SectionGeometry.GetLocalSize() - FVector2D(2,2)),
-				SelectedSectionOverlay,
-				DrawEffects,
-				SelectionColor.GetValue().CopyWithNewOpacity(0.8f)
-			);
-		}
-
 		return LayerId;
 	}
 
@@ -428,7 +489,10 @@ struct FSequencerSectionPainterImpl : FSequencerSectionPainter
 		const float RangeEndPixel = TimeToPixelConverter.FrameToPixel(UE::MovieScene::DiscreteExclusiveUpper(Segment.Range));
 		const float RangeSizePixel = RangeEndPixel - RangeStartPixel;
 
-		FGeometry RangeGeometry = InnerSectionGeometry.MakeChild(FVector2D(RangeSizePixel, InnerSectionGeometry.Size.Y), FSlateLayoutTransform(FVector2D(RangeStartPixel, 0.f)));
+		TViewModelPtr<IGeometryExtension> Geometry = SectionModel->FindAncestorOfType<IOutlinerExtension>().ImplicitCast();
+		const float EasingHeight = (Geometry ? Geometry->GetVirtualGeometry().GetHeight() : InnerSectionGeometry.Size.Y) - 2.f;
+
+		FGeometry RangeGeometry = InnerSectionGeometry.MakeChild(FVector2D(RangeSizePixel, EasingHeight), FSlateLayoutTransform(FVector2D(RangeStartPixel, 1.f)));
 		if (!FSlateRect::DoRectanglesIntersect(RangeGeometry.GetLayoutBoundingRect(), ParentClippingRect))
 		{
 			return;
@@ -624,45 +688,111 @@ struct FSequencerSectionPainterImpl : FSequencerSectionPainter
 	{
 		using namespace UE::Sequencer;
 
-		const ESlateDrawEffect DrawEffects = bParentEnabled ? ESlateDrawEffect::None : ESlateDrawEffect::DisabledEffect;
-		static const FSlateBrush* EmptySpaceBrush = FAppStyle::GetBrush("Sequencer.Section.EmptySpace");
-
-		// Attach contiguous regions together
-		TOptional<FSlateRect> CurrentArea;
+		TSharedPtr<STrackLane> OwningTrackLane = SectionWidget.WeakOwningTrackLane.Pin();
+		if (!OwningTrackLane)
+		{
+			return;
+		}
 
-		for (const FSectionLayoutElement& Element : SectionWidget.Layout->GetElements())
+		// Find all our children
+		TSet<FWeakViewModelPtr> AllDescendents;
+		for (const FViewModelPtr& Descendent : SectionModel->GetDescendants())
 		{
-			const bool bIsEmptySpace = Element.GetModel()->IsA<FChannelModel>() && Element.GetChannels().Num() == 0;
-			const bool bExistingEmptySpace = CurrentArea.IsSet();
+			AllDescendents.Add(Descendent);
+		}
 
-			if (bIsEmptySpace && bExistingEmptySpace && FMath::IsNearlyEqual(CurrentArea->Bottom, Element.GetOffset()))
-			{
-				CurrentArea->Bottom = Element.GetOffset() + Element.GetHeight();
-				continue;
-			}
+		// Find all our child track lanes
+		TArray<TSharedPtr<STrackLane>> EmptyChildLanes;
+		Algo::Transform(OwningTrackLane->GetChildLanes(), EmptyChildLanes, [](const TWeakPtr<STrackLane>& In){ return In.Pin(); });
 
-			if (bExistingEmptySpace)
+		// Remove any invalid track lanes, or track lanes that we made a widget for
+		for (int32 Index = EmptyChildLanes.Num()-1; Index >= 0; --Index)
+		{
+			if (EmptyChildLanes[Index])
 			{
-				FPaintGeometry PaintGeom = SectionGeometry.MakeChild(CurrentArea->GetSize(), FSlateLayoutTransform(CurrentArea->GetTopLeft2f())).ToPaintGeometry();
-				FSlateDrawElement::MakeBox(DrawElements, LayerId, PaintGeom, EmptySpaceBrush, DrawEffects);
-				CurrentArea.Reset();
+				for (const TPair<FWeakViewModelPtr, TSharedPtr<ITrackLaneWidget>>& Pair : EmptyChildLanes[Index]->GetAllWidgets())
+				{
+					if (AllDescendents.Contains(Pair.Key))
+					{
+						EmptyChildLanes.RemoveAtSwap(Index, 1, false);
+						// Move onto the next child lane
+						break;
+					}
+				}
 			}
-
-			if (bIsEmptySpace)
+			else
 			{
-				CurrentArea = FSlateRect::FromPointAndExtent(FVector2D(0.f, Element.GetOffset()), FVector2D(SectionGeometry.Size.X, Element.GetHeight()));
+				EmptyChildLanes.RemoveAtSwap(Index, 1, false);
 			}
 		}
 
-		if (CurrentArea.IsSet())
+		// Paint empty space for remaining track lanes
+		const ESlateDrawEffect DrawEffects = bParentEnabled ? ESlateDrawEffect::None : ESlateDrawEffect::DisabledEffect;
+		static const FSlateBrush* EmptySpaceBrush = FAppStyle::GetBrush("Sequencer.Section.EmptySpace");
+
+		for (TSharedPtr<STrackLane> EmptyTrackLane : EmptyChildLanes)
 		{
-			FPaintGeometry PaintGeom = SectionGeometry.MakeChild(CurrentArea->GetSize(), FSlateLayoutTransform(CurrentArea->GetTopLeft2f())).ToPaintGeometry();
-			FSlateDrawElement::MakeBox(DrawElements, LayerId, PaintGeom, EmptySpaceBrush, DrawEffects);
+			const float LaneTop = EmptyTrackLane->GetVerticalPosition() - OwningTrackLane->GetVerticalPosition();
+			const FVector2D LaneSize(SectionGeometry.GetLocalSize().X, EmptyTrackLane->GetOutlinerItem()->GetOutlinerSizing().Height);
+
+			FSlateDrawElement::MakeBox(
+				DrawElements,
+				LayerId,
+				SectionGeometry.MakeChild(
+					LaneSize,
+					FSlateLayoutTransform(FVector2D(0.f, LaneTop))
+				).ToPaintGeometry(),
+				EmptySpaceBrush,
+				DrawEffects
+			);
 		}
 
 		++LayerId;
 	}
 
+	void DrawBorder(const FLinearColor& FinalTint)
+	{
+		UMovieSceneSection* SectionObject = SectionModel->GetSection();
+		UMovieSceneTrack* Track = SectionObject->GetTypedOuter<UMovieSceneTrack>();
+
+		// Show the locked border if it's locked or read only
+		const bool bLocked = SectionObject->IsLocked() || SectionObject->IsReadOnly();
+
+		// Only show section to key border if we have more than one section
+		const bool bIsSectionToKey = Track && Track->GetAllSections().Num() > 1 && Track->GetSectionToKey() == SectionObject;
+
+		const ESlateDrawEffect DrawEffects = bParentEnabled ? ESlateDrawEffect::None : ESlateDrawEffect::DisabledEffect;
+
+		if (bLocked)
+		{
+			static const FName SelectionBorder("Sequencer.Section.LockedBorder");
+
+			FSlateDrawElement::MakeBox(
+				DrawElements,
+				LayerId,
+				SectionGeometry.ToPaintGeometry(),
+				FAppStyle::GetBrush(SelectionBorder),
+				DrawEffects,
+				FStyleColors::Error.GetColor(FWidgetStyle())
+			);
+			++LayerId;
+		}
+		else if (bIsSectionToKey)
+		{
+			static const FName SelectionBorder("Sequencer.Section.LockedBorder");
+
+			FSlateDrawElement::MakeBox(
+				DrawElements,
+				LayerId,
+				SectionGeometry.ToPaintGeometry(),
+				FAppStyle::GetBrush(SelectionBorder),
+				DrawEffects,
+				FStyleColors::Success.GetColor(FWidgetStyle())
+			);
+			++LayerId;
+		}
+	}
+
 	TOptional<FLinearColor> SelectionColor;
 	FSequencer& Sequencer;
 	const SSequencerSection& SectionWidget;
@@ -671,16 +801,18 @@ struct FSequencerSectionPainterImpl : FSequencerSectionPainter
 
 	/** The clipping rectangle of the parent widget */
 	FSlateRect ParentClippingRect;
+
+	bool bClipRectEnabled;
 };
 
-void SSequencerSection::Construct( const FArguments& InArgs, TSharedPtr<FSequencer> InSequencer, TSharedPtr<FSectionModel> InSectionModel)
+void SSequencerSection::Construct( const FArguments& InArgs, TSharedPtr<FSequencer> InSequencer, TSharedPtr<FSectionModel> InSectionModel, TSharedPtr<STrackLane> OwningTrackLane)
 {
 	using namespace UE::MovieScene;
 
 	Sequencer = InSequencer;
 	WeakSectionModel = InSectionModel;
+	WeakOwningTrackLane = OwningTrackLane;
 	SectionInterface = InSectionModel->GetSectionInterface();
-	Layout = FSectionLayout(InSectionModel);
 	HandleOffsetPx = 0.f;
 
 	SetEnabled(MakeAttributeSP(this, &SSequencerSection::IsEnabled));
@@ -697,7 +829,22 @@ void SSequencerSection::Construct( const FArguments& InArgs, TSharedPtr<FSequenc
 
 	ChildSlot
 	[
-		SectionInterface->GenerateSectionWidget()
+		SNew(SOverlay)
+
+		+ SOverlay::Slot()
+		[
+			SAssignNew(ChildLaneWidgets, SCompoundTrackLaneView)
+			.TimeToPixel(FGetTimeToPixel::CreateLambda([this](const FGeometry& AllottedGeometry){
+				FGeometry SectionGeometry = MakeSectionGeometryWithoutHandles( AllottedGeometry, SectionInterface );
+				UMovieSceneSection& Section = *SectionInterface->GetSectionObject();
+				return ConstructTimeConverterForSection(SectionGeometry, Section, GetSequencer());
+			}))
+		]
+
+		+ SOverlay::Slot()
+		[
+			SectionInterface->GenerateSectionWidget()
+		]
 	];
 }
 
@@ -827,144 +974,6 @@ int32 SSequencerSection::GetOverlapPriority() const
 	return 0;
 }
 
-void SSequencerSection::GetKeysUnderMouse( const FVector2D& MousePosition, const FGeometry& AllottedGeometry, TArray<FSequencerSelectedKey>& OutKeys, float KeyHeightFraction ) const
-{
-	FGeometry SectionGeometry = MakeSectionGeometryWithoutHandles( AllottedGeometry, SectionInterface );
-
-	UMovieSceneSection& Section = *SectionInterface->GetSectionObject();
-	FTimeToPixel TimeToPixelConverter = ConstructTimeConverterForSection(SectionGeometry, Section, GetSequencer());
-	const FVector2D MousePixel  = SectionGeometry.AbsoluteToLocal( MousePosition );
-
-	// HitTest 
-	const FFrameTime HalfKeySizeFrames = TimeToPixelConverter.PixelDeltaToFrame( SequencerSectionConstants::KeySize.X*.5f );
-	const FFrameTime MouseFrameTime    = TimeToPixelConverter.PixelToFrame( MousePixel.X );
-
-	TRange<FFrameNumber> HitTestRange = TRange<FFrameNumber>( (MouseFrameTime-HalfKeySizeFrames).CeilToFrame(), (MouseFrameTime+HalfKeySizeFrames).CeilToFrame());
-	
-	if (HitTestRange.IsEmpty())
-	{
-		return;
-	}
-
-	// Search every key area until we find the one under the mouse
-	for (const FSectionLayoutElement& Element : Layout->GetElements())
-	{
-		const FGeometry KeyAreaGeometry = Element.ComputeGeometry(AllottedGeometry);
-		const FVector2D LocalMousePixel = KeyAreaGeometry.AbsoluteToLocal(MousePosition);
-		const float LocalKeyPosY = KeyAreaGeometry.GetLocalSize().Y * .5f;
-
-		// Check that this section is under our mouse, and discard it from potential selection if the mouse is higher than the key's height. We have to
-		// check keys on a per-section basis (and not for the overall SectionGeometry) because keys are offset on tracks that have expandable 
-		// ranges (ie: Audio, Animation) which otherwise makes them fail the height-threshold check.
-		if (!KeyAreaGeometry.IsUnderLocation(MousePosition) || FMath::Abs(LocalKeyPosY - LocalMousePixel.Y) > SequencerSectionConstants::KeySize.Y * KeyHeightFraction)
-		{
-			continue;
-		}
-
-		for (TWeakPtr<FChannelModel> WeakChannel : Element.GetChannels())
-		{
-			if (TSharedPtr<FChannelModel> Channel = WeakChannel.Pin())
-			{
-				TArray<FKeyHandle> KeyHandles;
-				Channel->GetKeyArea()->GetKeyHandles(KeyHandles, HitTestRange);
-
-				// Only ever select one key from any given key area
-				if (KeyHandles.Num())
-				{
-					OutKeys.Add( FSequencerSelectedKey( Section, Channel, KeyHandles[0] ) );
-				}
-			}
-		}
-
-		// The mouse is in this key area so it cannot possibly be in any other key area
-		return;
-	}
-}
-
-
-void SSequencerSection::CreateKeysUnderMouse( const FVector2D& MousePosition, const FGeometry& AllottedGeometry, TArrayView<const FSequencerSelectedKey> InPressedKeys, TArray<FSequencerSelectedKey>& OutKeys )
-{
-	TSharedPtr<FSectionModel> SectionModel = WeakSectionModel.Pin();
-	if (!SectionModel)
-	{
-		return;
-	}
-
-	UMovieSceneSection& Section = *SectionInterface->GetSectionObject();
-
-	if (Section.IsReadOnly())
-	{
-		return;
-	}
-
-	FGeometry SectionGeometry = MakeSectionGeometryWithoutHandles( AllottedGeometry, SectionInterface );
-	FTimeToPixel TimeToPixelConverter = ConstructTimeConverterForSection(SectionGeometry, Section, GetSequencer());
-
-	// If the pressed key exists, offset the new key and look for it in the newly laid out key areas
-	if (InPressedKeys.Num())
-	{
-		Section.Modify();
-
-		// Offset by 1 pixel worth of time if possible
-		const FFrameTime TimeFuzz = TimeToPixelConverter.PixelDeltaToFrame(1.f);
-
-		for (const FSequencerSelectedKey& PressedKey : InPressedKeys)
-		{
-			const FFrameNumber CurrentTime = PressedKey.WeakChannel.Pin()->GetKeyArea()->GetKeyTime(PressedKey.KeyHandle);
-			const FKeyHandle NewHandle = PressedKey.WeakChannel.Pin()->GetKeyArea()->DuplicateKey(PressedKey.KeyHandle);
-
-			PressedKey.WeakChannel.Pin()->GetKeyArea()->SetKeyTime(NewHandle, CurrentTime + TimeFuzz.FrameNumber);
-			OutKeys.Add(FSequencerSelectedKey(Section, PressedKey.WeakChannel, NewHandle));
-		}
-	}
-	else
-	{
-		TSharedPtr<FTrackModel>             TrackModel             = SectionModel->FindAncestorOfType<FTrackModel>();
-		TSharedPtr<IObjectBindingExtension> ObjectBindingExtension = SectionModel->FindAncestorOfType<IObjectBindingExtension>();
-
-		FGuid ObjectBinding = ObjectBindingExtension ? ObjectBindingExtension->GetObjectGuid() : FGuid();
-
-		FVector2D LocalSpaceMousePosition = SectionGeometry.AbsoluteToLocal( MousePosition );
-		const FFrameTime CurrentTime = TimeToPixelConverter.PixelToFrame(LocalSpaceMousePosition.X);
-
-		TArray<TSharedRef<IKeyArea>> ValidKeyAreasUnderCursor;
-
-		// Search every key area until we find the one under the mouse
-		for (const FSectionLayoutElement& Element : Layout->GetElements())
-		{
-			// Compute the current key area geometry
-			FGeometry KeyAreaGeometryPadded = Element.ComputeGeometry(AllottedGeometry);
-
-			// Is the key area under the mouse
-			if( !KeyAreaGeometryPadded.IsUnderLocation( MousePosition ) )
-			{
-				continue;
-			}
-
-			for (TWeakPtr<FChannelModel> WeakChannel : Element.GetChannels())
-			{
-				if (TSharedPtr<FChannelModel> Channel = WeakChannel.Pin())
-				{
-					ValidKeyAreasUnderCursor.Add(Channel->GetKeyArea().ToSharedRef());
-				}
-			}
-		}
-
-		FScopedTransaction Transaction(NSLOCTEXT("Sequencer", "CreateKeysUnderMouse", "Create keys under mouse"));
-		FAddKeyOperation::FromKeyAreas(TrackModel->GetTrackEditor().Get(), ValidKeyAreasUnderCursor).Commit(CurrentTime.FrameNumber, GetSequencer());
-
-		// Get the keys under the mouse as the newly created keys. Check with the full height of the key track area.
-		const float KeyHeightFraction = 1.f;
-		GetKeysUnderMouse(MousePosition, AllottedGeometry, OutKeys, KeyHeightFraction);
-	}
-
-	if (OutKeys.Num())
-	{
-		Layout = FSectionLayout(SectionModel);
-	}
-}
-
-
 bool SSequencerSection::CheckForEasingHandleInteraction( const FPointerEvent& MouseEvent, const FGeometry& SectionGeometry )
 {
 	UMovieSceneSection* ThisSection = SectionInterface->GetSectionObject();
@@ -1051,6 +1060,18 @@ bool SSequencerSection::CheckForEdgeInteraction( const FPointerEvent& MouseEvent
 		return false;
 	}
 
+	const ISequencerEditTool* EditTool = GetSequencer().GetViewModel()->GetTrackArea()->GetEditTool();
+	TSharedPtr<ITrackAreaHotspot> Hotspot = EditTool ? EditTool->GetDragHotspot() : nullptr;
+	if (!Hotspot)
+	{
+		Hotspot = GetSequencer().GetViewModel()->GetTrackArea()->GetHotspot();
+	}
+
+	if (Hotspot && !Hotspot->IsA<FSectionHotspotBase>())
+	{
+		return false;
+	}
+
 	TArray<TSharedPtr<FSectionModel>> AllUnderlappingSections;
 	AllUnderlappingSections.Add(WeakSectionModel.Pin());
 	for (const FOverlappingSections& Segment : UnderlappingSegments)
@@ -1110,9 +1131,25 @@ bool SSequencerSection::CheckForEdgeInteraction( const FPointerEvent& MouseEvent
 
 bool SSequencerSection::CheckForEasingAreaInteraction( const FPointerEvent& MouseEvent, const FGeometry& SectionGeometry )
 {
+	TSharedPtr<FSectionModel> ThisSectionModel = WeakSectionModel.Pin();
+	if (!ThisSectionModel)
+	{
+		return false;
+	}
+
+	FVector2D LocalMousePos = SectionGeometry.AbsoluteToLocal(MouseEvent.GetScreenSpacePosition());
+
+	TViewModelPtr<IGeometryExtension> HeaderGeometry = ThisSectionModel->FindAncestorOfType<IOutlinerExtension>().ImplicitCast();
+	const float EasingHeight = HeaderGeometry ? HeaderGeometry->GetVirtualGeometry().GetHeight() : SectionGeometry.GetLocalSize().Y;
+
+	if (LocalMousePos.Y < 0 || LocalMousePos.Y > EasingHeight)
+	{
+		return false;
+	}
+
 	UMovieSceneSection* ThisSection = SectionInterface->GetSectionObject();
 	FTimeToPixel TimeToPixelConverter = ConstructTimeConverterForSection(MakeSectionGeometryWithoutHandles(SectionGeometry, SectionInterface), *ThisSection, GetSequencer());
-	const FFrameNumber MouseTime = TimeToPixelConverter.PixelToFrame(SectionGeometry.AbsoluteToLocal(MouseEvent.GetScreenSpacePosition()).X).FrameNumber;
+	const FFrameNumber MouseTime = TimeToPixelConverter.PixelToFrame(LocalMousePos.X).FrameNumber;
 
 	// First off, set the hotspot to an easing area if necessary
 	for (const FOverlappingSections& Segment : UnderlappingEasingSegments)
@@ -1169,14 +1206,6 @@ int32 SSequencerSection::OnPaint( const FPaintArgs& Args, const FGeometry& Allot
 	UMovieSceneTrack* Track = SectionObject->GetTypedOuter<UMovieSceneTrack>();
 	const bool bTrackDisabled = Track && (Track->IsEvalDisabled() || Track->IsRowEvalDisabled(SectionObject->GetRowIndex()));
 	const bool bEnabled = bParentEnabled && SectionObject->IsActive() && !(bTrackDisabled);
-	const bool bLocked = SectionObject->IsLocked() || SectionObject->IsReadOnly();
-
-	bool bSetSectionToKey = false;
-    // Only show section to key border if we have more than one section
-	if (Track && Track->GetAllSections().Num() > 1 && Track->GetSectionToKey() == SectionObject)
-	{
-		bSetSectionToKey = true;
-	}
 
 	const ESlateDrawEffect DrawEffects = bEnabled ? ESlateDrawEffect::None : ESlateDrawEffect::DisabledEffect;
 
@@ -1208,92 +1237,35 @@ int32 SSequencerSection::OnPaint( const FPaintArgs& Args, const FGeometry& Allot
 
 	Painter.bIsSelected = GetSequencer().GetSelection().IsSelected(SectionModel);
 
-	for (const FSectionLayoutElement& Element : Layout->GetElements())
-	{
-		TSharedPtr<FViewModel> Model = Element.GetModel();
-		if (FCategoryModel* Category = Model->CastThis<FCategoryModel>())
-		{
-			TArray<TSharedRef<IKeyArea>> ChildKeyAreas;
-
-			for (TSharedPtr<FChannelModel> Channel : Category->GetDescendantsOfType<FChannelModel>())
-			{
-				if (TSharedPtr<IKeyArea> KeyArea = Channel->GetKeyArea())
-				{
-					ChildKeyAreas.Add(KeyArea.ToSharedRef());
-				}
-			}
-
-			FKeyAreaElement KeyAreaElement;
-			KeyAreaElement.KeyAreas = ChildKeyAreas;
-			KeyAreaElement.KeyAreaGeometry = Element.ComputeGeometry(AllottedGeometry);
-			KeyAreaElement.Type = (FKeyAreaElement::EType)Element.GetType();
-			Painter.KeyAreaElements.Add(KeyAreaElement);
-		}
-		else
-		{
-			FKeyAreaElement KeyAreaElement;
-			for (TWeakPtr<FChannelModel> WeakChannel : Element.GetChannels())
-			{
-				if (TSharedPtr<FChannelModel> Channel = WeakChannel.Pin())
-				{
-					KeyAreaElement.KeyAreas.Add(Channel->GetKeyArea().ToSharedRef());
-				}
-			}
-			KeyAreaElement.KeyAreaGeometry = Element.ComputeGeometry(AllottedGeometry);
-			KeyAreaElement.Type = (FKeyAreaElement::EType)Element.GetType();
-			Painter.KeyAreaElements.Add(KeyAreaElement);
-		}
-	}
-
-	FSlateClippingZone ClippingZone(Painter.SectionClippingRect);
-	OutDrawElements.PushClip(ClippingZone);
-
 	// Ask the interface to draw the section
 	LayerId = SectionInterface->OnPaintSection(Painter);
 
-	LayerId = SCompoundWidget::OnPaint( Args, AllottedGeometry, MyCullingRect, OutDrawElements, LayerId, InWidgetStyle, bEnabled );
-
 	FLinearColor SelectionColor = FAppStyle::GetSlateColor(SequencerSectionConstants::SelectionColorName).GetColor(FWidgetStyle());
 	DrawSectionHandles(AllottedGeometry, OutDrawElements, LayerId, DrawEffects, SelectionColor, Hotspot);
 
+	LayerId = SCompoundWidget::OnPaint( Args, AllottedGeometry, MyCullingRect, OutDrawElements, LayerId, InWidgetStyle, bEnabled );
+
 	Painter.LayerId = LayerId;
 	PaintEasingHandles( Painter, SelectionColor, Hotspot );
 
+	// --------------------------------------------
+	// Draw section selection tint
+	const float SectionThrobValue = GetSectionSelectionThrobValue();
+	if (SectionThrobValue != 0.f && GetSequencer().GetSelection().IsSelected(SectionModel))
 	{
-		FKeyRendererPaintArgs KeyRenderArgs;
-		KeyRenderArgs.KeyThrobValue = GetKeySelectionThrobValue();
-		KeyRenderArgs.SectionThrobValue = GetSectionSelectionThrobValue();
-
-		KeyRenderer.Paint(Layout.GetValue(), InWidgetStyle, KeyRenderArgs, &GetSequencer(), Painter);
-	}
-
-	LayerId = Painter.LayerId;
-	if (bLocked)
-	{
-		static const FName SelectionBorder("Sequencer.Section.LockedBorder");
+		static const FName BackgroundTrackTintBrushName("Sequencer.Section.BackgroundTint");
 
 		FSlateDrawElement::MakeBox(
 			OutDrawElements,
-			LayerId,
+			LayerId++,
 			AllottedGeometry.ToPaintGeometry(),
-			FAppStyle::GetBrush(SelectionBorder),
+			FAppStyle::GetBrush(BackgroundTrackTintBrushName),
 			DrawEffects,
-			FLinearColor::Red
+			SelectionColor.CopyWithNewOpacity(SectionThrobValue)
 		);
 	}
-	else if (bSetSectionToKey)
-	{
-		static const FName SelectionBorder("Sequencer.Section.LockedBorder");
 
-		FSlateDrawElement::MakeBox(
-			OutDrawElements,
-			LayerId,
-			AllottedGeometry.ToPaintGeometry(),
-			FAppStyle::GetBrush(SelectionBorder),
-			DrawEffects,
-			FLinearColor::Green
-		);
-	}
+	LayerId = Painter.LayerId;
 
 	// Section name with drop shadow
 	FText SectionTitle = SectionInterface->GetSectionTitle();
@@ -1307,6 +1279,9 @@ int32 SSequencerSection::OnPaint( const FPaintArgs& Args, const FGeometry& Allot
 
 	if (!SectionTitle.IsEmpty())
 	{
+		FSlateClippingZone ClippingZone(Painter.SectionClippingRect);
+		OutDrawElements.PushClip(ClippingZone);
+
 		FVector2D TopLeft = SectionGeometry.AbsoluteToLocal(Painter.SectionClippingRect.GetTopLeft()) + FVector2D(1.f, -1.f);
 
 		FSlateFontInfo FontInfo = FAppStyle::GetFontStyle("NormalFont");
@@ -1322,20 +1297,6 @@ int32 SSequencerSection::OnPaint( const FPaintArgs& Args, const FGeometry& Allot
 			FontInfo.Size = FMath::Max(FMath::FloorToInt(FontInfo.Size - 6.f), 11);
 		}
 
-		// Drop shadow
-		FSlateDrawElement::MakeText(
-			OutDrawElements,
-			LayerId,
-			SectionGeometry.MakeChild(
-				FVector2D(SectionGeometry.Size.X, GetFontHeight()),
-				FSlateLayoutTransform(TopLeft + FVector2D(ContentPadding.Left, ContentPadding.Top) + FVector2D(1.f, 1.f))
-			).ToPaintGeometry(),
-			SectionTitle,
-			FontInfo,
-			DrawEffects,
-			FLinearColor(0,0,0,.5f * Painter.GhostAlpha)
-		);
-
 		FSlateDrawElement::MakeText(
 			OutDrawElements,
 			LayerId,
@@ -1346,11 +1307,16 @@ int32 SSequencerSection::OnPaint( const FPaintArgs& Args, const FGeometry& Allot
 			SectionTitle,
 			FontInfo,
 			DrawEffects,
-			FColor(200, 200, 200, static_cast<uint8>(Painter.GhostAlpha * 255))
+			FColor(192, 192, 192, static_cast<uint8>(Painter.GhostAlpha * 255))
 		);
+
+		OutDrawElements.PopClip();
 	}
 
-	OutDrawElements.PopClip();
+	if (Painter.bClipRectEnabled)
+	{
+		OutDrawElements.PopClip();
+	}
 	return LayerId + 1;
 }
 
@@ -1362,6 +1328,11 @@ void SSequencerSection::PaintEasingHandles( FSequencerSectionPainter& InPainter,
 		return;
 	}
 
+	if (Hotspot && !Hotspot->IsA<FSectionHotspotBase>())
+	{
+		return;
+	}
+
 	TArray<TSharedPtr<FSectionModel>> AllUnderlappingSections;
 	if (IsSectionHighlighted(SectionInterface->GetSectionObject(), Hotspot))
 	{
@@ -1495,8 +1466,16 @@ void SSequencerSection::DrawSectionHandles( const FGeometry& AllottedGeometry, F
 		return;
 	}
 
+	if (Hotspot && !Hotspot->IsA<FSectionHotspotBase>())
+	{
+		return;
+	}
+
 	TOptional<FSlateClippingState> PreviousClipState = OutDrawElements.GetClippingState();
-	OutDrawElements.PopClip();
+	if (PreviousClipState.IsSet())
+	{
+		OutDrawElements.PopClip();
+	}
 
 	OutDrawElements.PushClip(FSlateClippingZone(AllottedGeometry.GetLayoutBoundingRect()));
 
@@ -1557,10 +1536,10 @@ void SSequencerSection::DrawSectionHandles( const FGeometry& AllottedGeometry, F
 		const float ThisHandleOffset = UnderlappingSectionObj == ThisSection ? HandleOffsetPx : 0.f;
 		FVector2D GripSize( UnderlappingSection->GetSectionGripSize(), AllottedGeometry.Size.Y );
 
-		float Opacity = 0.5f;
+		float Opacity = 0.85f;
 		if (ThisHandleOffset != 0)
 		{
-			Opacity = FMath::Clamp(.5f + ThisHandleOffset / GripSize.X * .5f, .5f, 1.f);
+			Opacity = FMath::Clamp(Opacity + ThisHandleOffset / GripSize.X * (1.f-Opacity), Opacity, 1.f);
 		}
 
 		const FSlateBrush* LeftGripBrush = FAppStyle::GetBrush("Sequencer.Section.GripLeft");
@@ -1614,8 +1593,6 @@ void SSequencerSection::Tick( const FGeometry& AllottedGeometry, const double In
 {
 	if( GetVisibility() == EVisibility::Visible )
 	{
-		Layout = FSectionLayout(WeakSectionModel.Pin());
-
 		UMovieSceneSection* Section = SectionInterface->GetSectionObject();
 		if (Section && Section->HasStartFrame() && Section->HasEndFrame())
 		{
@@ -1640,9 +1617,9 @@ void SSequencerSection::Tick( const FGeometry& AllottedGeometry, const double In
 	}
 }
 
-void SSequencerSection::AddChildLane(TSharedPtr<ITrackLaneWidget> ChildWidget)
+void SSequencerSection::AddChildView(TSharedPtr<ITrackLaneWidget> ChildWidget, TWeakPtr<STrackLane> InWeakOwningLane)
 {
-	
+	ChildLaneWidgets->AddWeakWidget(ChildWidget, InWeakOwningLane);
 }
 
 void SSequencerSection::OnModifiedIndirectly(UMovieSceneSignedObject* Object)
@@ -1655,42 +1632,9 @@ void SSequencerSection::OnModifiedIndirectly(UMovieSceneSignedObject* Object)
 
 FReply SSequencerSection::OnMouseButtonDown( const FGeometry& MyGeometry, const FPointerEvent& MouseEvent )
 {
-	TArrayView<const FSequencerSelectedKey> HoveredKeys;
-	
-	// The hovered key is defined from the sequencer hotspot
-	if (TSharedPtr<FKeyHotspot> KeyHotspot = HotspotCast<FKeyHotspot>(GetSequencer().GetViewModel()->GetTrackArea()->GetHotspot()))
-	{
-		HoveredKeys = KeyHotspot->Keys;
-	}
-
-	if (MouseEvent.GetEffectingButton() == EKeys::MiddleMouseButton)
-	{
-		GEditor->BeginTransaction(NSLOCTEXT("Sequencer", "CreateKey_Transaction", "Create Key"));
-
-		// Generate a key and set it as the PressedKey
-		TArray<FSequencerSelectedKey> NewKeys;
-		CreateKeysUnderMouse(MouseEvent.GetScreenSpacePosition(), MyGeometry, HoveredKeys, NewKeys);
-
-		if (NewKeys.Num())
-		{
-			GetSequencer().GetSelection().EmptySelectedKeys();
-			for (const FSequencerSelectedKey& NewKey : NewKeys)
-			{
-				GetSequencer().GetSelection().AddToSelection(NewKey);
-			}
-
-			// Pass the event to the tool to copy the hovered key and move it
-			GetSequencer().GetViewModel()->GetTrackArea()->SetHotspot( MakeShared<FKeyHotspot>(MoveTemp(NewKeys), Sequencer) );
-
-			// Return unhandled so that the EditTool can handle the mouse down based on the newly created keyframe and prepare to move it
-			return FReply::Unhandled();
-		}
-	}
-
 	return FReply::Unhandled();
 }
 
-
 FGeometry SSequencerSection::MakeSectionGeometryWithoutHandles( const FGeometry& AllottedGeometry, const TSharedPtr<ISequencerSection>& InSectionInterface ) const
 {
 	return AllottedGeometry.MakeChild(
@@ -1724,18 +1668,6 @@ FReply SSequencerSection::OnMouseButtonDoubleClick( const FGeometry& MyGeometry,
 
 	if( MouseEvent.GetEffectingButton() == EKeys::LeftMouseButton )
 	{
-		TArray<FSequencerSelectedKey> Keys;
-		GetKeysUnderMouse(MouseEvent.GetScreenSpacePosition(), MyGeometry, Keys);
-		TArray<FKeyHandle> KeyHandles;
-		for (FSequencerSelectedKey Key : Keys)
-		{
-			KeyHandles.Add(Key.KeyHandle);
-		}
-		if (KeyHandles.Num() > 0)
-		{
-			return SectionInterface->OnKeyDoubleClicked(KeyHandles);
-		}
-
 		FReply Reply = SectionInterface->OnSectionDoubleClicked( MyGeometry, MouseEvent );
 		if (!Reply.IsEventHandled())
 		{
@@ -1765,15 +1697,7 @@ FReply SSequencerSection::OnMouseButtonDoubleClick( const FGeometry& MyGeometry,
 
 FReply SSequencerSection::OnMouseMove( const FGeometry& MyGeometry, const FPointerEvent& MouseEvent )
 {
-	// Checked for hovered key
-	TArray<FSequencerSelectedKey> KeysUnderMouse;
-	GetKeysUnderMouse( MouseEvent.GetScreenSpacePosition(), MyGeometry, KeysUnderMouse );
-	if ( KeysUnderMouse.Num() )
-	{
-		GetSequencer().GetViewModel()->GetTrackArea()->SetHotspot( MakeShared<FKeyHotspot>( MoveTemp(KeysUnderMouse), Sequencer ) );
-	}
-	// Check other interaction points in order of importance
-	else if (
+	if ( MouseEvent.GetCursorDelta().Size() > 0 &&
 		!CheckForEasingHandleInteraction(MouseEvent, MyGeometry) &&
 		!CheckForEdgeInteraction(MouseEvent, MyGeometry) &&
 		!CheckForEasingAreaInteraction(MouseEvent, MyGeometry))
@@ -1784,35 +1708,16 @@ FReply SSequencerSection::OnMouseMove( const FGeometry& MyGeometry, const FPoint
 
 	return FReply::Unhandled();
 }
-	
+
 FReply SSequencerSection::OnMouseButtonUp(const FGeometry& MyGeometry, const FPointerEvent& MouseEvent)
 {
-	if (MouseEvent.GetEffectingButton() == EKeys::MiddleMouseButton)
-	{
-		// Snap keys on mouse up since we want to create keys at the exact mouse position (ie. to keep the newly created keys under the mouse 
-		// while dragging) but obey snapping rules if necessary
-		if (GetSequencer().GetSequencerSettings()->GetIsSnapEnabled() && GetSequencer().GetSequencerSettings()->GetSnapKeyTimesToInterval())
-		{
-			GetSequencer().SnapToFrame();
-
-			for (const FSequencerSelectedKey& SelectedKey : GetSequencer().GetSelection().GetSelectedKeys())
-			{
-				const FFrameNumber CurrentTime = SelectedKey.WeakChannel.Pin()->GetKeyArea()->GetKeyTime(SelectedKey.KeyHandle);
-				GetSequencer().SetLocalTime(CurrentTime, ESnapTimeMode::STM_Interval);
-				break;
-			}
-		}
-		GEditor->EndTransaction();
-
-		// Return unhandled so that the EditTool can handle the mouse up based on the newly created keyframe and finish moving it
-		return FReply::Unhandled();
-	}
 	return FReply::Unhandled();
 }
 
 void SSequencerSection::OnMouseLeave( const FPointerEvent& MouseEvent )
 {
 	SCompoundWidget::OnMouseLeave( MouseEvent );
+
 	GetSequencer().GetViewModel()->GetTrackArea()->SetHotspot(nullptr);
 }
 
diff --git a/Engine/Source/Editor/Sequencer/Private/SSequencerSection.h b/Engine/Source/Editor/Sequencer/Private/SSequencerSection.h
--- a/Engine/Source/Editor/Sequencer/Private/SSequencerSection.h
+++ b/Engine/Source/Editor/Sequencer/Private/SSequencerSection.h
@@ -7,11 +7,7 @@
 #include "Input/Reply.h"
 #include "Widgets/DeclarativeSyntaxSupport.h"
 #include "Widgets/SCompoundWidget.h"
-#include "SequencerSelectedKey.h"
 #include "Rendering/RenderingCommon.h"
-#include "SectionLayout.h"
-#include "SequencerKeyRenderer.h"
-#include "DisplayNodes/SequencerKeyTimeCache.h"
 
 #include "EventHandlers/ISignedObjectEventHandler.h"
 #include "MVVM/ViewModels/SectionModel.h"
@@ -27,6 +23,7 @@ namespace UE
 namespace Sequencer
 {
 
+class SCompoundTrackLaneView;
 class FSectionModel;
 struct ITrackAreaHotspot;
 
@@ -37,7 +34,7 @@ public:
 	{}
 	SLATE_END_ARGS()
 
-	void Construct( const FArguments& InArgs, TSharedPtr<FSequencer> Sequencer, TSharedPtr<FSectionModel> InSectionModel);
+	void Construct( const FArguments& InArgs, TSharedPtr<FSequencer> Sequencer, TSharedPtr<FSectionModel> InSectionModel, TSharedPtr<STrackLane> InOwningTrackLane);
 	~SSequencerSection();
 
 	TSharedPtr<ISequencerSection> GetSectionInterface() const { return SectionInterface; }
@@ -46,27 +43,6 @@ public:
 
 private:
 
-	/**
-	 * Determines the key that is under the mouse
-	 *
-	 * @param MousePosition		The current screen space position of the mouse
-	 * @param AllottedGeometry	The geometry of the mouse event
-	 * @param OutKeys			Array to populate with the keys that are under the mouse
-	 * @param KeyHeightFraction Optional key height fraction. Keys are by default offset and half height. Pass in 1 if checking full height of key track area.  
-	 * @return The key that is under the mouse.  Invalid if there is no key under the mouse
-	 */
-	void GetKeysUnderMouse( const FVector2D& MousePosition, const FGeometry& AllottedGeometry, TArray<FSequencerSelectedKey>& OutKeys, float KeyHeightFraction=0.5f ) const;
-
-	/**
-	 * Creates a key at the mouse position
-	 *
-	 * @param MousePosition		The current screen space position of the mouse
-	 * @param AllottedGeometry	The geometry of the mouse event
-	 * @param InPressedKeys     Keys if pressed
-	 * @param OutKeys			Array to populate with the new keys
-	 */
-	void CreateKeysUnderMouse( const FVector2D& MousePosition, const FGeometry& AllottedGeometry, TArrayView<const FSequencerSelectedKey> InPressedKeys, TArray<FSequencerSelectedKey>& OutKeys );
-
 	/**
 	 * Checks for user interaction (via the mouse) with the left and right edge of a section
 	 *
@@ -94,7 +70,7 @@ private:
 	virtual void ReportParentGeometry(const FGeometry& InParentGeometry) override;
 	virtual TSharedRef<const SWidget> AsWidget() const override { return AsShared(); }
 	virtual bool AcceptsChildren() const override { return true; }
-	virtual void AddChildLane(TSharedPtr<ITrackLaneWidget> ChildWidget) override;
+	virtual void AddChildView(TSharedPtr<ITrackLaneWidget> ChildWidget, TWeakPtr<STrackLane> InWeakOwningLane) override;
 
 	/**
 	 * Paint the easing handles for this section
@@ -139,12 +115,12 @@ public:
 	/** Indicate that the current key selection should throb the specified number of times. A single throb takes 0.2s. */
 	static void ThrobKeySelection(int32 ThrobCount = 1);
 
-	/** Get a value between 0 and 1 that indicates the amount of throb-scale to apply to the currently selected sections */
-	static float GetSectionSelectionThrobValue();
-
 	/** Get a value between 0 and 1 that indicates the amount of throb-scale to apply to the currently selected keys */
 	static float GetKeySelectionThrobValue();
 
+	/** Get a value between 0 and 1 that indicates the amount of throb-scale to apply to the currently selected sections */
+	static float GetSectionSelectionThrobValue();
+
 	/**
 	 * Check to see whether the specified section is highlighted
 	 */
@@ -156,8 +132,10 @@ private:
 	TSharedPtr<ISequencerSection> SectionInterface;
 	/** Section model */
 	TWeakPtr<FSectionModel> WeakSectionModel;
-	/** Cached layout generated each tick */
-	TOptional<FSectionLayout> Layout;
+	/** Widget container for child lanes */
+	TSharedPtr<SCompoundTrackLaneView> ChildLaneWidgets;
+	/** The track lane that this widget is on */
+	TWeakPtr<STrackLane> WeakOwningTrackLane;
 	/** Cached parent geometry to pass down to any section interfaces that need it during tick */
 	FGeometry ParentGeometry;
 	/** The end time for a throbbing animation for selected sections */
@@ -170,8 +148,6 @@ private:
 	TArray<FOverlappingSections> UnderlappingSegments;
 	/** Array of segments that define other sections that reside below this one */
 	TArray<FOverlappingSections> UnderlappingEasingSegments;
-	/** structure responsible for rendering keys  */
-	FKeyRenderer KeyRenderer;
 
 	MovieScene::TNonIntrusiveEventHandler<MovieScene::ISignedObjectEventHandler> TrackModifiedBinding;
 
diff --git a/Engine/Source/Editor/Sequencer/Private/Sequencer.cpp b/Engine/Source/Editor/Sequencer/Private/Sequencer.cpp
--- a/Engine/Source/Editor/Sequencer/Private/Sequencer.cpp
+++ b/Engine/Source/Editor/Sequencer/Private/Sequencer.cpp
@@ -158,6 +158,7 @@
 #include "EntitySystem/MovieSceneEntitySystemLinker.h"
 #include "EntitySystem/MovieScenePreAnimatedStateSystem.h"
 #include "Systems/MovieSceneMotionVectorSimulationSystem.h"
+#include "IKeyArea.h"
 
 #include "EngineModule.h"
 
@@ -10709,6 +10710,15 @@ void FSequencer::BindCommands()
 		FCanExecuteAction::CreateLambda( []{ return true; } ),
 		FIsActionChecked::CreateLambda( [this]{ return Settings->GetShowLayerBars(); } ) );
 
+	SequencerCommandBindings->MapAction(
+		Commands.ToggleKeyBars,
+		FExecuteAction::CreateLambda( [this]{
+			Settings->SetShowKeyBars( !Settings->GetShowKeyBars() );
+			RefreshUI();
+		} ),
+		FCanExecuteAction::CreateLambda( []{ return true; } ),
+		FIsActionChecked::CreateLambda( [this]{ return Settings->GetShowKeyBars(); } ) );
+
 	SequencerCommandBindings->MapAction(
 		Commands.ToggleChannelColors,
 		FExecuteAction::CreateLambda( [this]{
diff --git a/Engine/Source/Editor/Sequencer/Private/SequencerChannelTraits.cpp b/Engine/Source/Editor/Sequencer/Private/SequencerChannelTraits.cpp
--- a/Engine/Source/Editor/Sequencer/Private/SequencerChannelTraits.cpp
+++ b/Engine/Source/Editor/Sequencer/Private/SequencerChannelTraits.cpp
@@ -13,6 +13,7 @@ void DrawKeys(FMovieSceneChannel* Channel, TArrayView<const FKeyHandle> InHandle
 	// By default just render diamonds for keys
 	FKeyDrawParams DefaultParams;
 	DefaultParams.BorderBrush = DefaultParams.FillBrush = FAppStyle::Get().GetBrush("Sequencer.KeyDiamond");
+	DefaultParams.ConnectionStyle = EKeyConnectionStyle::Solid;
 
 	for (FKeyDrawParams& Param : OutKeyDrawParams)
 	{
diff --git a/Engine/Source/Editor/Sequencer/Private/SequencerCommands.cpp b/Engine/Source/Editor/Sequencer/Private/SequencerCommands.cpp
--- a/Engine/Source/Editor/Sequencer/Private/SequencerCommands.cpp
+++ b/Engine/Source/Editor/Sequencer/Private/SequencerCommands.cpp
@@ -137,6 +137,7 @@ void FSequencerCommands::RegisterCommands()
 
 	UI_COMMAND( FindInContentBrowser, "Find in Content Browser", "Find the viewed sequence asset in the content browser", EUserInterfaceActionType::Button, FInputChord() );
 	UI_COMMAND( ToggleLayerBars, "Layer Bars", "Show/hide the layer bars to edit keyframes in bulk", EUserInterfaceActionType::ToggleButton, FInputChord() );
+	UI_COMMAND( ToggleKeyBars, "Key Bars", "Show/hide key bar connectors for quickly retiming pairs of keys", EUserInterfaceActionType::ToggleButton, FInputChord() );
 	UI_COMMAND( ToggleChannelColors, "Channel Colors", "Show/hide the channel colors in the track area", EUserInterfaceActionType::ToggleButton, FInputChord() );
 	UI_COMMAND( ToggleShowSelectedNodesOnly, "Selected Nodes Only", "Show selected nodes only", EUserInterfaceActionType::ToggleButton, FInputChord() );
 	
diff --git a/Engine/Source/Editor/Sequencer/Private/SequencerContextMenus.cpp b/Engine/Source/Editor/Sequencer/Private/SequencerContextMenus.cpp
--- a/Engine/Source/Editor/Sequencer/Private/SequencerContextMenus.cpp
+++ b/Engine/Source/Editor/Sequencer/Private/SequencerContextMenus.cpp
@@ -7,6 +7,7 @@
 #include "SequencerCommands.h"
 #include "SSequencer.h"
 #include "SectionLayout.h"
+#include "IKeyArea.h"
 #include "SSequencerSection.h"
 #include "SequencerSettings.h"
 #include "MVVM/Views/ITrackAreaHotspot.h"
@@ -442,7 +443,13 @@ void FSectionContextMenu::PopulateMenu(FMenuBuilder& MenuBuilder, TSharedPtr<FEx
 				LOCTEXT("KeySection", "Key This Section"),
 				LOCTEXT("KeySection_ToolTip", "This section will get changed when we modify the property externally"),
 				FSlateIcon(),
-				FUIAction(FExecuteAction::CreateLambda([=] { return Shared->SetSectionToKey(); }))
+				FUIAction(
+					FExecuteAction::CreateLambda([=] { return Shared->SetSectionToKey(); }),
+					FCanExecuteAction(),
+					FIsActionChecked::CreateLambda([=] { return Shared->IsSectionToKey(); })
+				),
+				NAME_None,
+				EUserInterfaceActionType::Check
 			);
 		}
 	}
@@ -673,6 +680,7 @@ void FSectionContextMenu::SetSectionToKey()
 		return;
 	}
 
+	const bool bToggle = IsSectionToKey();
 	for (TWeakObjectPtr<UMovieSceneSection> WeakSection : Sequencer->GetSelection().GetSelectedSections())
 	{
 		if (UMovieSceneSection* Section = WeakSection.Get())
@@ -682,13 +690,29 @@ void FSectionContextMenu::SetSectionToKey()
 			{
 				FScopedTransaction Transaction(LOCTEXT("SetSectionToKey", "Set Section To Key"));
 				Track->Modify();
-				Track->SetSectionToKey(Section);
+				Track->SetSectionToKey(bToggle ? nullptr : Section);
 			}
 		}
 		break;
 	}
 }
 
+bool FSectionContextMenu::IsSectionToKey() const
+{
+	for (TWeakObjectPtr<UMovieSceneSection> WeakSection : Sequencer->GetSelection().GetSelectedSections())
+	{
+		if (UMovieSceneSection* Section = WeakSection.Get())
+		{
+			UMovieSceneTrack* Track = Section->GetTypedOuter<UMovieSceneTrack>();
+			if (Track && Track->GetSectionToKey() != Section)
+			{
+				return false;
+			}
+		}
+	}
+	return true;
+}
+
 bool FSectionContextMenu::CanSetSectionToKey() const
 {
 	if (Sequencer->GetSelection().GetSelectedSections().Num() != 1)
diff --git a/Engine/Source/Editor/Sequencer/Private/SequencerContextMenus.h b/Engine/Source/Editor/Sequencer/Private/SequencerContextMenus.h
--- a/Engine/Source/Editor/Sequencer/Private/SequencerContextMenus.h
+++ b/Engine/Source/Editor/Sequencer/Private/SequencerContextMenus.h
@@ -52,6 +52,8 @@ private:
 
 	void SetSectionToKey();
 
+	bool IsSectionToKey() const;
+
 	bool CanSetSectionToKey() const;
 
 	void AutoSizeSection();
diff --git a/Engine/Source/Editor/Sequencer/Private/SequencerHotspots.cpp b/Engine/Source/Editor/Sequencer/Private/SequencerHotspots.cpp
--- a/Engine/Source/Editor/Sequencer/Private/SequencerHotspots.cpp
+++ b/Engine/Source/Editor/Sequencer/Private/SequencerHotspots.cpp
@@ -9,14 +9,17 @@
 #include "MVVM/Views/STrackAreaView.h"
 #include "Tools/SequencerEditTool_Movement.h"
 #include "Tools/SequencerEditTool_Selection.h"
+#include "Tools/SequencerSnapField.h"
 #include "Widgets/Layout/SBox.h"
 #include "Styling/AppStyle.h"
 #include "Channels/MovieSceneChannel.h"
 #include "Channels/MovieSceneChannelProxy.h"
 #include "MovieSceneTimeHelpers.h"
 #include "SequencerCommonHelpers.h"
+#include "MovieSceneSignedObject.h"
 
 #include "MVVM/ViewModels/SectionModel.h"
+#include "MVVM/ViewModels/VirtualTrackArea.h"
 
 #define LOCTEXT_NAMESPACE "SequencerHotspots"
 
@@ -26,6 +29,7 @@ namespace Sequencer
 {
 
 UE_SEQUENCER_DEFINE_CASTABLE(FKeyHotspot);
+UE_SEQUENCER_DEFINE_CASTABLE(FKeyBarHotspot);
 UE_SEQUENCER_DEFINE_CASTABLE(FSectionEasingAreaHotspot);
 UE_SEQUENCER_DEFINE_CASTABLE(FSectionEasingHandleHotspot);
 UE_SEQUENCER_DEFINE_CASTABLE(FSectionHotspot);
@@ -181,6 +185,149 @@ bool FKeyHotspot::PopulateContextMenu(FMenuBuilder& MenuBuilder, TSharedPtr<FExt
 	return true;
 }
 
+void FKeyBarHotspot::UpdateOnHover(FTrackAreaViewModel& InTrackArea) const
+{
+}
+TOptional<FFrameNumber> FKeyBarHotspot::GetTime() const
+{
+	FFrameNumber Time = 0;
+
+	if (LeadingKeys.Num())
+	{
+		TArrayView<const FSequencerSelectedKey> FirstKey(&LeadingKeys.Last(), 1);
+		TArrayView<FFrameNumber> FirstKeyTime(&Time, 1);
+		GetKeyTimes(FirstKey, FirstKeyTime);
+	}
+
+	return Time;
+}
+
+FCursorReply FKeyBarHotspot::GetCursor() const
+{
+	return FCursorReply::Cursor(EMouseCursor::ResizeLeftRight);
+}
+
+TSharedPtr<ISequencerEditToolDragOperation> FKeyBarHotspot::InitiateDrag(const FPointerEvent& MouseEvent)
+{
+	struct FKeyBarDrag : ISequencerEditToolDragOperation, UE::Sequencer::ISnapCandidate
+	{
+		FFrameTime RelativeStartTime;
+
+		TArray<FFrameTime> RelativeStartTimes;
+		TArray<FSequencerSelectedKey> AllLinearKeys;
+
+		TSet<FSequencerSelectedKey> AllKeys;
+
+		TSharedPtr<FKeyBarHotspot> Hotspot;
+		TUniquePtr<FScopedTransaction> Transaction;
+		TSet<UMovieSceneSection*> ModifiedSections;
+		FSequencerSnapField SnapField;
+	
+		FKeyBarDrag(TSharedPtr<FKeyBarHotspot> InHotspot, TSharedPtr<FSequencer> InSequencer)
+			: Hotspot(InHotspot)
+		{
+			AllKeys.Append(Hotspot->LeadingKeys);
+			AllKeys.Append(Hotspot->TrailingKeys);
+			AllKeys.Append(InSequencer->GetSelection().GetSelectedKeys());
+
+			AllLinearKeys = AllKeys.Array();
+
+			SnapField = FSequencerSnapField(*InSequencer, *this);
+		}
+
+		void OnBeginDrag(const FPointerEvent& MouseEvent, FVector2D LocalMousePos, const FVirtualTrackArea& VirtualTrackArea) override
+		{
+			RelativeStartTime = VirtualTrackArea.PixelToFrame(LocalMousePos.X);
+
+			// Cache off the starting times
+			TArray<FFrameNumber> AbsoluteStartTimes;
+			AbsoluteStartTimes.SetNumUninitialized(AllLinearKeys.Num());
+			RelativeStartTimes.SetNumUninitialized(AllLinearKeys.Num());
+
+			// Make the times relative to the initial drag position
+			GetKeyTimes(AllLinearKeys, AbsoluteStartTimes);
+			for (int32 Index = 0; Index < AbsoluteStartTimes.Num() ; ++Index)
+			{
+				RelativeStartTimes[Index] = (AbsoluteStartTimes[Index] - RelativeStartTime);
+			}
+
+			Transaction = MakeUnique<FScopedTransaction>(LOCTEXT("DragKeyBarTransaction", "Move Keys"));
+
+			for (const FSequencerSelectedKey& Key : AllLinearKeys)
+			{
+				if (!ModifiedSections.Contains(Key.Section))
+				{
+					ModifiedSections.Add(Key.Section);
+				}
+			}
+		}
+		void OnDrag(const FPointerEvent& MouseEvent, FVector2D LocalMousePos, const FVirtualTrackArea& VirtualTrackArea) override
+		{
+			UE::MovieScene::FScopedSignedObjectModifyDefer Defer(true);
+
+			const FFrameTime NewTime = VirtualTrackArea.PixelToFrame(LocalMousePos.X);
+
+			for (UMovieSceneSection* Section : ModifiedSections)
+			{
+				Section->Modify();
+			}
+
+			// Set the position of leading and trailing keys
+			TArray<FFrameTime> NewKeyTimes;
+			NewKeyTimes.SetNumUninitialized(RelativeStartTimes.Num());
+
+			for (int32 Index = 0; Index < RelativeStartTimes.Num(); ++Index)
+			{
+				NewKeyTimes[Index] = NewTime + RelativeStartTimes[Index];
+			}
+
+			const float PixelSnapWidth = 20.f;
+			const int32 SnapThreshold = VirtualTrackArea.PixelDeltaToFrame(PixelSnapWidth).FloorToFrame().Value;
+			if (TOptional<FSequencerSnapField::FSnapResult> SnappedTime = SnapField.Snap(NewKeyTimes, SnapThreshold))
+			{
+				FFrameTime SnappedDiff = SnappedTime->SnappedTime - SnappedTime->OriginalTime;
+				for (int32 Index = 0; Index < RelativeStartTimes.Num(); ++Index)
+				{
+					NewKeyTimes[Index] = NewKeyTimes[Index] + SnappedDiff;
+				}
+			}
+
+			TArray<FFrameNumber> NewKeyFrames;
+			NewKeyFrames.SetNumUninitialized(NewKeyTimes.Num());
+			for (int32 Index = 0; Index < NewKeyTimes.Num(); ++Index)
+			{
+				NewKeyFrames[Index] = NewKeyTimes[Index].RoundToFrame();
+			}
+			SetKeyTimes(AllLinearKeys, NewKeyFrames);
+		}
+		void OnEndDrag( const FPointerEvent& MouseEvent, FVector2D LocalMousePos, const FVirtualTrackArea& VirtualTrackArea) override
+		{
+			Transaction.Reset();
+		}
+		FCursorReply GetCursor() const override
+		{
+			return FCursorReply::Cursor(EMouseCursor::ResizeLeftRight);
+		}
+		int32 OnPaint(const FGeometry& AllottedGeometry, const FSlateRect& MyCullingRect, FSlateWindowElementList& OutDrawElements, int32 LayerId) const override
+		{
+			return LayerId;
+		}
+		bool IsKeyApplicable(FKeyHandle KeyHandle, const UE::Sequencer::FViewModelPtr& Owner) const override
+		{
+			using namespace UE::Sequencer;
+
+			TSharedPtr<FChannelModel> Channel = Owner.ImplicitCast();
+			return Channel && !AllKeys.Contains(FSequencerSelectedKey(*Channel->GetSection(), Channel, KeyHandle));
+		}
+		bool AreSectionBoundsApplicable(UMovieSceneSection* Section) const override
+		{
+			return true;
+		}
+	};
+
+	return MakeShared<FKeyBarDrag>(SharedThis(this), WeakSequencer.Pin());
+}
+
 UMovieSceneSection* FSectionHotspotBase::GetSection() const
 {
 	if (TSharedPtr<FSectionModel> Model = WeakSectionModel.Pin())
diff --git a/Engine/Source/Editor/Sequencer/Private/SequencerHotspots.h b/Engine/Source/Editor/Sequencer/Private/SequencerHotspots.h
--- a/Engine/Source/Editor/Sequencer/Private/SequencerHotspots.h
+++ b/Engine/Source/Editor/Sequencer/Private/SequencerHotspots.h
@@ -79,6 +79,30 @@ struct FKeyHotspot
 	TWeakPtr<FSequencer> WeakSequencer;
 };
 
+/** A hotspot representing a key bar */
+struct FKeyBarHotspot
+	: ITrackAreaHotspot
+	, TSharedFromThis<FKeyBarHotspot>
+{
+	UE_SEQUENCER_DECLARE_CASTABLE(FKeyBarHotspot, ITrackAreaHotspot);
+
+	FKeyBarHotspot(TArray<FSequencerSelectedKey>&& InLeadingKeys, TArray<FSequencerSelectedKey>&& InTrailingKeys, TWeakPtr<FSequencer> InWeakSequencer)
+		: LeadingKeys(MoveTemp(InLeadingKeys))
+		, TrailingKeys(MoveTemp(InTrailingKeys))
+		, WeakSequencer(InWeakSequencer)
+	{ }
+
+	virtual void UpdateOnHover(FTrackAreaViewModel& InTrackArea) const override;
+	virtual TOptional<FFrameNumber> GetTime() const override;
+	virtual TSharedPtr<ISequencerEditToolDragOperation> InitiateDrag(const FPointerEvent& MouseEvent) override;
+	virtual FCursorReply GetCursor() const override;
+
+	/** The keys that are part of this hotspot */
+	TArray<FSequencerSelectedKey> LeadingKeys;
+	TArray<FSequencerSelectedKey> TrailingKeys;
+	TWeakPtr<FSequencer> WeakSequencer;
+};
+
 /** A hotspot representing a section */
 struct FSectionHotspotBase
 	: ITrackAreaHotspot, IMouseHandlerHotspot
diff --git a/Engine/Source/Editor/Sequencer/Private/SequencerKeyRenderer.cpp b/Engine/Source/Editor/Sequencer/Private/SequencerKeyRenderer.cpp
deleted file mode 100644
diff --git a/Engine/Source/Editor/Sequencer/Private/SequencerKeyRenderer.h b/Engine/Source/Editor/Sequencer/Private/SequencerKeyRenderer.h
deleted file mode 100644
diff --git a/Engine/Source/Editor/Sequencer/Private/SequencerSettings.cpp b/Engine/Source/Editor/Sequencer/Private/SequencerSettings.cpp
--- a/Engine/Source/Editor/Sequencer/Private/SequencerSettings.cpp
+++ b/Engine/Source/Editor/Sequencer/Private/SequencerSettings.cpp
@@ -45,6 +45,7 @@ USequencerSettings::USequencerSettings( const FObjectInitializer& ObjectInitiali
 	ZeroPadFrames = 0;
 	JumpFrameIncrement = FFrameNumber(5);
 	bShowLayerBars = true;
+	bShowKeyBars = true;
 	bInfiniteKeyAreas = false;
 	bShowChannelColors = false;
 	bShowStatusBar = true;
@@ -552,6 +553,19 @@ void USequencerSettings::SetShowLayerBars(bool InbShowLayerBars)
 	}
 }
 
+bool USequencerSettings::GetShowKeyBars() const
+{
+	return bShowKeyBars;
+}
+
+void USequencerSettings::SetShowKeyBars(bool InbShowKeyBars)
+{
+	if (bShowKeyBars != InbShowKeyBars)
+	{
+		bShowKeyBars = InbShowKeyBars;
+		SaveConfig();
+	}
+}
 
 bool USequencerSettings::GetInfiniteKeyAreas() const
 {
@@ -638,7 +652,7 @@ void USequencerSettings::RemoveKeyAreaCurveExtents(const FString& ChannelName)
 	SaveConfig();
 }
 
-void USequencerSettings::SetKeyAreaCurveExtents(const FString& ChannelName, float InMin, float InMax)
+void USequencerSettings::SetKeyAreaCurveExtents(const FString& ChannelName, double InMin, double InMax)
 {
 	RemoveKeyAreaCurveExtents(ChannelName);
 
@@ -649,7 +663,7 @@ void USequencerSettings::SetKeyAreaCurveExtents(const FString& ChannelName, floa
 	SaveConfig();
 }
 
-void USequencerSettings::GetKeyAreaCurveExtents(const FString& ChannelName, float& InMin, float& InMax) const
+void USequencerSettings::GetKeyAreaCurveExtents(const FString& ChannelName, double& OutMin, double& OutMax) const
 {
 	TArray<FString> ChannelsArray;
 	KeyAreaCurveExtents.ParseIntoArray(ChannelsArray, TEXT(":"));
@@ -661,8 +675,8 @@ void USequencerSettings::GetKeyAreaCurveExtents(const FString& ChannelName, floa
 
 		if (ExtentsArray.Num() == 3 && ExtentsArray[0] == ChannelName)
 		{
-			InMin = FCString::Atof(*ExtentsArray[1]);
-			InMax = FCString::Atof(*ExtentsArray[2]);
+			OutMin = FCString::Atod(*ExtentsArray[1]);
+			OutMax = FCString::Atod(*ExtentsArray[2]);
 			return;
 		}
 	}
diff --git a/Engine/Source/Editor/Sequencer/Private/Tools/DragOperation_Stretch.cpp b/Engine/Source/Editor/Sequencer/Private/Tools/DragOperation_Stretch.cpp
--- a/Engine/Source/Editor/Sequencer/Private/Tools/DragOperation_Stretch.cpp
+++ b/Engine/Source/Editor/Sequencer/Private/Tools/DragOperation_Stretch.cpp
@@ -122,10 +122,10 @@ void FEditToolDragOperation_Stretch::OnDrag(const FPointerEvent& MouseEvent, FVe
 		const float SnapThresholdPx = VirtualTrackArea.PixelToSeconds(PixelSnapWidth) - VirtualTrackArea.PixelToSeconds(0.f);
 		const int32 SnapThreshold = (SnapThresholdPx * Sequencer->GetFocusedTickResolution()).FloorToFrame().Value;
 
-		TOptional<FFrameNumber> SnappedResult = SnapField->Snap(CurrentTime.RoundToFrame(), SnapThreshold);
+		TOptional<FSequencerSnapField::FSnapResult> SnappedResult = SnapField->Snap(CurrentTime.RoundToFrame(), SnapThreshold);
 		if (SnappedResult.IsSet())
 		{
-			CurrentTime = SnappedResult.GetValue();
+			CurrentTime = SnappedResult->SnappedTime;
 		}
 	}
 
@@ -163,10 +163,10 @@ void FEditToolDragOperation_Stretch::OnEndDrag( const FPointerEvent& MouseEvent,
 		const float SnapThresholdPx = VirtualTrackArea.PixelToSeconds(PixelSnapWidth) - VirtualTrackArea.PixelToSeconds(0.f);
 		const int32 SnapThreshold = (SnapThresholdPx * Sequencer->GetFocusedTickResolution()).FloorToFrame().Value;
 
-		TOptional<FFrameNumber> SnappedResult = SnapField->Snap(CurrentTime.RoundToFrame(), SnapThreshold);
+		TOptional<FSequencerSnapField::FSnapResult> SnappedResult = SnapField->Snap(CurrentTime.RoundToFrame(), SnapThreshold);
 		if (SnappedResult.IsSet())
 		{
-			CurrentTime = SnappedResult.GetValue();
+			CurrentTime = SnappedResult->SnappedTime;
 		}
 	}
 
diff --git a/Engine/Source/Editor/Sequencer/Private/Tools/EditToolDragOperations.cpp b/Engine/Source/Editor/Sequencer/Private/Tools/EditToolDragOperations.cpp
--- a/Engine/Source/Editor/Sequencer/Private/Tools/EditToolDragOperations.cpp
+++ b/Engine/Source/Editor/Sequencer/Private/Tools/EditToolDragOperations.cpp
@@ -52,24 +52,24 @@ protected:
 	TSet<UMovieSceneSection*> SectionsToExclude;
 };
 
-TOptional<FSequencerSnapField::FSnapResult> SnapToInterval(const TArray<FFrameNumber>& InTimes, int32 FrameThreshold, FFrameRate Resolution, FFrameRate DisplayRate, ESequencerScrubberStyle ScrubStyle)
+TOptional<FSequencerSnapField::FSnapResult> SnapToInterval(const TArray<FFrameTime>& InTimes, int32 FrameThreshold, FFrameRate Resolution, FFrameRate DisplayRate, ESequencerScrubberStyle ScrubStyle)
 {
 	TOptional<FSequencerSnapField::FSnapResult> Result;
 
-	FFrameNumber SnapAmount(0);
-	for (FFrameNumber Time : InTimes)
+	FFrameTime SnapAmount(0);
+	for (FFrameTime Time : InTimes)
 	{
 		// Convert from resolution to DisplayRate, round to frame, then back again. We floor to frames when using the frame block scrubber, and round using the vanilla scrubber
 		FFrameTime   DisplayTime      = FFrameRate::TransformTime(Time, Resolution, DisplayRate);
 		FFrameNumber PlayIntervalTime = ScrubStyle == ESequencerScrubberStyle::FrameBlock ? DisplayTime.FloorToFrame() : DisplayTime.RoundToFrame();
-		FFrameNumber IntervalSnap     = FFrameRate::TransformTime(PlayIntervalTime, DisplayRate, Resolution  ).FloorToFrame();
+		FFrameNumber IntervalSnap     = FFrameRate::TransformTime(PlayIntervalTime, DisplayRate, Resolution).FloorToFrame();
 
-		FFrameNumber ThisSnapAmount   = IntervalSnap - Time;
+		FFrameTime ThisSnapAmount   = IntervalSnap - Time;
 		if (FMath::Abs(ThisSnapAmount) <= FrameThreshold)
 		{
 			if (!Result.IsSet() || FMath::Abs(ThisSnapAmount) < SnapAmount)
 			{
-				Result = FSequencerSnapField::FSnapResult{Time, IntervalSnap};
+				Result = FSequencerSnapField::FSnapResult{Time, IntervalSnap, 1.f};
 				SnapAmount = ThisSnapAmount;
 			}
 		}
@@ -79,7 +79,7 @@ TOptional<FSequencerSnapField::FSnapResult> SnapToInterval(const TArray<FFrameNu
 }
 
 /** How many pixels near the mouse has to be before snapping occurs */
-const float PixelSnapWidth = 10.f;
+const float PixelSnapWidth = 20.f;
 
 
 
@@ -277,7 +277,7 @@ void FResizeSection::OnDrag(const FPointerEvent& MouseEvent, FVector2D LocalMous
 	// Snapping
 	if ( Settings->GetIsSnapEnabled() )
 	{
-		TArray<FFrameNumber> SectionTimes;
+		TArray<FFrameTime> SectionTimes;
 		for (UMovieSceneSection* Section : Sections)
 		{
 			SectionTimes.Add(SectionInitTimes[Section] + DeltaTime);
@@ -302,7 +302,7 @@ void FResizeSection::OnDrag(const FPointerEvent& MouseEvent, FVector2D LocalMous
 		if (SnappedTime.IsSet())
 		{
 			// Add the snapped amount onto the delta
-			DeltaTime += SnappedTime->Snapped - SnappedTime->Original;
+			DeltaTime += (SnappedTime->SnappedTime - SnappedTime->OriginalTime).RoundToFrame();
 		}
 	}
 	
@@ -590,7 +590,7 @@ void FManipulateSectionEasing::OnDrag(const FPointerEvent& MouseEvent, FVector2D
 	// Snapping
 	if (Settings->GetIsSnapEnabled())
 	{
-		TArray<FFrameNumber> SnapTimes;
+		TArray<FFrameTime> SnapTimes;
 		if (bEaseIn)
 		{
 			FFrameNumber DesiredTime = (DeltaTime + Section->GetInclusiveStartFrame() + InitValue.Get(0)).RoundToFrame();
@@ -621,7 +621,7 @@ void FManipulateSectionEasing::OnDrag(const FPointerEvent& MouseEvent, FVector2D
 		if (SnappedTime.IsSet())
 		{
 			// Add the snapped amount onto the delta
-			DeltaTime += SnappedTime->Snapped - SnappedTime->Original;
+			DeltaTime += SnappedTime->SnappedTime - SnappedTime->OriginalTime;
 		}
 	}
 
@@ -778,12 +778,11 @@ void FMoveKeysAndSections::OnDrag(const FPointerEvent& MouseEvent, FVector2D Loc
 	// Calculate snapping first which modifies our MouseTime to reflect where it would have to be for the closest snap to work.
 	if (Settings->GetIsSnapEnabled())
 	{
-		float SnapThresholdPx = VirtualTrackArea.PixelToSeconds(PixelSnapWidth) - VirtualTrackArea.PixelToSeconds(0.f);
-		int32 SnapThreshold = (SnapThresholdPx * TickResolution).FloorToFrame().Value;
+		FFrameTime SnapThreshold = VirtualTrackArea.PixelDeltaToFrame(PixelSnapWidth);
 
 		// The edge of each bounded section as well as each individual key is a valid marker to try and snap to intervals/sections/etc.
 		// We take our stored offsets and add them to our current time to figure out where on the timeline the are currently.
-		TArray<FFrameNumber> ValidSnapMarkers;
+		TArray<FFrameTime> ValidSnapMarkers;
 
 		// If they have both keys and settings selected then we snap to the interval if either one of them is enabled, otherwise respect the individual setting.
 		const bool bSnapToInterval = (KeysAsArray.Num() > 0 && Settings->GetSnapKeyTimesToInterval()) || (Sections.Num() > 0 && Settings->GetSnapSectionTimesToInterval());
@@ -793,7 +792,7 @@ void FMoveKeysAndSections::OnDrag(const FPointerEvent& MouseEvent, FVector2D Loc
 		ValidSnapMarkers.SetNumUninitialized(RelativeSnapOffsets.Num());
 		for (int32 Index = 0; Index < RelativeSnapOffsets.Num(); ++Index)
 		{
-			ValidSnapMarkers[Index] = (RelativeSnapOffsets[Index] + MouseTime).FloorToFrame();
+			ValidSnapMarkers[Index] = (RelativeSnapOffsets[Index] + MouseTime);
 		}
 
 		// Now we'll try and snap all of these points to the closest valid snap marker (which may be a section or interval)
@@ -817,7 +816,7 @@ void FMoveKeysAndSections::OnDrag(const FPointerEvent& MouseEvent, FVector2D Loc
 		if (SnappedTime.IsSet())
 		{
 			// Add the snapped amount onto the mouse time so the resulting delta brings us in alignment.
-			MouseTime += (SnappedTime->Snapped - SnappedTime->Original);
+			MouseTime += (SnappedTime->SnappedTime - SnappedTime->OriginalTime);
 		}
 	}
 
diff --git a/Engine/Source/Editor/Sequencer/Private/Tools/SequencerEditTool_Movement.cpp b/Engine/Source/Editor/Sequencer/Private/Tools/SequencerEditTool_Movement.cpp
--- a/Engine/Source/Editor/Sequencer/Private/Tools/SequencerEditTool_Movement.cpp
+++ b/Engine/Source/Editor/Sequencer/Private/Tools/SequencerEditTool_Movement.cpp
@@ -500,18 +500,9 @@ FCursorReply FSequencerEditTool_Movement::OnCursorQuery(const FGeometry& MyGeome
 {
 	using namespace UE::Sequencer;
 
-	TSharedPtr<FTrackAreaViewModel> TrackArea = Sequencer.GetViewModel()->GetTrackArea();
-	TSharedPtr<ITrackAreaHotspot> Hotspot = DelayedDrag.IsSet()
-		? DelayedDrag->Hotspot
-		: TrackArea->GetHotspot();
-
-	if (Hotspot.IsValid())
+	if (DragOperation)
 	{
-		FCursorReply Reply = Hotspot->GetCursor();
-		if (Reply.IsEventHandled())
-		{
-			return Reply;
-		}
+		return DragOperation->GetCursor();
 	}
 
 	return FCursorReply::Cursor(EMouseCursor::CardinalCross);
diff --git a/Engine/Source/Editor/Sequencer/Private/Tools/SequencerEditTool_Selection.cpp b/Engine/Source/Editor/Sequencer/Private/Tools/SequencerEditTool_Selection.cpp
--- a/Engine/Source/Editor/Sequencer/Private/Tools/SequencerEditTool_Selection.cpp
+++ b/Engine/Source/Editor/Sequencer/Private/Tools/SequencerEditTool_Selection.cpp
@@ -319,6 +319,11 @@ FSequencerEditTool_Selection::FSequencerEditTool_Selection(FSequencer& InSequenc
 
 FCursorReply FSequencerEditTool_Selection::OnCursorQuery(const FGeometry& MyGeometry, const FPointerEvent& CursorEvent) const
 {
+	if (DragOperation)
+	{
+		return DragOperation->GetCursor();
+	}
+
 	return FCursorReply::Cursor(EMouseCursor::Crosshairs);
 }
 
diff --git a/Engine/Source/Editor/Sequencer/Private/Tools/SequencerSnapField.cpp b/Engine/Source/Editor/Sequencer/Private/Tools/SequencerSnapField.cpp
--- a/Engine/Source/Editor/Sequencer/Private/Tools/SequencerSnapField.cpp
+++ b/Engine/Source/Editor/Sequencer/Private/Tools/SequencerSnapField.cpp
@@ -53,6 +53,10 @@ FSequencerSnapField::FSequencerSnapField(const FSequencer& InSequencer, UE::Sequ
 
 void FSequencerSnapField::AddExplicitSnap(UE::Sequencer::FSnapPoint InSnap)
 {
+	if (InSnap.Weighting == 1.f && InSnap.Type != UE::Sequencer::FSnapPoint::Key)
+	{
+		InSnap.Weighting = 10.f;
+	}
 	SortedSnaps.Add(InSnap);
 }
 
@@ -133,6 +137,7 @@ void FSequencerSnapField::Finalize()
 	{
 		const FFrameNumber CurrentTime = SortedSnaps[Index].Time;
 
+		float FinalWeight = SortedSnaps[Index].Weighting;
 		int32 NumToMerge = 0;
 		for (int32 DuplIndex = Index + 1; DuplIndex < SortedSnaps.Num(); ++DuplIndex)
 		{
@@ -141,16 +146,18 @@ void FSequencerSnapField::Finalize()
 				break;
 			}
 			++NumToMerge;
+			FinalWeight += SortedSnaps[DuplIndex].Weighting;
 		}
 
 		if (NumToMerge)
 		{
+			SortedSnaps[Index].Weighting = FinalWeight;
 			SortedSnaps.RemoveAt(Index + 1, NumToMerge, false);
 		}
 	}
 }
 
-TOptional<FFrameNumber> FSequencerSnapField::Snap(FFrameNumber InTime, int32 Threshold) const
+TOptional<FSequencerSnapField::FSnapResult> FSequencerSnapField::Snap(const FFrameTime& InTime, const FFrameTime& Threshold) const
 {
 	int32 Min = 0;
 	int32 Max = SortedSnaps.Num();
@@ -160,70 +167,81 @@ TOptional<FFrameNumber> FSequencerSnapField::Snap(FFrameNumber InTime, int32 Thr
 	{
 		int32 SearchIndex = Min + (Max - Min) / 2;
 
-		FFrameNumber ProspectiveSnapPos = SortedSnaps[SearchIndex].Time;
-		if (ProspectiveSnapPos > InTime + Threshold)
+		UE::Sequencer::FSnapPoint ProspectiveSnap = SortedSnaps[SearchIndex];
+		if (ProspectiveSnap.Time > InTime + Threshold)
 		{
 			Max = SearchIndex;
 		}
-		else if (ProspectiveSnapPos < InTime - Threshold)
+		else if (ProspectiveSnap.Time < InTime - Threshold)
 		{
 			Min = SearchIndex + 1;
 		}
 		else
 		{
-			// Linearly search forwards and backwards to find the closest snap
+			// Linearly search forwards and backwards to find the closest or heaviest snap
 
-			FFrameNumber SnapDelta = ProspectiveSnapPos - InTime;
+			float SnapWeight = 0.f;
+			FFrameTime SnapDelta = ProspectiveSnap.Time - InTime;
 
 			// Search forwards while we're in the threshold
-			for (int32 FwdIndex = SearchIndex+1; FwdIndex < Max-1 && SortedSnaps[FwdIndex].Time < InTime + Threshold; ++FwdIndex)
+			for (int32 FwdIndex = SearchIndex; FwdIndex < Max && SortedSnaps[FwdIndex].Time < InTime + Threshold; ++FwdIndex)
 			{
-				FFrameNumber ThisSnapDelta = InTime - SortedSnaps[FwdIndex].Time;
-				if (FMath::Abs(ThisSnapDelta) < FMath::Abs(SnapDelta))
+				FFrameTime ThisSnapDelta = InTime - SortedSnaps[FwdIndex].Time;
+				float ThisSnapWeight = SortedSnaps[FwdIndex].Weighting;
+				if (ThisSnapWeight > SnapWeight || (ThisSnapWeight == SnapWeight && FMath::Abs(ThisSnapDelta) < FMath::Abs(SnapDelta)))
 				{
 					SnapDelta = ThisSnapDelta;
-					ProspectiveSnapPos = SortedSnaps[FwdIndex].Time;
+					SnapWeight = ThisSnapWeight;
+					ProspectiveSnap = SortedSnaps[FwdIndex];
 				}
 			}
 
 			// Search backwards while we're in the threshold
-			for (int32 BckIndex = SearchIndex-1; BckIndex >= Min && SortedSnaps[BckIndex].Time > InTime + Threshold; --BckIndex)
+			for (int32 BckIndex = SearchIndex-1; BckIndex >= Min && SortedSnaps[BckIndex].Time > InTime - Threshold; --BckIndex)
 			{
-				FFrameNumber ThisSnapDelta = InTime - SortedSnaps[BckIndex].Time;
-				if (FMath::Abs(ThisSnapDelta) < FMath::Abs(SnapDelta))
+				FFrameTime ThisSnapDelta = InTime - SortedSnaps[BckIndex].Time;
+				float ThisSnapWeight = SortedSnaps[BckIndex].Weighting;
+				if (ThisSnapWeight > SnapWeight || (ThisSnapWeight == SnapWeight && FMath::Abs(ThisSnapDelta) < FMath::Abs(SnapDelta)))
 				{
 					SnapDelta = ThisSnapDelta;
-					ProspectiveSnapPos = SortedSnaps[BckIndex].Time;
+					SnapWeight = ThisSnapWeight;
+					ProspectiveSnap = SortedSnaps[BckIndex];
 				}
 			}
 
-			return ProspectiveSnapPos;
+			if (SnapWeight != 0.f)
+			{
+				FSnapResult Result = { InTime, ProspectiveSnap.Time, ProspectiveSnap.Weighting };
+				return Result;
+			}
+
+			break;
 		}
 	}
 
-	return TOptional<FFrameNumber>();
+	return TOptional<FSequencerSnapField::FSnapResult>();
 }
 
-TOptional<FSequencerSnapField::FSnapResult> FSequencerSnapField::Snap(const TArray<FFrameNumber>& InTimes, int32 Threshold) const
+TOptional<FSequencerSnapField::FSnapResult> FSequencerSnapField::Snap(const TArray<FFrameTime>& InTimes, const FFrameTime& Threshold) const
 {
 	TOptional<FSnapResult> ProspectiveSnap;
-	FFrameNumber SnapDelta(0);
 
-	for (FFrameNumber Time : InTimes)
+	int32 NumSnaps = 0;
+	float MaxSnapWeight = 0.f;
+	for (FFrameTime Time : InTimes)
 	{
-		TOptional<FFrameNumber> ThisSnap = Snap(Time, Threshold);
+		TOptional<FSnapResult> ThisSnap = Snap(Time, Threshold);
 		if (!ThisSnap.IsSet())
 		{
 			continue;
 		}
 
-		FFrameNumber ThisSnapDelta = ThisSnap.GetValue() - Time;
-		if (!ProspectiveSnap.IsSet() || FMath::Abs(ThisSnapDelta) < FMath::Abs(SnapDelta))
+		if (!ProspectiveSnap.IsSet() || ThisSnap->SnappedWeight > MaxSnapWeight)
 		{
-			ProspectiveSnap = FSnapResult{ Time, ThisSnap.GetValue() };
-			SnapDelta = ThisSnapDelta;
+			ProspectiveSnap = ThisSnap;
+			MaxSnapWeight = ThisSnap->SnappedWeight;
 		}
 	}
 
 	return ProspectiveSnap;
-}
+}
\ No newline at end of file
diff --git a/Engine/Source/Editor/Sequencer/Private/Tools/SequencerSnapField.h b/Engine/Source/Editor/Sequencer/Private/Tools/SequencerSnapField.h
--- a/Engine/Source/Editor/Sequencer/Private/Tools/SequencerSnapField.h
+++ b/Engine/Source/Editor/Sequencer/Private/Tools/SequencerSnapField.h
@@ -19,9 +19,11 @@ public:
 	struct FSnapResult
 	{
 		/** The time before it was snapped */
-		FFrameNumber Original;
+		FFrameTime OriginalTime;
 		/** The time after it was snapped */
-		FFrameNumber Snapped;
+		FFrameNumber SnappedTime;
+		/** The total weight of the time that was snapped to */
+		float SnappedWeight = 0.f;
 	};
 
 	FSequencerSnapField(){}
@@ -40,10 +42,10 @@ public:
 	void Finalize();
 
 	/** Snap the specified time to this field with the given threshold */
-	TOptional<FFrameNumber> Snap(FFrameNumber InTime, int32 Threshold) const;
+	TOptional<FSnapResult> Snap(const FFrameTime& InTime, const FFrameTime& Threshold) const;
 
 	/** Snap the specified times to this field with the given threshold. Will return the closest snap value of the entire intersection. */
-	TOptional<FSnapResult> Snap(const TArray<FFrameNumber>& InTimes, int32 Threshold) const;
+	TOptional<FSnapResult> Snap(const TArray<FFrameTime>& InTimes, const FFrameTime& Threshold) const;
 
 private:
 	/** Array of snap points, approximately grouped, and sorted in ascending order by time */
diff --git a/Engine/Source/Editor/Sequencer/Private/Widgets/SSequencerLayerBar.cpp b/Engine/Source/Editor/Sequencer/Private/Widgets/SSequencerLayerBar.cpp
--- a/Engine/Source/Editor/Sequencer/Private/Widgets/SSequencerLayerBar.cpp
+++ b/Engine/Source/Editor/Sequencer/Private/Widgets/SSequencerLayerBar.cpp
@@ -215,12 +215,14 @@ void SSequencerLayerBar::Construct(const FArguments& InArgs, TWeakPtr<FSequencer
 	ChildSlot
 	[
 		SNew(SHorizontalBox)
+		.Visibility(this, &SSequencerLayerBar::GetHandleVisibility)
 
 		+ SHorizontalBox::Slot()
+		.Padding(FMargin(1.f, 2.f))
 		.AutoWidth()
 		[
 			SNew(SBox)
-			.WidthOverride(16.f)
+			.WidthOverride(8.f)
 			[
 				SNew(SSequencerLayerBarHandle, SharedThis(this), EStretchConstraint::AnchorToEnd)
 				.BackgroundBrush(FAppStyle::Get().GetBrush("Sequencer.LayerBar.HandleLeft"))
@@ -233,10 +235,11 @@ void SSequencerLayerBar::Construct(const FArguments& InArgs, TWeakPtr<FSequencer
 		]
 
 		+ SHorizontalBox::Slot()
+		.Padding(FMargin(1.f, 2.f))
 		.AutoWidth()
 		[
 			SNew(SBox)
-			.WidthOverride(16.f)
+			.WidthOverride(8.f)
 			[
 				SNew(SSequencerLayerBarHandle, SharedThis(this), EStretchConstraint::AnchorToStart)
 				.BackgroundBrush(FAppStyle::Get().GetBrush("Sequencer.LayerBar.HandleRight"))
@@ -273,6 +276,11 @@ FTrackLaneScreenAlignment SSequencerLayerBar::GetAlignment(const FTimeToPixel& I
 	return FTrackLaneScreenAlignment{};
 }
 
+EVisibility SSequencerLayerBar::GetHandleVisibility() const
+{
+	return IsHovered() ? EVisibility::Visible : EVisibility::Collapsed;
+}
+
 int32 SSequencerLayerBar::OnPaint(const FPaintArgs& Args, const FGeometry& AllottedGeometry, const FSlateRect& MyCullingRect, FSlateWindowElementList& OutDrawElements, int32 LayerId, const FWidgetStyle& InWidgetStyle, bool bParentEnabled) const
 {
 	const ESlateDrawEffect DrawEffects = bParentEnabled
@@ -280,8 +288,7 @@ int32 SSequencerLayerBar::OnPaint(const FPaintArgs& Args, const FGeometry& Allot
 		: ESlateDrawEffect::DisabledEffect;
 
 	static const FSlateBrush* LayerBarBackgroundBrush = FAppStyle::Get().GetBrush("Sequencer.LayerBar.Background");
-	static const FSlateBrush* LayerBarBorderBrush     = FAppStyle::Get().GetBrush("Sequencer.LayerBar.Border");
-	static const FSlateBrush* SelectedSectionOverlay  = FAppStyle::Get().GetBrush("Sequencer.Section.SelectedSectionOverlay");
+	static const FSlateBrush* SelectedSectionOverlay  = FAppStyle::Get().GetBrush("Sequencer.Section.CollapsedSelectedSectionOverlay");
 
 	// Draw the background
 	FSlateDrawElement::MakeBox(
@@ -291,15 +298,6 @@ int32 SSequencerLayerBar::OnPaint(const FPaintArgs& Args, const FGeometry& Allot
 		LayerBarBackgroundBrush,
 		DrawEffects
 	);
-
-	// Draw the border
-	FSlateDrawElement::MakeBox(
-		OutDrawElements,
-		LayerId,
-		AllottedGeometry.ToPaintGeometry(),
-		LayerBarBorderBrush,
-		DrawEffects
-	);
 	++LayerId;
 
 	TSharedPtr<FSequencer>		Sequencer		 = GetSequencer();
diff --git a/Engine/Source/Editor/Sequencer/Private/Widgets/SSequencerLayerBar.h b/Engine/Source/Editor/Sequencer/Private/Widgets/SSequencerLayerBar.h
--- a/Engine/Source/Editor/Sequencer/Private/Widgets/SSequencerLayerBar.h
+++ b/Engine/Source/Editor/Sequencer/Private/Widgets/SSequencerLayerBar.h
@@ -5,11 +5,8 @@
 #include "Math/Range.h"
 #include "Misc/FrameNumber.h"
 
-#include "Layout/Children.h"
-#include "Layout/Visibility.h"
-
+#include "Widgets/SCompoundWidget.h"
 #include "Widgets/DeclarativeSyntaxSupport.h"
-#include "Widgets/SPanel.h"
 
 #include "MVVM/Extensions/ITrackLaneExtension.h"
 #include "MVVM/Extensions/ViewModelExtensionCollection.h"
@@ -58,6 +55,8 @@ private:
 	void OnMouseEnter(const FGeometry& MyGeometry, const FPointerEvent& MouseEvent) override;
 	void OnMouseLeave(const FPointerEvent& MouseEvent) override;
 
+	EVisibility GetHandleVisibility() const;
+
 private:
 
 	TWeakPtr<FLayerBarModel> WeakLayerBar;
diff --git a/Engine/Source/Editor/Sequencer/Public/IKeyArea.h b/Engine/Source/Editor/Sequencer/Public/IKeyArea.h
--- a/Engine/Source/Editor/Sequencer/Public/IKeyArea.h
+++ b/Engine/Source/Editor/Sequencer/Public/IKeyArea.h
@@ -132,6 +132,13 @@ public:
 	 */
 	TSharedPtr<ISequencerSection> GetSectionInterface() const;
 
+	/**
+	 * Check whether we should show this key area's curve on its channel UI or not
+	 *
+	 * @return True to show the curve, false otherwise
+	 */
+	bool ShouldShowCurve() const;
+
 public:
 
 	/**
diff --git a/Engine/Source/Editor/Sequencer/Public/ISectionLayoutBuilder.h b/Engine/Source/Editor/Sequencer/Public/ISectionLayoutBuilder.h
--- a/Engine/Source/Editor/Sequencer/Public/ISectionLayoutBuilder.h
+++ b/Engine/Source/Editor/Sequencer/Public/ISectionLayoutBuilder.h
@@ -3,11 +3,18 @@
 #pragma once
 
 #include "CoreMinimal.h"
+#include "Templates/Function.h"
 
 class IKeyArea;
 
 struct FMovieSceneChannelHandle;
 
+namespace UE::Sequencer
+{
+	class FCategoryModel;
+	class FChannelModel;
+}
+
 /**
  * Builds an inner layout for a section
  */
@@ -22,20 +29,20 @@ public:
 	 * @param CategoryName	The name of the category
 	 * @param DisplayLabel	The localized display label for the category
 	 */
-	virtual void PushCategory( FName CategoryName, const FText& DisplayLabel ) = 0;
-	
+	virtual void PushCategory( FName CategoryName, const FText& DisplayLabel, TFunction<TSharedPtr<UE::Sequencer::FCategoryModel>(FName, const FText&)> OptionalFactory) = 0;
+
 	/**
 	 * Sets the section as a key area itself
 	 * @param Channel		The channel that is to be assigned as the top level channel for this section
 	 */
-	virtual void SetTopLevelChannel( const FMovieSceneChannelHandle& Channel ) = 0;
+	virtual void SetTopLevelChannel( const FMovieSceneChannelHandle& Channel, TFunction<TSharedPtr<UE::Sequencer::FChannelModel>(FName, const FMovieSceneChannelHandle&)> OptionalFactory ) = 0;
 
 	/**
 	 * Adds a channel onto the layout. If a category is pushed, the key area will appear as a child of the current category
 	 *
 	 * @param Channel		A handle to the channel to be added to the layout
 	 */
-	virtual void AddChannel( const FMovieSceneChannelHandle& Channel ) = 0;
+	virtual void AddChannel( const FMovieSceneChannelHandle& Channel, TFunction<TSharedPtr<UE::Sequencer::FChannelModel>(FName, const FMovieSceneChannelHandle&)> OptionalFactory ) = 0;
 
 	/**
 	 * Pops a category off the stack
diff --git a/Engine/Source/Editor/Sequencer/Public/ISequencerChannelInterface.h b/Engine/Source/Editor/Sequencer/Public/ISequencerChannelInterface.h
--- a/Engine/Source/Editor/Sequencer/Public/ISequencerChannelInterface.h
+++ b/Engine/Source/Editor/Sequencer/Public/ISequencerChannelInterface.h
@@ -145,6 +145,15 @@ struct ISequencerChannelInterface
 	 */
 	virtual void DrawKeys_Raw(FMovieSceneChannel* Channel, TArrayView<const FKeyHandle> InKeyHandles, const UMovieSceneSection* InOwner, TArrayView<FKeyDrawParams> OutKeyDrawParams) const = 0;
 
+	/**
+	 * Whether this channel should draw a curve on its editor UI
+	 *
+	 * @param Channel               The channel to query
+	 * @param InSection             The section that owns the channel
+	 * @return true to show the curve on the UI, false otherwise
+	 */
+	virtual bool ShouldShowCurve_Raw(const FMovieSceneChannel* Channel, UMovieSceneSection* InSection) const = 0;
+
 	/**
 	 * Whether this channel supports curve models
 	 */
diff --git a/Engine/Source/Editor/Sequencer/Public/ISequencerSection.h b/Engine/Source/Editor/Sequencer/Public/ISequencerSection.h
--- a/Engine/Source/Editor/Sequencer/Public/ISequencerSection.h
+++ b/Engine/Source/Editor/Sequencer/Public/ISequencerSection.h
@@ -10,6 +10,7 @@
 #include "Widgets/SNullWidget.h"
 #include "Layout/Margin.h"
 #include "MovieSceneSection.h"
+#include "Channels/MovieSceneChannelHandle.h"
 
 class FMenuBuilder;
 class FSequencerSectionPainter;
@@ -18,6 +19,13 @@ class ISequencer;
 class ISequencerSection;
 class ISequencerTrackEditor;
 struct FSlateBrush;
+struct FMovieSceneChannelMetaData;
+
+namespace UE::Sequencer
+{
+	class FCategoryModel;
+	class FChannelModel;
+}
 
 /** Enumerates which edge is being resized */
 UENUM()
@@ -35,7 +43,7 @@ namespace SequencerSectionConstants
 	/** The size of each key */
 	const FVector2D KeySize(12.0f, 12.0f);
 
-	const float DefaultSectionGripSize = 7.0f;
+	const float DefaultSectionGripSize = 8.0f;
 
 	const float DefaultSectionHeight = 15.f;
 
@@ -67,6 +75,16 @@ struct FSequencerSectionPropertyDetailsViewCustomizationParams
 class ISequencerSection
 {
 public:
+	/** Structure used during key area creation to group channels by their group name */
+	struct FChannelData
+	{
+		/** Handle to the channel */
+		FMovieSceneChannelHandle Channel;
+
+		/** The channel's editor meta data */
+		const FMovieSceneChannelMetaData& MetaData;
+	};
+
 	virtual ~ISequencerSection(){}
 	/**
 	 * The MovieSceneSection data being visualized
@@ -138,6 +156,16 @@ public:
 	 */	
 	SEQUENCER_API virtual void GenerateSectionLayout( class ISectionLayoutBuilder& LayoutBuilder );
 
+	/**
+	 * Create a custom category model
+	 */
+	virtual TSharedPtr<UE::Sequencer::FCategoryModel> ConstructCategoryModel(FName InCategoryName, const FText& InDisplayText, TArrayView<const FChannelData> Channels) const { return nullptr; }
+
+	/**
+	 * Create a custom channel model
+	 */
+	virtual TSharedPtr<UE::Sequencer::FChannelModel> ConstructChannelModel(FName InCategoryName, const FMovieSceneChannelHandle& InChannelHandle) const { return nullptr; }
+
 	/**
 	 * @return The height of the section
 	 */
diff --git a/Engine/Source/Editor/Sequencer/Public/KeyDrawParams.h b/Engine/Source/Editor/Sequencer/Public/KeyDrawParams.h
deleted file mode 100644
diff --git a/Engine/Source/Editor/Sequencer/Public/MVVM/Extensions/ISnappableExtension.h b/Engine/Source/Editor/Sequencer/Public/MVVM/Extensions/ISnappableExtension.h
--- a/Engine/Source/Editor/Sequencer/Public/MVVM/Extensions/ISnappableExtension.h
+++ b/Engine/Source/Editor/Sequencer/Public/MVVM/Extensions/ISnappableExtension.h
@@ -21,11 +21,31 @@ struct SEQUENCER_API FSnapPoint
 {
 	enum ESnapType { Key, SectionBounds, CustomSection, PlaybackRange, CurrentTime, InOutRange, Mark };
 
-	/** The type of snap */
-	ESnapType Type;
+	FSnapPoint(ESnapType InType, FFrameNumber InTime, float InWeighting = -1.f)
+		: Time(InTime)
+		, Weighting(InWeighting)
+		, Type(InType)
+	{
+		if (InWeighting == -1.f)
+		{
+			if (Type != UE::Sequencer::FSnapPoint::Key)
+			{
+				Weighting = 10.f;
+			}
+			else
+			{
+				Weighting = 1.f;
+			}
+		}
+	}
 
 	/** The time of the snap */
 	FFrameNumber Time;
+
+	float Weighting = 1.f;
+
+	/** The type of snap */
+	ESnapType Type;
 };
 
 class SEQUENCER_API ISnapField
diff --git a/Engine/Source/Editor/Sequencer/Public/MVVM/ViewModels/CategoryModel.h b/Engine/Source/Editor/Sequencer/Public/MVVM/ViewModels/CategoryModel.h
--- a/Engine/Source/Editor/Sequencer/Public/MVVM/ViewModels/CategoryModel.h
+++ b/Engine/Source/Editor/Sequencer/Public/MVVM/ViewModels/CategoryModel.h
@@ -26,7 +26,7 @@ namespace Sequencer
  * For instance, this represents the "Location" category of a single transform section, which would
  * contain the X, Y, and Z translation channels.
  */
-class FCategoryModel
+class SEQUENCER_API FCategoryModel
 	: public FViewModel
 	, public FLinkedOutlinerExtension
 	, public FGeometryExtensionShim
@@ -51,6 +51,10 @@ public:
 	TSharedPtr<ITrackLaneWidget> CreateTrackLaneView(const FCreateTrackLaneViewParams& InParams) override;
 	FTrackLaneVirtualAlignment ArrangeVirtualTrackLaneView() const override;
 
+private:
+
+	FLinearColor GetKeyBarColor() const;
+
 private:
 
 	FViewModelListHead Children;
@@ -62,7 +66,7 @@ private:
  * For instance, this represents the "Location" category entry in the Sequence outliner, which would contain
  * the X, Y, and Z translation channels of all the corresponding sections in the track area.
  */
-class FCategoryGroupModel
+class SEQUENCER_API FCategoryGroupModel
 	: public FOutlinerItemModel
 	, public ITrackAreaExtension
 	, public ICompoundOutlinerExtension
diff --git a/Engine/Source/Editor/Sequencer/Public/MVVM/ViewModels/ChannelModel.h b/Engine/Source/Editor/Sequencer/Public/MVVM/ViewModels/ChannelModel.h
--- a/Engine/Source/Editor/Sequencer/Public/MVVM/ViewModels/ChannelModel.h
+++ b/Engine/Source/Editor/Sequencer/Public/MVVM/ViewModels/ChannelModel.h
@@ -9,6 +9,9 @@
 #include "MVVM/Extensions/LinkedOutlinerExtension.h"
 #include "MVVM/Extensions/ITrackAreaExtension.h"
 #include "MVVM/Extensions/IDeletableExtension.h"
+#include "MVVM/Extensions/ITrackLaneExtension.h"
+#include "MVVM/Extensions/IKeyExtension.h"
+
 #include "EventHandlers/ISignedObjectEventHandler.h"
 #include "Tree/ICurveEditorTreeItem.h"
 
@@ -39,10 +42,12 @@ class FChannelModel
 	, public FLinkedOutlinerExtension
 	, public FGeometryExtensionShim
 	, public FLinkedOutlinerComputedSizingShim
+	, public ITrackLaneExtension
+	, public IKeyExtension
 {
 public:
 
-	UE_SEQUENCER_DECLARE_CASTABLE(FChannelModel, FViewModel, FLinkedOutlinerExtension, IGeometryExtension);
+	UE_SEQUENCER_DECLARE_CASTABLE(FChannelModel, FViewModel, FLinkedOutlinerExtension, IGeometryExtension, ITrackLaneExtension, IKeyExtension);
 
 	FChannelModel(FName InChannelName, TWeakPtr<ISequencerSection> InSection, FMovieSceneChannelHandle InChannel);
 	~FChannelModel();
@@ -70,6 +75,21 @@ public:
 	/** Returns the desired sizing for the track area row */
 	FOutlinerSizing GetDesiredSizing() const;
 
+	/*~ ITrackLaneExtension */
+	TSharedPtr<ITrackLaneWidget> CreateTrackLaneView(const FCreateTrackLaneViewParams& InParams) override;
+	FTrackLaneVirtualAlignment ArrangeVirtualTrackLaneView() const override;
+
+	/*~ IKeyExtension */
+	bool UpdateCachedKeys(TSharedPtr<FCachedKeys>& OutCachedKeys) const override;
+	bool GetFixedExtents(double& OutFixedMin, double& OutFixedMax) const override;
+	int32 CustomPaint(const FGeometry& KeyGeometry, int32 LayerId) const override;
+	void DrawKeys(TArrayView<const FKeyHandle> InKeyHandles, TArrayView<FKeyDrawParams> OutKeyDrawParams) override;
+	TUniquePtr<FCurveModel> CreateCurveModel() override;
+
+private:
+
+	FLinearColor GetKeyBarColor() const;
+
 private:
 
 	TSharedPtr<IKeyArea> KeyArea;
@@ -80,7 +100,7 @@ private:
  * Model for the outliner entry associated with all sections' channels of a given common name.
  * For instance, this represents the "Location.X" entry in the Sequencer outliner.
  */
-class FChannelGroupModel
+class SEQUENCER_API FChannelGroupModel
 	: public FViewModel
 	, public ITrackAreaExtension
 {
@@ -141,7 +161,7 @@ protected:
  * Model for the outliner entry associated with all sections' channels of a given common name.
  * For instance, this represents the "Location.X" entry in the Sequencer outliner.
  */
-class FChannelGroupOutlinerModel
+class SEQUENCER_API FChannelGroupOutlinerModel
 	: public TOutlinerModelMixin<FChannelGroupModel>
 	, public ICompoundOutlinerExtension
 	, public IDeletableExtension
diff --git a/Engine/Source/Editor/Sequencer/Public/MVVM/ViewModels/OutlinerItemModel.h b/Engine/Source/Editor/Sequencer/Public/MVVM/ViewModels/OutlinerItemModel.h
--- a/Engine/Source/Editor/Sequencer/Public/MVVM/ViewModels/OutlinerItemModel.h
+++ b/Engine/Source/Editor/Sequencer/Public/MVVM/ViewModels/OutlinerItemModel.h
@@ -119,10 +119,6 @@ template<typename BaseType>
 class TOutlinerModelMixin : public BaseType, public FOutlinerItemModelMixin
 {
 public:
-	TOutlinerModelMixin()
-	{
-		this->RegisterChildList(&this->OutlinerChildList);
-	}
 
 	template<typename... ArgTypes>
 	TOutlinerModelMixin(ArgTypes&&... InArgs)
diff --git a/Engine/Source/Editor/Sequencer/Public/MVVM/ViewModels/SectionModel.h b/Engine/Source/Editor/Sequencer/Public/MVVM/ViewModels/SectionModel.h
--- a/Engine/Source/Editor/Sequencer/Public/MVVM/ViewModels/SectionModel.h
+++ b/Engine/Source/Editor/Sequencer/Public/MVVM/ViewModels/SectionModel.h
@@ -78,6 +78,9 @@ public:
 	/** Gets an array of sections whose easing bounds underlap the specified section */
 	TArray<FOverlappingSections> GetEasingSegments();
 
+	int32 GetPreRollFrames() const;
+	int32 GetPostRollFrames() const;
+
 public:
 
 	/*~ ILayerBarExtension Interface */
diff --git a/Engine/Source/Editor/Sequencer/Public/MVVM/ViewModels/TrackModel.h b/Engine/Source/Editor/Sequencer/Public/MVVM/ViewModels/TrackModel.h
--- a/Engine/Source/Editor/Sequencer/Public/MVVM/ViewModels/TrackModel.h
+++ b/Engine/Source/Editor/Sequencer/Public/MVVM/ViewModels/TrackModel.h
@@ -102,6 +102,7 @@ public:
 	/*~ ITrackAreaExtension */
 	FTrackAreaParameters GetTrackAreaParameters() const override;
 	FViewModelVariantIterator GetTrackAreaModelList() const override;
+	FViewModelVariantIterator GetTopLevelChildTrackAreaModels() const override;
 
 	/*~ ICurveEditorTreeItem */
 	void CreateCurveModels(TArray<TUniquePtr<FCurveModel>>& OutCurveModels) override;
diff --git a/Engine/Source/Editor/Sequencer/Public/MVVM/ViewModels/TrackRowModel.h b/Engine/Source/Editor/Sequencer/Public/MVVM/ViewModels/TrackRowModel.h
--- a/Engine/Source/Editor/Sequencer/Public/MVVM/ViewModels/TrackRowModel.h
+++ b/Engine/Source/Editor/Sequencer/Public/MVVM/ViewModels/TrackRowModel.h
@@ -65,6 +65,7 @@ public:
 	/*~ ITrackAreaExtension */
 	FTrackAreaParameters GetTrackAreaParameters() const override;
 	FViewModelVariantIterator GetTrackAreaModelList() const override;
+	FViewModelVariantIterator GetTopLevelChildTrackAreaModels() const override;
 
 	/*~ IDimmableExtension */
 	bool IsDimmed() const override;
diff --git a/Engine/Source/Editor/Sequencer/Public/MVVM/Views/SChannelView.h b/Engine/Source/Editor/Sequencer/Public/MVVM/Views/SChannelView.h
new file mode 100644
--- /dev/null
+++ b/Engine/Source/Editor/Sequencer/Public/MVVM/Views/SChannelView.h
@@ -0,0 +1,118 @@
+// Copyright Epic Games, Inc. All Rights Reserved.
+
+#pragma once
+
+#include "CoreTypes.h"
+#include "Widgets/SCompoundWidget.h"
+#include "MVVM/Extensions/ITrackLaneExtension.h"
+#include "MVVM/Views/KeyRenderer.h"
+#include "MVVM/ViewModelPtr.h"
+
+#include "Templates/SharedPointer.h"
+
+class FSequencer;
+struct FSequencerSelectedKey;
+
+namespace UE::Sequencer
+{
+
+struct ITrackAreaHotspot;
+
+class FTrackAreaViewModel;
+class FChannelModel;
+class FSectionModel;
+class STrackLane;
+
+/** Cache structure that stores all the sequencer-specific things that contribute to the cached screen-space renderered key states */
+struct FChannelViewKeyCachedState
+{
+	/** The current track area hotspot */
+	TWeakPtr<ITrackAreaHotspot> WeakHotspot;
+	/** DO NOT DEREFERENCE. Raw ptr to the last used hotspot. Only used to differentiate an expired hotspot from an explicitly null one. */
+	ITrackAreaHotspot* RawHotspotDoNotUse = nullptr;
+
+	/** The min/max tick value relating to the FMovieSceneSubSequenceData::ValidPlayRange bounds, or the current playback range */
+	FFrameNumber ValidPlayRangeMin, ValidPlayRangeMax;
+
+	/** The current view range */
+	TRange<FFrameTime> VisibleRange;
+
+	/** The size of the keys to draw */
+	FVector2D KeySizePx;
+
+	/** The value of FSequencerSelection::GetSerialNumber when this cache was created */
+	uint32 SelectionSerial;
+
+	/** The value of FSequencerSelectionPreview::GetSelectionHash when this cache was created */
+	uint32 SelectionPreviewHash;
+
+	/** Whether to draw a visual representation of the curve on this track as well */
+	bool bShowCurve;
+
+	/** Whether to show key bars */
+	bool bShowKeyBars;
+
+	/** Whether to gather keys for all children and collapse them into this renderer (normally only true when an item is collapsed) */
+	bool bCollapseChildren;
+
+	/** Default constructor for SWidget construction - not to be used under other circumstances */
+	FChannelViewKeyCachedState();
+
+	/** Real constructor for use when rendering */
+	FChannelViewKeyCachedState(TRange<FFrameTime> InVisibleRange, TSharedPtr<ITrackAreaHotspot> Hotspot, FViewModelPtr Model, FSequencer* Sequencer);
+
+	/** Compare this cache to another, returning what (if anything) has changed */
+	EViewDependentCacheFlags CompareTo(const FChannelViewKeyCachedState& Other) const;
+};
+
+/**
+ * Widget for displaying a channel within a section
+ */
+class SEQUENCER_API SChannelView
+	: public SCompoundWidget
+	, public ITrackLaneWidget
+{
+public:
+	SLATE_BEGIN_ARGS(SChannelView){}
+		SLATE_ATTRIBUTE(FLinearColor, KeyBarColor)
+	SLATE_END_ARGS()
+
+	void Construct(const FArguments& InArgs, const FViewModelPtr& InViewModel, TSharedPtr<FTimeToPixel> InTimeToPixel, TSharedPtr<FTrackAreaViewModel> TrackArea);
+
+	/*~ SWidget */
+	virtual int32 OnPaint(const FPaintArgs& Args, const FGeometry& AllottedGeometry, const FSlateRect& MyCullingRect, FSlateWindowElementList& OutDrawElements, int32 LayerId, const FWidgetStyle& InWidgetStyle, bool bParentEnabled) const override;
+	virtual FReply OnMouseButtonDown(const FGeometry& MyGeometry, const FPointerEvent& MouseEvent) override;
+	virtual FReply OnMouseButtonDoubleClick(const FGeometry& MyGeometry, const FPointerEvent& MouseEvent) override;
+	virtual FReply OnMouseMove(const FGeometry& MyGeometry, const FPointerEvent& MouseEvent) override;
+	virtual FReply OnMouseButtonUp(const FGeometry& MyGeometry, const FPointerEvent& MouseEvent) override;
+	virtual void OnMouseLeave(const FPointerEvent& MouseEvent) override;
+
+	/*~ ITrackLaneWidget */
+	virtual TSharedRef<const SWidget> AsWidget() const;
+	virtual FTrackLaneScreenAlignment GetAlignment(const FTimeToPixel& InTimeToPixel, const FGeometry& InParentGeometry) const;
+
+private:
+
+	int32 DrawLane(const FPaintArgs& Args, const FGeometry& AllottedGeometry, const FSlateRect& MyCullingRect, FSlateWindowElementList& OutDrawElements, int32 LayerId, const FWidgetStyle& InWidgetStyle, bool bParentEnabled) const;
+
+	void CreateKeysUnderMouse(const FVector2D& MousePosition, const FGeometry& AllottedGeometry, TArrayView<const FSequencerSelectedKey> InPressedKeys, TArray<FSequencerSelectedKey>& OutKeys);
+	void GetKeysUnderMouse(const FVector2D& MousePosition, const FGeometry& AllottedGeometry, TArray<FSequencerSelectedKey>& OutKeys) const;
+	void GetKeysAtPixelX(float LocalMousePixelX, TArray<FSequencerSelectedKey>& OutKeys) const;
+
+	TViewModelPtr<FSectionModel> GetSection() const;
+
+	TSharedPtr<FSequencer> LegacyGetSequencer() const;
+
+private:
+
+	FWeakViewModelPtr WeakModel;
+	TWeakPtr<STrackLane> WeakTrackLane;
+	TWeakPtr<FTrackAreaViewModel> WeakTrackArea;
+	TSharedPtr<FTimeToPixel> TrackAreaTimeToPixel;
+	FKeyRenderer KeyRenderer;
+	mutable TOptional<FChannelViewKeyCachedState> KeyRendererCache;
+
+	TAttribute<FLinearColor> KeyBarColor;
+};
+
+} // namespace UE::Sequencer
diff --git a/Engine/Source/Editor/Sequencer/Public/SequencerChannelInterface.h b/Engine/Source/Editor/Sequencer/Public/SequencerChannelInterface.h
--- a/Engine/Source/Editor/Sequencer/Public/SequencerChannelInterface.h
+++ b/Engine/Source/Editor/Sequencer/Public/SequencerChannelInterface.h
@@ -164,6 +164,19 @@ struct TSequencerChannelInterfaceCommon : ISequencerChannelInterface
 		return SupportsCurveEditorModels(InChannel.Cast<ChannelType>());
 	}
 
+	/**
+	 * Whether this channel should draw a curve on its editor UI
+	 *
+	 * @param Channel               The channel to query
+	 * @param InSection             The section that owns the channel
+	 * @return true to show the curve on the UI, false otherwise
+	 */
+	virtual bool ShouldShowCurve_Raw(const FMovieSceneChannel* Channel, UMovieSceneSection* InSection) const override
+	{
+		using namespace Sequencer;
+		return ShouldShowCurve(static_cast<const ChannelType*>(Channel), InSection);
+	}
+
 	/**
 	 * Create a new model for this channel that can be used on the curve editor interface
 	 *
diff --git a/Engine/Source/Editor/Sequencer/Public/SequencerChannelTraits.h b/Engine/Source/Editor/Sequencer/Public/SequencerChannelTraits.h
--- a/Engine/Source/Editor/Sequencer/Public/SequencerChannelTraits.h
+++ b/Engine/Source/Editor/Sequencer/Public/SequencerChannelTraits.h
@@ -4,7 +4,7 @@
 
 #include "Templates/SharedPointer.h"
 #include "ISequencerChannelInterface.h"
-#include "KeyDrawParams.h"
+#include "MVVM/Views/KeyDrawParams.h"
 #include "Channels/MovieSceneChannelData.h"
 #include "MovieSceneClipboard.h"
 #include "SequencerClipboardReconciler.h"
@@ -357,4 +357,17 @@ namespace Sequencer
 	 */
 	SEQUENCER_API TUniquePtr<FCurveModel> CreateCurveEditorModel(const FMovieSceneChannelHandle& ChannelHandle, UMovieSceneSection* OwningSection, TSharedRef<ISequencer> InSequencer);
 
+
+	/**
+	 * Whether this channel should draw a curve on its editor UI
+	 *
+	 * @param Channel               The channel to query
+	 * @param InSection             The section that owns the channel
+	 * @return true to show the curve on the UI, false otherwise
+	 */
+	inline bool ShouldShowCurve(const FMovieSceneChannel* Channel, UMovieSceneSection* InSection)
+	{
+		return false;
+	}
+
 }	// namespace Sequencer
diff --git a/Engine/Source/Editor/Sequencer/Public/SequencerCommands.h b/Engine/Source/Editor/Sequencer/Public/SequencerCommands.h
--- a/Engine/Source/Editor/Sequencer/Public/SequencerCommands.h
+++ b/Engine/Source/Editor/Sequencer/Public/SequencerCommands.h
@@ -346,6 +346,9 @@ public:
 	/** Toggles whether to show layer bars to edit keyframes in bulk. */
 	TSharedPtr< FUICommandInfo > ToggleLayerBars;
 
+	/** Show/hide key bar connectors for quickly retiming pairs of keys*/
+	TSharedPtr< FUICommandInfo > ToggleKeyBars;
+
 	/** Toggles whether to show channel colors in the track area. */
 	TSharedPtr< FUICommandInfo > ToggleChannelColors;
 
diff --git a/Engine/Source/Editor/Sequencer/Public/SequencerSectionPainter.h b/Engine/Source/Editor/Sequencer/Public/SequencerSectionPainter.h
--- a/Engine/Source/Editor/Sequencer/Public/SequencerSectionPainter.h
+++ b/Engine/Source/Editor/Sequencer/Public/SequencerSectionPainter.h
@@ -74,9 +74,6 @@ public:
 	/** The full geometry of the section. This is the width of the track area in the case of infinite sections */
 	FGeometry SectionGeometry;
 
-	/** Key area info with key area geometry for per key area track painting */
-	TArray<FKeyAreaElement> KeyAreaElements;
-	
 	/** The full clipping rectangle for the section */
 	FSlateRect SectionClippingRect;
 	
diff --git a/Engine/Source/Editor/Sequencer/Public/SequencerSettings.h b/Engine/Source/Editor/Sequencer/Public/SequencerSettings.h
--- a/Engine/Source/Editor/Sequencer/Public/SequencerSettings.h
+++ b/Engine/Source/Editor/Sequencer/Public/SequencerSettings.h
@@ -274,6 +274,11 @@ public:
 	/** Set whether to show layer bars */ 
 	void SetShowLayerBars(bool bInShowLayerBars);
 
+	/** @return true if showing key bars */
+	bool GetShowKeyBars() const;
+	/** Set whether to show key bars */ 
+	void SetShowKeyBars(bool bInShowKeyBars);
+
 	/** @return true if key areas are infinite */
 	bool GetInfiniteKeyAreas() const;
 	/** Set whether to show channel colors */
@@ -294,9 +299,9 @@ public:
 	/** @ Remove curve extents for the given channel */
 	void RemoveKeyAreaCurveExtents(const FString& ChannelName);
 	/** @return Get the key area curve extents for the given channel */
-	void GetKeyAreaCurveExtents(const FString& ChannelName, float& InMin, float& InMax) const;
+	void GetKeyAreaCurveExtents(const FString& ChannelName, double& InMin, double& InMax) const;
 	/** Set the key area curve extents for the given channel */
-	void SetKeyAreaCurveExtents(const FString& ChannelName, float InMin, float InMax);
+	void SetKeyAreaCurveExtents(const FString& ChannelName, double InMin, double InMax);
 
 	/** @return The key area height when showing curves */
 	float GetKeyAreaHeightWithCurves() const;
@@ -530,6 +535,10 @@ protected:
 	UPROPERTY( config, EditAnywhere, Category=Timeline )
 	bool bShowLayerBars;
 
+	/** Enable or disable key bar connections. */
+	UPROPERTY( config, EditAnywhere, Category=Timeline )
+	bool bShowKeyBars;
+
 	/** Enable or disable setting key area sections as infinite by default. */
 	UPROPERTY( config, EditAnywhere, Category=Timeline )
 	bool bInfiniteKeyAreas;
diff --git a/Engine/Source/Editor/SequencerCore/Private/MVVM/Extensions/IKeyExtension.cpp b/Engine/Source/Editor/SequencerCore/Private/MVVM/Extensions/IKeyExtension.cpp
new file mode 100644
--- /dev/null
+++ b/Engine/Source/Editor/SequencerCore/Private/MVVM/Extensions/IKeyExtension.cpp
@@ -0,0 +1,82 @@
+// Copyright Epic Games, Inc. All Rights Reserved.
+
+#include "MVVM/Extensions/IKeyExtension.h"
+
+namespace UE::Sequencer
+{
+
+bool FCachedKeys::FCachedKey::IsValid() const
+{
+	return Handle != FKeyHandle::Invalid();
+}
+
+void FCachedKeys::FCachedKey::Reset()
+{
+	Handle = FKeyHandle::Invalid();
+}
+
+void FCachedKeys::GetKeysInRange(const TRange<FFrameTime>& Range, TArrayView<const FFrameTime>* OutKeyTimes, TArrayView<const FKeyHandle>* OutHandles) const
+{
+	GetKeysInRangeWithBounds(Range, OutKeyTimes, OutHandles, nullptr, nullptr);
+}
+
+void FCachedKeys::GetKeysInRangeWithBounds(const TRange<FFrameTime>& Range, TArrayView<const FFrameTime>* OutKeyTimes, TArrayView<const FKeyHandle>* OutHandles, FCachedKey* OutLeadingKey, FCachedKey* OutTrailingKey) const
+{
+	// Binary search the first time that's >= the lower bound, minus the dilation amount
+	const int32 FirstVisibleIndex = Algo::LowerBound(KeyTimes, Range.GetLowerBoundValue());
+	// Binary search the last time that's > the upper bound
+	const int32 LastVisibleIndex = Algo::UpperBound(KeyTimes, Range.GetUpperBoundValue());
+
+	int32 Num = LastVisibleIndex - FirstVisibleIndex;
+	if (KeyTimes.IsValidIndex(FirstVisibleIndex) && LastVisibleIndex <= KeyTimes.Num())
+	{
+		if (OutKeyTimes)
+		{
+			*OutKeyTimes = MakeArrayView(&KeyTimes[FirstVisibleIndex], Num);
+		}
+
+		if (OutHandles)
+		{
+			*OutHandles = MakeArrayView(&KeyHandles[FirstVisibleIndex], Num);
+		}
+	}
+	else
+	{
+		if (OutKeyTimes)
+		{
+			*OutKeyTimes = TArrayView<const FFrameTime>();
+		}
+
+		if (OutHandles)
+		{
+			*OutHandles = TArrayView<const FKeyHandle>();
+		}
+	}
+
+	if (OutLeadingKey)
+	{
+		if (KeyTimes.IsValidIndex(FirstVisibleIndex-1))
+		{
+			*OutLeadingKey = FCachedKey{ KeyTimes[FirstVisibleIndex-1], KeyHandles[FirstVisibleIndex-1] };
+		}
+		else
+		{
+			*OutLeadingKey = FCachedKey();
+		}
+	}
+
+	if (OutTrailingKey)
+	{
+		if (KeyTimes.IsValidIndex(LastVisibleIndex))
+		{
+			*OutTrailingKey = FCachedKey{ KeyTimes[LastVisibleIndex], KeyHandles[LastVisibleIndex] };
+		}
+		else
+		{
+			*OutTrailingKey = FCachedKey();
+		}
+	}
+}
+
+} // namespace UE::Sequencer
+
diff --git a/Engine/Source/Editor/SequencerCore/Private/MVVM/ViewModelTypeIDs.cpp b/Engine/Source/Editor/SequencerCore/Private/MVVM/ViewModelTypeIDs.cpp
--- a/Engine/Source/Editor/SequencerCore/Private/MVVM/ViewModelTypeIDs.cpp
+++ b/Engine/Source/Editor/SequencerCore/Private/MVVM/ViewModelTypeIDs.cpp
@@ -23,6 +23,7 @@
 #include "MVVM/Extensions/IRenameableExtension.h"
 #include "MVVM/Extensions/LinkedOutlinerExtension.h"
 #include "MVVM/Extensions/IDeletableExtension.h"
+#include "MVVM/Extensions/IKeyExtension.h"
 #include "MVVM/ViewModels/EditorViewModel.h"
 #include "MVVM/ViewModels/OutlinerSpacer.h"
 #include "MVVM/ViewModels/OutlinerViewModel.h"
@@ -59,6 +60,7 @@ UE_SEQUENCER_DEFINE_VIEW_MODEL_TYPE_ID(IStretchableExtension);
 UE_SEQUENCER_DEFINE_VIEW_MODEL_TYPE_ID(ITrackAreaExtension);
 UE_SEQUENCER_DEFINE_VIEW_MODEL_TYPE_ID(ITrackLaneExtension);
 UE_SEQUENCER_DEFINE_VIEW_MODEL_TYPE_ID(IDeletableExtension);
+UE_SEQUENCER_DEFINE_VIEW_MODEL_TYPE_ID(IKeyExtension);
 
 // Extension types
 UE_SEQUENCER_DEFINE_VIEW_MODEL_TYPE_ID(IDynamicExtension);
diff --git a/Engine/Source/Editor/SequencerCore/Private/MVVM/ViewModels/TrackAreaViewModel.cpp b/Engine/Source/Editor/SequencerCore/Private/MVVM/ViewModels/TrackAreaViewModel.cpp
--- a/Engine/Source/Editor/SequencerCore/Private/MVVM/ViewModels/TrackAreaViewModel.cpp
+++ b/Engine/Source/Editor/SequencerCore/Private/MVVM/ViewModels/TrackAreaViewModel.cpp
@@ -34,9 +34,9 @@ void FTrackAreaViewModel::AddEditTool(TSharedPtr<ISequencerEditTool> InNewTool)
 	EditTools.Add(MoveTemp(InNewTool));
 }
 
-FTimeToPixel FTrackAreaViewModel::GetTimeToPixel(const FGeometry& AllottedGeometry) const
+FTimeToPixel FTrackAreaViewModel::GetTimeToPixel(float TrackAreaWidth) const
 {
-	return FTimeToPixel(AllottedGeometry, GetViewRange(), GetTickResolution());
+	return FTimeToPixel(TrackAreaWidth, GetViewRange(), GetTickResolution());
 }
 
 FFrameRate FTrackAreaViewModel::GetTickResolution() const
diff --git a/Engine/Source/Editor/SequencerCore/Private/MVVM/Views/KeyRenderer.cpp b/Engine/Source/Editor/SequencerCore/Private/MVVM/Views/KeyRenderer.cpp
new file mode 100644
--- /dev/null
+++ b/Engine/Source/Editor/SequencerCore/Private/MVVM/Views/KeyRenderer.cpp
@@ -0,0 +1,938 @@
+// Copyright Epic Games, Inc. All Rights Reserved.
+
+#include "MVVM/Views/KeyRenderer.h"
+#include "MVVM/Extensions/IHoveredExtension.h"
+#include "MVVM/Extensions/IOutlinerExtension.h"
+#include "MVVM/Extensions/LinkedOutlinerExtension.h"
+#include "MVVM/ViewModels/ViewModelIterators.h"
+
+#include "Styling/AppStyle.h"
+
+#include "CurveModel.h"
+#include "CurveEditor.h"
+
+namespace UE::Sequencer
+{
+
+void FKeyRenderer::Initialize(const FViewModelPtr& InViewModel)
+{
+	WeakViewModel = InViewModel;
+}
+
+bool FKeyRenderer::HasCurve() const
+{
+	return PrecomputedCurve.Num() >= 2;
+}
+
+void FKeyRenderer::HitTestKeys(const FFrameTime& Time, TArray<FKeysForModel>& OutAllKeys) const
+{
+	const FFrameTime HalfKeyWidth = KeyWidthInFrames * .5f;
+
+	// Hit test the actual structure used for drawing keys to ensure we are 100% accurate
+	// PrecomputedKeys is an array sorted by time, so we can binary search it to find the hit key
+	const int32 LowerIndex = Algo::LowerBoundBy(PrecomputedKeys, Time - HalfKeyWidth, &FKey::FinalKeyPosition);
+	const int32 UpperIndex = Algo::UpperBoundBy(PrecomputedKeys, Time + HalfKeyWidth, &FKey::FinalKeyPosition);
+
+	if (LowerIndex >= PrecomputedKeys.Num())
+	{
+		// No hit - we ran off the end of the drawn keys
+		return;
+	}
+
+	// We may have hit many overlapping keys, so we need to choose the one closest to the time
+	// Assume the first one is correct, then keep checking the distance between that and the next
+	// Break as soon as we find a key that is further away than the last (we know we're moving away from the time at that point)
+	int32 HitIndex = LowerIndex;
+	for (int32 Index = LowerIndex+1; Index < UpperIndex && Index < PrecomputedKeys.Num(); ++Index)
+	{
+		const FKey& PrevKey = PrecomputedKeys[Index-1];
+		const FKey& ThisKey = PrecomputedKeys[Index];
+
+		// Find the distance from the middle of the key to the hit time
+		FFrameTime PrevKeyDistance = FMath::Abs((PrevKey.KeyTickStart + PrevKey.KeyTickEnd) * .5f - Time);
+		FFrameTime ThisKeyDistance = FMath::Abs((ThisKey.KeyTickStart + ThisKey.KeyTickEnd) * .5f - Time);
+
+		if (ThisKeyDistance <= PrevKeyDistance)
+		{
+			HitIndex = Index;
+		}
+		else
+		{
+			break;
+		}
+	}
+
+	if (HitIndex < PrecomputedKeys.Num() && FMath::Abs(PrecomputedKeys[HitIndex].FinalKeyPosition - Time) <= HalfKeyWidth)
+	{
+		FFrameTime Start = PrecomputedKeys[HitIndex].KeyTickStart;
+		FFrameTime End   = PrecomputedKeys[HitIndex].KeyTickEnd;
+
+		// We have a hit. Gather anything from the key draw info that overlaps the hit range
+		for (const FCachedKeyDrawInformation& Entry : KeyDrawInfo)
+		{
+			const int32 FrameStart = Algo::LowerBound(Entry.FramesInRange, Start);
+			const int32 FrameEnd   = Algo::UpperBound(Entry.FramesInRange, End);
+			const int32 NumKeys    = FrameEnd - FrameStart;
+
+			if (NumKeys > 0 && FrameStart < Entry.HandlesInRange.Num())
+			{
+				FKeysForModel Keys;
+				Keys.Model = Entry.WeakKeyExtension.Pin();
+				Keys.Keys.SetNumUninitialized(NumKeys);
+
+				for (int32 Index = 0; Index < NumKeys; ++Index)
+				{
+					Keys.Keys[Index] = Entry.HandlesInRange[FrameStart + Index];
+				}
+
+				OutAllKeys.Emplace(MoveTemp(Keys));
+			}
+		}
+	}
+}
+
+bool FKeyRenderer::HitTestKeyBar(const FFrameTime& Time, TArray<FKeysForModel>& OutLeadingKeys, TArray<FKeysForModel>& OutTrailingKeys) const
+{
+	const FFrameTime HalfKeyWidth = KeyWidthInFrames * .5f;
+
+	// Hit test the actual structure used for drawing key bars to ensure we are 100% accurate
+	// PrecomputedKeyBars is an array sorted by time, so we can binary search it to find the hit area
+	const int32 HitIndex = Algo::LowerBoundBy(PrecomputedKeyBars, Time, &FKeyBar::EndTime);
+	if (HitIndex >= PrecomputedKeyBars.Num() || PrecomputedKeyBars[HitIndex].StartTime > Time)
+	{
+		// No hit - we ran off the end of the drawn keys, or didn't hit a bar
+		return false;
+	}
+
+	// Generate leading keys array
+	{
+		FFrameTime LeadingStart = PrecomputedKeyBars[HitIndex].StartTime - HalfKeyWidth;
+		FFrameTime LeadingEnd   = PrecomputedKeyBars[HitIndex].StartTime + HalfKeyWidth;
+
+		for (const FCachedKeyDrawInformation& Entry : KeyDrawInfo)
+		{
+			const int32 FrameStart = Algo::LowerBound(Entry.CachedKeys->KeyTimes, LeadingStart);
+			const int32 FrameEnd   = Algo::UpperBound(Entry.CachedKeys->KeyTimes, LeadingEnd);
+			const int32 NumKeys    = FrameEnd - FrameStart;
+
+			if (NumKeys > 0 && FrameStart < Entry.CachedKeys->KeyHandles.Num())
+			{
+				FKeysForModel Keys;
+				Keys.Model = Entry.WeakKeyExtension.Pin();
+				Keys.Keys.SetNumUninitialized(NumKeys);
+
+				for (int32 Index = 0; Index < NumKeys; ++Index)
+				{
+					Keys.Keys[Index] = Entry.CachedKeys->KeyHandles[FrameStart + Index];
+				}
+
+				OutLeadingKeys.Emplace(MoveTemp(Keys));
+			}
+		}
+	}
+
+	// Generate trailing keys array
+	{
+		FFrameTime TrailingStart = PrecomputedKeyBars[HitIndex].EndTime - HalfKeyWidth;
+		FFrameTime TrailingEnd   = PrecomputedKeyBars[HitIndex].EndTime + HalfKeyWidth;
+
+		for (const FCachedKeyDrawInformation& Entry : KeyDrawInfo)
+		{
+			const int32 FrameStart = Algo::LowerBound(Entry.CachedKeys->KeyTimes, TrailingStart);
+			const int32 FrameEnd   = Algo::UpperBound(Entry.CachedKeys->KeyTimes, TrailingEnd);
+			const int32 NumKeys    = FrameEnd - FrameStart;
+
+			if (NumKeys > 0 && FrameStart < Entry.CachedKeys->KeyHandles.Num())
+			{
+				FKeysForModel Keys;
+				Keys.Model = Entry.WeakKeyExtension.Pin();
+				Keys.Keys.SetNumUninitialized(NumKeys);
+
+				for (int32 Index = 0; Index < NumKeys; ++Index)
+				{
+					Keys.Keys[Index] = Entry.CachedKeys->KeyHandles[FrameStart + Index];
+				}
+
+				OutTrailingKeys.Emplace(MoveTemp(Keys));
+			}
+		}
+	}
+
+	return true;
+}
+
+void FKeyRenderer::Update(const FKeyBatchParameters& Params, const FGeometry& AllottedGeometry) const
+{
+	if (EnumHasAnyFlags(Params.CacheState, EViewDependentCacheFlags::Empty))
+	{
+		KeyDrawInfo.Reset();
+		return;
+	}
+
+	if (EnumHasAnyFlags(Params.CacheState, EViewDependentCacheFlags::DataChanged))
+	{
+		CacheKeyExtensions(Params);
+	}
+
+	FKeyBatchParameters NewParams = Params;
+	NewParams.CacheState |= UpdateViewIndependentData();
+
+	// We can reuse this key layout - update all the cached key positions
+	UpdateViewDependentData(NewParams, AllottedGeometry);
+}
+
+EViewDependentCacheFlags FKeyRenderer::UpdateViewIndependentData() const
+{
+	EViewDependentCacheFlags CacheFlags = EViewDependentCacheFlags::None;
+
+	for (FCachedKeyDrawInformation& CachedKeyDrawInfo : KeyDrawInfo)
+	{
+		CacheFlags |= CachedKeyDrawInfo.UpdateViewIndependentData();
+	}
+
+	return CacheFlags;
+}
+
+void FKeyRenderer::CacheKeyExtensions(const FKeyBatchParameters& Params) const
+{
+	KeyDrawInfo.Empty();
+
+	FViewModelPtr ViewModel = WeakViewModel.Pin();
+	if (!ViewModel)
+	{
+		return;
+	}
+
+	// Walk through the current model and accumulate key renderer extension for any collapsed children
+	constexpr bool bIncludeThis = true;
+	for (FParentFirstChildIterator ChildIt(ViewModel, bIncludeThis) ; ChildIt; ++ChildIt)
+	{
+		FViewModelPtr CurrentViewModel = *ChildIt;
+
+		TViewModelPtr<IKeyExtension> KeyRenderer = CurrentViewModel.ImplicitCast();
+		if (KeyRenderer)
+		{
+			KeyDrawInfo.Emplace(KeyRenderer);
+		}
+
+		if (Params.bCollapseChildren == false)
+		{
+			break;
+		}
+	}
+}
+
+void FKeyRenderer::UpdateViewDependentData(const FKeyBatchParameters& Params, const FGeometry& AllottedGeometry) const
+{
+	if (Params.CacheState == EViewDependentCacheFlags::None)
+	{
+		// Cache is still hot - nothing to do
+		return;
+	}
+
+	KeyWidthInFrames = Params.TimeToPixel.PixelDeltaToFrame(Params.KeySizePx.X);
+
+	const bool bHasAnySelection        = Params.ClientInterface && Params.ClientInterface->HasAnySelectedKeys();
+	const bool bHasAnySelectionPreview = Params.ClientInterface && Params.ClientInterface->HasAnyPreviewSelectedKeys();
+	const bool bHasAnyHoveredKeys      = Params.ClientInterface && Params.ClientInterface->HasAnyHoveredKeys();
+
+	// ------------------------------------------------------------------------------
+	// Update view-dependent data for each draw info
+	for (FCachedKeyDrawInformation& Info : KeyDrawInfo)
+	{
+		Info.CacheViewDependentData(Params);
+	}
+
+	// ------------------------------------------------------------------------------
+	// If the data has changed, or key state has changed, or the view has been zoomed
+	// we cannot preserve any keys (because we don't know whether they are still valid)
+	const bool bCanPreserveKeys = !EnumHasAnyFlags(Params.CacheState, EViewDependentCacheFlags::DataChanged | EViewDependentCacheFlags::ViewZoomed | EViewDependentCacheFlags::KeyStateChanged);
+
+	double PreserveStartFrame = TNumericLimits<double>::Max();
+	TArray<FKey> PreservedKeys;
+
+	// Attempt to preserve any previously computed key draw information
+	if (bCanPreserveKeys && PrecomputedKeys.Num() != 0)
+	{
+		const FFrameTime LowerBoundFrame = Params.VisibleRange.GetLowerBoundValue();
+		const FFrameTime UpperBoundFrame = Params.VisibleRange.GetUpperBoundValue();
+
+		const int32 PreserveStartIndex = Algo::LowerBoundBy(PrecomputedKeys, LowerBoundFrame, &FKey::KeyTickStart);
+		const int32 PreserveEndIndex   = Algo::UpperBoundBy(PrecomputedKeys, UpperBoundFrame, &FKey::KeyTickEnd);
+
+		const int32 PreserveNum = PreserveEndIndex - PreserveStartIndex;
+		if (PreserveNum > 0)
+		{
+			PreservedKeys = TArray<FKey>(PrecomputedKeys.GetData() + PreserveStartIndex, PreserveNum);
+			PreserveStartFrame = PreservedKeys[0].KeyTickStart.AsDecimal();
+
+			FFrameTime ActualPreserveEndFrame = PreservedKeys.Last().KeyTickEnd;
+			for (FCachedKeyDrawInformation& Info : KeyDrawInfo)
+			{
+				Info.PreserveToIndex = Algo::UpperBound(Info.FramesInRange, ActualPreserveEndFrame);
+			}
+		}
+	}
+
+	// ------------------------------------------------------------------------------
+	// Begin precomputation of keys to draw
+	PrecomputedKeys.Reset();
+
+	static float PixelOverlapThreshold = 3.f;
+	const double TimeOverlapThreshold = Params.TimeToPixel.PixelDeltaToFrame(PixelOverlapThreshold).AsDecimal();
+
+	auto AnythingLeftToDraw = [](const FCachedKeyDrawInformation& In)
+	{
+		return In.NextUnhandledIndex < In.FramesInRange.Num();
+	};
+
+	// Keep iterating all the cached key positions until we've moved through everything
+	// As stated above - this loop does not scale well for large numbers of KeyDrawInfo
+	// Which is generally not a problem, but is troublesome for Control Rigs
+	while (KeyDrawInfo.ContainsByPredicate(AnythingLeftToDraw))
+	{
+		// Determine the next key position to draw
+		double CardinalKeyFrame = TNumericLimits<double>::Max();
+		for (const FCachedKeyDrawInformation& Info : KeyDrawInfo)
+		{
+			if (Info.NextUnhandledIndex < Info.FramesInRange.Num())
+			{
+				CardinalKeyFrame = FMath::Min(CardinalKeyFrame, Info.FramesInRange[Info.NextUnhandledIndex].AsDecimal());
+			}
+		}
+
+		// If the cardinal time overlaps the preserved range, skip those keys
+		if (CardinalKeyFrame >= PreserveStartFrame && PreservedKeys.Num() != 0)
+		{
+			PrecomputedKeys.Append(PreservedKeys);
+			PreservedKeys.Empty();
+			for (FCachedKeyDrawInformation& Info : KeyDrawInfo)
+			{
+				Info.NextUnhandledIndex = Info.PreserveToIndex;
+			}
+			continue;
+		}
+
+		// Start grouping keys at the current key time plus 99% of the threshold to ensure that we group at the center of keys
+		// and that we avoid floating point precision issues where there is only one key [(KeyTime + TimeOverlapThreshold) - KeyTime != TimeOverlapThreshold] for some floats
+		double CardinalKeyTime = CardinalKeyFrame + TimeOverlapThreshold*0.9994f;
+
+		// Track whether all of the keys are within the valid range
+		bool bIsInRange = true;
+
+		const FFrameTime ValidPlayRangeMin = Params.ValidPlayRangeMin;
+		const FFrameTime ValidPlayRangeMax = Params.ValidPlayRangeMax;
+
+		double AverageKeyTime = 0.0;
+		int32 NumKeyTimes = 0;
+
+		FFrameTime KeyTickStart = TNumericLimits<FFrameNumber>::Max();
+		FFrameTime KeyTickEnd   = TNumericLimits<FFrameNumber>::Lowest();
+		EKeyConnectionStyle ConnectionStyle = EKeyConnectionStyle::None;
+
+		auto HandleKey = [&bIsInRange, &AverageKeyTime, &NumKeyTimes, &KeyTickStart, &KeyTickEnd, ValidPlayRangeMin, ValidPlayRangeMax](const FFrameTime& KeyFrame)
+		{
+			if (bIsInRange && (KeyFrame < ValidPlayRangeMin || KeyFrame >= ValidPlayRangeMax))
+			{
+				bIsInRange = false;
+			}
+
+			KeyTickStart = FMath::Min(KeyFrame, KeyTickStart);
+			KeyTickEnd   = FMath::Max(KeyFrame, KeyTickEnd);
+
+			AverageKeyTime += KeyFrame.AsDecimal();
+			++NumKeyTimes;
+		};
+
+
+		bool bFoundKey = false;
+		FKey NewKey;
+
+		int32 NumPreviewSelected = 0;
+		int32 NumPreviewNotSelected = 0;
+		int32 NumSelected = 0;
+		int32 NumHovered = 0;
+		int32 TotalNumKeys = 0;
+		int32 NumOverlaps = 0;
+
+		// Determine the ranges of keys considered to reside at this position
+		for (int32 DrawIndex = 0; DrawIndex < KeyDrawInfo.Num(); ++DrawIndex)
+		{
+			FCachedKeyDrawInformation& Info = KeyDrawInfo[DrawIndex];
+			TViewModelPtr<IKeyExtension> KeyExtension = Info.WeakKeyExtension.Pin();
+			if (!KeyExtension)
+			{
+				continue;
+			}
+			if (Info.NextUnhandledIndex >= Info.FramesInRange.Num())
+			{
+				NewKey.Flags |= EKeyRenderingFlags::PartialKey;
+				continue;
+			}
+			else if (!FMath::IsNearlyEqual(Info.FramesInRange[Info.NextUnhandledIndex].AsDecimal(), CardinalKeyTime, TimeOverlapThreshold))
+			{
+				NewKey.Flags |= EKeyRenderingFlags::PartialKey;
+				continue;
+			}
+
+			int32 ThisNumOverlaps = -1;
+			do
+			{
+				HandleKey(Info.FramesInRange[Info.NextUnhandledIndex]);
+
+				if (!bFoundKey)
+				{
+					NewKey.Params = Info.DrawParams[Info.NextUnhandledIndex];
+					bFoundKey = true;
+				}
+				else if (Info.DrawParams[Info.NextUnhandledIndex] != NewKey.Params)
+				{
+					NewKey.Flags |= EKeyRenderingFlags::PartialKey;
+				}
+
+				NewKey.Params.ConnectionStyle |= Info.DrawParams[Info.NextUnhandledIndex].ConnectionStyle;
+
+				// Avoid creating FSequencerSelectedKeys unless absolutely necessary
+				FKeyHandle ThisKeyHandle = Info.HandlesInRange[Info.NextUnhandledIndex];
+
+				if (bHasAnySelection && Params.ClientInterface->IsKeySelected(KeyExtension, ThisKeyHandle))
+				{
+					++NumSelected;
+				}
+				if (bHasAnySelectionPreview)
+				{
+					EKeySelectionPreviewState SelectionState = Params.ClientInterface->GetPreviewSelectionState(KeyExtension, ThisKeyHandle);
+
+					NumPreviewSelected    += int32(SelectionState == EKeySelectionPreviewState::Selected);
+					NumPreviewNotSelected += int32(SelectionState == EKeySelectionPreviewState::NotSelected);
+				}
+				if (bHasAnyHoveredKeys && Params.ClientInterface->IsKeyHovered(KeyExtension, ThisKeyHandle))
+				{
+					++NumHovered;
+				}
+
+				++TotalNumKeys;
+				++Info.NextUnhandledIndex;
+				++ThisNumOverlaps;
+			}
+			while (Info.NextUnhandledIndex < Info.FramesInRange.Num() && FMath::IsNearlyEqual(Info.FramesInRange[Info.NextUnhandledIndex].AsDecimal(), CardinalKeyTime, TimeOverlapThreshold));
+
+			NumOverlaps += ThisNumOverlaps;
+		}
+
+		if (NumKeyTimes == 0) //-V547
+		{
+			// This is not actually possible since HandleKey must have been called
+			// at least once, but it needs to be here to avoid a static analysis warning
+			break;
+		}
+
+		NewKey.FinalKeyPosition = FFrameTime::FromDecimal(AverageKeyTime / NumKeyTimes);
+		NewKey.KeyTickStart = KeyTickStart;
+		NewKey.KeyTickEnd = KeyTickEnd;
+
+		if (EnumHasAnyFlags(NewKey.Flags, EKeyRenderingFlags::PartialKey))
+		{
+			static const FSlateBrush* PartialKeyBrush = FAppStyle::GetBrush("Sequencer.PartialKey");
+			NewKey.Params.FillBrush = NewKey.Params.BorderBrush = PartialKeyBrush;
+		}
+
+		// Determine the key color based on its selection/hover states
+		if (NumPreviewSelected == TotalNumKeys)
+		{
+			NewKey.Flags |= EKeyRenderingFlags::PreviewSelected;
+		}
+		else if (NumPreviewNotSelected == TotalNumKeys)
+		{
+			NewKey.Flags |= EKeyRenderingFlags::PreviewNotSelected;
+		}
+		else if (NumSelected == TotalNumKeys)
+		{
+			NewKey.Flags |= EKeyRenderingFlags::Selected;
+		}
+		else if (NumSelected != 0)
+		{
+			NewKey.Flags |= EKeyRenderingFlags::AnySelected;
+		}
+		else if (NumHovered == TotalNumKeys)
+		{
+			NewKey.Flags |= EKeyRenderingFlags::Hovered;
+		}
+
+		if (NumOverlaps > 0)
+		{
+			NewKey.Flags |= EKeyRenderingFlags::Overlaps;
+		}
+
+		if (!bIsInRange)
+		{
+			NewKey.Flags |= EKeyRenderingFlags::OutOfRange;
+		}
+
+		const int32 NumStyles = FMath::CountBits(static_cast<__underlying_type(EKeyConnectionStyle)>(NewKey.Params.ConnectionStyle));
+		if (NumStyles > 1)
+		{
+			NewKey.Params.ConnectionStyle = EKeyConnectionStyle::Dotted;
+		}
+
+		PrecomputedKeys.Add(NewKey);
+	}
+
+	// Now add connecting bars for the keys
+
+	PrecomputedKeyBars.Reset();
+	if (Params.bShowKeyBars)
+	{
+		const FKeyDrawParams* LeadingKeyParams = nullptr;
+		const FKeyDrawParams* TrailingKeyParams = nullptr;
+
+		const FCachedKeys::FCachedKey* LeadingKey = nullptr;
+		const FCachedKeys::FCachedKey* TrailingKey = nullptr;
+
+		for (const FCachedKeyDrawInformation& Info : KeyDrawInfo)
+		{
+			if (Info.LeadingKey.IsValid())
+			{
+				if (!LeadingKey || Info.LeadingKey.Time > LeadingKey->Time)
+				{
+					LeadingKey = &Info.LeadingKey;
+					LeadingKeyParams = &Info.LeadingKeyParams;
+				}
+			}
+			if (Info.TrailingKey.IsValid())
+			{
+				if (!TrailingKey || Info.TrailingKey.Time < TrailingKey->Time)
+				{
+					TrailingKey = &Info.TrailingKey;
+					TrailingKeyParams = &Info.TrailingKeyParams;
+				}
+			}
+		}
+
+		if (PrecomputedKeys.Num() == 0)
+		{
+			if (LeadingKey && TrailingKey)
+			{
+				PrecomputedKeyBars.Add(FKeyBar{ LeadingKey->Time, TrailingKey->Time, EKeyBarRenderingFlags::None });
+			}
+		}
+		else
+		{
+			if (LeadingKey && Params.TimeToPixel.FrameDeltaToPixel(PrecomputedKeys[0].FinalKeyPosition - LeadingKey->Time) > Params.KeySizePx.X)
+			{
+				PrecomputedKeyBars.Add(FKeyBar{ LeadingKey->Time, PrecomputedKeys[0].FinalKeyPosition, EKeyBarRenderingFlags::None, LeadingKeyParams->ConnectionStyle });
+			}
+
+			for (int32 Index = 0; Index < PrecomputedKeys.Num()-1; ++Index)
+			{
+				FKey ThisKey = PrecomputedKeys[Index];
+				FKey NextKey = PrecomputedKeys[Index+1];
+				if (Params.TimeToPixel.FrameDeltaToPixel(NextKey.FinalKeyPosition - ThisKey.FinalKeyPosition) > Params.KeySizePx.X)
+				{
+					PrecomputedKeyBars.Add(FKeyBar{ ThisKey.FinalKeyPosition, NextKey.FinalKeyPosition, EKeyBarRenderingFlags::None, ThisKey.Params.ConnectionStyle });
+				}
+			}
+
+			if (TrailingKey && Params.TimeToPixel.FrameDeltaToPixel(TrailingKey->Time - PrecomputedKeys.Last().FinalKeyPosition) > Params.KeySizePx.X)
+			{
+				PrecomputedKeyBars.Add(FKeyBar{ PrecomputedKeys.Last().FinalKeyPosition, TrailingKey->Time, EKeyBarRenderingFlags::None, PrecomputedKeys.Last().Params.ConnectionStyle });
+			}
+		}
+	}
+
+	PrecomputedCurve.Reset();
+	if (Params.bShowCurve)
+	{
+		for (FCachedKeyDrawInformation& Info : KeyDrawInfo)
+		{
+			TViewModelPtr<IKeyExtension> KeyExtension = Info.WeakKeyExtension.Pin();
+			if (!KeyExtension)
+			{
+				continue;
+			}
+
+			TUniquePtr<FCurveModel> CurveModel = KeyExtension->CreateCurveModel();
+			if (CurveModel)
+			{
+				// Curve editor expects things in seconds
+				const double MinTimeSeconds = Params.VisibleRange.GetLowerBoundValue() / Params.TimeToPixel.GetTickResolution();
+				const double MaxTimeSeconds = Params.VisibleRange.GetUpperBoundValue() / Params.TimeToPixel.GetTickResolution();
+				const double TimeRangeSeconds = MaxTimeSeconds - MinTimeSeconds;
+
+				double Min = 0.0;
+				double Max = 1.0;
+				if (!KeyExtension->GetFixedExtents(Min, Max))
+				{
+					CurveModel->GetValueRange(Min, Max);
+				}
+
+				if (Min == Max)
+				{
+					PrecomputedCurve.Add(MakeTuple(Params.TimeToPixel.FrameToPixel(Params.VisibleRange.GetLowerBoundValue()), 0.0f));
+					PrecomputedCurve.Add(MakeTuple(Params.TimeToPixel.FrameToPixel(Params.VisibleRange.GetUpperBoundValue()), 0.0f));
+				}
+				else
+				{
+					const double Range = Max - Min;
+
+					FCurveEditor Dummy;
+
+					// Set up a screen space that encompasses the entire curve, for the duration of the visible range
+					FCurveEditorScreenSpace ScreenSpace(
+						AllottedGeometry.GetLocalSize(),
+						MinTimeSeconds, MaxTimeSeconds,
+						Min, Max);
+
+					CurveModel->DrawCurve(Dummy, ScreenSpace, PrecomputedCurve);
+
+					// Put the points in pixel space, and normalize the values
+					for (TTuple<double, double>& Pair : PrecomputedCurve)
+					{
+						Pair.Key = Params.TimeToPixel.SecondsToPixel(Pair.Key);
+						if (Range > 0.0)
+						{
+							Pair.Value = (Pair.Value - Min) / Range;
+						}
+					}
+				}
+			}
+		}
+	}
+}
+
+int32 FKeyRenderer::DrawCurve(const FKeyBatchParameters& Params, const FGeometry& AllottedGeometry, const FSlateRect& MyCullingRect, const FKeyRendererPaintArgs& PaintArgs, int32 LayerId) const
+{
+	PaintArgs.DrawElements->PushClip(FSlateClippingZone(AllottedGeometry));
+
+	TArray<FVector2D> CurvePoints;
+	for (const TTuple<double, double>& CurveKey : PrecomputedCurve)
+	{
+		CurvePoints.Add(
+			FVector2D(
+				CurveKey.Get<0>(),
+				(1.0f - CurveKey.Get<1>()) * (AllottedGeometry.GetLocalSize().Y-2.f) + 1.f // Inset the curve 1 px so it doesn't clip
+			)
+		);
+	}
+
+	const float CurveThickness = 1.f;
+	const bool  bAntiAliasCurves = true;
+
+	FSlateDrawElement::MakeLines(
+		*PaintArgs.DrawElements,
+		LayerId,
+		AllottedGeometry.ToPaintGeometry(),
+		CurvePoints,
+		PaintArgs.DrawEffects,
+		PaintArgs.CurveColor,
+		bAntiAliasCurves,
+		CurveThickness
+	);
+
+	++LayerId;
+
+	PaintArgs.DrawElements->PopClip();
+	for (const FCachedKeyDrawInformation& CachedKeyDrawInfo : KeyDrawInfo)
+	{
+		LayerId = CachedKeyDrawInfo.DrawExtra(AllottedGeometry, LayerId);
+	}
+
+	return LayerId;
+}
+
+int32 FKeyRenderer::Draw(const FKeyBatchParameters& Params, const FGeometry& AllottedGeometry, const FSlateRect& MyCullingRect, const FKeyRendererPaintArgs& PaintArgs, int32 LayerId) const
+{
+	// Draw key bars first
+	{
+		static const FSlateBrush* DottedKeyBarBrush = FAppStyle::GetBrush("Sequencer.KeyBar.Dotted");
+		static const FSlateBrush* DashedKeyBarBrush = FAppStyle::GetBrush("Sequencer.KeyBar.Dashed");
+		static const FSlateBrush* SolidKeyBarBrush = FAppStyle::GetBrush("Sequencer.KeyBar.Solid");
+
+		const float KeyBarFadeWidthPx = MyCullingRect.GetSize2f().X;
+		const float KeyBarFadeThresholdPx = KeyBarFadeWidthPx * 0.25f;
+
+		static float KeyBarHeight = 2.f;
+		const float KeyBarTop = AllottedGeometry.GetLocalSize().Y * .5f - KeyBarHeight*.5f;
+		for (const FKeyBar& KeyBar : PrecomputedKeyBars)
+		{
+			const float StartPx = Params.TimeToPixel.FrameToPixel(KeyBar.StartTime);
+			const float EndPx = Params.TimeToPixel.FrameToPixel(KeyBar.EndTime);
+			const float KeyBarWidth = EndPx - StartPx;
+
+			const float FadeAlpha = 1.f - FMath::Clamp((KeyBarWidth - KeyBarFadeThresholdPx) / (KeyBarFadeWidthPx - KeyBarFadeThresholdPx), 0.f, 1.f);
+			const float Opacity = FMath::Lerp(.25f, 1.f, FadeAlpha);
+
+			const FSlateBrush* Brush = SolidKeyBarBrush;
+			if (EnumHasAnyFlags(KeyBar.ConnectionStyle, EKeyConnectionStyle::Dotted))
+			{
+				Brush = DottedKeyBarBrush;
+			}
+			else if (EnumHasAnyFlags(KeyBar.ConnectionStyle, EKeyConnectionStyle::Dashed))
+			{
+				Brush = DashedKeyBarBrush;
+			}
+
+			if (Opacity < 1.f && Brush == SolidKeyBarBrush)
+			{
+				TArray<FSlateGradientStop> Stops;
+				Stops.Add(FSlateGradientStop(FVector2d(0.0, 0.0), PaintArgs.KeyBarColor));
+				Stops.Add(FSlateGradientStop(FVector2d(KeyBarFadeThresholdPx*0.5f, 0.0), PaintArgs.KeyBarColor.CopyWithNewOpacity(Opacity)));
+				Stops.Add(FSlateGradientStop(FVector2d(KeyBarWidth - KeyBarFadeThresholdPx*0.5f * 0.75f, 0.0), PaintArgs.KeyBarColor.CopyWithNewOpacity(Opacity)));
+				Stops.Add(FSlateGradientStop(FVector2d(KeyBarWidth, 0.0), PaintArgs.KeyBarColor));
+
+				FSlateDrawElement::MakeGradient(
+					*PaintArgs.DrawElements,
+					LayerId,
+					AllottedGeometry.ToPaintGeometry(
+						FVector2D(KeyBarWidth, KeyBarHeight),
+						FSlateLayoutTransform(FVector2D(StartPx, KeyBarTop))
+					),
+					MoveTemp(Stops),
+					EOrientation::Orient_Vertical,
+					PaintArgs.DrawEffects
+				);
+			}
+			else
+			{
+				FSlateDrawElement::MakeBox(
+					*PaintArgs.DrawElements,
+					LayerId,
+					// Center the key along Y.  Ensure the middle of the key is at the actual key time
+					AllottedGeometry.ToPaintGeometry(
+						FVector2D(EndPx - StartPx, KeyBarHeight),
+						FSlateLayoutTransform(FVector2D(StartPx, KeyBarTop))
+					),
+					Brush,
+					PaintArgs.DrawEffects,
+					PaintArgs.KeyBarColor.CopyWithNewOpacity(Opacity)
+				);
+			}
+		}
+		++LayerId;
+	}
+
+	for (const FKey& Key : PrecomputedKeys)
+	{
+		FKeyDrawParams KeyDrawParams = Key.Params;
+
+		if (EnumHasAnyFlags(Key.Flags, EKeyRenderingFlags::PartialKey))
+		{
+			KeyDrawParams.FillOffset = FVector2D(0.f, 0.f);
+			KeyDrawParams.FillTint   = KeyDrawParams.BorderTint  = FLinearColor::White;
+		}
+
+		const bool bSelected = EnumHasAnyFlags(Key.Flags, EKeyRenderingFlags::Selected);
+		// Determine the key color based on its selection/hover states
+		if (EnumHasAnyFlags(Key.Flags, EKeyRenderingFlags::PreviewSelected))
+		{
+			FLinearColor PreviewSelectionColor = PaintArgs.SelectionColor.LinearRGBToHSV();
+			PreviewSelectionColor.R += 0.1f; // +10% hue
+			PreviewSelectionColor.G = 0.6f; // 60% saturation
+			KeyDrawParams.BorderTint = KeyDrawParams.FillTint = PreviewSelectionColor.HSVToLinearRGB();
+		}
+		else if (EnumHasAnyFlags(Key.Flags, EKeyRenderingFlags::PreviewNotSelected))
+		{
+			KeyDrawParams.BorderTint = FLinearColor(0.05f, 0.05f, 0.05f, 1.0f);
+		}
+		else if (bSelected)
+		{
+			KeyDrawParams.BorderTint = PaintArgs.SelectionColor;
+			KeyDrawParams.FillTint = FLinearColor(0.05f, 0.05f, 0.05f, 1.0f);
+		}
+		else if (EnumHasAnyFlags(Key.Flags, EKeyRenderingFlags::AnySelected))
+		{
+			// partially selected
+			KeyDrawParams.BorderTint = PaintArgs.SelectionColor.CopyWithNewOpacity(0.5f);
+			KeyDrawParams.FillTint = FLinearColor(0.05f, 0.05f, 0.05f, 0.5f);
+		}
+		else if (EnumHasAnyFlags(Key.Flags, EKeyRenderingFlags::Hovered))
+		{
+			KeyDrawParams.BorderTint = FLinearColor(1.0f, 1.0f, 1.0f, 1.0f);
+			KeyDrawParams.FillTint = FLinearColor(1.0f, 1.0f, 1.0f, 1.0f);
+		}
+		else
+		{
+			KeyDrawParams.BorderTint = FLinearColor(0.05f, 0.05f, 0.05f, 1.0f);
+		}
+
+		// Color keys with overlaps with a red border
+		if (EnumHasAnyFlags(Key.Flags, EKeyRenderingFlags::Overlaps))
+		{
+			KeyDrawParams.BorderTint = FLinearColor(0.83f, 0.12f, 0.12f, 1.0f); // Red
+		}
+
+		const ESlateDrawEffect KeyDrawEffects = EnumHasAnyFlags(Key.Flags, EKeyRenderingFlags::OutOfRange) ? ESlateDrawEffect::DisabledEffect : PaintArgs.DrawEffects;
+
+		// draw border
+		const FVector2D KeySize = bSelected ? Params.KeySizePx + PaintArgs.ThrobAmount * PaintArgs.KeyThrobValue : Params.KeySizePx;
+
+		static const float BrushBorderWidth = 2.0f;
+
+		const float KeyPositionPx = Params.TimeToPixel.FrameToPixel(Key.FinalKeyPosition);
+
+		FSlateDrawElement::MakeBox(
+			*PaintArgs.DrawElements,
+			// always draw selected keys on top of other keys
+			bSelected ? LayerId + 1 : LayerId,
+			// Center the key along Y.  Ensure the middle of the key is at the actual key time
+			AllottedGeometry.ToPaintGeometry(
+				FVector2D(
+					KeyPositionPx - FMath::CeilToFloat(KeySize.X / 2.0f),
+					((AllottedGeometry.GetLocalSize().Y / 2.0f) - (KeySize.Y / 2.0f))
+				),
+				KeySize
+			),
+			KeyDrawParams.BorderBrush,
+			KeyDrawEffects,
+			KeyDrawParams.BorderTint
+		);
+
+		// draw fill
+		FSlateDrawElement::MakeBox(
+			*PaintArgs.DrawElements,
+			// always draw selected keys on top of other keys
+			bSelected ? LayerId + 2 : LayerId + 1,
+			// Center the key along Y.  Ensure the middle of the key is at the actual key time
+			AllottedGeometry.ToPaintGeometry(
+				KeyDrawParams.FillOffset + 
+				FVector2D(
+					(KeyPositionPx - FMath::CeilToFloat((KeySize.X / 2.0f) - BrushBorderWidth)),
+					((AllottedGeometry.GetLocalSize().Y / 2.0f) - ((KeySize.Y / 2.0f) - BrushBorderWidth))
+				),
+				KeySize - 2.0f * BrushBorderWidth
+			),
+			KeyDrawParams.FillBrush,
+			KeyDrawEffects,
+			KeyDrawParams.FillTint
+		);
+	}
+
+	return LayerId + 2;
+}
+
+int32 FKeyRenderer::FCachedKeyDrawInformation::DrawExtra(const FGeometry& AllottedGeometry, int32 LayerId) const
+{
+	TViewModelPtr<IKeyExtension> KeyExtension = WeakKeyExtension.Pin();
+	if (KeyExtension)
+	{
+		return KeyExtension->CustomPaint(AllottedGeometry, LayerId);
+	}
+	return LayerId;
+}
+
+EViewDependentCacheFlags FKeyRenderer::FCachedKeyDrawInformation::UpdateViewIndependentData()
+{
+	TViewModelPtr<IKeyExtension> KeyExtension = WeakKeyExtension.Pin();
+	if (KeyExtension)
+	{
+		const bool bDataChanged = KeyExtension->UpdateCachedKeys(CachedKeys);
+
+		return bDataChanged ? EViewDependentCacheFlags::DataChanged : EViewDependentCacheFlags::None;
+	}
+	else if (CachedKeys.IsValid())
+	{
+		CachedKeys = nullptr;
+		return EViewDependentCacheFlags::DataChanged;
+	}
+
+	return EViewDependentCacheFlags::None;
+}
+
+void FKeyRenderer::FCachedKeyDrawInformation::CacheViewDependentData(const FKeyBatchParameters& Params)
+{
+	TViewModelPtr<IKeyExtension> KeyExtension = WeakKeyExtension.Pin();
+	if (!KeyExtension)
+	{
+		CachedKeys = nullptr;
+
+		DrawParams.Empty();
+		FramesInRange = {};
+		HandlesInRange = {};
+
+		LeadingKey.Reset();
+		TrailingKey.Reset();
+
+		PreserveToIndex = FramesInRange.Num();
+		NextUnhandledIndex = 0;
+		return;
+	}
+
+	if (EnumHasAnyFlags(Params.CacheState, EViewDependentCacheFlags::DataChanged | EViewDependentCacheFlags::ViewChanged | EViewDependentCacheFlags::ViewZoomed))
+	{
+		TArrayView<const FFrameTime> OldFramesInRange = FramesInRange;
+
+		FFrameTime KeyWidth = Params.TimeToPixel.PixelDeltaToFrame(Params.KeySizePx.X);
+
+		// Dialte the visible range by the size of a key
+		TRange<FFrameTime> DilatedVisibleRange = Params.VisibleRange;
+		DilatedVisibleRange.SetLowerBoundValue(DilatedVisibleRange.GetLowerBoundValue() - KeyWidth*.5f);
+		DilatedVisibleRange.SetUpperBoundValue(DilatedVisibleRange.GetUpperBoundValue() + KeyWidth*.5f);
+
+		// Gather all the key handles in this view range, plus the next out-of-range key for connecting bars
+		CachedKeys->GetKeysInRangeWithBounds(DilatedVisibleRange, &FramesInRange, &HandlesInRange, &LeadingKey, &TrailingKey);
+
+		if (LeadingKey.IsValid())
+		{
+			KeyExtension->DrawKeys(MakeArrayView(&LeadingKey.Handle, 1), TArrayView<FKeyDrawParams>(&LeadingKeyParams, 1));
+		}
+		if (TrailingKey.IsValid())
+		{
+			KeyExtension->DrawKeys(MakeArrayView(&TrailingKey.Handle, 1), TArrayView<FKeyDrawParams>(&TrailingKeyParams, 1));
+		}
+
+		bool bDrawnKeys = false;
+		if (!EnumHasAnyFlags(Params.CacheState, EViewDependentCacheFlags::DataChanged) && OldFramesInRange.Num() != 0 && FramesInRange.Num() != 0)
+		{
+			// Try and preserve draw params if possible
+			const int32 PreserveStart = Algo::LowerBound(OldFramesInRange, FramesInRange[0]);
+			const int32 PreserveEnd   = Algo::UpperBound(OldFramesInRange, FramesInRange.Last());
+
+			const int32 PreserveNum   = PreserveEnd - PreserveStart;
+			if (PreserveNum > 0)
+			{
+				TArray<FKeyDrawParams> NewDrawParams;
+
+				const int32 HeadNum = Algo::LowerBound(FramesInRange, OldFramesInRange[PreserveStart]);
+				if (HeadNum > 0)
+				{
+					NewDrawParams.SetNum(HeadNum);
+					KeyExtension->DrawKeys(HandlesInRange.Slice(0, HeadNum), NewDrawParams);
+				}
+
+				NewDrawParams.Append(DrawParams.GetData() + PreserveStart, PreserveNum);
+
+				const int32 TailStart = Algo::LowerBound(FramesInRange, OldFramesInRange[PreserveEnd-1]);
+				const int32 TailNum = FramesInRange.Num() - TailStart;
+
+				if (TailNum > 0)
+				{
+					NewDrawParams.SetNum(FramesInRange.Num());
+					KeyExtension->DrawKeys(HandlesInRange.Slice(TailStart, TailNum), TArrayView<FKeyDrawParams>(NewDrawParams).Slice(TailStart, TailNum));
+				}
+
+				DrawParams = MoveTemp(NewDrawParams);
+				bDrawnKeys = true;
+			}
+		}
+
+		if (!bDrawnKeys)
+		{
+			DrawParams.SetNum(FramesInRange.Num());
+
+			if (FramesInRange.Num())
+			{
+				// Draw these keys
+				KeyExtension->DrawKeys(HandlesInRange, DrawParams);
+			}
+		}
+
+		check(DrawParams.Num() == FramesInRange.Num() && FramesInRange.Num() == HandlesInRange.Num());
+	}
+
+	// Always reset the pointers to the current key that needs processing
+	PreserveToIndex = FramesInRange.Num();
+	NextUnhandledIndex = 0;
+}
+
+} // namespace UE::Sequencer
+
diff --git a/Engine/Source/Editor/SequencerCore/Private/MVVM/Views/SCompoundTrackLaneView.cpp b/Engine/Source/Editor/SequencerCore/Private/MVVM/Views/SCompoundTrackLaneView.cpp
new file mode 100644
--- /dev/null
+++ b/Engine/Source/Editor/SequencerCore/Private/MVVM/Views/SCompoundTrackLaneView.cpp
@@ -0,0 +1,115 @@
+// Copyright Epic Games, Inc. All Rights Reserved.
+
+#include "MVVM/Views/SCompoundTrackLaneView.h"
+#include "MVVM/Views/STrackLane.h"
+#include "MVVM/Extensions/ITrackLaneExtension.h"
+#include "Widgets/SWeakWidget.h"
+
+#include "TimeToPixel.h"
+
+namespace UE::Sequencer
+{
+
+SCompoundTrackLaneView::FSlot::FSlot(TSharedPtr<ITrackLaneWidget> InInterface, TWeakPtr<STrackLane> InOwningLane)
+	: Interface(InInterface)
+	, WeakOwningLane(InOwningLane)
+{}
+SCompoundTrackLaneView::FSlot::FSlot(TWeakPtr<ITrackLaneWidget> InWeakInterface, TWeakPtr<STrackLane> InOwningLane)
+	: WeakInterface(InWeakInterface)
+	, WeakOwningLane(InOwningLane)
+{}
+
+SCompoundTrackLaneView::SCompoundTrackLaneView()
+	: Children(this)
+{
+}
+
+SCompoundTrackLaneView::~SCompoundTrackLaneView()
+{
+}
+
+void SCompoundTrackLaneView::Construct(const FArguments& InArgs)
+{
+	TimeToPixelDelegate = InArgs._TimeToPixel;
+
+	// Has to be bound dynamically (otherwise it wouldn't update when section ranges change)
+	check(TimeToPixelDelegate.IsBound());
+
+	SetClipping(EWidgetClipping::ClipToBounds);
+}
+
+void SCompoundTrackLaneView::AddWeakWidget(TSharedPtr<ITrackLaneWidget> InWidget, TWeakPtr<STrackLane> InOwningLane)
+{
+	FSlot::FSlotArguments SlotArguments(MakeUnique<FSlot>(TWeakPtr<ITrackLaneWidget>(InWidget), InOwningLane));
+	SlotArguments.AttachWidget(
+		SNew(SWeakWidget).PossiblyNullContent(InWidget->AsWidget())
+	);
+	Children.AddSlot(MoveTemp(SlotArguments));
+}
+
+void SCompoundTrackLaneView::AddStrongWidget(TSharedPtr<ITrackLaneWidget> InWidget, TWeakPtr<STrackLane> InOwningLane)
+{
+	FSlot::FSlotArguments SlotArguments(MakeUnique<FSlot>(InWidget, InOwningLane));
+	SlotArguments.AttachWidget(InWidget->AsWidget());
+	Children.AddSlot(MoveTemp(SlotArguments));
+}
+
+void SCompoundTrackLaneView::OnArrangeChildren(const FGeometry& AllottedGeometry, FArrangedChildren& ArrangedChildren) const
+{
+	FTimeToPixel TimeToPixel = TimeToPixelDelegate.Execute(AllottedGeometry);
+
+	for (int32 WidgetIndex = 0; WidgetIndex < Children.Num(); ++WidgetIndex)
+	{
+		const FSlot&        Slot   = Children[WidgetIndex];
+		TSharedRef<SWidget> Widget = Slot.GetWidget();
+
+		EVisibility WidgetVisibility = Widget->GetVisibility();
+		if (ArrangedChildren.Accepts(WidgetVisibility))
+		{
+			TSharedPtr<ITrackLaneWidget> Interface = Slot.GetInterface();
+			TSharedPtr<STrackLane>       TrackLane = Slot.WeakOwningLane.Pin();
+
+			if (!Interface || !TrackLane)
+			{
+				continue;
+			}
+
+			FGeometry LaneGeometry = AllottedGeometry;
+
+			if (TSharedPtr<STrackLane> ParentLane = TrackLane->GetParentLane())
+			{
+				// If we have a parent track lane, make the geometry relative to that
+				LaneGeometry = AllottedGeometry.MakeChild(
+					FVector2D(AllottedGeometry.GetLocalSize().X, TrackLane->GetDesiredSize().Y),
+					FSlateLayoutTransform(FVector2D(0.f, TrackLane->GetVerticalPosition() - ParentLane->GetVerticalPosition()))
+				);
+			}
+
+			FTrackLaneScreenAlignment ScreenAlignment = Interface->GetAlignment(TimeToPixel, LaneGeometry);
+			if (ScreenAlignment.IsVisible())
+			{
+				FArrangedWidget ArrangedWidget = ScreenAlignment.ArrangeWidget(Widget, LaneGeometry);
+				ArrangedChildren.AddWidget(WidgetVisibility, ArrangedWidget);
+			}
+		}
+	}
+}
+
+FVector2D SCompoundTrackLaneView::ComputeDesiredSize(float) const
+{
+	return FVector2D(10.f, 10.f);
+}
+
+FChildren* SCompoundTrackLaneView::GetChildren()
+{
+	return &Children;
+}
+
+int32 SCompoundTrackLaneView::OnPaint(const FPaintArgs& Args, const FGeometry& AllottedGeometry, const FSlateRect& MyCullingRect, FSlateWindowElementList& OutDrawElements, int32 LayerId, const FWidgetStyle& InWidgetStyle, bool bParentEnabled) const
+{
+	SPanel::OnPaint(Args, AllottedGeometry, MyCullingRect, OutDrawElements, LayerId, InWidgetStyle, bParentEnabled);
+	return LayerId;
+}
+
+} // namespace UE::Sequencer
+
diff --git a/Engine/Source/Editor/SequencerCore/Private/MVVM/Views/SOutlinerView.cpp b/Engine/Source/Editor/SequencerCore/Private/MVVM/Views/SOutlinerView.cpp
--- a/Engine/Source/Editor/SequencerCore/Private/MVVM/Views/SOutlinerView.cpp
+++ b/Engine/Source/Editor/SequencerCore/Private/MVVM/Views/SOutlinerView.cpp
@@ -1104,7 +1104,7 @@ TSharedPtr<STrackLane> SOutlinerView::FindOrCreateParentLane(TViewModelPtr<IOutl
 				if (!TrackLane)
 				{
 					// Add a track slot for the nested track lane parent if it doesn't already exist
-					TrackLane = SNew(STrackLane, TrackArea->GetViewModel(), OutlinerExtItem, TrackAreaParameters, SharedThis(this));
+					TrackLane = SNew(STrackLane, TrackArea, OutlinerExtItem, nullptr, TrackAreaParameters, SharedThis(this));
 					TrackArea->AddTrackSlot(OutlinerExtItem, TrackLane);
 				}
 
@@ -1136,13 +1136,8 @@ void SOutlinerView::CreateTrackLanesForRow(TSharedRef<SOutlinerViewRow> InRow, T
 				if (!TrackLane)
 				{
 					// Add a track slot for the row
-					TrackLane = SNew(STrackLane, TrackArea->GetViewModel(), CurDataModel, Parameters, SharedThis(this));
-
 					TSharedPtr<STrackLane> ParentLane = FindOrCreateParentLane(CurDataModel);
-					if (ParentLane)
-					{
-						TrackLane->SetParent(ParentLane);
-					}
+					TrackLane = SNew(STrackLane, TrackArea, CurDataModel, ParentLane, Parameters, SharedThis(this));
 
 					TrackArea->AddTrackSlot(CurDataModel, TrackLane);
 				}
diff --git a/Engine/Source/Editor/SequencerCore/Private/MVVM/Views/STrackAreaView.cpp b/Engine/Source/Editor/SequencerCore/Private/MVVM/Views/STrackAreaView.cpp
--- a/Engine/Source/Editor/SequencerCore/Private/MVVM/Views/STrackAreaView.cpp
+++ b/Engine/Source/Editor/SequencerCore/Private/MVVM/Views/STrackAreaView.cpp
@@ -56,11 +56,21 @@ void STrackAreaView::Construct(const FArguments& InArgs, TWeakPtr<FTrackAreaView
 	// Space for the edit tool
 	InputStack.AddHandler(nullptr);
 
-	if (TSharedPtr<FTrackAreaViewModel> ViewModel = WeakViewModel.Pin())
-	{
-		InputStack.OnBeginCapture.AddSP(ViewModel.Get(), &FTrackAreaViewModel::LockEditTool);
-		InputStack.OnEndCapture.AddSP(ViewModel.Get(), &FTrackAreaViewModel::UnlockEditTool);
-	}
+	TSharedPtr<FTrackAreaViewModel> ViewModel = WeakViewModel.Pin();
+	check(ViewModel);
+
+	InputStack.OnBeginCapture.AddSP(ViewModel.Get(), &FTrackAreaViewModel::LockEditTool);
+	InputStack.OnEndCapture.AddSP(ViewModel.Get(), &FTrackAreaViewModel::UnlockEditTool);
+
+	// Set some default time to pixel
+	TimeToPixel = MakeShared<FTimeToPixel>(ViewModel->GetTimeToPixel(100.f));
+
+	SetClipping(EWidgetClipping::ClipToBoundsAlways);
+}
+
+const FTrackAreaViewLayers& STrackAreaView::GetPaintLayers() const
+{
+	return LayerIds;
 }
 
 TSharedPtr<FTrackAreaViewModel> STrackAreaView::GetViewModel() const
@@ -68,6 +78,11 @@ TSharedPtr<FTrackAreaViewModel> STrackAreaView::GetViewModel() const
 	return WeakViewModel.Pin();
 }
 
+TSharedPtr<FTimeToPixel> STrackAreaView::GetTimeToPixel() const
+{
+	return TimeToPixel;
+}
+
 FLinearColor STrackAreaView::BlendDefaultTrackColor(FLinearColor InColor)
 {
 	static FLinearColor BaseColor(FColor(71,71,71));
@@ -104,6 +119,14 @@ TSharedPtr<STrackLane> STrackAreaView::FindTrackSlot(const TViewModelPtr<IOutlin
 
 void STrackAreaView::OnArrangeChildren( const FGeometry& AllottedGeometry, FArrangedChildren& ArrangedChildren ) const
 {
+	TSharedPtr<FTrackAreaViewModel> TrackArea = WeakViewModel.Pin();
+	if (!TrackArea)
+	{
+		return;
+	}
+
+	*const_cast<STrackAreaView*>(this)->TimeToPixel = FTimeToPixel(TrackArea->GetTimeToPixel(AllottedGeometry.GetLocalSize().X));
+
 	for (int32 ChildIndex = 0; ChildIndex < Children.Num(); ++ChildIndex)
 	{
 		const FTrackAreaSlot& CurChild = Children[ChildIndex];
@@ -171,6 +194,9 @@ int32 STrackAreaView::OnPaint(const FPaintArgs& Args, const FGeometry& AllottedG
 		return LayerId;
 	}
 
+	// Reassign the time <-> pixel space for this frame
+	*GetTimeToPixel() = TrackArea->GetTimeToPixel(AllottedGeometry.GetLocalSize().X);
+
 	// TODO: TrackEditors are drawn in subclass, maybe should be done in viewmodel?
 
 	// paint the child widgets
@@ -179,15 +205,20 @@ int32 STrackAreaView::OnPaint(const FPaintArgs& Args, const FGeometry& AllottedG
 
 	const FPaintArgs NewArgs = Args.WithNewParent(this);
 
+	LayerIds.LaneBackgrounds = LayerId;
+
+	float MaxLayerId = LayerId;
 	for (int32 ChildIndex = 0; ChildIndex < ArrangedChildren.Num(); ++ChildIndex)
 	{
 		FArrangedWidget& CurWidget = ArrangedChildren[ChildIndex];
 		FSlateRect ChildClipRect = MyCullingRect.IntersectionWith( CurWidget.Geometry.GetLayoutBoundingRect() );
-		const int32 ThisWidgetLayerId = CurWidget.Widget->Paint( NewArgs, CurWidget.Geometry, ChildClipRect, OutDrawElements, LayerId + 2, InWidgetStyle, ShouldBeEnabled( bParentEnabled ) );
+		const int32 ThisWidgetLayerId = CurWidget.Widget->Paint( NewArgs, CurWidget.Geometry, ChildClipRect, OutDrawElements, LayerId, InWidgetStyle, ShouldBeEnabled( bParentEnabled ) );
 
-		LayerId = FMath::Max(LayerId, ThisWidgetLayerId);
+		MaxLayerId = FMath::Max(MaxLayerId, ThisWidgetLayerId);
 	}
 
+	LayerId = MaxLayerId;
+
 	if (const ISequencerEditTool* EditTool = TrackArea->GetEditTool())
 	{
 		LayerId = EditTool->OnPaint(AllottedGeometry, MyCullingRect, OutDrawElements, LayerId + 1);
@@ -217,9 +248,8 @@ int32 STrackAreaView::OnPaint(const FPaintArgs& Args, const FGeometry& AllottedG
 
 		if (DropFrameRange.IsSet())
 		{
-			FTimeToPixel TimeToPixel = TrackArea->GetTimeToPixel(AllottedGeometry);
-			DropMinX = TimeToPixel.FrameToPixel(DropFrameRange.GetValue().GetLowerBoundValue());
-			DropMaxX = TimeToPixel.FrameToPixel(DropFrameRange.GetValue().GetUpperBoundValue());
+			DropMinX = TimeToPixel->FrameToPixel(DropFrameRange.GetValue().GetLowerBoundValue());
+			DropMaxX = TimeToPixel->FrameToPixel(DropFrameRange.GetValue().GetUpperBoundValue());
 		}
 
 		// Top
@@ -412,6 +442,15 @@ FCursorReply STrackAreaView::OnCursorQuery( const FGeometry& MyGeometry, const F
 
 	if (TSharedPtr<FTrackAreaViewModel> TrackArea = WeakViewModel.Pin())
 	{
+		if (TSharedPtr<ITrackAreaHotspot> Hotspot = TrackArea->GetHotspot())
+		{
+			FCursorReply HotspotCursor = Hotspot->GetCursor();
+			if (HotspotCursor.IsEventHandled())
+			{
+				return HotspotCursor;
+			}
+		}
+
 		if (ISequencerEditTool* EditTool  = TrackArea->GetEditTool())
 		{
 			return EditTool->OnCursorQuery(MyGeometry, CursorEvent);
diff --git a/Engine/Source/Editor/SequencerCore/Private/MVVM/Views/STrackLane.cpp b/Engine/Source/Editor/SequencerCore/Private/MVVM/Views/STrackLane.cpp
--- a/Engine/Source/Editor/SequencerCore/Private/MVVM/Views/STrackLane.cpp
+++ b/Engine/Source/Editor/SequencerCore/Private/MVVM/Views/STrackLane.cpp
@@ -12,10 +12,11 @@
 #include "MVVM/SharedViewModelData.h"
 #include "MVVM/ViewModels/EditorViewModel.h"
 #include "MVVM/ViewModels/OutlinerViewModel.h"
-#include "MVVM/ViewModels/TrackAreaViewModel.h"
 #include "MVVM/ViewModels/ViewModel.h"
 #include "MVVM/ViewModels/ViewModelIterators.h"
+#include "MVVM/ViewModels/TrackAreaViewModel.h"
 #include "MVVM/Views/SOutlinerView.h"
+#include "MVVM/Views/STrackAreaView.h"
 
 namespace UE
 {
@@ -40,15 +41,22 @@ STrackLane::~STrackLane()
 {
 }
 
-void STrackLane::Construct(const FArguments& InArgs, TWeakPtr<FTrackAreaViewModel> InTrackAreaViewModel, TWeakViewModelPtr<IOutlinerExtension> InWeakOutlinerItem, const FTrackAreaParameters& InTrackParams, const TSharedRef<SOutlinerView>& InTreeView)
+void STrackLane::Construct(const FArguments& InArgs, TWeakPtr<STrackAreaView> InTrackAreaView, TWeakViewModelPtr<IOutlinerExtension> InWeakOutlinerItem, TSharedPtr<STrackLane> InParentLane, const FTrackAreaParameters& InTrackParams, const TSharedRef<SOutlinerView>& InTreeView)
 {
 	bWidgetsDirty = true;
 
-	WeakTrackAreaViewModel = InTrackAreaViewModel;
+	WeakTrackAreaView = InTrackAreaView;
 	WeakOutlinerItem = InWeakOutlinerItem;
 	TreeView = InTreeView;
+	ParentLane = InParentLane;
+
+	if (ParentLane)
+	{
+		ParentLane->WeakChildLanes.Add(SharedThis(this));
+	}
 
 	TrackParams = InTrackParams;
+	TimeToPixel = InTrackAreaView.Pin()->GetTimeToPixel();
 
 	FViewModelPtr OutlinerItem = GetOutlinerItem();
 	if (ensure(OutlinerItem))
@@ -68,6 +76,16 @@ TViewModelPtr<IOutlinerExtension> STrackLane::GetOutlinerItem() const
 	return WeakOutlinerItem.Pin();
 }
 
+TSharedPtr<STrackLane> STrackLane::GetParentLane() const
+{
+	return ParentLane;
+}
+
+TArrayView<const TWeakPtr<STrackLane>> STrackLane::GetChildLanes() const
+{
+	return WeakChildLanes;
+}
+
 bool STrackLane::IsPinned() const
 {
 	TSharedPtr<IPinnableExtension> Pinnable = WeakOutlinerItem.ImplicitPin();
@@ -84,8 +102,13 @@ void STrackLane::OnHierarchyUpdated()
 	RecreateWidgets();
 }
 
+TSharedPtr<ITrackLaneWidget> STrackLane::FindWidgetForModel(const FWeakViewModelPtr& InModel) const
+{
+	return WidgetsByModel.FindRef(InModel);
+}
+
 void STrackLane::RecreateWidgets()
-{	
+{
 	if (!bWidgetsDirty)
 	{
 		return;
@@ -93,44 +116,93 @@ void STrackLane::RecreateWidgets()
 
 	bWidgetsDirty = false;
 	Children.Empty();
+	WidgetsByModel.Reset();
+
+	WeakChildLanes.Remove(TWeakPtr<STrackLane>());
 
+	TSharedPtr<STrackAreaView>      TrackAreaView      = WeakTrackAreaView.Pin();
 	TSharedPtr<ITrackAreaExtension> TrackAreaExtension = WeakOutlinerItem.ImplicitPin();
-	TSharedPtr<FTrackAreaViewModel> TrackAreaViewModel = WeakTrackAreaViewModel.Pin();
-	if (!TrackAreaViewModel || !TrackAreaExtension)
+	TSharedPtr<SOutlinerView>       PinnedTree         = TreeView.Pin();
+	if (!TrackAreaView || !TrackAreaExtension || !PinnedTree)
 	{
 		return;
 	}
 
-	TArray<TSharedPtr<ITrackLaneWidget>> ChildrenWidgets;
+	FCreateTrackLaneViewParams ViewParams(TrackAreaView->GetViewModel()->GetEditor());
+	ViewParams.OwningTrackLane = SharedThis(this);
+	ViewParams.TimeToPixel = TimeToPixel;
 
 	// Construct views for this track lane
-
-	FCreateTrackLaneViewParams CreateLaneParams(TrackAreaViewModel->GetEditor());
 	for (TTypedIterator<ITrackLaneExtension, FViewModelVariantIterator> It(TrackAreaExtension->GetTrackAreaModelList()); It; ++It)
 	{
-		TSharedPtr<ITrackLaneWidget> NewView = It->CreateTrackLaneView(CreateLaneParams);
+		TViewModelPtr<ITrackLaneExtension> Model = *It;
+
+		TSharedPtr<ITrackLaneWidget> ParentView;
 
-		if (NewView)
+		if (ParentLane)
 		{
-			if (NewView->AcceptsChildren())
+			FViewModelPtr Parent = Model.AsModel()->GetParent();
+			// Find the parent model that exists in the parent lane
+			while (Parent && !ParentView)
 			{
-				for (const TViewModelPtr<ITrackLaneExtension>& Child : It.GetCurrentItem()->GetDescendantsOfType<ITrackLaneExtension>())
-				{
-					TSharedPtr<ITrackLaneWidget> ChildView = Child->CreateTrackLaneView(CreateLaneParams);
-					NewView->AddChildLane(ChildView);
-				}
+				ParentView = ParentLane->FindWidgetForModel(Parent);
+				Parent = Parent->GetParent();
 			}
-			ChildrenWidgets.Add(NewView);
 		}
-	}
 
-	// Add children to the widget
+		TSharedPtr<ITrackLaneWidget> NewView = Model->CreateTrackLaneView(ViewParams);
+		WidgetsByModel.Add(Model, NewView);
 
-	for (const TSharedPtr<ITrackLaneWidget>& ChildWidget : ChildrenWidgets)
+		// If we have a parent view, add the widget to that
+		if (ParentView)
+		{
+			ParentView->AddChildView(NewView, SharedThis(this));
+		}
+		// Otherwise we add it directly to this lane
+		else
+		{
+			FSlot::FSlotArguments SlotArguments(MakeUnique<FSlot>(NewView));
+			SlotArguments.AttachWidget(NewView->AsWidget());
+			Children.AddSlot(MoveTemp(SlotArguments));
+		}
+	}
+
+	// Construct views for top level children track lane
+	for (TTypedIterator<ITrackAreaExtension, FViewModelVariantIterator> TopLevelIt(TrackAreaExtension->GetTopLevelChildTrackAreaModels()); TopLevelIt; ++TopLevelIt)
 	{
-		FSlot::FSlotArguments SlotArguments(MakeUnique<FSlot>(ChildWidget));
-		SlotArguments.AttachWidget(ChildWidget->AsWidget());
-		Children.AddSlot(MoveTemp(SlotArguments));
+		TViewModelPtr<ITrackAreaExtension> InlineModel = *TopLevelIt;
+		if (InlineModel->GetTrackAreaParameters().LaneType != ETrackAreaLaneType::None)
+		{
+			// Add these widgets directly to this lane as well
+			for (TTypedIterator<ITrackLaneExtension, FViewModelVariantIterator> It(InlineModel->GetTrackAreaModelList()); It; ++It)
+			{
+				TViewModelPtr<ITrackLaneExtension> Model = *It;
+				TSharedPtr<ITrackLaneWidget> ParentView;
+				FViewModelPtr Parent = Model.AsModel()->GetParent();
+				// Find the parent model that exists in the parent lane
+				while (Parent && !ParentView)
+				{
+					ParentView = FindWidgetForModel(Parent);
+					Parent = Parent->GetParent();
+				}
+
+				TSharedPtr<ITrackLaneWidget> NewView = Model->CreateTrackLaneView(ViewParams);
+				WidgetsByModel.Add(Model, NewView);
+
+				// If we have a parent view, add the widget to that
+				if (ParentView)
+				{
+					ParentView->AddChildView(NewView, SharedThis(this));
+				}
+				// Otherwise we add it directly to this lane
+				else
+				{
+					FSlot::FSlotArguments SlotArguments(MakeUnique<FSlot>(NewView));
+					SlotArguments.AttachWidget(NewView->AsWidget());
+					Children.AddSlot(MoveTemp(SlotArguments));
+				}
+			}
+		}
 	}
 }
 
@@ -138,15 +210,13 @@ void STrackLane::OnArrangeChildren( const FGeometry& AllottedGeometry, FArranged
 {
 	for (int32 WidgetIndex = 0; WidgetIndex < Children.Num(); ++WidgetIndex)
 	{
-		const FSlot&        Slot = Children[WidgetIndex];
-		TSharedRef<SWidget> Widget    = Slot.GetWidget();
+		const FSlot&        Slot   = Children[WidgetIndex];
+		TSharedRef<SWidget> Widget = Slot.GetWidget();
 
 		EVisibility WidgetVisibility = Widget->GetVisibility();
-		TSharedPtr<FTrackAreaViewModel> TrackAreaViewModel = WeakTrackAreaViewModel.Pin();
-		if (TrackAreaViewModel.IsValid() && ArrangedChildren.Accepts(WidgetVisibility))
+		if (ArrangedChildren.Accepts(WidgetVisibility))
 		{
-			FTimeToPixel TimeToPixel = TrackAreaViewModel->GetTimeToPixel(AllottedGeometry);
-			FTrackLaneScreenAlignment ScreenAlignment = Slot.Interface->GetAlignment(TimeToPixel, AllottedGeometry);
+			FTrackLaneScreenAlignment ScreenAlignment = Slot.Interface->GetAlignment(*TimeToPixel, AllottedGeometry);
 			if (ScreenAlignment.IsVisible())
 			{
 				FArrangedWidget ArrangedWidget = ScreenAlignment.ArrangeWidget(Widget, AllottedGeometry);
@@ -178,12 +248,15 @@ int32 STrackLane::PaintLaneBackground(const FGeometry& AllottedGeometry, const F
 	static const FName BorderName("Sequencer.AnimationOutliner.DefaultBorder");
 	static const FName SelectionColorName("SelectionColor");
 
-	TViewModelPtr<IOutlinerExtension> OutlinerItem = WeakOutlinerItem.Pin();
-	if (!OutlinerItem || OutlinerItem->IsFilteredOut())
+	TSharedPtr<STrackAreaView>        TrackAreaView = WeakTrackAreaView.Pin();;
+	TViewModelPtr<IOutlinerExtension> OutlinerItem  = WeakOutlinerItem.Pin();
+	if (!TrackAreaView || !OutlinerItem || OutlinerItem->IsFilteredOut())
 	{
 		return LayerId;
 	}
 
+	const int32 PaintLayerId = TrackAreaView->GetPaintLayers().LaneBackgrounds;
+
 	TViewModelPtr<IHoveredExtension> Hoverable = OutlinerItem.ImplicitCast();
 
 	float TotalNodeHeight = OutlinerItem->GetOutlinerSizing().GetTotalHeight();
@@ -197,7 +270,7 @@ int32 STrackLane::PaintLaneBackground(const FGeometry& AllottedGeometry, const F
 
 		FSlateDrawElement::MakeBox(
 			OutDrawElements,
-			LayerId,
+			PaintLayerId,
 			AllottedGeometry.ToPaintGeometry(
 				FVector2D(0, 0),
 				FVector2D(AllottedGeometry.GetLocalSize().X, TotalNodeHeight)
@@ -228,7 +301,7 @@ int32 STrackLane::PaintLaneBackground(const FGeometry& AllottedGeometry, const F
 		{
 			FSlateDrawElement::MakeBox(
 				OutDrawElements,
-				LayerId,
+				PaintLayerId,
 				AllottedGeometry.ToPaintGeometry(
 					FVector2D(0, 0),
 					FVector2D(AllottedGeometry.GetLocalSize().X, TotalNodeHeight)
@@ -240,7 +313,7 @@ int32 STrackLane::PaintLaneBackground(const FGeometry& AllottedGeometry, const F
 		}
 	}
 
-	return LayerId + 1;
+	return LayerId;
 }
 
 void STrackLane::Tick(const FGeometry& AllottedGeometry, const double InCurrentTime, const float InDeltaTime)
@@ -277,8 +350,6 @@ void STrackLane::Tick(const FGeometry& AllottedGeometry, const double InCurrentT
 
 FVector2D STrackLane::ComputeDesiredSize(float LayoutScale) const
 {
-	TSharedPtr<FOutlinerViewModel> Outliner = WeakTrackAreaViewModel.Pin()->GetEditor()->GetOutliner();
-
 	TViewModelPtr<IOutlinerExtension>  OutlinerItem = WeakOutlinerItem.Pin();
 	TViewModelPtr<ITrackAreaExtension> TrackArea    = OutlinerItem.ImplicitCast();
 
@@ -309,6 +380,8 @@ FVector2D STrackLane::ComputeDesiredSize(float LayoutScale) const
 				}
 			}
 		}
+
+		Height -= Parameters.TrackLanePadding.Top + Parameters.TrackLanePadding.Bottom;
 	}
 
 	return FVector2D(100.f, Height);
@@ -321,6 +394,10 @@ void STrackLane::SetVerticalPosition(float InPosition)
 
 float STrackLane::GetVerticalPosition() const
 {
+	if (TViewModelPtr<ITrackAreaExtension> TrackArea = WeakOutlinerItem.ImplicitPin())
+	{
+		return Position + TrackArea->GetTrackAreaParameters().TrackLanePadding.Top;
+	}
 	return Position;
 }
 
@@ -393,11 +470,6 @@ void STrackLane::PositionParentTrackLanes(TViewModelPtr<IOutlinerExtension> InIt
 	}
 }
 
-void STrackLane::SetParent(TSharedPtr<STrackLane> InParentLane)
-{
-	ParentLane = InParentLane;
-}
-
 FReply STrackLane::OnMouseButtonDown(const FGeometry& MyGeometry, const FPointerEvent& MouseEvent)
 {
 	static float HitThreshold = 5.f;
diff --git a/Engine/Source/Editor/SequencerCore/Private/MVVM/Views/STrackLane.h b/Engine/Source/Editor/SequencerCore/Private/MVVM/Views/STrackLane.h
deleted file mode 100644
diff --git a/Engine/Source/Editor/SequencerCore/Public/MVVM/Extensions/IKeyExtension.h b/Engine/Source/Editor/SequencerCore/Public/MVVM/Extensions/IKeyExtension.h
new file mode 100644
--- /dev/null
+++ b/Engine/Source/Editor/SequencerCore/Public/MVVM/Extensions/IKeyExtension.h
@@ -0,0 +1,62 @@
+// Copyright Epic Games, Inc. All Rights Reserved.
+
+#pragma once
+
+#include "CoreTypes.h"
+#include "MVVM/ViewModelTypeID.h"
+#include "Curves/KeyHandle.h"
+#include "Misc/FrameTime.h"
+
+struct FGuid;
+struct FPointerEvent;
+struct FKeyDrawParams;
+struct FGeometry;
+class FCurveModel;
+
+namespace UE::Sequencer
+{
+
+struct FCachedKeys
+{
+	/** Handles to each of the keys */
+	TArray<FKeyHandle> KeyHandles;
+	/** Frame times for each of the keys in a data-dependent time-space */
+	TArray<FFrameTime> KeyTimes;
+
+	struct FCachedKey
+	{
+		FFrameTime Time;
+		FKeyHandle Handle = FKeyHandle::Invalid();
+
+		bool IsValid() const;
+		void Reset();
+	};
+
+	void GetKeysInRange(const TRange<FFrameTime>& Range, TArrayView<const FFrameTime>* OutKeyTimes, TArrayView<const FKeyHandle>* OutHandles) const;
+	void GetKeysInRangeWithBounds(const TRange<FFrameTime>& Range, TArrayView<const FFrameTime>* OutKeyTimes, TArrayView<const FKeyHandle>* OutHandles, FCachedKey* OutLeadingKey, FCachedKey* OutTrailingKey) const;
+};
+
+class SEQUENCERCORE_API IKeyExtension
+{
+public:
+
+	UE_SEQUENCER_DECLARE_VIEW_MODEL_TYPE_ID(IKeyExtension)
+
+
+	virtual ~IKeyExtension(){}
+
+	virtual bool UpdateCachedKeys(TSharedPtr<FCachedKeys>& OutCachedKeys) const = 0;
+
+	virtual bool GetFixedExtents(double& OutFixedMin, double& OutFixedMax) const = 0;
+
+	virtual int32 CustomPaint(const FGeometry& KeyGeometry, int32 LayerId) const = 0;
+
+	virtual void DrawKeys(TArrayView<const FKeyHandle> InKeyHandles, TArrayView<FKeyDrawParams> OutKeyDrawParams) = 0;
+
+	virtual TUniquePtr<FCurveModel> CreateCurveModel() = 0;
+};
+
+
+
+} // namespace UE::Sequencer
+
diff --git a/Engine/Source/Editor/SequencerCore/Public/MVVM/Extensions/ITrackAreaExtension.h b/Engine/Source/Editor/SequencerCore/Public/MVVM/Extensions/ITrackAreaExtension.h
--- a/Engine/Source/Editor/SequencerCore/Public/MVVM/Extensions/ITrackAreaExtension.h
+++ b/Engine/Source/Editor/SequencerCore/Public/MVVM/Extensions/ITrackAreaExtension.h
@@ -20,9 +20,14 @@ enum class ETrackAreaLaneType : uint8
 	None
 };
 
-struct SEQUENCERCORE_API FTrackAreaParameters
+struct FTrackAreaParameters
 {
-	ETrackAreaLaneType LaneType;
+	ETrackAreaLaneType LaneType = ETrackAreaLaneType::None;
+	struct
+	{
+		float Top = 0.f;
+		float Bottom = 0.f;
+	} TrackLanePadding;
 };
 
 /**
@@ -38,6 +43,7 @@ public:
 
 	virtual FTrackAreaParameters GetTrackAreaParameters() const = 0;
 	virtual FViewModelVariantIterator GetTrackAreaModelList() const = 0;
+	virtual FViewModelVariantIterator GetTopLevelChildTrackAreaModels() const { return FViewModelVariantIterator(); }
 
 	template<typename T>
 	TTypedIterator<T, FViewModelVariantIterator> GetTrackAreaModelListAs() const
diff --git a/Engine/Source/Editor/SequencerCore/Public/MVVM/Extensions/ITrackLaneExtension.h b/Engine/Source/Editor/SequencerCore/Public/MVVM/Extensions/ITrackLaneExtension.h
--- a/Engine/Source/Editor/SequencerCore/Public/MVVM/Extensions/ITrackLaneExtension.h
+++ b/Engine/Source/Editor/SequencerCore/Public/MVVM/Extensions/ITrackLaneExtension.h
@@ -3,8 +3,7 @@
 #pragma once
 
 #include "MVVM/ViewModelTypeID.h"
-#include "MVVM/ViewModels/EditorViewModel.h"
-#include "Widgets/SCompoundWidget.h"
+#include "Types/SlateEnums.h"
 
 struct FGeometry;
 
@@ -18,6 +17,9 @@ namespace UE
 namespace Sequencer
 {
 
+class FViewModel;
+class STrackLane;
+class FEditorViewModel;
 struct FTrackLaneScreenAlignment;
 
 struct SEQUENCERCORE_API FTrackLaneVerticalArrangement
@@ -124,7 +126,7 @@ public:
 	/**
 	 * Add a new child to this lane
 	 */
-	virtual void AddChildLane(TSharedPtr<ITrackLaneWidget> ChildWidget) {}
+	virtual void AddChildView(TSharedPtr<ITrackLaneWidget> ChildWidget, TWeakPtr<STrackLane> InWeakOwningLane) {}
 };
 
 /** Parameters for creating a track lane widget */
@@ -135,6 +137,12 @@ struct SEQUENCERCORE_API FCreateTrackLaneViewParams
 	{}
 
 	const TSharedPtr<FEditorViewModel> Editor;
+
+	TSharedPtr<FViewModel> ParentModel;
+
+	TSharedPtr<STrackLane> OwningTrackLane;
+
+	TSharedPtr<FTimeToPixel> TimeToPixel;
 };
 
 /** Extension for view-models that can create track lanes in the track area */
diff --git a/Engine/Source/Editor/SequencerCore/Public/MVVM/ViewModelPtr.h b/Engine/Source/Editor/SequencerCore/Public/MVVM/ViewModelPtr.h
--- a/Engine/Source/Editor/SequencerCore/Public/MVVM/ViewModelPtr.h
+++ b/Engine/Source/Editor/SequencerCore/Public/MVVM/ViewModelPtr.h
@@ -60,6 +60,19 @@ struct FViewModelPtr
 		return *this;
 	}
 
+	/** Construction from a shared ptr */
+	template<typename ModelType>
+	FViewModelPtr(const TSharedRef<ModelType>& InModel)
+		: Model(InModel)
+	{}
+	/** Assignment from a shared ptr */
+	template<typename ModelType>
+	FViewModelPtr& operator=(const TSharedRef<ModelType>& InModel)
+	{
+		Model = InModel;
+		return *this;
+	}
+
 	/** Construction from a raw ptr */
 	FViewModelPtr(FViewModel* InModel)
 		: Model(InModel->AsShared())
@@ -224,15 +237,28 @@ struct TViewModelPtr : FViewModelPtr
 	/** Construction from a regular view model ptr */
 	template<typename SharedType>
 	TViewModelPtr(const TSharedPtr<SharedType>& InModel)
-		: Extension(nullptr)
 	{
-		Assign(InModel);
+		Assign(InModel, InModel.Get());
 	}
 	/** Assignment from a regular view model ptr */
 	template<typename SharedType>
 	TViewModelPtr& operator=(const TSharedPtr<SharedType>& InModel)
 	{
-		Assign(InModel);
+		Assign(InModel, InModel.Get());
+		return *this;
+	}
+
+	/** Construction from a regular view model ref */
+	template<typename SharedType>
+	TViewModelPtr(const TSharedRef<SharedType>& InModel)
+	{
+		Assign(InModel, &InModel.Get());
+	}
+	/** Assignment from a regular view model ptr */
+	template<typename SharedType>
+	TViewModelPtr& operator=(const TSharedRef<SharedType>& InModel)
+	{
+		Assign(InModel, &InModel.Get());
 		return *this;
 	}
 
@@ -304,6 +330,14 @@ struct TViewModelPtr : FViewModelPtr
 		return TSharedPtr<ExtensionType>(Model, Extension);
 	}
 
+	/**
+	 * Conversion to a normal weak pointer of this extension type
+	 */
+	operator TWeakPtr<ExtensionType>() const
+	{
+		return TSharedPtr<ExtensionType>(Model, Extension);
+	}
+
 public:
 
 	/**
diff --git a/Engine/Source/Editor/SequencerCore/Public/MVVM/ViewModels/TrackAreaViewModel.h b/Engine/Source/Editor/SequencerCore/Public/MVVM/ViewModels/TrackAreaViewModel.h
--- a/Engine/Source/Editor/SequencerCore/Public/MVVM/ViewModels/TrackAreaViewModel.h
+++ b/Engine/Source/Editor/SequencerCore/Public/MVVM/ViewModels/TrackAreaViewModel.h
@@ -36,7 +36,7 @@ public:
 	TSharedPtr<FEditorViewModel> GetEditor() const;
 
 	/** Create a time/pixel converter for the given geometry */
-	FTimeToPixel GetTimeToPixel(const FGeometry& AllottedGeometry) const;
+	FTimeToPixel GetTimeToPixel(float TrackAreaWidth) const;
 
 	/** Gets the current tick resolution of the editor */
 	virtual FFrameRate GetTickResolution() const;
diff --git a/Engine/Source/Editor/SequencerCore/Public/MVVM/Views/KeyDrawParams.h b/Engine/Source/Editor/SequencerCore/Public/MVVM/Views/KeyDrawParams.h
new file mode 100644
--- /dev/null
+++ b/Engine/Source/Editor/SequencerCore/Public/MVVM/Views/KeyDrawParams.h
@@ -0,0 +1,56 @@
+// Copyright Epic Games, Inc. All Rights Reserved.
+
+#pragma once
+
+#include "Math/Color.h"
+
+struct FSlateBrush;
+
+enum class EKeyConnectionStyle
+{
+	None = 0,
+	Dotted = 1 << 0,
+	Dashed = 1 << 1,
+	Solid  = 1 << 2,
+};
+ENUM_CLASS_FLAGS(EKeyConnectionStyle)
+
+/**
+ * Structure defining how a key should be drawn
+ */
+struct FKeyDrawParams
+{
+	FKeyDrawParams()
+		: BorderBrush(nullptr), FillBrush(nullptr)
+		, BorderTint(FLinearColor::White), FillTint(FLinearColor::White)
+		, FillOffset(0.f, 0.f)
+		, ConnectionStyle(EKeyConnectionStyle::None)
+	{}
+
+	friend bool operator==(const FKeyDrawParams& A, const FKeyDrawParams& B)
+	{
+		return A.BorderBrush == B.BorderBrush && A.FillBrush == B.FillBrush && A.BorderTint == B.BorderTint && A.FillTint == B.FillTint && A.FillOffset == B.FillOffset && A.ConnectionStyle == B.ConnectionStyle;
+	}
+	friend bool operator!=(const FKeyDrawParams& A, const FKeyDrawParams& B)
+	{
+		return !(A == B);
+	}
+
+	/** Brush to use for drawing the key's border */
+	const FSlateBrush* BorderBrush;
+
+	/** Brush to use for drawing the key's filled area */
+	const FSlateBrush* FillBrush;
+
+	/** Tint to be used for the key's border */
+	FLinearColor BorderTint;
+
+	/** Tint to be used for the key's filled area */
+	FLinearColor FillTint;
+
+	/** The amount to offset the fill brush from the keys center */
+	FVector2D FillOffset;
+
+	/** Flags denoting how to draw connecting lines between this key and the next */
+	EKeyConnectionStyle ConnectionStyle;
+};
diff --git a/Engine/Source/Editor/SequencerCore/Public/MVVM/Views/KeyRenderer.h b/Engine/Source/Editor/SequencerCore/Public/MVVM/Views/KeyRenderer.h
new file mode 100644
--- /dev/null
+++ b/Engine/Source/Editor/SequencerCore/Public/MVVM/Views/KeyRenderer.h
@@ -0,0 +1,341 @@
+// Copyright Epic Games, Inc. All Rights Reserved.
+
+#pragma once
+
+#include "Math/Color.h"
+#include "Math/Vector2D.h"
+#include "Math/Range.h"
+#include "Containers/ArrayView.h"
+#include "Misc/FrameNumber.h"
+#include "Misc/FrameTime.h"
+#include "MVVM/Views/KeyDrawParams.h"
+#include "MVVM/Extensions/IKeyExtension.h"
+#include "Rendering/RenderingCommon.h"
+#include "Curves/KeyHandle.h"
+#include "TimeToPixel.h"
+#include "MVVM/ViewModelPtr.h"
+
+struct FSlateBrush;
+class FWidgetStyle;
+class FSlateWindowElementList;
+
+namespace UE::Sequencer
+{
+
+/** Key selection preview state - mirrored from FSequencer while selection concepts are ported to SequencerCore */
+enum class EKeySelectionPreviewState : uint8
+{
+	Undefined,
+	Selected,
+	NotSelected
+};
+
+/**
+ * Paint arguments required for painting keys on a sequencer track
+ */
+struct FKeyRendererPaintArgs
+{
+	/** Draw elements to paint onto */
+	FSlateWindowElementList* DrawElements;
+
+	/** The amount to throb selected keys by */
+	float KeyThrobValue = 0.f;
+
+	/** The amount to throb selected sections by */
+	float SectionThrobValue = 0.f;
+
+	/** Fixed amount to throb newly created keys by */
+	FVector2D ThrobAmount = FVector2D(12.f, 12.f);
+
+	/** Default draw effect */
+	ESlateDrawEffect DrawEffects;
+
+	/** Selection color to use for selected keys */
+	FLinearColor SelectionColor;
+
+	/** Curve color to paint underlying curves with */
+	FLinearColor CurveColor = FLinearColor(0.8f, 0.8f, 0.f, 0.7);
+
+	/** Key bar color */
+	FLinearColor KeyBarColor = FColor(160, 160, 160);
+};
+
+/** Flag enum signifying how the cache has changed since it was last generated */
+enum class EViewDependentCacheFlags : uint8
+{
+	None             = 0,       // The cache is still entirely valid, simply redraw the keys
+	DataChanged      = 1 << 0,  // The underlying keyframes have changed - everything needs regenerating
+	KeyStateChanged  = 1 << 1,  // The selection, hover or preview selection state of the keys has changed
+	ViewChanged      = 1 << 2,  // The view range has changed - view dependent data needs regenerating, but some cache data may be preserved
+	ViewZoomed       = 1 << 3,  // The view range has been zoomed - view dependent data needs regenerating, no key grouping can be preserved
+
+	Empty            = 1 << 4,  // The view range is empty and there is nothing to show
+
+	All = DataChanged | KeyStateChanged | ViewChanged | ViewZoomed,
+
+};
+
+/** Interface for defining selection states for keys */
+class IKeyRendererInterface
+{
+public:
+	virtual ~IKeyRendererInterface(){}
+
+	virtual bool HasAnySelectedKeys() const = 0;
+	virtual bool HasAnyPreviewSelectedKeys() const = 0;
+	virtual bool HasAnyHoveredKeys() const = 0;
+
+	virtual bool IsKeySelected(const TViewModelPtr<IKeyExtension>&, FKeyHandle) const = 0;
+	virtual bool IsKeyHovered(const TViewModelPtr<IKeyExtension>&, FKeyHandle) const = 0;
+	virtual EKeySelectionPreviewState GetPreviewSelectionState(const TViewModelPtr<IKeyExtension>&, FKeyHandle) const = 0;
+};
+
+struct FKeyBatchParameters
+{
+	FKeyBatchParameters(const FTimeToPixel& InTimeToPixel)
+		: ClientInterface(nullptr)
+		, TimeToPixel(InTimeToPixel)
+		, KeySizePx(12.f, 12.f)
+		, ValidPlayRangeMin(TNumericLimits<int32>::Min())
+		, ValidPlayRangeMax(TNumericLimits<int32>::Max())
+		, CacheState(EViewDependentCacheFlags::None)
+		, bShowCurve(false)
+		, bShowKeyBars(true)
+		, bCollapseChildren(false)
+	{}
+
+	/** Interface ptr for defining selection states for keys. Can be null if selection is not supported. */
+	IKeyRendererInterface* ClientInterface;
+
+	/** Time to pixel convertor */
+	FTimeToPixel TimeToPixel;
+
+	/** Visible range to clip rendering to */
+	TRange<FFrameTime> VisibleRange;
+
+	/** Screen-space size of keys in px */
+	FVector2D KeySizePx;
+
+	/** Valid range for keys - keys outside this range are drawn de-saturated */
+	FFrameNumber ValidPlayRangeMin;
+	FFrameNumber ValidPlayRangeMax;
+
+	/** Cache state that defines whether and how we need to update any cached view (in)dependent data */
+	EViewDependentCacheFlags CacheState;
+
+	/** Whether to show a curve underneath the keys */
+	uint8 bShowCurve : 1;
+
+	/** Whether to show key bars connecting each key */
+	uint8 bShowKeyBars : 1;
+
+	/** Whether to collapse all children of the model into this key renderer as well, drawing all nested keys as groups within this renderer */
+	uint8 bCollapseChildren : 1;
+};
+
+/**
+ * Utility class for efficiently drawing large numbers of keys on a track lane.
+ */
+struct SEQUENCERCORE_API FKeyRenderer
+{
+	/** Flag enum signifying states for a particular key or group of keys */
+	enum class EKeyRenderingFlags
+	{
+		None               = 0,
+		PartialKey         = 1 << 0,  // Indicates that this key comprises multiple keys of different types, or a partially keyed collapsed channel
+		Selected           = 1 << 1,  // Only if.NumSelected == Key.TotalNumKeys
+		PreviewSelected    = 1 << 2,  // Only if NumPreviewSelected == NumKeys
+		PreviewNotSelected = 1 << 3,  // Only if NumPreviewNotSelected == NumKeys
+		AnySelected        = 1 << 4,  // If any are selected
+		Hovered            = 1 << 5,  // Only if NumKeys == NumHovered
+		Overlaps           = 1 << 6,  // If NumKeys > 1
+		OutOfRange         = 1 << 7,  // If any of the keys fall outside of the valid range
+	};
+
+	enum class EKeyBarRenderingFlags
+	{
+		None               = 0,
+		Hovered            = 1 << 0,
+		OutOfRange         = 1 << 1,
+	};
+
+	struct FKeysForModel
+	{
+		FViewModelPtr Model;
+		TArray<FKeyHandle> Keys;
+	};
+
+	/**
+	 * Initialize this renderer with a generic model
+	 */
+	void Initialize(const FViewModelPtr& InViewModel);
+
+	/**
+	 * Update any invalidated data required for drawing keys
+	 */
+	void Update(const FKeyBatchParameters& Params, const FGeometry& AllottedGeometry) const;
+
+	/**
+	 * Draw this batch's curve
+	 */
+	int32 DrawCurve(const FKeyBatchParameters& Params, const FGeometry& AllottedGeometry, const FSlateRect& MyCullingRect, const FKeyRendererPaintArgs& Args, int32 LayerId) const;
+
+	/**
+	 * Draw this batch
+	 */
+	int32 Draw(const FKeyBatchParameters& Params, const FGeometry& AllottedGeometry, const FSlateRect& MyCullingRect, const FKeyRendererPaintArgs& Args, int32 LayerId) const;
+
+	/**
+	 * Retrieve keys at the specified position
+	 */
+	void HitTestKeys(const FFrameTime& Time, TArray<FKeysForModel>& OutAllKeys) const;
+
+	/**
+	 * Retrieve keys at the specified position
+	 */
+	bool HitTestKeyBar(const FFrameTime& Time, TArray<FKeysForModel>& LeadingKeys, TArray<FKeysForModel>& TrailingKeys) const;
+
+	/**
+	 * Returns true if this renderer has a curve, false otherwise
+	 */
+	bool HasCurve() const;
+
+private:
+
+	/** Container that caches the key positions for a given key area, along with those that overlap the current visible range */
+	struct FCachedKeyDrawInformation
+	{
+		/** Construction from the key area this represents */
+		FCachedKeyDrawInformation(TViewModelPtr<IKeyExtension> KeyExtension)
+			: NextUnhandledIndex(0)
+			, PreserveToIndex(0)
+			, WeakKeyExtension(KeyExtension)
+		{
+			KeyExtension->UpdateCachedKeys(CachedKeys);
+		}
+
+
+		/**
+		 * Attempt to update data that is not dependent upon the current view
+		 * @return EViewDependentCacheFlags::DataChanged if view-indipendent data was updated, EViewDependentCacheFlags::None otherwise
+		 */
+		EViewDependentCacheFlags UpdateViewIndependentData();
+
+
+		/**
+		 * Ensure that view-dependent data (such as which keys need drawing and how) is up to date
+		 *
+		 * @param Params    The batch parameters
+		 */
+		void CacheViewDependentData(const FKeyBatchParameters& Params);
+
+
+		/**
+		 * Draw additional content in addition to keys for a particular IKeyArea
+		 *
+		 * @param Painter			 The painter to add the created geometry
+		 * @param KeyGeometry		 Allocated geometry to draw in
+		 */
+		int32 DrawExtra(const FGeometry& KeyGeometry, int32 LayerId) const;
+
+	public:
+
+		/** Index into the array views for the next unhandled key */
+		int32 NextUnhandledIndex;
+
+		/** Index into the array views for the first index proceeding a preserved range (TimesInRange.Num() if not) */
+		int32 PreserveToIndex;
+
+		/** Cached array view retrieved from CachedKeys for the key frames that overlap the current time */
+		TArrayView<const FFrameTime> FramesInRange;
+		/** Cached array view retrieved from CachedKeys for the key handles that overlap the current time */
+		TArrayView<const FKeyHandle> HandlesInRange;
+		/** Draw params for each of the keys visible on screen */
+		TArray<FKeyDrawParams> DrawParams;
+
+		FKeyDrawParams LeadingKeyParams;
+		FKeyDrawParams TrailingKeyParams;
+
+		FCachedKeys::FCachedKey LeadingKey;
+		FCachedKeys::FCachedKey TrailingKey;
+
+		/** Construction from the key area this represents */
+		TSharedPtr<FCachedKeys> CachedKeys;
+
+		TWeakViewModelPtr<IKeyExtension> WeakKeyExtension;
+	};
+
+
+	/** Cached parameters for drawing a single key */
+	struct FKey
+	{
+		/** Paint parameters for this key */
+		FKeyDrawParams Params;
+
+		/** The tick range that this key occupies (significant when this FKey represents multiple overlapping keys) */
+		FFrameTime KeyTickStart, KeyTickEnd;
+
+		/** The time that this key should be drawn - represents the average time for overlapping keys */
+		FFrameTime FinalKeyPosition;
+
+		/** Flags that specify how to draw this key */
+		EKeyRenderingFlags Flags = EKeyRenderingFlags::None;
+	};
+
+
+	/** Cached parameters for drawing a single key */
+	struct FKeyBar
+	{
+		/** The tick range that this key occupies (significant when this FKey represents multiple overlapping keys) */
+		FFrameTime StartTime, EndTime;
+
+		/** Flags that specify how to draw this key bar */
+		EKeyBarRenderingFlags Flags = EKeyBarRenderingFlags::None;
+
+		/** Key bar connection style */
+		EKeyConnectionStyle ConnectionStyle = EKeyConnectionStyle::None;
+	};
+
+	/**
+	 * Cache all the key extensions from our model
+	 */
+	void CacheKeyExtensions(const FKeyBatchParameters& Params) const;
+
+	/**
+	 * Attempt to update data that is not dependent upon the current view
+	 * @return EViewDependentCacheFlags::DataChanged if view-indipendent data was updated, EViewDependentCacheFlags::None otherwise
+	 */
+	EViewDependentCacheFlags UpdateViewIndependentData() const;
+
+
+	/**
+	 * Ensure that view-dependent data (such as which keys need drawing and how) is up to date
+	 */
+	void UpdateViewDependentData(const FKeyBatchParameters& Params, const FGeometry& AllottedGeometry) const;
+
+private:
+
+	/** The model that we are drawing keys for */
+	FWeakViewModelPtr WeakViewModel;
+
+	/** Array of cached draw info for each of the key areas that comprise this batch */
+	mutable TArray<FCachedKeyDrawInformation, TInlineAllocator<1>> KeyDrawInfo;
+
+	/** Computed final draw info for keys */
+	mutable TArray<FKey> PrecomputedKeys;
+
+	/** Computed final draw info for key bars */
+	mutable TArray<FKeyBar> PrecomputedKeyBars;
+
+	/** Computed final draw curve info */
+	mutable TArray<TTuple<double, double>> PrecomputedCurve;
+
+	/** Cached width of a key in frames */
+	mutable FFrameTime KeyWidthInFrames;
+};
+
+ENUM_CLASS_FLAGS(EViewDependentCacheFlags);
+ENUM_CLASS_FLAGS(FKeyRenderer::EKeyRenderingFlags);
+
+} // namespace UE::Sequencer
+
diff --git a/Engine/Source/Editor/SequencerCore/Public/MVVM/Views/SCompoundTrackLaneView.h b/Engine/Source/Editor/SequencerCore/Public/MVVM/Views/SCompoundTrackLaneView.h
new file mode 100644
--- /dev/null
+++ b/Engine/Source/Editor/SequencerCore/Public/MVVM/Views/SCompoundTrackLaneView.h
@@ -0,0 +1,80 @@
+// Copyright Epic Games, Inc. All Rights Reserved.
+
+#pragma once
+
+#include "CoreTypes.h"
+
+#include "Widgets/DeclarativeSyntaxSupport.h"
+#include "Widgets/SPanel.h"
+#include "Layout/Children.h"
+
+struct FTimeToPixel;
+struct FGeometry;
+
+class FPaintArgs;
+class FSlateWindowElementList;
+class FTrackAreaViewModel;
+
+namespace UE::Sequencer
+{
+
+class STrackLane;
+class ITrackLaneWidget;
+
+DECLARE_DELEGATE_RetVal_OneParam(FTimeToPixel, FGetTimeToPixel, const FGeometry&);
+
+/**
+ * 
+ */
+class SEQUENCERCORE_API SCompoundTrackLaneView
+	: public SPanel
+{
+public:
+
+	SLATE_BEGIN_ARGS(SCompoundTrackLaneView){}
+
+		SLATE_ARGUMENT(FGetTimeToPixel, TimeToPixel)
+
+	SLATE_END_ARGS()
+
+	SCompoundTrackLaneView();
+	~SCompoundTrackLaneView();
+
+	void Construct(const FArguments& InArgs);
+
+	void AddWeakWidget(TSharedPtr<ITrackLaneWidget> InWidget, TWeakPtr<STrackLane> InOwningLane);
+	void AddStrongWidget(TSharedPtr<ITrackLaneWidget> InWidget, TWeakPtr<STrackLane> InOwningLane);
+
+	/*~ SPanel Interface */
+	void OnArrangeChildren( const FGeometry& AllottedGeometry, FArrangedChildren& ArrangedChildren ) const override;
+	FVector2D ComputeDesiredSize(float) const override;
+	FChildren* GetChildren() override;
+
+	virtual int32 OnPaint(const FPaintArgs& Args, const FGeometry& AllottedGeometry, const FSlateRect& MyCullingRect, FSlateWindowElementList& OutDrawElements, int32 LayerId, const FWidgetStyle& InWidgetStyle, bool bParentEnabled) const override;
+
+private:
+
+	struct FSlot : TSlotBase<FSlot>
+	{
+		FSlot(TSharedPtr<ITrackLaneWidget> InInterface, TWeakPtr<STrackLane> InOwningLane);
+		FSlot(TWeakPtr<ITrackLaneWidget> InWeakInterface, TWeakPtr<STrackLane> InOwningLane);
+
+		TSharedPtr<ITrackLaneWidget> GetInterface() const
+		{
+			return Interface ? Interface : WeakInterface.Pin();
+		}
+
+		TSharedPtr<ITrackLaneWidget> Interface;
+		TWeakPtr<ITrackLaneWidget> WeakInterface;
+
+		TWeakPtr<STrackLane> WeakOwningLane;
+	};
+
+	/** All the widgets in the panel */
+	TPanelChildren<FSlot> Children;
+
+	FGetTimeToPixel TimeToPixelDelegate;
+};
+
+} // namespace UE::Sequencer
+
diff --git a/Engine/Source/Editor/SequencerCore/Public/MVVM/Views/STrackAreaView.h b/Engine/Source/Editor/SequencerCore/Public/MVVM/Views/STrackAreaView.h
--- a/Engine/Source/Editor/SequencerCore/Public/MVVM/Views/STrackAreaView.h
+++ b/Engine/Source/Editor/SequencerCore/Public/MVVM/Views/STrackAreaView.h
@@ -18,6 +18,7 @@ class FArrangedChildren;
 class FPaintArgs;
 class FSlateWindowElementList;
 class ISequencerEditTool;
+struct FTimeToPixel;
 
 namespace UE
 {
@@ -46,6 +47,12 @@ public:
 };
 
 
+struct FTrackAreaViewLayers
+{
+	int32 LaneBackgrounds = 0;
+};
+
+
 /**
  * The area where tracks( rows of sections ) are displayed.
  */
@@ -70,6 +77,8 @@ public:
 	
 	void SetVirtualPosition(float InVirtualTop);
 
+	TSharedPtr<FTimeToPixel> GetTimeToPixel() const;
+
 public:
 
 	/** Empty the track area */
@@ -93,6 +102,7 @@ public:
 	void SetIsPinned(bool bInIsPinned) { bIsPinned = bInIsPinned; }
 	bool IsPinned() const { return bIsPinned; }
 
+	const FTrackAreaViewLayers& GetPaintLayers() const;
 	static FLinearColor BlendDefaultTrackColor(FLinearColor InColor);
 
 public:
@@ -151,14 +161,18 @@ protected:
 	/** Weak pointer to the dropped node */
 	TWeakViewModelPtr<IOutlinerExtension> WeakDroppedItem;
 
+	mutable TSharedPtr<FTimeToPixel> TimeToPixel;
+
 	float VirtualTop;
 
-	/** Whether the dropped node is allowed to be dropped onto */
-	bool bAllowDrop;
+	mutable FTrackAreaViewLayers LayerIds;
 
 	/** The frame range of the section about to be dropped */
 	TOptional<TRange<FFrameNumber>> DropFrameRange;
 
+	/** Whether the dropped node is allowed to be dropped onto */
+	bool bAllowDrop;
+
 	/** Whether this TrackArea is for pinned nodes or non-pinned nodes */
 	bool bShowPinnedNodes;
 
diff --git a/Engine/Source/Editor/SequencerCore/Public/MVVM/Views/STrackLane.h b/Engine/Source/Editor/SequencerCore/Public/MVVM/Views/STrackLane.h
new file mode 100644
--- /dev/null
+++ b/Engine/Source/Editor/SequencerCore/Public/MVVM/Views/STrackLane.h
@@ -0,0 +1,140 @@
+// Copyright Epic Games, Inc. All Rights Reserved.
+
+#pragma once
+
+#include "CoreTypes.h"
+
+#include "Widgets/DeclarativeSyntaxSupport.h"
+#include "Widgets/SPanel.h"
+#include "Layout/Children.h"
+
+#include "MVVM/Extensions/ITrackAreaExtension.h"
+#include "MVVM/Extensions/ITrackLaneExtension.h"
+#include "MVVM/Extensions/IOutlinerExtension.h"
+
+class FPaintArgs;
+class FSlateWindowElementList;
+class FTrackAreaViewModel;
+
+namespace UE
+{
+namespace Sequencer
+{
+
+class SOutlinerView;
+class STrackAreaView;
+
+/**
+ * A wrapper widget responsible for positioning a track lane within the section area
+ */
+class SEQUENCERCORE_API STrackLane
+	: public SPanel
+{
+public:
+	SLATE_BEGIN_ARGS(STrackLane){}
+	SLATE_END_ARGS()
+
+	STrackLane();
+	~STrackLane();
+
+	/** Construct this widget */
+	void Construct(const FArguments& InArgs, TWeakPtr<STrackAreaView> InTrackAreaView, TWeakViewModelPtr<IOutlinerExtension> InWeakOutlinerItem, TSharedPtr<STrackLane> InParentLane, const FTrackAreaParameters& InTrackParams, const TSharedRef<SOutlinerView>& InTreeView);
+
+	/*~ SWidget Interface */
+	int32 OnPaint(const FPaintArgs& Args, const FGeometry& AllottedGeometry, const FSlateRect& MyCullingRect, FSlateWindowElementList& OutDrawElements, int32 LayerId, const FWidgetStyle& InWidgetStyle, bool bParentEnabled) const override;
+	void Tick(const FGeometry& AllottedGeometry, const double InCurrentTime, const float InDeltaTime) override;
+	FReply OnMouseButtonDown(const FGeometry& MyGeometry, const FPointerEvent& MouseEvent) override;
+	FReply OnMouseButtonUp(const FGeometry& MyGeometry, const FPointerEvent& MouseEvent) override;
+	FReply OnMouseMove(const FGeometry& MyGeometry, const FPointerEvent& MouseEvent) override;
+	FCursorReply OnCursorQuery(const FGeometry& MyGeometry, const FPointerEvent& CursorEvent) const override;
+
+	/*~ SPanel Interface */
+	void OnArrangeChildren( const FGeometry& AllottedGeometry, FArrangedChildren& ArrangedChildren ) const override;
+	FVector2D ComputeDesiredSize(float) const override;
+	FChildren* GetChildren() override { return &Children; }
+
+	/** Gets the outliner view-model that created this track lane */
+	TViewModelPtr<IOutlinerExtension> GetOutlinerItem() const;
+
+	TSharedPtr<STrackLane> GetParentLane() const;
+
+	TArrayView<const TWeakPtr<STrackLane>> GetChildLanes() const;
+
+	const TMap<FWeakViewModelPtr, TSharedPtr<ITrackLaneWidget>>& GetAllWidgets() const { return WidgetsByModel; }
+
+	/** Gets the vertical offset of this track lane */
+	float GetVerticalPosition() const;
+	/** Sets the vertical offset of this track lane */
+	void SetVerticalPosition(float InPosition);
+
+	/**
+	 * Sets the vertical offset of this track lane and all parent track lanes given the vertical offset of
+	 * a particular outliner item inside the current track lane (i.e. an outliner item that's a child of the 
+	 * item that created this track lane)
+	 * Note that the given item *could* actually be the one that created this track lane.
+	 */
+	void PositionParentTrackLanes(TViewModelPtr<IOutlinerExtension> InItem, float InItemTop);
+
+	/** Returns whether the outliner item that created this track lane is pinned or not */
+	bool IsPinned() const;
+
+	TSharedPtr<ITrackLaneWidget> FindWidgetForModel(const FWeakViewModelPtr& InModel) const;
+
+protected:
+
+	int32 PaintLaneBackground(const FGeometry& AllottedGeometry, const FSlateRect& MyCullingRect, FSlateWindowElementList& OutDrawElements, int32 LayerId, const FWidgetStyle& InWidgetStyle) const;
+
+	void OnHierarchyUpdated();
+	void RecreateWidgets();
+
+private:
+	TWeakPtr<STrackAreaView> WeakTrackAreaView;
+
+	/** The outliner item that created this lane */
+	TWeakViewModelPtr<IOutlinerExtension> WeakOutlinerItem;
+
+	/** Pointer back to the tree view for virtual <-> physical space conversions. Important: weak ptr to avoid circular references */
+	TWeakPtr<SOutlinerView> TreeView;
+
+	TMap<FWeakViewModelPtr, TSharedPtr<ITrackLaneWidget>> WidgetsByModel;
+
+	TArray<TWeakPtr<STrackLane>> WeakChildLanes;
+
+	/** Parent track lane - intentionally a _strong_ ptr to the parent in order to keep it alive even if it is scrolled out of view */
+	TSharedPtr<STrackLane> ParentLane;
+
+	TSharedPtr<FTimeToPixel> TimeToPixel;
+
+	struct FSlot : TSlotBase<FSlot>
+	{
+		FSlot(TSharedPtr<ITrackLaneWidget> InInterface);
+		~FSlot();
+
+		TSharedRef<ITrackLaneWidget> Interface;
+	};
+
+	/** All the widgets in the panel */
+	TPanelChildren<FSlot> Children;
+
+	/** Our desired size last frame */
+	TOptional<FVector2D> LastDesiredSize;
+
+	struct FDragParameters
+	{
+		FDragParameters(float InOriginalHeight, float InDragStartY) : OriginalHeight(InOriginalHeight), DragStartY(InDragStartY) {}
+
+		float OriginalHeight;
+		float DragStartY;
+	};
+	TOptional<FDragParameters> DragParameters;
+
+	FTrackAreaParameters TrackParams;
+
+	float Position;
+
+	bool bWidgetsDirty;
+};
+
+} // namespace Sequencer
+} // namespace UE
+
diff --git a/Engine/Source/Editor/SequencerCore/Public/TimeToPixel.h b/Engine/Source/Editor/SequencerCore/Public/TimeToPixel.h
--- a/Engine/Source/Editor/SequencerCore/Public/TimeToPixel.h
+++ b/Engine/Source/Editor/SequencerCore/Public/TimeToPixel.h
@@ -106,6 +106,17 @@ public:
 		return ( PixelDelta / PixelsPerSecond ) * TickResolution;
 	}
 
+	/**
+	 * Converts a pixel delta value to delta seconds time 
+	 *
+	 * @param PixelDelta The delta value in pixel space
+	 * @return The equivalent delta time in seconds
+	 */
+	double PixelDeltaToSeconds( float PixelDelta ) const
+	{
+		return ( PixelDelta / PixelsPerSecond );
+	}
+
 	/**
 	 * Retrieve the tick resolution of the current sequence
 	 */
@@ -117,7 +128,7 @@ public:
 	/**
 	 * Make this converter relative to the specified time (ie, such that pixel 0 == FrameAmount)
 	 */
-	FTimeToPixel RelativeTo(FFrameNumber FrameZero) const
+	FTimeToPixel RelativeTo(const FFrameTime& FrameZero) const
 	{
 		FTimeToPixel Copy = *this;
 		Copy.ViewRangeStartSeconds = FrameZero / TickResolution;
diff --git a/Engine/Source/Editor/SequencerCore/SequencerCore.Build.cs b/Engine/Source/Editor/SequencerCore/SequencerCore.Build.cs
--- a/Engine/Source/Editor/SequencerCore/SequencerCore.Build.cs
+++ b/Engine/Source/Editor/SequencerCore/SequencerCore.Build.cs
@@ -15,6 +15,8 @@ public class SequencerCore : ModuleRules
 		PublicDependencyModuleNames.AddRange(
 			new string[] {
 				"Core",
+				"CurveEditor",
+				"Engine",
 				"GraphEditor",
 				"InputCore",
 				"Slate",
diff --git a/Engine/Source/Runtime/MovieScene/Private/MovieSceneSection.cpp b/Engine/Source/Runtime/MovieScene/Private/MovieSceneSection.cpp
--- a/Engine/Source/Runtime/MovieScene/Private/MovieSceneSection.cpp
+++ b/Engine/Source/Runtime/MovieScene/Private/MovieSceneSection.cpp
@@ -701,7 +701,7 @@ void UMovieSceneSection::EvaluateEasing(FFrameTime InTime, TOptional<float>& Out
 		}
 	}
 
-	if (HasEndFrame() && Easing.EaseOut.GetObject() && GetEaseOutRange().Contains(InTime.FrameNumber))
+	if (HasEndFrame() && Easing.EaseOut.GetObject() && (GetEaseOutRange().Contains(InTime.FrameNumber) || GetEaseOutRange().GetUpperBoundValue() == InTime.FrameNumber))
 	{
 		const int32  EaseFrame     = (InTime.FrameNumber - GetExclusiveEndFrame() + Easing.GetEaseOutDuration()).Value;
 		const double EaseOutInterp = (double(EaseFrame) + InTime.GetSubFrame()) / Easing.GetEaseOutDuration();
