Sequencer: Add flag to control the visibility of Sequencer Outliner columns when attaching a Sequencer.

diff --git a/Engine/Source/Editor/Sequencer/Private/LevelEditorSequencerIntegration.cpp b/Engine/Source/Editor/Sequencer/Private/LevelEditorSequencerIntegration.cpp
--- a/Engine/Source/Editor/Sequencer/Private/LevelEditorSequencerIntegration.cpp
+++ b/Engine/Source/Editor/Sequencer/Private/LevelEditorSequencerIntegration.cpp
@@ -1023,7 +1023,6 @@ void FLevelEditorSequencerIntegration::AttachOutlinerColumn()
 		true, TOptional<float>(), LOCTEXT("SpawnableColumnName", "Spawnable"));
 
 	SceneOutlinerModule.RegisterDefaultColumnType< Sequencer::FSequencerSpawnableColumn >(SpawnColumnInfo);
-	AcquiredResources.Add([=]{ this->DetachOutlinerColumn(); });
 
 	FSceneOutlinerColumnInfo ColumnInfo(ESceneOutlinerColumnVisibility::Visible, 15, 
 		FCreateSceneOutlinerColumn::CreateRaw( this, &FLevelEditorSequencerIntegration::CreateSequencerInfoColumn), 
@@ -1251,7 +1250,10 @@ void FLevelEditorSequencerIntegration::AddSequencer(TSharedRef<ISequencer> InSeq
 	}
 
 	ActivateRealtimeViewports();
-	AttachOutlinerColumn();
+	if (Options.bAttachOutlinerColumns)
+	{
+		AttachOutlinerColumn();
+	}
 	OnSequencersChanged.Broadcast();
 }
 
@@ -1287,6 +1289,12 @@ void FLevelEditorSequencerIntegration::RemoveSequencer(TSharedRef<ISequencer> In
 	{
 		AcquiredResources.Release();
 	}
+	
+	auto HasOutlinerColumns = [](const FSequencerAndOptions& In) { return In.Sequencer.IsValid() && In.Options.bAttachOutlinerColumns; };
+	if (!BoundSequencers.FindByPredicate(HasOutlinerColumns))
+	{
+		DetachOutlinerColumn();
+	}
 
 	OnSequencersChanged.Broadcast();
 }
diff --git a/Engine/Source/Editor/Sequencer/Public/LevelEditorSequencerIntegration.h b/Engine/Source/Editor/Sequencer/Public/LevelEditorSequencerIntegration.h
--- a/Engine/Source/Editor/Sequencer/Public/LevelEditorSequencerIntegration.h
+++ b/Engine/Source/Editor/Sequencer/Public/LevelEditorSequencerIntegration.h
@@ -29,11 +29,13 @@ struct FLevelEditorSequencerIntegrationOptions
 		: bRequiresLevelEvents(true)
 		, bRequiresActorEvents(false)
 		, bForceRefreshDetails(true)
+		, bAttachOutlinerColumns(true)
 	{}
 
 	bool bRequiresLevelEvents : 1;
 	bool bRequiresActorEvents : 1;
 	bool bForceRefreshDetails : 1;
+	bool bAttachOutlinerColumns : 1;
 };
 
 
