Added missing include and fixed ambiguous symbol

diff --git a/Engine/Source/Editor/Sequencer/Private/Tools/SequencerEditTool_Movement.cpp b/Engine/Source/Editor/Sequencer/Private/Tools/SequencerEditTool_Movement.cpp
--- a/Engine/Source/Editor/Sequencer/Private/Tools/SequencerEditTool_Movement.cpp
+++ b/Engine/Source/Editor/Sequencer/Private/Tools/SequencerEditTool_Movement.cpp
@@ -12,6 +12,7 @@
 #include "MVVM/Views/ITrackAreaHotspot.h"
 #include "SequencerHotspots.h"
 #include "MVVM/ViewModels/VirtualTrackArea.h"
+#include "MVVM/ViewModels/TrackAreaViewModel.h"
 #include "SequencerSettings.h"
 #include "Tools/EditToolDragOperations.h"
 #include "IKeyArea.h"
diff --git a/Engine/Source/Editor/SequencerCore/Public/MVVM/Views/STrackLane.h b/Engine/Source/Editor/SequencerCore/Public/MVVM/Views/STrackLane.h
--- a/Engine/Source/Editor/SequencerCore/Public/MVVM/Views/STrackLane.h
+++ b/Engine/Source/Editor/SequencerCore/Public/MVVM/Views/STrackLane.h
@@ -14,7 +14,6 @@
 
 class FPaintArgs;
 class FSlateWindowElementList;
-class FTrackAreaViewModel;
 
 namespace UE
 {
