Sequencer: Added flag to optionally enable Sequencer Editor mode when initializing a Sequencer instance.
Disable Sequencer Editor mode for DaySequenceActorPreview.

diff --git a/Engine/Source/Editor/Sequencer/Private/LevelEditorSequencerIntegration.cpp b/Engine/Source/Editor/Sequencer/Private/LevelEditorSequencerIntegration.cpp
--- a/Engine/Source/Editor/Sequencer/Private/LevelEditorSequencerIntegration.cpp
+++ b/Engine/Source/Editor/Sequencer/Private/LevelEditorSequencerIntegration.cpp
@@ -232,7 +232,6 @@ void FLevelEditorSequencerIntegration::Initialize(const FLevelEditorSequencerInt
 
 	AddLevelViewportMenuExtender();
 	ActivateDetailHandler(Options);
-	ActivateSequencerEditorMode();
 
 	{
 		FLevelEditorModule& LevelEditorModule = FModuleManager::Get().GetModuleChecked<FLevelEditorModule>("LevelEditor");
@@ -357,12 +356,20 @@ void FLevelEditorSequencerIntegration::OnPostSaveExternalActors(UWorld* World)
 
 void FLevelEditorSequencerIntegration::OnNewCurrentLevel()
 {
-	ActivateSequencerEditorMode();
+	auto IsSequenceEditor = [](const FSequencerAndOptions& In) { return In.Sequencer.IsValid() && In.Options.bActivateSequencerEdMode; };
+	if (BoundSequencers.FindByPredicate(IsSequenceEditor))
+	{
+		ActivateSequencerEditorMode();
+	}
 }
 
 void FLevelEditorSequencerIntegration::OnMapOpened(const FString& Filename, bool bLoadAsTemplate)
 {
-	ActivateSequencerEditorMode();
+	auto IsSequenceEditor = [](const FSequencerAndOptions& In) { return In.Sequencer.IsValid() && In.Options.bActivateSequencerEdMode; };
+	if (BoundSequencers.FindByPredicate(IsSequenceEditor))
+	{
+		ActivateSequencerEditorMode();
+	}
 }
 
 void FLevelEditorSequencerIntegration::OnLevelAdded(ULevel* InLevel, UWorld* InWorld)
@@ -564,10 +571,8 @@ void FLevelEditorSequencerIntegration::UpdateDetails(bool bForceRefresh)
 
 void FLevelEditorSequencerIntegration::ActivateSequencerEditorMode()
 {
-
 	// Release the sequencer mode if we already enabled it
-	FName ResourceName("SequencerMode");
-	AcquiredResources.Release(ResourceName);
+	DeactivateSequencerEditorMode();
 
 	// Activate the default mode in case FEditorModeTools::Tick isn't run before here. 
 	// This can be removed once a general fix for UE-143791 has been implemented.
@@ -587,27 +592,25 @@ void FLevelEditorSequencerIntegration::ActivateSequencerEditorMode()
 			SequencerEdMode->AddSequencer(Pinned);
 		}
 	}
+}
 
-	// Acquire the resource, which allows us to deactivate the mode later
-	AcquiredResources.Add(
-		ResourceName,
-		[] {
-
-			FEditorModeID ModeID = TEXT("SequencerToolsEditMode");
 
-			if (GLevelEditorModeTools().IsModeActive(ModeID))
-			{
-				GLevelEditorModeTools().DeactivateMode(ModeID);
-			}
-			if (GLevelEditorModeTools().IsModeActive(FSequencerEdMode::EM_SequencerMode))
-			{
-				GLevelEditorModeTools().DeactivateMode(FSequencerEdMode::EM_SequencerMode);
-			}
-		}
-	);
+void FLevelEditorSequencerIntegration::DeactivateSequencerEditorMode()
+{
+	const FEditorModeID ModeID = TEXT("SequencerToolsEditMode");
+	
+	if (GLevelEditorModeTools().IsModeActive(ModeID))
+	{
+		GLevelEditorModeTools().DeactivateMode(ModeID);
+	}
+	if (GLevelEditorModeTools().IsModeActive(FSequencerEdMode::EM_SequencerMode))
+	{
+		GLevelEditorModeTools().DeactivateMode(FSequencerEdMode::EM_SequencerMode);
+	}
 }
 
 
+
 void FLevelEditorSequencerIntegration::OnPreBeginPIE(bool bIsSimulating)
 {
 	IterateAllSequencers(
@@ -1248,6 +1251,10 @@ void FLevelEditorSequencerIntegration::AddSequencer(TSharedRef<ISequencer> InSeq
 	{
 		SequencerEdMode->AddSequencer(DerivedSequencerPtr);
 	}
+	else if (Options.bActivateSequencerEdMode)
+	{
+		ActivateSequencerEditorMode();
+	}
 
 	ActivateRealtimeViewports();
 	if (Options.bAttachOutlinerColumns)
@@ -1284,14 +1291,27 @@ void FLevelEditorSequencerIntegration::RemoveSequencer(TSharedRef<ISequencer> In
 
 	KeyFrameHandler->Remove(InSequencer);
 
-	auto IsValidSequencer = [](const FSequencerAndOptions& In){ return In.Sequencer.IsValid(); };
-	if (!BoundSequencers.FindByPredicate(IsValidSequencer))
+	bool bHasValidSequencer = false;
+	bool bHasSequencerEditor = false;
+	bool bHasOutlinerColumns = false;
+	for (const FSequencerAndOptions& In : BoundSequencers)
+	{
+		if (In.Sequencer.IsValid())
+		{
+			bHasValidSequencer = true;
+			bHasSequencerEditor = bHasSequencerEditor || In.Options.bActivateSequencerEdMode;
+			bHasOutlinerColumns = bHasOutlinerColumns || In.Options.bAttachOutlinerColumns;
+		}
+	}
+	if (!bHasValidSequencer)
 	{
 		AcquiredResources.Release();
 	}
-	
-	auto HasOutlinerColumns = [](const FSequencerAndOptions& In) { return In.Sequencer.IsValid() && In.Options.bAttachOutlinerColumns; };
-	if (!BoundSequencers.FindByPredicate(HasOutlinerColumns))
+	if (!bHasSequencerEditor)
+	{
+		DeactivateSequencerEditorMode();
+	}
+	if (!bHasOutlinerColumns)
 	{
 		DetachOutlinerColumn();
 	}
diff --git a/Engine/Source/Editor/Sequencer/Public/LevelEditorSequencerIntegration.h b/Engine/Source/Editor/Sequencer/Public/LevelEditorSequencerIntegration.h
--- a/Engine/Source/Editor/Sequencer/Public/LevelEditorSequencerIntegration.h
+++ b/Engine/Source/Editor/Sequencer/Public/LevelEditorSequencerIntegration.h
@@ -30,12 +30,14 @@ struct FLevelEditorSequencerIntegrationOptions
 		, bRequiresActorEvents(false)
 		, bForceRefreshDetails(true)
 		, bAttachOutlinerColumns(true)
+		, bActivateSequencerEdMode(true)
 	{}
 
 	bool bRequiresLevelEvents : 1;
 	bool bRequiresActorEvents : 1;
 	bool bForceRefreshDetails : 1;
 	bool bAttachOutlinerColumns : 1;
+	bool bActivateSequencerEdMode : 1;
 };
 
 
@@ -174,6 +176,7 @@ private:
 private:
 
 	void ActivateSequencerEditorMode();
+	void DeactivateSequencerEditorMode();
 	void AddLevelViewportMenuExtender();
 	void ActivateDetailHandler(const FLevelEditorSequencerIntegrationOptions& Options);
 	void AttachOutlinerColumn();
