Sequencer: Tiny section width fixes
  - We previously did not arrange sections with degenerate (ie, [0,0] ) ranges on their track lanes. Unfortunately UMG often creates degenerate section ranges so we need to show those properly.
  - Also fixed resize-handle offsets that were not rendering correctly outside the section range when the pixel width of the section was small enough.

diff --git a/Engine/Source/Editor/Sequencer/Private/SSequencerSection.cpp b/Engine/Source/Editor/Sequencer/Private/SSequencerSection.cpp
--- a/Engine/Source/Editor/Sequencer/Private/SSequencerSection.cpp
+++ b/Engine/Source/Editor/Sequencer/Private/SSequencerSection.cpp
@@ -830,6 +830,7 @@ void SSequencerSection::Construct( const FArguments& InArgs, TSharedPtr<FSequenc
 	UpdateUnderlappingSegments();
 
 	ChildSlot
+	.Padding(MakeAttributeSP(this, &SSequencerSection::GetHandleOffsetPadding))
 	[
 		SNew(SOverlay)
 
@@ -874,6 +875,11 @@ EVisibility SSequencerSection::GetTopLevelChannelGroupVisibility() const
 	return EVisibility::Visible;
 }
 
+FMargin SSequencerSection::GetHandleOffsetPadding() const
+{
+	return FMargin(HandleOffsetPx, 0.0);
+}
+
 FText SSequencerSection::GetToolTipText() const
 {
 	const UMovieSceneSection* SectionObject = SectionInterface->GetSectionObject();
@@ -974,7 +980,7 @@ FTrackLaneScreenAlignment SSequencerSection::GetAlignment(const FTimeToPixel& In
 	{
 		constexpr float MinSectionWidth = 1.f;
 
-		const float FinalSectionWidth = FMath::Max(MinSectionWidth, ScreenAlignment.WidthPx);
+		const float FinalSectionWidth = FMath::Max(MinSectionWidth + SectionInterface->GetSectionGripSize() * 2.f, ScreenAlignment.WidthPx);
 		const float GripOffset        = (FinalSectionWidth - ScreenAlignment.WidthPx) / 2.f;
 
 		ScreenAlignment.LeftPosPx -= GripOffset;
diff --git a/Engine/Source/Editor/Sequencer/Private/SSequencerSection.h b/Engine/Source/Editor/Sequencer/Private/SSequencerSection.h
--- a/Engine/Source/Editor/Sequencer/Private/SSequencerSection.h
+++ b/Engine/Source/Editor/Sequencer/Private/SSequencerSection.h
@@ -109,6 +109,8 @@ private:
 
 	EVisibility GetTopLevelChannelGroupVisibility() const;
 
+	FMargin GetHandleOffsetPadding() const;
+
 public:
 
 	/** Indicate that the current section selection should throb the specified number of times. A single throb takes 0.2s. */
diff --git a/Engine/Source/Editor/SequencerCore/Public/MVVM/Extensions/ITrackLaneExtension.h b/Engine/Source/Editor/SequencerCore/Public/MVVM/Extensions/ITrackLaneExtension.h
--- a/Engine/Source/Editor/SequencerCore/Public/MVVM/Extensions/ITrackLaneExtension.h
+++ b/Engine/Source/Editor/SequencerCore/Public/MVVM/Extensions/ITrackLaneExtension.h
@@ -55,7 +55,7 @@ struct SEQUENCERCORE_API FTrackLaneVirtualAlignment
 
 	bool IsVisible() const
 	{
-		return !Range.IsEmpty() && !Range.IsDegenerate();
+		return !Range.IsEmpty();
 	}
 
 	static FTrackLaneVirtualAlignment Fixed(const TRange<FFrameNumber>& InRange, float InFixedHeight, EVerticalAlignment InVAlign = VAlign_Center)
