Sequencer: Broadcast GetOnOutlinerNodeSelectionChanged

This used to happen in SequencerHelpers::PerformDefaultSelection. This fixes a bug where selecting a key should also select the object that the key/section/track/binding is bound to.

diff --git a/Engine/Source/Editor/Sequencer/Private/SequencerHotspots.cpp b/Engine/Source/Editor/Sequencer/Private/SequencerHotspots.cpp
--- a/Engine/Source/Editor/Sequencer/Private/SequencerHotspots.cpp
+++ b/Engine/Source/Editor/Sequencer/Private/SequencerHotspots.cpp
@@ -59,6 +59,7 @@ FHotspotSelectionManager::FHotspotSelectionManager(const FPointerEvent* InMouseE
 FHotspotSelectionManager::~FHotspotSelectionManager()
 {
 	Selection->ResumeBroadcast();
+	Selection->GetOnOutlinerNodeSelectionChanged().Broadcast();
 }
 
 void FHotspotSelectionManager::ConditionallyClearSelection()
