Sequencer: make an arbitrary section range when a new section is created past the playback end.

Before, this would result in an error because the start of the section is past the end of the section. Now, we take the current time and make a section that ends at the end of the view range.

diff --git a/Engine/Source/Editor/Sequencer/Private/SequencerUtilities.cpp b/Engine/Source/Editor/Sequencer/Private/SequencerUtilities.cpp
--- a/Engine/Source/Editor/Sequencer/Private/SequencerUtilities.cpp
+++ b/Engine/Source/Editor/Sequencer/Private/SequencerUtilities.cpp
@@ -207,7 +207,15 @@ void FSequencerUtilities::PopulateMenu_CreateNewSection(FMenuBuilder& MenuBuilde
 
 			Track->Modify();
 
-			NewSection->SetRange(TRange<FFrameNumber>(CurrentTime.Time.FrameNumber, PlaybackEnd));
+			FFrameNumber NewSectionRangeEnd = PlaybackEnd;
+			if (PlaybackEnd <= CurrentTime.Time.FrameNumber)
+			{
+				const FAnimatedRange ViewRange = Sequencer->GetViewRange();
+				const FFrameRate TickResolution = Sequencer->GetFocusedTickResolution();
+				NewSectionRangeEnd = (ViewRange.GetUpperBoundValue() * TickResolution).FloorToFrame();
+			}
+			
+			NewSection->SetRange(TRange<FFrameNumber>(CurrentTime.Time.FrameNumber, NewSectionRangeEnd));
 			NewSection->SetOverlapPriority(OverlapPriority);
 			NewSection->SetRowIndex(RowIndex);
 			NewSection->SetBlendType(BlendType);
