Non-unity mac editor fix.

diff --git a/Engine/Source/Editor/Sequencer/Private/MVVM/ViewModels/ChannelModel.cpp b/Engine/Source/Editor/Sequencer/Private/MVVM/ViewModels/ChannelModel.cpp
--- a/Engine/Source/Editor/Sequencer/Private/MVVM/ViewModels/ChannelModel.cpp
+++ b/Engine/Source/Editor/Sequencer/Private/MVVM/ViewModels/ChannelModel.cpp
@@ -7,6 +7,7 @@
 #include "Channels/MovieSceneChannelHandle.h"
 #include "Channels/MovieSceneSectionChannelOverrideRegistry.h"
 #include "CurveModel.h"
+#include "DetailsViewArgs.h"
 #include "HAL/PlatformCrt.h"
 #include "IKeyArea.h"
 #include "ISequencerChannelInterface.h"
@@ -23,8 +24,10 @@
 #include "MVVM/Views/SSequencerKeyNavigationButtons.h"
 #include "Math/UnrealMathUtility.h"
 #include "Misc/AssertionMacros.h"
+#include "Modules/ModuleManager.h"
 #include "MovieSceneTrack.h"
 #include "MVVM/Views/SChannelView.h"
+#include "PropertyEditorModule.h"
 #include "SKeyAreaEditorSwitcher.h"
 #include "SequencerCoreFwd.h"
 #include "SequencerSectionPainter.h"
