Sequencer: Update key data only if channel is valid.

The signature of the section can be cached when the channel for the key area isn't established yet, so the cache never gets filled again when it does get established.

diff --git a/Engine/Source/Editor/Sequencer/Private/MVVM/ViewModels/ChannelModel.cpp b/Engine/Source/Editor/Sequencer/Private/MVVM/ViewModels/ChannelModel.cpp
--- a/Engine/Source/Editor/Sequencer/Private/MVVM/ViewModels/ChannelModel.cpp
+++ b/Engine/Source/Editor/Sequencer/Private/MVVM/ViewModels/ChannelModel.cpp
@@ -136,7 +136,7 @@ bool FChannelModel::UpdateCachedKeys(TSharedPtr<FCachedKeys>& OutCachedKeys) con
 			UMovieSceneSection* Section = InChannel->GetSection();
 			if (!Section || !CachedSignature.IsValid() || Section->GetSignature() != CachedSignature)
 			{
-				CachedSignature = Section ? Section->GetSignature() : FGuid();
+				CachedSignature = Section && InChannel->GetKeyArea() && InChannel->GetKeyArea()->ResolveChannel() ? Section->GetSignature() : FGuid();
 
 				KeyTimes.Reset();
 				KeyHandles.Reset();
