Sequencer: Add additional shortcuts for Jump to Start/Jump to End that can be handled in the viewport

diff --git a/Engine/Source/Editor/Sequencer/Private/Sequencer.cpp b/Engine/Source/Editor/Sequencer/Private/Sequencer.cpp
--- a/Engine/Source/Editor/Sequencer/Private/Sequencer.cpp
+++ b/Engine/Source/Editor/Sequencer/Private/Sequencer.cpp
@@ -11145,6 +11145,14 @@ void FSequencer::BindCommands()
 		Commands.JumpToEnd,
 		FExecuteAction::CreateSP(this, &FSequencer::JumpToEnd));
 
+	SequencerCommandBindings->MapAction(
+		Commands.JumpToStartViewport,
+		FExecuteAction::CreateSP(this, &FSequencer::JumpToStart));
+
+	SequencerCommandBindings->MapAction(
+		Commands.JumpToEndViewport,
+		FExecuteAction::CreateSP(this, &FSequencer::JumpToEnd));
+
 	SequencerCommandBindings->MapAction(
 		Commands.StepForward,
 		FExecuteAction::CreateSP(this, &FSequencer::StepForward),
@@ -11318,6 +11326,8 @@ void FSequencer::BindCommands()
 		CurveEditorSharedBindings->MapAction(Commands.PlayForward, *SequencerCommandBindings->GetActionForCommand(Commands.PlayForward));
 		CurveEditorSharedBindings->MapAction(Commands.JumpToStart, *SequencerCommandBindings->GetActionForCommand(Commands.JumpToStart));
 		CurveEditorSharedBindings->MapAction(Commands.JumpToEnd, *SequencerCommandBindings->GetActionForCommand(Commands.JumpToEnd));
+		CurveEditorSharedBindings->MapAction(Commands.JumpToStartViewport, *SequencerCommandBindings->GetActionForCommand(Commands.JumpToStartViewport));
+		CurveEditorSharedBindings->MapAction(Commands.JumpToEndViewport, *SequencerCommandBindings->GetActionForCommand(Commands.JumpToEndViewport));
 		CurveEditorSharedBindings->MapAction(Commands.ShuttleBackward, *SequencerCommandBindings->GetActionForCommand(Commands.ShuttleBackward));
 		CurveEditorSharedBindings->MapAction(Commands.ShuttleForward, *SequencerCommandBindings->GetActionForCommand(Commands.ShuttleForward));
 		CurveEditorSharedBindings->MapAction(Commands.Pause, *SequencerCommandBindings->GetActionForCommand(Commands.Pause));
diff --git a/Engine/Source/Editor/Sequencer/Private/SequencerCommands.cpp b/Engine/Source/Editor/Sequencer/Private/SequencerCommands.cpp
--- a/Engine/Source/Editor/Sequencer/Private/SequencerCommands.cpp
+++ b/Engine/Source/Editor/Sequencer/Private/SequencerCommands.cpp
@@ -12,11 +12,13 @@
 void FSequencerCommands::RegisterCommands()
 {
 	UI_COMMAND( TogglePlay, "Toggle Play", "Toggle the timeline playing", EUserInterfaceActionType::Button, FInputChord(EKeys::SpaceBar) );
-	UI_COMMAND( TogglePlayViewport, "Toggle Play(Viewport)", "Toggle the timeline playing in all viewports and sequencer", EUserInterfaceActionType::Button, FInputChord(EModifierKey::Shift, EKeys::SpaceBar));
-	UI_COMMAND( ScrubTimeViewport, "Scrub Time(Viewport)", "Scrub mouse left and right to change time", EUserInterfaceActionType::Button, FInputChord(EKeys::B) );
+	UI_COMMAND( TogglePlayViewport, "Toggle Play (Viewport)", "Toggle the timeline playing in all viewports and sequencer", EUserInterfaceActionType::Button, FInputChord(EModifierKey::Shift, EKeys::SpaceBar));
+	UI_COMMAND( ScrubTimeViewport, "Scrub Time (Viewport)", "Scrub mouse left and right to change time", EUserInterfaceActionType::Button, FInputChord(EKeys::B) );
 	UI_COMMAND( PlayForward, "Play Forward", "Play the timeline forward", EUserInterfaceActionType::Button, FInputChord(EKeys::Down) );
 	UI_COMMAND( JumpToStart, "Jump to Start", "Jump to the start of the playback range", EUserInterfaceActionType::Button, FInputChord(EKeys::Up) );
 	UI_COMMAND( JumpToEnd, "Jump to End", "Jump to the end of the playback range", EUserInterfaceActionType::Button, FInputChord(EModifierKey::Control, EKeys::Up) );
+	UI_COMMAND( JumpToStartViewport, "Jump to Start (Viewport)", "Jump to the start of the playback range in all viewports and sequencer", EUserInterfaceActionType::Button, FInputChord() );
+	UI_COMMAND( JumpToEndViewport, "Jump to End (Viewport)", "Jump to the end of the playback range in all viewports and sequencer", EUserInterfaceActionType::Button, FInputChord() );
 	UI_COMMAND( ShuttleBackward, "Shuttle Backward", "Shuttle backward", EUserInterfaceActionType::Button, FInputChord(EKeys::J) );
 	UI_COMMAND( ShuttleForward, "Shuttle Forward", "Shuttle forward", EUserInterfaceActionType::Button, FInputChord(EKeys::L) );
 	UI_COMMAND( Pause, "Pause", "Pause playback", EUserInterfaceActionType::Button, FInputChord(EKeys::K) );
@@ -28,7 +30,7 @@ void FSequencerCommands::RegisterCommands()
 	UI_COMMAND( JumpForward, "Jump Forward", "Jump the timeline forward a user defined number of frames/times", EUserInterfaceActionType::Button, FInputChord(EModifierKey::Shift, EKeys::Right) );
 	UI_COMMAND( JumpBackward, "Jump Backward", "Jump the timeline backward a user defined number of frames/times", EUserInterfaceActionType::Button, FInputChord(EModifierKey::Shift, EKeys::Left) );
 	UI_COMMAND( StepToNextKey, "Step to Next Key (Viewport)", "Step to the next key in all viewports and sequencer", EUserInterfaceActionType::Button, FInputChord(EKeys::Period) );
-	UI_COMMAND( StepToPreviousKey, "Step to Previous Key(Viewport)", "Step to the previous key in all viewports and sequencer", EUserInterfaceActionType::Button, FInputChord(EKeys::Comma) );
+	UI_COMMAND( StepToPreviousKey, "Step to Previous Key (Viewport)", "Step to the previous key in all viewports and sequencer", EUserInterfaceActionType::Button, FInputChord(EKeys::Comma) );
 	UI_COMMAND( StepToNextCameraKey, "Step to Next Camera Key", "Step to the next camera key", EUserInterfaceActionType::Button, FInputChord() );
 	UI_COMMAND( StepToPreviousCameraKey, "Step to Previous Camera Key", "Step to the previous camera key", EUserInterfaceActionType::Button, FInputChord() );
 	UI_COMMAND( StepToNextShot, "Step to Next Shot", "Step to the next shot", EUserInterfaceActionType::Button, FInputChord(EModifierKey::Shift, EKeys::Period) );
diff --git a/Engine/Source/Editor/Sequencer/Public/SequencerCommands.h b/Engine/Source/Editor/Sequencer/Public/SequencerCommands.h
--- a/Engine/Source/Editor/Sequencer/Public/SequencerCommands.h
+++ b/Engine/Source/Editor/Sequencer/Public/SequencerCommands.h
@@ -42,7 +42,13 @@ public:
 
 	/** Jump to end of playback */
 	TSharedPtr< FUICommandInfo > JumpToEnd;
-	
+
+	/** Jump to start of playback */
+	TSharedPtr< FUICommandInfo > JumpToStartViewport;
+
+	/** Jump to end of playback */
+	TSharedPtr< FUICommandInfo > JumpToEndViewport;
+
 	/** Shuttle forward */
 	TSharedPtr< FUICommandInfo > ShuttleForward;
 
