Sequencer: work-around fix for transform-track crashes

diff --git a/Engine/Source/Editor/MovieSceneTools/Private/TrackEditors/TransformTrackEditor.cpp b/Engine/Source/Editor/MovieSceneTools/Private/TrackEditors/TransformTrackEditor.cpp
--- a/Engine/Source/Editor/MovieSceneTools/Private/TrackEditors/TransformTrackEditor.cpp
+++ b/Engine/Source/Editor/MovieSceneTools/Private/TrackEditors/TransformTrackEditor.cpp
@@ -708,8 +708,13 @@ void F3DTransformTrackEditor::GetTransformKeys( const TOptional<FTransformData>&
 
 	bool IsChannelKeyable[9];
 	const FName DoubleChannelTypeName = FMovieSceneDoubleChannel::StaticStruct()->GetFName();
+	TWeakPtr<FMovieSceneChannelProxy> WeakSectionProxy = Section->GetChannelProxy().AsWeak();
 	for (int32 Index = 0; Index < sizeof(IsChannelKeyable); ++Index)
 	{
+		if (ChannelHandles[Index] == FMovieSceneChannelHandle())
+		{
+			ChannelHandles[Index] = FMovieSceneChannelHandle(WeakSectionProxy, FMovieSceneDoubleChannel::StaticStruct()->GetFName(), Index);
+		}
 		IsChannelKeyable[Index] = (ChannelHandles[Index].GetChannelTypeName() == DoubleChannelTypeName);
 	}
 
diff --git a/Engine/Source/Editor/Sequencer/Private/MVVM/ViewModels/SectionModel.cpp b/Engine/Source/Editor/Sequencer/Private/MVVM/ViewModels/SectionModel.cpp
--- a/Engine/Source/Editor/Sequencer/Private/MVVM/ViewModels/SectionModel.cpp
+++ b/Engine/Source/Editor/Sequencer/Private/MVVM/ViewModels/SectionModel.cpp
@@ -543,13 +543,16 @@ TArray<FOverlappingSections> FSectionModel::GetEasingSegments()
 
 FCachedChannelModels::FCachedChannelModels(TSharedPtr<FSectionModel> InSection)
 {
-	for (TSharedPtr<FChannelModel> ChannelModel : InSection->GetDescendantsOfType<FChannelModel>())
+	if (InSection)
 	{
-		if (TSharedPtr<IKeyArea> KeyArea = ChannelModel->GetKeyArea())
+		for (TSharedPtr<FChannelModel> ChannelModel : InSection->GetDescendantsOfType<FChannelModel>())
 		{
-			const FName ChannelName = KeyArea->GetName();
-			const FMovieSceneChannelHandle ChannelHandle = KeyArea->GetChannel();
-			ChannelsByName.Add(ChannelName, FCachedChannelModel{ ChannelHandle, ChannelModel });
+			if (TSharedPtr<IKeyArea> KeyArea = ChannelModel->GetKeyArea())
+			{
+				const FName ChannelName = KeyArea->GetName();
+				const FMovieSceneChannelHandle ChannelHandle = KeyArea->GetChannel();
+				ChannelsByName.Add(ChannelName, FCachedChannelModel{ ChannelHandle, ChannelModel });
+			}
 		}
 	}
 }
