Sequencer: Added generic ways to create custom views for channels
 - This is now supported by either overriding ISequencerSection::CreateChannelModel, or by defining a

TSharedPtr<UE::Sequencer::STrackLaneView> CreateChannelView(TMovieSceneChannelHandle<FMyChannelType>, TWeakPtr<FChannelModel> const UE::Sequencer::FCreateTrackLaneViewParams&);

overload for the sequencer channel traits of that type.

diff --git a/Engine/Source/Editor/MovieSceneTools/Public/Channels/PerlinNoiseChannelInterface.h b/Engine/Source/Editor/MovieSceneTools/Public/Channels/PerlinNoiseChannelInterface.h
--- a/Engine/Source/Editor/MovieSceneTools/Public/Channels/PerlinNoiseChannelInterface.h
+++ b/Engine/Source/Editor/MovieSceneTools/Public/Channels/PerlinNoiseChannelInterface.h
@@ -226,5 +226,14 @@ struct TPerlinNoiseChannelInterface : ISequencerChannelInterface
 
 		return LayerId + 1;
 	}
+
+	virtual TSharedPtr<UE::Sequencer::FChannelModel> CreateChannelModel_Raw(const FMovieSceneChannelHandle& InChannelHandle, FName InChannelName) const override
+	{
+		return nullptr;
+	}
+	virtual TSharedPtr<UE::Sequencer::STrackAreaLaneView> CreateChannelView_Raw(const FMovieSceneChannelHandle& InChannelHandle, TWeakPtr<UE::Sequencer::FChannelModel> InWeakModel, const UE::Sequencer::FCreateTrackLaneViewParams& Parameters) const override
+	{
+		return nullptr;
+	}
 };
 
diff --git a/Engine/Source/Editor/Sequencer/Private/ISequencerSection.cpp b/Engine/Source/Editor/Sequencer/Private/ISequencerSection.cpp
--- a/Engine/Source/Editor/Sequencer/Private/ISequencerSection.cpp
+++ b/Engine/Source/Editor/Sequencer/Private/ISequencerSection.cpp
@@ -20,6 +20,9 @@
 #include "SequencerSectionPainter.h"
 #include "Templates/Tuple.h"
 #include "Templates/UnrealTemplate.h"
+#include "ISequencerChannelInterface.h"
+#include "ISequencerModule.h"
+#include "Modules/ModuleManager.h"
 
 struct FMovieSceneChannel;
 
@@ -53,6 +56,8 @@ struct FGroupData
 
 void ISequencerSection::GenerateSectionLayout( ISectionLayoutBuilder& LayoutBuilder )
 {
+	using namespace UE::Sequencer;
+
 	UMovieSceneSection* Section = GetSectionObject();
 	if (!Section)
 	{
@@ -96,9 +101,22 @@ void ISequencerSection::GenerateSectionLayout( ISectionLayoutBuilder& LayoutBuil
 		return;
 	}
 
-	auto ChannelFactory = [this](FName InChannelName, const FMovieSceneChannelHandle& InChannel)
+
+	ISequencerModule* SequencerModule = &FModuleManager::LoadModuleChecked<ISequencerModule>("Sequencer");
+
+	auto ChannelFactory = [this, SequencerModule](FName InChannelName, const FMovieSceneChannelHandle& InChannel)
 	{
-		return this->ConstructChannelModel(InChannelName, InChannel);
+		TSharedPtr<FChannelModel> ChannelModel = this->ConstructChannelModel(InChannelName, InChannel);
+		if (!ChannelModel)
+		{
+			ISequencerChannelInterface* EditorInterface = SequencerModule->FindChannelEditorInterface(InChannel.GetChannelTypeName());
+			if (EditorInterface)
+			{
+				ChannelModel = EditorInterface->CreateChannelModel_Raw(InChannel, InChannelName);
+			}
+		}
+
+		return ChannelModel;
 	};
 
 	// Collapse single channels to the top level track node if allowed
diff --git a/Engine/Source/Editor/Sequencer/Private/MVVM/ViewModels/CategoryModel.cpp b/Engine/Source/Editor/Sequencer/Private/MVVM/ViewModels/CategoryModel.cpp
--- a/Engine/Source/Editor/Sequencer/Private/MVVM/ViewModels/CategoryModel.cpp
+++ b/Engine/Source/Editor/Sequencer/Private/MVVM/ViewModels/CategoryModel.cpp
@@ -25,6 +25,8 @@
 #include "Widgets/DeclarativeSyntaxSupport.h"
 #include "Widgets/Layout/SBox.h"
 #include "Widgets/SBoxPanel.h"
+#include "MVVM/Views/STrackLane.h"
+#include "MVVM/ViewModels/SequencerEditorViewModel.h"
 
 class SWidget;
 
@@ -61,7 +63,7 @@ FOutlinerSizing FCategoryModel::GetDesiredSizing() const
 
 TSharedPtr<ITrackLaneWidget> FCategoryModel::CreateTrackLaneView(const FCreateTrackLaneViewParams& InParams)
 {
-	return SNew(SChannelView, SharedThis(this), InParams.TimeToPixel, InParams.Editor->GetTrackArea())
+	return SNew(SChannelView, SharedThis(this), InParams.OwningTrackLane->GetTrackAreaView())
 		.KeyBarColor(this, &FCategoryModel::GetKeyBarColor);
 }
 
diff --git a/Engine/Source/Editor/Sequencer/Private/MVVM/ViewModels/ChannelModel.cpp b/Engine/Source/Editor/Sequencer/Private/MVVM/ViewModels/ChannelModel.cpp
--- a/Engine/Source/Editor/Sequencer/Private/MVVM/ViewModels/ChannelModel.cpp
+++ b/Engine/Source/Editor/Sequencer/Private/MVVM/ViewModels/ChannelModel.cpp
@@ -22,6 +22,7 @@
 #include "MVVM/Views/SOutlinerView.h"
 #include "MVVM/Views/SOutlinerItemViewBase.h"
 #include "MVVM/Views/SSequencerKeyNavigationButtons.h"
+#include "MVVM/Views/STrackLane.h"
 #include "Math/UnrealMathUtility.h"
 #include "Misc/AssertionMacros.h"
 #include "Modules/ModuleManager.h"
@@ -108,7 +109,17 @@ FOutlinerSizing FChannelModel::GetDesiredSizing() const
 
 TSharedPtr<ITrackLaneWidget> FChannelModel::CreateTrackLaneView(const FCreateTrackLaneViewParams& InParams)
 {
-	return SNew(SChannelView, SharedThis(this), InParams.TimeToPixel, InParams.Editor->GetTrackArea())
+	ISequencerChannelInterface* EditorInterface = KeyArea->FindChannelEditorInterface();
+	if (EditorInterface)
+	{
+		TSharedPtr<STrackAreaLaneView> CustomWidget = EditorInterface->CreateChannelView_Raw(KeyArea->GetChannel(), SharedThis(this), InParams);
+		if (CustomWidget)
+		{
+			return CustomWidget;
+		}
+	}
+
+	return SNew(SChannelView, SharedThis(this), InParams.OwningTrackLane->GetTrackAreaView())
 		.KeyBarColor(this, &FChannelModel::GetKeyBarColor);
 }
 
diff --git a/Engine/Source/Editor/Sequencer/Private/MVVM/Views/SChannelView.cpp b/Engine/Source/Editor/Sequencer/Private/MVVM/Views/SChannelView.cpp
--- a/Engine/Source/Editor/Sequencer/Private/MVVM/Views/SChannelView.cpp
+++ b/Engine/Source/Editor/Sequencer/Private/MVVM/Views/SChannelView.cpp
@@ -5,6 +5,7 @@
 #include "MVVM/ViewModels/SectionModel.h"
 #include "MVVM/ViewModels/SequenceModel.h"
 #include "MVVM/ViewModels/TrackAreaViewModel.h"
+#include "MVVM/Views/STrackAreaView.h"
 
 #include "IKeyArea.h"
 #include "Sequencer.h"
@@ -212,11 +213,14 @@ EViewDependentCacheFlags FChannelViewKeyCachedState::CompareTo(const FChannelVie
 	return Flags;
 }
 
-void SChannelView::Construct(const FArguments& InArgs, const FViewModelPtr& InViewModel, TSharedPtr<FTimeToPixel> InTimeToPixel, TSharedPtr<FTrackAreaViewModel> TrackArea)
+void SChannelView::Construct(const FArguments& InArgs, const FViewModelPtr& InViewModel, TSharedPtr<STrackAreaView> InTrackAreaView)
 {
-	WeakModel = InViewModel;
-	WeakTrackArea = TrackArea;
-	TrackAreaTimeToPixel = InTimeToPixel;
+	STrackAreaLaneView::Construct(
+		STrackAreaLaneView::FArguments()
+		[
+			InArgs._Content.Widget
+		]
+		, InViewModel, InTrackAreaView);
 
 	KeyBarColor = InArgs._KeyBarColor;
 
@@ -225,7 +229,8 @@ void SChannelView::Construct(const FArguments& InArgs, const FViewModelPtr& InVi
 
 FReply SChannelView::OnMouseButtonDown(const FGeometry& MyGeometry, const FPointerEvent& MouseEvent)
 {
-	TViewModelPtr<FTrackAreaViewModel> TrackArea = WeakTrackArea.Pin();
+	TSharedPtr<STrackAreaView>      TrackAreaView = WeakTrackAreaView.Pin();
+	TSharedPtr<FTrackAreaViewModel> TrackArea     = TrackAreaView ? TrackAreaView->GetViewModel() : nullptr;
 	if (!TrackArea)
 	{
 		return FReply::Unhandled();
@@ -269,7 +274,8 @@ FReply SChannelView::OnMouseButtonDown(const FGeometry& MyGeometry, const FPoint
 
 FReply SChannelView::OnMouseMove(const FGeometry& MyGeometry, const FPointerEvent& MouseEvent)
 {
-	TViewModelPtr<FTrackAreaViewModel> TrackArea = WeakTrackArea.Pin();
+	TSharedPtr<STrackAreaView>      TrackAreaView = WeakTrackAreaView.Pin();
+	TSharedPtr<FTrackAreaViewModel> TrackArea     = TrackAreaView ? TrackAreaView->GetViewModel() : nullptr;
 	if (!TrackArea)
 	{
 		return FReply::Unhandled();
@@ -313,11 +319,7 @@ FReply SChannelView::OnMouseMove(const FGeometry& MyGeometry, const FPointerEven
 				TRange<FFrameNumber> SectionRange = Section->GetRange();
 
 				// Make the time <-> pixel converter relative to the section range rather than the view range
-				FTimeToPixel RelativeTimeToPixel = *TrackAreaTimeToPixel;
-				if (SectionRange.GetLowerBound().IsClosed())
-				{
-					RelativeTimeToPixel = TrackAreaTimeToPixel->RelativeTo(SectionRange.GetLowerBoundValue());
-				}
+				FTimeToPixel RelativeTimeToPixel = GetRelativeTimeToPixel();
 
 				TArray<FKeyRenderer::FKeysForModel> LeadingKeys, TrailingKeys;
 				if (KeyRenderer.HitTestKeyBar(RelativeTimeToPixel.PixelToFrame(MousePixel.X), LeadingKeys, TrailingKeys))
@@ -368,7 +370,8 @@ void SChannelView::OnMouseLeave(const FPointerEvent& MouseEvent)
 {
 	SCompoundWidget::OnMouseLeave(MouseEvent);
 
-	TViewModelPtr<FTrackAreaViewModel> TrackArea = WeakTrackArea.Pin();
+	TSharedPtr<STrackAreaView>      TrackAreaView = WeakTrackAreaView.Pin();
+	TSharedPtr<FTrackAreaViewModel> TrackArea     = TrackAreaView ? TrackAreaView->GetViewModel() : nullptr;
 	if (TrackArea)
 	{
 		TrackArea->SetHotspot(nullptr);
@@ -417,24 +420,19 @@ int32 SChannelView::OnPaint(const FPaintArgs& Args, const FGeometry& AllottedGeo
 
 	const bool bIncludeThis = true;
 
-	TViewModelPtr<FTrackAreaViewModel> TrackArea = WeakTrackArea.Pin();
-	TViewModelPtr<FSectionModel>  Section  = Model->FindAncestorOfType<FSectionModel>(bIncludeThis);
-	TViewModelPtr<FSequenceModel> Sequence = Model->FindAncestorOfType<FSequenceModel>();
+	TSharedPtr<STrackAreaView>      TrackAreaView = WeakTrackAreaView.Pin();
+	TSharedPtr<FTrackAreaViewModel> TrackArea     = TrackAreaView ? TrackAreaView->GetViewModel() : nullptr;
+	TViewModelPtr<FSectionModel>    Section       = Model->FindAncestorOfType<FSectionModel>(bIncludeThis);
+	TViewModelPtr<FSequenceModel>   Sequence      = Model->FindAncestorOfType<FSequenceModel>();
 
-	UMovieSceneSection* SectionObject = Section->GetSection();
+	UMovieSceneSection* SectionObject = Section ? Section->GetSection() : nullptr;
 	if (!Section || !Sequence || !SectionObject || !TrackArea)
 	{
 		return LayerId;
 	}
 
-	TRange<FFrameNumber> SectionRange = Section->GetRange();
-
 	// Make the time <-> pixel converter relative to the section range rather than the view range
-	FTimeToPixel RelativeTimeToPixel = *TrackAreaTimeToPixel;
-	if (SectionRange.GetLowerBound().IsClosed())
-	{
-		RelativeTimeToPixel = TrackAreaTimeToPixel->RelativeTo(SectionRange.GetLowerBoundValue());
-	}
+	FTimeToPixel RelativeTimeToPixel = GetRelativeTimeToPixel();
 
 	FVector2D TopLeft = AllottedGeometry.AbsoluteToLocal(MyCullingRect.GetTopLeft());
 	FVector2D BottomRight = AllottedGeometry.AbsoluteToLocal(MyCullingRect.GetBottomRight());
@@ -575,24 +573,6 @@ int32 SChannelView::DrawLane(const FPaintArgs& Args, const FGeometry& AllottedGe
 	return LayerId;
 }
 
-TSharedRef<const SWidget> SChannelView::AsWidget() const
-{
-	return AsShared();
-}
-
-FTrackLaneScreenAlignment SChannelView::GetAlignment(const FTimeToPixel& InTimeToPixel, const FGeometry& InParentGeometry) const
-{
-	TSharedPtr<ITrackLaneExtension> TrackLaneExtension = WeakModel.ImplicitPin();
-	if (TrackLaneExtension)
-	{
-		FTrackLaneVirtualAlignment VirtualAlignment = TrackLaneExtension->ArrangeVirtualTrackLaneView();
-		FTrackLaneScreenAlignment  ScreenAlignment  = VirtualAlignment.ToScreen(InTimeToPixel, InParentGeometry);
-
-		return ScreenAlignment;
-	}
-	return FTrackLaneScreenAlignment();
-}
-
 TViewModelPtr<FSectionModel> SChannelView::GetSection() const
 {
 	FViewModelPtr Model = WeakModel.Pin();
@@ -640,25 +620,15 @@ void SChannelView::GetKeysUnderMouse(const FVector2D& MousePosition, const FGeom
 void SChannelView::GetKeysAtPixelX(float LocalMousePixelX, TArray<FSequencerSelectedKey>& OutKeys) const
 {
 	TViewModelPtr<FSectionModel> Section = GetSection();
-	if (!Section)
-	{
-		return;
-	}
+	UMovieSceneSection* SectionObject = Section ? Section->GetSection() : nullptr;
 
-	UMovieSceneSection* SectionObject = Section->GetSection();
 	if (!SectionObject)
 	{
 		return;
 	}
 
-	TRange<FFrameNumber> SectionRange = Section->GetRange();
-
 	// Make the time <-> pixel converter relative to the section range rather than the view range
-	FTimeToPixel RelativeTimeToPixel = *TrackAreaTimeToPixel;
-	if (SectionRange.GetLowerBound().IsClosed())
-	{
-		RelativeTimeToPixel = TrackAreaTimeToPixel->RelativeTo(SectionRange.GetLowerBoundValue());
-	}
+	FTimeToPixel RelativeTimeToPixel = GetRelativeTimeToPixel();
 
 	TArray<FKeyRenderer::FKeysForModel> Results;
 	KeyRenderer.HitTestKeys(RelativeTimeToPixel.PixelToFrame(LocalMousePixelX), Results);
@@ -685,14 +655,7 @@ void SChannelView::CreateKeysUnderMouse(const FVector2D& MousePosition, const FG
 		return;
 	}
 
-	TRange<FFrameNumber> SectionRange = Section->GetRange();
-
-	// Make the time <-> pixel converter relative to the section range rather than the view range
-	FTimeToPixel RelativeTimeToPixel = *TrackAreaTimeToPixel;
-	if (SectionRange.GetLowerBound().IsClosed())
-	{
-		RelativeTimeToPixel = TrackAreaTimeToPixel->RelativeTo(SectionRange.GetLowerBoundValue());
-	}
+	FTimeToPixel RelativeTimeToPixel = GetRelativeTimeToPixel();
 
 	// If the pressed key exists, offset the new key and look for it in the newly laid out key areas
 	if (InPressedKeys.Num())
diff --git a/Engine/Source/Editor/Sequencer/Public/ISequencerChannelInterface.h b/Engine/Source/Editor/Sequencer/Public/ISequencerChannelInterface.h
--- a/Engine/Source/Editor/Sequencer/Public/ISequencerChannelInterface.h
+++ b/Engine/Source/Editor/Sequencer/Public/ISequencerChannelInterface.h
@@ -28,6 +28,13 @@ class FMovieSceneClipboardBuilder;
 class FMovieSceneClipboardKeyTrack;
 class FTrackInstancePropertyBindings;
 
+namespace UE::Sequencer
+{
+	class FChannelModel;
+	class STrackAreaLaneView;
+	struct FCreateTrackLaneViewParams;
+}
+
 /** Utility struct representing a number of selected keys on a single channel */
 struct FExtendKeyMenuParams
 {
@@ -166,6 +173,25 @@ struct ISequencerChannelInterface
 	 */
 	virtual TUniquePtr<FCurveModel> CreateCurveEditorModel_Raw(const FMovieSceneChannelHandle& Channel, UMovieSceneSection* OwningSection, TSharedRef<ISequencer> InSequencer) const = 0;
 
+	/**
+	 * Create a new channel model for this type of channel
+	 *
+	 * @param InChannelHandle    The channel handle to create a model for
+	 * @param InChannelName      The identifying name of this channel
+	 * @return (Optional) A new model to be added to a curve editor
+	 */
+	virtual TSharedPtr<UE::Sequencer::FChannelModel> CreateChannelModel_Raw(const FMovieSceneChannelHandle& InChannelHandle, FName InChannelName) const = 0;
+
+	/**
+	 * Create a new channel view for this type of channel
+	 *
+	 * @param InChannelHandle    The channel handle to create a model for
+	 * @param InWeakModel        The model that is creating the view. Should not be Pinned persistently.
+	 * @param Parameters         View construction parameters
+	 * @return (Optional) A new model to be added to a curve editor
+	 */
+	virtual TSharedPtr<UE::Sequencer::STrackAreaLaneView> CreateChannelView_Raw(const FMovieSceneChannelHandle& InChannelHandle, TWeakPtr<UE::Sequencer::FChannelModel> InWeakModel, const UE::Sequencer::FCreateTrackLaneViewParams& Parameters) const = 0;
+
 	/**
 	 * Draw additional content in addition to keys for a particular channel
 	 *
diff --git a/Engine/Source/Editor/Sequencer/Public/MVVM/ViewModels/ChannelModel.h b/Engine/Source/Editor/Sequencer/Public/MVVM/ViewModels/ChannelModel.h
--- a/Engine/Source/Editor/Sequencer/Public/MVVM/ViewModels/ChannelModel.h
+++ b/Engine/Source/Editor/Sequencer/Public/MVVM/ViewModels/ChannelModel.h
@@ -48,7 +48,7 @@ class FSectionModel;
  * Model for a single channel inside a section.
  * For instance, this represents the "Location.X" channel of a single transform section.
  */
-class FChannelModel
+class SEQUENCER_API FChannelModel
 	: public FViewModel
 	, public FLinkedOutlinerExtension
 	, public FGeometryExtensionShim
diff --git a/Engine/Source/Editor/Sequencer/Public/MVVM/Views/SChannelView.h b/Engine/Source/Editor/Sequencer/Public/MVVM/Views/SChannelView.h
--- a/Engine/Source/Editor/Sequencer/Public/MVVM/Views/SChannelView.h
+++ b/Engine/Source/Editor/Sequencer/Public/MVVM/Views/SChannelView.h
@@ -6,6 +6,7 @@
 #include "Widgets/SCompoundWidget.h"
 #include "MVVM/Extensions/ITrackLaneExtension.h"
 #include "MVVM/Views/KeyRenderer.h"
+#include "MVVM/Views/STrackAreaLaneView.h"
 #include "MVVM/ViewModelPtr.h"
 
 #include "Templates/SharedPointer.h"
@@ -69,15 +70,15 @@ struct FChannelViewKeyCachedState
  * Widget for displaying a channel within a section
  */
 class SEQUENCER_API SChannelView
-	: public SCompoundWidget
-	, public ITrackLaneWidget
+	: public STrackAreaLaneView
 {
 public:
 	SLATE_BEGIN_ARGS(SChannelView){}
 		SLATE_ATTRIBUTE(FLinearColor, KeyBarColor)
+		SLATE_DEFAULT_SLOT(FArguments, Content)
 	SLATE_END_ARGS()
 
-	void Construct(const FArguments& InArgs, const FViewModelPtr& InViewModel, TSharedPtr<FTimeToPixel> InTimeToPixel, TSharedPtr<FTrackAreaViewModel> TrackArea);
+	void Construct(const FArguments& InArgs, const FViewModelPtr& InViewModel, TSharedPtr<STrackAreaView> InTrackAreaView);
 
 	/*~ SWidget */
 	virtual int32 OnPaint(const FPaintArgs& Args, const FGeometry& AllottedGeometry, const FSlateRect& MyCullingRect, FSlateWindowElementList& OutDrawElements, int32 LayerId, const FWidgetStyle& InWidgetStyle, bool bParentEnabled) const override;
@@ -87,11 +88,7 @@ public:
 	virtual FReply OnMouseButtonUp(const FGeometry& MyGeometry, const FPointerEvent& MouseEvent) override;
 	virtual void OnMouseLeave(const FPointerEvent& MouseEvent) override;
 
-	/*~ ITrackLaneWidget */
-	virtual TSharedRef<const SWidget> AsWidget() const;
-	virtual FTrackLaneScreenAlignment GetAlignment(const FTimeToPixel& InTimeToPixel, const FGeometry& InParentGeometry) const;
-
-private:
+protected:
 
 	int32 DrawLane(const FPaintArgs& Args, const FGeometry& AllottedGeometry, const FSlateRect& MyCullingRect, FSlateWindowElementList& OutDrawElements, int32 LayerId, const FWidgetStyle& InWidgetStyle, bool bParentEnabled) const;
 
@@ -103,12 +100,8 @@ private:
 
 	TSharedPtr<FSequencer> LegacyGetSequencer() const;
 
-private:
+protected:
 
-	FWeakViewModelPtr WeakModel;
-	TWeakPtr<STrackLane> WeakTrackLane;
-	TWeakPtr<FTrackAreaViewModel> WeakTrackArea;
-	TSharedPtr<FTimeToPixel> TrackAreaTimeToPixel;
 	FKeyRenderer KeyRenderer;
 	mutable TOptional<FChannelViewKeyCachedState> KeyRendererCache;
 
diff --git a/Engine/Source/Editor/Sequencer/Public/SequencerChannelInterface.h b/Engine/Source/Editor/Sequencer/Public/SequencerChannelInterface.h
--- a/Engine/Source/Editor/Sequencer/Public/SequencerChannelInterface.h
+++ b/Engine/Source/Editor/Sequencer/Public/SequencerChannelInterface.h
@@ -189,6 +189,33 @@ struct TSequencerChannelInterfaceCommon : ISequencerChannelInterface
 		return CreateCurveEditorModel(InChannel.Cast<ChannelType>(), OwningSection, InSequencer);
 	}
 
+	/**
+	 * Create a new channel model for this type of channel
+	 *
+	 * @param InChannelHandle    The channel handle to create a model for
+	 * @param InChannelName      The identifying name of this channel
+	 * @return (Optional) A new model to be added to a curve editor
+	 */
+	virtual TSharedPtr<UE::Sequencer::FChannelModel> CreateChannelModel_Raw(const FMovieSceneChannelHandle& InChannelHandle, FName InChannelName) const override
+	{
+		using namespace Sequencer;
+		return CreateChannelModel(InChannelHandle.Cast<ChannelType>(), InChannelName);
+	}
+
+	/**
+	 * Create a new channel view for this type of channel
+	 *
+	 * @param InChannelHandle    The channel handle to create a model for
+	 * @param InWeakModel        The model that is creating the view. Should not be Pinned persistently.
+	 * @param Parameters         View construction parameters
+	 * @return (Optional) A new model to be added to a curve editor
+	 */
+	virtual TSharedPtr<UE::Sequencer::STrackAreaLaneView> CreateChannelView_Raw(const FMovieSceneChannelHandle& InChannelHandle, TWeakPtr<UE::Sequencer::FChannelModel> InWeakModel, const UE::Sequencer::FCreateTrackLaneViewParams& Parameters) const override
+	{
+		using namespace Sequencer;
+		return CreateChannelView(InChannelHandle.Cast<ChannelType>(), InWeakModel, Parameters);
+	}
+
 	/**
 	 * Create an editor on the sequencer node tree
 	 *
diff --git a/Engine/Source/Editor/Sequencer/Public/SequencerChannelTraits.h b/Engine/Source/Editor/Sequencer/Public/SequencerChannelTraits.h
--- a/Engine/Source/Editor/Sequencer/Public/SequencerChannelTraits.h
+++ b/Engine/Source/Editor/Sequencer/Public/SequencerChannelTraits.h
@@ -74,6 +74,14 @@ struct FSequencerChannelPaintArgs
 	bool bParentEnabled = true;
 };
 
+namespace UE::Sequencer
+{
+	class FChannelModel;
+	class STrackAreaLaneView;
+
+	struct FCreateTrackLaneViewParams;
+}
+
 /**
  * Stub/default implementations for ISequencerChannelInterface functions.
  * Custom behaviour should be implemented by overloading the relevant function with the necessary channel/data types.
@@ -405,6 +413,30 @@ namespace Sequencer
 	 */
 	SEQUENCER_API TUniquePtr<FCurveModel> CreateCurveEditorModel(const FMovieSceneChannelHandle& ChannelHandle, UMovieSceneSection* OwningSection, TSharedRef<ISequencer> InSequencer);
 
+	/**
+	 * Create a new channel model for this type of channel
+	 *
+	 * @param InChannelHandle    The channel handle to create a model for
+	 * @param InChannelName      The identifying name of this channel
+	 * @return (Optional) A new model to be used as part of the Sequencer MVVM framework
+	 */
+	inline TSharedPtr<UE::Sequencer::FChannelModel> CreateChannelModel(const FMovieSceneChannelHandle& InChannelHandle, FName InChannelName)
+	{
+		return nullptr;
+	}
+
+	/**
+	 * Create a new channel view for this type of channel
+	 *
+	 * @param InChannelHandle    The channel handle to create a model for
+	 * @param InWeakModel        The model that is creating the view. Should not be Pinned persistently.
+	 * @param Parameters         View construction parameters
+	 * @return (Optional) A new view to be shown on the track area
+	 */
+	inline TSharedPtr<UE::Sequencer::STrackAreaLaneView> CreateChannelView(const FMovieSceneChannelHandle& InChannelHandle, TWeakPtr<UE::Sequencer::FChannelModel> InWeakModel, const UE::Sequencer::FCreateTrackLaneViewParams& Parameters)
+	{
+		return nullptr;
+	}	
 
 	/**
 	 * Whether this channel should draw a curve on its editor UI
diff --git a/Engine/Source/Editor/SequencerCore/Private/MVVM/Views/STrackAreaLaneView.cpp b/Engine/Source/Editor/SequencerCore/Private/MVVM/Views/STrackAreaLaneView.cpp
new file mode 100644
--- /dev/null
+++ b/Engine/Source/Editor/SequencerCore/Private/MVVM/Views/STrackAreaLaneView.cpp
@@ -0,0 +1,57 @@
+// Copyright Epic Games, Inc. All Rights Reserved.
+
+#include "MVVM/Views/STrackAreaLaneView.h"
+#include "MVVM/ViewModels/TrackAreaViewModel.h"
+
+
+namespace UE::Sequencer
+{
+
+void STrackAreaLaneView::Construct(const FArguments& InArgs, const FViewModelPtr& InViewModel, TSharedPtr<STrackAreaView> InTrackAreaView)
+{
+	WeakModel = InViewModel;
+	WeakTrackAreaView = InTrackAreaView;
+	TrackAreaTimeToPixel = InTrackAreaView->GetTimeToPixel();
+
+	ChildSlot
+	[
+		InArgs._Content.Widget
+	];
+}
+
+TSharedRef<const SWidget> STrackAreaLaneView::AsWidget() const
+{
+	return AsShared();
+}
+
+FTrackLaneScreenAlignment STrackAreaLaneView::GetAlignment(const FTimeToPixel& InTimeToPixel, const FGeometry& InParentGeometry) const
+{
+	TSharedPtr<ITrackLaneExtension> TrackLaneExtension = WeakModel.ImplicitPin();
+	if (TrackLaneExtension)
+	{
+		FTrackLaneVirtualAlignment VirtualAlignment = TrackLaneExtension->ArrangeVirtualTrackLaneView();
+		FTrackLaneScreenAlignment  ScreenAlignment  = VirtualAlignment.ToScreen(InTimeToPixel, InParentGeometry);
+
+		return ScreenAlignment;
+	}
+	return FTrackLaneScreenAlignment();
+}
+
+FTimeToPixel STrackAreaLaneView::GetRelativeTimeToPixel() const
+{
+	FTimeToPixel RelativeTimeToPixel = *TrackAreaTimeToPixel;
+
+	TSharedPtr<ITrackLaneExtension> TrackLaneExtension = WeakModel.ImplicitPin();
+	if (TrackLaneExtension)
+	{
+		FTrackLaneVirtualAlignment VirtualAlignment = TrackLaneExtension->ArrangeVirtualTrackLaneView();
+		if (VirtualAlignment.Range.GetLowerBound().IsClosed())
+		{
+			RelativeTimeToPixel = TrackAreaTimeToPixel->RelativeTo(VirtualAlignment.Range.GetLowerBoundValue());
+		}
+	}
+
+	return RelativeTimeToPixel;
+}
+
+} // namespace UE::Sequencer
diff --git a/Engine/Source/Editor/SequencerCore/Public/MVVM/Views/STrackAreaLaneView.h b/Engine/Source/Editor/SequencerCore/Public/MVVM/Views/STrackAreaLaneView.h
new file mode 100644
--- /dev/null
+++ b/Engine/Source/Editor/SequencerCore/Public/MVVM/Views/STrackAreaLaneView.h
@@ -0,0 +1,40 @@
+// Copyright Epic Games, Inc. All Rights Reserved.
+
+#pragma once
+
+#include "CoreTypes.h"
+#include "Widgets/SCompoundWidget.h"
+#include "MVVM/ViewModelPtr.h"
+#include "MVVM/Extensions/ITrackLaneExtension.h"
+
+
+namespace UE::Sequencer
+{
+
+class STrackAreaView;
+
+class SEQUENCERCORE_API STrackAreaLaneView
+	: public SCompoundWidget
+	, public ITrackLaneWidget
+{
+public:
+	SLATE_BEGIN_ARGS(STrackAreaLaneView){}
+		SLATE_DEFAULT_SLOT(FArguments, Content)
+	SLATE_END_ARGS()
+
+	void Construct(const FArguments& InArgs, const FViewModelPtr& InViewModel, TSharedPtr<STrackAreaView> InTrackAreaView);
+
+	/*~ ITrackLaneWidget */
+	virtual TSharedRef<const SWidget> AsWidget() const;
+	virtual FTrackLaneScreenAlignment GetAlignment(const FTimeToPixel& InTimeToPixel, const FGeometry& InParentGeometry) const;
+
+	FTimeToPixel GetRelativeTimeToPixel() const;
+
+protected:
+
+	FWeakViewModelPtr WeakModel;
+	TWeakPtr<STrackAreaView> WeakTrackAreaView;
+	TSharedPtr<FTimeToPixel> TrackAreaTimeToPixel;
+};
+
+} // namespace UE::Sequencer
