Sequencer: Optimized key interactions for large numbers of keys

A few changes here:

 - FKeyHandleLookupTable used to roll its own sparse array using a TArray of TOptionals. It was also designed to lazily update its map of handles to indices if it encountered one that was not correct. Unfortunately this results in an n^2 operation if you try and resolve key handles for all the elements in a channel when they are not valid (because it has to look through all the key handles to find the correct index). I have changed this to always keep the handle/index mappings up to date when moving keys around. In this way the operation scales poorly with the number of keys added or removed from the channel, rather than the total number of keys (which is still bad, but not as bad as always scaling poorly for large numbers of keys).

  - Various optimizations to FKeyRenderer. When preserving computed keys (ie, when we are simply padding new keys into or old keys out of view), we no longer re-allocate the entire array (which can be very large), and instead opt to relocate-construct the preserved keys within the same allocation. When the array gets large enough we always end up going to the crt allocator which is extremely slow, so keeping the allocation alive is beneficial here.

  - FKeyRenderer cached key computation is now performed using a parallel for loop when the number of channels or keys is large. This should really be improved by doing the same work for _all keys rendered in the current frame_ rahter than only going wide for each section, but this is a good win for control rigs at least.

  - When determining whether to draw a group of keys as selected, hovered or preview selected, we now only consider the first 500 keys before assuming all the others have the same state. This saves us from iterating thousands and thousands of keys to see if they are all selected when zoomed out very far.

 - FSequencerSelectedKey construction and comparison was taking up a lot of CPU. I have added raw FKeyHandle sets to the various selection classes for direct comparisons which speeds up checking for selected and hovered keys greatly. FKeyHandle is a unique index regardless of what channel it exists on, so this is safe.

diff --git a/Engine/Source/Editor/Sequencer/Private/MVVM/Views/SChannelView.cpp b/Engine/Source/Editor/Sequencer/Private/MVVM/Views/SChannelView.cpp
--- a/Engine/Source/Editor/Sequencer/Private/MVVM/Views/SChannelView.cpp
+++ b/Engine/Source/Editor/Sequencer/Private/MVVM/Views/SChannelView.cpp
@@ -48,53 +48,43 @@ TArray<FSequencerSelectedKey> KeyRendererResultToSelectedKeys(TArrayView<const F
 /** Key renderer interface for defining selection and hover states. Should be removed in favor of SequencerCore knowing about selection in the future */
 struct FSequencerKeyRendererInterface : IKeyRendererInterface
 {
-	UMovieSceneSection* SectionObject;
-	const TSet<FSequencerSelectedKey>* SelectedKeys;
-	const TMap<FSequencerSelectedKey, ESelectionPreviewState>* SelectionPreview;
-	TArrayView<const FSequencerSelectedKey> HoveredKeys;
+	const TSet<FKeyHandle>* SelectedKeys;
+	const TMap<FKeyHandle, ESelectionPreviewState>* SelectionPreview;
+	const TSet<FKeyHandle>* HoveredKeys;
 
-	FSequencerKeyRendererInterface(UMovieSceneSection* InSectionObject, FSequencer* Sequencer, TSharedPtr<ITrackAreaHotspot> Hotspot)
+	FSequencerKeyRendererInterface(FSequencer* Sequencer, TSharedPtr<ITrackAreaHotspot> Hotspot)
 	{
-		SectionObject    = InSectionObject;
-		SelectedKeys     = &Sequencer->GetSelection().GetSelectedKeys();
-		SelectionPreview = &Sequencer->GetSelectionPreview().GetDefinedKeyStates();
-		if (TSharedPtr<FKeyHotspot> KeyHotspot = HotspotCast<FKeyHotspot>(Hotspot))
-		{
-			HoveredKeys = KeyHotspot->Keys;
-		}
+		TSharedPtr<FKeyHotspot> KeyHotspot = HotspotCast<FKeyHotspot>(Hotspot);
+
+		SelectedKeys     = &Sequencer->GetSelection().GetRawSelectedKeys();
+		SelectionPreview = &Sequencer->GetSelectionPreview().GetRawDefinedKeyStates();
+		HoveredKeys      = KeyHotspot ? &KeyHotspot->RawKeys : nullptr;
 	}
 
-	bool HasAnySelectedKeys() const
+	bool HasAnySelectedKeys() const override
 	{
 		return SelectedKeys->Num() != 0;
 	}
-	bool HasAnyPreviewSelectedKeys() const
+	bool HasAnyPreviewSelectedKeys() const override
 	{
 		return SelectionPreview->Num() != 0;
 	}
-	bool HasAnyHoveredKeys() const
+	bool HasAnyHoveredKeys() const override
 	{
-		return HoveredKeys.Num() != 0;
+		return HoveredKeys && HoveredKeys->Num() != 0;
 	}
 
-	bool IsKeySelected(const TViewModelPtr<IKeyExtension>& InOwner, FKeyHandle InKey) const
+	bool IsKeySelected(const TViewModelPtr<IKeyExtension>& InOwner, FKeyHandle InKey) const override
 	{
-		TSharedPtr<FChannelModel> Channel = InOwner.ImplicitCast();
-		FSequencerSelectedKey SelectedKey(*SectionObject, Channel, InKey);
-		return SelectedKeys->Contains(SelectedKey);
+		return SelectedKeys->Contains(InKey);
 	}
-	bool IsKeyHovered(const TViewModelPtr<IKeyExtension>& InOwner, FKeyHandle InKey) const
+	bool IsKeyHovered(const TViewModelPtr<IKeyExtension>& InOwner, FKeyHandle InKey) const override
 	{
-		TSharedPtr<FChannelModel> Channel = InOwner.ImplicitCast();
-		FSequencerSelectedKey SelectedKey(*SectionObject, Channel, InKey);
-		return HoveredKeys.Contains(SelectedKey);
+		return HoveredKeys ? HoveredKeys->Contains(InKey) : false;
 	}
-	EKeySelectionPreviewState GetPreviewSelectionState(const TViewModelPtr<IKeyExtension>& InOwner, FKeyHandle InKey) const
+	EKeySelectionPreviewState GetPreviewSelectionState(const TViewModelPtr<IKeyExtension>& InOwner, FKeyHandle InKey) const override
 	{
-		TSharedPtr<FChannelModel> Channel = InOwner.ImplicitCast();
-		FSequencerSelectedKey SelectedKey(*SectionObject, Channel, InKey);
-
-		if (const ESelectionPreviewState* State = SelectionPreview->Find(SelectedKey))
+		if (const ESelectionPreviewState* State = SelectionPreview->Find(InKey))
 		{
 			switch(*State)
 			{
@@ -119,6 +109,7 @@ FChannelViewKeyCachedState::FChannelViewKeyCachedState()
 	, bShowCurve(false)
 	, bShowKeyBars(true)
 	, bCollapseChildren(false)
+	, bIsChannelHovered(false)
 {}
 
 FChannelViewKeyCachedState::FChannelViewKeyCachedState(TRange<FFrameTime> InVisibleRange, TSharedPtr<ITrackAreaHotspot> Hotspot, FViewModelPtr Model, FSequencer* Sequencer)
@@ -139,6 +130,7 @@ FChannelViewKeyCachedState::FChannelViewKeyCachedState(TRange<FFrameTime> InVisi
 	bShowCurve = Channel && Channel->GetKeyArea()->ShouldShowCurve();
 	bShowKeyBars = Sequencer->GetSequencerSettings()->GetShowKeyBars();
 	bCollapseChildren = Outliner && Outliner->GetLinkedOutlinerItem() && Outliner->GetLinkedOutlinerItem()->IsExpanded() == false;
+	bIsChannelHovered = false;
 	if (!Hotspot || Hotspot->CastThis<FKeyHotspot>() || Hotspot->CastThis<FKeyBarHotspot>())
 	{
 		WeakHotspot = Hotspot;
@@ -150,19 +142,22 @@ EViewDependentCacheFlags FChannelViewKeyCachedState::CompareTo(const FChannelVie
 {
 	EViewDependentCacheFlags Flags = EViewDependentCacheFlags::None;
 
-	// Check if either hotspot expired
-	if ( (WeakHotspot.Pin() == nullptr && RawHotspotDoNotUse != nullptr) ||
-		 (Other.WeakHotspot.Pin() == nullptr && Other.RawHotspotDoNotUse != nullptr) )
+	if (bIsChannelHovered || bIsChannelHovered != Other.bIsChannelHovered)
 	{
-		Flags |= EViewDependentCacheFlags::KeyStateChanged;
-	}
+		// Check if either hotspot expired
+		if ( (WeakHotspot.Pin() == nullptr && RawHotspotDoNotUse != nullptr) ||
+			 (Other.WeakHotspot.Pin() == nullptr && Other.RawHotspotDoNotUse != nullptr) )
+		{
+			Flags |= EViewDependentCacheFlags::KeyStateChanged;
+		}
 
-	// Check if the hotspot has changed (this is necessary because theoretically both hotspots
-	// could still be _alive_ but Other.WeakHotspot is the active one)
-	if (WeakHotspot != Other.WeakHotspot)
-	{
-		// Hotspot has changed - maybe we (un)hovered some keys or a key bar
-		Flags |= EViewDependentCacheFlags::KeyStateChanged;
+		// Check if the hotspot has changed (this is necessary because theoretically both hotspots
+		// could still be _alive_ but Other.WeakHotspot is the active one)
+		if (WeakHotspot != Other.WeakHotspot)
+		{
+			// Hotspot has changed - maybe we (un)hovered some keys or a key bar
+			Flags |= EViewDependentCacheFlags::KeyStateChanged;
+		}
 	}
 
 	if (ValidPlayRangeMin != Other.ValidPlayRangeMin || ValidPlayRangeMax != Other.ValidPlayRangeMax)
@@ -236,7 +231,7 @@ FReply SChannelView::OnMouseButtonDown(const FGeometry& MyGeometry, const FPoint
 		return FReply::Unhandled();
 	}
 
-	TArrayView<const FSequencerSelectedKey> HoveredKeys;
+	TSet<FSequencerSelectedKey> HoveredKeys;
 
 	// The hovered key is defined from the sequencer hotspot
 	if (TSharedPtr<FKeyHotspot> KeyHotspot = HotspotCast<FKeyHotspot>(TrackArea->GetHotspot()))
@@ -275,13 +270,15 @@ FReply SChannelView::OnMouseButtonDown(const FGeometry& MyGeometry, const FPoint
 FReply SChannelView::OnMouseMove(const FGeometry& MyGeometry, const FPointerEvent& MouseEvent)
 {
 	TSharedPtr<STrackAreaView>      TrackAreaView = WeakTrackAreaView.Pin();
+	TViewModelPtr<FSectionModel>    Section       = GetSection();
 	TSharedPtr<FTrackAreaViewModel> TrackArea     = TrackAreaView ? TrackAreaView->GetViewModel() : nullptr;
-	if (!TrackArea)
+	if (!TrackArea || !Section)
 	{
 		return FReply::Unhandled();
 	}
 
 	TSharedPtr<FSequencer> Sequencer = LegacyGetSequencer();
+	
 
 	// Don't fiddle with hotspots if there is an active drag
 	const ISequencerEditTool* EditTool = TrackArea->GetEditTool();
@@ -291,51 +288,78 @@ FReply SChannelView::OnMouseMove(const FGeometry& MyGeometry, const FPointerEven
 	}
 
 	// Checked for hovered key
-	TArray<FSequencerSelectedKey> KeysUnderMouse;
-	GetKeysUnderMouse(MouseEvent.GetScreenSpacePosition(), MyGeometry, KeysUnderMouse);
-	if (KeysUnderMouse.Num())
+	TSharedPtr<FKeyHotspot>    KeyHotspot    = HotspotCast<FKeyHotspot>(TrackArea->GetHotspot());
+	TSharedPtr<FKeyBarHotspot> KeyBarHotspot = HotspotCast<FKeyBarHotspot>(TrackArea->GetHotspot());
+	
+	const FTimeToPixel RelativeTimeToPixel = GetRelativeTimeToPixel();
+	const FVector2D    MousePixel          = MyGeometry.AbsoluteToLocal(MouseEvent.GetScreenSpacePosition());
+	const float        HalfHeightPx        = MyGeometry.GetLocalSize().Y*.5f;
+
+
+	// Perform validation of existing hotspots, and hittesting of new ones if necessary.
+	// WARNING: This code will early-return once it has found a valid hotspot, so no 
+	//          code should be added after here if it must always run
+	if (KeyRendererCache.IsSet())
 	{
-		TSharedPtr<FKeyHotspot> KeyHotspot = HotspotCast<FKeyHotspot>(TrackArea->GetHotspot());
-		if (!KeyHotspot || KeyHotspot->Keys != KeysUnderMouse)
+		if (KeyHotspot)
+		{
+			const FFrameTime KeyTime       = KeyHotspot->GetTime().Get(0);
+			const FVector2D  HalfKeySizePx = KeyRendererCache->KeySizePx * .5f;
+			const float      KeyPixel      = RelativeTimeToPixel.FrameToPixel(KeyTime);
+
+			if (FMath::Abs(MousePixel.Y - HalfHeightPx) <= HalfKeySizePx.Y &&
+				FMath::Abs(MousePixel.X - KeyPixel) <= HalfKeySizePx.X)
+			{
+				// Key hotspot is still valid
+				return FReply::Handled();
+			}
+		}
+		else if (KeyBarHotspot)
 		{
-			TrackArea->SetHotspot(MakeShared<FKeyHotspot>(MoveTemp(KeysUnderMouse), Sequencer));
+			if (FMath::Abs(MousePixel.Y - HalfHeightPx) <= 4.f &&
+				KeyBarHotspot->Range.Contains(RelativeTimeToPixel.PixelToFrame(MousePixel.X)))
+			{
+				// Keybar hotspot is still valid
+				return FReply::Handled();
+			}
 		}
+	}
+
+	TArray<FSequencerSelectedKey> KeysUnderMouse;
+	GetKeysUnderMouse(MouseEvent.GetScreenSpacePosition(), MyGeometry, KeysUnderMouse);
 
+	// If we are hovering over a new key, set that as the hotspot
+	if (KeysUnderMouse.Num())
+	{
+		TrackArea->SetHotspot(MakeShared<FKeyHotspot>(MoveTemp(KeysUnderMouse), Sequencer));
 		return FReply::Handled();
 	}
 
+	// Keys always render in the middle of the geometry so do a quick vertical test first
+	// We actually hittest a larger vertical area than the bars render to make it easier for user
+	// to hit key bars
+	if (FMath::Abs(MousePixel.Y - HalfHeightPx) <= 4.f)
 	{
-		TViewModelPtr<FSectionModel> Section = GetSection();
-		if (Section)
+		FKeyRenderer::FKeyBar KeyBar;
+		if (KeyRenderer.HitTestKeyBar(RelativeTimeToPixel.PixelToFrame(MousePixel.X), KeyBar))
 		{
-			const float HalfHeightPx = MyGeometry.GetLocalSize().Y*.5f;
-			const FVector2D MousePixel = MyGeometry.AbsoluteToLocal(MouseEvent.GetScreenSpacePosition());
-
-			// Keys always render in the middle of the geometry so do a quick vertical test first
-			// We actually hittest a larger vertical area than the bars render to make it easier for user
-			// to hit key bars
-			if (FMath::Abs(MousePixel.Y - HalfHeightPx) <= 4.f)
-			{
-				TRange<FFrameNumber> SectionRange = Section->GetRange();
-
-				// Make the time <-> pixel converter relative to the section range rather than the view range
-				FTimeToPixel RelativeTimeToPixel = GetRelativeTimeToPixel();
-
-				TArray<FKeyRenderer::FKeysForModel> LeadingKeys, TrailingKeys;
-				if (KeyRenderer.HitTestKeyBar(RelativeTimeToPixel.PixelToFrame(MousePixel.X), LeadingKeys, TrailingKeys))
-				{
-					UMovieSceneSection* SectionObject = Section->GetSection();
-					TrackArea->SetHotspot(MakeShared<FKeyBarHotspot>(
-						KeyRendererResultToSelectedKeys(LeadingKeys, SectionObject),
-						KeyRendererResultToSelectedKeys(TrailingKeys, SectionObject),
-						Sequencer)
-					);
-					return FReply::Handled();
-				}
-			}
+			UMovieSceneSection* SectionObject = Section->GetSection();
+			TrackArea->SetHotspot(MakeShared<FKeyBarHotspot>(
+				KeyBar.Range,
+				KeyRendererResultToSelectedKeys(KeyBar.LeadingKeys, SectionObject),
+				KeyRendererResultToSelectedKeys(KeyBar.TrailingKeys, SectionObject),
+				Sequencer)
+			);
+
+			return FReply::Handled();
 		}
 	}
 
+	// Otherwise reset the hotspot if it is no longer valid
+	if (KeyHotspot || KeyBarHotspot)
+	{
+		TrackArea->SetHotspot(nullptr);
+	}
 	return FReply::Unhandled();
 }
 
@@ -457,8 +481,9 @@ int32 SChannelView::OnPaint(const FPaintArgs& Args, const FGeometry& AllottedGeo
 		LayerId = Channel->CustomPaint(ChannelPaintArgs, LayerId);
 	}
 
-	FSequencerKeyRendererInterface KeyRenderInterface(SectionObject, Sequencer.Get(), TrackArea->GetHotspot());
+	FSequencerKeyRendererInterface KeyRenderInterface(Sequencer.Get(), TrackArea->GetHotspot());
 	FChannelViewKeyCachedState NewCachedState(VisibleRange, TrackArea->GetHotspot(), Model, Sequencer.Get());
+	NewCachedState.bIsChannelHovered = IsHovered();
 
 	FKeyBatchParameters Params(RelativeTimeToPixel);
 	Params.CacheState = KeyRendererCache.IsSet() ? NewCachedState.CompareTo(KeyRendererCache.GetValue()) : EViewDependentCacheFlags::All;
@@ -635,7 +660,7 @@ void SChannelView::GetKeysAtPixelX(float LocalMousePixelX, TArray<FSequencerSele
 	OutKeys = KeyRendererResultToSelectedKeys(Results, SectionObject);
 }
 
-void SChannelView::CreateKeysUnderMouse(const FVector2D& MousePosition, const FGeometry& AllottedGeometry, TArrayView<const FSequencerSelectedKey> InPressedKeys, TArray<FSequencerSelectedKey>& OutKeys)
+void SChannelView::CreateKeysUnderMouse(const FVector2D& MousePosition, const FGeometry& AllottedGeometry, const TSet<FSequencerSelectedKey>& InPressedKeys, TArray<FSequencerSelectedKey>& OutKeys)
 {
 	FViewModelPtr Model = WeakModel.Pin();
 	TViewModelPtr<FSectionModel> Section = GetSection();
diff --git a/Engine/Source/Editor/Sequencer/Private/SSequencerSection.cpp b/Engine/Source/Editor/Sequencer/Private/SSequencerSection.cpp
--- a/Engine/Source/Editor/Sequencer/Private/SSequencerSection.cpp
+++ b/Engine/Source/Editor/Sequencer/Private/SSequencerSection.cpp
@@ -359,6 +359,14 @@ struct FSequencerSectionPainterImpl : FSequencerSectionPainter
 			return;
 		}
 
+		TSharedPtr<STrackLane> TrackLane = SectionWidget.WeakOwningTrackLane.Pin();
+		if (!TrackLane)
+		{
+			return;
+		}
+
+		const float LaneHeight = TrackLane->GetDesiredSize().X;
+
 		TSharedRef<FSlateFontCache> FontCache = FSlateApplication::Get().GetRenderer()->GetFontCache();
 
 		UEnum* Enum = FindObjectChecked<UEnum>(nullptr, TEXT("/Script/MovieScene.EMovieSceneBlendType"), true);
@@ -376,15 +384,17 @@ struct FSequencerSectionPainterImpl : FSequencerSectionPainter
 			FontInfo.Size = FMath::Max(FMath::FloorToInt(FontInfo.Size - 6.f), 11);
 		}
 
-		FVector2D TextOffset = SectionModel->GetRange() == TRange<FFrameNumber>::All() ? FVector2D(0.f, -1.f) : FVector2D(1.f, -1.f);
-		FVector2D BottomLeft = SectionGeometry.AbsoluteToLocal(SectionClippingRect.GetBottomLeft()) + TextOffset;
+		const float FontHeight = GetFontHeight();
+
+		FVector2D TextOffset = SectionModel->GetRange() == TRange<FFrameNumber>::All() ? FVector2D(0.f, LaneHeight - FontHeight - 1.f) : FVector2D(1.f, LaneHeight - FontHeight - 1.f);
+		FVector2D TextPosition = SectionGeometry.AbsoluteToLocal(SectionClippingRect.GetTopLeft()) + TextOffset;
 
 		FSlateDrawElement::MakeText(
 			DrawElements,
 			LayerId,
 			SectionGeometry.MakeChild(
-				FVector2D(SectionGeometry.Size.X, GetFontHeight()),
-				FSlateLayoutTransform(BottomLeft - FVector2D(0.f, GetFontHeight()+1.f))
+				FVector2D(SectionGeometry.Size.X, FontHeight),
+				FSlateLayoutTransform(TextPosition)
 			).ToPaintGeometry(),
 			DisplayText,
 			FontInfo,
@@ -851,7 +861,7 @@ void SSequencerSection::Construct( const FArguments& InArgs, TSharedPtr<FSequenc
 
 		+ SOverlay::Slot()
 		[
-			SNew(SChannelView, InSectionModel, OwningTrackLane->GetTimeToPixel(), OwningTrackLane->GetTrackAreaView()->GetViewModel())
+			SNew(SChannelView, InSectionModel, OwningTrackLane->GetTrackAreaView())
 			.Visibility(this, &SSequencerSection::GetTopLevelChannelGroupVisibility)
 		]
 	];
@@ -1269,13 +1279,19 @@ int32 SSequencerSection::OnPaint( const FPaintArgs& Args, const FGeometry& Allot
 	LayerId = SectionInterface->OnPaintSection(Painter);
 
 	FLinearColor SelectionColor = FAppStyle::GetSlateColor(SequencerSectionConstants::SelectionColorName).GetColor(FWidgetStyle());
-	DrawSectionHandles(AllottedGeometry, OutDrawElements, LayerId, DrawEffects, SelectionColor, Hotspot);
 
 	LayerId = SCompoundWidget::OnPaint( Args, AllottedGeometry, MyCullingRect, OutDrawElements, LayerId, InWidgetStyle, bEnabled );
 
+	// We paint section handles after channels so they draw on top, but we allow some leeway for channels to paint on a higher layer than they
+	// reported. This enables us to still draw and interact with keys that overlap section handles, even if we drew the handles chronologically later
+	DrawSectionHandles(AllottedGeometry, OutDrawElements, LayerId, DrawEffects, SelectionColor, Hotspot);
+
 	Painter.LayerId = LayerId;
 	PaintEasingHandles( Painter, SelectionColor, Hotspot );
 
+	// Artificially increase the layer now to ensure that we are now painting above keys now
+	LayerId += 10;
+
 	// --------------------------------------------
 	// Draw section selection tint
 	const float SectionThrobValue = GetSectionSelectionThrobValue();
@@ -1801,7 +1817,13 @@ bool SSequencerSection::IsSectionHighlighted(UMovieSceneSection* InSection, TSha
 
 	if (FKeyHotspot* KeyHotspot = Hotspot->CastThis<FKeyHotspot>())
 	{
-		return KeyHotspot->Keys.ContainsByPredicate([InSection](const FSequencerSelectedKey& Key){ return Key.Section == InSection; });
+		for (const FSequencerSelectedKey& Key : KeyHotspot->Keys)
+		{
+			if (Key.Section == InSection)
+			{
+				return true;
+			}
+		}
 	}
 	else if (FSectionEasingAreaHotspot* EasingAreaHotspot = Hotspot->CastThis<FSectionEasingAreaHotspot>())
 	{
diff --git a/Engine/Source/Editor/Sequencer/Private/SequencerHotspots.cpp b/Engine/Source/Editor/Sequencer/Private/SequencerHotspots.cpp
--- a/Engine/Source/Editor/Sequencer/Private/SequencerHotspots.cpp
+++ b/Engine/Source/Editor/Sequencer/Private/SequencerHotspots.cpp
@@ -147,15 +147,27 @@ void FHotspotSelectionManager::SelectModelExclusive(TSharedPtr<FViewModel> InMod
 	}
 }
 
+FKeyHotspot::FKeyHotspot(const TArray<FSequencerSelectedKey>& InKeys, TWeakPtr<FSequencer> InWeakSequencer)
+	: Keys(InKeys)
+	, WeakSequencer(InWeakSequencer)
+{ 
+	RawKeys.Reserve(Keys.Num());
+	for (const FSequencerSelectedKey& Key : InKeys)
+	{
+		RawKeys.Add(Key.KeyHandle);
+	}
+}
+
 void FKeyHotspot::HandleMouseSelection(FHotspotSelectionManager& SelectionManager)
 {
+	TArray<FSequencerSelectedKey> KeysArray = Keys.Array();
 	if (SelectionManager.MouseEvent->GetEffectingButton() == EKeys::RightMouseButton)
 	{
-		SelectionManager.SelectKeysExclusive(Keys);
+		SelectionManager.SelectKeysExclusive(KeysArray);
 	}
 	else
 	{
-		SelectionManager.ToggleKeys(Keys);
+		SelectionManager.ToggleKeys(KeysArray);
 	}
 }
 
@@ -168,11 +180,12 @@ TOptional<FFrameNumber> FKeyHotspot::GetTime() const
 {
 	FFrameNumber Time = 0;
 
-	if (Keys.Num())
+	for (const FSequencerSelectedKey& Key : Keys)
 	{
-		TArrayView<const FSequencerSelectedKey> FirstKey(&Keys[0], 1);
+		TArrayView<const FSequencerSelectedKey> FirstKey(&Key, 1);
 		TArrayView<FFrameNumber> FirstKeyTime(&Time, 1);
 		GetKeyTimes(FirstKey, FirstKeyTime);
+		break;
 	}
 
 	return Time;
diff --git a/Engine/Source/Editor/Sequencer/Private/SequencerHotspots.h b/Engine/Source/Editor/Sequencer/Private/SequencerHotspots.h
--- a/Engine/Source/Editor/Sequencer/Private/SequencerHotspots.h
+++ b/Engine/Source/Editor/Sequencer/Private/SequencerHotspots.h
@@ -63,10 +63,7 @@ struct FKeyHotspot
 {
 	UE_SEQUENCER_DECLARE_CASTABLE(FKeyHotspot, ITrackAreaHotspot, IMouseHandlerHotspot);
 
-	FKeyHotspot(TArray<FSequencerSelectedKey> InKeys, TWeakPtr<FSequencer> InWeakSequencer)
-		: Keys(MoveTemp(InKeys))
-		, WeakSequencer(InWeakSequencer)
-	{ }
+	SEQUENCER_API FKeyHotspot(const TArray<FSequencerSelectedKey>& InKeys, TWeakPtr<FSequencer> InWeakSequencer);
 
 	virtual void UpdateOnHover(FTrackAreaViewModel& InTrackArea) const override;
 	virtual TOptional<FFrameNumber> GetTime() const override;
@@ -75,7 +72,8 @@ struct FKeyHotspot
 	virtual void HandleMouseSelection(FHotspotSelectionManager& SelectionManager) override;
 
 	/** The keys that are part of this hotspot */
-	TArray<FSequencerSelectedKey> Keys;
+	TSet<FSequencerSelectedKey> Keys;
+	TSet<FKeyHandle> RawKeys;
 	TWeakPtr<FSequencer> WeakSequencer;
 };
 
@@ -86,10 +84,11 @@ struct FKeyBarHotspot
 {
 	UE_SEQUENCER_DECLARE_CASTABLE(FKeyBarHotspot, ITrackAreaHotspot);
 
-	FKeyBarHotspot(TArray<FSequencerSelectedKey>&& InLeadingKeys, TArray<FSequencerSelectedKey>&& InTrailingKeys, TWeakPtr<FSequencer> InWeakSequencer)
+	FKeyBarHotspot(const TRange<FFrameTime>& InRange, TArray<FSequencerSelectedKey>&& InLeadingKeys, TArray<FSequencerSelectedKey>&& InTrailingKeys, TWeakPtr<FSequencer> InWeakSequencer)
 		: LeadingKeys(MoveTemp(InLeadingKeys))
 		, TrailingKeys(MoveTemp(InTrailingKeys))
 		, WeakSequencer(InWeakSequencer)
+		, Range(InRange)
 	{ }
 
 	virtual void UpdateOnHover(FTrackAreaViewModel& InTrackArea) const override;
@@ -101,6 +100,7 @@ struct FKeyBarHotspot
 	TArray<FSequencerSelectedKey> LeadingKeys;
 	TArray<FSequencerSelectedKey> TrailingKeys;
 	TWeakPtr<FSequencer> WeakSequencer;
+	TRange<FFrameTime> Range;
 };
 
 /** A hotspot representing a section */
diff --git a/Engine/Source/Editor/Sequencer/Private/SequencerSelection.cpp b/Engine/Source/Editor/Sequencer/Private/SequencerSelection.cpp
--- a/Engine/Source/Editor/Sequencer/Private/SequencerSelection.cpp
+++ b/Engine/Source/Editor/Sequencer/Private/SequencerSelection.cpp
@@ -58,6 +58,11 @@ const TSet<FSequencerSelectedKey>& FSequencerSelection::GetSelectedKeys() const
 	return SelectedKeys;
 }
 
+const TSet<FKeyHandle>& FSequencerSelection::GetRawSelectedKeys() const
+{
+	return RawSelectedKeys;
+}
+
 TSet<TWeakObjectPtr<UMovieSceneSection>> FSequencerSelection::GetSelectedSections() const
 {
 	using namespace UE::Sequencer;
@@ -259,6 +264,7 @@ void FSequencerSelection::AddToSelection(const FSequencerSelectedKey& Key)
 	++SerialNumber;
 
 	SelectedKeys.Add(Key);
+	RawSelectedKeys.Add(Key.KeyHandle);
 
 	if (TViewModelPtr<IOutlinerExtension> OutlinerItem = Key.WeakChannel.Pin()->GetLinkedOutlinerItem())
 	{
@@ -377,6 +383,7 @@ void FSequencerSelection::RemoveFromSelection(const FSequencerSelectedKey& Key)
 	++SerialNumber;
 
 	SelectedKeys.Remove(Key);
+	RawSelectedKeys.Remove(Key.KeyHandle);
 
 	if (TViewModelPtr<IOutlinerExtension> OutlinerItem = Key.WeakChannel.Pin()->GetLinkedOutlinerItem())
 	{
@@ -506,6 +513,7 @@ void FSequencerSelection::EmptySelectedKeys()
 		}
 	}
 	SelectedKeys.Empty();
+	RawSelectedKeys.Empty();
 
 	if ( IsBroadcasting() )
 	{
@@ -709,6 +717,7 @@ void FSequencerSelection::RevalidateSelection()
 	{
 		if (!It->Section || !It->WeakChannel.Pin().Get())
 		{
+			RawSelectedKeys.Remove(It->KeyHandle);
 			It.RemoveCurrent();
 		}
 	}
diff --git a/Engine/Source/Editor/Sequencer/Private/SequencerSelection.h b/Engine/Source/Editor/Sequencer/Private/SequencerSelection.h
--- a/Engine/Source/Editor/Sequencer/Private/SequencerSelection.h
+++ b/Engine/Source/Editor/Sequencer/Private/SequencerSelection.h
@@ -41,6 +41,9 @@ public:
 	/** Gets a set of the selected keys. */
 	const TSet<FSequencerSelectedKey>& GetSelectedKeys() const;
 
+	/** Gets a set of the selected keys. */
+	const TSet<FKeyHandle>& GetRawSelectedKeys() const;
+
 	/** Gets a set of the selected sections */
 	TSet<TWeakObjectPtr<UMovieSceneSection>> GetSelectedSections() const;
 
@@ -199,6 +202,8 @@ private:
 	TSharedPtr<FViewModel> RootModel;
 
 	TSet<FSequencerSelectedKey> SelectedKeys;
+	/** Selected keys added directly by FKeyHandle, derived from SelectedKeys on mutation */
+	TSet<FKeyHandle> RawSelectedKeys;
 	TSet<TWeakPtr<FViewModel>> SelectedTrackAreaItems;
 	TSet<TWeakPtr<FViewModel>> SelectedOutlinerItems;
 	mutable TSet<TWeakPtr<FViewModel>> NodesWithSelectedKeysOrSections;
diff --git a/Engine/Source/Editor/Sequencer/Private/SequencerSelectionPreview.cpp b/Engine/Source/Editor/Sequencer/Private/SequencerSelectionPreview.cpp
--- a/Engine/Source/Editor/Sequencer/Private/SequencerSelectionPreview.cpp
+++ b/Engine/Source/Editor/Sequencer/Private/SequencerSelectionPreview.cpp
@@ -8,15 +8,17 @@
 
 namespace UE::Sequencer { class FViewModel; }
 
-void FSequencerSelectionPreview::SetSelectionState(FSequencerSelectedKey Key, ESelectionPreviewState InState)
+void FSequencerSelectionPreview::SetSelectionState(const FSequencerSelectedKey& Key, ESelectionPreviewState InState)
 {
 	if (InState == ESelectionPreviewState::Undefined)
 	{
 		DefinedKeyStates.Remove(Key);
+		RawDefinedKeyStates.Remove(Key.KeyHandle);
 	}
 	else
 	{
 		DefinedKeyStates.Add(Key, InState);
+		RawDefinedKeyStates.Add(Key.KeyHandle, InState);
 	}
 
 	CachedSelectionHash.Reset();
@@ -44,6 +46,16 @@ ESelectionPreviewState FSequencerSelectionPreview::GetSelectionState(const FSequ
 	return ESelectionPreviewState::Undefined;
 }
 
+ESelectionPreviewState FSequencerSelectionPreview::GetSelectionState(FKeyHandle Key) const
+{
+	if (const ESelectionPreviewState* State = RawDefinedKeyStates.Find(Key))
+	{
+		return *State;
+	}
+	return ESelectionPreviewState::Undefined;
+}
+
+
 ESelectionPreviewState FSequencerSelectionPreview::GetSelectionState(TWeakPtr<UE::Sequencer::FViewModel> InModel) const
 {
 	if (const ESelectionPreviewState* State = DefinedModelStates.Find(InModel))
@@ -62,6 +74,7 @@ void FSequencerSelectionPreview::Empty()
 void FSequencerSelectionPreview::EmptyDefinedKeyStates()
 {
 	DefinedKeyStates.Reset();
+	RawDefinedKeyStates.Reset();
 	CachedSelectionHash.Reset();
 }
 
@@ -77,7 +90,7 @@ uint32 FSequencerSelectionPreview::GetSelectionHash() const
 	{
 		uint32 NewHash = 0;
 
-		for (TPair<FSequencerSelectedKey, ESelectionPreviewState> Pair : DefinedKeyStates)
+		for (TPair<FKeyHandle, ESelectionPreviewState> Pair : RawDefinedKeyStates)
 		{
 			NewHash = HashCombine(NewHash, HashCombine(GetTypeHash(Pair.Key), GetTypeHash(Pair.Value)));
 		}
diff --git a/Engine/Source/Editor/Sequencer/Private/SequencerSelectionPreview.h b/Engine/Source/Editor/Sequencer/Private/SequencerSelectionPreview.h
--- a/Engine/Source/Editor/Sequencer/Private/SequencerSelectionPreview.h
+++ b/Engine/Source/Editor/Sequencer/Private/SequencerSelectionPreview.h
@@ -5,6 +5,7 @@
 #include "Containers/Map.h"
 #include "HAL/Platform.h"
 #include "Misc/Optional.h"
+#include "Curves/KeyHandle.h"
 #include "SequencerSelectedKey.h"
 #include "Templates/SharedPointer.h"
 
@@ -26,18 +27,20 @@ public:
 
 	/** Access the defined key states */
 	const TMap<FSequencerSelectedKey, ESelectionPreviewState>& GetDefinedKeyStates() const { return DefinedKeyStates; }
+	const TMap<FKeyHandle, ESelectionPreviewState>& GetRawDefinedKeyStates() const { return RawDefinedKeyStates; }
 
 	/** Access the defined model states */
 	const TMap<TWeakPtr<UE::Sequencer::FViewModel>, ESelectionPreviewState>& GetDefinedModelStates() const { return DefinedModelStates; }
 
 	/** Adds a key to the selection */
-	void SetSelectionState(FSequencerSelectedKey Key, ESelectionPreviewState InState);
+	void SetSelectionState(const FSequencerSelectedKey& Key, ESelectionPreviewState InState);
 
 	/** Adds a model to the selection */
 	void SetSelectionState(TWeakPtr<UE::Sequencer::FViewModel> InModel, ESelectionPreviewState InState);
 
 	/** Returns the selection state for the the specified key. */
 	ESelectionPreviewState GetSelectionState(const FSequencerSelectedKey& Key) const;
+	ESelectionPreviewState GetSelectionState(FKeyHandle Key) const;
 
 	/** Returns the selection state for the the specified model. */
 	ESelectionPreviewState GetSelectionState(TWeakPtr<UE::Sequencer::FViewModel> InModel) const;
@@ -57,6 +60,8 @@ public:
 private:
 
 	TMap<FSequencerSelectedKey, ESelectionPreviewState> DefinedKeyStates;
+	TMap<FKeyHandle, ESelectionPreviewState> RawDefinedKeyStates;
+
 	TMap<TWeakPtr<UE::Sequencer::FViewModel>, ESelectionPreviewState> DefinedModelStates;
 
 	/** Cached hash of this whole selection preview state. */
diff --git a/Engine/Source/Editor/Sequencer/Private/Tools/SequencerEditTool_Movement.cpp b/Engine/Source/Editor/Sequencer/Private/Tools/SequencerEditTool_Movement.cpp
--- a/Engine/Source/Editor/Sequencer/Private/Tools/SequencerEditTool_Movement.cpp
+++ b/Engine/Source/Editor/Sequencer/Private/Tools/SequencerEditTool_Movement.cpp
@@ -181,16 +181,20 @@ TSharedPtr<UE::Sequencer::ISequencerEditToolDragOperation> FSequencerEditTool_Mo
 		{
 			if (KeyHotspot)
 			{
-				TArrayView<const FSequencerSelectedKey> HoveredKeys = KeyHotspot->Keys;
+				const TSet<FSequencerSelectedKey>& HoveredKeys = KeyHotspot->Keys;
 
-				auto AnyUnselectedKey = [&Selection](const FSequencerSelectedKey& InKey)
+				bool bUniqueDrag = false;
+				for (const FSequencerSelectedKey& Key : HoveredKeys)
 				{
-					return !Selection.IsSelected(InKey);
+					// If any are not selected, we'll treat this as a unique drag
+					if (!Selection.IsSelected(Key))
+					{
+						bUniqueDrag = true;
+					}
 				};
 
-				if (HoveredKeys.ContainsByPredicate(AnyUnselectedKey))
+				if (bUniqueDrag)
 				{
-					// If any are not selected, we'll treat this as a unique drag
 					Selection.EmptySelectedKeys();
 					Selection.EmptySelectedTrackAreaItems();
 					for (const FSequencerSelectedKey& Key : HoveredKeys)
@@ -236,16 +240,20 @@ TSharedPtr<UE::Sequencer::ISequencerEditToolDragOperation> FSequencerEditTool_Mo
 		// Moving key(s)?
 		else if (KeyHotspot)
 		{
-			TArrayView<const FSequencerSelectedKey> HoveredKeys = KeyHotspot->Keys;
+			const TSet<FSequencerSelectedKey>& HoveredKeys = KeyHotspot->Keys;
 
-			auto AnyUnselectedKey = [&Selection](const FSequencerSelectedKey& InKey)
+			bool bUniqueDrag = false;
+			for (const FSequencerSelectedKey& Key : HoveredKeys)
 			{
-				return !Selection.IsSelected(InKey);
+				// If any are not selected, we'll treat this as a unique drag
+				if (!Selection.IsSelected(Key))
+				{
+					bUniqueDrag = true;
+				}
 			};
 
-			if (HoveredKeys.ContainsByPredicate(AnyUnselectedKey))
+			if (bUniqueDrag)
 			{
-				// If any are not selected, we'll treat this as a unique drag
 				Selection.EmptySelectedKeys();
 				Selection.EmptySelectedTrackAreaItems();
 				for (const FSequencerSelectedKey& Key : HoveredKeys)
diff --git a/Engine/Source/Editor/Sequencer/Private/Tools/SequencerEditTool_Selection.cpp b/Engine/Source/Editor/Sequencer/Private/Tools/SequencerEditTool_Selection.cpp
--- a/Engine/Source/Editor/Sequencer/Private/Tools/SequencerEditTool_Selection.cpp
+++ b/Engine/Source/Editor/Sequencer/Private/Tools/SequencerEditTool_Selection.cpp
@@ -27,7 +27,7 @@ struct FSelectionPreviewVisitor final
 {
 	FSelectionPreviewVisitor(FSequencerSelectionPreview& InSelectionPreview, FSequencerSelection& InSelection, ESelectionPreviewState InSetStateTo, bool bInPinned)
 		: SelectionPreview(InSelectionPreview)
-		, ExistingSelection(InSelection)
+		, ExistingSelection(InSelection.GetRawSelectedKeys())
 		, SetStateTo(InSetStateTo)
 		, bPinned(bInPinned)
 	{}
@@ -42,12 +42,10 @@ struct FSelectionPreviewVisitor final
 			return;
 		}
 
-		FSequencerSelectedKey Key(*Section, TSharedPtr<FChannelModel>(Channel), KeyHandle);
-
 		// Under default behavior keys have priority, so if a key is changing selection state then we remove any sections from the selection. The user can bypass this
 		// by holding down the control key which will allow selecting both keys and sections.
 		bool bKeySelectionHasPriority = !FSlateApplication::Get().GetModifierKeys().IsControlDown();
-		bool bKeyIsSelected = ExistingSelection.IsSelected(Key);
+		bool bKeyIsSelected = ExistingSelection.Contains(KeyHandle);
 
 		if (bKeySelectionHasPriority && 
 			((bKeyIsSelected && SetStateTo == ESelectionPreviewState::NotSelected) ||
@@ -57,6 +55,7 @@ struct FSelectionPreviewVisitor final
 			SelectionPreview.EmptyDefinedModelStates();
 		}
 
+		FSequencerSelectedKey Key(*Section, TSharedPtr<FChannelModel>(Channel), KeyHandle);
 		SelectionPreview.SetSelectionState(Key, SetStateTo);
 	}
 
@@ -84,7 +83,7 @@ struct FSelectionPreviewVisitor final
 private:
 
 	FSequencerSelectionPreview& SelectionPreview;
-	FSequencerSelection& ExistingSelection;
+	const TSet<FKeyHandle>& ExistingSelection;
 	ESelectionPreviewState SetStateTo;
 	bool bPinned;
 };
diff --git a/Engine/Source/Editor/Sequencer/Public/MVVM/Views/SChannelView.h b/Engine/Source/Editor/Sequencer/Public/MVVM/Views/SChannelView.h
--- a/Engine/Source/Editor/Sequencer/Public/MVVM/Views/SChannelView.h
+++ b/Engine/Source/Editor/Sequencer/Public/MVVM/Views/SChannelView.h
@@ -56,6 +56,9 @@ struct FChannelViewKeyCachedState
 	/** Whether to gather keys for all children and collapse them into this renderer (normally only true when an item is collapsed) */
 	bool bCollapseChildren;
 
+	/** Whether this channel is hovered */
+	bool bIsChannelHovered;
+
 	/** Default constructor for SWidget construction - not to be used under other circumstances */
 	FChannelViewKeyCachedState();
 
@@ -92,7 +95,7 @@ protected:
 
 	int32 DrawLane(const FPaintArgs& Args, const FGeometry& AllottedGeometry, const FSlateRect& MyCullingRect, FSlateWindowElementList& OutDrawElements, int32 LayerId, const FWidgetStyle& InWidgetStyle, bool bParentEnabled) const;
 
-	void CreateKeysUnderMouse(const FVector2D& MousePosition, const FGeometry& AllottedGeometry, TArrayView<const FSequencerSelectedKey> InPressedKeys, TArray<FSequencerSelectedKey>& OutKeys);
+	void CreateKeysUnderMouse(const FVector2D& MousePosition, const FGeometry& AllottedGeometry, const TSet<FSequencerSelectedKey>& InPressedKeys, TArray<FSequencerSelectedKey>& OutKeys);
 	void GetKeysUnderMouse(const FVector2D& MousePosition, const FGeometry& AllottedGeometry, TArray<FSequencerSelectedKey>& OutKeys) const;
 	void GetKeysAtPixelX(float LocalMousePixelX, TArray<FSequencerSelectedKey>& OutKeys) const;
 
diff --git a/Engine/Source/Editor/SequencerCore/Private/MVVM/Views/KeyRenderer.cpp b/Engine/Source/Editor/SequencerCore/Private/MVVM/Views/KeyRenderer.cpp
--- a/Engine/Source/Editor/SequencerCore/Private/MVVM/Views/KeyRenderer.cpp
+++ b/Engine/Source/Editor/SequencerCore/Private/MVVM/Views/KeyRenderer.cpp
@@ -11,6 +11,9 @@
 #include "CurveModel.h"
 #include "CurveEditor.h"
 
+#include "Async/ParallelFor.h"
+#include "Algo/Accumulate.h"
+
 namespace UE::Sequencer
 {
 
@@ -91,19 +94,21 @@ void FKeyRenderer::HitTestKeys(const FFrameTime& Time, TArray<FKeysForModel>& Ou
 	}
 }
 
-bool FKeyRenderer::HitTestKeyBar(const FFrameTime& Time, TArray<FKeysForModel>& OutLeadingKeys, TArray<FKeysForModel>& OutTrailingKeys) const
+bool FKeyRenderer::HitTestKeyBar(const FFrameTime& Time, FKeyBar& OutKeyBar) const
 {
 	const FFrameTime HalfKeyWidth = KeyWidthInFrames * .5f;
 
 	// Hit test the actual structure used for drawing key bars to ensure we are 100% accurate
 	// PrecomputedKeyBars is an array sorted by time, so we can binary search it to find the hit area
-	const int32 HitIndex = Algo::LowerBoundBy(PrecomputedKeyBars, Time, &FKeyBar::EndTime);
+	const int32 HitIndex = Algo::LowerBoundBy(PrecomputedKeyBars, Time, &FCachedKeyBar::EndTime);
 	if (HitIndex >= PrecomputedKeyBars.Num() || PrecomputedKeyBars[HitIndex].StartTime > Time)
 	{
 		// No hit - we ran off the end of the drawn keys, or didn't hit a bar
 		return false;
 	}
 
+ 	OutKeyBar.Range = TRange<FFrameTime>::Inclusive(PrecomputedKeyBars[HitIndex].StartTime+HalfKeyWidth, PrecomputedKeyBars[HitIndex].EndTime-HalfKeyWidth);
+
 	// Generate leading keys array
 	{
 		FFrameTime LeadingStart = PrecomputedKeyBars[HitIndex].StartTime - HalfKeyWidth;
@@ -126,7 +131,7 @@ bool FKeyRenderer::HitTestKeyBar(const FFrameTime& Time, TArray<FKeysForModel>&
 					Keys.Keys[Index] = Entry.CachedKeys->KeyHandles[FrameStart + Index];
 				}
 
-				OutLeadingKeys.Emplace(MoveTemp(Keys));
+				OutKeyBar.LeadingKeys.Emplace(MoveTemp(Keys));
 			}
 		}
 	}
@@ -153,7 +158,7 @@ bool FKeyRenderer::HitTestKeyBar(const FFrameTime& Time, TArray<FKeysForModel>&
 					Keys.Keys[Index] = Entry.CachedKeys->KeyHandles[FrameStart + Index];
 				}
 
-				OutTrailingKeys.Emplace(MoveTemp(Keys));
+				OutKeyBar.TrailingKeys.Emplace(MoveTemp(Keys));
 			}
 		}
 	}
@@ -236,12 +241,14 @@ void FKeyRenderer::UpdateViewDependentData(const FKeyBatchParameters& Params, co
 	const bool bHasAnySelectionPreview = Params.ClientInterface && Params.ClientInterface->HasAnyPreviewSelectedKeys();
 	const bool bHasAnyHoveredKeys      = Params.ClientInterface && Params.ClientInterface->HasAnyHoveredKeys();
 
+	const int32 NumFramesInRange = Algo::Accumulate(KeyDrawInfo, 0, [](int32 Current, const FCachedKeyDrawInformation& In) { return Current + In.Num(); });
+
 	// ------------------------------------------------------------------------------
 	// Update view-dependent data for each draw info
-	for (FCachedKeyDrawInformation& Info : KeyDrawInfo)
-	{
-		Info.CacheViewDependentData(Params);
-	}
+	const bool bForceSingleThread = NumFramesInRange <= 50000;
+	ParallelFor(KeyDrawInfo.Num(), [this, &Params](int32 Index){
+		this->KeyDrawInfo[Index].CacheViewDependentData(Params);
+	}, bForceSingleThread);
 
 	// ------------------------------------------------------------------------------
 	// If the data has changed, or key state has changed, or the view has been zoomed
@@ -355,6 +362,8 @@ void FKeyRenderer::UpdateViewDependentData(const FKeyBatchParameters& Params, co
 		int32 TotalNumKeys = 0;
 		int32 NumOverlaps = 0;
 
+		static int32 MaxNumKeysToConsider = 500;
+
 		// Determine the ranges of keys considered to reside at this position
 		for (int32 DrawIndex = 0; DrawIndex < KeyDrawInfo.Num(); ++DrawIndex)
 		{
@@ -395,18 +404,18 @@ void FKeyRenderer::UpdateViewDependentData(const FKeyBatchParameters& Params, co
 				// Avoid creating FSequencerSelectedKeys unless absolutely necessary
 				FKeyHandle ThisKeyHandle = Info.HandlesInRange[Info.NextUnhandledIndex];
 
-				if (bHasAnySelection && Params.ClientInterface->IsKeySelected(KeyExtension, ThisKeyHandle))
+				if (bHasAnySelection && NumSelected < MaxNumKeysToConsider && Params.ClientInterface->IsKeySelected(KeyExtension, ThisKeyHandle))
 				{
 					++NumSelected;
 				}
-				if (bHasAnySelectionPreview)
+				if (bHasAnySelectionPreview && NumPreviewSelected < MaxNumKeysToConsider && NumPreviewNotSelected < MaxNumKeysToConsider)
 				{
 					EKeySelectionPreviewState SelectionState = Params.ClientInterface->GetPreviewSelectionState(KeyExtension, ThisKeyHandle);
 
 					NumPreviewSelected    += int32(SelectionState == EKeySelectionPreviewState::Selected);
 					NumPreviewNotSelected += int32(SelectionState == EKeySelectionPreviewState::NotSelected);
 				}
-				if (bHasAnyHoveredKeys && Params.ClientInterface->IsKeyHovered(KeyExtension, ThisKeyHandle))
+				if (bHasAnyHoveredKeys && NumHovered < MaxNumKeysToConsider && Params.ClientInterface->IsKeyHovered(KeyExtension, ThisKeyHandle))
 				{
 					++NumHovered;
 				}
@@ -438,15 +447,15 @@ void FKeyRenderer::UpdateViewDependentData(const FKeyBatchParameters& Params, co
 		}
 
 		// Determine the key color based on its selection/hover states
-		if (NumPreviewSelected == TotalNumKeys)
+		if (NumPreviewSelected == FMath::Min(TotalNumKeys, MaxNumKeysToConsider))
 		{
 			NewKey.Flags |= EKeyRenderingFlags::PreviewSelected;
 		}
-		else if (NumPreviewNotSelected == TotalNumKeys)
+		else if (NumPreviewNotSelected == FMath::Min(TotalNumKeys, MaxNumKeysToConsider))
 		{
 			NewKey.Flags |= EKeyRenderingFlags::PreviewNotSelected;
 		}
-		else if (NumSelected == TotalNumKeys)
+		else if (NumSelected == FMath::Min(TotalNumKeys, MaxNumKeysToConsider))
 		{
 			NewKey.Flags |= EKeyRenderingFlags::Selected;
 		}
@@ -454,7 +463,7 @@ void FKeyRenderer::UpdateViewDependentData(const FKeyBatchParameters& Params, co
 		{
 			NewKey.Flags |= EKeyRenderingFlags::AnySelected;
 		}
-		else if (NumHovered == TotalNumKeys)
+		else if (NumHovered == FMath::Min(TotalNumKeys, MaxNumKeysToConsider))
 		{
 			NewKey.Flags |= EKeyRenderingFlags::Hovered;
 		}
@@ -513,7 +522,7 @@ void FKeyRenderer::UpdateViewDependentData(const FKeyBatchParameters& Params, co
 		{
 			if (LeadingKey && TrailingKey)
 			{
-				PrecomputedKeyBars.Add(FKeyBar{ LeadingKey->Time, TrailingKey->Time, EKeyBarRenderingFlags::None });
+				PrecomputedKeyBars.Add(FCachedKeyBar{ LeadingKey->Time, TrailingKey->Time, EKeyBarRenderingFlags::None });
 			}
 		}
 		else
@@ -521,7 +530,7 @@ void FKeyRenderer::UpdateViewDependentData(const FKeyBatchParameters& Params, co
 			if (LeadingKey && LeadingKeyParams->ConnectionStyle != EKeyConnectionStyle::None &&
 				Params.TimeToPixel.FrameDeltaToPixel(PrecomputedKeys[0].FinalKeyPosition - LeadingKey->Time) > Params.KeySizePx.X)
 			{
-				PrecomputedKeyBars.Add(FKeyBar{ LeadingKey->Time, PrecomputedKeys[0].FinalKeyPosition, EKeyBarRenderingFlags::None, LeadingKeyParams->ConnectionStyle });
+				PrecomputedKeyBars.Add(FCachedKeyBar{ LeadingKey->Time, PrecomputedKeys[0].FinalKeyPosition, EKeyBarRenderingFlags::None, LeadingKeyParams->ConnectionStyle });
 			}
 
 			for (int32 Index = 0; Index < PrecomputedKeys.Num()-1; ++Index)
@@ -531,14 +540,14 @@ void FKeyRenderer::UpdateViewDependentData(const FKeyBatchParameters& Params, co
 				if (ThisKey.Params.ConnectionStyle != EKeyConnectionStyle::None &&
 					Params.TimeToPixel.FrameDeltaToPixel(NextKey.FinalKeyPosition - ThisKey.FinalKeyPosition) > Params.KeySizePx.X)
 				{
-					PrecomputedKeyBars.Add(FKeyBar{ ThisKey.FinalKeyPosition, NextKey.FinalKeyPosition, EKeyBarRenderingFlags::None, ThisKey.Params.ConnectionStyle });
+					PrecomputedKeyBars.Add(FCachedKeyBar{ ThisKey.FinalKeyPosition, NextKey.FinalKeyPosition, EKeyBarRenderingFlags::None, ThisKey.Params.ConnectionStyle });
 				}
 			}
 
 			if (TrailingKey && PrecomputedKeys.Last().Params.ConnectionStyle != EKeyConnectionStyle::None &&
 				Params.TimeToPixel.FrameDeltaToPixel(TrailingKey->Time - PrecomputedKeys.Last().FinalKeyPosition) > Params.KeySizePx.X)
 			{
-				PrecomputedKeyBars.Add(FKeyBar{ PrecomputedKeys.Last().FinalKeyPosition, TrailingKey->Time, EKeyBarRenderingFlags::None, PrecomputedKeys.Last().Params.ConnectionStyle });
+				PrecomputedKeyBars.Add(FCachedKeyBar{ PrecomputedKeys.Last().FinalKeyPosition, TrailingKey->Time, EKeyBarRenderingFlags::None, PrecomputedKeys.Last().Params.ConnectionStyle });
 			}
 		}
 	}
@@ -649,7 +658,7 @@ int32 FKeyRenderer::Draw(const FKeyBatchParameters& Params, const FGeometry& All
 
 		static float KeyBarHeight = 2.f;
 		const float KeyBarTop = AllottedGeometry.GetLocalSize().Y * .5f - KeyBarHeight*.5f;
-		for (const FKeyBar& KeyBar : PrecomputedKeyBars)
+		for (const FCachedKeyBar& KeyBar : PrecomputedKeyBars)
 		{
 			const float StartPx = Params.TimeToPixel.FrameToPixel(KeyBar.StartTime);
 			const float EndPx = Params.TimeToPixel.FrameToPixel(KeyBar.EndTime);
@@ -874,36 +883,56 @@ void FKeyRenderer::FCachedKeyDrawInformation::CacheViewDependentData(const FKeyB
 			const int32 PreserveEnd   = Algo::UpperBound(OldFramesInRange, FramesInRange.Last());
 
 			const int32 PreserveNum   = PreserveEnd - PreserveStart;
+
 			if (PreserveNum > 0)
 			{
-				TArray<FKeyDrawParams> NewDrawParams;
+				const int32 HeadNum   = Algo::LowerBound(FramesInRange, OldFramesInRange[PreserveStart]);
+				const int32 TailStart = Algo::LowerBound(FramesInRange, OldFramesInRange[PreserveEnd-1]);
+				const int32 TailNum   = FramesInRange.Num() - TailStart;
 
-				const int32 HeadNum = Algo::LowerBound(FramesInRange, OldFramesInRange[PreserveStart]);
-				if (HeadNum > 0)
+				if (HeadNum == 0 && TailNum == 0)
 				{
-					NewDrawParams.SetNum(HeadNum);
-					KeyExtension->DrawKeys(HandlesInRange.Slice(0, HeadNum), NewDrawParams);
+					// If we're preserving everything, just set the flag to avoid reallocating the whole array
+					bDrawnKeys = true;
 				}
+				else
+				{
+					const int32 HeadPreserveDiff = HeadNum - PreserveStart;
+
+					// Insert elements at the front of the array if we need to (this probably means we scrolled to the right to reveal new leading keys)
+					if (HeadPreserveDiff > 0)
+					{
+						DrawParams.InsertUninitialized(0, HeadPreserveDiff);
+					}
+					// Remove elements at the front of the array if we need to (this probably means we scrolled to the left and hid leading keys)
+					else if (HeadPreserveDiff < 0)
+					{
+						DrawParams.RemoveAt(0, -HeadPreserveDiff, false);
+					}
 
-				NewDrawParams.Append(DrawParams.GetData() + PreserveStart, PreserveNum);
+					// Draw newly revealed keys at the start
+					if (HeadNum > 0)
+					{
+						KeyExtension->DrawKeys(HandlesInRange.Slice(0, HeadNum), TArrayView<FKeyDrawParams>(DrawParams).Slice(0, HeadNum));
+					}
 
-				const int32 TailStart = Algo::LowerBound(FramesInRange, OldFramesInRange[PreserveEnd-1]);
-				const int32 TailNum = FramesInRange.Num() - TailStart;
+					// Expand or contract the array to the final size
+					DrawParams.SetNum(FramesInRange.Num(), false);
 
-				if (TailNum > 0)
-				{
-					NewDrawParams.SetNum(FramesInRange.Num());
-					KeyExtension->DrawKeys(HandlesInRange.Slice(TailStart, TailNum), TArrayView<FKeyDrawParams>(NewDrawParams).Slice(TailStart, TailNum));
-				}
+					// Draw newly revealed keys at the end
+					if (TailNum > 0)
+					{
+						KeyExtension->DrawKeys(HandlesInRange.Slice(TailStart, TailNum), TArrayView<FKeyDrawParams>(DrawParams).Slice(TailStart, TailNum));
+					}
 
-				DrawParams = MoveTemp(NewDrawParams);
-				bDrawnKeys = true;
+					bDrawnKeys = true;
+				}
 			}
 		}
 
 		if (!bDrawnKeys)
 		{
-			DrawParams.SetNum(FramesInRange.Num());
+			DrawParams.SetNum(FramesInRange.Num(), false);
 
 			if (FramesInRange.Num())
 			{
@@ -913,6 +942,12 @@ void FKeyRenderer::FCachedKeyDrawInformation::CacheViewDependentData(const FKeyB
 		}
 
 		check(DrawParams.Num() == FramesInRange.Num() && FramesInRange.Num() == HandlesInRange.Num());
+
+		// Shrink the draw params if we have excess slack
+		if (DrawParams.GetSlack() > DrawParams.Num() / 3)
+		{
+			DrawParams.Shrink();
+		}
 	}
 
 	// Always reset the pointers to the current key that needs processing
diff --git a/Engine/Source/Editor/SequencerCore/Private/MVVM/Views/STrackAreaLaneView.cpp b/Engine/Source/Editor/SequencerCore/Private/MVVM/Views/STrackAreaLaneView.cpp
--- a/Engine/Source/Editor/SequencerCore/Private/MVVM/Views/STrackAreaLaneView.cpp
+++ b/Engine/Source/Editor/SequencerCore/Private/MVVM/Views/STrackAreaLaneView.cpp
@@ -1,7 +1,9 @@
 // Copyright Epic Games, Inc. All Rights Reserved.
 
 #include "MVVM/Views/STrackAreaLaneView.h"
+#include "MVVM/Views/STrackAreaView.h"
 #include "MVVM/ViewModels/TrackAreaViewModel.h"
+#include "TimeToPixel.h"
 
 
 namespace UE::Sequencer
diff --git a/Engine/Source/Editor/SequencerCore/Public/MVVM/Views/KeyRenderer.h b/Engine/Source/Editor/SequencerCore/Public/MVVM/Views/KeyRenderer.h
--- a/Engine/Source/Editor/SequencerCore/Public/MVVM/Views/KeyRenderer.h
+++ b/Engine/Source/Editor/SequencerCore/Public/MVVM/Views/KeyRenderer.h
@@ -165,6 +165,18 @@ struct SEQUENCERCORE_API FKeyRenderer
 		TArray<FKeyHandle> Keys;
 	};
 
+	struct FKeyBar
+	{
+		TRange<FFrameTime> Range = TRange<FFrameTime>::Empty();
+		TArray<FKeysForModel> LeadingKeys;
+		TArray<FKeysForModel> TrailingKeys;
+
+		bool IsValid() const
+		{
+			return !Range.IsEmpty();
+		}
+	};
+
 	/**
 	 * Initialize this renderer with a generic model
 	 */
@@ -193,7 +205,7 @@ struct SEQUENCERCORE_API FKeyRenderer
 	/**
 	 * Retrieve keys at the specified position
 	 */
-	bool HitTestKeyBar(const FFrameTime& Time, TArray<FKeysForModel>& LeadingKeys, TArray<FKeysForModel>& TrailingKeys) const;
+	bool HitTestKeyBar(const FFrameTime& Time, FKeyBar& OutKeyBar) const;
 
 	/**
 	 * Returns true if this renderer has a curve, false otherwise
@@ -229,6 +241,14 @@ private:
 		 */
 		void CacheViewDependentData(const FKeyBatchParameters& Params);
 
+		/**
+		 * Return the size of this key draw information
+		 */
+		int32 Num() const
+		{
+			return FramesInRange.Num();
+		}
+
 	public:
 
 		/** Index into the array views for the next unhandled key */
@@ -275,7 +295,7 @@ private:
 
 
 	/** Cached parameters for drawing a single key */
-	struct FKeyBar
+	struct FCachedKeyBar
 	{
 		/** The tick range that this key occupies (significant when this FKey represents multiple overlapping keys) */
 		FFrameTime StartTime, EndTime;
@@ -316,7 +336,7 @@ private:
 	mutable TArray<FKey> PrecomputedKeys;
 
 	/** Computed final draw info for key bars */
-	mutable TArray<FKeyBar> PrecomputedKeyBars;
+	mutable TArray<FCachedKeyBar> PrecomputedKeyBars;
 
 	/** Computed final draw curve info */
 	mutable TArray<TTuple<double, double>> PrecomputedCurve;
diff --git a/Engine/Source/Runtime/Engine/Classes/Curves/KeyHandle.h b/Engine/Source/Runtime/Engine/Classes/Curves/KeyHandle.h
--- a/Engine/Source/Runtime/Engine/Classes/Curves/KeyHandle.h
+++ b/Engine/Source/Runtime/Engine/Classes/Curves/KeyHandle.h
@@ -184,9 +184,14 @@ public:
 		return Ar;
 	}
 
+private:
+
+	/** Relocate key handles and indices proceeding the specified index by shifting them forwards or backwards */
+	void RelocateKeyHandles(int32 StartAtIndex, int32 DeltaIndex);
+
 private:
 	/** Array of optional key handles that reside in corresponding indices for an externally owned serial data structure */
-	TArray<TOptional<FKeyHandle>> KeyHandles;
+	TSparseArray<FKeyHandle> KeyHandles;
 
 	/** Map of which key handles go to which indices. Indices may point to incorrect indices. Check against entry in KeyHandles array to verify. */
 	TMap<FKeyHandle, int32> KeyHandlesToIndices;
diff --git a/Engine/Source/Runtime/Engine/Private/Curves/KeyHandle.cpp b/Engine/Source/Runtime/Engine/Private/Curves/KeyHandle.cpp
--- a/Engine/Source/Runtime/Engine/Private/Curves/KeyHandle.cpp
+++ b/Engine/Source/Runtime/Engine/Private/Curves/KeyHandle.cpp
@@ -153,50 +153,20 @@ void FKeyHandleMap::EnsureIndexHasAHandle(int32 KeyIndex)
 int32 FKeyHandleLookupTable::GetIndex(FKeyHandle KeyHandle)
 {
 	const int32* Index = KeyHandlesToIndices.Find(KeyHandle);
-	if (!Index)
-	{
-		// If it's not even in the map, there's no way this could be a valid handle for this container
-		return INDEX_NONE;
-	}
-	else if (KeyHandles.IsValidIndex(*Index) && KeyHandles[*Index] == KeyHandle)
-	{
-		return *Index;
-	}
-
-	// slow lookup and cache
-	const int32 NewCacheIndex = KeyHandles.IndexOfByPredicate(
-		[KeyHandle](const TOptional<FKeyHandle>& PredKeyHandle)
-		{
-			return PredKeyHandle.IsSet() && PredKeyHandle.GetValue() == KeyHandle;
-		}
-	);
-
-	if (NewCacheIndex == INDEX_NONE)
-	{
-		return INDEX_NONE;
-	}
-
-	KeyHandlesToIndices.Add(KeyHandle, NewCacheIndex);
-	return NewCacheIndex;
+	return Index ? *Index : INDEX_NONE;
 }
 
 FKeyHandle FKeyHandleLookupTable::FindOrAddKeyHandle(int32 Index)
 {
-	if (KeyHandles.IsValidIndex(Index) && KeyHandles[Index].IsSet())
+	if (KeyHandles.IsValidIndex(Index))
 	{
-		return KeyHandles[Index].GetValue();
-	}
-	
-	int32 NumToAdd = Index + 1 - KeyHandles.Num();
-	if (NumToAdd > 0)
-	{
-		KeyHandles.AddDefaulted(NumToAdd);
+		return KeyHandles[Index];
 	}
 
 	// Allocate a new key handle
 	FKeyHandle NewKeyHandle;
 
-	KeyHandles[Index] = NewKeyHandle;
+	KeyHandles.Insert(Index, NewKeyHandle);
 	KeyHandlesToIndices.Add(NewKeyHandle, Index);
 
 	return NewKeyHandle;
@@ -206,14 +176,19 @@ void FKeyHandleLookupTable::MoveHandle(int32 OldIndex, int32 NewIndex)
 {
 	if (KeyHandles.IsValidIndex(OldIndex))
 	{
-		TOptional<FKeyHandle> Handle = KeyHandles[OldIndex];
+		FKeyHandle Handle = KeyHandles[OldIndex];
+
+		KeyHandles.RemoveAt(OldIndex);
+		RelocateKeyHandles(OldIndex, -1);
 
-		KeyHandles.RemoveAt(OldIndex, 1, false);
-		KeyHandles.Insert(Handle, NewIndex);
-		if (Handle.IsSet())
+		if (NewIndex < KeyHandles.GetMaxIndex())
 		{
-			KeyHandlesToIndices.Add(Handle.GetValue(), NewIndex);
+			// Move proceeding keys forward to make space
+			RelocateKeyHandles(NewIndex, 1);
 		}
+
+		KeyHandles.Insert(NewIndex, Handle);
+		KeyHandlesToIndices.Add(Handle, NewIndex);
 	}
 }
 
@@ -221,24 +196,26 @@ FKeyHandle FKeyHandleLookupTable::AllocateHandle(int32 Index)
 {
 	FKeyHandle NewKeyHandle;
 
-	int32 NumToAdd = Index + 1 - KeyHandles.Num();
-	if (NumToAdd > 0)
+	if (Index < KeyHandles.GetMaxIndex())
 	{
-		KeyHandles.AddDefaulted(NumToAdd);
+		// Move proceeding keys forward to make space
+		RelocateKeyHandles(Index, 1);
 	}
 
-	KeyHandles.Insert(NewKeyHandle, Index);
+	KeyHandles.Insert(Index, NewKeyHandle);
 	KeyHandlesToIndices.Add(NewKeyHandle, Index);
 	return NewKeyHandle;
 }
 
 void FKeyHandleLookupTable::DeallocateHandle(int32 Index)
 {
-	TOptional<FKeyHandle> KeyHandle = KeyHandles[Index];
-	KeyHandles.RemoveAt(Index, 1, false);
-	if (KeyHandle.IsSet())
+	if (KeyHandles.IsValidIndex(Index))
 	{
-		KeyHandlesToIndices.Remove(KeyHandle.GetValue());
+		KeyHandlesToIndices.Remove(KeyHandles[Index]);
+		KeyHandles.RemoveAt(Index, 1);
+
+		// Move proceeding keys into the gap we just made
+		RelocateKeyHandles(Index+1, -1);
 	}
 }
 
@@ -253,9 +230,77 @@ bool FKeyHandleLookupTable::Serialize(FArchive& Ar)
 	// We're only concerned with Undo/Redo transactions
 	if (Ar.IsTransacting())
 	{
-		Ar << KeyHandles;
+		// Serialize the sparse array so as to preserve indices (by default it only serializes valid entries)
+		int32 NumHandles = KeyHandles.GetMaxIndex();
+		Ar << NumHandles;
+
+		FKeyHandle InvalidHandle = FKeyHandle::Invalid();
+
+		if (Ar.IsLoading())
+		{
+			KeyHandles.Empty(NumHandles);
+			for (int32 Index = 0; Index < NumHandles; ++Index)
+			{
+				FKeyHandle Handle = InvalidHandle;
+				Ar << Handle;
+				if (Handle != InvalidHandle)
+				{
+					KeyHandles.Insert(Index, Handle);
+				}
+			}
+		}
+		else if (Ar.IsSaving())
+		{
+			for (int32 Index = 0; Index < NumHandles; ++Index)
+			{
+				FKeyHandle Handle = KeyHandles.IsAllocated(Index) ? KeyHandles[Index] : InvalidHandle;
+				Ar << Handle;
+			}
+		}
+
 		Ar << KeyHandlesToIndices;
 	}
 
 	return true;
+}
+
+void FKeyHandleLookupTable::RelocateKeyHandles(int32 StartAtIndex, int32 DeltaIndex)
+{
+	if (DeltaIndex == 0)
+	{
+		return;
+	}
+
+	const int32 OldNumKeys = KeyHandles.GetMaxIndex();
+	const int32 NewNumKeys = KeyHandles.GetMaxIndex() + DeltaIndex;
+
+	// Iterate forwards when removing elements, and backwards when adding elements
+	const int32 StartIndex = DeltaIndex > 0 ? OldNumKeys-1   : StartAtIndex;
+	const int32 EndIndex   = DeltaIndex > 0 ? StartAtIndex-1 : OldNumKeys;
+	const int32 Inc        = DeltaIndex > 0 ? -1             : 1;
+
+	if (DeltaIndex > 0)
+	{
+		// Reserve enough space for the new elements
+		KeyHandles.Reserve(NewNumKeys);
+	}
+
+	// Move handles and fixup indices
+	for (int32 FixupIndex = StartIndex; FixupIndex != EndIndex; FixupIndex += Inc)
+	{
+		if (KeyHandles.IsAllocated(FixupIndex))
+		{
+			FKeyHandle Handle = KeyHandles[FixupIndex];
+
+			KeyHandles.Insert(FixupIndex+DeltaIndex, Handle);
+			KeyHandlesToIndices.Add(Handle, FixupIndex+DeltaIndex);
+
+			KeyHandles.RemoveAt(FixupIndex);
+		}
+	}
+
+	if (DeltaIndex < 0)
+	{
+		KeyHandles.Shrink();
+	}
 }
\ No newline at end of file
