Sequencer: Corrected color rendering of top-level key bars

diff --git a/Engine/Source/Editor/Sequencer/Private/SSequencerSection.cpp b/Engine/Source/Editor/Sequencer/Private/SSequencerSection.cpp
--- a/Engine/Source/Editor/Sequencer/Private/SSequencerSection.cpp
+++ b/Engine/Source/Editor/Sequencer/Private/SSequencerSection.cpp
@@ -271,7 +271,7 @@ struct FSequencerSectionPainterImpl : FSequencerSectionPainter
 
 			FLinearColor FillTint = BlendedTint.LinearRGBToHSV();
 			FillTint.G *= .5f;
-			FillTint.B = FMath::Max(.03f, Tint.B*.1f);
+			FillTint.B = FMath::Max(.03f, FillTint.B*.1f);
 			FSlateDrawElement::MakeBox(
 				DrawElements,
 				LayerId,
@@ -862,6 +862,7 @@ void SSequencerSection::Construct( const FArguments& InArgs, TSharedPtr<FSequenc
 		+ SOverlay::Slot()
 		[
 			SNew(SChannelView, InSectionModel, OwningTrackLane->GetTrackAreaView())
+			.KeyBarColor(this, &SSequencerSection::GetTopLevelKeyBarColor)
 			.Visibility(this, &SSequencerSection::GetTopLevelChannelGroupVisibility)
 		]
 	];
@@ -885,6 +886,26 @@ EVisibility SSequencerSection::GetTopLevelChannelGroupVisibility() const
 	return EVisibility::Visible;
 }
 
+FLinearColor SSequencerSection::GetTopLevelKeyBarColor() const
+{
+	TSharedPtr<FSectionModel> SectionModel = WeakSectionModel.Pin();
+	TSharedPtr<ITrackExtension> Track = SectionModel ? SectionModel->FindAncestorOfType<ITrackExtension>() : nullptr;
+	UMovieSceneTrack* TrackObject = Track ? Track->GetTrack() : nullptr;
+
+	FLinearColor Tint = FLinearColor::White;
+	if (TrackObject)
+	{
+		Tint = FSequencerSectionPainter::BlendColor(TrackObject->GetColorTint()).CopyWithNewOpacity(1.f).LinearRGBToHSV();
+		// These values relate to those that are draw in the expanded section fill in PaintSectionBackground
+		// Except the Saturation is 75% rather than 50%, and the value is 25% rather than 10%
+		//Tint.G *= .75f;
+		Tint.B = FMath::Max(.03f, Tint.B*.25f);
+		Tint = Tint.HSVToLinearRGB();
+	}
+
+	return Tint;
+}
+
 FMargin SSequencerSection::GetHandleOffsetPadding() const
 {
 	return FMargin(HandleOffsetPx, 0.0);
diff --git a/Engine/Source/Editor/Sequencer/Private/SSequencerSection.h b/Engine/Source/Editor/Sequencer/Private/SSequencerSection.h
--- a/Engine/Source/Editor/Sequencer/Private/SSequencerSection.h
+++ b/Engine/Source/Editor/Sequencer/Private/SSequencerSection.h
@@ -109,6 +109,8 @@ private:
 
 	EVisibility GetTopLevelChannelGroupVisibility() const;
 
+	FLinearColor GetTopLevelKeyBarColor() const;
+
 	FMargin GetHandleOffsetPadding() const;
 
 public:
