Sequencer: Fix folders getting renamed when reordering them.

diff --git a/Engine/Source/Editor/Sequencer/Private/MVVM/ViewModels/FolderModel.cpp b/Engine/Source/Editor/Sequencer/Private/MVVM/ViewModels/FolderModel.cpp
--- a/Engine/Source/Editor/Sequencer/Private/MVVM/ViewModels/FolderModel.cpp
+++ b/Engine/Source/Editor/Sequencer/Private/MVVM/ViewModels/FolderModel.cpp
@@ -534,7 +534,7 @@ void FFolderModel::PerformDrop(const FViewModelPtr& TargetModel, const FDragDrop
 		return;
 	}
 
-	const FScopedTransaction Transaction(FText::Format(LOCTEXT("MoveItems", "Move into {0}."), FText::FromName(Folder->GetFolderName())));
+	const FScopedTransaction Transaction(FText::Format(LOCTEXT("MoveItems", "Move into {0}"), FText::FromName(Folder->GetFolderName())));
 
 	Folder->SetFlags(RF_Transactional);
 	Folder->Modify();
diff --git a/Engine/Source/Editor/Sequencer/Private/MVVM/ViewModels/SequenceModel.cpp b/Engine/Source/Editor/Sequencer/Private/MVVM/ViewModels/SequenceModel.cpp
--- a/Engine/Source/Editor/Sequencer/Private/MVVM/ViewModels/SequenceModel.cpp
+++ b/Engine/Source/Editor/Sequencer/Private/MVVM/ViewModels/SequenceModel.cpp
@@ -214,7 +214,7 @@ void FSequenceModel::PerformDrop(const FViewModelPtr& TargetModel, const FDragDr
 		return;
 	}
 
-	const FScopedTransaction Transaction(LOCTEXT("MoveItems", "Move to root."));
+	const FScopedTransaction Transaction(LOCTEXT("MoveItems", "Move to Root"));
 
 	MovieScene->SetFlags(RF_Transactional);
 	MovieScene->Modify();
@@ -253,7 +253,7 @@ void FSequenceModel::PerformDrop(const FViewModelPtr& TargetModel, const FDragDr
 
 			bool bSuccess = false;
 
-			// Handle dropoping a folder into another folder
+			// Handle dropping a folder into another folder
 			// @todo: if we ever support folders within object bindings this will need to have better validation
 			if (TSharedPtr<FFolderModel> DraggedFolderModel = DraggedModel.ImplicitCast())
 			{
@@ -267,7 +267,11 @@ void FSequenceModel::PerformDrop(const FViewModelPtr& TargetModel, const FDragDr
 					}
 
 					// Give this folder a unique name inside its new parent if necessary
-					FName FolderName = UMovieSceneFolder::MakeUniqueChildFolderName(DraggedFolder->GetFolderName(), MovieScene->GetRootFolders());
+					TArray<UMovieSceneFolder*> ExistingFolders;
+					MovieScene->GetRootFolders(ExistingFolders);
+					ExistingFolders.Remove(DraggedFolder);
+
+					FName FolderName = UMovieSceneFolder::MakeUniqueChildFolderName(DraggedFolder->GetFolderName(), ExistingFolders);
 					if (FolderName != DraggedFolder->GetFolderName())
 					{
 						DraggedFolder->SetFlags(RF_Transactional);
