Add key enabled state

diff --git a/Engine/Source/Editor/Sequencer/Private/MVVM/Views/SSequencerKeyNavigationButtons.h b/Engine/Source/Editor/Sequencer/Private/MVVM/Views/SSequencerKeyNavigationButtons.h
--- a/Engine/Source/Editor/Sequencer/Private/MVVM/Views/SSequencerKeyNavigationButtons.h
+++ b/Engine/Source/Editor/Sequencer/Private/MVVM/Views/SSequencerKeyNavigationButtons.h
@@ -47,7 +47,7 @@ public:
 			.Time(this, &SSequencerKeyNavigationButtons::GetTime)
 			.OnSetTime(this, &SSequencerKeyNavigationButtons::SetTime)
 			.OnAddKey(this, &SSequencerKeyNavigationButtons::HandleAddKey)
-			.IsEnabled(!InSequencer->IsReadOnly())
+			.IsEnabled_Lambda([this] () { return WeakSequencer.IsValid() ? !WeakSequencer.Pin()->IsReadOnly() : false; })
 		];
 	}
 
diff --git a/Engine/Source/Editor/SequencerCore/Public/MVVM/Views/SKeyNavigationButtons.h b/Engine/Source/Editor/SequencerCore/Public/MVVM/Views/SKeyNavigationButtons.h
--- a/Engine/Source/Editor/SequencerCore/Public/MVVM/Views/SKeyNavigationButtons.h
+++ b/Engine/Source/Editor/SequencerCore/Public/MVVM/Views/SKeyNavigationButtons.h
@@ -37,7 +37,7 @@ public:
 	DECLARE_DELEGATE_TwoParams(FOnAddKey, FFrameTime, TSharedPtr<FViewModel>)
 
 	SLATE_BEGIN_ARGS(SKeyNavigationButtons) : _IsEnabled(true) {}
-		SLATE_ARGUMENT(bool, IsEnabled)
+		SLATE_ATTRIBUTE(bool, IsEnabled)
 
 		SLATE_ARGUMENT(FText, PreviousKeyToolTip)
 		SLATE_ARGUMENT(FText, NextKeyToolTip)
