Sequencer modifications for MetaHuman Identity creation

diff --git a/Engine/Source/Editor/Sequencer/Private/SSequencer.cpp b/Engine/Source/Editor/Sequencer/Private/SSequencer.cpp
--- a/Engine/Source/Editor/Sequencer/Private/SSequencer.cpp
+++ b/Engine/Source/Editor/Sequencer/Private/SSequencer.cpp
@@ -311,6 +311,7 @@ void SSequencer::Construct(const FArguments& InArgs, TSharedRef<FSequencer> InSe
 		.AutoHeight()
 		[
 			SNew(SBox)
+			.Visibility(this, &SSequencer::GetShowSequencerToolbar)
 			.Padding(FMargin(CommonPadding,0.0f,0.0f,0.f))
 			[
 				SNew(SWrapBox)
@@ -735,7 +736,7 @@ void SSequencer::Construct(const FArguments& InArgs, TSharedRef<FSequencer> InSe
 					.Padding(ResizeBarPadding)
 					[
 						SNew( SSequencerSectionOverlay, TimeSliderControllerRef )
-						.Visibility( EVisibility::HitTestInvisible )
+						.Visibility(this, &SSequencer::GetShowTickLines)
 						.DisplayScrubPosition( false )
 						.DisplayTickLines( true )
 						.Clipping(EWidgetClipping::ClipToBounds)
@@ -3598,12 +3599,20 @@ EVisibility SSequencer::GetStatusBarVisibility() const
 	return GetSequencerSettings()->GetShowStatusBar() ? EVisibility::SelfHitTestInvisible : EVisibility::Hidden;
 }
 
-EFrameNumberDisplayFormats SSequencer::GetTimeDisplayFormat() const
+EVisibility SSequencer::GetShowTickLines() const
 {
-	return GetSequencerSettings()->GetTimeDisplayFormat();
+	return GetSequencerSettings()->GetShowTickLines() ? EVisibility::SelfHitTestInvisible : EVisibility::Hidden;
 }
 
+EVisibility SSequencer::GetShowSequencerToolbar() const
+{
+	return GetSequencerSettings()->GetShowSequencerToolbar() ? EVisibility::Visible : EVisibility::Collapsed;
+}
 
+EFrameNumberDisplayFormats SSequencer::GetTimeDisplayFormat() const
+{
+	return GetSequencerSettings()->GetTimeDisplayFormat();
+}
 void SSequencer::OnColumnFillCoefficientChanged(float FillCoefficient, int32 ColumnIndex)
 {
 	ColumnFillCoefficients[ColumnIndex] = FillCoefficient;
diff --git a/Engine/Source/Editor/Sequencer/Private/SSequencer.h b/Engine/Source/Editor/Sequencer/Private/SSequencer.h
--- a/Engine/Source/Editor/Sequencer/Private/SSequencer.h
+++ b/Engine/Source/Editor/Sequencer/Private/SSequencer.h
@@ -581,6 +581,12 @@ private:
 	/** Gets whether the status bar should be visible. */
 	EVisibility GetStatusBarVisibility() const;
 
+	/** Gets whether the tick lines should be drawn. */
+	EVisibility GetShowTickLines() const;
+
+	/** Gets whether the sequencer toolbar should be displayed */
+	EVisibility GetShowSequencerToolbar() const;
+
 	/** What is the preferred display format for time values. */
 	EFrameNumberDisplayFormats GetTimeDisplayFormat() const;
 
diff --git a/Engine/Source/Editor/Sequencer/Private/SequencerSettings.cpp b/Engine/Source/Editor/Sequencer/Private/SequencerSettings.cpp
--- a/Engine/Source/Editor/Sequencer/Private/SequencerSettings.cpp
+++ b/Engine/Source/Editor/Sequencer/Private/SequencerSettings.cpp
@@ -61,6 +61,8 @@ USequencerSettings::USequencerSettings( const FObjectInitializer& ObjectInitiali
 	TrajectoryPathCap = 250;
 	FrameNumberDisplayFormat = EFrameNumberDisplayFormats::Seconds;
 	bAutoExpandNodesOnSelection = true;
+	bShowTickLines = true;
+	bShowSequencerToolbar = true;
 }
 
 void USequencerSettings::PostEditChangeProperty(FPropertyChangedEvent& PropertyChangedEvent)
@@ -610,6 +612,34 @@ void USequencerSettings::SetShowStatusBar(bool InbShowStatusBar)
 	}
 }
 
+bool USequencerSettings::GetShowTickLines() const
+{
+	return bShowTickLines;
+}
+
+void USequencerSettings::SetShowTickLines(bool bInDrawTickLines)
+{
+	if(bShowTickLines != bInDrawTickLines)
+	{
+		bShowTickLines = bInDrawTickLines;
+		SaveConfig();
+	}
+}
+
+bool USequencerSettings::GetShowSequencerToolbar() const
+{
+	return bShowSequencerToolbar;
+}
+
+void USequencerSettings::SetShowSequencerToolbar(bool bInShowSequencerToolbar)
+{
+	if(bShowSequencerToolbar != bInShowSequencerToolbar)
+	{
+		bShowSequencerToolbar = bInShowSequencerToolbar;
+		SaveConfig();
+	}
+}
+
 bool USequencerSettings::HasKeyAreaCurveExtents(const FString& ChannelName) const
 {
 	TArray<FString> ChannelsArray;
diff --git a/Engine/Source/Editor/Sequencer/Public/SequencerSettings.h b/Engine/Source/Editor/Sequencer/Public/SequencerSettings.h
--- a/Engine/Source/Editor/Sequencer/Public/SequencerSettings.h
+++ b/Engine/Source/Editor/Sequencer/Public/SequencerSettings.h
@@ -294,6 +294,16 @@ public:
 	/** Set whether to show status bar */
 	void SetShowStatusBar(bool bInShowStatusBar);
 
+	/** @return true if showing tick lines */
+	bool GetShowTickLines() const;
+	/** Set whether to show status bar */
+	void SetShowTickLines(bool bInDrawTickLines);
+
+	/** @return true if showing sequencer toolbar */
+	bool GetShowSequencerToolbar() const;
+	/** Set whether to show sequencer toolbar bar */
+	void SetShowSequencerToolbar(bool bInDrawTickLines);
+
 	/** @return Whether the given channel has curve extents */
 	bool HasKeyAreaCurveExtents(const FString& ChannelName) const;
 	/** @ Remove curve extents for the given channel */
@@ -551,6 +561,14 @@ protected:
 	UPROPERTY(config, EditAnywhere, Category = Timeline)
 	bool bShowStatusBar;
 
+	/** Enable or disable displaying the tick lines. */
+	UPROPERTY(config, EditAnywhere, Category = Timeline)
+	bool bShowTickLines;
+
+	/** Enable or disable displaying the sequencer toolbar. */
+	UPROPERTY(config, EditAnywhere, Category = Timeline)
+	bool bShowSequencerToolbar;
+
 	/** The key area curve extents, stored per channel name */
 	UPROPERTY(config, EditAnywhere, Category = Timeline)
 	FString KeyAreaCurveExtents;
