Sequencer: invoke individual keying operations for each object being keyed via the custom details view "+" key button

Everything else in Sequencer code is receiving per-object keying operations. For instance, when auto-key is enabled, changing a property value will trigger one property-changed event per object, even if the details view is currently bulk editing multiple objects. We therefore make our button do the same so that various assumptions in the code aren't failing.

diff --git a/Engine/Source/Editor/Sequencer/Private/LevelEditorSequencerIntegration.cpp b/Engine/Source/Editor/Sequencer/Private/LevelEditorSequencerIntegration.cpp
--- a/Engine/Source/Editor/Sequencer/Private/LevelEditorSequencerIntegration.cpp
+++ b/Engine/Source/Editor/Sequencer/Private/LevelEditorSequencerIntegration.cpp
@@ -125,14 +125,20 @@ public:
 	{
 		TArray<UObject*> Objects;
 		KeyedPropertyHandle.GetOuterObjects( Objects );
-		FKeyPropertyParams KeyPropertyParams(Objects, KeyedPropertyHandle, ESequencerKeyMode::ManualKeyForced);
 
+		TArray<UObject*> EachObject;
+		EachObject.SetNum(1);
 		for (const TWeakPtr<ISequencer>& WeakSequencer : Sequencers)
 		{
 			TSharedPtr<ISequencer> Sequencer = WeakSequencer.Pin();
 			if (Sequencer.IsValid())
 			{
-				Sequencer->KeyProperty(KeyPropertyParams);
+				for (UObject* Object : Objects)
+				{
+					EachObject[0] = Object;
+					FKeyPropertyParams KeyPropertyParams(EachObject, KeyedPropertyHandle, ESequencerKeyMode::ManualKeyForced);
+					Sequencer->KeyProperty(KeyPropertyParams);
+				}
 			}
 		}
 	}
