Sequencer: Don't snap keybars if snapping is disabled

diff --git a/Engine/Source/Editor/Sequencer/Private/SequencerHotspots.cpp b/Engine/Source/Editor/Sequencer/Private/SequencerHotspots.cpp
--- a/Engine/Source/Editor/Sequencer/Private/SequencerHotspots.cpp
+++ b/Engine/Source/Editor/Sequencer/Private/SequencerHotspots.cpp
@@ -3,6 +3,7 @@
 #include "SequencerHotspots.h"
 #include "MVVM/Extensions/IObjectBindingExtension.h"
 #include "SequencerCommonHelpers.h"
+#include "SequencerSettings.h"
 #include "SSequencer.h"
 #include "Tools/EditToolDragOperations.h"
 #include "SequencerContextMenus.h"
@@ -237,7 +238,7 @@ TSharedPtr<ISequencerEditToolDragOperation> FKeyBarHotspot::InitiateDrag(const F
 		TUniquePtr<FScopedTransaction> Transaction;
 		TSet<UMovieSceneSection*> ModifiedSections;
 		FSequencerSnapField SnapField;
-	
+
 		FKeyBarDrag(TSharedPtr<FKeyBarHotspot> InHotspot, TSharedPtr<FSequencer> InSequencer)
 			: Hotspot(InHotspot)
 		{
@@ -296,14 +297,17 @@ TSharedPtr<ISequencerEditToolDragOperation> FKeyBarHotspot::InitiateDrag(const F
 				NewKeyTimes[Index] = NewTime + RelativeStartTimes[Index];
 			}
 
-			const float PixelSnapWidth = 20.f;
-			const int32 SnapThreshold = VirtualTrackArea.PixelDeltaToFrame(PixelSnapWidth).FloorToFrame().Value;
-			if (TOptional<FSequencerSnapField::FSnapResult> SnappedTime = SnapField.Snap(NewKeyTimes, SnapThreshold))
+			if (Hotspot.IsValid() && Hotspot->WeakSequencer.IsValid() && Hotspot->WeakSequencer.Pin()->GetSequencerSettings()->GetIsSnapEnabled())
 			{
-				FFrameTime SnappedDiff = SnappedTime->SnappedTime - SnappedTime->OriginalTime;
-				for (int32 Index = 0; Index < RelativeStartTimes.Num(); ++Index)
+				const float PixelSnapWidth = 20.f;
+				const int32 SnapThreshold = VirtualTrackArea.PixelDeltaToFrame(PixelSnapWidth).FloorToFrame().Value;
+				if (TOptional<FSequencerSnapField::FSnapResult> SnappedTime = SnapField.Snap(NewKeyTimes, SnapThreshold))
 				{
-					NewKeyTimes[Index] = NewKeyTimes[Index] + SnappedDiff;
+					FFrameTime SnappedDiff = SnappedTime->SnappedTime - SnappedTime->OriginalTime;
+					for (int32 Index = 0; Index < RelativeStartTimes.Num(); ++Index)
+					{
+						NewKeyTimes[Index] = NewKeyTimes[Index] + SnappedDiff;
+					}
 				}
 			}
 
