Sequencer: Get grouped sections on move

diff --git a/Engine/Source/Editor/Sequencer/Private/Tools/EditToolDragOperations.cpp b/Engine/Source/Editor/Sequencer/Private/Tools/EditToolDragOperations.cpp
--- a/Engine/Source/Editor/Sequencer/Private/Tools/EditToolDragOperations.cpp
+++ b/Engine/Source/Editor/Sequencer/Private/Tools/EditToolDragOperations.cpp
@@ -693,7 +693,26 @@ void FMoveKeysAndSections::AddModel(TSharedPtr<UE::Sequencer::FViewModel> Model)
 
 	if (FSectionModel* SectionModel = Model->CastThis<FSectionModel>())
 	{
-		Sections.Add(SectionModel->GetSection());
+		if (UMovieSceneSection* Section = SectionModel->GetSection())
+		{
+			Sections.Add(Section);
+	
+			if (UMovieScene* MovieScene = Section->GetTypedOuter<UMovieScene>())
+			{
+				// If the section is in a group, we also want to add the sections it is grouped with
+				if (const FMovieSceneSectionGroup* SectionGroup = MovieScene->GetSectionGroup(*Section))
+				{
+					for (TWeakObjectPtr<UMovieSceneSection> WeakGroupedSection : *SectionGroup)
+					{
+						// Verify sections are still valid, and are not infinite.
+						if (WeakGroupedSection.IsValid())
+						{
+							Sections.Add(WeakGroupedSection.Get());
+						}
+					}
+				}
+			}
+		}
 	}
 	if (TSharedPtr<IDraggableTrackAreaExtension> DraggableItem = Model->CastThisShared<IDraggableTrackAreaExtension>())
 	{
