Sequencer: Increase layer id for selection hash

diff --git a/Engine/Source/Editor/Sequencer/Private/SSequencerSection.cpp b/Engine/Source/Editor/Sequencer/Private/SSequencerSection.cpp
--- a/Engine/Source/Editor/Sequencer/Private/SSequencerSection.cpp
+++ b/Engine/Source/Editor/Sequencer/Private/SSequencerSection.cpp
@@ -225,7 +225,7 @@ struct FSequencerSectionPainterImpl : FSequencerSectionPainter
 			{
 				FSlateDrawElement::MakeBox(
 					DrawElements,
-					LayerId,
+					++LayerId,
 					ExpandedSectionGeometry.ToPaintGeometry(FVector2D(1, 1), ExpandedSectionGeometry.GetLocalSize() - FVector2D(2,2)),
 					CollapsedSelectedSectionOverlay,
 					DrawEffects,
