Sequencer: Let top-level channels do extra drawing too

Before, only non-top-level channels (e.g. transform track channels) were doing extra drawing. Now, the channel of, say, a float track, will also do extra drawing.

diff --git a/Engine/Source/Editor/Sequencer/Private/MVVM/Views/SChannelView.cpp b/Engine/Source/Editor/Sequencer/Private/MVVM/Views/SChannelView.cpp
--- a/Engine/Source/Editor/Sequencer/Private/MVVM/Views/SChannelView.cpp
+++ b/Engine/Source/Editor/Sequencer/Private/MVVM/Views/SChannelView.cpp
@@ -467,6 +467,14 @@ int32 SChannelView::OnPaint(const FPaintArgs& Args, const FGeometry& AllottedGeo
 	);
 
 	TViewModelPtr<FChannelModel> Channel = Model.ImplicitCast();
+	if (!Channel)
+	{
+		TOptional<FViewModelChildren> TopLevelChannels = Model->FindChildList(FTrackModel::GetTopLevelChannelType());
+		if (TopLevelChannels.IsSet())
+		{
+			Channel = TopLevelChannels->FindFirstChildOfType<FChannelModel>();
+		}
+	}
 	if (Channel)
 	{
 		FSequencerChannelPaintArgs ChannelPaintArgs = {
