Sequencer: Draw blend type offset fix. Track lane height instead of width.

Also, gave it a bit of space from the edge of the section because the new handles are flush against the sides and have thickness.

diff --git a/Engine/Source/Editor/Sequencer/Private/SSequencerSection.cpp b/Engine/Source/Editor/Sequencer/Private/SSequencerSection.cpp
--- a/Engine/Source/Editor/Sequencer/Private/SSequencerSection.cpp
+++ b/Engine/Source/Editor/Sequencer/Private/SSequencerSection.cpp
@@ -365,7 +365,7 @@ struct FSequencerSectionPainterImpl : FSequencerSectionPainter
 			return;
 		}
 
-		const float LaneHeight = TrackLane->GetDesiredSize().X;
+		const float LaneHeight = TrackLane->GetDesiredSize().Y;
 
 		TSharedRef<FSlateFontCache> FontCache = FSlateApplication::Get().GetRenderer()->GetFontCache();
 
@@ -386,7 +386,8 @@ struct FSequencerSectionPainterImpl : FSequencerSectionPainter
 
 		const float FontHeight = GetFontHeight();
 
-		FVector2D TextOffset = SectionModel->GetRange() == TRange<FFrameNumber>::All() ? FVector2D(0.f, LaneHeight - FontHeight - 1.f) : FVector2D(1.f, LaneHeight - FontHeight - 1.f);
+		// Offset more to the right of the lower bound since there's a handle there
+		FVector2D TextOffset = SectionModel->GetRange().HasLowerBound() ? FVector2D(8.f, LaneHeight - FontHeight - 4.f) : FVector2D(1.f, LaneHeight - FontHeight - 4.f);
 		FVector2D TextPosition = SectionGeometry.AbsoluteToLocal(SectionClippingRect.GetTopLeft()) + TextOffset;
 
 		FSlateDrawElement::MakeText(
