Sequencer: PopulateContextMenu and HandleMouseSelection for FKeyBarHotspot so that they act more like FSectionHotspot.

When right mouse click is performed on a keybar, the first section is selected (this is for the paste destinations in the paste menu). Also, return the context menu for the first section of the key bar.

diff --git a/Engine/Source/Editor/Sequencer/Private/SequencerHotspots.cpp b/Engine/Source/Editor/Sequencer/Private/SequencerHotspots.cpp
--- a/Engine/Source/Editor/Sequencer/Private/SequencerHotspots.cpp
+++ b/Engine/Source/Editor/Sequencer/Private/SequencerHotspots.cpp
@@ -7,6 +7,7 @@
 #include "SSequencer.h"
 #include "Tools/EditToolDragOperations.h"
 #include "SequencerContextMenus.h"
+#include "SequencerNodeTree.h"
 #include "MVVM/Views/STrackAreaView.h"
 #include "Tools/SequencerEditTool_Movement.h"
 #include "Tools/SequencerEditTool_Selection.h"
@@ -218,11 +219,61 @@ TOptional<FFrameNumber> FKeyBarHotspot::GetTime() const
 	return Time;
 }
 
+bool FKeyBarHotspot::PopulateContextMenu(FMenuBuilder& MenuBuilder, TSharedPtr<FExtender> MenuExtender, FFrameTime MouseDownTime)
+{
+	TArrayView<const FSequencerSelectedKey> FirstKey(&LeadingKeys.Last(), 1);
+
+	if (FirstKey.Num() == 0 || !FirstKey[0].Section)
+	{
+		return false;
+	}
+
+	TSharedPtr<FSequencer> Sequencer = WeakSequencer.Pin();
+
+	TSharedPtr<FSectionModel> SectionModel = Sequencer->GetNodeTree()->GetSectionModel(FirstKey[0].Section);
+	if (!SectionModel)
+	{
+		return false;
+	}
+	
+	FSectionContextMenu::BuildMenu(MenuBuilder, MenuExtender, *Sequencer, MouseDownTime);
+
+	TSharedPtr<IObjectBindingExtension> ObjectBinding = SectionModel->FindAncestorOfType<IObjectBindingExtension>();
+	SectionModel->GetSectionInterface()->BuildSectionContextMenu(MenuBuilder, ObjectBinding ? ObjectBinding->GetObjectGuid() : FGuid());
+
+	return true;
+}
+
 FCursorReply FKeyBarHotspot::GetCursor() const
 {
 	return FCursorReply::Cursor(EMouseCursor::ResizeLeftRight);
 }
 
+void FKeyBarHotspot::HandleMouseSelection(FHotspotSelectionManager& SelectionManager)
+{
+	TArrayView<const FSequencerSelectedKey> FirstKey(&LeadingKeys.Last(), 1);
+
+	if (FirstKey.Num() == 0 || !FirstKey[0].Section)
+	{
+		return;
+	}
+
+	TSharedPtr<FSequencer> Sequencer = WeakSequencer.Pin();
+
+	TSharedPtr<FSectionModel> SectionModel = Sequencer->GetNodeTree()->GetSectionModel(FirstKey[0].Section);
+	if (!SectionModel)
+	{
+		return;
+	}
+
+	// Base-class only handles RMB selection so that the other handles and interactive controls
+	// that act as hotspots and still operate correctly with Left click
+	if (SectionModel && SelectionManager.MouseEvent->GetEffectingButton() == EKeys::RightMouseButton)
+	{
+		SelectionManager.SelectModelExclusive(SectionModel);
+	}
+}
+
 TSharedPtr<ISequencerEditToolDragOperation> FKeyBarHotspot::InitiateDrag(const FPointerEvent& MouseEvent)
 {
 	struct FKeyBarDrag : ISequencerEditToolDragOperation, UE::Sequencer::ISnapCandidate
diff --git a/Engine/Source/Editor/Sequencer/Private/SequencerHotspots.h b/Engine/Source/Editor/Sequencer/Private/SequencerHotspots.h
--- a/Engine/Source/Editor/Sequencer/Private/SequencerHotspots.h
+++ b/Engine/Source/Editor/Sequencer/Private/SequencerHotspots.h
@@ -79,10 +79,11 @@ struct FKeyHotspot
 
 /** A hotspot representing a key bar */
 struct FKeyBarHotspot
-	: ITrackAreaHotspot
+	: ITrackAreaHotspot, IMouseHandlerHotspot
 	, TSharedFromThis<FKeyBarHotspot>
+
 {
-	UE_SEQUENCER_DECLARE_CASTABLE(FKeyBarHotspot, ITrackAreaHotspot);
+	UE_SEQUENCER_DECLARE_CASTABLE(FKeyBarHotspot, ITrackAreaHotspot, IMouseHandlerHotspot);
 
 	FKeyBarHotspot(const TRange<FFrameTime>& InRange, TArray<FSequencerSelectedKey>&& InLeadingKeys, TArray<FSequencerSelectedKey>&& InTrailingKeys, TWeakPtr<FSequencer> InWeakSequencer)
 		: LeadingKeys(MoveTemp(InLeadingKeys))
@@ -93,9 +94,13 @@ struct FKeyBarHotspot
 
 	virtual void UpdateOnHover(FTrackAreaViewModel& InTrackArea) const override;
 	virtual TOptional<FFrameNumber> GetTime() const override;
+	virtual bool PopulateContextMenu(FMenuBuilder& MenuBuilder, TSharedPtr<FExtender> MenuExtender, FFrameTime MouseDownTime) override;
 	virtual TSharedPtr<ISequencerEditToolDragOperation> InitiateDrag(const FPointerEvent& MouseEvent) override;
 	virtual FCursorReply GetCursor() const override;
 
+	/** IMouseHandlerHotspot */
+	virtual void HandleMouseSelection(FHotspotSelectionManager& SelectionManager) override;
+
 	/** The keys that are part of this hotspot */
 	TArray<FSequencerSelectedKey> LeadingKeys;
 	TArray<FSequencerSelectedKey> TrailingKeys;
