Sequencer: Set all of the new possessables' parents after replacing the parent/actor.

diff --git a/Engine/Source/Editor/Sequencer/Private/SequencerUtilities.cpp b/Engine/Source/Editor/Sequencer/Private/SequencerUtilities.cpp
--- a/Engine/Source/Editor/Sequencer/Private/SequencerUtilities.cpp
+++ b/Engine/Source/Editor/Sequencer/Private/SequencerUtilities.cpp
@@ -2510,7 +2510,7 @@ FGuid FSequencerUtilities::AssignActor(TSharedRef<ISequencer> Sequencer, AActor*
 		// Defer replacing this object until the components have been updated
 	}
 
-	auto UpdateComponent = [&](FGuid OldComponentGuid, UActorComponent* NewComponent)
+	auto UpdateComponent = [&](FGuid OldComponentGuid, UActorComponent* NewComponent, TArray<FGuid>& NewComponentGuids)
 	{
 		FMovieSceneSequenceIDRef FocusedGuid = Sequencer->GetFocusedTemplateID();
 
@@ -2532,13 +2532,11 @@ FGuid FSequencerUtilities::AssignActor(TSharedRef<ISequencer> Sequencer, AActor*
 		Sequencer->State.Invalidate(OldComponentGuid, FocusedGuid);
 		Sequencer->State.Invalidate(NewPossessable.GetGuid(), FocusedGuid);
 
-		FMovieScenePossessable* ThisPossessable = OwnerMovieScene->FindPossessable(NewPossessable.GetGuid());
-		if (ensure(ThisPossessable))
-		{
-			ThisPossessable->SetParent(NewGuid, OwnerMovieScene);
-		}
+		NewComponentGuids.Add(NewPossessable.GetGuid());
 	};
 
+	TArray<FGuid> NewComponentGuids;
+
 	// Handle components
 	AActor* ActorToReplace = Cast<AActor>(RuntimeObject);
 	if (ActorToReplace != nullptr && ActorToReplace->IsActorBeingDestroyed() == false)
@@ -2555,7 +2553,7 @@ FGuid FSequencerUtilities::AssignActor(TSharedRef<ISequencer> Sequencer, AActor*
 					{
 						if (NewComponent->GetFullName(Actor) == ComponentToReplace->GetFullName(ActorToReplace))
 						{
-							UpdateComponent(ComponentGuid, NewComponent);
+							UpdateComponent(ComponentGuid, NewComponent, NewComponentGuids);
 							bComponentWasUpdated = true;
 						}
 					}
@@ -2585,7 +2583,7 @@ FGuid FSequencerUtilities::AssignActor(TSharedRef<ISequencer> Sequencer, AActor*
 				UActorComponent** ComponentPtr = ComponentNameToComponent.Find(OldPossessable.GetName());
 				if (ComponentPtr != nullptr)
 				{
-					UpdateComponent(OldPossessable.GetGuid(), *ComponentPtr);
+					UpdateComponent(OldPossessable.GetGuid(), *ComponentPtr, NewComponentGuids);
 				}
 			}
 		}
@@ -2598,6 +2596,15 @@ FGuid FSequencerUtilities::AssignActor(TSharedRef<ISequencer> Sequencer, AActor*
 	Sequencer->State.Invalidate(InObjectBinding, Sequencer->GetFocusedTemplateID());
 	Sequencer->State.Invalidate(NewPossessableActor.GetGuid(), Sequencer->GetFocusedTemplateID());
 
+	for (const FGuid& NewComponentGuid : NewComponentGuids)
+	{
+		FMovieScenePossessable* ThisPossessable = OwnerMovieScene->FindPossessable(NewComponentGuid);
+		if (ensure(ThisPossessable))
+		{
+			ThisPossessable->SetParent(NewGuid, OwnerMovieScene);
+		}
+	}
+
 	// Try to fix up folders
 	TArray<UMovieSceneFolder*> FoldersToCheck;
 	for (UMovieSceneFolder* Folder : Sequencer->GetFocusedMovieSceneSequence()->GetMovieScene()->GetRootFolders())
