Sequencer: Fix trail drawing so that it's drawn for the selected object binding as well (this was a regression from MVVM refactor)

diff --git a/Engine/Source/Editor/Sequencer/Private/SequencerEdMode.cpp b/Engine/Source/Editor/Sequencer/Private/SequencerEdMode.cpp
--- a/Engine/Source/Editor/Sequencer/Private/SequencerEdMode.cpp
+++ b/Engine/Source/Editor/Sequencer/Private/SequencerEdMode.cpp
@@ -849,24 +849,27 @@ void FSequencerEdMode::DrawTracks3D(FPrimitiveDrawInterface* PDI)
 				continue;
 			}
 
-			bool bSelected = false;
-			for (TSharedPtr<FViewModel> Child : ObjectBindingNode->GetDescendants())
+			bool bSelected = Selection.IsSelected(ObjectBindingNode);
+			if (!bSelected)
 			{
-				if (Selection.IsSelected(Child) || Selection.NodeHasSelectedKeysOrSections(Child))
+				for (TSharedPtr<FViewModel> Child : ObjectBindingNode->GetDescendants())
 				{
-					bSelected = true;
-					// Stop traversing
-					break;
+					if (Selection.IsSelected(Child) || Selection.NodeHasSelectedKeysOrSections(Child))
+					{
+						bSelected = true;
+						// Stop traversing
+						break;
+					}
 				}
-			}
 
-			// If one of our parent is selected, we're considered selected
-			for (TSharedPtr<FViewModel> Parent : ObjectBindingNode->GetAncestors())
-			{
-				if (Selection.IsSelected(Parent) || Selection.NodeHasSelectedKeysOrSections(Parent))
+				// If one of our parent is selected, we're considered selected
+				for (TSharedPtr<FViewModel> Parent : ObjectBindingNode->GetAncestors())
 				{
-					bSelected = true;
-					break;
+					if (Selection.IsSelected(Parent) || Selection.NodeHasSelectedKeysOrSections(Parent))
+					{
+						bSelected = true;
+						break;
+					}
 				}
 			}
 
