Add option to enable playback range spin box in time range widget

diff --git a/Engine/Source/Editor/Sequencer/Private/SSequencer.cpp b/Engine/Source/Editor/Sequencer/Private/SSequencer.cpp
--- a/Engine/Source/Editor/Sequencer/Private/SSequencer.cpp
+++ b/Engine/Source/Editor/Sequencer/Private/SSequencer.cpp
@@ -234,10 +234,16 @@ void SSequencer::Construct(const FArguments& InArgs, TSharedRef<FSequencer> InSe
 	bMirrorLabels = true;
 	TSharedRef<ITimeSlider> BottomTimeSlider = SequencerWidgets.CreateTimeSlider( TimeSliderControllerRef, TAttribute<EVisibility>(this, &SSequencer::GetBottomTimeSliderVisibility), bMirrorLabels );
 
+	EShowRange Ranges = EShowRange(EShowRange::WorkingRange | EShowRange::ViewRange);
+	if(InArgs._ShowPlaybackRangeInTimeSlider)
+	{
+		Ranges |= EShowRange::PlaybackRange;
+	}
+
 	// Create bottom time range slider
 	TSharedRef<ITimeSlider> BottomTimeRange = SequencerWidgets.CreateTimeRange(
 		FTimeRangeArgs(
-			EShowRange(EShowRange::WorkingRange | EShowRange::ViewRange),
+			Ranges,
 			TimeSliderControllerRef,
 			TAttribute<EVisibility>(this, &SSequencer::GetTimeRangeVisibility),
 			NumericTypeInterface.ToSharedRef()
diff --git a/Engine/Source/Editor/Sequencer/Private/SSequencer.h b/Engine/Source/Editor/Sequencer/Private/SSequencer.h
--- a/Engine/Source/Editor/Sequencer/Private/SSequencer.h
+++ b/Engine/Source/Editor/Sequencer/Private/SSequencer.h
@@ -283,6 +283,9 @@ public:
 		/** Selection handler to pass to outliner view */
 		SLATE_ATTRIBUTE( TSharedPtr<UE::Sequencer::IOutlinerSelectionHandler>, SelectionHandler )
 
+		/** Whether to display the playback range spin box in time range slider */
+		SLATE_ARGUMENT( bool, ShowPlaybackRangeInTimeSlider )
+
 	SLATE_END_ARGS()
 
 
diff --git a/Engine/Source/Editor/Sequencer/Private/Sequencer.cpp b/Engine/Source/Editor/Sequencer/Private/Sequencer.cpp
--- a/Engine/Source/Editor/Sequencer/Private/Sequencer.cpp
+++ b/Engine/Source/Editor/Sequencer/Private/Sequencer.cpp
@@ -480,7 +480,8 @@ void FSequencer::InitSequencer(const FSequencerInitParams& InitParams, const TSh
 		.OnInitToolMenuContext(InitParams.ViewParams.OnInitToolMenuContext)
 		.AddMenuExtender(InitParams.ViewParams.AddMenuExtender)
 		.ToolbarExtender(InitParams.ViewParams.ToolbarExtender)
-		.SelectionHandler(SelectionHandler);
+		.SelectionHandler(SelectionHandler)
+		.ShowPlaybackRangeInTimeSlider(InitParams.ViewParams.bShowPlaybackRangeInTimeSlider);
 
 	SelectionHandler->SetTreeViews(SequencerWidget->GetTreeView(), SequencerWidget->GetPinnedTreeView());
 
diff --git a/Engine/Source/Editor/Sequencer/Public/ISequencerModule.h b/Engine/Source/Editor/Sequencer/Public/ISequencerModule.h
--- a/Engine/Source/Editor/Sequencer/Public/ISequencerModule.h
+++ b/Engine/Source/Editor/Sequencer/Public/ISequencerModule.h
@@ -105,6 +105,9 @@ struct FSequencerViewParams
 	/** Whether the sequencer is read-only */
 	bool bReadOnly;
 
+	/** Whether the playback range spin box should be visible in TimeRange widget */
+	bool bShowPlaybackRangeInTimeSlider;
+
 	/** Style of scrubber to use */
 	ESequencerScrubberStyle ScrubberStyle;
 
@@ -112,6 +115,7 @@ struct FSequencerViewParams
 		: OnGetPlaybackSpeeds(ISequencer::FOnGetPlaybackSpeeds::CreateLambda([]() { return TArray<float>{ 0.1f, 0.25f, 0.5f, 1.f, 2.f, 5.f, 10.f }; }))
 		, UniqueName(MoveTemp(InName))
 		, bReadOnly(false)
+		, bShowPlaybackRangeInTimeSlider(false)
 		, ScrubberStyle(ESequencerScrubberStyle::Vanilla)
 	{ }
 };
