Sequencer: Remove override channel parameters menu from the sequencer outliner

This was affecting all sections in the row, but it's likely that the user wants to only change parameters for one section (which can be done in the section's context menu). Plus, this was the only place in the UI that gave access to channel-related parameters in the sequencer outliner, so it's better to remove this precedent.

diff --git a/Engine/Source/Editor/Sequencer/Private/MVVM/ViewModels/ChannelModel.cpp b/Engine/Source/Editor/Sequencer/Private/MVVM/ViewModels/ChannelModel.cpp
--- a/Engine/Source/Editor/Sequencer/Private/MVVM/ViewModels/ChannelModel.cpp
+++ b/Engine/Source/Editor/Sequencer/Private/MVVM/ViewModels/ChannelModel.cpp
@@ -464,26 +464,6 @@ void FChannelGroupModel::BuildChannelOverrideMenu(FMenuBuilder& MenuBuilder)
 		}
 	}
 
-	// Show channel parameters.
-	if (ExistingOverrides.Num() > 0)
-	{
-		MenuBuilder.BeginSection("ChannelParameters", LOCTEXT("ChannelGroupParameterMenuSectionName", "Channel Parameters"));
-
-		for (const TTuple<UClass*, TArray<UMovieSceneChannelOverrideContainer*>>& ExistingOverride : ExistingOverrides)
-		{
-			const FText ChannelOverrideClassName = ExistingOverride.Key->GetDisplayNameText();
-			MenuBuilder.AddSubMenu(
-				FText::Format(LOCTEXT("OverrideChannelParametersMenuLabel", "{0} Parameters"), ChannelOverrideClassName),
-				FText::Format(
-					LOCTEXT("OverrideChannelParametersMenuTooltip", "Edit parameters for all channels of type {0}"),
-					ChannelOverrideClassName),
-				FNewMenuDelegate::CreateSP(this, &FChannelGroupOutlinerModel::BuildChannelOverrideParametersMenu, ExistingOverride.Value)
-			);
-		}
-
-		MenuBuilder.EndSection();
-	}
-
 	// Add menu entries for overriding channels, with only types that are common among all sections.
 	TMap<TSubclassOf<UMovieSceneChannelOverrideContainer>, int32> CandidateCounts;
 	for (const UMovieSceneChannelOverrideContainer::FOverrideCandidates& CandidateOverrides : AllCandidateOverrides)
@@ -651,26 +631,6 @@ void FChannelGroupModel::RemoveChannelOverrides()
 	Sequencer->NotifyMovieSceneDataChanged(EMovieSceneDataChangeType::MovieSceneStructureItemsChanged);
 }
 
-void FChannelGroupModel::BuildChannelOverrideParametersMenu(FMenuBuilder& MenuBuilder, TArray<UMovieSceneChannelOverrideContainer*> ChannelParameters)
-{
-	FPropertyEditorModule& EditModule = FModuleManager::Get().GetModuleChecked<FPropertyEditorModule>("PropertyEditor");
-
-	FDetailsViewArgs DetailsViewArgs;
-	DetailsViewArgs.bAllowSearch = false;
-	DetailsViewArgs.NameAreaSettings = FDetailsViewArgs::HideNameArea;
-	DetailsViewArgs.bHideSelectionTip = true;
-	DetailsViewArgs.bShowOptions = false;
-	DetailsViewArgs.bShowScrollBar = false;
-
-	TSharedRef<IDetailsView> DetailsView = EditModule.CreateDetailView(DetailsViewArgs);
-
-	TArray<UObject*> Objects;
-	Objects.Append(ChannelParameters);
-	DetailsView->SetObjects(Objects, true);
-
-	MenuBuilder.AddWidget(DetailsView, FText(), true, false);
-}
-
 FChannelGroupOutlinerModel::FChannelGroupOutlinerModel(FName InChannelName, const FText& InDisplayText)
 	: TOutlinerModelMixin<FChannelGroupModel>(InChannelName, InDisplayText)
 {
diff --git a/Engine/Source/Editor/Sequencer/Public/MVVM/ViewModels/ChannelModel.h b/Engine/Source/Editor/Sequencer/Public/MVVM/ViewModels/ChannelModel.h
--- a/Engine/Source/Editor/Sequencer/Public/MVVM/ViewModels/ChannelModel.h
+++ b/Engine/Source/Editor/Sequencer/Public/MVVM/ViewModels/ChannelModel.h
@@ -176,8 +176,6 @@ private:
 	void BuildChannelOverrideMenu(FMenuBuilder& MenuBuilder, UMovieSceneChannelOverrideContainer::FOverrideCandidates OverrideCandidates);
 	void OverrideChannels(TSubclassOf<UMovieSceneChannelOverrideContainer> OverrideClass);
 	void RemoveChannelOverrides();
-	
-	void BuildChannelOverrideParametersMenu(FMenuBuilder& MenuBuilder, TArray<UMovieSceneChannelOverrideContainer*> ChannelParameters);
 
 protected:
 
