Sequencer: Expand section if necessary when dragging key bars

diff --git a/Engine/Source/Editor/Sequencer/Private/SequencerSelectedKey.cpp b/Engine/Source/Editor/Sequencer/Private/SequencerSelectedKey.cpp
--- a/Engine/Source/Editor/Sequencer/Private/SequencerSelectedKey.cpp
+++ b/Engine/Source/Editor/Sequencer/Private/SequencerSelectedKey.cpp
@@ -83,6 +83,14 @@ void SetKeyTimes(TArrayView<const FSequencerSelectedKey> InSelectedKeys, TArrayV
 			for (int32 Index : ChannelInfo.OriginalIndices)
 			{
 				KeyTimesScratch.Add(InTimes[Index]);
+
+				if (UMovieSceneSection* Section = ChannelInfo.OwningSection)
+				{
+					if (!Section->GetRange().Contains(InTimes[Index]))
+					{
+						Section->ExpandToFrame(InTimes[Index]);
+					}
+				}
 			}
 
 			Channel->SetKeyTimes(ChannelInfo.KeyHandles, KeyTimesScratch);
