Sequencer: Snapping for key bars

Moved snap to interval to the snap field which is enabled with SetSnapToInterval().  Different operations will enable it depending on their use case. For example, dragging key bars, should have snap to interval enabled if the snap key time to interval setting is enabled. Dragging sections should have snap to interval enabled if the snap section time to interval setting is enabled.

diff --git a/Engine/Source/Editor/Sequencer/Private/SequencerHotspots.cpp b/Engine/Source/Editor/Sequencer/Private/SequencerHotspots.cpp
--- a/Engine/Source/Editor/Sequencer/Private/SequencerHotspots.cpp
+++ b/Engine/Source/Editor/Sequencer/Private/SequencerHotspots.cpp
@@ -300,6 +300,7 @@ TSharedPtr<ISequencerEditToolDragOperation> FKeyBarHotspot::InitiateDrag(const F
 			AllLinearKeys = AllKeys.Array();
 
 			SnapField = FSequencerSnapField(*InSequencer, *this);
+			SnapField.SetSnapToInterval(InSequencer->GetSequencerSettings()->GetSnapKeyTimesToInterval());
 		}
 
 		void OnBeginDrag(const FPointerEvent& MouseEvent, FVector2D LocalMousePos, const FVirtualTrackArea& VirtualTrackArea) override
diff --git a/Engine/Source/Editor/Sequencer/Private/Tools/EditToolDragOperations.cpp b/Engine/Source/Editor/Sequencer/Private/Tools/EditToolDragOperations.cpp
--- a/Engine/Source/Editor/Sequencer/Private/Tools/EditToolDragOperations.cpp
+++ b/Engine/Source/Editor/Sequencer/Private/Tools/EditToolDragOperations.cpp
@@ -53,31 +53,6 @@ protected:
 	TSet<UMovieSceneSection*> SectionsToExclude;
 };
 
-TOptional<FSequencerSnapField::FSnapResult> SnapToInterval(const TArray<FFrameTime>& InTimes, int32 FrameThreshold, FFrameRate Resolution, FFrameRate DisplayRate, ESequencerScrubberStyle ScrubStyle)
-{
-	TOptional<FSequencerSnapField::FSnapResult> Result;
-
-	FFrameTime SnapAmount(0);
-	for (FFrameTime Time : InTimes)
-	{
-		// Convert from resolution to DisplayRate, round to frame, then back again. We floor to frames when using the frame block scrubber, and round using the vanilla scrubber
-		FFrameTime   DisplayTime      = FFrameRate::TransformTime(Time, Resolution, DisplayRate);
-		FFrameNumber PlayIntervalTime = ScrubStyle == ESequencerScrubberStyle::FrameBlock ? DisplayTime.FloorToFrame() : DisplayTime.RoundToFrame();
-		FFrameNumber IntervalSnap     = FFrameRate::TransformTime(PlayIntervalTime, DisplayRate, Resolution).FloorToFrame();
-
-		FFrameTime ThisSnapAmount   = IntervalSnap - Time;
-		if (FMath::Abs(ThisSnapAmount) <= FrameThreshold)
-		{
-			if (!Result.IsSet() || FMath::Abs(ThisSnapAmount) < SnapAmount)
-			{
-				Result = FSequencerSnapField::FSnapResult{Time, IntervalSnap, 1.f};
-				SnapAmount = ThisSnapAmount;
-			}
-		}
-	}
-
-	return Result;
-}
 
 /** How many pixels near the mouse has to be before snapping occurs */
 const float PixelSnapWidth = 20.f;
@@ -163,6 +138,7 @@ void FResizeSection::OnBeginDrag(const FPointerEvent& MouseEvent, FVector2D Loca
 	TSet<FSequencerSelectedKey> EmptyKeySet;
 	FInvalidKeyAndSectionSnappingCandidates SnapCandidates(EmptyKeySet, Sections);
 	SnapField = FSequencerSnapField(Sequencer, SnapCandidates, ESequencerEntity::Section | ESequencerEntity::Key);
+	SnapField.GetValue().SetSnapToInterval(Sequencer.GetSequencerSettings()->GetSnapSectionTimesToInterval());
 
 	SectionInitTimes.Empty();
 
@@ -294,12 +270,6 @@ void FResizeSection::OnDrag(const FPointerEvent& MouseEvent, FVector2D LocalMous
 			SnappedTime = SnapField->Snap(SectionTimes, SnapThreshold);
 		}
 
-		if (!SnappedTime.IsSet() && Settings->GetSnapSectionTimesToInterval())
-		{
-			int32 IntervalSnapThreshold = FMath::RoundToInt( ( TickResolution / DisplayRate ).AsDecimal() );
-			SnappedTime = SnapToInterval(SectionTimes, IntervalSnapThreshold, TickResolution, DisplayRate, ScrubStyle);
-		}
-
 		if (SnappedTime.IsSet())
 		{
 			// Add the snapped amount onto the delta
@@ -565,6 +535,7 @@ void FManipulateSectionEasing::OnBeginDrag(const FPointerEvent& MouseEvent, FVec
 		// Construct a snap field of all section bounds
 		ISnapCandidate SnapCandidates;
 		SnapField = FSequencerSnapField(Sequencer, SnapCandidates, ESequencerEntity::Section);
+		SnapField.GetValue().SetSnapToInterval(Sequencer.GetSequencerSettings()->GetSnapSectionTimesToInterval());
 	}
 
 	InitValue = bEaseIn ? Section->Easing.GetEaseInDuration() : Section->Easing.GetEaseOutDuration();
@@ -613,12 +584,6 @@ void FManipulateSectionEasing::OnDrag(const FPointerEvent& MouseEvent, FVector2D
 			SnappedTime = SnapField->Snap(SnapTimes, SnapThreshold);
 		}
 
-		if (!SnappedTime.IsSet() && Settings->GetSnapSectionTimesToInterval())
-		{
-			int32 IntervalSnapThreshold = FMath::RoundToInt( ( TickResolution / DisplayRate ).AsDecimal() );
-			SnappedTime = SnapToInterval(SnapTimes, IntervalSnapThreshold, TickResolution, DisplayRate, ScrubStyle);
-		}
-
 		if (SnappedTime.IsSet())
 		{
 			// Add the snapped amount onto the delta
@@ -808,29 +773,27 @@ void FMoveKeysAndSections::OnDrag(const FPointerEvent& MouseEvent, FVector2D Loc
 		const bool bSnapToInterval = (KeysAsArray.Num() > 0 && Settings->GetSnapKeyTimesToInterval()) || (Sections.Num() > 0 && Settings->GetSnapSectionTimesToInterval());
 		const bool bSnapToLikeTypes = (KeysAsArray.Num() > 0 && Settings->GetSnapKeyTimesToKeys()) || (Sections.Num() > 0 && Settings->GetSnapSectionTimesToSections());
 
+		SnapField.GetValue().SetSnapToInterval(bSnapToInterval);
+
 		// RelativeSnapOffsets contains both our sections and our keys, and we add them all as potential things that can snap to stuff.
-		ValidSnapMarkers.SetNumUninitialized(RelativeSnapOffsets.Num());
-		for (int32 Index = 0; Index < RelativeSnapOffsets.Num(); ++Index)
+		if (bSnapToLikeTypes)
 		{
-			ValidSnapMarkers[Index] = (RelativeSnapOffsets[Index] + MouseTime);
+			ValidSnapMarkers.SetNumUninitialized(RelativeSnapOffsets.Num());
+			for (int32 Index = 0; Index < RelativeSnapOffsets.Num(); ++Index)
+			{
+				ValidSnapMarkers[Index] = (RelativeSnapOffsets[Index] + MouseTime);
+			}
 		}
 
 		// Now we'll try and snap all of these points to the closest valid snap marker (which may be a section or interval)
 		TOptional<FSequencerSnapField::FSnapResult> SnappedTime;
 
-		if (bSnapToLikeTypes)
+		if (bSnapToLikeTypes || bSnapToInterval)
 		{
 			// This may or may not set the SnappedTime depending on if there are any sections within the threshold.
 			SnappedTime = SnapField->Snap(ValidSnapMarkers, SnapThreshold);
 		}
 
-		if (!SnappedTime.IsSet() && bSnapToInterval)
-		{
-			// Snap to the nearest interval (if enabled). Snapping to other objects has over interval.
-			int32 IntervalSnapThreshold = FMath::RoundToInt((TickResolution / DisplayRate).AsDecimal());
-			SnappedTime = SnapToInterval(ValidSnapMarkers, IntervalSnapThreshold, TickResolution, DisplayRate, ScrubStyle);
-		}
-
 		// If they actually snapped to something (snapping may be on but settings might dictate nothing to snap to) add the difference
 		// to our current MouseTime so that MouseTime reflects the amount needed to move to get to the whole snap point.
 		if (SnappedTime.IsSet())
diff --git a/Engine/Source/Editor/Sequencer/Private/Tools/SequencerSnapField.cpp b/Engine/Source/Editor/Sequencer/Private/Tools/SequencerSnapField.cpp
--- a/Engine/Source/Editor/Sequencer/Private/Tools/SequencerSnapField.cpp
+++ b/Engine/Source/Editor/Sequencer/Private/Tools/SequencerSnapField.cpp
@@ -121,6 +121,10 @@ void FSequencerSnapField::Initialize(const FSequencer& InSequencer, UE::Sequence
 	{
 		SortedSnaps.Append(Visitor.Snaps);
 	}
+
+	TickResolution = InSequencer.GetFocusedTickResolution();
+	DisplayRate = InSequencer.GetFocusedDisplayRate();
+	ScrubStyle = InSequencer.GetScrubStyle();
 }
 
 void FSequencerSnapField::Finalize()
@@ -231,6 +235,30 @@ TOptional<FSequencerSnapField::FSnapResult> FSequencerSnapField::Snap(const TArr
 	for (FFrameTime Time : InTimes)
 	{
 		TOptional<FSnapResult> ThisSnap = Snap(Time, Threshold);
+
+		if (bSnapToInterval)
+		{
+			FFrameTime ThisTime = ThisSnap.IsSet() ? ThisSnap->SnappedTime : Time;
+			FFrameTime IntervalTime = FFrameRate::TransformTime(ThisTime, TickResolution, DisplayRate);
+			FFrameNumber PlayIntervalTime = ScrubStyle == ESequencerScrubberStyle::FrameBlock ? IntervalTime.FloorToFrame() : IntervalTime.RoundToFrame();
+			FFrameNumber IntervalSnap = FFrameRate::TransformTime(PlayIntervalTime, DisplayRate, TickResolution).FloorToFrame();
+
+			const int32 IntervalSnapThreshold = FMath::RoundToInt((TickResolution / DisplayRate).AsDecimal());
+
+			FFrameTime ThisSnapAmount = IntervalSnap - ThisTime;
+			if (FMath::Abs(ThisSnapAmount) <= IntervalSnapThreshold)
+			{
+				if (ThisSnap.IsSet())
+				{
+					ThisSnap->SnappedTime = IntervalSnap;
+				}
+				else
+				{
+					ThisSnap = { Time, IntervalSnap, 1.f };
+				}
+			}
+		}
+
 		if (!ThisSnap.IsSet())
 		{
 			continue;
@@ -244,4 +272,4 @@ TOptional<FSequencerSnapField::FSnapResult> FSequencerSnapField::Snap(const TArr
 	}
 
 	return ProspectiveSnap;
-}
\ No newline at end of file
+}
diff --git a/Engine/Source/Editor/Sequencer/Private/Tools/SequencerSnapField.h b/Engine/Source/Editor/Sequencer/Private/Tools/SequencerSnapField.h
--- a/Engine/Source/Editor/Sequencer/Private/Tools/SequencerSnapField.h
+++ b/Engine/Source/Editor/Sequencer/Private/Tools/SequencerSnapField.h
@@ -10,6 +10,8 @@
 class IKeyArea;
 class FSequencer;
 
+enum class ESequencerScrubberStyle : uint8;
+
 /** A snapping field that provides efficient snapping calculations on a range of values */
 class FSequencerSnapField
 {
@@ -31,16 +33,14 @@ public:
 	/** Construction from a sequencer and a snap canidate implementation. Optionally provide an entity mask to completely ignore some entity types */
 	FSequencerSnapField(const FSequencer& InSequencer, UE::Sequencer::ISnapCandidate& Candidate, uint32 EntityMask = ESequencerEntity::Everything);
 
-	/** Move construction / assignment */
-	FSequencerSnapField(FSequencerSnapField&& In) : SortedSnaps(MoveTemp(In.SortedSnaps)) {}
-	FSequencerSnapField& operator=(FSequencerSnapField&& In) { SortedSnaps = MoveTemp(In.SortedSnaps); return *this; }
-
 	void Initialize(const FSequencer& InSequencer, UE::Sequencer::ISnapCandidate& Candidate, uint32 EntityMask = ESequencerEntity::Everything);
 
 	void AddExplicitSnap(UE::Sequencer::FSnapPoint InSnap);
 
 	void Finalize();
 
+	void SetSnapToInterval(bool bInSnapToInterval) { bSnapToInterval = bInSnapToInterval; }
+
 	/** Snap the specified time to this field with the given threshold */
 	TOptional<FSnapResult> Snap(const FFrameTime& InTime, const FFrameTime& Threshold) const;
 
@@ -50,4 +50,10 @@ public:
 private:
 	/** Array of snap points, approximately grouped, and sorted in ascending order by time */
 	TArray<UE::Sequencer::FSnapPoint> SortedSnaps;
+
+	FFrameRate TickResolution;
+	FFrameRate DisplayRate;
+	ESequencerScrubberStyle ScrubStyle;
+
+	bool bSnapToInterval = false;
 };
