Sequencer: Fix painting of multi-selection hash overlay, and extend this overlay to the section contents

diff --git a/Engine/Source/Editor/Sequencer/Private/SSequencerSection.cpp b/Engine/Source/Editor/Sequencer/Private/SSequencerSection.cpp
--- a/Engine/Source/Editor/Sequencer/Private/SSequencerSection.cpp
+++ b/Engine/Source/Editor/Sequencer/Private/SSequencerSection.cpp
@@ -242,7 +242,7 @@ struct FSequencerSectionPainterImpl : FSequencerSectionPainter
 				FVector2D(ExpandedSectionGeometry.GetLocalSize().X, HeaderHeight),
 				FSlateLayoutTransform()
 			);
-			FPaintGeometry ContentsGeometry = ExpandedSectionGeometry.ToPaintGeometry(
+			FGeometry ContentsGeometry = ExpandedSectionGeometry.MakeChild(
 				FVector2D(ExpandedSectionGeometry.GetLocalSize().X, ExpandedSectionGeometry.GetLocalSize().Y - HeaderHeight),
 				FSlateLayoutTransform(FVector2D(0.f, HeaderHeight)) 
 			);
@@ -256,30 +256,42 @@ struct FSequencerSectionPainterImpl : FSequencerSectionPainter
 				BlendedTint
 			);
 
+			FLinearColor FillTint = BlendedTint.LinearRGBToHSV();
+			FillTint.G *= .5f;
+			FillTint.B = FMath::Max(.03f, FillTint.B*.1f);
+			FSlateDrawElement::MakeBox(
+				DrawElements,
+				++LayerId,
+				ContentsGeometry.ToPaintGeometry(),
+				SectionContentsBackgroundBrush,
+				DrawEffects,
+				FillTint.HSVToLinearRGB()
+			);
+
 			// Draw the selection hash
 			if (SelectionColor.IsSet())
 			{
 				FSlateDrawElement::MakeBox(
 					DrawElements,
-					LayerId,
-					HeaderGeometry.ToPaintGeometry(FVector2D(1, 1), HeaderGeometry.GetLocalSize() - FVector2D(2,2)),
+					++LayerId,
+					HeaderGeometry.ToPaintGeometry(FVector2D(1, 1), HeaderGeometry.GetLocalSize() - FVector2D(2, 2)),
 					SectionHeaderSelectedSectionOverlay,
 					DrawEffects,
 					SelectionColor.GetValue().CopyWithNewOpacity(0.8f)
 				);
-			}
 
-			FLinearColor FillTint = BlendedTint.LinearRGBToHSV();
-			FillTint.G *= .5f;
-			FillTint.B = FMath::Max(.03f, FillTint.B*.1f);
-			FSlateDrawElement::MakeBox(
-				DrawElements,
-				LayerId,
-				ContentsGeometry,
-				SectionContentsBackgroundBrush,
-				DrawEffects,
-				FillTint.HSVToLinearRGB()
-			);
+				FLinearColor FillSelectionColor = SelectionColor.GetValue().LinearRGBToHSV();
+				FillSelectionColor.G *= .5f;
+				FillSelectionColor.B = FMath::Max(.03f, FillSelectionColor.B * .1f);
+				FSlateDrawElement::MakeBox(
+					DrawElements,
+					++LayerId,
+					ContentsGeometry.ToPaintGeometry(FVector2D(1, 1), ContentsGeometry.GetLocalSize() - FVector2D(2, 2)),
+					SectionHeaderSelectedSectionOverlay,
+					DrawEffects,
+					FillSelectionColor.HSVToLinearRGB().CopyWithNewOpacity(0.8f)
+				);
+			}
 		}
 
 		if (!ensure(!bClipRectEnabled))
