Sequencer: Clean internal state when recycling VMs

This fixes a problem of pre-recycling state leaking through after an update

diff --git a/Engine/Source/Editor/Sequencer/Private/MVVM/ViewModels/CategoryModel.cpp b/Engine/Source/Editor/Sequencer/Private/MVVM/ViewModels/CategoryModel.cpp
--- a/Engine/Source/Editor/Sequencer/Private/MVVM/ViewModels/CategoryModel.cpp
+++ b/Engine/Source/Editor/Sequencer/Private/MVVM/ViewModels/CategoryModel.cpp
@@ -229,6 +229,11 @@ void FCategoryGroupModel::Delete()
 	}
 }
 
+void FCategoryGroupModel::OnRecycle()
+{
+	Categories.Empty();
+}
+
 } // namespace Sequencer
 } // namespace UE
 
diff --git a/Engine/Source/Editor/Sequencer/Private/MVVM/ViewModels/ChannelModel.cpp b/Engine/Source/Editor/Sequencer/Private/MVVM/ViewModels/ChannelModel.cpp
--- a/Engine/Source/Editor/Sequencer/Private/MVVM/ViewModels/ChannelModel.cpp
+++ b/Engine/Source/Editor/Sequencer/Private/MVVM/ViewModels/ChannelModel.cpp
@@ -133,6 +133,11 @@ FTrackLaneVirtualAlignment FChannelModel::ArrangeVirtualTrackLaneView() const
 	return FTrackLaneVirtualAlignment::Proportional(TRange<FFrameNumber>(), 1.f);
 }
 
+void FChannelModel::OnRecycle()
+{
+	KeyArea = nullptr;
+}
+
 bool FChannelModel::UpdateCachedKeys(TSharedPtr<FCachedKeys>& OutCachedKeys) const
 {
 	struct FSequencerCachedKeys : FCachedKeys
@@ -374,6 +379,11 @@ FViewModelVariantIterator FChannelGroupModel::GetTrackAreaModelList() const
 	return &Channels;
 }
 
+void FChannelGroupModel::OnRecycle()
+{
+	Channels.Empty();
+}
+
 void FChannelGroupModel::CreateCurveModels(TArray<TUniquePtr<FCurveModel>>& OutCurveModels)
 {
 	TSharedPtr<FSequenceModel> SequenceModel = FindAncestorOfType<FSequenceModel>();
diff --git a/Engine/Source/Editor/Sequencer/Private/MVVM/ViewModels/FolderModel.cpp b/Engine/Source/Editor/Sequencer/Private/MVVM/ViewModels/FolderModel.cpp
--- a/Engine/Source/Editor/Sequencer/Private/MVVM/ViewModels/FolderModel.cpp
+++ b/Engine/Source/Editor/Sequencer/Private/MVVM/ViewModels/FolderModel.cpp
@@ -7,6 +7,7 @@
 #include "MVVM/ViewModels/LayerBarModel.h"
 #include "MVVM/ViewModels/ObjectBindingModel.h"
 #include "MVVM/ViewModels/ViewModelIterators.h"
+#include "MVVM/Extensions/IRecyclableExtension.h"
 #include "MVVM/Views/SOutlinerItemViewBase.h"
 #include "MVVM/ViewModels/SequencerEditorViewModel.h"
 #include "MVVM/ViewModels/OutlinerViewModelDragDropOp.h"
@@ -111,7 +112,7 @@ void FFolderModel::RepopulateChildren()
 
 	// Keep our existing children alive
 	FScopedViewModelListHead RecycledOutliner(This, EViewModelListType::Recycled);
-	OutlinerChildren.MoveChildrenTo(RecycledOutliner.GetChildren());
+	OutlinerChildren.MoveChildrenTo<IRecyclableExtension>(RecycledOutliner.GetChildren(), IRecyclableExtension::CallOnRecycle);
 
 	// Create/recycle models for any children
 	for (UMovieSceneFolder* ChildFolder : Folder->GetChildFolders())
diff --git a/Engine/Source/Editor/Sequencer/Private/MVVM/ViewModels/ObjectBindingModel.cpp b/Engine/Source/Editor/Sequencer/Private/MVVM/ViewModels/ObjectBindingModel.cpp
--- a/Engine/Source/Editor/Sequencer/Private/MVVM/ViewModels/ObjectBindingModel.cpp
+++ b/Engine/Source/Editor/Sequencer/Private/MVVM/ViewModels/ObjectBindingModel.cpp
@@ -11,6 +11,7 @@
 #include "MVVM/ViewModels/SequencerEditorViewModel.h"
 #include "MVVM/ViewModels/OutlinerViewModelDragDropOp.h"
 #include "MVVM/Views/SOutlinerObjectBindingView.h"
+#include "MVVM/Extensions/IRecyclableExtension.h"
 #include "Modules/ModuleManager.h"
 #include "ISequencerModule.h"
 #include "MovieScene.h"
@@ -208,7 +209,7 @@ void FObjectBindingModel::OnConstruct()
 	check(Binding);
 
 	FScopedViewModelListHead RecycledHead(AsShared(), EViewModelListType::Recycled);
-	GetChildrenForList(&OutlinerChildList).MoveChildrenTo(RecycledHead.GetChildren());
+	GetChildrenForList(&OutlinerChildList).MoveChildrenTo<IRecyclableExtension>(RecycledHead.GetChildren(), IRecyclableExtension::CallOnRecycle);
 
 	for (UMovieSceneTrack* Track : Binding->GetTracks())
 	{
diff --git a/Engine/Source/Editor/Sequencer/Private/MVVM/ViewModels/TrackModel.cpp b/Engine/Source/Editor/Sequencer/Private/MVVM/ViewModels/TrackModel.cpp
--- a/Engine/Source/Editor/Sequencer/Private/MVVM/ViewModels/TrackModel.cpp
+++ b/Engine/Source/Editor/Sequencer/Private/MVVM/ViewModels/TrackModel.cpp
@@ -2,6 +2,7 @@
 
 #include "MVVM/ViewModels/TrackModel.h"
 #include "MVVM/Extensions/IObjectBindingExtension.h"
+#include "MVVM/Extensions/IRecyclableExtension.h"
 #include "MVVM/ViewModels/ChannelModel.h"
 #include "MVVM/ViewModels/FolderModel.h"
 #include "MVVM/ViewModels/ViewModelIterators.h"
@@ -198,8 +199,8 @@ void FTrackModel::ForceUpdate()
 		TSharedPtr<FViewModel> SectionsTail;
 
 		FScopedViewModelListHead RecycledModels(AsShared(), EViewModelListType::Recycled);
-		GetChildrenForList(&SectionList).MoveChildrenTo(RecycledModels.GetChildren());
-		OutlinerChildren.MoveChildrenTo(RecycledModels.GetChildren());
+		GetChildrenForList(&SectionList).MoveChildrenTo<IRecyclableExtension>(RecycledModels.GetChildren(), IRecyclableExtension::CallOnRecycle);
+		OutlinerChildren.MoveChildrenTo<IRecyclableExtension>(RecycledModels.GetChildren(), IRecyclableExtension::CallOnRecycle);
 
 		// Add all sections directly to this track row
 		for (UMovieSceneSection* Section : Track->GetAllSections())
@@ -244,12 +245,8 @@ void FTrackModel::ForceUpdate()
 		// This should only be required if this track previously represented
 		// a single row, but now there are multiple rows
 		FScopedViewModelListHead RecycledModels(AsShared(), EViewModelListType::Recycled);
-		GetChildrenForList(&SectionList).MoveChildrenTo(RecycledModels.GetChildren());
-		OutlinerChildren.MoveChildrenTo(RecycledModels.GetChildren());
-
-		// Also clear our outliner children but keep them alive, because there could be category groups or channel
-		// groups in there, if we previously only had 1 row before
-		TSharedPtr<FViewModel> OldOutlinerItems = OutlinerChildList.ReliquishList();
+		GetChildrenForList(&SectionList).MoveChildrenTo<IRecyclableExtension>(RecycledModels.GetChildren(), IRecyclableExtension::CallOnRecycle);
+		OutlinerChildren.MoveChildrenTo<IRecyclableExtension>(RecycledModels.GetChildren(), IRecyclableExtension::CallOnRecycle);
 
 		// We need to build row models so let's grab the storage for that
 		FTrackRowModelStorageExtension* TrackRowModelStorage = SequenceModel->CastDynamic<FTrackRowModelStorageExtension>();
@@ -287,7 +284,7 @@ void FTrackModel::ForceUpdate()
 
 				// Keep sections on rows alive as well
 				RowModels[RowIndex].OldSections = MakeUnique<FScopedViewModelListHead>(TrackRowModel, EViewModelListType::Recycled);
-				TrackRowModel->GetSectionModels().MoveChildrenTo(RowModels[RowIndex].OldSections->GetChildren());
+				TrackRowModel->GetSectionModels().MoveChildrenTo<IRecyclableExtension>(RowModels[RowIndex].OldSections->GetChildren(), IRecyclableExtension::CallOnRecycle);
 			}
 		}
 
diff --git a/Engine/Source/Editor/Sequencer/Private/MVVM/ViewModels/TrackModelLayoutBuilder.cpp b/Engine/Source/Editor/Sequencer/Private/MVVM/ViewModels/TrackModelLayoutBuilder.cpp
--- a/Engine/Source/Editor/Sequencer/Private/MVVM/ViewModels/TrackModelLayoutBuilder.cpp
+++ b/Engine/Source/Editor/Sequencer/Private/MVVM/ViewModels/TrackModelLayoutBuilder.cpp
@@ -8,6 +8,7 @@
 #include "HAL/Platform.h"
 #include "ISequencerSection.h"
 #include "MVVM/Extensions/IOutlinerExtension.h"
+#include "MVVM/Extensions/IRecyclableExtension.h"
 #include "MVVM/ViewModels/CategoryModel.h"
 #include "MVVM/ViewModels/ChannelModel.h"
 #include "MVVM/ViewModels/SectionModel.h"
@@ -57,7 +58,7 @@ void FHierarchicalModelListRefresher::RecycleChildren(const TSharedPtr<FViewMode
 		RecycledLists.Last().Emplace(InModel, EViewModelListType::Recycled);
 
 		FViewModelChildren RecycledChildren = InModel->GetChildList(EViewModelListType::Recycled);
-		InExistingChildren.MoveChildrenTo(RecycledChildren);
+		InExistingChildren.MoveChildrenTo<IRecyclableExtension>(RecycledChildren, IRecyclableExtension::CallOnRecycle);
 	}
 }
 
diff --git a/Engine/Source/Editor/Sequencer/Public/MVVM/ViewModels/CategoryModel.h b/Engine/Source/Editor/Sequencer/Public/MVVM/ViewModels/CategoryModel.h
--- a/Engine/Source/Editor/Sequencer/Public/MVVM/ViewModels/CategoryModel.h
+++ b/Engine/Source/Editor/Sequencer/Public/MVVM/ViewModels/CategoryModel.h
@@ -9,6 +9,7 @@
 #include "MVVM/Extensions/IDeletableExtension.h"
 #include "MVVM/Extensions/IGeometryExtension.h"
 #include "MVVM/Extensions/IOutlinerExtension.h"
+#include "MVVM/Extensions/IRecyclableExtension.h"
 #include "MVVM/Extensions/ITrackAreaExtension.h"
 #include "MVVM/Extensions/ITrackLaneExtension.h"
 #include "MVVM/Extensions/LinkedOutlinerExtension.h"
@@ -81,9 +82,10 @@ class SEQUENCER_API FCategoryGroupModel
 	, public ITrackAreaExtension
 	, public ICompoundOutlinerExtension
 	, public IDeletableExtension
+	, public IRecyclableExtension
 {
 public:
-	UE_SEQUENCER_DECLARE_CASTABLE(FCategoryGroupModel, FOutlinerItemModel, ITrackAreaExtension, ICompoundOutlinerExtension, IDeletableExtension);
+	UE_SEQUENCER_DECLARE_CASTABLE(FCategoryGroupModel, FOutlinerItemModel, ITrackAreaExtension, ICompoundOutlinerExtension, IDeletableExtension, IRecyclableExtension);
 
 	explicit FCategoryGroupModel(FName InCategoryName, const FText& InDisplayText);
 
@@ -127,6 +129,9 @@ public:
 	bool CanDelete(FText* OutErrorMessage) const override;
 	void Delete() override;
 
+	/*~ IRecyclableExtension */
+	void OnRecycle() override;
+
 private:
 
 	TArray<TWeakViewModelPtr<FCategoryModel>> Categories;
diff --git a/Engine/Source/Editor/Sequencer/Public/MVVM/ViewModels/ChannelModel.h b/Engine/Source/Editor/Sequencer/Public/MVVM/ViewModels/ChannelModel.h
--- a/Engine/Source/Editor/Sequencer/Public/MVVM/ViewModels/ChannelModel.h
+++ b/Engine/Source/Editor/Sequencer/Public/MVVM/ViewModels/ChannelModel.h
@@ -13,6 +13,7 @@
 #include "MVVM/Extensions/IGeometryExtension.h"
 #include "MVVM/Extensions/IKeyExtension.h"
 #include "MVVM/Extensions/IOutlinerExtension.h"
+#include "MVVM/Extensions/IRecyclableExtension.h"
 #include "MVVM/Extensions/ITrackAreaExtension.h"
 #include "MVVM/Extensions/ITrackLaneExtension.h"
 #include "MVVM/Extensions/LinkedOutlinerExtension.h"
@@ -54,11 +55,12 @@ class SEQUENCER_API FChannelModel
 	, public FGeometryExtensionShim
 	, public FLinkedOutlinerComputedSizingShim
 	, public ITrackLaneExtension
+	, public IRecyclableExtension
 	, public IKeyExtension
 {
 public:
 
-	UE_SEQUENCER_DECLARE_CASTABLE(FChannelModel, FViewModel, FLinkedOutlinerExtension, IGeometryExtension, ITrackLaneExtension, IKeyExtension);
+	UE_SEQUENCER_DECLARE_CASTABLE(FChannelModel, FViewModel, FLinkedOutlinerExtension, IGeometryExtension, ITrackLaneExtension, IRecyclableExtension, IKeyExtension);
 
 	FChannelModel(FName InChannelName, TWeakPtr<ISequencerSection> InSection, FMovieSceneChannelHandle InChannel);
 	~FChannelModel();
@@ -92,6 +94,9 @@ public:
 	TSharedPtr<ITrackLaneWidget> CreateTrackLaneView(const FCreateTrackLaneViewParams& InParams) override;
 	FTrackLaneVirtualAlignment ArrangeVirtualTrackLaneView() const override;
 
+	/*~ IRecyclableExtension */
+	void OnRecycle() override;
+
 	/*~ IKeyExtension */
 	bool UpdateCachedKeys(TSharedPtr<FCachedKeys>& OutCachedKeys) const override;
 	bool GetFixedExtents(double& OutFixedMin, double& OutFixedMax) const override;
@@ -115,10 +120,11 @@ private:
 class SEQUENCER_API FChannelGroupModel
 	: public FViewModel
 	, public ITrackAreaExtension
+	, public IRecyclableExtension
 {
 public:
 
-	UE_SEQUENCER_DECLARE_CASTABLE(FChannelGroupModel, FViewModel, ITrackAreaExtension);
+	UE_SEQUENCER_DECLARE_CASTABLE(FChannelGroupModel, FViewModel, ITrackAreaExtension, IRecyclableExtension);
 
 	FChannelGroupModel(FName InChannelName, const FText& InDisplayText);
 	~FChannelGroupModel();
@@ -164,6 +170,9 @@ public:
 	FTrackAreaParameters GetTrackAreaParameters() const override;
 	FViewModelVariantIterator GetTrackAreaModelList() const override;
 
+	/*~ IRecyclableExtension */
+	void OnRecycle() override;
+
 	void CreateCurveModels(TArray<TUniquePtr<FCurveModel>>& OutCurveModels);
 	bool HasCurves() const;
 
diff --git a/Engine/Source/Editor/SequencerCore/Private/MVVM/Extensions/IRecyclableExtension.cpp b/Engine/Source/Editor/SequencerCore/Private/MVVM/Extensions/IRecyclableExtension.cpp
new file mode 100644
--- /dev/null
+++ b/Engine/Source/Editor/SequencerCore/Private/MVVM/Extensions/IRecyclableExtension.cpp
@@ -0,0 +1,21 @@
+// Copyright Epic Games, Inc. All Rights Reserved.
+
+#include "MVVM/Extensions/IRecyclableExtension.h"
+#include "MVVM/ViewModelPtr.h"
+
+namespace UE
+{
+namespace Sequencer
+{
+
+void IRecyclableExtension::CallOnRecycle(const TViewModelPtr<IRecyclableExtension>& RecyclableItem)
+{
+	if (RecyclableItem)
+	{
+		RecyclableItem->OnRecycle();
+	}
+}
+
+} // namespace Sequencer
+} // namespace UE
+
diff --git a/Engine/Source/Editor/SequencerCore/Private/MVVM/ViewModelTypeIDs.cpp b/Engine/Source/Editor/SequencerCore/Private/MVVM/ViewModelTypeIDs.cpp
--- a/Engine/Source/Editor/SequencerCore/Private/MVVM/ViewModelTypeIDs.cpp
+++ b/Engine/Source/Editor/SequencerCore/Private/MVVM/ViewModelTypeIDs.cpp
@@ -15,6 +15,7 @@
 #include "MVVM/Extensions/IMutableExtension.h"
 #include "MVVM/Extensions/IOutlinerExtension.h"
 #include "MVVM/Extensions/IPinnableExtension.h"
+#include "MVVM/Extensions/IRecyclableExtension.h"
 #include "MVVM/Extensions/IRenameableExtension.h"
 #include "MVVM/Extensions/IResizableExtension.h"
 #include "MVVM/Extensions/ISelectableExtension.h"
@@ -53,6 +54,7 @@ UE_SEQUENCER_DEFINE_VIEW_MODEL_TYPE_ID(IMutableExtension);
 UE_SEQUENCER_DEFINE_VIEW_MODEL_TYPE_ID(IOutlinerDropTargetOutlinerExtension);
 UE_SEQUENCER_DEFINE_VIEW_MODEL_TYPE_ID(IOutlinerExtension);
 UE_SEQUENCER_DEFINE_VIEW_MODEL_TYPE_ID(IPinnableExtension);
+UE_SEQUENCER_DEFINE_VIEW_MODEL_TYPE_ID(IRecyclableExtension);
 UE_SEQUENCER_DEFINE_VIEW_MODEL_TYPE_ID(IRenameableExtension);
 UE_SEQUENCER_DEFINE_VIEW_MODEL_TYPE_ID(IResizableExtension);
 UE_SEQUENCER_DEFINE_VIEW_MODEL_TYPE_ID(ISelectableExtension);
diff --git a/Engine/Source/Editor/SequencerCore/Public/MVVM/Extensions/IRecyclableExtension.h b/Engine/Source/Editor/SequencerCore/Public/MVVM/Extensions/IRecyclableExtension.h
new file mode 100644
--- /dev/null
+++ b/Engine/Source/Editor/SequencerCore/Public/MVVM/Extensions/IRecyclableExtension.h
@@ -0,0 +1,30 @@
+// Copyright Epic Games, Inc. All Rights Reserved.
+
+#pragma once
+
+#include "CoreTypes.h"
+#include "MVVM/ViewModelTypeID.h"
+
+namespace UE
+{
+namespace Sequencer
+{
+
+template<typename T> struct TViewModelPtr;
+
+class SEQUENCERCORE_API IRecyclableExtension
+{
+public:
+
+	UE_SEQUENCER_DECLARE_VIEW_MODEL_TYPE_ID(IRecyclableExtension)
+
+	static void CallOnRecycle(const TViewModelPtr<IRecyclableExtension>& RecyclableItem);
+
+	virtual ~IRecyclableExtension(){}
+
+	virtual void OnRecycle() = 0;
+};
+
+} // namespace Sequencer
+} // namespace UE
+
diff --git a/Engine/Source/Editor/SequencerCore/Public/MVVM/ViewModels/ViewModelHierarchy.h b/Engine/Source/Editor/SequencerCore/Public/MVVM/ViewModels/ViewModelHierarchy.h
--- a/Engine/Source/Editor/SequencerCore/Public/MVVM/ViewModels/ViewModelHierarchy.h
+++ b/Engine/Source/Editor/SequencerCore/Public/MVVM/ViewModels/ViewModelHierarchy.h
@@ -210,6 +210,19 @@ struct SEQUENCERCORE_API FViewModelChildren
 	 */
 	void MoveChildrenTo(const FViewModelChildren& OutDestination);
 
+	/**
+	 * Invoke a callback on these children, and move them all into the specified destination list, changing parent pointers if necessary
+	 */
+	template<typename T, typename Callback>
+	void MoveChildrenTo(const FViewModelChildren& OutDestination, Callback&& InCallback)
+	{
+		for (const TViewModelPtr<T>& Item : IterateSubList<T>())
+		{
+			InCallback(Item);
+		}
+		MoveChildrenTo(OutDestination);
+	}
+
 	/**
 	 * Empty this specific sub-list, resetting all parent ptrs
 	 */
