Sequencer: Fixed widget interaction between top level channel views and custom section widgets

diff --git a/Engine/Source/Editor/Sequencer/Private/SSequencerSection.cpp b/Engine/Source/Editor/Sequencer/Private/SSequencerSection.cpp
--- a/Engine/Source/Editor/Sequencer/Private/SSequencerSection.cpp
+++ b/Engine/Source/Editor/Sequencer/Private/SSequencerSection.cpp
@@ -892,10 +892,11 @@ EVisibility SSequencerSection::GetTopLevelChannelGroupVisibility() const
 	TSharedPtr<FSectionModel> SectionModel = WeakSectionModel.Pin();
 	TViewModelPtr<IOutlinerExtension> Outliner = SectionModel ? SectionModel->FindAncestorOfType<IOutlinerExtension>() : nullptr;
 
-	if (!Outliner || Outliner->IsExpanded())
+	if (!SectionModel || !Outliner || Outliner->IsExpanded() || !SectionModel->GetDescendantsOfType<FChannelModel>())
 	{
 		return EVisibility::Collapsed;
 	}
+
 	return EVisibility::Visible;
 }
 
