Sequencer: Prevent layer bars from contributing to Actor selection events

diff --git a/Engine/Source/Editor/Sequencer/Private/SequencerSelection.cpp b/Engine/Source/Editor/Sequencer/Private/SequencerSelection.cpp
--- a/Engine/Source/Editor/Sequencer/Private/SequencerSelection.cpp
+++ b/Engine/Source/Editor/Sequencer/Private/SequencerSelection.cpp
@@ -123,7 +123,7 @@ const TSet<TWeakPtr<UE::Sequencer::FViewModel>>& FSequencerSelection::GetNodesWi
 	for (TWeakPtr<FViewModel> WeakTrackAreaModel : SelectedTrackAreaItems)
 	{
 		TSharedPtr<FViewModel> TrackAreaModel = WeakTrackAreaModel.Pin();
-		if (TrackAreaModel)
+		if (TrackAreaModel && TrackAreaModel->IsA<FSectionModel>())
 		{
 			TViewModelPtr<IOutlinerExtension> OutlinerItem = TrackAreaModel->FindAncestorOfType<IOutlinerExtension>();
 			if (OutlinerItem)
