Sequencer: Don't allow selected infinite sections to be moved/dragged

Movement tool is not activated if the section hotspot is an infinite section

diff --git a/Engine/Source/Editor/Sequencer/Private/SequencerHotspots.cpp b/Engine/Source/Editor/Sequencer/Private/SequencerHotspots.cpp
--- a/Engine/Source/Editor/Sequencer/Private/SequencerHotspots.cpp
+++ b/Engine/Source/Editor/Sequencer/Private/SequencerHotspots.cpp
@@ -439,7 +439,10 @@ void FSectionHotspotBase::UpdateOnHover(FTrackAreaViewModel& InTrackArea) const
 	TSharedPtr<ISequencer> Sequencer = WeakSequencer.Pin();
 	if (Sequencer->GetSelection().IsSelected(WeakSectionModel))
 	{
-		InTrackArea.AttemptToActivateTool(FSequencerEditTool_Movement::Identifier);
+		if (ThisSection && ThisSection->GetRange() != TRange<FFrameNumber>::All())
+		{
+			InTrackArea.AttemptToActivateTool(FSequencerEditTool_Movement::Identifier);
+		}
 	}
 	else if (ThisSection)
 	{
