Sequencer: Fixed edge-case that caused object bindings to be parented incorrectly

diff --git a/Engine/Source/Editor/Sequencer/Private/MVVM/ObjectBindingModelStorageExtension.cpp b/Engine/Source/Editor/Sequencer/Private/MVVM/ObjectBindingModelStorageExtension.cpp
--- a/Engine/Source/Editor/Sequencer/Private/MVVM/ObjectBindingModelStorageExtension.cpp
+++ b/Engine/Source/Editor/Sequencer/Private/MVVM/ObjectBindingModelStorageExtension.cpp
@@ -192,6 +192,14 @@ TSharedPtr<FViewModel> FObjectBindingModelStorageExtension::GetOrCreateModelForB
 	if (DesiredParent.IsValid())
 	{
 		TSharedPtr<FViewModel> Parent = GetOrCreateModelForBinding(DesiredParent);
+
+		// If our parent doesn't have a parent, it will probably get destructed at the end of this
+		// function because nothing is holding onto it, so we add it to the root to keep it alive
+		if (Parent->GetParent() == nullptr)
+		{
+			OwnerModel->GetChildList(EViewModelListType::Outliner).AddChild(Parent);
+		}
+
 		Parent->GetChildList(EViewModelListType::Outliner).AddChild(ObjectModel);
 	}
 
