Sequencer: When rendering from a shot, enabling only the shot in the context of the root sequence

diff --git a/Engine/Source/Editor/Sequencer/Private/Sequencer.cpp b/Engine/Source/Editor/Sequencer/Private/Sequencer.cpp
--- a/Engine/Source/Editor/Sequencer/Private/Sequencer.cpp
+++ b/Engine/Source/Editor/Sequencer/Private/Sequencer.cpp
@@ -3473,7 +3473,13 @@ void FSequencer::RenderMovieInternal(TRange<FFrameNumber> Range, bool bSetFrameO
 	ISequencerModule& SequencerModule = FModuleManager::LoadModuleChecked<ISequencerModule>("Sequencer");
 	if (IMovieRendererInterface* MovieRenderer = SequencerModule.GetMovieRenderer(GetMovieRendererName()))
 	{
-		MovieRenderer->RenderMovie(GetRootMovieSceneSequence(), TArray<UMovieSceneCinematicShotSection*>());
+		TArray<UMovieSceneCinematicShotSection*> ShotSections;
+		if (UMovieSceneCinematicShotSection* ShotSection = Cast<UMovieSceneCinematicShotSection>(FindSubSection(GetFocusedTemplateID())))
+		{
+			ShotSections.Add(ShotSection);
+		}
+
+		MovieRenderer->RenderMovie(GetRootMovieSceneSequence(), ShotSections);
 		return;
 	}
 
