Sequencer: With the new Left click does Marquee alt option to always move the camera we need to redraw the viewports when doing so or the gizmo may lag. This is one in some other places also and mimics how the default alt drag works.

diff --git a/Engine/Source/Editor/Sequencer/Private/SequencerEdMode.cpp b/Engine/Source/Editor/Sequencer/Private/SequencerEdMode.cpp
--- a/Engine/Source/Editor/Sequencer/Private/SequencerEdMode.cpp
+++ b/Engine/Source/Editor/Sequencer/Private/SequencerEdMode.cpp
@@ -335,6 +335,7 @@ bool FSequencerEdMode::InputDelta(FEditorViewportClient* InViewportClient, FView
 			InViewportClient->SetCurrentWidgetAxis(EAxisList::None);
 			InViewportClient->PeformDefaultCameraMovement(InDrag, InRot, InScale);
 			GUnrealEd->UpdatePivotLocationForSelection();
+			GUnrealEd->RedrawLevelEditingViewports();
 		}
 
 		return true;
@@ -343,6 +344,7 @@ bool FSequencerEdMode::InputDelta(FEditorViewportClient* InViewportClient, FView
 	{
 		return DragToolHandler.InputDelta(InViewportClient, InViewport, InDrag, InRot, InScale);
 		GUnrealEd->UpdatePivotLocationForSelection();
+		GUnrealEd->RedrawLevelEditingViewports();
 
 	}
 	return FEdMode::InputDelta(InViewportClient, InViewport, InDrag, InRot, InScale);
