Sequencer: fix default expansion state initialization

Object bindings already start expanded thanks to GetDefaultExpansionState. We don't need to mess that up by setting SetExpansion in the constructor, which actually clears the init flag and prevents looking up the saved state.

diff --git a/Engine/Source/Editor/Sequencer/Private/MVVM/ViewModels/ObjectBindingModel.cpp b/Engine/Source/Editor/Sequencer/Private/MVVM/ViewModels/ObjectBindingModel.cpp
--- a/Engine/Source/Editor/Sequencer/Private/MVVM/ViewModels/ObjectBindingModel.cpp
+++ b/Engine/Source/Editor/Sequencer/Private/MVVM/ViewModels/ObjectBindingModel.cpp
@@ -176,9 +176,6 @@ FObjectBindingModel::FObjectBindingModel(FSequenceModel* InOwnerModel, const FMo
 {
 	RegisterChildList(&TrackAreaList);
 
-	// Objects are expanded by default
-	SetExpansion(true);
-
 	SetIdentifier(*ObjectBindingID.ToString());
 }
 
