Sequencer: Save the tree view width percentage in sequencer settings so that the split between the tree view and the track area is saved/restored.

diff --git a/Engine/Source/Editor/Sequencer/Private/SSequencer.cpp b/Engine/Source/Editor/Sequencer/Private/SSequencer.cpp
--- a/Engine/Source/Editor/Sequencer/Private/SSequencer.cpp
+++ b/Engine/Source/Editor/Sequencer/Private/SSequencer.cpp
@@ -256,9 +256,6 @@ void SSequencer::Construct(const FArguments& InArgs, TSharedRef<FSequencer> InSe
 	RootCustomization.AddMenuExtender = InArgs._AddMenuExtender;
 	RootCustomization.ToolbarExtender = InArgs._ToolbarExtender;
 
-	ColumnFillCoefficients[0] = 0.3f;
-	ColumnFillCoefficients[1] = 0.7f;
-
 	TAttribute<float> FillCoefficient_0, FillCoefficient_1;
 	{
 		FillCoefficient_0.Bind(TAttribute<float>::FGetter::CreateSP(this, &SSequencer::GetColumnFillCoefficient, 0));
@@ -2920,6 +2917,17 @@ void SSequencer::OnOutlinerSearchChanged( const FText& Filter )
 	}
 }
 
+float SSequencer::GetColumnFillCoefficient(int32 ColumnIndex) const
+{
+	if (ColumnIndex == 0)
+	{
+		return GetSequencerSettings()->GetTreeViewWidth();
+	}
+	else
+	{
+		return 1.f - GetSequencerSettings()->GetTreeViewWidth();
+	}
+}
 
 void SSequencer::OnDragEnter( const FGeometry& MyGeometry, const FDragDropEvent& DragDropEvent )
 {
@@ -3617,9 +3625,17 @@ EFrameNumberDisplayFormats SSequencer::GetTimeDisplayFormat() const
 {
 	return GetSequencerSettings()->GetTimeDisplayFormat();
 }
+
 void SSequencer::OnColumnFillCoefficientChanged(float FillCoefficient, int32 ColumnIndex)
 {
-	ColumnFillCoefficients[ColumnIndex] = FillCoefficient;
+	if (ColumnIndex == 0)
+	{
+		GetSequencerSettings()->SetTreeViewWidth(FillCoefficient);
+	}
+	else
+	{
+		GetSequencerSettings()->SetTreeViewWidth(1.f - FillCoefficient);
+	}
 }
 
 void SSequencer::OnCurveEditorVisibilityChanged(bool bShouldBeVisible)
diff --git a/Engine/Source/Editor/Sequencer/Private/SSequencer.h b/Engine/Source/Editor/Sequencer/Private/SSequencer.h
--- a/Engine/Source/Editor/Sequencer/Private/SSequencer.h
+++ b/Engine/Source/Editor/Sequencer/Private/SSequencer.h
@@ -521,10 +521,7 @@ private:
 	/**
 	 * @return The fill percentage of the animation outliner
 	 */
-	float GetColumnFillCoefficient(int32 ColumnIndex) const
-	{
-		return ColumnFillCoefficients[ColumnIndex];
-	}
+	float GetColumnFillCoefficient(int32 ColumnIndex) const;
 
 	/**
 	 * Called when one or more assets are dropped into the widget
@@ -702,10 +699,7 @@ private:
 	/** Container for the toolbar, so that we can re-create it as needed. */
 	TSharedPtr<SBox> ToolbarContainer;
 
-	/** The fill coefficients of each column in the grid. */
-	float ColumnFillCoefficients[2];
-
-	/** Whether the active timer is currently registered */
+	/** Whether the active timer is currently registered */ 
 	bool bIsActiveTimerRegistered;
 
 	/** Whether the user is selecting. Ignore selection changes from the level when the user is selecting. */
diff --git a/Engine/Source/Editor/Sequencer/Private/SequencerSettings.cpp b/Engine/Source/Editor/Sequencer/Private/SequencerSettings.cpp
--- a/Engine/Source/Editor/Sequencer/Private/SequencerSettings.cpp
+++ b/Engine/Source/Editor/Sequencer/Private/SequencerSettings.cpp
@@ -61,6 +61,7 @@ USequencerSettings::USequencerSettings( const FObjectInitializer& ObjectInitiali
 	TrajectoryPathCap = 250;
 	FrameNumberDisplayFormat = EFrameNumberDisplayFormats::Seconds;
 	bAutoExpandNodesOnSelection = true;
+	TreeViewWidth = 0.3f;
 	bShowTickLines = true;
 	bShowSequencerToolbar = true;
 }
@@ -914,4 +915,13 @@ void USequencerSettings::SetAutoExpandNodesOnSelection(bool bInAutoExpandNodesOn
 		bAutoExpandNodesOnSelection = bInAutoExpandNodesOnSelection;
 		SaveConfig();
 	}
+}
+
+void USequencerSettings::SetTreeViewWidth(float InTreeViewWidth)
+{
+	if (InTreeViewWidth != TreeViewWidth)
+	{
+		TreeViewWidth = InTreeViewWidth;
+		SaveConfig();
+	}
 }
\ No newline at end of file
diff --git a/Engine/Source/Editor/Sequencer/Public/SequencerSettings.h b/Engine/Source/Editor/Sequencer/Public/SequencerSettings.h
--- a/Engine/Source/Editor/Sequencer/Public/SequencerSettings.h
+++ b/Engine/Source/Editor/Sequencer/Public/SequencerSettings.h
@@ -397,6 +397,11 @@ public:
 	/** Sets whether or not to expand the outliner tree view when a child element is selected (from outside of the tree view). */
 	void SetAutoExpandNodesOnSelection(bool bInAutoExpandNodesOnSelection);
 
+	/** Gets the tree view width percentage */
+	float GetTreeViewWidth() const { return TreeViewWidth; }
+	/** Sets the tree view width percentage */
+	void SetTreeViewWidth(float InTreeViewWidth);
+
 protected:
 
 	/** The auto change mode (auto-key, auto-track or none). */
@@ -630,9 +635,13 @@ protected:
 	FString MovieRendererName;
 
 	/** Whether to expand the sequencer tree view when a child element is selected (from outside of the tree view). */
-	UPROPERTY(config, EditAnywhere, Category= General)
+	UPROPERTY(config, EditAnywhere, Category = General)
 	bool bAutoExpandNodesOnSelection;
 
+	/** The tree view width percentage */
+	UPROPERTY(config, EditAnywhere, Category = General)
+	float TreeViewWidth;
+
 	FOnEvaluateSubSequencesInIsolationChanged OnEvaluateSubSequencesInIsolationChangedEvent;
 	FOnShowSelectedNodesOnlyChanged OnShowSelectedNodesOnlyChangedEvent;
 	FOnAllowEditsModeChanged OnAllowEditsModeChangedEvent;
