Sequencer: Store the package dirty flag before attaching/detaching actors during restore state. Attach can call detach (if it was previously attached), and this dirties the package. This causes Sequencer to unnecessairly check out actors in OFPA as Sequencer restores the world state on save, but the attach track ends up dirtying the actor when putting it back to its original state. This causes an erronous package dirty which then shows up as a file that needs to be checked out/saved even though from the user perspective you haven't actually changed any assets in the level.

diff --git a/Engine/Plugins/Editor/DataValidation/Source/DataValidation/Private/EditorValidatorSubsystem.cpp b/Engine/Plugins/Editor/DataValidation/Source/DataValidation/Private/EditorValidatorSubsystem.cpp
--- a/Engine/Plugins/Editor/DataValidation/Source/DataValidation/Private/EditorValidatorSubsystem.cpp
+++ b/Engine/Plugins/Editor/DataValidation/Source/DataValidation/Private/EditorValidatorSubsystem.cpp
@@ -20,6 +20,7 @@
 #include "DataValidationChangelist.h"
 #include "DataValidationModule.h"
 #include "Engine/Level.h"
+#include "Editor.h"
 
 #include UE_INLINE_GENERATED_CPP_BY_NAME(EditorValidatorSubsystem)
 
@@ -217,6 +218,13 @@ int32 UEditorValidatorSubsystem::ValidateAssetsWithSettings(const TArray<FAssetD
 	SlowTask.Visibility = ESlowTaskVisibility::ForceVisible;
 	SlowTask.MakeDialogDelayed(.1f);
 
+	// Broadcast the Editor event before we start validating. This lets other systems (such as Sequencer) restore the state
+	// of the level to what is actually saved on disk before performing validation.
+	if (FEditorDelegates::OnPreAssetValidation.IsBound())
+	{
+		FEditorDelegates::OnPreAssetValidation.Broadcast();
+	}
+
 	FMessageLog DataValidationLog("AssetCheck");
 
 	int32 NumAdded = 0;
@@ -349,6 +357,12 @@ int32 UEditorValidatorSubsystem::ValidateAssetsWithSettings(const TArray<FAssetD
 		DataValidationLog.Open(EMessageSeverity::Info, true);
 	}
 
+	// Broadcast now that we're complete so other systems can go back to their previous state.
+	if (FEditorDelegates::OnPostAssetValidation.IsBound())
+	{
+		FEditorDelegates::OnPostAssetValidation.Broadcast();
+	}
+
 	return NumInvalidFiles + NumFilesWithWarnings;
 }
 
diff --git a/Engine/Source/Editor/Sequencer/Private/LevelEditorSequencerIntegration.cpp b/Engine/Source/Editor/Sequencer/Private/LevelEditorSequencerIntegration.cpp
--- a/Engine/Source/Editor/Sequencer/Private/LevelEditorSequencerIntegration.cpp
+++ b/Engine/Source/Editor/Sequencer/Private/LevelEditorSequencerIntegration.cpp
@@ -51,6 +51,8 @@
 #include "UnrealEdMisc.h"
 #include "Editor/UnrealEdEngine.h"
 #include "EditorSupportDelegates.h"
+#include "Subsystems/UnrealEditorSubsystem.h"
+
 
 #define LOCTEXT_NAMESPACE "LevelEditorSequencerIntegration"
 
@@ -194,6 +196,14 @@ void FLevelEditorSequencerIntegration::Initialize(const FLevelEditorSequencerInt
 		FDelegateHandle Handle = FEditorDelegates::PostSaveExternalActors.AddRaw(this, &FLevelEditorSequencerIntegration::OnPostSaveExternalActors);
 		AcquiredResources.Add([=]{ FEditorDelegates::PostSaveExternalActors.Remove(Handle); });
 	}
+	{
+		FDelegateHandle Handle = FEditorDelegates::OnPreAssetValidation.AddRaw(this, &FLevelEditorSequencerIntegration::OnPreAssetValidation);
+		AcquiredResources.Add([=] { FEditorDelegates::OnPreAssetValidation.Remove(Handle); });
+	}
+	{
+		FDelegateHandle Handle = FEditorDelegates::OnPostAssetValidation.AddRaw(this, &FLevelEditorSequencerIntegration::OnPostAssetValidation);
+		AcquiredResources.Add([=] { FEditorDelegates::OnPostAssetValidation.Remove(Handle); });
+	}
 	{
 		FDelegateHandle Handle = FEditorDelegates::PreBeginPIE.AddRaw(this, &FLevelEditorSequencerIntegration::OnPreBeginPIE);
 		AcquiredResources.Add([=]{ FEditorDelegates::PreBeginPIE.Remove(Handle); });
@@ -360,6 +370,25 @@ void FLevelEditorSequencerIntegration::OnPostSaveExternalActors(UWorld* World)
 	ResetToAnimatedState(World);
 }
 
+void FLevelEditorSequencerIntegration::OnPreAssetValidation()
+{
+	// Asset validation doesn't have a world context, so we'll just use the editor world.
+	UUnrealEditorSubsystem* UnrealEditorSubsystem = GEditor->GetEditorSubsystem<UUnrealEditorSubsystem>();
+	if(UnrealEditorSubsystem && UnrealEditorSubsystem->GetEditorWorld())
+	{
+		RestoreToSavedState(UnrealEditorSubsystem->GetEditorWorld());
+	}
+}
+
+void FLevelEditorSequencerIntegration::OnPostAssetValidation()
+{
+	UUnrealEditorSubsystem* UnrealEditorSubsystem = GEditor->GetEditorSubsystem<UUnrealEditorSubsystem>();
+	if(UnrealEditorSubsystem && UnrealEditorSubsystem->GetEditorWorld())
+	{
+		ResetToAnimatedState(UnrealEditorSubsystem->GetEditorWorld());
+	}
+}
+
 void FLevelEditorSequencerIntegration::OnNewCurrentLevel()
 {
 	auto IsSequenceEditor = [](const FSequencerAndOptions& In) { return In.Sequencer.IsValid() && In.Options.bActivateSequencerEdMode; };
diff --git a/Engine/Source/Editor/Sequencer/Public/LevelEditorSequencerIntegration.h b/Engine/Source/Editor/Sequencer/Public/LevelEditorSequencerIntegration.h
--- a/Engine/Source/Editor/Sequencer/Public/LevelEditorSequencerIntegration.h
+++ b/Engine/Source/Editor/Sequencer/Public/LevelEditorSequencerIntegration.h
@@ -105,6 +105,12 @@ private:
 	/** Called after any number of external actors has been saved. The sequencer puts everything back to its initial state. */
 	void OnPostSaveExternalActors(UWorld* World);
 
+	/** Called before asset validation is run on assets. The sequencer puts everything back to its initial state. */
+	void OnPreAssetValidation();
+	
+	/** Called after asset validation has finished. The sequencer re-evaluates to hide the fact we did this from users. */
+	void OnPostAssetValidation();
+
 	/** Called after a level has been added */
 	void OnLevelAdded(ULevel* InLevel, UWorld* InWorld);
 
diff --git a/Engine/Source/Editor/UnrealEd/Private/Editor.cpp b/Engine/Source/Editor/UnrealEd/Private/Editor.cpp
--- a/Engine/Source/Editor/UnrealEd/Private/Editor.cpp
+++ b/Engine/Source/Editor/UnrealEd/Private/Editor.cpp
@@ -118,6 +118,8 @@ FEditorDelegates::FOnPreSaveWorldWithContext			FEditorDelegates::PreSaveWorldWit
 FEditorDelegates::FOnPostSaveWorldWithContext			FEditorDelegates::PostSaveWorldWithContext;
 FEditorDelegates::FOnPreSaveExternalActors				FEditorDelegates::PreSaveExternalActors;
 FEditorDelegates::FOnPostSaveExternalActors				FEditorDelegates::PostSaveExternalActors;
+FSimpleMulticastDelegate								FEditorDelegates::OnPreAssetValidation;
+FSimpleMulticastDelegate								FEditorDelegates::OnPostAssetValidation;
 FEditorDelegates::FOnFinishPickingBlueprintClass		FEditorDelegates::OnFinishPickingBlueprintClass;
 FEditorDelegates::FOnNewAssetCreation					FEditorDelegates::OnConfigureNewAssetProperties;
 FEditorDelegates::FOnNewAssetCreation					FEditorDelegates::OnNewAssetCreated;
diff --git a/Engine/Source/Editor/UnrealEd/Public/Editor.h b/Engine/Source/Editor/UnrealEd/Public/Editor.h
--- a/Engine/Source/Editor/UnrealEd/Public/Editor.h
+++ b/Engine/Source/Editor/UnrealEd/Public/Editor.h
@@ -239,6 +239,10 @@ struct UNREALED_API FEditorDelegates
 	static FOnPreSaveExternalActors PreSaveExternalActors;
 	/** Called after any number of external actors has been saved */
 	static FOnPostSaveExternalActors PostSaveExternalActors;
+	/** Called before any asset validation happens via the Asset Validation subsystem. */
+	static FSimpleMulticastDelegate OnPreAssetValidation;
+	/** Called after asset validation happens by the Asset Validation subsystem. */
+	static FSimpleMulticastDelegate OnPostAssetValidation;
 	/** Called when finishing picking a new blueprint class during construction */
 	static FOnFinishPickingBlueprintClass OnFinishPickingBlueprintClass;
 	/** Called when beginning configuration of a new asset */
diff --git a/Engine/Source/Runtime/MovieScene/Private/MovieSceneCommonHelpers.cpp b/Engine/Source/Runtime/MovieScene/Private/MovieSceneCommonHelpers.cpp
--- a/Engine/Source/Runtime/MovieScene/Private/MovieSceneCommonHelpers.cpp
+++ b/Engine/Source/Runtime/MovieScene/Private/MovieSceneCommonHelpers.cpp
@@ -16,6 +16,7 @@
 #include "Sound/SoundNodeWavePlayer.h"
 #include "MovieSceneTrack.h"
 #include "Engine/Engine.h"
+#include "UObject/Package.h"
 
 bool MovieSceneHelpers::IsSectionKeyable(const UMovieSceneSection* Section)
 {
@@ -551,6 +552,24 @@ UObject* MovieSceneHelpers::MakeSpawnableTemplateFromInstance(UObject& InSourceO
 	return NewInstance;
 }
 
+
+MovieSceneHelpers::FMovieSceneScopedPackageDirtyGuard::FMovieSceneScopedPackageDirtyGuard(USceneComponent* InComponent)
+{
+	Component = InComponent;
+	if (Component && Component->GetPackage())
+	{
+		bPackageWasDirty = Component->GetPackage()->IsDirty();
+	}
+}
+
+MovieSceneHelpers::FMovieSceneScopedPackageDirtyGuard::~FMovieSceneScopedPackageDirtyGuard()
+{
+	if (Component && Component->GetPackage())
+	{
+		Component->GetPackage()->SetDirtyFlag(bPackageWasDirty);
+	}
+}
+
 FTrackInstancePropertyBindings::FTrackInstancePropertyBindings( FName InPropertyName, const FString& InPropertyPath )
 	: PropertyPath( InPropertyPath )
 	, PropertyName( InPropertyName )
@@ -948,4 +967,4 @@ template<> void FTrackInstancePropertyBindings::SetCurrentValue<UObject*>(UObjec
 	{
 		InRuntimeObject.ProcessEvent(NotifyFunction, nullptr);
 	}
-}
+}
\ No newline at end of file
diff --git a/Engine/Source/Runtime/MovieScene/Public/MovieSceneCommonHelpers.h b/Engine/Source/Runtime/MovieScene/Public/MovieSceneCommonHelpers.h
--- a/Engine/Source/Runtime/MovieScene/Public/MovieSceneCommonHelpers.h
+++ b/Engine/Source/Runtime/MovieScene/Public/MovieSceneCommonHelpers.h
@@ -25,6 +25,23 @@ class MOVIESCENE_API MovieSceneHelpers
 {
 public:
 
+	/*
+	* Helper struct to cache the package dirty state and then to restore it
+	* after this leaves scope. This is for a few minor areas where calling
+	* functions on actors dirties them, but Sequencer doesn't actually want
+	* the package to be dirty as it causes Sequencer to unnecessairly dirty
+	* actors.
+	*/
+	struct MOVIESCENE_API FMovieSceneScopedPackageDirtyGuard
+	{
+		FMovieSceneScopedPackageDirtyGuard(class USceneComponent* InComponent);
+		virtual ~FMovieSceneScopedPackageDirtyGuard();
+
+	private:
+		class USceneComponent* Component;
+		bool bPackageWasDirty;
+	};
+
 	/** 
 	 * @return Whether the section is keyable (active, on a track that is not muted, etc 
 	 */
diff --git a/Engine/Source/Runtime/MovieSceneTracks/Private/MovieSceneTracksComponentTypes.cpp b/Engine/Source/Runtime/MovieSceneTracks/Private/MovieSceneTracksComponentTypes.cpp
--- a/Engine/Source/Runtime/MovieSceneTracks/Private/MovieSceneTracksComponentTypes.cpp
+++ b/Engine/Source/Runtime/MovieSceneTracks/Private/MovieSceneTracksComponentTypes.cpp
@@ -347,8 +347,11 @@ void FComponentAttachParams::ApplyAttach(USceneComponent* ChildComponentToAttach
 {
 	if (ChildComponentToAttach->GetAttachParent() != NewAttachParent || ChildComponentToAttach->GetAttachSocketName() != SocketName)
 	{
-		FAttachmentTransformRules AttachmentRules(AttachmentLocationRule, AttachmentRotationRule, AttachmentScaleRule, false);
+		// Attachment changes may try to mark a package as dirty but this prevents us from restoring the level to the pre-animated
+		// state correctly which causes issues with validation.
+		MovieSceneHelpers::FMovieSceneScopedPackageDirtyGuard DirtyFlagGuard(ChildComponentToAttach);
 
+		FAttachmentTransformRules AttachmentRules(AttachmentLocationRule, AttachmentRotationRule, AttachmentScaleRule, false);
 		ChildComponentToAttach->AttachToComponent(NewAttachParent, AttachmentRules, SocketName);
 	}
 
@@ -365,11 +368,15 @@ void FComponentDetachParams::ApplyDetach(USceneComponent* ChildComponentToAttach
 	// Detach if there was no pre-existing parent
 	if (!NewAttachParent)
 	{
+		MovieSceneHelpers::FMovieSceneScopedPackageDirtyGuard DirtyFlagGuard(ChildComponentToAttach);
+
 		FDetachmentTransformRules DetachmentRules(DetachmentLocationRule, DetachmentRotationRule, DetachmentScaleRule, false);
 		ChildComponentToAttach->DetachFromComponent(DetachmentRules);
 	}
 	else
 	{
+		MovieSceneHelpers::FMovieSceneScopedPackageDirtyGuard DirtyFlagGuard(ChildComponentToAttach);
+
 		ChildComponentToAttach->AttachToComponent(NewAttachParent, FAttachmentTransformRules::KeepRelativeTransform, SocketName);
 	}
 }
