Sequencer: Use node path for group identifier otherwise tracks in folders won't get the correct path for groups.

diff --git a/Engine/Source/Editor/Sequencer/Private/Sequencer.cpp b/Engine/Source/Editor/Sequencer/Private/Sequencer.cpp
--- a/Engine/Source/Editor/Sequencer/Private/Sequencer.cpp
+++ b/Engine/Source/Editor/Sequencer/Private/Sequencer.cpp
@@ -6030,7 +6030,7 @@ void FSequencer::AddSelectedNodesToNewNodeGroup()
 
 		if (Groupable)
 		{
-			NodesToAdd.Add(Groupable->GetIdentifierForGrouping());
+			NodesToAdd.Add(IOutlinerExtension::GetPathName(Item));
 		}
 	}
 
@@ -6090,7 +6090,7 @@ void FSequencer::AddNodesToExistingNodeGroup(const TArray<TWeakPtr<UE::Sequencer
 
 		if (Groupable)
 		{
-			NodesToAdd.Add(Groupable->GetIdentifierForGrouping());
+			NodesToAdd.Add(IOutlinerExtension::GetPathName(Item));
 		}
 	}
 
