Sequencer: Fix track rows not being renameable

diff --git a/Engine/Source/Editor/Sequencer/Private/MVVM/ViewModels/TrackRowModel.cpp b/Engine/Source/Editor/Sequencer/Private/MVVM/ViewModels/TrackRowModel.cpp
--- a/Engine/Source/Editor/Sequencer/Private/MVVM/ViewModels/TrackRowModel.cpp
+++ b/Engine/Source/Editor/Sequencer/Private/MVVM/ViewModels/TrackRowModel.cpp
@@ -25,6 +25,7 @@
 #include "Tracks/MovieScene3DTransformTrack.h"
 #include "Tracks/MovieScenePrimitiveMaterialTrack.h"
 
+#include "ScopedTransaction.h"
 #include "Styling/AppStyle.h"
 
 #define LOCTEXT_NAMESPACE "TrackRowModel"
@@ -226,6 +227,41 @@ TSharedRef<SWidget> FTrackRowModel::CreateOutlinerView(const FCreateOutlinerView
 			InParams.TreeViewRow);
 }
 
+bool FTrackRowModel::CanRename() const
+{
+	UMovieSceneNameableTrack* NameableTrack = Cast<UMovieSceneNameableTrack>(GetTrack());
+	return NameableTrack && NameableTrack->CanRename();
+}
+
+void FTrackRowModel::Rename(const FText& NewName)
+{
+	UMovieSceneNameableTrack* NameableTrack = ::Cast<UMovieSceneNameableTrack>(GetTrack());
+
+	if (NameableTrack && !NameableTrack->GetTrackRowDisplayName(GetRowIndex()).EqualTo(NewName))
+	{
+		const FScopedTransaction Transaction(NSLOCTEXT("SequencerTrackRowNode", "RenameTrackRow", "Rename Track Row"));
+		NameableTrack->SetTrackRowDisplayName(NewName, GetRowIndex());
+
+		SetIdentifier(FName(*NewName.ToString()));
+
+		// HACK: this should not exist but is required to make renaming emitters work in niagara
+		if (TSharedPtr<FSequenceModel> OwnerModel = FindAncestorOfType<FSequenceModel>())
+		{
+			OwnerModel->GetSequencer()->NotifyMovieSceneDataChanged(EMovieSceneDataChangeType::TrackValueChanged);
+		}
+	}
+}
+
+bool FTrackRowModel::IsRenameValidImpl(const FText& NewName, FText& OutErrorMessage) const
+{
+	UMovieSceneNameableTrack* NameableTrack = ::Cast<UMovieSceneNameableTrack>(GetTrack());
+	if (NameableTrack)
+	{
+		return NameableTrack->ValidateDisplayName(NewName, OutErrorMessage);
+	}
+	return false;
+}
+
 bool FTrackRowModel::IsResizable() const
 {
 	UMovieSceneTrack* Track = GetTrack();
diff --git a/Engine/Source/Editor/Sequencer/Public/MVVM/ViewModels/TrackRowModel.h b/Engine/Source/Editor/Sequencer/Public/MVVM/ViewModels/TrackRowModel.h
--- a/Engine/Source/Editor/Sequencer/Public/MVVM/ViewModels/TrackRowModel.h
+++ b/Engine/Source/Editor/Sequencer/Public/MVVM/ViewModels/TrackRowModel.h
@@ -6,6 +6,7 @@
 #include "MVVM/ViewModels/OutlinerItemModel.h"
 #include "MVVM/Extensions/ITrackExtension.h"
 #include "MVVM/Extensions/ITrackAreaExtension.h"
+#include "MVVM/Extensions/IRenameableExtension.h"
 #include "MVVM/Extensions/IResizableExtension.h"
 #include "MVVM/Extensions/IDeletableExtension.h"
 
@@ -27,12 +28,13 @@ class FTrackRowModel
 	: public FOutlinerItemModel
 	, public ITrackAreaExtension
 	, public ITrackExtension
+	, public IRenameableExtension
 	, public IResizableExtension
 	, public IDeletableExtension
 {
 public:
 
-	UE_SEQUENCER_DECLARE_CASTABLE(FTrackRowModel, FOutlinerItemModel, ITrackAreaExtension, ITrackExtension, IDeletableExtension);
+	UE_SEQUENCER_DECLARE_CASTABLE(FTrackRowModel, FOutlinerItemModel, ITrackAreaExtension, ITrackExtension, IDeletableExtension, IRenameableExtension);
 
 	explicit FTrackRowModel(UMovieSceneTrack* InTrack, int32 InRowIndex);
 	~FTrackRowModel();
@@ -53,6 +55,11 @@ public:
 	FSlateFontInfo GetLabelFont() const override;
 	const FSlateBrush* GetIconBrush() const override;
 
+	/*~ IRenameableExtension */
+	bool CanRename() const override;
+	void Rename(const FText& NewName) override;
+	bool IsRenameValidImpl(const FText& NewName, FText& OutErrorMessage) const override;
+
 	/*~ IResizableExtension */
 	bool IsResizable() const override;
 	void Resize(float NewSize) override;
