Cooker Memory: Change FDeferredSignedObjectChangeHandler to not be active when in commandlet; it is only needed in editor. One of its effects is to keep the MovieSceneSignedObjects in memory, and this causes garbage collection to not be as impactfuly in the cooker, so we need it disabled in the cooker.
#rnx

diff --git a/Engine/Source/Editor/Sequencer/Private/SequencerModule.cpp b/Engine/Source/Editor/Sequencer/Private/SequencerModule.cpp
--- a/Engine/Source/Editor/Sequencer/Private/SequencerModule.cpp
+++ b/Engine/Source/Editor/Sequencer/Private/SequencerModule.cpp
@@ -440,7 +440,12 @@ public:
 				FCoreDelegates::OnPostEngineInit.AddRaw(this, &FSequencerModule::RegisterMenus);
 			}
 
-			UMovieSceneSignedObject::SetDeferredHandler(MakeUnique<FDeferredSignedObjectChangeHandler>());
+			// Set the deferred handler for use by Sequence editors. It is not necessary in commandlets.
+			// TODO: Create/Delete the handler when Sequence editors are opened/closed.
+			if (!IsRunningCommandlet())
+			{
+				UMovieSceneSignedObject::SetDeferredHandler(MakeUnique<FDeferredSignedObjectChangeHandler>());
+			}
 
 			FPropertyEditorModule& EditModule = FModuleManager::Get().GetModuleChecked<FPropertyEditorModule>("PropertyEditor");
 			OnGetGlobalRowExtensionHandle = EditModule.GetGlobalRowExtensionDelegate().AddStatic(&RegisterKeyframeExtensionHandler);
