Sequencer: Override ConstructEntitySystemLinker so that it is unique.

This fixes an issue where having the Day Sequence Editor and Level Sequence Editor both active at the same time would have shared the EntitySystemLinker. That led to a spawnable being removed from the registry with the id entry from the day sequence instead of the id entry from the level sequence. Subsequently, that spawnable wouldn't be spawned again after an autosave/restore state because it was still found in the registry.

diff --git a/Engine/Source/Editor/Sequencer/Private/Sequencer.cpp b/Engine/Source/Editor/Sequencer/Private/Sequencer.cpp
--- a/Engine/Source/Editor/Sequencer/Private/Sequencer.cpp
+++ b/Engine/Source/Editor/Sequencer/Private/Sequencer.cpp
@@ -4137,13 +4137,19 @@ void FSequencer::GetViewportSettings(TMap<FViewportClient*, EMovieSceneViewportP
 	}
 }
 
+UMovieSceneEntitySystemLinker* FSequencer::ConstructEntitySystemLinker()
+{
+	FName EntitySystemLinkerName = MakeUniqueObjectName(GetPlaybackContext(), UMovieSceneEntitySystemLinker::StaticClass(), TEXT("SequencerEntitySystemLinker"));
+	UMovieSceneEntitySystemLinker* EntitySystemLinker = NewObject<UMovieSceneEntitySystemLinker>(GetPlaybackContext(), EntitySystemLinkerName);
+	EntitySystemLinker->SetLinkerRole(UE::MovieScene::EEntitySystemLinkerRole::Standalone);
+	return EntitySystemLinker;
+}
 
 EMovieScenePlayerStatus::Type FSequencer::GetPlaybackStatus() const
 {
 	return PlaybackState;
 }
 
-
 void FSequencer::SetPlaybackStatus(EMovieScenePlayerStatus::Type InPlaybackStatus)
 {
 	PlaybackState = InPlaybackStatus;
diff --git a/Engine/Source/Editor/Sequencer/Private/Sequencer.h b/Engine/Source/Editor/Sequencer/Private/Sequencer.h
--- a/Engine/Source/Editor/Sequencer/Private/Sequencer.h
+++ b/Engine/Source/Editor/Sequencer/Private/Sequencer.h
@@ -815,6 +815,7 @@ public:
 	virtual void SetPlaybackStatus(EMovieScenePlayerStatus::Type InPlaybackStatus) override;
 	virtual FMovieSceneSpawnRegister& GetSpawnRegister() override { return *SpawnRegister; }
 	virtual bool IsPreview() const override { return SilentModeCount != 0; }
+	virtual UMovieSceneEntitySystemLinker* ConstructEntitySystemLinker() override;
 
 protected:
 
