Sequencer: Override ConstructEntitySystemLinker so that it is unique.

This fixes an issue where having the Day Sequence Editor and Level Sequence Editor both active at the same time would have shared the EntitySystemLinker. That led to a spawnable being removed from the registry with the id entry from the day sequence instead of the id entry from the level sequence. Subsequently, that spawnable wouldn't be spawned again after an autosave/restore state because it was still found in the registry.

diff --git a/Engine/Source/Editor/Sequencer/Private/Sequencer.cpp b/Engine/Source/Editor/Sequencer/Private/Sequencer.cpp
--- a/Engine/Source/Editor/Sequencer/Private/Sequencer.cpp
+++ b/Engine/Source/Editor/Sequencer/Private/Sequencer.cpp
@@ -4147,6 +4147,9 @@ UMovieSceneEntitySystemLinker* FSequencer::ConstructEntitySystemLinker()
 
 	FName EntitySystemLinkerName = MakeUniqueObjectName(Context, UMovieSceneEntitySystemLinker::StaticClass(), TEXT("SequencerEntitySystemLinker"));
 	UMovieSceneEntitySystemLinker* EntitySystemLinker = NewObject<UMovieSceneEntitySystemLinker>(Context, EntitySystemLinkerName);
+	EntitySystemLinker->SetLinkerRole(UE::MovieScene::EEntitySystemLinkerRole::Standalone);
+	return EntitySystemLinker;
+}
 
 EMovieScenePlayerStatus::Type FSequencer::GetPlaybackStatus() const
 {
