Sequencer: Fix paste bindings so that it fixes the parent bindings after the new bindings have been assigned

diff --git a/Engine/Source/Editor/Sequencer/Private/SequencerUtilities.cpp b/Engine/Source/Editor/Sequencer/Private/SequencerUtilities.cpp
--- a/Engine/Source/Editor/Sequencer/Private/SequencerUtilities.cpp
+++ b/Engine/Source/Editor/Sequencer/Private/SequencerUtilities.cpp
@@ -2206,16 +2206,6 @@ bool FSequencerUtilities::PasteBindings(const FString& TextToImport, TSharedRef<
 		}
 	}
 
-	// Fix up parent guids
-	for (auto PossessableGuid : PossessableGuids)
-	{
-		FMovieScenePossessable* Possessable = MovieScene->FindPossessable(PossessableGuid);
-		if (Possessable && OldToNewGuidMap.Contains(Possessable->GetParent()) && PossessableGuid != OldToNewGuidMap[Possessable->GetParent()])
-		{
-			Possessable->SetParent(OldToNewGuidMap[Possessable->GetParent()], MovieScene);
-		}
-	}
-
 	// Fix possessable actor bindings
 	for (int32 PossessableGuidIndex = 0; PossessableGuidIndex < PossessableGuids.Num(); ++PossessableGuidIndex)
 	{
@@ -2251,6 +2241,14 @@ bool FSequencerUtilities::PasteBindings(const FString& TextToImport, TSharedRef<
 								GuidToFolderMap.Remove(PossessableGuids[PossessableGuidIndex]);
 							}
 
+							for (TPair<FGuid, FGuid>& OldToNewGuid : OldToNewGuidMap)
+							{
+								if (OldToNewGuid.Value == PossessableGuids[PossessableGuidIndex])
+								{
+									OldToNewGuid.Value = NewGuid;
+								}
+							}
+
 							PossessableGuids[PossessableGuidIndex] = NewGuid;
 						}
 					}
@@ -2259,6 +2257,16 @@ bool FSequencerUtilities::PasteBindings(const FString& TextToImport, TSharedRef<
 		}
 	}
 	
+	// Fix up parent guids
+	for (auto PossessableGuid : PossessableGuids)
+	{
+		FMovieScenePossessable* Possessable = MovieScene->FindPossessable(PossessableGuid);
+		if (Possessable && OldToNewGuidMap.Contains(Possessable->GetParent()) && PossessableGuid != OldToNewGuidMap[Possessable->GetParent()])
+		{
+			Possessable->SetParent(OldToNewGuidMap[Possessable->GetParent()], MovieScene);
+		}
+	}
+
 	// Set up folders
 	for (auto PossessableGuid : PossessableGuids)
 	{
