Sequencer: Remove Master from Sequencer

diff --git a/Engine/Source/Editor/Sequencer/Private/Sequencer.cpp b/Engine/Source/Editor/Sequencer/Private/Sequencer.cpp
--- a/Engine/Source/Editor/Sequencer/Private/Sequencer.cpp
+++ b/Engine/Source/Editor/Sequencer/Private/Sequencer.cpp
@@ -2630,7 +2630,7 @@ UMovieSceneSection* FSequencer::FindNextOrPreviousShot(UMovieSceneSequence* Sequ
 {
 	UMovieScene* OwnerMovieScene = Sequence->GetMovieScene();
 
-	UMovieSceneTrack* CinematicShotTrack = OwnerMovieScene->FindMasterTrack(UMovieSceneCinematicShotTrack::StaticClass());
+	UMovieSceneTrack* CinematicShotTrack = OwnerMovieScene->FindTrack(UMovieSceneCinematicShotTrack::StaticClass());
 	if (!CinematicShotTrack)
 	{
 		return nullptr;
@@ -2728,7 +2728,7 @@ void FSequencer::SetPlaybackRangeToAllShots()
 	UMovieSceneSequence* Sequence = GetFocusedMovieSceneSequence();
 	UMovieScene* OwnerMovieScene = Sequence->GetMovieScene();
 
-	UMovieSceneTrack* CinematicShotTrack = OwnerMovieScene->FindMasterTrack(UMovieSceneCinematicShotTrack::StaticClass());
+	UMovieSceneTrack* CinematicShotTrack = OwnerMovieScene->FindTrack(UMovieSceneCinematicShotTrack::StaticClass());
 	if (!CinematicShotTrack || CinematicShotTrack->GetAllSections().Num() == 0)
 	{
 		return;
@@ -5428,7 +5428,7 @@ void FSequencer::AddSubSequence(UMovieSceneSequence* Sequence)
 	const FScopedTransaction Transaction( LOCTEXT("UndoAddingObject", "Add Object to MovieScene") );
 	OwnerMovieScene->Modify();
 
-	UMovieSceneSubTrack* SubTrack = OwnerMovieScene->AddMasterTrack<UMovieSceneSubTrack>();
+	UMovieSceneSubTrack* SubTrack = OwnerMovieScene->AddTrack<UMovieSceneSubTrack>();
 
 	FFrameNumber Duration = ConvertFrameTime(
 		Sequence->GetMovieScene()->GetPlaybackRange().Size<FFrameNumber>(),
@@ -7323,10 +7323,10 @@ void FSequencer::MoveNodeToFolder(TSharedRef<FViewModel> NodeToMove, UMovieScene
 		{
 			FFolderModel* NodeParentFolder = ParentNode->CastThisChecked<FFolderModel>();
 			NodeParentFolder->GetFolder()->Modify();
-			NodeParentFolder->GetFolder()->RemoveChildMasterTrack(TrackNode->GetTrack());
+			NodeParentFolder->GetFolder()->RemoveChildTrack(TrackNode->GetTrack());
 		}
 
-		DestinationFolder->AddChildMasterTrack(TrackNode->GetTrack());
+		DestinationFolder->AddChildTrack(TrackNode->GetTrack());
 	}
 	else if (IObjectBindingExtension* ObjectBindingNode = NodeToMove->CastThis<IObjectBindingExtension>())
 	{
@@ -7426,7 +7426,7 @@ TArray<TSharedRef<UE::Sequencer::FViewModel> > FSequencer::GetSelectedNodesInFol
 			{
 				if (TrackNode->GetTrack())
 				{
-					if (Folder->GetFolder()->GetChildMasterTracks().Contains(TrackNode->GetTrack()))
+					if (Folder->GetFolder()->GetChildTracks().Contains(TrackNode->GetTrack()))
 					{
 						NodesToFolders.Add(SelectedNode);
 					}
@@ -7823,7 +7823,7 @@ void FSequencer::RemoveSelectedNodesFromFolders()
 			{
 				if (TrackNode->GetTrack())
 				{
-					Folder->GetFolder()->RemoveChildMasterTrack(TrackNode->GetTrack());
+					Folder->GetFolder()->RemoveChildTrack(TrackNode->GetTrack());
 				}
 			}
 		}
@@ -8825,7 +8825,7 @@ void FSequencer::OnAddTrack(const TWeakObjectPtr<UMovieSceneTrack>& InTrack, con
 		if (SelectedParentFolders.Num() == 1)
 		{
 			SelectedParentFolders[0]->Modify();
-			SelectedParentFolders[0]->AddChildMasterTrack(InTrack.Get());
+			SelectedParentFolders[0]->AddChildTrack(InTrack.Get());
 		}
 	}
 
@@ -9489,11 +9489,11 @@ TSet<FFrameNumber> FSequencer::GetVerticalFrames() const
 		UMovieScene* FocusedMovieScene = FocusedMovieSequence->GetMovieScene();
 		if (FocusedMovieScene != nullptr)
 		{
-			for (UMovieSceneTrack* MasterTrack : FocusedMovieScene->GetMasterTracks())
+			for (UMovieSceneTrack* Track : FocusedMovieScene->GetTracks())
 			{
-				if (MasterTrack && MasterTrack->DisplayOptions.bShowVerticalFrames)
+				if (Track && Track->DisplayOptions.bShowVerticalFrames)
 				{
-					AddVerticalFrames(VerticalFrames, MasterTrack);
+					AddVerticalFrames(VerticalFrames, Track);
 				}
 			}
 
@@ -10214,7 +10214,7 @@ void FSequencer::ExportFBX()
 		
 		// Select selected nodes if there are selected nodes
 		TArray<FGuid> Bindings;
-		TArray<UMovieSceneTrack*> MasterTracks;
+		TArray<UMovieSceneTrack*> Tracks;
 		UMovieScene* MovieScene = GetFocusedMovieSceneSequence()->GetMovieScene();
 		for (const TWeakPtr<FViewModel>& Node : Selection.GetSelectedOutlinerItems())
 		{
@@ -10236,9 +10236,9 @@ void FSequencer::ExportFBX()
 			else if (ITrackExtension* TrackNode = ICastable::CastWeakPtr<ITrackExtension>(Node))
 			{
 				UMovieSceneTrack* Track = TrackNode->GetTrack();
-				if (Track && MovieScene->IsAMasterTrack(*Track))
+				if (Track && MovieScene->ContainsTrack(*Track))
 				{
-					MasterTracks.Add(Track);
+					Tracks.Add(Track);
 				}
 			}
 		}
@@ -10246,7 +10246,7 @@ void FSequencer::ExportFBX()
 		FString FileExtension = FPaths::GetExtension(ExportFilename);
 		if (FileExtension == TEXT("fbx"))
 		{
-			ExportFBXInternal(ExportFilename, Bindings, (Bindings.Num() + MasterTracks.Num()) > 0 ? MasterTracks : MovieScene->GetMasterTracks());
+			ExportFBXInternal(ExportFilename, Bindings, (Bindings.Num() + Tracks.Num()) > 0 ? Tracks : MovieScene->GetTracks());
 		}
 		else
 		{
@@ -10283,7 +10283,7 @@ void FSequencer::ExportFBX()
 }
 
 
-void FSequencer::ExportFBXInternal(const FString& ExportFilename, const TArray<FGuid>& Bindings, const TArray<UMovieSceneTrack*>& MasterTracks)
+void FSequencer::ExportFBXInternal(const FString& ExportFilename, const TArray<FGuid>& Bindings, const TArray<UMovieSceneTrack*>& Tracks)
 {
 	{
 		UnFbx::FFbxExporter* Exporter = UnFbx::FFbxExporter::GetInstance();
@@ -10306,7 +10306,7 @@ void FSequencer::ExportFBXInternal(const FString& ExportFilename, const TArray<F
 					SetLocalTimeDirectly(UE::MovieScene::DiscreteInclusiveLower(GetTimeBounds()));
 				}
 
-				if (MovieSceneToolHelpers::ExportFBX(World, MovieScene, this, Bindings, MasterTracks, NodeNameAdapter, Template, ExportFilename, RootToLocalTransform))
+				if (MovieSceneToolHelpers::ExportFBX(World, MovieScene, this, Bindings, Tracks, NodeNameAdapter, Template, ExportFilename, RootToLocalTransform))
 				{
 					FNotificationInfo Info(NSLOCTEXT("Sequencer", "ExportFBXSucceeded", "FBX Export Succeeded."));
 					Info.Hyperlink = FSimpleDelegate::CreateStatic([](FString InFilename) { FPlatformProcess::ExploreFolder(*InFilename); }, ExportFilename);
@@ -10409,7 +10409,7 @@ void FSequencer::TrimOrExtendSection(bool bTrimOrExtendLeft)
 	}
 	else
 	{
-		for (UMovieSceneTrack* Track : MovieScene->GetMasterTracks())
+		for (UMovieSceneTrack* Track : MovieScene->GetTracks())
 		{
 			MovieSceneToolHelpers::TrimOrExtendSection(Track, TOptional<int32>(), GetLocalTime(), bTrimOrExtendLeft, Settings->GetDeleteKeysWhenTrimming());
 		}
@@ -11319,17 +11319,17 @@ void FSequencer::BindCommands()
 	SequencerCommandBindings->MapAction(
 		Commands.SetSelectionRangeToNextShot,
 		FExecuteAction::CreateSP( this, &FSequencer::SetSelectionRangeToShot, true ),
-		FCanExecuteAction::CreateSP( this, &FSequencer::IsViewingMasterSequence ) );
+		FCanExecuteAction::CreateSP( this, &FSequencer::IsViewingRootSequence ) );
 
 	SequencerCommandBindings->MapAction(
 		Commands.SetSelectionRangeToPreviousShot,
 		FExecuteAction::CreateSP( this, &FSequencer::SetSelectionRangeToShot, false ),
-		FCanExecuteAction::CreateSP( this, &FSequencer::IsViewingMasterSequence ) );
+		FCanExecuteAction::CreateSP( this, &FSequencer::IsViewingRootSequence ) );
 
 	SequencerCommandBindings->MapAction(
 		Commands.SetPlaybackRangeToAllShots,
 		FExecuteAction::CreateSP( this, &FSequencer::SetPlaybackRangeToAllShots ),
-		FCanExecuteAction::CreateSP( this, &FSequencer::IsViewingMasterSequence ) );
+		FCanExecuteAction::CreateSP( this, &FSequencer::IsViewingRootSequence ) );
 
 	SequencerCommandBindings->MapAction(
 		Commands.RefreshUI,
diff --git a/Engine/Source/Editor/Sequencer/Private/Sequencer.h b/Engine/Source/Editor/Sequencer/Private/Sequencer.h
--- a/Engine/Source/Editor/Sequencer/Private/Sequencer.h
+++ b/Engine/Source/Editor/Sequencer/Private/Sequencer.h
@@ -910,7 +910,7 @@ protected:
 	void OnScrubPositionParentChanged(FMovieSceneSequenceID InScrubPositionParent);
 
 	/** Exports sequence to a FBX file */
-	void ExportFBXInternal(const FString& Filename, const TArray<FGuid>& Bindings, const TArray<UMovieSceneTrack*>& MasterTracks);
+	void ExportFBXInternal(const FString& Filename, const TArray<FGuid>& Bindings, const TArray<UMovieSceneTrack*>& Tracks);
 
 protected:
 
@@ -1097,8 +1097,8 @@ private:
 	/** Get actors that want to rerun construction scripts */
 	void GetConstructionScriptActors(UMovieScene*, FMovieSceneSequenceIDRef SequenceID, TSet<TWeakObjectPtr<AActor> >& BoundActors, TArray < TPair<FMovieSceneSequenceID, FGuid> >& BoundGuids);
 
-	/** Check whether we're viewing the master sequence or not */
-	bool IsViewingMasterSequence() const { return ActiveTemplateIDs.Num() == 1; }
+	/** Check whether we're viewing the root sequence or not */
+	bool IsViewingRootSequence() const { return ActiveTemplateIDs.Num() == 1; }
 
 	/** Recompile any dirty director blueprints in the sequence hierarchy */
 	void RecompileDirtyDirectors();
diff --git a/Engine/Source/Editor/Sequencer/Private/SequencerCommands.cpp b/Engine/Source/Editor/Sequencer/Private/SequencerCommands.cpp
--- a/Engine/Source/Editor/Sequencer/Private/SequencerCommands.cpp
+++ b/Engine/Source/Editor/Sequencer/Private/SequencerCommands.cpp
@@ -171,7 +171,7 @@ void FSequencerCommands::RegisterCommands()
 	UI_COMMAND( ImportFBX, "Import...", "Import the animation from an FBX file.", EUserInterfaceActionType::Button, FInputChord() );
 	UI_COMMAND( ExportFBX, "Export...", "Export the selected objects (all if none selected) and animation to an FBX file. (Shots and sub-scenes not supported)", EUserInterfaceActionType::Button, FInputChord() );
 
-	UI_COMMAND( ToggleEvaluateSubSequencesInIsolation, "Evaluate Sub Sequences In Isolation", "When enabled, will only evaluate the currently focused sequence; otherwise evaluate from the master sequence.", EUserInterfaceActionType::ToggleButton, FInputChord() );
+	UI_COMMAND( ToggleEvaluateSubSequencesInIsolation, "Evaluate Sub Sequences In Isolation", "When enabled, will only evaluate the currently focused sequence; otherwise evaluate from the root sequence.", EUserInterfaceActionType::ToggleButton, FInputChord() );
 
 	UI_COMMAND( QuickTreeSearch, "Quick Tree Search", "Jumps keyboard focus to the tree searchbox to allow searching for tracks in the current Sequence.", EUserInterfaceActionType::Button, FInputChord(EModifierKey::Control, EKeys::F));
 	
diff --git a/Engine/Source/Editor/Sequencer/Private/SequencerEdMode.cpp b/Engine/Source/Editor/Sequencer/Private/SequencerEdMode.cpp
--- a/Engine/Source/Editor/Sequencer/Private/SequencerEdMode.cpp
+++ b/Engine/Source/Editor/Sequencer/Private/SequencerEdMode.cpp
@@ -935,7 +935,7 @@ void FSequencerEdMode::DrawAudioTracks(FPrimitiveDrawInterface* PDI)
 		{
 			UMovieSceneAudioTrack* AudioTrack = Cast<UMovieSceneAudioTrack>(Track);
 
-			if (!AudioTrack || !AudioTrack->IsAMasterTrack())
+			if (!AudioTrack)
 			{
 				continue;
 			}
diff --git a/Engine/Source/Editor/Sequencer/Private/SequencerUtilities.cpp b/Engine/Source/Editor/Sequencer/Private/SequencerUtilities.cpp
--- a/Engine/Source/Editor/Sequencer/Private/SequencerUtilities.cpp
+++ b/Engine/Source/Editor/Sequencer/Private/SequencerUtilities.cpp
@@ -1334,7 +1334,7 @@ bool FSequencerUtilities::PasteFolders(const FString& TextToImport, FMovieSceneP
 		OutFolders.Add(CopiedFolder);
 
 		// Clear the folder contents, those relationships will be made when the tracks are pasted
-		CopiedFolder->ClearChildMasterTracks();
+		CopiedFolder->ClearChildTracks();
 		CopiedFolder->ClearChildObjectBindings();
 
 		bool bHasParent = false;
@@ -1390,8 +1390,8 @@ void FSequencerUtilities::CopyTracks(const TArray<UMovieSceneTrack*>& Tracks, co
 
 		UMovieSceneTrack* DuplicatedTrack = Cast<UMovieSceneTrack>(StaticDuplicateObject(Track, CopyableTrack));
 		CopyableTrack->Track = DuplicatedTrack;
-		CopyableTrack->bIsAMasterTrack = MovieScene->IsAMasterTrack(*Track);
-		CopyableTrack->bIsACameraCutTrack = Track->IsA<UMovieSceneCameraCutTrack>();
+		CopyableTrack->bIsRootTrack = MovieScene->ContainsTrack(*Track);
+		CopyableTrack->bIsCameraCutTrack = Track->IsA<UMovieSceneCameraCutTrack>();
 
 		UMovieSceneFolder* Folder = nullptr;
 		for (UMovieSceneFolder* RootFolder : MovieScene->GetRootFolders())
@@ -1465,14 +1465,14 @@ bool FSequencerUtilities::PasteTracks(const FString& TextToImport, FMovieScenePa
 
 	const FScopedTransaction Transaction(LOCTEXT("PasteTracks", "Paste Tracks"));
 
-	int32 NumMasterOrCameraCutTracks = 0;
+	int32 NumRootOrCameraCutTracks = 0;
 	int32 NumTracks = 0;
 
 	for (UMovieSceneCopyableTrack* CopyableTrack : ImportedTracks)
 	{
-		if (CopyableTrack->bIsAMasterTrack || CopyableTrack->bIsACameraCutTrack)
+		if (CopyableTrack->bIsRootTrack || CopyableTrack->bIsCameraCutTrack)
 		{
-			++NumMasterOrCameraCutTracks;
+			++NumRootOrCameraCutTracks;
 		}
 		else
 		{
@@ -1481,7 +1481,7 @@ bool FSequencerUtilities::PasteTracks(const FString& TextToImport, FMovieScenePa
 	}
 
 	int32 NumTracksPasted = 0;
-	int32 NumMasterOrCameraCutTracksPasted = 0;
+	int32 NumRootOrCameraCutTracksPasted = 0;
 
 	for (const FMovieSceneBindingProxy& ObjectBinding : PasteTracksParams.Bindings)
 	{
@@ -1496,7 +1496,7 @@ bool FSequencerUtilities::PasteTracks(const FString& TextToImport, FMovieScenePa
 
 		for (UMovieSceneCopyableTrack* CopyableTrack : NewTracks)
 		{
-			if (!CopyableTrack->bIsAMasterTrack && !CopyableTrack->bIsACameraCutTrack)
+			if (!CopyableTrack->bIsRootTrack && !CopyableTrack->bIsCameraCutTrack)
 			{
 				UMovieSceneTrack* NewTrack = CopyableTrack->Track;
 				NewTrack->ClearFlags(RF_Transient);
@@ -1539,10 +1539,10 @@ bool FSequencerUtilities::PasteTracks(const FString& TextToImport, FMovieScenePa
 	UMovieScene* MovieScene = PasteTracksParams.Sequence ? PasteTracksParams.Sequence->GetMovieScene() : nullptr;
 	if (MovieScene)
 	{
-		// Add as master track or set camera cut track
+		// Add as root track or set camera cut track
 		for (UMovieSceneCopyableTrack* CopyableTrack : ImportedTracks)
 		{
-			if (CopyableTrack->bIsAMasterTrack || CopyableTrack->bIsACameraCutTrack)
+			if (CopyableTrack->bIsRootTrack || CopyableTrack->bIsCameraCutTrack)
 			{
 				UMovieSceneTrack* NewTrack = CopyableTrack->Track;
 				NewTrack->ClearFlags(RF_Transient);
@@ -1563,22 +1563,22 @@ bool FSequencerUtilities::PasteTracks(const FString& TextToImport, FMovieScenePa
 					MovieScene->SetCameraCutTrack(NewTrack);
 					if (ParentFolder != nullptr)
 					{
-						ParentFolder->AddChildMasterTrack(NewTrack);
+						ParentFolder->AddChildTrack(NewTrack);
 					}
 
-					++NumMasterOrCameraCutTracksPasted;
+					++NumRootOrCameraCutTracksPasted;
 				}
 				else
 				{
-					if (MovieScene->AddGivenMasterTrack(NewTrack))
+					if (MovieScene->AddGivenTrack(NewTrack))
 					{
 						if (ParentFolder != nullptr)
 						{
-							ParentFolder->AddChildMasterTrack(NewTrack);
+							ParentFolder->AddChildTrack(NewTrack);
 						}
 					}
 
-					++NumMasterOrCameraCutTracksPasted;
+					++NumRootOrCameraCutTracksPasted;
 				}
 
 				OutTracks.Add(NewTrack);
@@ -1586,9 +1586,9 @@ bool FSequencerUtilities::PasteTracks(const FString& TextToImport, FMovieScenePa
 		}
 	}
 	
-	if (NumMasterOrCameraCutTracksPasted < NumMasterOrCameraCutTracks)
+	if (NumRootOrCameraCutTracksPasted < NumRootOrCameraCutTracks)
 	{
-		FNotificationInfo Info(LOCTEXT("PasteTracks_NoMasterTracks", "Can't paste track. Master track could not be pasted"));
+		FNotificationInfo Info(LOCTEXT("PasteTracks_NoTracks", "Can't paste track. Root track could not be pasted"));
 		OutErrors.Add(Info);
 	}
 
@@ -1598,7 +1598,7 @@ bool FSequencerUtilities::PasteTracks(const FString& TextToImport, FMovieScenePa
 		OutErrors.Add(Info);
 	}
 
-	return (NumMasterOrCameraCutTracksPasted + NumTracksPasted) > 0;
+	return (NumRootOrCameraCutTracksPasted + NumTracksPasted) > 0;
 }
 
 bool FSequencerUtilities::CanPasteTracks(const FString& TextToImport)
diff --git a/Engine/Source/Editor/Sequencer/Public/MVVM/TrackModelStorageExtension.h b/Engine/Source/Editor/Sequencer/Public/MVVM/TrackModelStorageExtension.h
--- a/Engine/Source/Editor/Sequencer/Public/MVVM/TrackModelStorageExtension.h
+++ b/Engine/Source/Editor/Sequencer/Public/MVVM/TrackModelStorageExtension.h
@@ -37,8 +37,8 @@ public:
 
 private:
 
-	void OnMasterTrackAdded(UMovieSceneTrack*) override;
-	void OnMasterTrackRemoved(UMovieSceneTrack*) override;
+	void OnTrackAdded(UMovieSceneTrack*) override;
+	void OnTrackRemoved(UMovieSceneTrack*) override;
 
 private:
 
diff --git a/Engine/Source/Editor/Sequencer/Public/MovieSceneTrackEditor.h b/Engine/Source/Editor/Sequencer/Public/MovieSceneTrackEditor.h
--- a/Engine/Source/Editor/Sequencer/Public/MovieSceneTrackEditor.h
+++ b/Engine/Source/Editor/Sequencer/Public/MovieSceneTrackEditor.h
@@ -129,7 +129,42 @@ public:
 	FFindOrCreateTrackResult FindOrCreateTrackForObject( const FGuid& ObjectHandle, TSubclassOf<UMovieSceneTrack> TrackClass, FName PropertyName = NAME_None, bool bCreateTrackIfMissing = true );
 
 	template<typename TrackClass>
-	struct FFindOrCreateMasterTrackResult
+	struct FFindOrCreateRootTrackResult
+	{
+		FFindOrCreateRootTrackResult() : Track(nullptr), bWasCreated(false) {}
+
+		TrackClass* Track;
+		bool bWasCreated;
+	};
+
+	/**
+	 * Find or add a track of the specified type in the focused movie scene.
+	 *
+	 * @param TrackClass The class of the track to find or add.
+	 * @return The track results.
+	 */
+	template<typename TrackClass>
+	FFindOrCreateRootTrackResult<TrackClass> FindOrCreateRootTrack()
+	{
+		FFindOrCreateRootTrackResult<TrackClass> Result;
+		bool bTrackExisted;
+
+		UMovieScene* MovieScene = GetSequencer()->GetFocusedMovieSceneSequence()->GetMovieScene();
+		Result.Track = MovieScene->FindTrack<TrackClass>();
+		bTrackExisted = Result.Track != nullptr;
+
+		if (Result.Track == nullptr)
+		{
+			Result.Track = MovieScene->AddTrack<TrackClass>();
+		}
+
+		Result.bWasCreated = bTrackExisted == false && Result.Track != nullptr;
+		return Result;
+	}
+
+	template<typename TrackClass> struct
+	UE_DEPRECATED(5.2, "FFindOrCreateMasterTrackResult is deprecated. Please use FFindOrCreateRootTrackResult instead")
+	FFindOrCreateMasterTrackResult
 	{
 		FFindOrCreateMasterTrackResult() : Track(nullptr), bWasCreated(false) {}
 
@@ -138,30 +173,32 @@ public:
 	};
 
 	/**
-	 * Find or add a master track of the specified type in the focused movie scene.
+	 * Find or add a track of the specified type in the focused movie scene.
 	 *
 	 * @param TrackClass The class of the track to find or add.
 	 * @return The track results.
 	 */
+	PRAGMA_DISABLE_DEPRECATION_WARNINGS // Suppress compiler warning on return of deprecated function
 	template<typename TrackClass>
+	UE_DEPRECATED(5.2, "FindOrCreateMasterTrack is deprecated. Please use FindOrCreateRootTrack instead")
 	FFindOrCreateMasterTrackResult<TrackClass> FindOrCreateMasterTrack()
 	{
 		FFindOrCreateMasterTrackResult<TrackClass> Result;
 		bool bTrackExisted;
 
 		UMovieScene* MovieScene = GetSequencer()->GetFocusedMovieSceneSequence()->GetMovieScene();
-		Result.Track = MovieScene->FindMasterTrack<TrackClass>();
+		Result.Track = MovieScene->FindTrack<TrackClass>();
 		bTrackExisted = Result.Track != nullptr;
 
 		if (Result.Track == nullptr)
 		{
-			Result.Track = MovieScene->AddMasterTrack<TrackClass>();
+			Result.Track = MovieScene->AddTrack<TrackClass>();
 		}
 
 		Result.bWasCreated = bTrackExisted == false && Result.Track != nullptr;
 		return Result;
 	}
-
+	PRAGMA_ENABLE_DEPRECATION_WARNINGS
 
 	/** @return The sequencer bound to this handler */
 	const TSharedPtr<ISequencer> GetSequencer() const;
diff --git a/Engine/Source/Editor/Sequencer/Public/SequencerSettings.h b/Engine/Source/Editor/Sequencer/Public/SequencerSettings.h
--- a/Engine/Source/Editor/Sequencer/Public/SequencerSettings.h
+++ b/Engine/Source/Editor/Sequencer/Public/SequencerSettings.h
@@ -602,7 +602,7 @@ protected:
 	UPROPERTY(config, EditAnywhere, Category = General)
 	bool bActivateRealtimeViewports;
 
-	/** When enabled, entering a sub sequence will evaluate that sub sequence in isolation, rather than from the master sequence */
+	/** When enabled, entering a sub sequence will evaluate that sub sequence in isolation, rather than from the root sequence */
 	UPROPERTY(config, EditAnywhere, Category=Playback)
 	bool bEvaluateSubSequencesInIsolation;
 
diff --git a/Engine/Source/Runtime/LevelSequence/Private/LevelSequenceDirector.cpp b/Engine/Source/Runtime/LevelSequence/Private/LevelSequenceDirector.cpp
--- a/Engine/Source/Runtime/LevelSequence/Private/LevelSequenceDirector.cpp
+++ b/Engine/Source/Runtime/LevelSequence/Private/LevelSequenceDirector.cpp
@@ -22,7 +22,7 @@ UWorld* ULevelSequenceDirector::GetWorld() const
 	return GetTypedOuter<UWorld>();
 }
 
-FQualifiedFrameTime ULevelSequenceDirector::GetMasterSequenceTime() const
+FQualifiedFrameTime ULevelSequenceDirector::GetRootSequenceTime() const
 {
 	using namespace UE::MovieScene;
 
diff --git a/Engine/Source/Runtime/LevelSequence/Private/LevelSequencePlayer.cpp b/Engine/Source/Runtime/LevelSequence/Private/LevelSequencePlayer.cpp
--- a/Engine/Source/Runtime/LevelSequence/Private/LevelSequencePlayer.cpp
+++ b/Engine/Source/Runtime/LevelSequence/Private/LevelSequencePlayer.cpp
@@ -463,10 +463,10 @@ void ULevelSequencePlayer::TakeFrameSnapshot(FLevelSequencePlayerSnapshot& OutSn
 	// In Playback Resolution
 	const FFrameTime CurrentSequenceTime		  = ConvertFrameTime(CurrentPlayTime, PlayPosition.GetInputRate(), PlayPosition.GetOutputRate());
 
-	OutSnapshot.MasterTime = FQualifiedFrameTime(CurrentPlayTime, PlayPosition.GetInputRate());
-	OutSnapshot.MasterName = Sequence->GetName();
+	OutSnapshot.RootTime = FQualifiedFrameTime(CurrentPlayTime, PlayPosition.GetInputRate());
+	OutSnapshot.RootName = Sequence->GetName();
 
-	OutSnapshot.CurrentShotName = OutSnapshot.MasterName;
+	OutSnapshot.CurrentShotName = OutSnapshot.RootName;
 	OutSnapshot.CurrentShotLocalTime = FQualifiedFrameTime(CurrentPlayTime, PlayPosition.GetInputRate());
 	OutSnapshot.CameraComponent = CachedCameraComponent.IsValid() ? CachedCameraComponent.Get() : nullptr;
 	OutSnapshot.ShotID = MovieSceneSequenceID::Invalid;
@@ -479,7 +479,7 @@ void ULevelSequencePlayer::TakeFrameSnapshot(FLevelSequencePlayerSnapshot& OutSn
 	OutSnapshot.SourceTimecode = MovieScene->GetEarliestTimecodeSource().Timecode.ToString();
 #endif
 
-	UMovieSceneCinematicShotTrack* ShotTrack = MovieScene->FindMasterTrack<UMovieSceneCinematicShotTrack>();
+	UMovieSceneCinematicShotTrack* ShotTrack = MovieScene->FindTrack<UMovieSceneCinematicShotTrack>();
 	if (ShotTrack)
 	{
 		UMovieSceneCinematicShotSection* ActiveShot = nullptr;
diff --git a/Engine/Source/Runtime/LevelSequence/Public/LevelSequenceDirector.h b/Engine/Source/Runtime/LevelSequence/Public/LevelSequenceDirector.h
--- a/Engine/Source/Runtime/LevelSequence/Public/LevelSequenceDirector.h
+++ b/Engine/Source/Runtime/LevelSequence/Public/LevelSequenceDirector.h
@@ -24,14 +24,18 @@ public:
 	void OnCreated();
 	
 	/**
-	 * Get the current time for the outermost (root, or master) sequence
-	 * @return The current playback position of the master sequence
+	 * Get the current time for the outermost (root) sequence
+	 * @return The current playback position of the outermost (root) sequence
 	 */
 	UFUNCTION(BlueprintCallable, Category = "Sequencer|Director")
-	FQualifiedFrameTime GetMasterSequenceTime() const;
+	FQualifiedFrameTime GetRootSequenceTime() const;
+
+	UE_DEPRECATED(5.2, "GetMasterSequenceTime is deprecated. Please use GetRootSequenceTime instead")
+	UFUNCTION(BlueprintCallable, Category = "Sequencer|Director", meta = (DeprecatedFunction, DeprecationMessage = "Use GetRootSequenceTime"))
+	FQualifiedFrameTime GetMasterSequenceTime() const { return GetRootSequenceTime(); }
 
 	/**
-	 * Get the current time for this director's sub-sequence (or the master sequence, if this is a master sequence director)
+	 * Get the current time for this director's sub-sequence (or the root sequence, if this is a root sequence director)
 	 * @return The current playback position of this director's sequence
 	 */
 	UFUNCTION(BlueprintCallable, Category = "Sequencer|Director")
diff --git a/Engine/Source/Runtime/LevelSequence/Public/LevelSequencePlayer.h b/Engine/Source/Runtime/LevelSequence/Public/LevelSequencePlayer.h
--- a/Engine/Source/Runtime/LevelSequence/Public/LevelSequencePlayer.h
+++ b/Engine/Source/Runtime/LevelSequence/Public/LevelSequencePlayer.h
@@ -28,7 +28,7 @@ class UCameraComponent;
 DECLARE_DYNAMIC_MULTICAST_DELEGATE_OneParam(FOnLevelSequencePlayerCameraCutEvent, UCameraComponent*, CameraComponent);
 
 USTRUCT(BlueprintType)
-struct UE_DEPRECATED(5.0, "Snapshot settings are deprecated. Use the frame rate from the FQualifiedFrameTime on the MasterTime and the ShotTime") FLevelSequenceSnapshotSettings
+struct UE_DEPRECATED(5.0, "Snapshot settings are deprecated. Use the frame rate from the FQualifiedFrameTime on the RootTime and the ShotTime") FLevelSequenceSnapshotSettings
 { GENERATED_BODY() };
 
 /**
@@ -40,10 +40,10 @@ struct FLevelSequencePlayerSnapshot
 	GENERATED_BODY()
 
 	UPROPERTY(VisibleAnywhere, BlueprintReadOnly, Category="General")
-	FString MasterName;
+	FString RootName;
 
 	UPROPERTY(VisibleAnywhere, BlueprintReadOnly, Category="General")
-	FQualifiedFrameTime MasterTime;
+	FQualifiedFrameTime RootTime;
 
 	UPROPERTY(VisibleAnywhere, BlueprintReadOnly, Category="General")
 	FQualifiedFrameTime SourceTime;
@@ -68,6 +68,14 @@ struct FLevelSequencePlayerSnapshot
 
 	UPROPERTY()
 	FMovieSceneSequenceID ShotID;
+
+private:
+
+	UPROPERTY()
+	FString MasterName_DEPRECATED;
+
+	UPROPERTY()
+	FQualifiedFrameTime MasterTime_DEPRECATED;
 };
 
 /**
diff --git a/Engine/Source/Runtime/MovieScene/Private/Compilation/MovieSceneCompiledDataManager.cpp b/Engine/Source/Runtime/MovieScene/Private/Compilation/MovieSceneCompiledDataManager.cpp
--- a/Engine/Source/Runtime/MovieScene/Private/Compilation/MovieSceneCompiledDataManager.cpp
+++ b/Engine/Source/Runtime/MovieScene/Private/Compilation/MovieSceneCompiledDataManager.cpp
@@ -729,7 +729,7 @@ void UMovieSceneCompiledDataManager::Compile(FMovieSceneCompiledDataID DataID, U
 				CompileTrack(&Entry, nullptr, Track, Params, &GatheredSignatures, &GatheredData);
 			}
 
-			for (UMovieSceneTrack* Track : MovieScene->GetMasterTracks())
+			for (UMovieSceneTrack* Track : MovieScene->GetTracks())
 			{
 				CompileTrack(&Entry, nullptr, Track, Params, &GatheredSignatures, &GatheredData);
 			}
@@ -850,7 +850,7 @@ void UMovieSceneCompiledDataManager::Gather(const FMovieSceneCompiledDataEntry&
 			GatherTrack(nullptr, Track, Params, TrackTemplate, OutCompilerData);
 		}
 
-		for (UMovieSceneTrack* Track : MovieScene->GetMasterTracks())
+		for (UMovieSceneTrack* Track : MovieScene->GetTracks())
 		{
 			GatherTrack(nullptr, Track, Params, TrackTemplate, OutCompilerData);
 		}
@@ -1391,7 +1391,7 @@ bool UMovieSceneCompiledDataManager::GenerateSubSequenceData(UMovieSceneSequence
 
 	bool bContainsSubSequences = false;
 
-	for (UMovieSceneTrack* Track : MovieScene->GetMasterTracks())
+	for (UMovieSceneTrack* Track : MovieScene->GetTracks())
 	{
 		if (UMovieSceneSubTrack* SubTrack = Cast<UMovieSceneSubTrack>(Track))
 		{
@@ -1495,7 +1495,7 @@ void UMovieSceneCompiledDataManager::PopulateSubSequenceTree(UMovieSceneSequence
 
 	check(RootPath && InOutHierarchy);
 
-	for (UMovieSceneTrack* Track : MovieScene->GetMasterTracks())
+	for (UMovieSceneTrack* Track : MovieScene->GetTracks())
 	{
 		if (UMovieSceneSubTrack* SubTrack = Cast<UMovieSceneSubTrack>(Track))
 		{
diff --git a/Engine/Source/Runtime/MovieScene/Private/EntitySystem/BuiltInComponentTypes.cpp b/Engine/Source/Runtime/MovieScene/Private/EntitySystem/BuiltInComponentTypes.cpp
--- a/Engine/Source/Runtime/MovieScene/Private/EntitySystem/BuiltInComponentTypes.cpp
+++ b/Engine/Source/Runtime/MovieScene/Private/EntitySystem/BuiltInComponentTypes.cpp
@@ -144,7 +144,7 @@ FBuiltInComponentTypes::FBuiltInComponentTypes()
 	Tags.NeedsUnlink             = ComponentRegistry->NewTag(TEXT("Needs Unlink"));
 	Tags.HasUnresolvedBinding    = ComponentRegistry->NewTag(TEXT("Has Unresolved Binding"));
 	Tags.MigratedFromFastPath    = ComponentRegistry->NewTag(TEXT("Migrated From Fast Path"));
-	Tags.Master                  = ComponentRegistry->NewTag(TEXT("Master"));
+	Tags.Root                    = ComponentRegistry->NewTag(TEXT("Root"));
 	Tags.ImportedEntity          = ComponentRegistry->NewTag(TEXT("Imported Entity"));
 	Tags.Finished                = ComponentRegistry->NewTag(TEXT("Finished Evaluating"));
 	Tags.Ignored                 = ComponentRegistry->NewTag(TEXT("Ignored"));
diff --git a/Engine/Source/Runtime/MovieScene/Private/EntitySystem/MovieSceneEvaluationHookSystem.cpp b/Engine/Source/Runtime/MovieScene/Private/EntitySystem/MovieSceneEvaluationHookSystem.cpp
--- a/Engine/Source/Runtime/MovieScene/Private/EntitySystem/MovieSceneEvaluationHookSystem.cpp
+++ b/Engine/Source/Runtime/MovieScene/Private/EntitySystem/MovieSceneEvaluationHookSystem.cpp
@@ -144,7 +144,7 @@ void UMovieSceneEvaluationHookSystem::OnRun(FSystemTaskPrerequisites& InPrerequi
 			FGraphEventRef SortTask = TGraphTask<FEvaluationHookSorter>::CreateTask(&Prereqs, Linker->EntityManager.GetDispatchThread())
 			.ConstructAndDispatchWhenReady(this);
 
-			Subsequents.AddMasterTask(SortTask);
+			Subsequents.AddRootTask(SortTask);
 		}
 	}
 	else if (HasEvents())
diff --git a/Engine/Source/Runtime/MovieScene/Private/EntitySystem/MovieSceneSequenceUpdaters.cpp b/Engine/Source/Runtime/MovieScene/Private/EntitySystem/MovieSceneSequenceUpdaters.cpp
--- a/Engine/Source/Runtime/MovieScene/Private/EntitySystem/MovieSceneSequenceUpdaters.cpp
+++ b/Engine/Source/Runtime/MovieScene/Private/EntitySystem/MovieSceneSequenceUpdaters.cpp
@@ -387,18 +387,18 @@ FInstanceHandle FSequenceUpdater_Hierarchical::GetOrCreateSequenceInstance(IMovi
 	return InstanceHandle;
 }
 
-void FSequenceUpdater_Hierarchical::OverrideRootSequence(UMovieSceneEntitySystemLinker* InLinker, FRootInstanceHandle MasterInstanceHandle, FMovieSceneSequenceID NewRootOverrideSequenceID)
+void FSequenceUpdater_Hierarchical::OverrideRootSequence(UMovieSceneEntitySystemLinker* InLinker, FRootInstanceHandle RootInstanceHandle, FMovieSceneSequenceID NewRootOverrideSequenceID)
 {
 	if (RootOverrideSequenceID != NewRootOverrideSequenceID)
 	{
 		if (RootOverrideSequenceID == MovieSceneSequenceID::Root)
 		{
-			// When specifying a new root override where there was none before (ie, when we were previously evaluating from the master)
-			// We unlink everything from the master sequence since we know they won't be necessary.
+			// When specifying a new root override where there was none before (ie, when we were previously evaluating from the root)
+			// We unlink everything from the root sequence since we know they won't be necessary.
 			// This is because the root sequence instance is handled separately in FSequenceUpdater_Hierarchical::Update, and it wouldn't
 			// get automatically unlinked like other sub sequences would (by way of not being present in the ActiveSequences map)
 			FInstanceRegistry* InstanceRegistry = InLinker->GetInstanceRegistry();
-			InstanceRegistry->MutateInstance(MasterInstanceHandle).Ledger.UnlinkEverything(InLinker);
+			InstanceRegistry->MutateInstance(RootInstanceHandle).Ledger.UnlinkEverything(InLinker);
 		}
 
 		InvalidateCachedData(InLinker);
@@ -426,8 +426,8 @@ void FSequenceUpdater_Hierarchical::Update(UMovieSceneEntitySystemLinker* Linker
 
 	if (RootOverrideSequenceID != MovieSceneSequenceID::Root)
 	{
-		const FMovieSceneSequenceHierarchy* MasterHierarchy = CompiledDataManager->FindHierarchy(CompiledDataID);
-		const FMovieSceneSubSequenceData*   SubData         = MasterHierarchy ? MasterHierarchy->FindSubData(RootOverrideSequenceID) : nullptr;
+		const FMovieSceneSequenceHierarchy* RootHierarchy = CompiledDataManager->FindHierarchy(CompiledDataID);
+		const FMovieSceneSubSequenceData*   SubData         = RootHierarchy ? RootHierarchy->FindSubData(RootOverrideSequenceID) : nullptr;
 		UMovieSceneSequence*                RootSequence    = SubData ? SubData->GetSequence() : nullptr;
 		if (ensure(RootSequence))
 		{
@@ -435,7 +435,7 @@ void FSequenceUpdater_Hierarchical::Update(UMovieSceneEntitySystemLinker* Linker
 			RootCompiledDataID = CompiledDataManager->GetDataID(RootSequence);
 			RootContext        = Context.Transform(SubData->RootToSequenceTransform, SubData->TickResolution);
 
-			RootOverridePath.Reset(RootOverrideSequenceID, MasterHierarchy);
+			RootOverridePath.Reset(RootOverrideSequenceID, RootHierarchy);
 
 			ActiveSequences.Add(RootOverrideSequenceID);
 		}
@@ -508,10 +508,10 @@ void FSequenceUpdater_Hierarchical::Update(UMovieSceneEntitySystemLinker* Linker
 
 		for (const FMovieSceneSubSequenceTreeEntry& Entry : RootOverrideHierarchy->GetTree().GetAllData(SubSequenceIt.Node()))
 		{
-			// When a root override path is specified, we always remap the 'local' sequence IDs to their equivalents from the master sequence.
-			FMovieSceneSequenceID SequenceIDFromMaster = RootOverridePath.ResolveChildSequenceID(Entry.SequenceID);
+			// When a root override path is specified, we always remap the 'local' sequence IDs to their equivalents from the root sequence.
+			FMovieSceneSequenceID SequenceIDFromRoot = RootOverridePath.ResolveChildSequenceID(Entry.SequenceID);
 
-			ActiveSequences.Add(SequenceIDFromMaster);
+			ActiveSequences.Add(SequenceIDFromRoot);
 
 			const FMovieSceneSubSequenceData* SubData = RootOverrideHierarchy->FindSubData(Entry.SequenceID);
 			checkf(SubData, TEXT("Sub data does not exist for a SequenceID that exists in the hierarchical tree - this indicates a corrupt compilation product."));
@@ -519,7 +519,7 @@ void FSequenceUpdater_Hierarchical::Update(UMovieSceneEntitySystemLinker* Linker
 			UMovieSceneSequence* SubSequence = SubData->GetSequence();
 			if (SubSequence == nullptr)
 			{
-				FInstanceHandle SubSequenceHandle = SequenceInstances.FindRef(SequenceIDFromMaster).Handle;
+				FInstanceHandle SubSequenceHandle = SequenceInstances.FindRef(SequenceIDFromRoot).Handle;
 				if (SubSequenceHandle.IsValid())
 				{
 					FSequenceInstance& SubSequenceInstance = InstanceRegistry->MutateInstance(SubSequenceHandle);
@@ -531,7 +531,7 @@ void FSequenceUpdater_Hierarchical::Update(UMovieSceneEntitySystemLinker* Linker
 				FMovieSceneCompiledDataID SubDataID = CompiledDataManager->GetDataID(SubSequence);
 
 				// Set the context for the root sequence instance
-				FInstanceHandle    SubSequenceHandle = GetOrCreateSequenceInstance(InPlayer, InstanceRegistry, InstanceHandle, SequenceIDFromMaster);
+				FInstanceHandle    SubSequenceHandle = GetOrCreateSequenceInstance(InPlayer, InstanceRegistry, InstanceHandle, SequenceIDFromRoot);
 				FSequenceInstance& SubSequenceInstance = InstanceRegistry->MutateInstance(SubSequenceHandle);
 
 				// Update the sub sequence's context
diff --git a/Engine/Source/Runtime/MovieScene/Private/EntitySystem/MovieSceneSpawnablesSystem.cpp b/Engine/Source/Runtime/MovieScene/Private/EntitySystem/MovieSceneSpawnablesSystem.cpp
--- a/Engine/Source/Runtime/MovieScene/Private/EntitySystem/MovieSceneSpawnablesSystem.cpp
+++ b/Engine/Source/Runtime/MovieScene/Private/EntitySystem/MovieSceneSpawnablesSystem.cpp
@@ -192,7 +192,7 @@ void UMovieSceneSpawnablesSystem::OnRun(FSystemTaskPrerequisites& InPrerequisite
 			IMovieScenePlayer* Player = Instance.GetPlayer();
 
 			// If the sequence instance has finished and it is a sub sequence, we do not destroy the spawnable
-			// if it is owned by the master sequence or externally. These will get destroyed or forgotten by the player when it ends
+			// if it is owned by the root sequence or externally. These will get destroyed or forgotten by the player when it ends
 			if (Instance.HasFinished() && Instance.IsSubSequence())
 			{
 				const UMovieSceneSequence* Sequence = Player->State.FindSequence(Instance.GetSequenceID());
diff --git a/Engine/Source/Runtime/MovieScene/Private/EntitySystem/MovieSceneSystemTaskDependencies.cpp b/Engine/Source/Runtime/MovieScene/Private/EntitySystem/MovieSceneSystemTaskDependencies.cpp
--- a/Engine/Source/Runtime/MovieScene/Private/EntitySystem/MovieSceneSystemTaskDependencies.cpp
+++ b/Engine/Source/Runtime/MovieScene/Private/EntitySystem/MovieSceneSystemTaskDependencies.cpp
@@ -53,20 +53,20 @@ void FSystemSubsequentTasks::ResetNode(uint16 InNodeID)
 	}
 }
 
-void FSystemSubsequentTasks::AddMasterTask(FGraphEventRef MasterTask)
+void FSystemSubsequentTasks::AddRootTask(FGraphEventRef RootTask)
 {
 	SCOPE_CYCLE_COUNTER(MovieSceneEval_SystemDependencyCost)
 
-	if (MasterTask)
+	if (RootTask)
 	{
 		if (!Subsequents)
 		{
 			Subsequents = MakeShared<FSystemTaskPrerequisites>();
 			Graph->Nodes.Array[NodeID].SubsequentTasks = Subsequents;
 		}
-		Subsequents->AddMasterTask(MasterTask);
+		Subsequents->AddRootTask(RootTask);
 
-		AllTasks->Add(MasterTask);
+		AllTasks->Add(RootTask);
 	}
 }
 
diff --git a/Engine/Source/Runtime/MovieScene/Private/EntitySystem/TrackInstance/MovieSceneTrackInstance.cpp b/Engine/Source/Runtime/MovieScene/Private/EntitySystem/TrackInstance/MovieSceneTrackInstance.cpp
--- a/Engine/Source/Runtime/MovieScene/Private/EntitySystem/TrackInstance/MovieSceneTrackInstance.cpp
+++ b/Engine/Source/Runtime/MovieScene/Private/EntitySystem/TrackInstance/MovieSceneTrackInstance.cpp
@@ -11,10 +11,10 @@
 
 void UMovieSceneTrackInstance::Initialize(UObject* InAnimatedObject, UMovieSceneEntitySystemLinker* InLinker)
 {
-	// We make the difference between a master track instance, and a bound track instance that lost its binding,
+	// We make the difference between a root track instance, and a bound track instance that lost its binding,
 	// by remembering if this track instance was initialized with a valid bound object.
 	AnimatedObject = InAnimatedObject;
-	bIsMasterTrackInstance = (InAnimatedObject == nullptr);
+	bIsRootTrackInstance = (InAnimatedObject == nullptr);
 
 	PrivateLinker = InLinker;
 
@@ -30,7 +30,7 @@ void UMovieSceneTrackInstance::Destroy()
 {
 	using namespace UE::MovieScene;
 
-	if (bIsMasterTrackInstance || !UE::MovieScene::FBuiltInComponentTypes::IsBoundObjectGarbage(AnimatedObject))
+	if (bIsRootTrackInstance || !UE::MovieScene::FBuiltInComponentTypes::IsBoundObjectGarbage(AnimatedObject))
 	{
 		OnDestroyed();
 	}
diff --git a/Engine/Source/Runtime/MovieScene/Private/EntitySystem/TrackInstance/MovieSceneTrackInstanceSystem.cpp b/Engine/Source/Runtime/MovieScene/Private/EntitySystem/TrackInstance/MovieSceneTrackInstanceSystem.cpp
--- a/Engine/Source/Runtime/MovieScene/Private/EntitySystem/TrackInstance/MovieSceneTrackInstanceSystem.cpp
+++ b/Engine/Source/Runtime/MovieScene/Private/EntitySystem/TrackInstance/MovieSceneTrackInstanceSystem.cpp
@@ -4,7 +4,7 @@
 #include "EntitySystem/TrackInstance/MovieSceneTrackInstance.h"
 #include "EntitySystem/MovieSceneBoundObjectInstantiator.h"
 #include "EntitySystem/MovieSceneBoundSceneComponentInstantiator.h"
-#include "EntitySystem/MovieSceneMasterInstantiatorSystem.h"
+#include "EntitySystem/MovieSceneRootInstantiatorSystem.h"
 #include "EntitySystem/MovieSceneEntitySystemTask.h"
 #include "EntitySystem/MovieSceneEntityManager.h"
 #include "EntitySystem/BuiltInComponentTypes.h"
@@ -76,7 +76,7 @@ UMovieSceneTrackInstanceInstantiator::UMovieSceneTrackInstanceInstantiator(const
 
 	if (HasAnyFlags(RF_ClassDefaultObject))
 	{
-		DefineImplicitPrerequisite(UMovieSceneMasterInstantiatorSystem::StaticClass(), GetClass());
+		DefineImplicitPrerequisite(UMovieSceneRootInstantiatorSystem::StaticClass(), GetClass());
 		DefineComponentConsumer(GetClass(), FBuiltInComponentTypes::Get()->BoundObject);
 	}
 }
@@ -287,6 +287,6 @@ void UMovieSceneTrackInstanceSystem::OnRun(FSystemTaskPrerequisites& InPrerequis
 		};
 
 		FGraphEventRef Task = FFunctionGraphTask::CreateAndDispatchWhenReady(MoveTemp(Run), GET_STATID(MovieSceneEval_GenericTrackInstanceTask), InPrerequisites.All(), Linker->EntityManager.GetGatherThread());
-		Subsequents.AddMasterTask(Task);
+		Subsequents.AddRootTask(Task);
 	}
 }
diff --git a/Engine/Source/Runtime/MovieScene/Private/Evaluation/Blending/MovieSceneAccumulatedBlendState.cpp b/Engine/Source/Runtime/MovieScene/Private/Evaluation/Blending/MovieSceneAccumulatedBlendState.cpp
--- a/Engine/Source/Runtime/MovieScene/Private/Evaluation/Blending/MovieSceneAccumulatedBlendState.cpp
+++ b/Engine/Source/Runtime/MovieScene/Private/Evaluation/Blending/MovieSceneAccumulatedBlendState.cpp
@@ -27,7 +27,7 @@ void FMovieSceneAccumulatedBlendState::Consolidate(TMap<FMovieSceneBlendingKey,
 	}
 	else
 	{
-		// Explicit nullptr means master tracks
+		// Explicit nullptr means root tracks
 		Consolidate(InOutBlendState);
 	}
 }
diff --git a/Engine/Source/Runtime/MovieScene/Private/Evaluation/MovieScenePreAnimatedState.cpp b/Engine/Source/Runtime/MovieScene/Private/Evaluation/MovieScenePreAnimatedState.cpp
--- a/Engine/Source/Runtime/MovieScene/Private/Evaluation/MovieScenePreAnimatedState.cpp
+++ b/Engine/Source/Runtime/MovieScene/Private/Evaluation/MovieScenePreAnimatedState.cpp
@@ -8,7 +8,7 @@
 #include "Evaluation/PreAnimatedState/MovieSceneRestoreStateParams.h"
 #include "Evaluation/PreAnimatedState/MovieScenePreAnimatedCaptureSources.h"
 #include "Evaluation/PreAnimatedState/MovieScenePreAnimatedObjectTokenStorage.h"
-#include "Evaluation/PreAnimatedState/MovieScenePreAnimatedMasterTokenStorage.h"
+#include "Evaluation/PreAnimatedState/MovieScenePreAnimatedRootTokenStorage.h"
 #include "Evaluation/PreAnimatedState/MovieScenePreAnimatedEntityCaptureSource.h"
 
 DECLARE_CYCLE_STAT(TEXT("Save Pre Animated State"), MovieSceneEval_SavePreAnimatedState, STATGROUP_MovieSceneEval);
diff --git a/Engine/Source/Runtime/MovieScene/Private/Evaluation/PreAnimatedState/MovieScenePreAnimatedStateExtension.cpp b/Engine/Source/Runtime/MovieScene/Private/Evaluation/PreAnimatedState/MovieScenePreAnimatedStateExtension.cpp
--- a/Engine/Source/Runtime/MovieScene/Private/Evaluation/PreAnimatedState/MovieScenePreAnimatedStateExtension.cpp
+++ b/Engine/Source/Runtime/MovieScene/Private/Evaluation/PreAnimatedState/MovieScenePreAnimatedStateExtension.cpp
@@ -6,7 +6,7 @@
 #include "Evaluation/PreAnimatedState/MovieSceneRestoreStateParams.h"
 #include "Evaluation/PreAnimatedState/MovieScenePreAnimatedCaptureSource.h"
 #include "Evaluation/PreAnimatedState/MovieScenePreAnimatedObjectTokenStorage.h"
-#include "Evaluation/PreAnimatedState/MovieScenePreAnimatedMasterTokenStorage.h"
+#include "Evaluation/PreAnimatedState/MovieScenePreAnimatedRootTokenStorage.h"
 #include "EntitySystem/MovieSceneEntitySystemTask.h"
 #include "EntitySystem/MovieSceneComponentPtr.h"
 #include "EntitySystem/BuiltInComponentTypes.h"
@@ -574,14 +574,14 @@ void FPreAnimatedStateExtension::SavePreAnimatedStateDirectly(FMovieSceneAnimTyp
 {
 	using namespace UE::MovieScene;
 
-	TSharedPtr<FAnimTypePreAnimatedStateMasterStorage> MasterStorage = WeakGenericMasterStorage.Pin();
-	if (!MasterStorage)
+	TSharedPtr<FAnimTypePreAnimatedStateRootStorage> RootStorage = WeakGenericRootStorage.Pin();
+	if (!RootStorage)
 	{
-		MasterStorage = GetOrCreateStorage<FAnimTypePreAnimatedStateMasterStorage>();
-		WeakGenericMasterStorage = MasterStorage;
+		RootStorage = GetOrCreateStorage<FAnimTypePreAnimatedStateRootStorage>();
+		WeakGenericRootStorage = RootStorage;
 	}
 
-	FPreAnimatedStateEntry   Entry        = MasterStorage->MakeEntry(InTokenType);
+	FPreAnimatedStateEntry   Entry        = RootStorage->MakeEntry(InTokenType);
 	FPreAnimatedStorageIndex StorageIndex = Entry.ValueHandle.StorageIndex;
 
 	AddSourceMetaData(Entry);
@@ -590,18 +590,18 @@ void FPreAnimatedStateExtension::SavePreAnimatedStateDirectly(FMovieSceneAnimTyp
 		? EPreAnimatedStorageRequirement::Transient
 		: EPreAnimatedStorageRequirement::Persistent;
 
-	if (!MasterStorage->IsStorageRequirementSatisfied(StorageIndex, Requirement))
+	if (!RootStorage->IsStorageRequirementSatisfied(StorageIndex, Requirement))
 	{
 		IMovieScenePreAnimatedGlobalTokenPtr Token = Producer.CacheExistingState();
 		if (Token.IsValid())
 		{
-			const bool bHasEverAninmated = MasterStorage->HasEverAnimated(StorageIndex);
+			const bool bHasEverAninmated = RootStorage->HasEverAnimated(StorageIndex);
 			if (!bHasEverAninmated)
 			{
 				Producer.InitializeForAnimation();
 			}
 
-			MasterStorage->AssignPreAnimatedValue(StorageIndex, Requirement, MoveTemp(Token));
+			RootStorage->AssignPreAnimatedValue(StorageIndex, Requirement, MoveTemp(Token));
 		}
 	}
 }
diff --git a/Engine/Source/Runtime/MovieScene/Private/MovieScene.cpp b/Engine/Source/Runtime/MovieScene/Private/MovieScene.cpp
--- a/Engine/Source/Runtime/MovieScene/Private/MovieScene.cpp
+++ b/Engine/Source/Runtime/MovieScene/Private/MovieScene.cpp
@@ -123,6 +123,12 @@ void UMovieScene::Serialize( FArchive& Ar )
 #if WITH_EDITOR
 	if (Ar.IsLoading())
 	{
+		if (MasterTracks_DEPRECATED.Num())
+		{
+			Tracks = MasterTracks_DEPRECATED;
+			MasterTracks_DEPRECATED.Empty();
+		}
+
 		UpgradeTimeRanges();
 		RemoveNullTracks();
 
@@ -940,10 +946,10 @@ TArray<UMovieSceneSection*> UMovieScene::GetAllSections() const
 {
 	TArray<UMovieSceneSection*> OutSections;
 
-	// Add all master type sections 
-	for( int32 TrackIndex = 0; TrackIndex < MasterTracks.Num(); ++TrackIndex )
+	// Add all type sections 
+	for( int32 TrackIndex = 0; TrackIndex < Tracks.Num(); ++TrackIndex )
 	{
-		OutSections.Append( MasterTracks[TrackIndex]->GetAllSections() );
+		OutSections.Append( Tracks[TrackIndex]->GetAllSections() );
 	}
 	
 	// Add all camera cut sections
@@ -1073,23 +1079,38 @@ bool UMovieScene::AddGivenTrack(UMovieSceneTrack* InTrack, const FGuid& ObjectGu
 
 bool UMovieScene::RemoveTrack(UMovieSceneTrack& Track)
 {
-	Modify();
-
-	bool bAnythingRemoved = false;
+	// Remove either a root track or a track from a binding
+	if (Tracks.Contains(&Track))
+	{
+		Modify();
 
-	for (auto& Binding : ObjectBindings)
+		const bool bRemoved = Tracks.RemoveSingle(&Track) != 0;
+		if (bRemoved)
+		{
+			EventHandlers.Trigger(&UE::MovieScene::ISequenceDataEventHandler::OnTrackRemoved, &Track);
+		}
+		return bRemoved;
+	}
+	else
 	{
-		if (Binding.RemoveTrack(Track, this))
+		Modify();
+
+		bool bAnythingRemoved = false;
+		
+		for (auto& Binding : ObjectBindings)
 		{
-			bAnythingRemoved = true;
+			if (Binding.RemoveTrack(Track, this))
+			{
+				bAnythingRemoved = true;
 
-			// The track was removed from the current binding, stop
-			// searching now as it cannot exist in any other binding
-			break;
+				// The track was removed from the current binding, stop
+				// searching now as it cannot exist in any other binding
+				break;
+			}
 		}
+	
+		return bAnythingRemoved;
 	}
-
-	return bAnythingRemoved;
 }
 
 bool UMovieScene::FindTrackBinding(const UMovieSceneTrack& InTrack, FGuid& OutGuid) const
@@ -1109,11 +1130,11 @@ bool UMovieScene::FindTrackBinding(const UMovieSceneTrack& InTrack, FGuid& OutGu
 	return false;
 }
 
-UMovieSceneTrack* UMovieScene::FindMasterTrack( TSubclassOf<UMovieSceneTrack> TrackClass ) const
+UMovieSceneTrack* UMovieScene::FindTrack( TSubclassOf<UMovieSceneTrack> TrackClass ) const
 {
 	UMovieSceneTrack* FoundTrack = nullptr;
 
-	for (const auto Track : MasterTracks)
+	for (const auto Track : Tracks)
 	{
 		if (Track->GetClass()->IsChildOf(TrackClass))
 		{
@@ -1126,7 +1147,7 @@ UMovieSceneTrack* UMovieScene::FindMasterTrack( TSubclassOf<UMovieSceneTrack> Tr
 }
 
 
-UMovieSceneTrack* UMovieScene::AddMasterTrack( TSubclassOf<UMovieSceneTrack> TrackClass )
+UMovieSceneTrack* UMovieScene::AddTrack( TSubclassOf<UMovieSceneTrack> TrackClass )
 {
 #if WITH_EDITOR
 	if (!IsTrackClassAllowed(TrackClass))
@@ -1138,15 +1159,15 @@ UMovieSceneTrack* UMovieScene::AddMasterTrack( TSubclassOf<UMovieSceneTrack> Tra
 	Modify();
 
 	UMovieSceneTrack* CreatedType = NewObject<UMovieSceneTrack>(this, TrackClass, NAME_None, RF_Transactional);
-	MasterTracks.Add( CreatedType );
+	Tracks.Add( CreatedType );
 
-	EventHandlers.Trigger(&UE::MovieScene::ISequenceDataEventHandler::OnMasterTrackAdded, CreatedType);
+	EventHandlers.Trigger(&UE::MovieScene::ISequenceDataEventHandler::OnTrackAdded, CreatedType);
 
 	return CreatedType;
 }
 
 
-bool UMovieScene::AddGivenMasterTrack(UMovieSceneTrack* InTrack)
+bool UMovieScene::AddGivenTrack(UMovieSceneTrack* InTrack)
 {
 #if WITH_EDITOR
 	if (!IsTrackClassAllowed(InTrack->GetClass()))
@@ -1155,13 +1176,13 @@ bool UMovieScene::AddGivenMasterTrack(UMovieSceneTrack* InTrack)
 	}
 #endif
 
-	if (!MasterTracks.Contains(InTrack))
+	if (!Tracks.Contains(InTrack))
 	{
 		Modify();
-		MasterTracks.Add(InTrack);
+		Tracks.Add(InTrack);
 		InTrack->Rename(nullptr, this);
 
-		EventHandlers.Trigger(&UE::MovieScene::ISequenceDataEventHandler::OnMasterTrackAdded, InTrack);
+		EventHandlers.Trigger(&UE::MovieScene::ISequenceDataEventHandler::OnTrackAdded, InTrack);
 
 		return true;
 	}
@@ -1169,24 +1190,11 @@ bool UMovieScene::AddGivenMasterTrack(UMovieSceneTrack* InTrack)
 }
 
 
-bool UMovieScene::RemoveMasterTrack(UMovieSceneTrack& Track) 
-{
-	Modify();
-
-	const bool bRemoved = MasterTracks.RemoveSingle(&Track) != 0;
-	if (bRemoved)
-	{
-		EventHandlers.Trigger(&UE::MovieScene::ISequenceDataEventHandler::OnMasterTrackRemoved, &Track);
-	}
-	return bRemoved;
-}
-
-
-bool UMovieScene::IsAMasterTrack(const UMovieSceneTrack& Track) const
+bool UMovieScene::ContainsTrack(const UMovieSceneTrack& InTrack) const
 {
-	for ( const UMovieSceneTrack* MasterTrack : MasterTracks)
+	for ( const UMovieSceneTrack* Track : Tracks)
 	{
-		if (&Track == MasterTrack)
+		if (&InTrack == Track)
 		{
 			return true;
 		}
@@ -1210,7 +1218,7 @@ UMovieSceneTrack* UMovieScene::AddCameraCutTrack( TSubclassOf<UMovieSceneTrack>
 		Modify();
 		CameraCutTrack = NewObject<UMovieSceneTrack>(this, TrackClass, NAME_None, RF_Transactional);
 
-		EventHandlers.Trigger(&UE::MovieScene::ISequenceDataEventHandler::OnMasterTrackAdded, CameraCutTrack);
+		EventHandlers.Trigger(&UE::MovieScene::ISequenceDataEventHandler::OnTrackAdded, CameraCutTrack);
 	}
 
 	return CameraCutTrack;
@@ -1231,7 +1239,7 @@ void UMovieScene::RemoveCameraCutTrack()
 		UMovieSceneTrack* TmpCameraCut = CameraCutTrack;
 		CameraCutTrack = nullptr;
 
-		EventHandlers.Trigger(&UE::MovieScene::ISequenceDataEventHandler::OnMasterTrackRemoved, TmpCameraCut);
+		EventHandlers.Trigger(&UE::MovieScene::ISequenceDataEventHandler::OnTrackRemoved, TmpCameraCut);
 	}
 }
 
@@ -1249,10 +1257,10 @@ void UMovieScene::SetCameraCutTrack(UMovieSceneTrack* InTrack)
 
 	if (OldCameraCutTrack)
 	{
-		EventHandlers.Trigger(&UE::MovieScene::ISequenceDataEventHandler::OnMasterTrackRemoved, OldCameraCutTrack);
+		EventHandlers.Trigger(&UE::MovieScene::ISequenceDataEventHandler::OnTrackRemoved, OldCameraCutTrack);
 	}
 
-	EventHandlers.Trigger(&UE::MovieScene::ISequenceDataEventHandler::OnMasterTrackAdded, CameraCutTrack);
+	EventHandlers.Trigger(&UE::MovieScene::ISequenceDataEventHandler::OnTrackAdded, CameraCutTrack);
 }
 
 
@@ -1285,7 +1293,7 @@ void UMovieScene::UpgradeTimeRanges()
 		// In this instance (UMG), playback always started at 0
 		TRangeBound<FFrameNumber> MaxFrame = TRangeBound<FFrameNumber>::Exclusive(0);
 
-		for (const UMovieSceneTrack* Track : MasterTracks)
+		for (const UMovieSceneTrack* Track : Tracks)
 		{
 			TOptional<TRangeBound<FFrameNumber>> MaxUpper = GetMaxUpperBound(Track);
 			if (MaxUpper.IsSet())
@@ -1366,11 +1374,11 @@ void UMovieScene::UpgradeTimeRanges()
 void UMovieScene::RemoveNullTracks()
 {
 	// Remove any null tracks
-	for( int32 TrackIndex = 0; TrackIndex < MasterTracks.Num(); )
+	for( int32 TrackIndex = 0; TrackIndex < Tracks.Num(); )
 	{
-		if (MasterTracks[TrackIndex] == nullptr)
+		if (Tracks[TrackIndex] == nullptr)
 		{
-			MasterTracks.RemoveAt(TrackIndex);
+			Tracks.RemoveAt(TrackIndex);
 		}
 		else
 		{
diff --git a/Engine/Source/Runtime/MovieScene/Private/MovieSceneFolder.cpp b/Engine/Source/Runtime/MovieScene/Private/MovieSceneFolder.cpp
--- a/Engine/Source/Runtime/MovieScene/Private/MovieSceneFolder.cpp
+++ b/Engine/Source/Runtime/MovieScene/Private/MovieSceneFolder.cpp
@@ -84,13 +84,13 @@ void UMovieSceneFolder::RemoveChildFolder( UMovieSceneFolder* InChildFolder )
 }
 
 
-const TArray<UMovieSceneTrack*>& UMovieSceneFolder::GetChildMasterTracks() const
+const TArray<UMovieSceneTrack*>& UMovieSceneFolder::GetChildTracks() const
 {
-	return ChildMasterTracks;
+	return ChildTracks;
 }
 
 
-void UMovieSceneFolder::AddChildMasterTrack( UMovieSceneTrack* InMasterTrack )
+void UMovieSceneFolder::AddChildTrack( UMovieSceneTrack* InTrack )
 {
 	Modify();
 
@@ -104,33 +104,33 @@ void UMovieSceneFolder::AddChildMasterTrack( UMovieSceneTrack* InMasterTrack )
 
 		for (UMovieSceneFolder* MovieSceneFolder : AllFolders)
 		{
-			MovieSceneFolder->RemoveChildMasterTrack(InMasterTrack);
+			MovieSceneFolder->RemoveChildTrack(InTrack);
 		}
 	}
 #endif
 
-	ChildMasterTracks.Add( InMasterTrack );
+	ChildTracks.Add( InTrack );
 
-	EventHandlers.Trigger(&UE::MovieScene::IFolderEventHandler::OnTrackAdded, InMasterTrack);
+	EventHandlers.Trigger(&UE::MovieScene::IFolderEventHandler::OnTrackAdded, InTrack);
 }
 
 
-void UMovieSceneFolder::RemoveChildMasterTrack( UMovieSceneTrack* InMasterTrack )
+void UMovieSceneFolder::RemoveChildTrack( UMovieSceneTrack* InTrack )
 {
 	Modify();
 
-	if (ChildMasterTracks.Remove( InMasterTrack ) > 0)
+	if (ChildTracks.Remove( InTrack ) > 0)
 	{
-		EventHandlers.Trigger(&UE::MovieScene::IFolderEventHandler::OnTrackRemoved, InMasterTrack);
+		EventHandlers.Trigger(&UE::MovieScene::IFolderEventHandler::OnTrackRemoved, InTrack);
 	}
 }
 
 
-void UMovieSceneFolder::ClearChildMasterTracks()
+void UMovieSceneFolder::ClearChildTracks()
 {
 	Modify();
 
-	ChildMasterTracks.Empty();
+	ChildTracks.Empty();
 }
 
 
@@ -145,7 +145,7 @@ void UMovieSceneFolder::AddChildObjectBinding(const FGuid& InObjectBinding )
 	Modify();
 
 #if WITH_EDITORONLY_DATA
-	// Ensure the added object  does not belong to any other folder in the same scene.
+	// Ensure the added object does not belong to any other folder in the same scene.
 	UMovieScene* OwningScene = GetTypedOuter<UMovieScene>();
 	if (OwningScene)
 	{
@@ -208,16 +208,16 @@ void UMovieSceneFolder::PostLoad()
 	UMovieScene* OwningScene = GetTypedOuter<UMovieScene>();
 	if (OwningScene)
 	{
-		// Validate child Master Tracks
-		for(int32 ChildMasterTrackIndex = 0; ChildMasterTrackIndex < ChildMasterTracks.Num(); ChildMasterTrackIndex++)
+		// Validate child Tracks
+		for(int32 ChildTrackIndex = 0; ChildTrackIndex < ChildTracks.Num(); ChildTrackIndex++)
 		{
-			const UMovieSceneTrack* ChildTrack = ChildMasterTracks[ChildMasterTrackIndex];
-			if (!OwningScene->GetMasterTracks().Contains(ChildTrack))
+			const UMovieSceneTrack* ChildTrack = ChildTracks[ChildTrackIndex];
+			if (!OwningScene->GetTracks().Contains(ChildTrack))
 			{
-				ChildMasterTracks.RemoveAt(ChildMasterTrackIndex);
-				ChildMasterTrackIndex--;
+				ChildTracks.RemoveAt(ChildTrackIndex);
+				ChildTrackIndex--;
 
-				UE_LOG(LogMovieScene, Warning, TEXT("Folder (%s) in Sequence (%s) contained a reference to a Master Track (%s) that no longer exists in the sequence, removing."), *GetFolderName().ToString(), *OwningScene->GetPathName(), *GetNameSafe(ChildTrack));
+				UE_LOG(LogMovieScene, Warning, TEXT("Folder (%s) in Sequence (%s) contained a reference to a Track (%s) that no longer exists in the sequence, removing."), *GetFolderName().ToString(), *OwningScene->GetPathName(), *GetNameSafe(ChildTrack));
 			}
 		}
 
@@ -281,7 +281,7 @@ UMovieSceneFolder* UMovieSceneFolder::FindFolderContaining(const FGuid& InObject
 
 UMovieSceneFolder* UMovieSceneFolder::FindFolderContaining(const UMovieSceneTrack* InTrack)
 {
-	if (ChildMasterTracks.Contains(InTrack))
+	if (ChildTracks.Contains(InTrack))
 	{
 		return this;
 	}
@@ -348,6 +348,15 @@ void UMovieSceneFolder::Serialize( FArchive& Archive )
 	if ( Archive.IsLoading() )
 	{
 		Super::Serialize( Archive );
+
+#if WITH_EDITOR
+		if (ChildMasterTracks_DEPRECATED.Num())
+		{
+			ChildTracks = ChildMasterTracks_DEPRECATED;
+			ChildMasterTracks_DEPRECATED.Empty();
+		}
+#endif
+
 		ChildObjectBindings.Empty();
 		for ( const FString& ChildObjectBindingString : ChildObjectBindingStrings )
 		{
diff --git a/Engine/Source/Runtime/MovieScene/Private/MovieSceneObjectBindingID.cpp b/Engine/Source/Runtime/MovieScene/Private/MovieSceneObjectBindingID.cpp
--- a/Engine/Source/Runtime/MovieScene/Private/MovieSceneObjectBindingID.cpp
+++ b/Engine/Source/Runtime/MovieScene/Private/MovieSceneObjectBindingID.cpp
@@ -62,12 +62,12 @@ FMovieSceneSequenceID ResolveExternalSequenceID(FMovieSceneSequenceID SourceSequ
 }
 
 
-FRelativeObjectBindingID::FRelativeObjectBindingID(FMovieSceneSequenceID SourceSequenceID, FMovieSceneSequenceID TargetSequenceID, const FGuid& TargetGuid, UMovieSceneSequence* MasterSequence)
+FRelativeObjectBindingID::FRelativeObjectBindingID(FMovieSceneSequenceID SourceSequenceID, FMovieSceneSequenceID TargetSequenceID, const FGuid& TargetGuid, UMovieSceneSequence* RootSequence)
 {
 	UMovieSceneCompiledDataManager* CompiledDataManager = UMovieSceneCompiledDataManager::GetPrecompiledData();
 
 	FMovieSceneSequenceHierarchy TempHierarchy;
-	UMovieSceneCompiledDataManager::CompileHierarchy(MasterSequence, &TempHierarchy, EMovieSceneServerClientMask::All);
+	UMovieSceneCompiledDataManager::CompileHierarchy(RootSequence, &TempHierarchy, EMovieSceneServerClientMask::All);
 
 	ConstructInternal(SourceSequenceID, TargetSequenceID, TargetGuid, &TempHierarchy);
 }
diff --git a/Engine/Source/Runtime/MovieScene/Private/MovieSceneSequence.cpp b/Engine/Source/Runtime/MovieScene/Private/MovieSceneSequence.cpp
--- a/Engine/Source/Runtime/MovieScene/Private/MovieSceneSequence.cpp
+++ b/Engine/Source/Runtime/MovieScene/Private/MovieSceneSequence.cpp
@@ -164,38 +164,38 @@ bool UMovieSceneSequence::OptimizeForCook()
 
 	bool bModified = false;
 
-	for (int32 MasterTrackIndex = 0; MasterTrackIndex < MovieScene->GetMasterTracks().Num(); )
+	for (int32 TrackIndex = 0; TrackIndex < MovieScene->GetTracks().Num(); )
 	{
-		UMovieSceneTrack* MasterTrack = MovieScene->GetMasterTracks()[MasterTrackIndex];
-		if (MasterTrack && MasterTrack->GetCookOptimizationFlags() == ECookOptimizationFlags::RemoveTrack)
+		UMovieSceneTrack* Track = MovieScene->GetTracks()[TrackIndex];
+		if (Track && Track->GetCookOptimizationFlags() == ECookOptimizationFlags::RemoveTrack)
 		{				
-			MasterTrack->RemoveForCook();
-			MovieScene->RemoveMasterTrack(*MasterTrack);
-			UE_LOG(LogMovieScene, Display, TEXT("Removing muted track: %s from: %s"), *MasterTrack->GetDisplayName().ToString(), *GetPathName());
+			Track->RemoveForCook();
+			MovieScene->RemoveTrack(*Track);
+			UE_LOG(LogMovieScene, Display, TEXT("Removing muted track: %s from: %s"), *Track->GetDisplayName().ToString(), *GetPathName());
 			bModified = true;
 			continue;
 		}
-		++MasterTrackIndex;
+		++TrackIndex;
 	}
 
 	// Go through the tracks again and look at sections
-	for (int32 MasterTrackIndex = 0; MasterTrackIndex < MovieScene->GetMasterTracks().Num(); ++MasterTrackIndex)
+	for (int32 TrackIndex = 0; TrackIndex < MovieScene->GetTracks().Num(); ++TrackIndex)
 	{
-		UMovieSceneTrack* MasterTrack =  MovieScene->GetMasterTracks()[MasterTrackIndex];
-		if (MasterTrack)
+		UMovieSceneTrack* Track =  MovieScene->GetTracks()[TrackIndex];
+		if (Track)
 		{
-			for (int32 MasterSectionIndex = 0; MasterSectionIndex < MasterTrack->GetAllSections().Num(); )
+			for (int32 SectionIndex = 0; SectionIndex < Track->GetAllSections().Num(); )
 			{
-				UMovieSceneSection* MasterSection = MasterTrack->GetAllSections()[MasterSectionIndex];
-				if (MasterSection && MasterSection->GetCookOptimizationFlags() == ECookOptimizationFlags::RemoveSection)
+				UMovieSceneSection* Section = Track->GetAllSections()[SectionIndex];
+				if (Section && Section->GetCookOptimizationFlags() == ECookOptimizationFlags::RemoveSection)
 				{
-					MasterSection->RemoveForCook();
-					MasterTrack->RemoveSection(*MasterSection);
-					UE_LOG(LogMovieScene, Display, TEXT("Removing muted section: %s from: %s"), *MasterSection->GetPathName(), *MasterTrack->GetDisplayName().ToString());
+					Section->RemoveForCook();
+					Track->RemoveSection(*Section);
+					UE_LOG(LogMovieScene, Display, TEXT("Removing muted section: %s from: %s"), *Section->GetPathName(), *Track->GetDisplayName().ToString());
 					bModified = true;
 					continue;
 				}
-				++MasterSectionIndex;
+				++SectionIndex;
 			}
 		}
 	}
diff --git a/Engine/Source/Runtime/MovieScene/Private/MovieSceneSequenceVisitor.cpp b/Engine/Source/Runtime/MovieScene/Private/MovieSceneSequenceVisitor.cpp
--- a/Engine/Source/Runtime/MovieScene/Private/MovieSceneSequenceVisitor.cpp
+++ b/Engine/Source/Runtime/MovieScene/Private/MovieSceneSequenceVisitor.cpp
@@ -90,7 +90,7 @@ void VisitTrackImpl(const FSequenceVisitParams& InParams, UMovieSceneTrack* InTr
 		}
 	}
 
-	if (InParams.bVisitTracks || InParams.bVisitMasterTracks)
+	if (InParams.bVisitTracks || InParams.bVisitRootTracks)
 	{
 		InVisitor.VisitTrack(InTrack, ObjectBinding, *SubSequenceSpace);
 	}
@@ -117,14 +117,14 @@ void VisitSequenceImpl(UMovieSceneSequence* Sequence, const FSequenceVisitParams
 		return;
 	}
 
-	if (InParams.bVisitMasterTracks)
+	if (InParams.bVisitRootTracks)
 	{
 		if (UMovieSceneTrack* Track = MovieScene->GetCameraCutTrack())
 		{
 			VisitTrackImpl(InParams, Track, InVisitor, FGuid(), InOutRootPath, SubSequenceSpace);
 		}
 
-		for (UMovieSceneTrack* Track : MovieScene->GetMasterTracks())
+		for (UMovieSceneTrack* Track : MovieScene->GetTracks())
 		{
 			VisitTrackImpl(InParams, Track, InVisitor, FGuid(), InOutRootPath, SubSequenceSpace);
 		}
diff --git a/Engine/Source/Runtime/MovieScene/Private/MovieSceneTimeHelpers.cpp b/Engine/Source/Runtime/MovieScene/Private/MovieSceneTimeHelpers.cpp
--- a/Engine/Source/Runtime/MovieScene/Private/MovieSceneTimeHelpers.cpp
+++ b/Engine/Source/Runtime/MovieScene/Private/MovieSceneTimeHelpers.cpp
@@ -120,7 +120,7 @@ void MigrateFrameTimes(FFrameRate SourceRate, FFrameRate DestinationRate, UMovie
 
 void TimeHelpers::MigrateFrameTimes(FFrameRate SourceRate, FFrameRate DestinationRate, UMovieScene* MovieScene, bool bApplyRecursively)
 {
-	int32 TotalNumTracks = MovieScene->GetMasterTracks().Num() + (MovieScene->GetCameraCutTrack() ? 1 : 0);
+	int32 TotalNumTracks = MovieScene->GetTracks().Num() + (MovieScene->GetCameraCutTrack() ? 1 : 0);
 	for (const FMovieSceneBinding& Binding : MovieScene->GetBindings())
 	{
 		TotalNumTracks += Binding.GetTracks().Num();
@@ -140,7 +140,7 @@ void TimeHelpers::MigrateFrameTimes(FFrameRate SourceRate, FFrameRate Destinatio
 	MovieScene->SetSelectionRange(MigrateFrameRange(MovieScene->GetSelectionRange(), SourceRate, DestinationRate));
 #endif
 
-	for (UMovieSceneTrack* Track : MovieScene->GetMasterTracks())
+	for (UMovieSceneTrack* Track : MovieScene->GetTracks())
 	{
 		SlowTask.EnterProgressFrame();
 		UE::MovieScene::MigrateFrameTimes(SourceRate, DestinationRate, Track);
diff --git a/Engine/Source/Runtime/MovieScene/Private/Sections/MovieSceneSubSection.cpp b/Engine/Source/Runtime/MovieScene/Private/Sections/MovieSceneSubSection.cpp
--- a/Engine/Source/Runtime/MovieScene/Private/Sections/MovieSceneSubSection.cpp
+++ b/Engine/Source/Runtime/MovieScene/Private/Sections/MovieSceneSubSection.cpp
@@ -214,7 +214,7 @@ void UMovieSceneSubSection::PostEditChangeProperty(FPropertyChangedEvent& Proper
 {
 	if (PropertyChangedEvent.Property && PropertyChangedEvent.Property->GetFName() == GET_MEMBER_NAME_CHECKED(UMovieSceneSubSection, SubSequence))
 	{
-		// Check whether the subsequence that was just set has master tracks that contain the sequence that this subsection is in.
+		// Check whether the subsequence that was just set has tracks that contain the sequence that this subsection is in.
 		UMovieScene* SubSequenceMovieScene = SubSequence ? SubSequence->GetMovieScene() : nullptr;
 
 		UMovieSceneSubTrack* TrackOuter = Cast<UMovieSceneSubTrack>(GetOuter());
@@ -225,9 +225,9 @@ void UMovieSceneSubSection::PostEditChangeProperty(FPropertyChangedEvent& Proper
 			{
 				TArray<UMovieSceneSubTrack*> SubTracks;
 
-				for (UMovieSceneTrack* MasterTrack : SubSequenceMovieScene->GetMasterTracks())
+				for (UMovieSceneTrack* Track : SubSequenceMovieScene->GetTracks())
 				{
-					if (UMovieSceneSubTrack* SubTrack = Cast<UMovieSceneSubTrack>(MasterTrack))
+					if (UMovieSceneSubTrack* SubTrack = Cast<UMovieSceneSubTrack>(Track))
 					{
 						SubTracks.Add(SubTrack);
 					}
diff --git a/Engine/Source/Runtime/MovieScene/Private/Tests/MovieSceneCompilerTests.cpp b/Engine/Source/Runtime/MovieScene/Private/Tests/MovieSceneCompilerTests.cpp
--- a/Engine/Source/Runtime/MovieScene/Private/Tests/MovieSceneCompilerTests.cpp
+++ b/Engine/Source/Runtime/MovieScene/Private/Tests/MovieSceneCompilerTests.cpp
@@ -33,7 +33,7 @@ bool FMovieSceneCompilerPerfTest::RunTest(const FString& Parameters)
 
 	for (int32 i = 0; i < 100; ++i)
 	{
-		UTestMovieSceneTrack* Track = Sequence->MovieScene->AddMasterTrack<UTestMovieSceneTrack>();
+		UTestMovieSceneTrack* Track = Sequence->MovieScene->AddTrack<UTestMovieSceneTrack>();
 
 		int32 NumSections = FMath::Rand() % 10;
 		for (int32 SectionIndex = 0; SectionIndex < NumSections; ++SectionIndex)
@@ -186,7 +186,7 @@ bool FMovieSceneCompilerRangeTest::RunTest(const FString& Parameters)
 	UTestMovieSceneSequence* Sequence = NewObject<UTestMovieSceneSequence>(GetTransientPackage());
 	for (TRange<FFrameNumber> Range : Ranges)
 	{
-		UTestMovieSceneTrack*   Track   = Sequence->MovieScene->AddMasterTrack<UTestMovieSceneTrack>();
+		UTestMovieSceneTrack*   Track   = Sequence->MovieScene->AddTrack<UTestMovieSceneTrack>();
 		UTestMovieSceneSection* Section = NewObject<UTestMovieSceneSection>(Track);
 
 		Section->SetRange(Range);
@@ -300,7 +300,7 @@ bool FMovieSceneCompilerEmptySpaceOnTheFlyTest::RunTest(const FString& Parameter
 	}
 
 	UTestMovieSceneSequence* Sequence = NewObject<UTestMovieSceneSequence>(GetTransientPackage());
-	UTestMovieSceneTrack*    Track    = Sequence->MovieScene->AddMasterTrack<UTestMovieSceneTrack>();
+	UTestMovieSceneTrack*    Track    = Sequence->MovieScene->AddTrack<UTestMovieSceneTrack>();
 
 	for (TRange<FFrameNumber> Range : SectionRanges)
 	{
@@ -406,7 +406,7 @@ bool FMovieSceneCompilerSubSequencesTest::RunTest(const FString& Parameters)
 	};
 
 	UTestMovieSceneSequence* RootSequence = NewObject<UTestMovieSceneSequence>(GetTransientPackage());
-	UTestMovieSceneSubTrack* RootSubTrack = RootSequence->MovieScene->AddMasterTrack<UTestMovieSceneSubTrack>();
+	UTestMovieSceneSubTrack* RootSubTrack = RootSequence->MovieScene->AddTrack<UTestMovieSceneSubTrack>();
 	
 	UTestMovieSceneSequence* Shot1Sequence = NewObject<UTestMovieSceneSequence>(GetTransientPackage());
 	Shot1Sequence->GetMovieScene()->SetPlaybackRange(0, 100);
@@ -416,7 +416,7 @@ bool FMovieSceneCompilerSubSequencesTest::RunTest(const FString& Parameters)
 	Shot1SubSection->SetSequence(Shot1Sequence);
 	RootSubTrack->SectionArray.Add(Shot1SubSection);
 	
-	UTestMovieSceneTrack* Shot1Track = Shot1Sequence->MovieScene->AddMasterTrack<UTestMovieSceneTrack>();
+	UTestMovieSceneTrack* Shot1Track = Shot1Sequence->MovieScene->AddTrack<UTestMovieSceneTrack>();
 	UTestMovieSceneSection* Shot1Section = NewObject<UTestMovieSceneSection>(Shot1Track);
 	Shot1Section->SetRange(TRange<FFrameNumber>(0, 60));
 	Shot1Track->SectionArray.Add(Shot1Section);
diff --git a/Engine/Source/Runtime/MovieScene/Private/Tests/MovieScenePreAnimatedStateTests.cpp b/Engine/Source/Runtime/MovieScene/Private/Tests/MovieScenePreAnimatedStateTests.cpp
--- a/Engine/Source/Runtime/MovieScene/Private/Tests/MovieScenePreAnimatedStateTests.cpp
+++ b/Engine/Source/Runtime/MovieScene/Private/Tests/MovieScenePreAnimatedStateTests.cpp
@@ -356,7 +356,7 @@ bool FMovieScenePreAnimatedStateTrackTypesTest::RunTest(const FString& Parameter
 
 	// Make a keep-state eval hook section
 	{
-		UTestMovieSceneEvalHookTrack*   Track = Sequence->MovieScene->AddMasterTrack<UTestMovieSceneEvalHookTrack>();
+		UTestMovieSceneEvalHookTrack*   Track = Sequence->MovieScene->AddTrack<UTestMovieSceneEvalHookTrack>();
 		UTestMovieSceneEvalHookSection* Section = NewObject<UTestMovieSceneEvalHookSection>(Track);
 
 		Section->EvalOptions.CompletionMode = EMovieSceneCompletionMode::KeepState;
@@ -368,7 +368,7 @@ bool FMovieScenePreAnimatedStateTrackTypesTest::RunTest(const FString& Parameter
 	}
 	// Make a restore-state eval hook section
 	{
-		UTestMovieSceneEvalHookTrack*   Track = Sequence->MovieScene->AddMasterTrack<UTestMovieSceneEvalHookTrack>();
+		UTestMovieSceneEvalHookTrack*   Track = Sequence->MovieScene->AddTrack<UTestMovieSceneEvalHookTrack>();
 		UTestMovieSceneEvalHookSection* Section = NewObject<UTestMovieSceneEvalHookSection>(Track);
 
 		Section->EvalOptions.CompletionMode = EMovieSceneCompletionMode::RestoreState;
@@ -450,7 +450,7 @@ bool FMovieScenePreAnimatedStateContextChangedTest::RunTest(const FString& Param
 
 	// Make a keep-state eval hook section
 	{
-		UTestMovieSceneEvalHookTrack*   Track = Sequence->MovieScene->AddMasterTrack<UTestMovieSceneEvalHookTrack>();
+		UTestMovieSceneEvalHookTrack*   Track = Sequence->MovieScene->AddTrack<UTestMovieSceneEvalHookTrack>();
 		UTestMovieSceneEvalHookSection* Section = NewObject<UTestMovieSceneEvalHookSection>(Track);
 
 		Section->EvalOptions.CompletionMode = EMovieSceneCompletionMode::KeepState;
diff --git a/Engine/Source/Runtime/MovieScene/Private/Tracks/MovieSceneSubTrack.cpp b/Engine/Source/Runtime/MovieScene/Private/Tracks/MovieSceneSubTrack.cpp
--- a/Engine/Source/Runtime/MovieScene/Private/Tracks/MovieSceneSubTrack.cpp
+++ b/Engine/Source/Runtime/MovieScene/Private/Tracks/MovieSceneSubTrack.cpp
@@ -118,7 +118,7 @@ bool UMovieSceneSubTrack::ContainsSequence(const UMovieSceneSequence& Sequence,
 			continue;
 		}
 
-		UMovieSceneSubTrack* SubSubTrack = SubMovieScene->FindMasterTrack<UMovieSceneSubTrack>();
+		UMovieSceneSubTrack* SubSubTrack = SubMovieScene->FindTrack<UMovieSceneSubTrack>();
 
 		if ((SubSubTrack != nullptr) && SubSubTrack->ContainsSequence(Sequence))
 		{
diff --git a/Engine/Source/Runtime/MovieScene/Public/EntitySystem/BuiltInComponentTypes.h b/Engine/Source/Runtime/MovieScene/Public/EntitySystem/BuiltInComponentTypes.h
--- a/Engine/Source/Runtime/MovieScene/Public/EntitySystem/BuiltInComponentTypes.h
+++ b/Engine/Source/Runtime/MovieScene/Public/EntitySystem/BuiltInComponentTypes.h
@@ -380,7 +380,7 @@ public:
 		FComponentTypeID MigratedFromFastPath;
 
 		FComponentTypeID ImportedEntity;
-		FComponentTypeID Master;
+		FComponentTypeID Root;
 
 		FComponentTypeID FixedTime;
 
diff --git a/Engine/Source/Runtime/MovieScene/Public/EntitySystem/MovieSceneEntitySystemTask.h b/Engine/Source/Runtime/MovieScene/Public/EntitySystem/MovieSceneEntitySystemTask.h
--- a/Engine/Source/Runtime/MovieScene/Public/EntitySystem/MovieSceneEntitySystemTask.h
+++ b/Engine/Source/Runtime/MovieScene/Public/EntitySystem/MovieSceneEntitySystemTask.h
@@ -558,7 +558,7 @@ struct TEntityTaskComponentsImpl<TIntegerSequence<int, Indices...>, T...>
 	/** Utility function called when the task is dispatched to populate the filter based on our component typs */
 	void PopulatePrerequisites(const FSystemTaskPrerequisites& InPrerequisites, FGraphEventArray* OutGatheredPrereqs) const
 	{
-		// Gather any master tasks
+		// Gather any root tasks
 		InPrerequisites.FilterByComponent(*OutGatheredPrereqs, FComponentTypeID::Invalid());
 
 		int Temp[] = { (UE::MovieScene::PopulatePrerequisites(&Accessors.template Get<Indices>(), InPrerequisites, OutGatheredPrereqs), 0)..., 0 };
@@ -568,7 +568,7 @@ struct TEntityTaskComponentsImpl<TIntegerSequence<int, Indices...>, T...>
 	/** Utility function called when the task is dispatched to populate the filter based on our component typs */
 	void PopulateSubsequents(const FGraphEventRef& InEvent, FSystemSubsequentTasks& OutSubsequents) const
 	{
-		OutSubsequents.AddMasterTask(InEvent);
+		OutSubsequents.AddRootTask(InEvent);
 
 		int Temp[] = { (UE::MovieScene::PopulateSubsequents(&Accessors.template Get<Indices>(), InEvent, OutSubsequents), 0)..., 0 };
 		(void)Temp;
diff --git a/Engine/Source/Runtime/MovieScene/Public/EntitySystem/MovieSceneEntitySystemTypes.h b/Engine/Source/Runtime/MovieScene/Public/EntitySystem/MovieSceneEntitySystemTypes.h
--- a/Engine/Source/Runtime/MovieScene/Public/EntitySystem/MovieSceneEntitySystemTypes.h
+++ b/Engine/Source/Runtime/MovieScene/Public/EntitySystem/MovieSceneEntitySystemTypes.h
@@ -95,7 +95,7 @@ ENUM_CLASS_FLAGS(EComplexFilterMode);
 
 /**
  * A numeric identifier used to represent a specific 'channel' within an interrogation linker.
- * Interrogation channels are used to identify groupings of tracks or entities that relate to the same output (eg: a property on an object; a master track etc).
+ * Interrogation channels are used to identify groupings of tracks or entities that relate to the same output (eg: a property on an object; a root track etc).
  * See FInterrogationKey for a combination of a channel with a specific interrogation index (or time)
  */
 struct FInterrogationChannel
diff --git a/Engine/Source/Runtime/MovieScene/Public/EntitySystem/MovieSceneRootInstantiatorSystem.h b/Engine/Source/Runtime/MovieScene/Public/EntitySystem/MovieSceneRootInstantiatorSystem.h
new file mode 100644
--- /dev/null
+++ b/Engine/Source/Runtime/MovieScene/Public/EntitySystem/MovieSceneRootInstantiatorSystem.h
@@ -0,0 +1,32 @@
+// Copyright Epic Games, Inc. All Rights Reserved.
+
+#pragma once
+
+#include "EntitySystem/MovieSceneEntityIDs.h"
+#include "EntitySystem/MovieSceneEntityInstantiatorSystem.h"
+#include "EntitySystem/MovieSceneEntitySystem.h"
+#include "EntitySystem/MovieSceneEntitySystemTypes.h"
+#include "UObject/ObjectMacros.h"
+#include "UObject/UObjectGlobals.h"
+
+#include "MovieSceneRootInstantiatorSystem.generated.h"
+
+class UObject;
+namespace UE { namespace MovieScene { struct FEntityAllocation; } }
+
+
+UCLASS()
+class MOVIESCENE_API UMovieSceneRootInstantiatorSystem : public UMovieSceneEntityInstantiatorSystem
+{
+	GENERATED_BODY()
+
+public:
+
+	UMovieSceneRootInstantiatorSystem(const FObjectInitializer& ObjInit);
+
+private:
+
+	virtual void OnRun(FSystemTaskPrerequisites& InPrerequisites, FSystemSubsequentTasks& Subsequents) override final;
+
+	void InstantiateAllocation(const UE::MovieScene::FEntityAllocation* ParentAllocation);
+};
\ No newline at end of file
diff --git a/Engine/Source/Runtime/MovieScene/Public/EntitySystem/MovieSceneSequenceInstance.h b/Engine/Source/Runtime/MovieScene/Public/EntitySystem/MovieSceneSequenceInstance.h
--- a/Engine/Source/Runtime/MovieScene/Public/EntitySystem/MovieSceneSequenceInstance.h
+++ b/Engine/Source/Runtime/MovieScene/Public/EntitySystem/MovieSceneSequenceInstance.h
@@ -113,7 +113,7 @@ public:
 	/**
 	 * Retrieve the SequenceID for this instance
 	 *
-	 * @return This sequence instance's SequenceID within the root-sequences hierachy, or MovieSceneSequenceID::Root for master sequence instances.
+	 * @return This sequence instance's SequenceID within the root-sequences hierachy, or MovieSceneSequenceID::Root for root sequence instances.
 	 */
 	FMovieSceneSequenceID GetSequenceID() const
 	{
diff --git a/Engine/Source/Runtime/MovieScene/Public/EntitySystem/MovieSceneSequenceUpdaters.h b/Engine/Source/Runtime/MovieScene/Public/EntitySystem/MovieSceneSequenceUpdaters.h
--- a/Engine/Source/Runtime/MovieScene/Public/EntitySystem/MovieSceneSequenceUpdaters.h
+++ b/Engine/Source/Runtime/MovieScene/Public/EntitySystem/MovieSceneSequenceUpdaters.h
@@ -116,7 +116,7 @@ public:
 	 * Override the sequence ID that should be considered the root sequence for this updater
 	 *
 	 * @param InLinker                    The linker that is owns this sequence
-	 * @param InstanceHandle              The instance handle for the master sequence
+	 * @param InstanceHandle              The instance handle for the root sequence
 	 * @param NewRootOverrideSequenceID   The new sequence ID to treat as the root
 	 */
 	virtual void OverrideRootSequence(UMovieSceneEntitySystemLinker* InLinker, FRootInstanceHandle InstanceHandle, FMovieSceneSequenceID NewRootOverrideSequenceID) = 0;
diff --git a/Engine/Source/Runtime/MovieScene/Public/EntitySystem/MovieSceneSystemTaskDependencies.h b/Engine/Source/Runtime/MovieScene/Public/EntitySystem/MovieSceneSystemTaskDependencies.h
--- a/Engine/Source/Runtime/MovieScene/Public/EntitySystem/MovieSceneSystemTaskDependencies.h
+++ b/Engine/Source/Runtime/MovieScene/Public/EntitySystem/MovieSceneSystemTaskDependencies.h
@@ -65,7 +65,7 @@ struct FSystemTaskPrerequisites
 
 	MOVIESCENE_API void FilterByComponent(FGraphEventArray& OutArray, std::initializer_list<FComponentTypeID> ComponentTypes) const;
 
-	void AddMasterTask(const FGraphEventRef& InNewTask)
+	void AddRootTask(const FGraphEventRef& InNewTask)
 	{
 		AddComponentTask(FComponentTypeID::Invalid(), InNewTask);
 	}
@@ -97,7 +97,7 @@ struct MOVIESCENE_API FSystemSubsequentTasks
 {
 	using FComponentTypeID = UE::MovieScene::FComponentTypeID;
 
-	void AddMasterTask(FGraphEventRef MasterTask);
+	void AddRootTask(FGraphEventRef RootTask);
 
 	void AddComponentTask(FComponentTypeID ComponentType, FGraphEventRef ComponentTask);
 
diff --git a/Engine/Source/Runtime/MovieScene/Public/EntitySystem/TrackInstance/MovieSceneTrackInstance.h b/Engine/Source/Runtime/MovieScene/Public/EntitySystem/TrackInstance/MovieSceneTrackInstance.h
--- a/Engine/Source/Runtime/MovieScene/Public/EntitySystem/TrackInstance/MovieSceneTrackInstance.h
+++ b/Engine/Source/Runtime/MovieScene/Public/EntitySystem/TrackInstance/MovieSceneTrackInstance.h
@@ -99,7 +99,7 @@ struct TStructOpsTypeTraits<FMovieSceneTrackInstanceInput> : public TStructOpsTy
 
 
 /**
- * Base class for all track instances. Can also be used for master tracks where AnimatedObject will be nullptr
+ * Base class for all track instances. Can also be used for root tracks where AnimatedObject will be nullptr
  */
 UCLASS(transient)
 class MOVIESCENE_API UMovieSceneTrackInstance : public UObject
@@ -138,7 +138,7 @@ public:
 
 
 	/**
-	 * Retrieve the object that is being animated (may be nullptr for master tracks)
+	 * Retrieve the object that is being animated (may be nullptr for root tracks)
 	 */
 	UObject* GetAnimatedObject() const
 	{
@@ -194,7 +194,7 @@ private:
 	TObjectPtr<UObject> AnimatedObject;
 
 	UPROPERTY()
-	bool bIsMasterTrackInstance;
+	bool bIsRootTrackInstance;
 
 	UPROPERTY()
 	TObjectPtr<UMovieSceneEntitySystemLinker> PrivateLinker;
diff --git a/Engine/Source/Runtime/MovieScene/Public/Evaluation/Blending/BlendableTokenStack.h b/Engine/Source/Runtime/MovieScene/Public/Evaluation/Blending/BlendableTokenStack.h
--- a/Engine/Source/Runtime/MovieScene/Public/Evaluation/Blending/BlendableTokenStack.h
+++ b/Engine/Source/Runtime/MovieScene/Public/Evaluation/Blending/BlendableTokenStack.h
@@ -51,7 +51,7 @@ public:
 	/**
 	 * Implemented by typed stacks to compute the final blended value for its data, and apply that result to the specified object
 	 *
-	 * @param InObject			The object to apply the result to. Nullptr for master track data.
+	 * @param InObject			The object to apply the result to. Nullptr for root track data.
 	 * @param Accumulator 		The accumulator that was used to gather and consolidate the blendable tokens
 	 * @param ActuatorTypeID	A unique identifier to the actuator that should be used to apply the result to the object
 	 * @param Context			Root-level movie scene evaluation context for the current frame
@@ -116,7 +116,7 @@ public:
 	/**
 	 * Implemented by typed stacks to compute the final blended value for its data, and apply that result to the specified object
 	 *
-	 * @param InObject			The object to apply the result to. Nullptr for master track data.
+	 * @param InObject			The object to apply the result to. Nullptr for root track data.
 	 * @param Accumulator 		The accumulator that was used to gather and consolidate the blendable tokens
 	 * @param ActuatorTypeID	A unique identifier to the actuator that should be used to apply the result to the object
 	 * @param Context			Root-level movie scene evaluation context for the current frame
diff --git a/Engine/Source/Runtime/MovieScene/Public/Evaluation/Blending/MovieSceneBlendingActuator.h b/Engine/Source/Runtime/MovieScene/Public/Evaluation/Blending/MovieSceneBlendingActuator.h
--- a/Engine/Source/Runtime/MovieScene/Public/Evaluation/Blending/MovieSceneBlendingActuator.h
+++ b/Engine/Source/Runtime/MovieScene/Public/Evaluation/Blending/MovieSceneBlendingActuator.h
@@ -43,7 +43,7 @@ public:
 	/**
 	 * Apply the the specified value to an object.
 	 *
-	 * @param InObject 		The object to apply the value to. nullptr where this actuator is being used for a master track.
+	 * @param InObject 		The object to apply the value to. nullptr where this actuator is being used for a root track.
 	 * @param InValue 		The value to apply to the object.
 	 * @param OriginalStack	Reference to the original stack from which the final result was derived.
 	 * @param Context		Movie scene context structure from the root level
@@ -65,7 +65,7 @@ public:
 	/**
 	 * Retrieve the current value of the specified object
 	 *
-	 * @param InObject 		The object to retrieve the value from. nullptr where this actuator is being used for a master track.
+	 * @param InObject 		The object to retrieve the value from. nullptr where this actuator is being used for a root track.
 	 * @param Player		The movie scene player currently running the sequence. Null when this actuator is used in an interrogation context.
 	 * @return The current value of the object
 	 */
diff --git a/Engine/Source/Runtime/MovieScene/Public/Evaluation/Blending/MovieSceneInitialValueStore.h b/Engine/Source/Runtime/MovieScene/Public/Evaluation/Blending/MovieSceneInitialValueStore.h
--- a/Engine/Source/Runtime/MovieScene/Public/Evaluation/Blending/MovieSceneInitialValueStore.h
+++ b/Engine/Source/Runtime/MovieScene/Public/Evaluation/Blending/MovieSceneInitialValueStore.h
@@ -49,7 +49,7 @@ struct TMovieSceneInitialValueStore
 	 * Constructor
 	 * @param InActuator 			The actuator that will apply the final value
 	 * @param InStack 				The stack of tokens that are currently being blended (used for saving pre-animated state)
-	 * @param InAnimatingObject 	(Optional) The object that is currently being animated (nullptr for master tracks)
+	 * @param InAnimatingObject 	(Optional) The object that is currently being animated (nullptr for root tracks)
 	 * @param InPlayer 				(Optional) The movie scene player currently playing the sequence (nullptr for data interrogation)
 	 */
 	TMovieSceneInitialValueStore(TMovieSceneBlendingActuator<DataType>& InActuator, const TBlendableTokenStack<DataType>& InStack, UObject* InAnimatingObject, IMovieScenePlayer* InPlayer)
@@ -103,8 +103,8 @@ private:
 	TMovieSceneBlendingActuator<DataType>& Actuator;
 	/** The stack of tokens that are being applied */
 	const TBlendableTokenStack<DataType>& Stack;
-	/** The object that is being animated (nullptr for master track animation) */
+	/** The object that is being animated (nullptr for root track animation) */
 	UObject* AnimatingObject;
-	/** Player that's playing back the sequence. Can be null. Potentially used when accessing values for master tracks. */
+	/** Player that's playing back the sequence. Can be null. Potentially used when accessing values for root tracks. */
 	IMovieScenePlayer* Player;
 };
diff --git a/Engine/Source/Runtime/MovieScene/Public/Evaluation/MovieSceneEvaluationState.h b/Engine/Source/Runtime/MovieScene/Public/Evaluation/MovieSceneEvaluationState.h
--- a/Engine/Source/Runtime/MovieScene/Public/Evaluation/MovieSceneEvaluationState.h
+++ b/Engine/Source/Runtime/MovieScene/Public/Evaluation/MovieSceneEvaluationState.h
@@ -64,7 +64,7 @@ struct FMovieSceneObjectCache
 	 * Set the sequence that this cache applies to
 	 *
 	 * @param InSequence		The sequence that this cache applies to
-	 * @param InSequenceID		The ID of the sequence within the master sequence
+	 * @param InSequenceID		The ID of the sequence within the root sequence
 	 */
 	MOVIESCENE_API void SetSequence(UMovieSceneSequence& InSequence, FMovieSceneSequenceIDRef InSequenceID, IMovieScenePlayer& Player);
 
@@ -166,7 +166,7 @@ private:
 	/** The sequence that we're caching objects for */
 	TWeakObjectPtr<UMovieSceneSequence> WeakSequence;
 
-	/** The sequence ID of the sequence within the master sequence */
+	/** The sequence ID of the sequence within the root sequence */
 	FMovieSceneSequenceID SequenceID;
 
 	template<typename ValueType>
diff --git a/Engine/Source/Runtime/MovieScene/Public/Evaluation/MovieSceneEvaluationTrack.h b/Engine/Source/Runtime/MovieScene/Public/Evaluation/MovieSceneEvaluationTrack.h
--- a/Engine/Source/Runtime/MovieScene/Public/Evaluation/MovieSceneEvaluationTrack.h
+++ b/Engine/Source/Runtime/MovieScene/Public/Evaluation/MovieSceneEvaluationTrack.h
@@ -391,7 +391,7 @@ public:
 
 private:
 
-	/** ID of the possessable or spawnable within the UMovieScene this track belongs to, if any. Zero guid where this relates to a master track. */
+	/** ID of the possessable or spawnable within the UMovieScene this track belongs to, if any. Zero guid where this relates to a root track. */
 	UPROPERTY()
 	FGuid ObjectBindingID;
 
diff --git a/Engine/Source/Runtime/MovieScene/Public/Evaluation/MovieSceneSequenceHierarchy.h b/Engine/Source/Runtime/MovieScene/Public/Evaluation/MovieSceneSequenceHierarchy.h
--- a/Engine/Source/Runtime/MovieScene/Public/Evaluation/MovieSceneSequenceHierarchy.h
+++ b/Engine/Source/Runtime/MovieScene/Public/Evaluation/MovieSceneSequenceHierarchy.h
@@ -236,7 +236,7 @@ struct FMovieSceneSubSequenceTree
 template<> struct TStructOpsTypeTraits<FMovieSceneSubSequenceTree> : public TStructOpsTypeTraitsBase2<FMovieSceneSubSequenceTree> { enum { WithSerializer = true, WithIdenticalViaEquality = true }; };
 
 /**
- * Structure that stores hierarchical information pertaining to all sequences contained within a master sequence
+ * Structure that stores hierarchical information pertaining to all sequences contained within a root sequence
  */
 USTRUCT()
 struct FMovieSceneSequenceHierarchy
diff --git a/Engine/Source/Runtime/MovieScene/Public/Evaluation/PersistentEvaluationData.h b/Engine/Source/Runtime/MovieScene/Public/Evaluation/PersistentEvaluationData.h
--- a/Engine/Source/Runtime/MovieScene/Public/Evaluation/PersistentEvaluationData.h
+++ b/Engine/Source/Runtime/MovieScene/Public/Evaluation/PersistentEvaluationData.h
@@ -65,7 +65,7 @@ struct FSharedPersistentDataKey
 	/** The actual shared ID */
 	FMovieSceneSharedDataId UniqueId;
 
-	/** The operand that this key relates to (may be invalid where the data pertains to master tracks) */
+	/** The operand that this key relates to (may be invalid where the data pertains to root tracks) */
 	FMovieSceneEvaluationOperand Operand;
 };
 
diff --git a/Engine/Source/Runtime/MovieScene/Public/Evaluation/PreAnimatedState/MovieScenePreAnimatedStateExtension.h b/Engine/Source/Runtime/MovieScene/Public/Evaluation/PreAnimatedState/MovieScenePreAnimatedStateExtension.h
--- a/Engine/Source/Runtime/MovieScene/Public/Evaluation/PreAnimatedState/MovieScenePreAnimatedStateExtension.h
+++ b/Engine/Source/Runtime/MovieScene/Public/Evaluation/PreAnimatedState/MovieScenePreAnimatedStateExtension.h
@@ -33,7 +33,7 @@ namespace UE
 namespace MovieScene
 {
 
-struct FAnimTypePreAnimatedStateMasterStorage;
+struct FAnimTypePreAnimatedStateRootStorage;
 struct FAnimTypePreAnimatedStateObjectStorage;
 struct FPreAnimatedEntityCaptureSource;
 struct FPreAnimatedStateExtension;
@@ -322,8 +322,8 @@ private:
 
 	/** Pointers to the storage for state bound to objects, organized by FMovieSceneAnimTypeID */
 	TWeakPtr<FAnimTypePreAnimatedStateObjectStorage> WeakGenericObjectStorage;
-	/** Pointers to the storage for state created from master tracks, or otherwise not bound to objects */
-	TWeakPtr<FAnimTypePreAnimatedStateMasterStorage> WeakGenericMasterStorage;
+	/** Pointers to the storage for state created from root tracks, or otherwise not bound to objects */
+	TWeakPtr<FAnimTypePreAnimatedStateRootStorage> WeakGenericRootStorage;
 
 private:
 
diff --git a/Engine/Source/Runtime/MovieScene/Public/EventHandlers/ISequenceDataEventHandler.h b/Engine/Source/Runtime/MovieScene/Public/EventHandlers/ISequenceDataEventHandler.h
--- a/Engine/Source/Runtime/MovieScene/Public/EventHandlers/ISequenceDataEventHandler.h
+++ b/Engine/Source/Runtime/MovieScene/Public/EventHandlers/ISequenceDataEventHandler.h
@@ -21,9 +21,9 @@ class ISequenceDataEventHandler
 {
 public:
 
-	virtual void OnMasterTrackAdded(UMovieSceneTrack*)     {}
+	virtual void OnTrackAdded(UMovieSceneTrack*)           {}
 
-	virtual void OnMasterTrackRemoved(UMovieSceneTrack*)   {}
+	virtual void OnTrackRemoved(UMovieSceneTrack*)         {}
 
 	virtual void OnBindingAdded(const FMovieSceneBinding&) {}
 
diff --git a/Engine/Source/Runtime/MovieScene/Public/MovieScene.h b/Engine/Source/Runtime/MovieScene/Public/MovieScene.h
--- a/Engine/Source/Runtime/MovieScene/Public/MovieScene.h
+++ b/Engine/Source/Runtime/MovieScene/Public/MovieScene.h
@@ -621,89 +621,106 @@ public:
 public:
 
 	/**
-	 * Adds a master track.
+	 * Adds a track.
 	 *
 	 * Note: The type should not already exist.
 	 *
 	 * @param TrackClass The class of the track to create
 	 * @param Type	The newly created type
-	 * @see FindMasterTrack, GetMasterTracks, IsMasterTrack, RemoveMasterTrack
+	 * @see FindTrack, GetTracks, IsTrack, RemoveTrack
 	 */
-	UMovieSceneTrack* AddMasterTrack(TSubclassOf<UMovieSceneTrack> TrackClass);
+	UMovieSceneTrack* AddTrack(TSubclassOf<UMovieSceneTrack> TrackClass);
 	
+	UE_DEPRECATED(5.2, "AddMasterTrack is deprecated. Please use AddTrack instead")
+	UMovieSceneTrack* AddMasterTrack(TSubclassOf<UMovieSceneTrack> TrackClass) { return AddTrack(TrackClass); }
+
 	/**
-	 * Adds a master track.
+	 * Adds a track.
 	 *
 	 * Note: The type should not already exist.
 	 *
 	 * @param TrackClass The class of the track to create
 	 * @param Type	The newly created type
-	 * @see FindMasterTrack, GetMasterTracks, IsMasterTrack, RemoveMasterTrack
+	 * @see FindTrack, GetTracks, IsTrack, RemoveTrack
 	 */
 	template<typename TrackClass>
-	TrackClass* AddMasterTrack()
+	TrackClass* AddTrack()
 	{
-		return Cast<TrackClass>(AddMasterTrack(TrackClass::StaticClass()));
+		return Cast<TrackClass>(AddTrack(TrackClass::StaticClass()));
 	}
 
+	template<typename TrackClass>
+	UE_DEPRECATED(5.2, "AddMasterTrack is deprecated. Please use AddTrack instead")
+	TrackClass* AddMasterTrack() { return AddTrack<TrackClass>(); }
+
 	/**
-	* Adds a given track as a master track
+	* Adds a given track as a track
 	*
 	* @param InTrack The track to add.
 	* @see  FindTrack, RemoveTrack
 	* @return true if the track is successfully added, false otherwise.
 	*/
-	bool AddGivenMasterTrack(UMovieSceneTrack* InTrack);
+	bool AddGivenTrack(UMovieSceneTrack* InTrack);
+
+	UE_DEPRECATED(5.2, "AddGivenMasterTrack is deprecated. Please use AddGivenTrack instead")
+	bool AddGivenMasterTrack(UMovieSceneTrack* InTrack) { return AddGivenTrack(InTrack); }
 
 	/**
-	 * Finds a master track (one not bound to a runtime objects).
+	 * Finds a track (one not bound to a runtime objects).
 	 *
 	 * @param TrackClass The class of the track to find.
 	 * @return The found track or nullptr if one does not exist.
-	 * @see AddMasterTrack, GetMasterTracks, IsMasterTrack, RemoveMasterTrack
+	 * @see AddTrack, GetTracks, IsTrack, RemoveTrack
 	 */
-	UMovieSceneTrack* FindMasterTrack(TSubclassOf<UMovieSceneTrack> TrackClass) const;
+	UMovieSceneTrack* FindTrack(TSubclassOf<UMovieSceneTrack> TrackClass) const;
+
+	UE_DEPRECATED(5.2, "FindMasterTrack is deprecated. Please use FindTrack instead")
+	UMovieSceneTrack* FindMasterTrack(TSubclassOf<UMovieSceneTrack> TrackClass) const { return FindTrack(TrackClass); }
 
 	/**
-	 * Finds a master track (one not bound to a runtime objects).
+	 * Finds a track (one not bound to a runtime objects).
 	 *
 	 * @param TrackClass The class of the track to find.
 	 * @return The found track or nullptr if one does not exist.
-	 * @see AddMasterTrack, GetMasterTracks, IsMasterTrack, RemoveMasterTrack
+	 * @see AddTrack, GetTracks, IsTrack, RemoveTrack
 	 */
 	template<typename TrackClass>
-	TrackClass* FindMasterTrack() const
+	TrackClass* FindTrack() const
 	{
-		return Cast<TrackClass>(FindMasterTrack(TrackClass::StaticClass()));
+		return Cast<TrackClass>(FindTrack(TrackClass::StaticClass()));
 	}
 
+	template<typename TrackClass>
+	UE_DEPRECATED(5.2, "FindMasterTrack is deprecated. Please use FindTrack instead")
+	TrackClass* FindMasterTrack() const { return FindTrack<TrackClass>(); }
+
 	/**
-	 * Get all master tracks.
+	 * Get all tracks.
 	 *
 	 * @return Track collection.
-	 * @see AddMasterTrack, FindMasterTrack, IsMasterTrack, RemoveMasterTrack
+	 * @see AddTrack, FindTrack, IsTrack, RemoveTrack
 	 */
-	const TArray<UMovieSceneTrack*>& GetMasterTracks() const
+	const TArray<UMovieSceneTrack*>& GetTracks() const
 	{
-		return MasterTracks;
+		return Tracks;
 	}
 
-	/**
-	 * Check whether the specified track is a master track in this scene.
-	 *
-	 * @return true if the track is a master track, false otherwise.
-	 * @see AddMasterTrack, FindMasterTrack, GetMasterTracks, RemoveMasterTrack
-	 */
-	bool IsAMasterTrack(const UMovieSceneTrack& Track) const;
+	UE_DEPRECATED(5.2, "GetMasterTracks is deprecated. Please use GetTracks instead")
+	const TArray<UMovieSceneTrack*>& GetMasterTracks() const { return GetTracks(); }
 
 	/**
-	 * Removes a master track.
+	 * Check whether the specified track is a track in this movie scene.
 	 *
-	 * @param Track The track to remove.
-	 * @return true if anything was removed.
-	 * @see AddMasterTrack, FindMasterTrack, GetMasterTracks, IsMasterTrack
+	 * @return true if the track is a track, false otherwise.
+	 * @see AddTrack, FindTrack, GetTracks, RemoveTrack
 	 */
-	bool RemoveMasterTrack(UMovieSceneTrack& Track);
+	bool ContainsTrack(const UMovieSceneTrack& Track) const;
+
+	UE_DEPRECATED(5.2, "IsAMasterTrack is deprecated. Please use ContainsTrack instead")
+	bool IsAMasterTrack(const UMovieSceneTrack& Track) const { return ContainsTrack(Track); }
+
+	UE_DEPRECATED(5.2, "RemoveMasterTrack is deprecated. Please use RemoveTrack instead")
+	bool RemoveMasterTrack(UMovieSceneTrack& Track) { return RemoveTrack(Track); }
 
 	/**
 	 * Move all the contents (tracks, child bindings) of the specified binding ID onto another
@@ -1230,9 +1247,9 @@ private:
 	UPROPERTY()
 	TMap<FName, FMovieSceneObjectBindingIDs> BindingGroups;
 
-	/** Master tracks which are not bound to spawned or possessed objects */
+	/** Tracks which are not bound to spawned or possessed objects */
 	UPROPERTY(Instanced)
-	TArray<TObjectPtr<UMovieSceneTrack>> MasterTracks;
+	TArray<TObjectPtr<UMovieSceneTrack>> Tracks;
 
 	/** The camera cut track is a specialized track for switching between cameras on a cinematic */
 	UPROPERTY(Instanced)
@@ -1333,10 +1350,11 @@ private:
 	UPROPERTY()
 	float FixedFrameInterval_DEPRECATED;
 
+	UPROPERTY()
+	TArray<TObjectPtr<UMovieSceneTrack>> MasterTracks_DEPRECATED;
+
 	//delegates
 	private:
 	FMovieSceneOnChannelChanged OnChannelChangedDelegate;
 #endif
-
-		
 };
diff --git a/Engine/Source/Runtime/MovieScene/Public/MovieSceneExecutionToken.h b/Engine/Source/Runtime/MovieScene/Public/MovieSceneExecutionToken.h
--- a/Engine/Source/Runtime/MovieScene/Public/MovieSceneExecutionToken.h
+++ b/Engine/Source/Runtime/MovieScene/Public/MovieSceneExecutionToken.h
@@ -69,7 +69,7 @@ struct IMovieScenePreAnimatedGlobalToken
 	virtual ~IMovieScenePreAnimatedGlobalToken() {}
 
 	/**
-	 * Restore global state for a master track.
+	 * Restore global state for a root track.
 	 *
 	 * @param Params Parameters for restoring state
 	 */
diff --git a/Engine/Source/Runtime/MovieScene/Public/MovieSceneFolder.h b/Engine/Source/Runtime/MovieScene/Public/MovieSceneFolder.h
--- a/Engine/Source/Runtime/MovieScene/Public/MovieSceneFolder.h
+++ b/Engine/Source/Runtime/MovieScene/Public/MovieSceneFolder.h
@@ -47,17 +47,29 @@ class MOVIESCENE_API UMovieSceneFolder : public UObject
 	/** Removes a child folder from this folder. Automatically calls Modify() on the folder object. */
 	void RemoveChildFolder( UMovieSceneFolder* InChildFolder );
 
-	/** Gets the master tracks contained by this folder. */
-	const TArray<UMovieSceneTrack*>& GetChildMasterTracks() const;
+	/** Gets the tracks contained by this folder. */
+	const TArray<UMovieSceneTrack*>& GetChildTracks() const;
 
-	/** Adds a master track to this folder. Automatically calls Modify() on the folder object. */
-	void AddChildMasterTrack( UMovieSceneTrack* InMasterTrack );
+	UE_DEPRECATED(5.2, "GetChildMasterTracks is deprecated. Please use GetChildTracks instead")
+	const TArray<UMovieSceneTrack*>& GetChildMasterTracks() const { return GetChildTracks(); }
 
-	/** Removes a master track from this folder. Automatically calls Modify() on the folder object. */
-	void RemoveChildMasterTrack( UMovieSceneTrack* InMasterTrack );
+	/** Adds a track to this folder. Automatically calls Modify() on the folder object. */
+	void AddChildTrack( UMovieSceneTrack* InTrack );
 
-	/** Clear all child master tracks from this folder. */
-	void ClearChildMasterTracks();
+	UE_DEPRECATED(5.2, "AddChildMasterTrack is deprecated. Please use AddChildTrack instead")
+	void AddChildMasterTrack(UMovieSceneTrack* InMasterTrack) { AddChildTrack(InMasterTrack); }
+
+	/** Removes a track from this folder. Automatically calls Modify() on the folder object. */
+	void RemoveChildTrack( UMovieSceneTrack* InTrack );
+
+	UE_DEPRECATED(5.2, "RemoveChildMasterTrack is deprecated. Please use RemoveChildTrack instead")
+	void RemoveChildMasterTrack(UMovieSceneTrack* InMasterTrack) { return RemoveChildTrack(InMasterTrack); }
+
+	/** Clear all child tracks from this folder. */
+	void ClearChildTracks();
+
+	UE_DEPRECATED(5.2, "ClearChildMasterTracks is deprecated. Please use ClearChildTracks instead")
+	void ClearChildMasterTracks() { ClearChildTracks(); }
 
 	/** Gets the guids for the object bindings contained by this folder. */
 	const TArray<FGuid>& GetChildObjectBindings() const;
@@ -149,9 +161,9 @@ private:
 	UPROPERTY()
 	TArray<TObjectPtr<UMovieSceneFolder>> ChildFolders;
 
-	/** The master tracks contained by this folder. */
+	/** The tracks contained by this folder. */
 	UPROPERTY()
-	TArray<TObjectPtr<UMovieSceneTrack>> ChildMasterTracks;
+	TArray<TObjectPtr<UMovieSceneTrack>> ChildTracks;
 
 	/** The guid strings used to serialize the guids for the object bindings contained by this folder. */
 	UPROPERTY()
@@ -169,6 +181,11 @@ private:
 
 	/** The guids for the object bindings contained by this folder. */
 	TArray<FGuid> ChildObjectBindings;
+
+#if WITH_EDITORONLY_DATA
+	UPROPERTY()
+	TArray<TObjectPtr<UMovieSceneTrack>> ChildMasterTracks_DEPRECATED;
+#endif
 };
 
 MOVIESCENE_API void GetMovieSceneFoldersRecursive(TArrayView<UMovieSceneFolder* const> InFoldersToRecurse, TArray<UMovieSceneFolder*>& OutFolders);
