Sequencer: clear the RF_Transient flag on channel overrides and constraints after a copy/paste

diff --git a/Engine/Source/Editor/Sequencer/Private/SequencerUtilities.cpp b/Engine/Source/Editor/Sequencer/Private/SequencerUtilities.cpp
--- a/Engine/Source/Editor/Sequencer/Private/SequencerUtilities.cpp
+++ b/Engine/Source/Editor/Sequencer/Private/SequencerUtilities.cpp
@@ -71,6 +71,17 @@ static EVisibility GetRolloverVisibility(TAttribute<bool> HoverState, TWeakPtr<S
 	}
 }
 
+static void ResetCopiedTracksFlags(UMovieSceneTrack* Track)
+{
+	TArray<UObject*> SectionSubObjects;
+	Track->ClearFlags(RF_Transient);
+	for (UMovieSceneSection* Section : Track->GetAllSections())
+	{
+		Section->ClearFlags(RF_Transient);
+		Section->PostPaste();
+	}
+}
+
 TSharedRef<SWidget> FSequencerUtilities::MakeAddButton(FText HoverText, FOnGetContent MenuContent, const TAttribute<bool>& HoverState, TWeakPtr<ISequencer> InSequencer)
 {
 	TAttribute<bool> IsEnabled = MakeAttributeLambda([InSequencer]() -> bool { return InSequencer.IsValid() ? !InSequencer.Pin()->IsReadOnly() : false; });
@@ -1499,11 +1510,7 @@ bool FSequencerUtilities::PasteTracks(const FString& TextToImport, FMovieScenePa
 			if (!CopyableTrack->bIsRootTrack && !CopyableTrack->bIsCameraCutTrack)
 			{
 				UMovieSceneTrack* NewTrack = CopyableTrack->Track;
-				NewTrack->ClearFlags(RF_Transient);
-				for (UMovieSceneSection* Section : NewTrack->GetAllSections())
-				{
-					Section->ClearFlags(RF_Transient);
-				}
+				ResetCopiedTracksFlags(NewTrack);
 
 				// Remove tracks with the same name before adding
 				for (const FMovieSceneBinding& Binding : MovieScene->GetBindings())
@@ -1545,11 +1552,7 @@ bool FSequencerUtilities::PasteTracks(const FString& TextToImport, FMovieScenePa
 			if (CopyableTrack->bIsRootTrack || CopyableTrack->bIsCameraCutTrack)
 			{
 				UMovieSceneTrack* NewTrack = CopyableTrack->Track;
-				NewTrack->ClearFlags(RF_Transient);
-				for (UMovieSceneSection* Section : NewTrack->GetAllSections())
-				{
-					Section->ClearFlags(RF_Transient);
-				}
+				ResetCopiedTracksFlags(NewTrack);
 
 				UMovieSceneFolder* ParentFolder = PasteTracksParams.ParentFolder;
 
@@ -1736,6 +1739,7 @@ bool FSequencerUtilities::PasteSections(const FString& TextToImport, FMovieScene
 			Track->Modify();
 
 			Section->ClearFlags(RF_Transient);
+			Section->PostPaste();
 			Section->Rename(nullptr, Track);
 			Track->AddSection(*Section);
 			if (Section->HasStartFrame())
@@ -2049,16 +2053,13 @@ bool FSequencerUtilities::PasteBindings(const FString& TextToImport, TSharedRef<
 			return false;
 		}
 
+		TArray<UObject*> SectionSubObjects;
 		for (UMovieSceneCopyableBinding* CopyableBinding : ImportedBindings)
 		{
 			// Clear transient flags on the imported tracks
 			for (UMovieSceneTrack* CopiedTrack : CopyableBinding->Tracks)
 			{
-				CopiedTrack->ClearFlags(RF_Transient);
-				for (UMovieSceneSection* Section : CopiedTrack->GetAllSections())
-                {
-					Section->ClearFlags(RF_Transient);
-                }
+				ResetCopiedTracksFlags(CopiedTrack);
 			}
 
 			UMovieSceneFolder* ParentFolder = PasteBindingsParams.ParentFolder;
@@ -2182,11 +2183,7 @@ bool FSequencerUtilities::PasteBindings(const FString& TextToImport, TSharedRef<
 				// Clear the transient flags on the copyable binding before assigning to the new spawnable
 				for (auto Track : NewBinding.GetTracks())
 				{
-					Track->ClearFlags(RF_Transient);
-					for (auto Section : Track->GetAllSections())
-					{
-						Section->ClearFlags(RF_Transient);
-					}
+					ResetCopiedTracksFlags(Track);
 				}
 
 				// Replace the auto-generated binding with our deserialized bindings (which has our tracks)
diff --git a/Engine/Source/Runtime/MovieScene/Private/Channels/MovieSceneSectionChannelOverrideRegistry.cpp b/Engine/Source/Runtime/MovieScene/Private/Channels/MovieSceneSectionChannelOverrideRegistry.cpp
--- a/Engine/Source/Runtime/MovieScene/Private/Channels/MovieSceneSectionChannelOverrideRegistry.cpp
+++ b/Engine/Source/Runtime/MovieScene/Private/Channels/MovieSceneSectionChannelOverrideRegistry.cpp
@@ -103,3 +103,20 @@ void UMovieSceneSectionChannelOverrideRegistry::PopulateEvaluationFieldImpl(cons
 	}
 }
 
+#if WITH_EDITOR
+
+void UMovieSceneSectionChannelOverrideRegistry::OnPostPaste()
+{
+	ClearFlags(RF_Transient);
+
+	for (TPair<FName, TObjectPtr<UMovieSceneChannelOverrideContainer>> Pair : Overrides)
+	{
+		if (ensure(Pair.Value))
+		{
+			Pair.Value->ClearFlags(RF_Transient);
+		}
+	}
+}
+
+#endif
+
diff --git a/Engine/Source/Runtime/MovieScene/Public/Channels/MovieSceneSectionChannelOverrideRegistry.h b/Engine/Source/Runtime/MovieScene/Public/Channels/MovieSceneSectionChannelOverrideRegistry.h
--- a/Engine/Source/Runtime/MovieScene/Public/Channels/MovieSceneSectionChannelOverrideRegistry.h
+++ b/Engine/Source/Runtime/MovieScene/Public/Channels/MovieSceneSectionChannelOverrideRegistry.h
@@ -75,7 +75,7 @@ public:
 	UMovieSceneChannelOverrideContainer* GetChannel(FName ChannelName) const;
 
 	/**
-	 *
+	 * Get all channel containers of a given type.
 	 */
 	template<typename ChannelContainerType>
 	void GetChannels(TArray<ChannelContainerType*>& OutChannels) const
@@ -106,6 +106,15 @@ public:
 	*/
 	void PopulateEvaluationFieldImpl(const TRange<FFrameNumber>& EffectiveRange, const FMovieSceneEvaluationFieldEntityMetaData& InMetaData, FMovieSceneEntityComponentFieldBuilder* OutFieldBuilder, UMovieSceneSection& OwnerSection);
 
+#if WITH_EDITOR
+
+	/**
+	 * Called by the owning section after it has been created for a paste operation.
+	 */
+	void OnPostPaste();
+
+#endif
+
 private:
 
 	/** Map of channel overrides. */
diff --git a/Engine/Source/Runtime/MovieScene/Public/MovieSceneSection.h b/Engine/Source/Runtime/MovieScene/Public/MovieSceneSection.h
--- a/Engine/Source/Runtime/MovieScene/Public/MovieSceneSection.h
+++ b/Engine/Source/Runtime/MovieScene/Public/MovieSceneSection.h
@@ -663,6 +663,8 @@ public:
 
 #if WITH_EDITOR
 	MOVIESCENE_API virtual void PostEditChangeProperty(FPropertyChangedEvent& PropertyChangedEvent) override;
+
+	MOVIESCENE_API virtual void PostPaste() {}
 #endif
 
 #if WITH_EDITOR
diff --git a/Engine/Source/Runtime/MovieSceneTracks/Private/Sections/MovieScene3DTransformSection.cpp b/Engine/Source/Runtime/MovieSceneTracks/Private/Sections/MovieScene3DTransformSection.cpp
--- a/Engine/Source/Runtime/MovieSceneTracks/Private/Sections/MovieScene3DTransformSection.cpp
+++ b/Engine/Source/Runtime/MovieSceneTracks/Private/Sections/MovieScene3DTransformSection.cpp
@@ -1090,6 +1090,20 @@ void UMovieScene3DTransformSection::OnChannelOverridesChanged()
 
 #if WITH_EDITOR
 
+void UMovieScene3DTransformSection::PostPaste()
+{
+	Super::PostPaste();
+	
+	if (OverrideRegistry)
+	{
+		OverrideRegistry->OnPostPaste();
+	}
+	if (Constraints)
+	{
+		Constraints->ClearFlags(RF_Transient);
+	}
+}
+
 void UMovieScene3DTransformSectionConstraints::PostEditUndo()
 {
 	Super::PostEditUndo();
diff --git a/Engine/Source/Runtime/MovieSceneTracks/Private/Sections/MovieSceneFloatSection.cpp b/Engine/Source/Runtime/MovieSceneTracks/Private/Sections/MovieSceneFloatSection.cpp
--- a/Engine/Source/Runtime/MovieSceneTracks/Private/Sections/MovieSceneFloatSection.cpp
+++ b/Engine/Source/Runtime/MovieSceneTracks/Private/Sections/MovieSceneFloatSection.cpp
@@ -93,3 +93,17 @@ void UMovieSceneFloatSection::OnChannelOverridesChanged()
 	ChannelProxy = nullptr;
 }
 
+#if WITH_EDITOR
+
+void UMovieSceneFloatSection::PostPaste()
+{
+	Super::PostPaste();
+	
+	if (OverrideRegistry)
+	{
+		OverrideRegistry->OnPostPaste();
+	}
+}
+
+#endif
+
diff --git a/Engine/Source/Runtime/MovieSceneTracks/Public/Sections/MovieScene3DTransformSection.h b/Engine/Source/Runtime/MovieSceneTracks/Public/Sections/MovieScene3DTransformSection.h
--- a/Engine/Source/Runtime/MovieSceneTracks/Public/Sections/MovieScene3DTransformSection.h
+++ b/Engine/Source/Runtime/MovieSceneTracks/Public/Sections/MovieScene3DTransformSection.h
@@ -235,6 +235,10 @@ public:
 	MOVIESCENETRACKS_API virtual void GetReferencedBindings(TArray<FGuid>& OutBindings) override;
 	MOVIESCENETRACKS_API virtual void PreSave(FObjectPreSaveContext SaveContext) override;
 
+#if WITH_EDITOR
+	MOVIESCENETRACKS_API virtual void PostPaste() override;
+#endif
+
 public:
 
 	/**
diff --git a/Engine/Source/Runtime/MovieSceneTracks/Public/Sections/MovieSceneFloatSection.h b/Engine/Source/Runtime/MovieSceneTracks/Public/Sections/MovieSceneFloatSection.h
--- a/Engine/Source/Runtime/MovieSceneTracks/Public/Sections/MovieSceneFloatSection.h
+++ b/Engine/Source/Runtime/MovieSceneTracks/Public/Sections/MovieSceneFloatSection.h
@@ -55,5 +55,9 @@ private:
 	virtual bool PopulateEvaluationFieldImpl(const TRange<FFrameNumber>& EffectiveRange, const FMovieSceneEvaluationFieldEntityMetaData& InMetaData, FMovieSceneEntityComponentFieldBuilder* OutFieldBuilder) override;
 
 	virtual EMovieSceneChannelProxyType CacheChannelProxy() override;
+
+#if WITH_EDITOR
+	MOVIESCENETRACKS_API virtual void PostPaste() override;
+#endif
 };
 
