Sequencer: Prevent infinite sections from being pasted on top of each other.

diff --git a/Engine/Source/Editor/Sequencer/Private/SequencerUtilities.cpp b/Engine/Source/Editor/Sequencer/Private/SequencerUtilities.cpp
--- a/Engine/Source/Editor/Sequencer/Private/SequencerUtilities.cpp
+++ b/Engine/Source/Editor/Sequencer/Private/SequencerUtilities.cpp
@@ -1741,6 +1741,17 @@ bool FSequencerUtilities::PasteSections(const FString& TextToImport, FMovieScene
 			Section->ClearFlags(RF_Transient);
 			Section->PostPaste();
 			Section->Rename(nullptr, Track);
+
+			if (Track->SupportsMultipleRows())
+			{
+				Section->SetRowIndex(RowIndex);
+			}
+			else if (!Section->HasStartFrame() && !Section->HasEndFrame())
+			{
+				// If the track doesn't support multiple rows and the pasted section is infinite, it should win out over existing sections
+				Track->RemoveAllAnimationData();
+			}
+
 			Track->AddSection(*Section);
 			if (Section->HasStartFrame())
 			{
@@ -1748,10 +1759,6 @@ bool FSequencerUtilities::PasteSections(const FString& TextToImport, FMovieScene
 				Section->MoveSection(NewStartFrame - Section->GetInclusiveStartFrame());
 			}
 
-			if (Track->SupportsMultipleRows())
-			{
-				Section->SetRowIndex(RowIndex);
-			}
 			OutSections.Add(Section);
 		}
 
