Sequencer: Spacer to match TrackView's color strip

#rnx

diff --git a/Engine/Source/Editor/Sequencer/Private/MVVM/Views/SOutlinerObjectBindingView.cpp b/Engine/Source/Editor/Sequencer/Private/MVVM/Views/SOutlinerObjectBindingView.cpp
--- a/Engine/Source/Editor/Sequencer/Private/MVVM/Views/SOutlinerObjectBindingView.cpp
+++ b/Engine/Source/Editor/Sequencer/Private/MVVM/Views/SOutlinerObjectBindingView.cpp
@@ -15,6 +15,7 @@
 #include "EditorStyleSet.h"
 #include "GameFramework/Actor.h"
 #include "Styling/SlateIconFinder.h"
+#include "Widgets/Layout/SBox.h"
 #include "Widgets/Layout/SSpacer.h"
 #include "Framework/Commands/UIAction.h"
 #include "Framework/MultiBox/MultiBoxBuilder.h"
@@ -60,6 +61,12 @@ void SOutlinerObjectBindingView::Construct(
 		.CustomContent()
 		[
 			CustomBoxPanel
+		]
+		.RightGutterContent()
+		[
+			// Spacer to match TrackView's color strip
+			SNew(SBox)
+			.WidthOverride(6.f)
 		], TWeakViewModelPtr<IOutlinerExtension>(InModel), InEditor, InTableRow
 	);
 }
