Sequencer: Save the tree view width percentage in sequencer settings so that the split between the tree view and the track area is saved/restored.

#rnx

diff --git a/Engine/Source/Editor/Sequencer/Private/SSequencer.cpp b/Engine/Source/Editor/Sequencer/Private/SSequencer.cpp
--- a/Engine/Source/Editor/Sequencer/Private/SSequencer.cpp
+++ b/Engine/Source/Editor/Sequencer/Private/SSequencer.cpp
@@ -114,6 +114,8 @@
 #include "Tracks/MovieSceneEventTrack.h"
 #include "ToolMenus.h"
 #include "MovieSceneToolHelpers.h"
+#include "Editor/UnrealEdEngine.h"
+#include "UnrealEdGlobals.h"
 
 #define LOCTEXT_NAMESPACE "Sequencer"
 
@@ -256,6 +258,15 @@ void SSequencer::Construct(const FArguments& InArgs, TSharedRef<FSequencer> InSe
 	RootCustomization.AddMenuExtender = InArgs._AddMenuExtender;
 	RootCustomization.ToolbarExtender = InArgs._ToolbarExtender;
 
+	ColumnFillCoefficients[0] = 0.3f;
+	ColumnFillCoefficients[1] = 0.7f;
+
+	if (GetSequencerSettings())
+	{
+		ColumnFillCoefficients[0] = GetSequencerSettings()->GetTreeViewWidth();
+		ColumnFillCoefficients[1] = 1.f - GetSequencerSettings()->GetTreeViewWidth();
+	}
+
 	TAttribute<float> FillCoefficient_0, FillCoefficient_1;
 	{
 		FillCoefficient_0.Bind(TAttribute<float>::FGetter::CreateSP(this, &SSequencer::GetColumnFillCoefficient, 0));
@@ -2917,18 +2928,6 @@ void SSequencer::OnOutlinerSearchChanged( const FText& Filter )
 	}
 }
 
-float SSequencer::GetColumnFillCoefficient(int32 ColumnIndex) const
-{
-	if (ColumnIndex == 0)
-	{
-		return GetSequencerSettings()->GetTreeViewWidth();
-	}
-	else
-	{
-		return 1.f - GetSequencerSettings()->GetTreeViewWidth();
-	}
-}
-
 void SSequencer::OnDragEnter( const FGeometry& MyGeometry, const FDragDropEvent& DragDropEvent )
 {
 	// @todo sequencer: Add drop validity cue
@@ -3628,13 +3627,22 @@ EFrameNumberDisplayFormats SSequencer::GetTimeDisplayFormat() const
 
 void SSequencer::OnColumnFillCoefficientChanged(float FillCoefficient, int32 ColumnIndex)
 {
-	if (ColumnIndex == 0)
-	{
-		GetSequencerSettings()->SetTreeViewWidth(FillCoefficient);
-	}
-	else
+	ColumnFillCoefficients[ColumnIndex] = FillCoefficient;
+
+	// Update on next tick so that we're not constantly saving config on dragging the divider
+	if (GEditor)
 	{
-		GetSequencerSettings()->SetTreeViewWidth(1.f - FillCoefficient);
+		GEditor->GetTimerManager()->SetTimerForNextTick([=]()
+		{
+			if (ColumnIndex == 0)
+			{
+				GetSequencerSettings()->SetTreeViewWidth(FillCoefficient);
+			}
+			else
+			{
+				GetSequencerSettings()->SetTreeViewWidth(1.f - FillCoefficient);
+			}
+		});
 	}
 }
 
diff --git a/Engine/Source/Editor/Sequencer/Private/SSequencer.h b/Engine/Source/Editor/Sequencer/Private/SSequencer.h
--- a/Engine/Source/Editor/Sequencer/Private/SSequencer.h
+++ b/Engine/Source/Editor/Sequencer/Private/SSequencer.h
@@ -521,7 +521,10 @@ private:
 	/**
 	 * @return The fill percentage of the animation outliner
 	 */
-	float GetColumnFillCoefficient(int32 ColumnIndex) const;
+	float GetColumnFillCoefficient(int32 ColumnIndex) const
+	{
+		return ColumnFillCoefficients[ColumnIndex];
+	}
 
 	/**
 	 * Called when one or more assets are dropped into the widget
@@ -699,7 +702,10 @@ private:
 	/** Container for the toolbar, so that we can re-create it as needed. */
 	TSharedPtr<SBox> ToolbarContainer;
 
-	/** Whether the active timer is currently registered */ 
+	/** The fill coefficients of each column in the grid. */
+	float ColumnFillCoefficients[2];
+
+	/** Whether the active timer is currently registered */
 	bool bIsActiveTimerRegistered;
 
 	/** Whether the user is selecting. Ignore selection changes from the level when the user is selecting. */
