Sequencer: Null check in Delete and Keep State for an object that is no longer bound by Sequencer

#rnx

diff --git a/Engine/Source/Editor/Sequencer/Private/Sequencer.cpp b/Engine/Source/Editor/Sequencer/Private/Sequencer.cpp
--- a/Engine/Source/Editor/Sequencer/Private/Sequencer.cpp
+++ b/Engine/Source/Editor/Sequencer/Private/Sequencer.cpp
@@ -5485,16 +5485,19 @@ bool FSequencer::OnRequestNodeDeleted( TSharedRef<FViewModel> NodeToBeDeleted, c
 		{
 			for (TWeakObjectPtr<> WeakObject : FindBoundObjects(ObjectBinding->GetObjectGuid(), SequenceID))
 			{
-				TArray<UObject*> SubObjects;
-				GetObjectsWithOuter(WeakObject.Get(), SubObjects);
+				if (WeakObject.IsValid())
+				{
+					TArray<UObject*> SubObjects;
+					GetObjectsWithOuter(WeakObject.Get(), SubObjects);
 
-				PreAnimatedState.DiscardAndRemoveEntityTokensForObject(*WeakObject.Get());
+					PreAnimatedState.DiscardAndRemoveEntityTokensForObject(*WeakObject.Get());
 
-				for (UObject* SubObject : SubObjects)
-				{
-					if (SubObject)
+					for (UObject* SubObject : SubObjects)
 					{
-						PreAnimatedState.DiscardAndRemoveEntityTokensForObject(*SubObject);
+						if (SubObject)
+						{
+							PreAnimatedState.DiscardAndRemoveEntityTokensForObject(*SubObject);
+						}
 					}
 				}
 			}
