Sequencer: Fixed a crash that could occur if the user opened a context menu on an object binding while non-object outliner items were selected

diff --git a/Engine/Source/Editor/Sequencer/Private/MVVM/ViewModels/ObjectBindingModel.cpp b/Engine/Source/Editor/Sequencer/Private/MVVM/ViewModels/ObjectBindingModel.cpp
--- a/Engine/Source/Editor/Sequencer/Private/MVVM/ViewModels/ObjectBindingModel.cpp
+++ b/Engine/Source/Editor/Sequencer/Private/MVVM/ViewModels/ObjectBindingModel.cpp
@@ -532,7 +532,7 @@ TSharedRef<SWidget> FObjectBindingModel::GetAddTrackMenuContent()
 	{
 		for (const TWeakPtr<FViewModel>& Node : Sequencer->GetSelection().GetSelectedOutlinerItems())
 		{
-			const FObjectBindingModel* ObjectBindingNode = Node.Pin()->CastThisChecked<FObjectBindingModel>();
+			const FObjectBindingModel* ObjectBindingNode = Node.Pin()->CastThis<FObjectBindingModel>();
 			if (!ObjectBindingNode)
 			{
 				continue;
