Sequencer: properly implement IsReadOnly for the sequencer editor view model

diff --git a/Engine/Source/Editor/Sequencer/Private/MVVM/ViewModels/SequencerEditorViewModel.cpp b/Engine/Source/Editor/Sequencer/Private/MVVM/ViewModels/SequencerEditorViewModel.cpp
--- a/Engine/Source/Editor/Sequencer/Private/MVVM/ViewModels/SequencerEditorViewModel.cpp
+++ b/Engine/Source/Editor/Sequencer/Private/MVVM/ViewModels/SequencerEditorViewModel.cpp
@@ -63,6 +63,12 @@ void FSequencerEditorViewModel::SetSequence(UMovieSceneSequence* InRootSequence)
 	SequenceModel->SetSequence(InRootSequence, MovieSceneSequenceID::Root);
 }
 
+bool FSequencerEditorViewModel::IsReadOnly() const
+{
+	TSharedPtr<ISequencer> Sequencer = WeakSequencer.Pin();
+	return !Sequencer || Sequencer->IsReadOnly();
+}
+
 } // namespace Sequencer
 } // namespace UE
 
diff --git a/Engine/Source/Editor/Sequencer/Public/MVVM/ViewModels/SequencerEditorViewModel.h b/Engine/Source/Editor/Sequencer/Public/MVVM/ViewModels/SequencerEditorViewModel.h
--- a/Engine/Source/Editor/Sequencer/Public/MVVM/ViewModels/SequencerEditorViewModel.h
+++ b/Engine/Source/Editor/Sequencer/Public/MVVM/ViewModels/SequencerEditorViewModel.h
@@ -43,6 +43,7 @@ protected:
 	virtual TSharedPtr<FViewModel> CreateRootModelImpl() override;
 	virtual TSharedPtr<FOutlinerViewModel> CreateOutlinerImpl() override;
 	virtual TSharedPtr<FTrackAreaViewModel> CreateTrackAreaImpl() override;
+	virtual bool IsReadOnly() const override;
 
 protected:
 
