Sequencer: add a tooltip for track outliner items that shows property binding information

diff --git a/Engine/Source/Editor/Sequencer/Private/MVVM/ViewModels/TrackModel.cpp b/Engine/Source/Editor/Sequencer/Private/MVVM/ViewModels/TrackModel.cpp
--- a/Engine/Source/Editor/Sequencer/Private/MVVM/ViewModels/TrackModel.cpp
+++ b/Engine/Source/Editor/Sequencer/Private/MVVM/ViewModels/TrackModel.cpp
@@ -533,28 +533,59 @@ FSlateColor FTrackModel::GetLabelColor() const
 			return bIsDimmed ? FSlateColor::UseSubduedForeground() : FSlateColor::UseForeground();
 		}
 
-		TSharedPtr<FSequenceModel> SequenceModel = FindAncestorOfType<FSequenceModel>();
-		TSharedPtr<ISequencer> Sequencer = SequenceModel->GetSequencer();
+		// Return a normal colour if we have at least one bound object for which the property binding resolves
+		// correctly. Otherwise, return a red colour indicating a binding issue.
+		TArray<UObject*> BoundObjects;
+		FindBoundObjects(BoundObjects);
+		for (UObject* BoundObject : BoundObjects)
+		{
+			FTrackInstancePropertyBindings PropertyBinding(PropertyTrack->GetPropertyName(), PropertyTrack->GetPropertyPath().ToString());
+			if (PropertyBinding.GetProperty(*BoundObject))
+			{
+				return bIsDimmed ? FSlateColor::UseSubduedForeground() : FSlateColor::UseForeground();
+			}
+		}
+		return bIsDimmed ? FSlateColor(FLinearColor::Red.Desaturate(0.6f)) : FLinearColor::Red;
+	}
 
-		if (TSharedPtr<IObjectBindingExtension> ParentBinding = FindAncestorOfType<IObjectBindingExtension>())
+	return FOutlinerItemModel::GetLabelColor();
+}
+
+FText FTrackModel::GetLabelToolTipText() const
+{
+	UMovieSceneTrack* Track = GetTrack();
+	if (!Track)
+	{
+		return FText();
+	}
+
+	if (UMovieScenePropertyTrack* PropertyTrack = Cast<UMovieScenePropertyTrack>(Track))
+	{
+		TArray<UObject*> BoundObjects;
+		FindBoundObjects(BoundObjects);
+		for (UObject* BoundObject : BoundObjects)
 		{
-			for (TWeakObjectPtr<> WeakObject : Sequencer->FindBoundObjects(ParentBinding->GetObjectGuid(), Sequencer->GetFocusedTemplateID()))
+			FTrackInstancePropertyBindings PropertyBinding(PropertyTrack->GetPropertyName(), PropertyTrack->GetPropertyPath().ToString());
+			if (FProperty* BoundProperty = PropertyBinding.GetProperty(*BoundObject))
 			{
-				if (UObject* Object = WeakObject.Get())
+				FString PropertyName = BoundProperty->GetMetaData(TEXT("DisplayName"));
+				if (PropertyName.IsEmpty())
 				{
-					FTrackInstancePropertyBindings PropertyBinding(PropertyTrack->GetPropertyName(), PropertyTrack->GetPropertyPath().ToString());
-					if (PropertyBinding.GetProperty(*Object))
-					{
-						return bIsDimmed ? FSlateColor::UseSubduedForeground() : FSlateColor::UseForeground();
-					}
+					PropertyName = BoundProperty->GetName();
+				}
+
+				FString CategoryName = BoundProperty->GetMetaData(TEXT("Category")).Replace(TEXT("|"), TEXT(" \u00BB "));
+				if (!CategoryName.IsEmpty())
+				{
+					CategoryName.Append(TEXT(" \u00BB "));
 				}
-			}
 
-			return bIsDimmed ? FSlateColor(FLinearColor::Red.Desaturate(0.6f)) : FLinearColor::Red;
+				return FText::FromString(FString::Printf(TEXT("%s%s\n(Path: %s)"), *CategoryName, *PropertyName, *PropertyBinding.GetPropertyPath()));
+			}
 		}
 	}
-
-	return FOutlinerItemModel::GetLabelColor();
+	
+	return Track->GetDisplayNameToolTipText();
 }
 
 TSharedRef<SWidget> FTrackModel::CreateOutlinerView(const FCreateOutlinerViewParams& InParams)
@@ -702,6 +733,32 @@ void FTrackModel::Delete()
 	}
 }
 
+bool FTrackModel::FindBoundObjects(TArray<UObject*>& OutBoundObjects) const
+{
+	TSharedPtr<FSequenceModel> SequenceModel = FindAncestorOfType<FSequenceModel>();
+	TSharedPtr<ISequencer> Sequencer = SequenceModel ? SequenceModel->GetSequencer() : nullptr;
+	if (!Sequencer)
+	{
+		return false;
+	}
+
+	TSharedPtr<IObjectBindingExtension> ParentBinding = FindAncestorOfType<IObjectBindingExtension>();
+	if (!ParentBinding)
+	{
+		return false;
+	}
+
+	TArrayView<TWeakObjectPtr<>> FoundBoundObjects = Sequencer->FindBoundObjects(ParentBinding->GetObjectGuid(), Sequencer->GetFocusedTemplateID());
+	OutBoundObjects.Reserve(OutBoundObjects.Num() + FoundBoundObjects.Num());
+	for (TWeakObjectPtr<> WeakObject : FoundBoundObjects)
+	{
+		if (UObject* Object = WeakObject.Get())
+		{
+			OutBoundObjects.Add(Object);
+		}
+	}
+	return true;
+}
 
 } // namespace Sequencer
 } // namespace UE
diff --git a/Engine/Source/Editor/Sequencer/Public/MVVM/ViewModels/TrackModel.h b/Engine/Source/Editor/Sequencer/Public/MVVM/ViewModels/TrackModel.h
--- a/Engine/Source/Editor/Sequencer/Public/MVVM/ViewModels/TrackModel.h
+++ b/Engine/Source/Editor/Sequencer/Public/MVVM/ViewModels/TrackModel.h
@@ -91,6 +91,7 @@ public:
 	const FSlateBrush* GetIconBrush() const override;
 	FText GetLabel() const override;
 	FSlateColor GetLabelColor() const override;
+	FText GetLabelToolTipText() const override;
 	
 	/*~ IDimmableExtension */
 	bool IsDimmed() const override;
@@ -139,6 +140,8 @@ private:
 
 	void ForceUpdate();
 
+	bool FindBoundObjects(TArray<UObject*>& OutBoundObjects) const;
+
 	/** A second children list for the sections inside this track */
 	FViewModelListHead SectionList;
 	FViewModelListHead TopLevelChannelList;
