Sequencer: Leave the tint as normal if there is no object binding

#rnx

diff --git a/Engine/Source/Editor/Sequencer/Private/MVVM/ViewModels/TrackModel.cpp b/Engine/Source/Editor/Sequencer/Private/MVVM/ViewModels/TrackModel.cpp
--- a/Engine/Source/Editor/Sequencer/Private/MVVM/ViewModels/TrackModel.cpp
+++ b/Engine/Source/Editor/Sequencer/Private/MVVM/ViewModels/TrackModel.cpp
@@ -530,7 +530,14 @@ FSlateColor FTrackModel::GetLabelColor() const
 		// 3D transform tracks don't map to property bindings as below
 		if (Track->IsA<UMovieScene3DTransformTrack>() || Track->IsA<UMovieScenePrimitiveMaterialTrack>())
 		{
-			return bIsDimmed ? FSlateColor::UseSubduedForeground() : FSlateColor::UseForeground();
+			return FOutlinerItemModel::GetLabelColor();
+		}
+
+		// If there is no object binding extension, don't tint it
+		TSharedPtr<IObjectBindingExtension> ParentBinding = FindAncestorOfType<IObjectBindingExtension>();
+		if (!ParentBinding)
+		{
+			return FOutlinerItemModel::GetLabelColor();
 		}
 
 		// Return a normal colour if we have at least one bound object for which the property binding resolves
