[Editor]
* ObjectBindingModel - Fixed comparison causing compile errors on odd msvc compiler versions (hololens etc)

diff --git a/Engine/Source/Editor/Sequencer/Private/MVVM/ViewModels/ObjectBindingModel.cpp b/Engine/Source/Editor/Sequencer/Private/MVVM/ViewModels/ObjectBindingModel.cpp
--- a/Engine/Source/Editor/Sequencer/Private/MVVM/ViewModels/ObjectBindingModel.cpp
+++ b/Engine/Source/Editor/Sequencer/Private/MVVM/ViewModels/ObjectBindingModel.cpp
@@ -541,7 +541,7 @@ TSharedRef<SWidget> FObjectBindingModel::GetAddTrackMenuContent()
 			const FGuid Guid = ObjectBindingNode->GetObjectGuid();
 			for (auto RuntimeObject : Sequencer->FindBoundObjects(Guid, OwnerModel->GetSequenceID()))
 			{
-				if (RuntimeObject != nullptr)
+				if (RuntimeObject.Get() != nullptr)
 				{
 					ObjectBindings.AddUnique(Guid);
 					ObjectClasses.Add(RuntimeObject->GetClass());
