Fixed non-unity/pch errors

diff --git a/Engine/Plugins/Runtime/nDisplay/Source/DisplayCluster/Private/Game/EngineClasses/Scene/DisplayClusterRootActor.cpp b/Engine/Plugins/Runtime/nDisplay/Source/DisplayCluster/Private/Game/EngineClasses/Scene/DisplayClusterRootActor.cpp
--- a/Engine/Plugins/Runtime/nDisplay/Source/DisplayCluster/Private/Game/EngineClasses/Scene/DisplayClusterRootActor.cpp
+++ b/Engine/Plugins/Runtime/nDisplay/Source/DisplayCluster/Private/Game/EngineClasses/Scene/DisplayClusterRootActor.cpp
@@ -43,6 +43,7 @@
 #include "Game/EngineClasses/Scene/DisplayClusterRootActorInitializer.h"
 
 #include "Algo/MaxElement.h"
+#include "UObject/Package.h"
 
 #if WITH_EDITOR
 #include "IConcertSyncClientModule.h"
diff --git a/Engine/Plugins/VirtualProduction/Takes/Source/TakesCore/Private/TakeRecorderSources.cpp b/Engine/Plugins/VirtualProduction/Takes/Source/TakesCore/Private/TakeRecorderSources.cpp
--- a/Engine/Plugins/VirtualProduction/Takes/Source/TakesCore/Private/TakeRecorderSources.cpp
+++ b/Engine/Plugins/VirtualProduction/Takes/Source/TakesCore/Private/TakeRecorderSources.cpp
@@ -19,6 +19,7 @@
 #include "Engine/TimecodeProvider.h"
 #include "Misc/App.h"
 #include "Editor.h"
+#include "UObject/Package.h"
 
 #include UE_INLINE_GENERATED_CPP_BY_NAME(TakeRecorderSources)
 
diff --git a/Engine/Source/Editor/Sequencer/Private/MVVM/ViewModels/SequenceModel.cpp b/Engine/Source/Editor/Sequencer/Private/MVVM/ViewModels/SequenceModel.cpp
--- a/Engine/Source/Editor/Sequencer/Private/MVVM/ViewModels/SequenceModel.cpp
+++ b/Engine/Source/Editor/Sequencer/Private/MVVM/ViewModels/SequenceModel.cpp
@@ -74,7 +74,7 @@ UMovieScene* FSequenceModel::GetMovieScene() const
 
 void FSequenceModel::SetSequence(UMovieSceneSequence* InSequence, FMovieSceneSequenceID InSequenceID)
 {
-	if (WeakSequence != InSequence || SequenceID != InSequenceID)
+	if (WeakSequence.Get() != InSequence || SequenceID != InSequenceID)
 	{
 		FViewModelHierarchyOperation HierarchyOperation(GetSharedData());
 
