Sequencer: Make paste object bindings more clever.

If there is no existing binding in the sequence being pasted into, the pasted binding will be bound to the original object.

If there is an existing binding in the sequence being pasted into, a dialog will prompt the user whether to
    1) Duplicate the bound object, in which case you'll get a duplicate copy of the object and the binding+tracks, or
    2) Do NOT duplicate, in which case the resulting binding will be unbound (red) and you can rebind it to something else

diff --git a/Engine/Source/Editor/Sequencer/Private/Sequencer.cpp b/Engine/Source/Editor/Sequencer/Private/Sequencer.cpp
--- a/Engine/Source/Editor/Sequencer/Private/Sequencer.cpp
+++ b/Engine/Source/Editor/Sequencer/Private/Sequencer.cpp
@@ -8019,8 +8019,50 @@ bool FSequencer::PasteObjectBindings(const FString& TextToImport, UMovieSceneFol
 		}
 	}
 
+	bool bDuplicateExistingActors = false;
+	TArray<FString> ObjectNames = FSequencerUtilities::GetPasteBindingsObjectNames(AsShared(), TextToImport);
+
+	// If the current movie scene has bindings that are bound to actors with the same name of the paste buffer, prompt the user what to do
+	UMovieScene* MovieScene = GetFocusedMovieSceneSequence()->GetMovieScene();
+	bool bWasPrompted = false;
+	for (int32 Index = 0; Index < MovieScene->GetPossessableCount() && !bWasPrompted; ++Index)
+	{
+		FGuid ThisGuid = MovieScene->GetPossessable(Index).GetGuid();
+
+		for (TWeakObjectPtr<> WeakObject : FindBoundObjects(ThisGuid, GetFocusedTemplateID()))
+		{
+			if (WeakObject.IsValid())
+			{
+				AActor* Actor = Cast<AActor>(WeakObject.Get());
+
+				if (Actor && ObjectNames.Contains(Actor->GetName()))
+				{
+					FText DuplicateActorsMsg = FText::Format(LOCTEXT("DuplicateActorsForPastedBinding", "Attempting to paste a binding that is already bound to {0}.\nShould the existing actor be duplicated for the pasted binding?"), FText::FromString(Actor->GetActorLabel()));
+
+					EAppReturnType::Type YesNoCancelReply = FMessageDialog::Open(EAppMsgType::YesNoCancel, DuplicateActorsMsg);
+
+					switch (YesNoCancelReply)
+					{
+					case EAppReturnType::Yes:
+						bDuplicateExistingActors = true;
+						break;
+
+					case EAppReturnType::No:
+						bDuplicateExistingActors = false;
+						break;
+
+					case EAppReturnType::Cancel:
+						return false;
+					}
+					bWasPrompted = true;				
+					break;
+				}
+			}
+		}
+	}
+
 	TArray<FMovieSceneBindingProxy> OutBindings;
-	return FSequencerUtilities::PasteBindings(TextToImport, AsShared(), FMovieScenePasteBindingsParams(TargetBindings, InParentFolder, InFolders), OutBindings, PasteErrors);
+	return FSequencerUtilities::PasteBindings(TextToImport, AsShared(), FMovieScenePasteBindingsParams(TargetBindings, InParentFolder, InFolders, bDuplicateExistingActors), OutBindings, PasteErrors);
 }
 
 bool FSequencer::PasteTracks(const FString& TextToImport, UMovieSceneFolder* InParentFolder, const TArray<UMovieSceneFolder*>& InFolders, TArray<FNotificationInfo>& PasteErrors, bool bClearSelection)
diff --git a/Engine/Source/Editor/Sequencer/Private/SequencerUtilities.cpp b/Engine/Source/Editor/Sequencer/Private/SequencerUtilities.cpp
--- a/Engine/Source/Editor/Sequencer/Private/SequencerUtilities.cpp
+++ b/Engine/Source/Editor/Sequencer/Private/SequencerUtilities.cpp
@@ -2106,51 +2106,60 @@ bool FSequencerUtilities::PasteBindings(const FString& TextToImport, TSharedRef<
 					}
 				}
 
-				TArray<AActor*> ActorsToDuplicate;
+				// Find the actors that this pasted binding should bind to
+				TArray<AActor*> ActorsToRebind;
 				for (TActorIterator<AActor> ActorItr(World); ActorItr; ++ActorItr)
 				{
 					AActor* Actor = *ActorItr;
 					if (Actor && CopyableBinding->BoundObjectNames.Contains(Actor->GetName()))
 					{
-						ActorsToDuplicate.Add(Actor);
+						// If this actor is already bound and we're not duplicating actors, don't bind to anything
+						if (!PasteBindingsParams.bDuplicateExistingActors && Sequencer->FindObjectId(*Actor, Sequencer->GetFocusedTemplateID()).IsValid())
+						{
+							continue;
+						}
+
+						ActorsToRebind.Add(Actor);
 						CopyableBinding->BoundObjectNames.Remove(Actor->GetName());
 					}
 				}
 
-				TArray<AActor*> DuplicatedActors;
-				if (ActorsToDuplicate.Num() != 0)
+				if (ActorsToRebind.Num() != 0)
 				{
-					GEditor->SelectNone(false, true);
-					for (AActor* ActorToDuplicate : ActorsToDuplicate)
+					if (PasteBindingsParams.bDuplicateExistingActors)
 					{
-						GEditor->SelectActor(ActorToDuplicate, true, false, false);
-					}
-
-					// Duplicate the bound actors
-					GEditor->edactDuplicateSelected(World->GetCurrentLevel(), false);
-
-					// Duplicating the bound actor through GEditor, edits the copy/paste clipboard. This is not desired from the user's 
-					// point of view since the user didn't explicitly invoke the copy operation. Instead, restore the copied contents
-					// of the clipboard after duplicating the actor
-					FPlatformApplicationMisc::ClipboardCopy(*TextToImport);
-
-					USelection* ActorSelection = GEditor->GetSelectedActors();
-					TArray<TWeakObjectPtr<AActor> > SelectedActors;
-					for (FSelectionIterator Iter(*ActorSelection); Iter; ++Iter)
-					{
-						AActor* Actor = Cast<AActor>(*Iter);
-						if (Actor)
+						GEditor->SelectNone(false, true);
+						for (AActor* ActorToRebind : ActorsToRebind)
+						{
+							GEditor->SelectActor(ActorToRebind, true, false, false);
+						}
+							
+						// Duplicate the bound actors
+						GEditor->edactDuplicateSelected(World->GetCurrentLevel(), false);
+
+						// Duplicating the bound actor through GEditor, edits the copy/paste clipboard. This is not desired from the user's 
+						// point of view since the user didn't explicitly invoke the copy operation. Instead, restore the copied contents
+						// of the clipboard after duplicating the actor
+						FPlatformApplicationMisc::ClipboardCopy(*TextToImport);
+
+						ActorsToRebind.Empty();
+						USelection* ActorSelection = GEditor->GetSelectedActors();
+						for (FSelectionIterator Iter(*ActorSelection); Iter; ++Iter)
 						{
-							DuplicatedActors.Add(Actor);
+							AActor* Actor = Cast<AActor>(*Iter);
+							if (Actor)
+							{
+								ActorsToRebind.Add(Actor);
 
-							CopyableBinding->BoundObjectNames.Add(Actor->GetName());
+								CopyableBinding->BoundObjectNames.Add(Actor->GetName());
+							}
 						}
 					}
 
-					// Bind the duplicated actors
-					if (DuplicatedActors.Num())
+					// Bind the actors
+					if (ActorsToRebind.Num())
 					{
-						AddActorsToBinding(Sequencer, DuplicatedActors, FMovieSceneBindingProxy(NewGuid, Sequence));
+						AddActorsToBinding(Sequencer, ActorsToRebind, FMovieSceneBindingProxy(NewGuid, Sequence));
 					}
 				}
 
@@ -2361,6 +2370,27 @@ bool FSequencerUtilities::CanPasteBindings(TSharedRef<ISequencer> Sequencer, con
 	return ObjectBindingFactory.CanCreateObjectsFromText(TextToImport);
 }
 
+TArray<FString> FSequencerUtilities::GetPasteBindingsObjectNames(TSharedRef<ISequencer> Sequencer, const FString& TextToImport)
+{
+	TArray<FString> ObjectNames;
+
+	TArray<UMovieSceneCopyableBinding*> ImportedBindings;
+	ImportObjectBindingsFromText(Sequencer.Get(), TextToImport, ImportedBindings);
+
+	for (UMovieSceneCopyableBinding* CopyableBinding : ImportedBindings)
+	{
+		if (CopyableBinding)
+		{
+			for (const FString& BoundObjectName : CopyableBinding->BoundObjectNames)
+			{
+				ObjectNames.Add(BoundObjectName);
+			}
+		}
+	}
+
+	return ObjectNames;
+}
+
 FGuid FSequencerUtilities::CreateBinding(TSharedRef<ISequencer> Sequencer, UObject& InObject, const FString& InName)
 {
 	const FScopedTransaction Transaction(LOCTEXT("CreateBinding", "Create New Binding"));
diff --git a/Engine/Source/Editor/Sequencer/Public/SequencerUtilities.h b/Engine/Source/Editor/Sequencer/Public/SequencerUtilities.h
--- a/Engine/Source/Editor/Sequencer/Public/SequencerUtilities.h
+++ b/Engine/Source/Editor/Sequencer/Public/SequencerUtilities.h
@@ -97,10 +97,11 @@ struct FMovieScenePasteBindingsParams
 {
 	GENERATED_BODY()
 
-	FMovieScenePasteBindingsParams(const TArray<FMovieSceneBindingProxy>& InBindings = TArray<FMovieSceneBindingProxy>(), UMovieSceneFolder* InParentFolder = nullptr, const TArray<UMovieSceneFolder*>& InFolders = TArray<UMovieSceneFolder*>())
+	FMovieScenePasteBindingsParams(const TArray<FMovieSceneBindingProxy>& InBindings = TArray<FMovieSceneBindingProxy>(), UMovieSceneFolder* InParentFolder = nullptr, const TArray<UMovieSceneFolder*>& InFolders = TArray<UMovieSceneFolder*>(), bool bInDuplicateExistingActors = false)
 		: Bindings(InBindings)
 		, ParentFolder(InParentFolder)
-		, Folders(InFolders) {}
+		, Folders(InFolders)
+		, bDuplicateExistingActors(bInDuplicateExistingActors) {}
 
 	UPROPERTY(BlueprintReadWrite, Category = "Movie Scene")
 	TArray<FMovieSceneBindingProxy> Bindings;
@@ -110,6 +111,9 @@ struct FMovieScenePasteBindingsParams
 
 	UPROPERTY(BlueprintReadWrite, Category = "Movie Scene")
 	TArray<TObjectPtr<UMovieSceneFolder>> Folders;
+
+	UPROPERTY(BlueprintReadWrite, Category = "Movie Scene")
+	bool bDuplicateExistingActors;
 };
 
 struct SEQUENCER_API FSequencerUtilities
@@ -172,6 +176,7 @@ struct SEQUENCER_API FSequencerUtilities
 	static void CopyBindings(TSharedRef<ISequencer> Sequencer, const TArray<FMovieSceneBindingProxy>& Bindings, const TArray<UMovieSceneFolder*>& InFolders, FString& ExportedText);
 	static bool PasteBindings(const FString& TextToImport, TSharedRef<ISequencer> Sequencer, FMovieScenePasteBindingsParams PasteBindingsParams, TArray<FMovieSceneBindingProxy>& OutBindings, TArray<FNotificationInfo>& OutErrors);
 	static bool CanPasteBindings(TSharedRef<ISequencer> Sequencer, const FString& TextToImport);
+	static TArray<FString> GetPasteBindingsObjectNames(TSharedRef<ISequencer> Sequencer, const FString& TextToImport);
 
 	/** Utility functions for managing bindings */
 	static FGuid CreateBinding(TSharedRef<ISequencer> Sequencer, UObject& InObject, const FString& InName);
