Sequencer: Added Reset Playhead When Navigating which resets the playhed to the start of the playback range when navigating in and out of subsequences

diff --git a/Engine/Source/Editor/Sequencer/Private/SSequencer.cpp b/Engine/Source/Editor/Sequencer/Private/SSequencer.cpp
--- a/Engine/Source/Editor/Sequencer/Private/SSequencer.cpp
+++ b/Engine/Source/Editor/Sequencer/Private/SSequencer.cpp
@@ -2434,6 +2434,7 @@ TSharedRef<SWidget> SSequencer::MakePlaybackMenu()
 		if (SequencerPtr.Pin()->IsLevelEditorSequencer())
 		{
 			MenuBuilder.AddMenuEntry( FSequencerCommands::Get().ToggleEvaluateSubSequencesInIsolation );
+			MenuBuilder.AddMenuEntry( FSequencerCommands::Get().ToggleResetPlayheadWhenNavigating );
 		}
 
 		MenuBuilder.AddMenuEntry( FSequencerCommands::Get().ToggleKeepCursorInPlaybackRangeWhileScrubbing );
diff --git a/Engine/Source/Editor/Sequencer/Private/Sequencer.cpp b/Engine/Source/Editor/Sequencer/Private/Sequencer.cpp
--- a/Engine/Source/Editor/Sequencer/Private/Sequencer.cpp
+++ b/Engine/Source/Editor/Sequencer/Private/Sequencer.cpp
@@ -1128,6 +1128,15 @@ void FSequencer::FocusSequenceInstance(UMovieSceneSubSection& InSubSection)
 
 	AddNodeGroupsCollectionChangedDelegate();
 
+	if (Settings->ShouldResetPlayheadWhenNavigating())
+	{
+		FQualifiedFrameTime CurrentTime = GetLocalTime();
+		if (!SubSequenceRange.Contains(CurrentTime.Time.FloorToFrame()))
+		{
+			SetLocalTime(SubSequenceRange.GetLowerBoundValue(), ESnapTimeMode::STM_None);
+		}
+	}
+
 	bNeedsEvaluate = true;
 	bGlobalMarkedFramesCached = false;
 }
@@ -11099,11 +11108,17 @@ void FSequencer::BindCommands()
 		FCanExecuteAction::CreateLambda([] { return true; }),
 		FIsActionChecked::CreateLambda([this] { return Settings->ShouldKeepCursorInPlayRangeWhileScrubbing(); }));
 
+	SequencerCommandBindings->MapAction(
+		Commands.ToggleResetPlayheadWhenNavigating,
+		FExecuteAction::CreateLambda([this] { Settings->SetResetPlayheadWhenNavigating(!Settings->ShouldResetPlayheadWhenNavigating()); }),
+		FCanExecuteAction::CreateLambda([] { return true; }),
+		FIsActionChecked::CreateLambda([this] { return Settings->ShouldResetPlayheadWhenNavigating(); }));
+
 	SequencerCommandBindings->MapAction(
 		Commands.ToggleKeepPlaybackRangeInSectionBounds,
-		FExecuteAction::CreateLambda( [this]{ Settings->SetKeepPlayRangeInSectionBounds( !Settings->ShouldKeepPlayRangeInSectionBounds() ); NotifyMovieSceneDataChanged( EMovieSceneDataChangeType::TrackValueChanged ); } ),
-		FCanExecuteAction::CreateLambda( []{ return true; } ),
-		FIsActionChecked::CreateLambda( [this]{ return Settings->ShouldKeepPlayRangeInSectionBounds(); } ) );
+		FExecuteAction::CreateLambda([this] { Settings->SetKeepPlayRangeInSectionBounds(!Settings->ShouldKeepPlayRangeInSectionBounds()); NotifyMovieSceneDataChanged(EMovieSceneDataChangeType::TrackValueChanged); }),
+		FCanExecuteAction::CreateLambda([] { return true; }),
+		FIsActionChecked::CreateLambda([this] { return Settings->ShouldKeepPlayRangeInSectionBounds(); }));
 
 	SequencerCommandBindings->MapAction(
 		Commands.ToggleEvaluateSubSequencesInIsolation,
diff --git a/Engine/Source/Editor/Sequencer/Private/SequencerCommands.cpp b/Engine/Source/Editor/Sequencer/Private/SequencerCommands.cpp
--- a/Engine/Source/Editor/Sequencer/Private/SequencerCommands.cpp
+++ b/Engine/Source/Editor/Sequencer/Private/SequencerCommands.cpp
@@ -54,6 +54,7 @@ void FSequencerCommands::RegisterCommands()
 	UI_COMMAND( ToggleAsyncEvaluation, "Async Evaluation", "When enabled, enables a single asynchronous evaluation once per-frame. When disabled, forces a full blocking evaluation every time this sequence is evaluated (should be avoided for real-time content).", EUserInterfaceActionType::ToggleButton, FInputChord() );
 	UI_COMMAND( ToggleDynamicWeighting, "Dynamic Weighting", "When enabled, all blendable tracks will cache their initial values to ensure that they are able to correctly blend in/out when dynamic weights are being used.", EUserInterfaceActionType::ToggleButton, FInputChord() );
 
+	UI_COMMAND( ToggleResetPlayheadWhenNavigating, "Reset Playhead When Navigating", "When checked, if the playhead is outside the playback range, it will be reset to the beginning when navigating in and out of subsequences", EUserInterfaceActionType::ToggleButton, FInputChord());
 	UI_COMMAND( ToggleKeepCursorInPlaybackRangeWhileScrubbing, "Keep Playhead in Playback Range While Scrubbing", "When checked, the playhead will be constrained to the current playback range while scrubbing", EUserInterfaceActionType::ToggleButton, FInputChord() );
 	UI_COMMAND( ToggleKeepPlaybackRangeInSectionBounds, "Keep Playback Range in Section Bounds", "When checked, the playback range will be synchronized to the section bounds", EUserInterfaceActionType::ToggleButton, FInputChord() );
 
diff --git a/Engine/Source/Editor/Sequencer/Private/SequencerSettings.cpp b/Engine/Source/Editor/Sequencer/Private/SequencerSettings.cpp
--- a/Engine/Source/Editor/Sequencer/Private/SequencerSettings.cpp
+++ b/Engine/Source/Editor/Sequencer/Private/SequencerSettings.cpp
@@ -38,6 +38,7 @@ USequencerSettings::USequencerSettings( const FObjectInitializer& ObjectInitiali
 	bIsolateCurveEditorToSelection = true;
 	LoopMode = ESequencerLoopMode::SLM_NoLoop;
 	bSnapKeysAndSectionsToPlayRange = false;
+	bResetPlayheadWhenNavigating = false;
 	bKeepCursorInPlayRangeWhileScrubbing = false;
 	bKeepPlayRangeInSectionBounds = true;
 	bCompileDirectorOnEvaluate = true;
@@ -453,6 +454,20 @@ void USequencerSettings::SetLoopMode(ESequencerLoopMode InLoopMode)
 	}
 }
 
+bool USequencerSettings::ShouldResetPlayheadWhenNavigating() const
+{
+	return bResetPlayheadWhenNavigating;
+}
+
+void USequencerSettings::SetResetPlayheadWhenNavigating(bool bInResetPlayheadWhenNavigating)
+{
+	if (bResetPlayheadWhenNavigating != bInResetPlayheadWhenNavigating)
+	{
+		bResetPlayheadWhenNavigating = bInResetPlayheadWhenNavigating;
+		SaveConfig();
+	}
+}
+
 bool USequencerSettings::ShouldKeepCursorInPlayRangeWhileScrubbing() const
 {
 	return bKeepCursorInPlayRangeWhileScrubbing;
diff --git a/Engine/Source/Editor/Sequencer/Public/SequencerCommands.h b/Engine/Source/Editor/Sequencer/Public/SequencerCommands.h
--- a/Engine/Source/Editor/Sequencer/Public/SequencerCommands.h
+++ b/Engine/Source/Editor/Sequencer/Public/SequencerCommands.h
@@ -145,6 +145,9 @@ public:
 	/** When enabled, all blendable tracks will always cache their initial values to ensure that they are able to correctly blend in/out when dynamic weights are being used. */
 	TSharedPtr< FUICommandInfo > ToggleDynamicWeighting;
 
+	/** Toggle resetting the playhead when navigating in and out of subsequences */
+	TSharedPtr< FUICommandInfo > ToggleResetPlayheadWhenNavigating;
+
 	/** Toggle constraining the time cursor to the playback range while scrubbing */
 	TSharedPtr< FUICommandInfo > ToggleKeepCursorInPlaybackRangeWhileScrubbing;
 
diff --git a/Engine/Source/Editor/Sequencer/Public/SequencerSettings.h b/Engine/Source/Editor/Sequencer/Public/SequencerSettings.h
--- a/Engine/Source/Editor/Sequencer/Public/SequencerSettings.h
+++ b/Engine/Source/Editor/Sequencer/Public/SequencerSettings.h
@@ -249,6 +249,11 @@ public:
 	/** Sets the loop mode. */
 	void SetLoopMode(ESequencerLoopMode InLoopMode);
 
+	/** @return true if the cursor reset when navigating in and out of subsequences, false otherwise */
+	bool ShouldResetPlayheadWhenNavigating() const;
+	/** Set whether or not the cursor should be reset when navigating in and out of subsequences */
+	void SetResetPlayheadWhenNavigating(bool bInResetPlayheadWhenNavigating);
+
 	/** @return true if the cursor should be kept within the playback range while scrubbing in sequencer, false otherwise */
 	bool ShouldKeepCursorInPlayRangeWhileScrubbing() const;
 	/** Set whether or not the cursor should be kept within the playback range while scrubbing in sequencer */
@@ -530,6 +535,10 @@ protected:
 	UPROPERTY( config )
 	TEnumAsByte<ESequencerLoopMode> LoopMode;
 
+	/** Enable or disable resetting the playhead when navigating in and out of subsequences. */
+	UPROPERTY(config, EditAnywhere, Category = Timeline, meta = (DisplayName = "Reset Playhead When Navigating"))
+	bool bResetPlayheadWhenNavigating;
+
 	/** Enable or disable keeping the playhead in the current playback range while scrubbing. */
 	UPROPERTY(config, EditAnywhere, Category = Timeline, meta = (DisplayName = "Keep Playhead in Play Range While Scrubbing"))
 	bool bKeepCursorInPlayRangeWhileScrubbing;
