Avoid turning the weak reference into strong refs in deferred signed object change handler by just removing AddReferencedObject from it

#rnx

diff --git a/Engine/Source/Editor/Sequencer/Private/SequencerModule.cpp b/Engine/Source/Editor/Sequencer/Private/SequencerModule.cpp
--- a/Engine/Source/Editor/Sequencer/Private/SequencerModule.cpp
+++ b/Engine/Source/Editor/Sequencer/Private/SequencerModule.cpp
@@ -59,7 +59,7 @@ namespace UE
 namespace Sequencer
 {
 
-struct FDeferredSignedObjectChangeHandler : UE::MovieScene::IDeferredSignedObjectChangeHandler, FGCObject
+struct FDeferredSignedObjectChangeHandler : UE::MovieScene::IDeferredSignedObjectChangeHandler
 {
 	FDeferredSignedObjectChangeHandler()
 	{
@@ -127,17 +127,6 @@ struct FDeferredSignedObjectChangeHandler : UE::MovieScene::IDeferredSignedObjec
 		SignedObjects.Add(SignedObject);
 	}
 
-	void AddReferencedObjects( FReferenceCollector& Collector ) override
-	{
-		for (TWeakObjectPtr<UMovieSceneSignedObject> WeakObject : SignedObjects)
-		{
-			if (UMovieSceneSignedObject* Object = WeakObject.Get())
-			{
-				Collector.AddReferencedObject(Object);
-			}
-		}
-	}
-
 	bool CreateImplicitScopedModifyDefer() override
 	{
 		ensure(!DeferImplicitChanges.IsSet());
@@ -150,11 +139,6 @@ struct FDeferredSignedObjectChangeHandler : UE::MovieScene::IDeferredSignedObjec
 		DeferImplicitChanges.Reset();
 	}
 
-	FString GetReferencerName() const override
-	{
-		return TEXT("FDeferredSignedObjectChangeHandler");
-	}
-
 	TSet<TWeakObjectPtr<UMovieSceneSignedObject>> SignedObjects;
 	TWeakObjectPtr<UTransBuffer>   WeakBuffer;
 	TOptional<UE::MovieScene::FScopedSignedObjectModifyDefer> DeferTransactionChanges;
