Sequencer: Fix so that batch edit sections work for selected sections

diff --git a/Engine/Source/Editor/Sequencer/Private/MVVM/ViewModels/TrackModel.cpp b/Engine/Source/Editor/Sequencer/Private/MVVM/ViewModels/TrackModel.cpp
--- a/Engine/Source/Editor/Sequencer/Private/MVVM/ViewModels/TrackModel.cpp
+++ b/Engine/Source/Editor/Sequencer/Private/MVVM/ViewModels/TrackModel.cpp
@@ -1,8 +1,10 @@
 // Copyright Epic Games, Inc. All Rights Reserved.
 
 #include "MVVM/ViewModels/TrackModel.h"
+
 #include "MVVM/Extensions/IObjectBindingExtension.h"
 #include "MVVM/Extensions/IRecyclableExtension.h"
+#include "MVVM/Extensions/ITrackExtension.h"
 #include "MVVM/ViewModels/ChannelModel.h"
 #include "MVVM/ViewModels/FolderModel.h"
 #include "MVVM/ViewModels/ViewModelIterators.h"
@@ -661,6 +663,18 @@ void FTrackModel::BuildContextMenu(FMenuBuilder& MenuBuilder)
 
 	// Find sections in the track to add batch properties for
 	TArray<TWeakObjectPtr<UObject>> TrackSections;
+
+	for (TWeakPtr<FViewModel> Node : GetEditor()->GetSequencer()->GetSelection().GetSelectedOutlinerItems())
+	{
+		if (ITrackExtension* TrackExtension = ICastable::CastWeakPtr<ITrackExtension>(Node))
+		{
+			for (UMovieSceneSection* Section : TrackExtension->GetSections())
+			{
+				TrackSections.Add(Section);
+			}
+		}
+	}
+
 	for (TSharedPtr<FViewModel> TrackAreaModel : GetTrackAreaModelList())
 	{
 		constexpr bool bIncludeThis = true;
@@ -668,7 +682,7 @@ void FTrackModel::BuildContextMenu(FMenuBuilder& MenuBuilder)
 		{
 			if (UMovieSceneSection* SectionObject = Section->GetSection())
 			{
-				TrackSections.Add(SectionObject);
+				TrackSections.AddUnique(SectionObject);
 			}
 		}
 	}
diff --git a/Engine/Source/Editor/Sequencer/Private/MVVM/ViewModels/TrackRowModel.cpp b/Engine/Source/Editor/Sequencer/Private/MVVM/ViewModels/TrackRowModel.cpp
--- a/Engine/Source/Editor/Sequencer/Private/MVVM/ViewModels/TrackRowModel.cpp
+++ b/Engine/Source/Editor/Sequencer/Private/MVVM/ViewModels/TrackRowModel.cpp
@@ -1,6 +1,8 @@
 // Copyright Epic Games, Inc. All Rights Reserved.
 
 #include "MVVM/ViewModels/TrackRowModel.h"
+
+#include "MVVM/Extensions/ITrackExtension.h"
 #include "MVVM/ViewModels/SectionModel.h"
 #include "MVVM/ViewModels/SequenceModel.h"
 #include "MVVM/ViewModels/ViewModelIterators.h"
@@ -315,6 +317,18 @@ void FTrackRowModel::BuildContextMenu(FMenuBuilder& MenuBuilder)
 
 	// Find sections in the track to add batch properties for
 	TArray<TWeakObjectPtr<UObject>> TrackSections;
+	
+	for (TWeakPtr<FViewModel> Node : GetEditor()->GetSequencer()->GetSelection().GetSelectedOutlinerItems())
+	{
+		if (ITrackExtension* TrackExtension = ICastable::CastWeakPtr<ITrackExtension>(Node))
+		{
+			for (UMovieSceneSection* Section : TrackExtension->GetSections())
+			{
+				TrackSections.Add(Section);
+			}
+		}
+	}
+	
 	for (TSharedPtr<FViewModel> TrackAreaModel : GetTrackAreaModelList())
 	{
 		constexpr bool bIncludeThis = true;
@@ -322,7 +336,7 @@ void FTrackRowModel::BuildContextMenu(FMenuBuilder& MenuBuilder)
 		{
 			if (UMovieSceneSection* SectionObject = Section->GetSection())
 			{
-				TrackSections.Add(SectionObject);
+				TrackSections.AddUnique(SectionObject);
 			}
 		}
 	}
