Sequencer: Save the tree view width percentage in sequencer settings so that the split between the tree view and the track area is saved/restored.

Rework to use OnSplitterFinishedResizing

#rnx

diff --git a/Engine/Source/Editor/Sequencer/Private/SSequencer.cpp b/Engine/Source/Editor/Sequencer/Private/SSequencer.cpp
--- a/Engine/Source/Editor/Sequencer/Private/SSequencer.cpp
+++ b/Engine/Source/Editor/Sequencer/Private/SSequencer.cpp
@@ -825,9 +825,10 @@ void SSequencer::Construct(const FArguments& InArgs, TSharedRef<FSequencer> InSe
 				+ SOverlay::Slot()
 				[
 					// track area virtual splitter overlay
-					SNew(SSequencerSplitterOverlay)
+					SAssignNew(TreeViewSplitter, SSequencerSplitterOverlay)
 					.Style(FAppStyle::Get(), "Sequencer.AnimationOutliner.Splitter")
 					.Visibility(EVisibility::SelfHitTestInvisible)
+					.OnSplitterFinishedResizing(this, &SSequencer::OnSplitterFinishedResizing)
 
 					+ SSplitter::Slot()
 					.Value(FillCoefficient_0)
@@ -3628,25 +3629,20 @@ EFrameNumberDisplayFormats SSequencer::GetTimeDisplayFormat() const
 	return GetSequencerSettings()->GetTimeDisplayFormat();
 }
 
+void SSequencer::OnSplitterFinishedResizing()
+{
+	SSplitter::FSlot const& LeftSplitterSlot = TreeViewSplitter->Splitter->SlotAt(0);
+	SSplitter::FSlot const& RightSplitterSlot = TreeViewSplitter->Splitter->SlotAt(1);
+
+	OnColumnFillCoefficientChanged(LeftSplitterSlot.GetSizeValue(), 0);
+	OnColumnFillCoefficientChanged(RightSplitterSlot.GetSizeValue(), 1);
+
+	GetSequencerSettings()->SetTreeViewWidth(LeftSplitterSlot.GetSizeValue());
+}
+
 void SSequencer::OnColumnFillCoefficientChanged(float FillCoefficient, int32 ColumnIndex)
 {
 	ColumnFillCoefficients[ColumnIndex] = FillCoefficient;
-
-	// Update on next tick so that we're not constantly saving config on dragging the divider
-	if (GEditor)
-	{
-		GEditor->GetTimerManager()->SetTimerForNextTick([=]()
-		{
-			if (ColumnIndex == 0)
-			{
-				GetSequencerSettings()->SetTreeViewWidth(FillCoefficient);
-			}
-			else
-			{
-				GetSequencerSettings()->SetTreeViewWidth(1.f - FillCoefficient);
-			}
-		});
-	}
 }
 
 void SSequencer::OnCurveEditorVisibilityChanged(bool bShouldBeVisible)
diff --git a/Engine/Source/Editor/Sequencer/Private/SSequencer.h b/Engine/Source/Editor/Sequencer/Private/SSequencer.h
--- a/Engine/Source/Editor/Sequencer/Private/SSequencer.h
+++ b/Engine/Source/Editor/Sequencer/Private/SSequencer.h
@@ -523,6 +523,7 @@ private:
 	 */
 	float GetColumnFillCoefficient(int32 ColumnIndex) const
 	{
+		ensure(ColumnIndex == 0 || ColumnIndex == 1);
 		return ColumnFillCoefficients[ColumnIndex];
 	}
 
@@ -596,6 +597,8 @@ private:
 	/** Called when a column fill percentage is changed by a splitter slot. */
 	void OnColumnFillCoefficientChanged(float FillCoefficient, int32 ColumnIndex);
 
+	void OnSplitterFinishedResizing();
+
 	/** Gets paint options for painting the playback range on sequencer */
 	FPaintPlaybackRangeArgs GetSectionPlaybackRangeArgs() const;
 
@@ -705,6 +708,8 @@ private:
 	/** The fill coefficients of each column in the grid. */
 	float ColumnFillCoefficients[2];
 
+	TSharedPtr<class SSequencerSplitterOverlay> TreeViewSplitter;
+
 	/** Whether the active timer is currently registered */
 	bool bIsActiveTimerRegistered;
 
