Sequencer: Fixed crash editing object channel keys

The actual crash was caused by a disparity between the actual channel value type (FMovieSceneObjectPathChannelKeyValue) and the type used to edit the value (FSoftObjectPtr).
When constructing an instance of the key struct, we would write into the instance property as if it were a FMovieSceneObjectPathChannelKeyValue, and trample memory. Ordinarily this doesn't happen because the types match, but object channels override the types here.
To fix this I made it the responsibility of PostConstructKeyInstance to assign the value, allowing the object path channel to correctly populate the value.

I also found that the CPP struct ops were not working correctly (something must have changed since this was first written; this class has a virtual destructor but its vptr was always nullptr because the struct ops were not being set-up/used any more), and in fact it is not well supported to have both CPP struct ops _and_ non C++ defined UPROPERTIES on a UScriptStruct. To this end I have removed the cpp struct ops from the generated struct and added code to do the initialization and destruction of the struct instance directly.

diff --git a/Engine/Source/Editor/MovieSceneTools/Private/Channels/BuiltInChannelEditors.cpp b/Engine/Source/Editor/MovieSceneTools/Private/Channels/BuiltInChannelEditors.cpp
--- a/Engine/Source/Editor/MovieSceneTools/Private/Channels/BuiltInChannelEditors.cpp
+++ b/Engine/Source/Editor/MovieSceneTools/Private/Channels/BuiltInChannelEditors.cpp
@@ -726,7 +726,7 @@ UMovieSceneKeyStructType* InstanceGeneratedStruct(FMovieSceneObjectPathChannel*
 		return nullptr;
 	}
 
-	FSoftObjectProperty* NewValueProperty = new FSoftObjectProperty(NewStruct, "Value", RF_NoFlags);
+	FObjectPtrProperty* NewValueProperty = new FObjectPtrProperty(NewStruct, "Value", RF_NoFlags);
 	NewValueProperty->SetPropertyFlags(CPF_Edit);
 	NewValueProperty->SetMetaData("Category", TEXT("Key"));
 	NewValueProperty->PropertyClass = PropertyClass;
@@ -741,29 +741,51 @@ UMovieSceneKeyStructType* InstanceGeneratedStruct(FMovieSceneObjectPathChannel*
 	return NewStruct;
 }
 
+
 void PostConstructKeyInstance(const TMovieSceneChannelHandle<FMovieSceneObjectPathChannel>& ChannelHandle, FKeyHandle InHandle, FStructOnScope* Struct)
-{	
+{
 	const UMovieSceneKeyStructType* GeneratedStructType = CastChecked<const UMovieSceneKeyStructType>(Struct->GetStruct());
 
-	FSoftObjectProperty* EditProperty = CastFieldChecked<FSoftObjectProperty>(GeneratedStructType->DestValueProperty.Get());
-	const uint8* PropertyAddress = EditProperty->ContainerPtrToValuePtr<uint8>(Struct->GetStructMemory());
+	uint8* StructMemory = Struct->GetStructMemory();
+
+	FObjectPtrProperty* ValueProperty = CastFieldChecked<FObjectPtrProperty>(GeneratedStructType->DestValueProperty.Get());
+	FStructProperty*     TimeProperty  = CastFieldChecked<FStructProperty>(GeneratedStructType->DestTimeProperty.Get());
+
+	const FFrameNumber*  TimeAddress   = TimeProperty->ContainerPtrToValuePtr<FFrameNumber>(StructMemory);
+	void*                ValueAddress  = ValueProperty->ContainerPtrToValuePtr<uint8>(StructMemory);
 
 	// It is safe to capture the property and address in this lambda because the lambda is owned by the struct itself, so cannot be invoked if the struct has been destroyed
-	auto CopyInstanceToKeyLambda = [ChannelHandle, InHandle, EditProperty, PropertyAddress](const FPropertyChangedEvent&)
+	auto CopyInstanceToKeyLambda = [ChannelHandle, InHandle, GeneratedStructType, ValueProperty, ValueAddress, TimeAddress](const FPropertyChangedEvent&)
 	{
 		if (FMovieSceneObjectPathChannel* DestinationChannel = ChannelHandle.Get())
 		{
 			const int32 KeyIndex = DestinationChannel->GetData().GetIndex(InHandle);
 			if (KeyIndex != INDEX_NONE)
 			{
-				UObject* ObjectPropertyValue = EditProperty->GetObjectPropertyValue(PropertyAddress);
+				UObject* ObjectPropertyValue = ValueProperty->GetObjectPropertyValue(ValueAddress);
 				DestinationChannel->GetData().GetValues()[KeyIndex] = ObjectPropertyValue;
+
+				// Set the new key time
+				DestinationChannel->SetKeyTime(InHandle, *TimeAddress);
 			}
 		}
 	};
 
 	FGeneratedMovieSceneKeyStruct* KeyStruct = reinterpret_cast<FGeneratedMovieSceneKeyStruct*>(Struct->GetStructMemory());
 	KeyStruct->OnPropertyChangedEvent = CopyInstanceToKeyLambda;
+
+	// Copy the initial value for the struct
+	FMovieSceneObjectPathChannel* Channel = ChannelHandle.Get();
+	if (Channel)
+	{
+		// Copy the initial value into the struct
+		const int32 KeyIndex = Channel->GetData().GetIndex(InHandle);
+		if (KeyIndex != INDEX_NONE)
+		{
+			UObject* InitialObject = Channel->GetData().GetValues()[KeyIndex].Get();
+			ValueProperty->SetObjectPropertyValue(ValueAddress, InitialObject);
+		}
+	}
 }
 
 template<typename ChannelType>
diff --git a/Engine/Source/Editor/Sequencer/Private/SequencerKeyStructGenerator.cpp b/Engine/Source/Editor/Sequencer/Private/SequencerKeyStructGenerator.cpp
--- a/Engine/Source/Editor/Sequencer/Private/SequencerKeyStructGenerator.cpp
+++ b/Engine/Source/Editor/Sequencer/Private/SequencerKeyStructGenerator.cpp
@@ -14,6 +14,48 @@ UMovieSceneKeyStructType::UMovieSceneKeyStructType(const FObjectInitializer& Obj
 	SourceTimesProperty  = nullptr;
 	DestValueProperty    = nullptr;
 	DestTimeProperty     = nullptr;
+
+	SetSuperStruct(FGeneratedMovieSceneKeyStruct::StaticStruct());
+	ensureMsgf(FGeneratedMovieSceneKeyStruct::StaticStruct()->PropertyLink == nullptr, TEXT("FGeneratedMovieSceneKeyStruct must not have any UPROPERTY members."));
+}
+
+void UMovieSceneKeyStructType::InitializeStruct(void* InDest, int32 ArrayDim) const
+{
+	const int32 Stride = GetStructureSize();
+
+	FMemory::Memzero(InDest, ArrayDim*Stride);
+
+	// Initialize the native super struct first
+	for (int32 Index = 0; Index < ArrayDim; ++Index)
+	{
+		uint8* ElementBytes = reinterpret_cast<uint8*>(InDest) + Index*Stride;
+		new (ElementBytes) FGeneratedMovieSceneKeyStruct();
+
+		// Initialize any additional properties on this struct
+		for (FProperty* Property = PropertyLink; Property; Property = Property->PropertyLinkNext)
+		{
+			Property->InitializeValue_InContainer(ElementBytes);
+		}
+	}
+}
+
+void UMovieSceneKeyStructType::DestroyStruct(void* InDest, int32 ArrayDim) const
+{
+	const int32 Stride = GetStructureSize();
+
+	for (int32 Index = 0; Index < ArrayDim; ++Index)
+	{
+		uint8* ElementBytes = reinterpret_cast<uint8*>(InDest) + Index*Stride;
+
+		// Destroy any additional properties on this struct
+		for (FProperty* Property = PropertyLink; Property; Property = Property->PropertyLinkNext)
+		{
+			Property->DestroyValue_InContainer(ElementBytes);
+		}
+
+		// Destroy the base class
+		reinterpret_cast<FGeneratedMovieSceneKeyStruct*>(ElementBytes)->~FGeneratedMovieSceneKeyStruct();
+	}
 }
 
 FSequencerKeyStructGenerator& FSequencerKeyStructGenerator::Get()
@@ -25,7 +67,6 @@ FSequencerKeyStructGenerator& FSequencerKeyStructGenerator::Get()
 UMovieSceneKeyStructType* FSequencerKeyStructGenerator::AllocateNewKeyStruct()
 {
 	UMovieSceneKeyStructType* NewStruct = NewObject<UMovieSceneKeyStructType>(GetTransientPackage(), NAME_None, RF_Public | RF_Standalone);
-	NewStruct->SetSuperStruct(FGeneratedMovieSceneKeyStruct::StaticStruct());
 	return NewStruct;
 }
 
@@ -103,8 +144,6 @@ void FSequencerKeyStructGenerator::FinalizeNewKeyStruct(UMovieSceneKeyStructType
 	InStruct->Bind();
 	InStruct->StaticLink(true);
 
-	UMovieSceneKeyStructType::DeferCppStructOps(InStruct->GetStructPathName(), new UScriptStruct::TCppStructOps<FGeneratedMovieSceneKeyStruct>);
-
 	check(InStruct->IsComplete());
 }
 
@@ -146,7 +185,6 @@ TSharedPtr<FStructOnScope> FSequencerKeyStructGenerator::CreateInitialStructInst
 
 	// Copy the initial time into the struct
 	{
-
 		const uint8* SrcTimeData  = GeneratedStructType->SourceTimesProperty->ContainerPtrToValuePtr<uint8>(SourceChannel);
 		uint8*       DestTimeData = GeneratedStructType->DestTimeProperty->ContainerPtrToValuePtr<uint8>(StructPtr);
 
@@ -154,14 +192,5 @@ TSharedPtr<FStructOnScope> FSequencerKeyStructGenerator::CreateInitialStructInst
 		GeneratedStructType->SourceTimesProperty->Inner->CopyCompleteValue(DestTimeData, SourceTimesArray.GetRawPtr(InitialKeyIndex));
 	}
 
-	// Copy the initial value into the struct
-	{
-		const uint8* SrcValueData  = GeneratedStructType->SourceValuesProperty->ContainerPtrToValuePtr<uint8>(SourceChannel);
-		uint8*       DestValueData = GeneratedStructType->DestValueProperty->ContainerPtrToValuePtr<uint8>(StructPtr);
-
-		FScriptArrayHelper SourceValuesArray(GeneratedStructType->SourceValuesProperty.Get(), SrcValueData);
-		GeneratedStructType->SourceValuesProperty->Inner->CopyCompleteValue(DestValueData, SourceValuesArray.GetRawPtr(InitialKeyIndex));
-	}
-
 	return Struct;
 }
diff --git a/Engine/Source/Editor/Sequencer/Public/SequencerKeyStructGenerator.h b/Engine/Source/Editor/Sequencer/Public/SequencerKeyStructGenerator.h
--- a/Engine/Source/Editor/Sequencer/Public/SequencerKeyStructGenerator.h
+++ b/Engine/Source/Editor/Sequencer/Public/SequencerKeyStructGenerator.h
@@ -41,6 +41,9 @@ public:
 
 	UMovieSceneKeyStructType(const FObjectInitializer& ObjInit);
 
+	void InitializeStruct(void* InDest, int32 ArrayDim) const override;
+	void DestroyStruct(void* Dest, int32 ArrayDim) const override;
+
 	/**
 	 * Check whether this generated struct is complete and ready to be used
 	 */
@@ -64,6 +67,9 @@ public:
 	/** The value property for this reflected struct, of the same type as SourceValuesProperty->Inner */
 	UPROPERTY()
 	TFieldPath<FProperty> DestValueProperty;
+
+private:
+	using UStruct::SetSuperStruct;
 };
 
 /**
@@ -192,6 +198,24 @@ inline UMovieSceneKeyStructType* InstanceGeneratedStruct(void* Channel, FSequenc
 template<typename ChannelType>
 void PostConstructKeyInstance(const TMovieSceneChannelHandle<ChannelType>& ChannelHandle, FKeyHandle InHandle, FStructOnScope* Struct)
 {
+	ChannelType* Channel = ChannelHandle.Get();
+	if (Channel)
+	{
+		const UMovieSceneKeyStructType* GeneratedStructType = CastChecked<const UMovieSceneKeyStructType>(Struct->GetStruct());
+		void* StructPtr = Struct->GetStructMemory();
+
+		const int32 InitialKeyIndex = Channel->GetData().GetIndex(InHandle);
+
+		// Copy the initial value into the struct
+		if (InitialKeyIndex != INDEX_NONE)
+		{
+			const uint8* SrcValueData  = GeneratedStructType->SourceValuesProperty->ContainerPtrToValuePtr<uint8>(Channel);
+			uint8*       DestValueData = GeneratedStructType->DestValueProperty->ContainerPtrToValuePtr<uint8>(StructPtr);
+
+			FScriptArrayHelper SourceValuesArray(GeneratedStructType->SourceValuesProperty.Get(), SrcValueData);
+			GeneratedStructType->SourceValuesProperty->Inner->CopyCompleteValue(DestValueData, SourceValuesArray.GetRawPtr(InitialKeyIndex));
+		}
+	}
 }
 
 template<typename ChannelType>
