Sequencer: Fixed array property handling in editor (adding property tracks for array elements & autokey)

diff --git a/Engine/Source/Editor/Sequencer/Private/MVVM/ViewModels/ObjectBindingModel.cpp b/Engine/Source/Editor/Sequencer/Private/MVVM/ViewModels/ObjectBindingModel.cpp
--- a/Engine/Source/Editor/Sequencer/Private/MVVM/ViewModels/ObjectBindingModel.cpp
+++ b/Engine/Source/Editor/Sequencer/Private/MVVM/ViewModels/ObjectBindingModel.cpp
@@ -67,14 +67,9 @@ void GetKeyablePropertyPaths(UClass* Class, void* ValuePtr, UStruct* PropertySou
 		{
 			PropertyPath.AddProperty(FPropertyInfo(Property));
 
-			bool bIsPropertyKeyable = Sequencer.CanKeyProperty(FCanKeyPropertyParams(Class, PropertyPath));
-			if (bIsPropertyKeyable)
-			{
-				KeyablePropertyPaths.Add(PropertyPath);
-			}
-
+			bool bIsPropertyKeyable = false;
 			FArrayProperty* ArrayProperty = CastField<FArrayProperty>(Property);
-			if (!bIsPropertyKeyable && ArrayProperty)
+			if (ArrayProperty)
 			{
 				FScriptArrayHelper ArrayHelper(ArrayProperty, ArrayProperty->ContainerPtrToValuePtr<void>(ValuePtr));
 				for (int32 Index = 0; Index < ArrayHelper.Num(); ++Index)
@@ -95,6 +90,15 @@ void GetKeyablePropertyPaths(UClass* Class, void* ValuePtr, UStruct* PropertySou
 				}
 			}
 
+			if (!bIsPropertyKeyable)
+			{
+				bIsPropertyKeyable = Sequencer.CanKeyProperty(FCanKeyPropertyParams(Class, PropertyPath));
+				if (bIsPropertyKeyable)
+				{
+					KeyablePropertyPaths.Add(PropertyPath);
+				}
+			}
+
 			if (!bIsPropertyKeyable || bRecurseAllProperties)
 			{
 				if (FStructProperty* StructProperty = CastField<FStructProperty>(Property))
diff --git a/Engine/Source/Editor/Sequencer/Private/SequencerObjectChangeListener.cpp b/Engine/Source/Editor/Sequencer/Private/SequencerObjectChangeListener.cpp
--- a/Engine/Source/Editor/Sequencer/Private/SequencerObjectChangeListener.cpp
+++ b/Engine/Source/Editor/Sequencer/Private/SequencerObjectChangeListener.cpp
@@ -191,10 +191,17 @@ bool IsHiddenFunction(const UStruct& PropertyStructure, FAnimatedPropertyKey Pro
 
 const FOnAnimatablePropertyChanged* FSequencerObjectChangeListener::FindPropertySetter(const UStruct& PropertyStructure, FAnimatedPropertyKey PropertyKey, const FProperty& Property) const
 {
+	const FArrayProperty* ArrayOwner = Property.GetTypedOwner<FArrayProperty>();
+	const FProperty* PropertyOrContainer = ArrayOwner ? ArrayOwner : &Property;
+
+	// If we are trying to set a property that exists within a container type (ie an array),
+	// we check for flags on the outer property but must not use setter functions that would not know the index to set
+	const bool bCanApplyFunction = PropertyOrContainer == &Property;
+
 	const FOnAnimatablePropertyChanged* DelegatePtr = PropertyChangedEventMap.Find(PropertyKey);
 	if (DelegatePtr != nullptr)
 	{
-		FString PropertyVarName = Property.GetName();
+		FString PropertyVarName = PropertyOrContainer->GetName();
 
 		// If this is a bool property, strip off the 'b' so that the "Set" functions to be 
 		// found are, for example, "SetHidden" instead of "SetbHidden"
@@ -203,55 +210,45 @@ const FOnAnimatablePropertyChanged* FSequencerObjectChangeListener::FindProperty
 			PropertyVarName.RemoveFromStart("b", ESearchCase::CaseSensitive);
 		}
 
-		const bool bFoundSetter = Property.HasSetter();
-
-		static const FString Set(TEXT("Set"));
-
-		const FString FunctionString = Set + PropertyVarName;
-
-		FName FunctionName = FName(*FunctionString);
-
-		static const FName DeprecatedFunctionName(TEXT("DeprecatedFunction"));
-		UFunction* Function = nullptr;
-		if (const UClass* Class = Cast<const UClass>(&PropertyStructure))
+		// Interp properties are always keyable
+		if (PropertyOrContainer->HasAnyPropertyFlags(CPF_Interp))
 		{
-			Function = Class->FindFunctionByName(FunctionName);
+			return DelegatePtr;
 		}
-		bool bFoundValidFunction = false;
-		if (Function && !Function->HasMetaData(DeprecatedFunctionName))
+
+		if (bCanApplyFunction)
 		{
-			// FIXME: FTrackInstancePropertyBindings::InvokeSetterFunction doesn't support array properties.
-			if (!CastField<const FArrayProperty>(&Property))
+			// If the function is hidden, we cannot use it
+			if (IsHiddenFunction(PropertyStructure, FAnimatedPropertyKey::FromProperty(&Property), Property.GetName()))
 			{
-				bFoundValidFunction = true;
+				return nullptr;
 			}
-		}
-
-		bool bFoundValidInterp = false;
-		bool bFoundEditDefaultsOnly = false;
-		bool bFoundEdit = false;
 
-		if (Property.HasAnyPropertyFlags(CPF_Interp))
-		{
-			bFoundValidInterp = true;
-		}
-
-		// @TODO: should we early out of our property path iteration if we find an "edit defaults only" property?
-		if (Property.HasAnyPropertyFlags(CPF_DisableEditOnInstance))
-		{
-			bFoundEditDefaultsOnly = true;
-		}
-		if (Property.HasAnyPropertyFlags(CPF_Edit))
-		{
-			bFoundEdit = true;
-		}
+			// If there is a native setter we can always animate the property
+			if (Property.HasSetter())
+			{
+				return DelegatePtr;
+			}
+	
+			// Check to see if we have a function of the form Set<PropertyName> that we can use
+			if (const UClass* Class = Cast<const UClass>(&PropertyStructure))
+			{
+				static const FString Set(TEXT("Set"));
+				static const FName DeprecatedFunctionName(TEXT("DeprecatedFunction"));
 
-		const bool bIsHiddenFunction = IsHiddenFunction(PropertyStructure, FAnimatedPropertyKey::FromProperty(&Property), Property.GetName());
+				FName FunctionName = FName(*(Set + PropertyVarName));
+				UFunction* Function = Class->FindFunctionByName(FunctionName);
 
-		// Valid if there's a setter function and the property is editable. Also valid if there's an interp keyword.
-		if (((bFoundValidFunction && bFoundEdit && !bFoundEditDefaultsOnly) || bFoundValidInterp || bFoundSetter) && !bIsHiddenFunction)
-		{
-			return DelegatePtr;
+				// @TODO: should we early out of our property path iteration if we find an "edit defaults only" property?
+				const bool bEditable           = Property.HasAnyPropertyFlags(CPF_Edit) && !Property.HasAnyPropertyFlags(CPF_DisableEditOnInstance);
+				const bool bFoundValidFunction = Function && !Function->HasMetaData(DeprecatedFunctionName);
+			
+				// Valid if there's a setter function and the property is editable.
+				if (bFoundValidFunction && bEditable)
+				{
+					return DelegatePtr;
+				}
+			}
 		}
 	}
 
@@ -285,6 +282,11 @@ bool FSequencerObjectChangeListener::CanKeyProperty_Internal(FCanKeyPropertyPara
 		FProperty* Property = CanKeyPropertyParams.PropertyPath.GetPropertyInfo(Index).Property.Get();
 		if (Property)
 		{
+			if (Property->IsA<FArrayProperty>())
+			{
+				continue;
+			}
+
 			const UStruct* PropertyContainer = CanKeyPropertyParams.FindPropertyContainer(Property);
 			if (PropertyContainer)
 			{
diff --git a/Engine/Source/Editor/Sequencer/Public/KeyPropertyParams.h b/Engine/Source/Editor/Sequencer/Public/KeyPropertyParams.h
--- a/Engine/Source/Editor/Sequencer/Public/KeyPropertyParams.h
+++ b/Engine/Source/Editor/Sequencer/Public/KeyPropertyParams.h
@@ -98,6 +98,19 @@ public:
 				int32 ArrayIndex = FMath::Max(0, PropertyInfo.ArrayIndex);
 				if (FArrayProperty* ArrayProp = CastField<FArrayProperty>(Property))
 				{
+					// Sometimes property paths have the array property twice, first with no array index,
+					// then a second so we skip over this property if that's the case
+					if (PropertyInfo.ArrayIndex == INDEX_NONE && i < PropertyPath.GetNumProperties()-1)
+					{
+						const FPropertyInfo& InnerPropertyInfo = PropertyPath.GetPropertyInfo(i+1);
+						FProperty* InnerProperty = InnerPropertyInfo.Property.Get();
+						if (InnerProperty && InnerProperty->GetOwner<FProperty>() == ArrayProp)
+						{
+							ArrayIndex = InnerPropertyInfo.ArrayIndex;
+							++i;
+						}
+					}
+
 					FScriptArrayHelper ParentArrayHelper(ArrayProp, ArrayProp->ContainerPtrToValuePtr<void>(ContainerPtr));
 					if (!ParentArrayHelper.IsValidIndex(ArrayIndex))
 					{
