Sequencer: editing of channel overrides is now properly undoable, and works on multiple channel groups

diff --git a/Engine/Source/Editor/Sequencer/Private/MVVM/ViewModels/ChannelModel.cpp b/Engine/Source/Editor/Sequencer/Private/MVVM/ViewModels/ChannelModel.cpp
--- a/Engine/Source/Editor/Sequencer/Private/MVVM/ViewModels/ChannelModel.cpp
+++ b/Engine/Source/Editor/Sequencer/Private/MVVM/ViewModels/ChannelModel.cpp
@@ -32,6 +32,7 @@
 #include "SKeyAreaEditorSwitcher.h"
 #include "SequencerCoreFwd.h"
 #include "SequencerSectionPainter.h"
+#include "SequencerSelection.h"
 #include "SequencerSettings.h"
 #include "SlotBase.h"
 #include "Styling/AppStyle.h"
@@ -300,6 +301,11 @@ uint32 FChannelGroupModel::GetChannelsSerialNumber() const
 	return ChannelsSerialNumber;
 }
 
+void FChannelGroupModel::OnChannelOverridesChanged()
+{
+	++ChannelsSerialNumber;
+}
+
 void FChannelGroupModel::CleanupChannels()
 {
 	const int32 NumRemoved = Channels.RemoveAll([](TWeakViewModelPtr<FChannelModel> Item) { return !Item.IsValid(); });
@@ -426,51 +432,81 @@ bool FChannelGroupModel::HasCurves() const
 
 void FChannelGroupModel::BuildChannelOverrideMenu(FMenuBuilder& MenuBuilder)
 {
+	TSharedPtr<FSequenceModel> SequenceModel = FindAncestorOfType<FSequenceModel>();
+	FChannelGroupOverrideHelper::BuildChannelOverrideMenu(MenuBuilder, SequenceModel);
+}
+
+void FChannelGroupOverrideHelper::BuildChannelOverrideMenu(FMenuBuilder& MenuBuilder, TSharedPtr<FSequenceModel> SequenceModel)
+{
+	if (!SequenceModel)
+	{
+		return;
+	}
+
+	TSharedPtr<ISequencer> Sequencer = SequenceModel->GetSequencer();
+	if (!Sequencer)
+	{
+		return;
+	}
+
+	// Get all selected outliner items. We will filter out those that aren't channel groups.
+	TArray<TSharedRef<FViewModel>> SelectedItems;
+	Sequencer->GetSelection().GetSelectedOutlinerItems(SelectedItems);
+
 	// Gather up information:
-	// - Candidates for overriding each of the channel types found in this group.
+	// - Candidates for overriding each of the channel types found in each group.
 	// - Already overriden channels.
 	int32 NumOverridableChannels = 0;
 	TMap<UClass*, TArray<UMovieSceneChannelOverrideContainer*>> ExistingOverrides;
 	TArray<UMovieSceneChannelOverrideContainer::FOverrideCandidates> AllCandidateOverrides;
-	for (int32 Index = 0; Index < GetChannels().Num(); ++Index)
+	for (int32 ChannelGroupIndex = 0; ChannelGroupIndex < SelectedItems.Num(); ++ChannelGroupIndex)
 	{
-		TSharedPtr<FChannelModel> Channel = GetChannel(Index);
-		if (!Channel)
+		TSharedPtr<FChannelGroupModel> SelectedChannelGroup = SelectedItems[ChannelGroupIndex]->CastThisShared<FChannelGroupModel>();
+		if (!SelectedChannelGroup)
 		{
 			continue;
 		}
-
-		UMovieSceneSection* Section = Channel->GetSection();
-		if (!Section)
+		
+		for (int32 ChannelIndex = 0; ChannelIndex < SelectedChannelGroup->GetChannels().Num(); ++ChannelIndex)
 		{
-			continue;
-		}
+			TSharedPtr<FChannelModel> Channel = SelectedChannelGroup->GetChannel(ChannelIndex);
+			if (!Channel)
+			{
+				continue;
+			}
 
-		IMovieSceneChannelOverrideProvider* OverrideProvider = Cast<IMovieSceneChannelOverrideProvider>(Section);
-		if (!OverrideProvider)
-		{
-			continue;
-		}
+			UMovieSceneSection* Section = Channel->GetSection();
+			if (!Section)
+			{
+				continue;
+			}
 
-		++NumOverridableChannels;
+			IMovieSceneChannelOverrideProvider* OverrideProvider = Cast<IMovieSceneChannelOverrideProvider>(Section);
+			if (!OverrideProvider)
+			{
+				continue;
+			}
 
-		UMovieSceneSectionChannelOverrideRegistry* OverrideRegistry = OverrideProvider->GetChannelOverrideRegistry(false);
-		UE::MovieScene::FChannelOverrideProviderTraitsHandle OverrideTraits = OverrideProvider->GetChannelOverrideProviderTraits();
+			++NumOverridableChannels;
 
-		UMovieSceneChannelOverrideContainer* OverrideContainer = OverrideRegistry ? 
-			OverrideRegistry->GetChannel(GetChannelName()) : nullptr;
-		if (OverrideContainer)
-		{
-			ExistingOverrides.FindOrAdd(OverrideContainer->GetClass()).Add(OverrideContainer);
-		}
-		
-		const FName DefaultChannelTypeName = OverrideTraits->GetDefaultChannelTypeName(GetChannelName());
-		if (!DefaultChannelTypeName.IsNone())
-		{
-			UMovieSceneChannelOverrideContainer::FOverrideCandidates OverrideCandidates;
-			UMovieSceneChannelOverrideContainer::GetOverrideCandidates(DefaultChannelTypeName, OverrideCandidates);
+			UMovieSceneSectionChannelOverrideRegistry* OverrideRegistry = OverrideProvider->GetChannelOverrideRegistry(false);
+			UE::MovieScene::FChannelOverrideProviderTraitsHandle OverrideTraits = OverrideProvider->GetChannelOverrideProviderTraits();
 
-			AllCandidateOverrides.Add(MoveTemp(OverrideCandidates));
+			UMovieSceneChannelOverrideContainer* OverrideContainer = OverrideRegistry ? 
+				OverrideRegistry->GetChannel(Channel->GetChannelName()) : nullptr;
+			if (OverrideContainer)
+			{
+				ExistingOverrides.FindOrAdd(OverrideContainer->GetClass()).Add(OverrideContainer);
+			}
+
+			const FName DefaultChannelTypeName = OverrideTraits->GetDefaultChannelTypeName(Channel->GetChannelName());
+			if (!DefaultChannelTypeName.IsNone())
+			{
+				UMovieSceneChannelOverrideContainer::FOverrideCandidates OverrideCandidates;
+				UMovieSceneChannelOverrideContainer::GetOverrideCandidates(DefaultChannelTypeName, OverrideCandidates);
+
+				AllCandidateOverrides.Add(MoveTemp(OverrideCandidates));
+			}
 		}
 	}
 
@@ -499,14 +535,14 @@ void FChannelGroupModel::BuildChannelOverrideMenu(FMenuBuilder& MenuBuilder)
 
 		if (CommonCandidateOverrides.Num() == 1)
 		{
-			BuildChannelOverrideMenu(MenuBuilder, CommonCandidateOverrides);
+			BuildChannelOverrideMenu(MenuBuilder, SequenceModel, CommonCandidateOverrides);
 		}
 		else if (CommonCandidateOverrides.Num() > 1)
 		{
 			MenuBuilder.AddSubMenu(
 				LOCTEXT("OverrideChannelChoiceLabel", "Override with..."),
 				LOCTEXT("OverrideChannelChoiceTooltip", "Overrides all channels with the specified channel"),
-				FNewMenuDelegate::CreateSP(this, &FChannelGroupOutlinerModel::BuildChannelOverrideMenu, CommonCandidateOverrides));
+				FNewMenuDelegate::CreateStatic(&FChannelGroupOverrideHelper::BuildChannelOverrideMenu, SequenceModel, CommonCandidateOverrides));
 		}
 
 		if (ExistingOverrides.Num() > 0)
@@ -516,7 +552,7 @@ void FChannelGroupModel::BuildChannelOverrideMenu(FMenuBuilder& MenuBuilder)
 				LOCTEXT("RemoveOverrideChannelTooltip", "Removes the channel override and revert back to the default channel type"),
 				FSlateIcon(),
 				FUIAction(
-					FExecuteAction::CreateSP(this, &FChannelGroupOutlinerModel::RemoveChannelOverrides),
+					FExecuteAction::CreateStatic(&FChannelGroupOverrideHelper::RemoveChannelOverrides, SequenceModel),
 					FCanExecuteAction()));
 		}
 
@@ -524,7 +560,7 @@ void FChannelGroupModel::BuildChannelOverrideMenu(FMenuBuilder& MenuBuilder)
 	}
 }
 
-void FChannelGroupModel::BuildChannelOverrideMenu(FMenuBuilder& MenuBuilder, UMovieSceneChannelOverrideContainer::FOverrideCandidates OverrideCandidates)
+void FChannelGroupOverrideHelper::BuildChannelOverrideMenu(FMenuBuilder& MenuBuilder, TSharedPtr<FSequenceModel> SequenceModel, UMovieSceneChannelOverrideContainer::FOverrideCandidates OverrideCandidates)
 {
 	for (TSubclassOf<UMovieSceneChannelOverrideContainer> OverrideCandidate : OverrideCandidates)
 	{
@@ -535,48 +571,91 @@ void FChannelGroupModel::BuildChannelOverrideMenu(FMenuBuilder& MenuBuilder, UMo
 				OverrideCandidate->GetDisplayNameText()),
 			FSlateIcon(),
 			FUIAction(
-				FExecuteAction::CreateSP(this, &FChannelGroupModel::OverrideChannels, OverrideCandidate),
+				FExecuteAction::CreateStatic(&FChannelGroupOverrideHelper::OverrideChannels, SequenceModel, OverrideCandidate),
 				FCanExecuteAction()));
 	}
 }
 
-void FChannelGroupModel::OverrideChannels(TSubclassOf<UMovieSceneChannelOverrideContainer> OverrideClass)
+void FChannelGroupOverrideHelper::OverrideChannels(TSharedPtr<FSequenceModel> SequenceModel, TSubclassOf<UMovieSceneChannelOverrideContainer> OverrideClass)
 {
-	const FScopedTransaction Transaction(LOCTEXT("OverrideChannels", "Override Channels"));
+	if (!SequenceModel)
+	{
+		return;
+	}
 
-	for (TWeakViewModelPtr<FChannelModel> WeakChannel : GetChannels())
+	TSharedPtr<ISequencer> Sequencer = SequenceModel->GetSequencer();
+	if (!Sequencer)
+	{
+		return;
+	}
+
+	bool bAnyChannelsChanged = false;
+	FScopedTransaction Transaction(LOCTEXT("OverrideChannels", "Override Channels"));
+
+	// Get all selected outliner items. We will filter out those that aren't channel groups.
+	TArray<TSharedRef<FViewModel>> SelectedItems;
+	Sequencer->GetSelection().GetSelectedOutlinerItems(SelectedItems);
+
+	for (int32 ChannelGroupIndex = 0; ChannelGroupIndex < SelectedItems.Num(); ++ChannelGroupIndex)
 	{
-		TSharedPtr<FChannelModel> Channel = WeakChannel.Pin();
-		if (!Channel || !Channel->GetSection())
+		TSharedPtr<FChannelGroupModel> SelectedChannelGroup = SelectedItems[ChannelGroupIndex]->CastThisShared<FChannelGroupModel>();
+		if (!SelectedChannelGroup)
 		{
 			continue;
 		}
 
-		UMovieSceneSection* Section = Channel->GetSection();
-		IMovieSceneChannelOverrideProvider* OverrideProvider = Cast<IMovieSceneChannelOverrideProvider>(Section);
-		if (!OverrideProvider)
+		bool bChannelsChanged = false;
+
+		for (TWeakViewModelPtr<FChannelModel> WeakChannel : SelectedChannelGroup->GetChannels())
 		{
-			continue;
+			TSharedPtr<FChannelModel> Channel = WeakChannel.Pin();
+			if (!Channel || !Channel->GetSection())
+			{
+				continue;
+			}
+
+			UMovieSceneSection* Section = Channel->GetSection();
+			IMovieSceneChannelOverrideProvider* OverrideProvider = Cast<IMovieSceneChannelOverrideProvider>(Section);
+			if (!OverrideProvider)
+			{
+				continue;
+			}
+
+			if (!Section->TryModify())
+			{
+				continue;
+			}
+
+			UMovieSceneSectionChannelOverrideRegistry* OverrideRegistry = OverrideProvider->GetChannelOverrideRegistry(true);
+			check(OverrideRegistry);
+
+			UMovieSceneChannelOverrideContainer* OverrideChannelContainer = NewObject<UMovieSceneChannelOverrideContainer>(
+					OverrideRegistry, OverrideClass, NAME_None, RF_Transactional);
+			OverrideChannelContainer->InitializeOverride(Channel->GetChannel());
+			OverrideRegistry->AddChannel(Channel->GetChannelName(), OverrideChannelContainer);	
+			
+			OverrideProvider->OnChannelOverridesChanged();
+
+			bAnyChannelsChanged = true;
+			bChannelsChanged = true;
 		}
 
-		if (!Section->TryModify())
+		if (bChannelsChanged)
 		{
-			continue;
+			SelectedChannelGroup->OnChannelOverridesChanged();
 		}
+	}
 
-		UMovieSceneSectionChannelOverrideRegistry* OverrideRegistry = OverrideProvider->GetChannelOverrideRegistry(true);
-		check(OverrideRegistry);
-
-		UMovieSceneChannelOverrideContainer* OverrideChannelContainer = NewObject<UMovieSceneChannelOverrideContainer>(Section, OverrideClass);
-		OverrideChannelContainer->InitializeOverride(Channel->GetChannel());
-		OverrideRegistry->AddChannel(GetChannelName(), OverrideChannelContainer);	
-		
-		OverrideProvider->OnChannelOverridesChanged();
-
-		++ChannelsSerialNumber;
+	if (!bAnyChannelsChanged)
+	{
+		Transaction.Cancel();
 	}
 
-	TSharedPtr<FSequenceModel> SequenceModel = FindAncestorOfType<FSequenceModel>();
+	Sequencer->NotifyMovieSceneDataChanged(EMovieSceneDataChangeType::MovieSceneStructureItemsChanged);
+}
+
+void FChannelGroupOverrideHelper::RemoveChannelOverrides(TSharedPtr<FSequenceModel> SequenceModel)
+{
 	if (!SequenceModel)
 	{
 		return;
@@ -588,55 +667,65 @@ void FChannelGroupModel::OverrideChannels(TSubclassOf<UMovieSceneChannelOverride
 		return;
 	}
 
-	Sequencer->NotifyMovieSceneDataChanged(EMovieSceneDataChangeType::MovieSceneStructureItemsChanged);
-}
+	bool bAnyChannelsChanged = false;
+	FScopedTransaction Transaction(LOCTEXT("RemoveChannelOverrides", "Remove Override Channels"));
 
-void FChannelGroupModel::RemoveChannelOverrides()
-{
-	const FScopedTransaction Transaction(LOCTEXT("RemoveChannelOverrides", "Remove Override Channels"));
+	// Get all selected outliner items. We will filter out those that aren't channel groups.
+	TArray<TSharedRef<FViewModel>> SelectedItems;
+	Sequencer->GetSelection().GetSelectedOutlinerItems(SelectedItems);
 
-	for (TWeakViewModelPtr<FChannelModel> WeakChannel : GetChannels())
+	for (int32 ChannelGroupIndex = 0; ChannelGroupIndex < SelectedItems.Num(); ++ChannelGroupIndex)
 	{
-		TSharedPtr<FChannelModel> Channel = WeakChannel.Pin();
-		if (!Channel || !Channel->GetSection())
+		TSharedPtr<FChannelGroupModel> SelectedChannelGroup = SelectedItems[ChannelGroupIndex]->CastThisShared<FChannelGroupModel>();
+		if (!SelectedChannelGroup)
 		{
 			continue;
 		}
 
-		IMovieSceneChannelOverrideProvider* OverrideProvider = Cast<IMovieSceneChannelOverrideProvider>(Channel->GetSection());
-		if (!OverrideProvider)
-		{
-			continue;
-		}
+		bool bChannelsChanged = false;
 
-		UMovieSceneSectionChannelOverrideRegistry* OverrideRegistry = OverrideProvider->GetChannelOverrideRegistry(false);
-		if (!OverrideRegistry)
+		for (TWeakViewModelPtr<FChannelModel> WeakChannel : SelectedChannelGroup->GetChannels())
 		{
-			continue;
-		}
+			TSharedPtr<FChannelModel> Channel = WeakChannel.Pin();
+			if (!Channel || !Channel->GetSection())
+			{
+				continue;
+			}
 
-		if (!Channel->GetSection()->TryModify())
-		{
-			continue;
-		}
+			IMovieSceneChannelOverrideProvider* OverrideProvider = Cast<IMovieSceneChannelOverrideProvider>(Channel->GetSection());
+			if (!OverrideProvider)
+			{
+				continue;
+			}
+
+			UMovieSceneSectionChannelOverrideRegistry* OverrideRegistry = OverrideProvider->GetChannelOverrideRegistry(false);
+			if (!OverrideRegistry)
+			{
+				continue;
+			}
+
+			if (!Channel->GetSection()->TryModify())
+			{
+				continue;
+			}
 
-		OverrideRegistry->RemoveChannel(GetChannelName());
+			OverrideRegistry->RemoveChannel(Channel->GetChannelName());
 
-		OverrideProvider->OnChannelOverridesChanged();
+			OverrideProvider->OnChannelOverridesChanged();
 
-		++ChannelsSerialNumber;
-	}
+			bAnyChannelsChanged = true;
+			bChannelsChanged = true;
+		}
 
-	TSharedPtr<FSequenceModel> SequenceModel = FindAncestorOfType<FSequenceModel>();
-	if (!SequenceModel)
-	{
-		return;
+		if (bChannelsChanged)
+		{
+			SelectedChannelGroup->OnChannelOverridesChanged();
+		}
 	}
 
-	TSharedPtr<ISequencer> Sequencer = SequenceModel->GetSequencer();
-	if (!Sequencer)
+	if (!bAnyChannelsChanged)
 	{
-		return;
+		Transaction.Cancel();
 	}
 
 	Sequencer->NotifyMovieSceneDataChanged(EMovieSceneDataChangeType::MovieSceneStructureItemsChanged);
diff --git a/Engine/Source/Editor/Sequencer/Public/MVVM/ViewModels/ChannelModel.h b/Engine/Source/Editor/Sequencer/Public/MVVM/ViewModels/ChannelModel.h
--- a/Engine/Source/Editor/Sequencer/Public/MVVM/ViewModels/ChannelModel.h
+++ b/Engine/Source/Editor/Sequencer/Public/MVVM/ViewModels/ChannelModel.h
@@ -44,6 +44,7 @@ namespace Sequencer
 {
 
 class FSectionModel;
+class FSequenceModel;
 
 /**
  * Model for a single channel inside a section.
@@ -164,6 +165,9 @@ public:
 	/** Gets a serial number representing if the list of channels has changed */
 	uint32 GetChannelsSerialNumber() const;
 
+	/** Notifies this view-model that the list of channels has changed */
+	void OnChannelOverridesChanged();
+
 public:
 
 	/*~ ITrackAreaExtension */
@@ -180,12 +184,6 @@ public:
 
 	void CleanupChannels();
 
-private:
-
-	void BuildChannelOverrideMenu(FMenuBuilder& MenuBuilder, UMovieSceneChannelOverrideContainer::FOverrideCandidates OverrideCandidates);
-	void OverrideChannels(TSubclassOf<UMovieSceneChannelOverrideContainer> OverrideClass);
-	void RemoveChannelOverrides();
-
 protected:
 
 	TArray<TWeakViewModelPtr<FChannelModel>> Channels;
@@ -242,6 +240,23 @@ private:
 	FOutlinerSizing ComputedSizing;
 };
 
+
+/**
+ * Utility class for building menus that add, edit, and remove channel overrides.
+ */
+class SEQUENCER_API FChannelGroupOverrideHelper
+{
+public:
+
+	static void BuildChannelOverrideMenu(FMenuBuilder& MenuBuilder, TSharedPtr<FSequenceModel> SequenceModel);
+
+private:
+
+	static void BuildChannelOverrideMenu(FMenuBuilder& MenuBuilder, TSharedPtr<FSequenceModel> SequenceModel, UMovieSceneChannelOverrideContainer::FOverrideCandidates OverrideCandidates);
+	static void OverrideChannels(TSharedPtr<FSequenceModel> SequenceModel, TSubclassOf<UMovieSceneChannelOverrideContainer> OverrideClass);
+	static void RemoveChannelOverrides(TSharedPtr<FSequenceModel> SequenceModel);
+};
+
 } // namespace Sequencer
 } // namespace UE
 
diff --git a/Engine/Source/Runtime/MovieScene/Private/Channels/MovieSceneSectionChannelOverrideRegistry.cpp b/Engine/Source/Runtime/MovieScene/Private/Channels/MovieSceneSectionChannelOverrideRegistry.cpp
--- a/Engine/Source/Runtime/MovieScene/Private/Channels/MovieSceneSectionChannelOverrideRegistry.cpp
+++ b/Engine/Source/Runtime/MovieScene/Private/Channels/MovieSceneSectionChannelOverrideRegistry.cpp
@@ -15,6 +15,7 @@ UMovieSceneSectionChannelOverrideRegistry::UMovieSceneSectionChannelOverrideRegi
 
 void UMovieSceneSectionChannelOverrideRegistry::AddChannel(FName ChannelName, UMovieSceneChannelOverrideContainer* ChannelContainer)
 {
+	Modify();
 	Overrides.Emplace(ChannelName, ChannelContainer);
 }
 
@@ -35,6 +36,7 @@ UMovieSceneChannelOverrideContainer* UMovieSceneSectionChannelOverrideRegistry::
 
 void UMovieSceneSectionChannelOverrideRegistry::RemoveChannel(FName ChannelName)
 {
+	Modify();
 	Overrides.Remove(ChannelName);
 }
 
diff --git a/Engine/Source/Runtime/MovieSceneTracks/Private/Sections/MovieScene3DTransformSection.cpp b/Engine/Source/Runtime/MovieSceneTracks/Private/Sections/MovieScene3DTransformSection.cpp
--- a/Engine/Source/Runtime/MovieSceneTracks/Private/Sections/MovieScene3DTransformSection.cpp
+++ b/Engine/Source/Runtime/MovieSceneTracks/Private/Sections/MovieScene3DTransformSection.cpp
@@ -1042,7 +1042,7 @@ UMovieSceneSectionChannelOverrideRegistry* UMovieScene3DTransformSection::GetCha
 {
 	if (bCreateIfMissing && OverrideRegistry == nullptr)
 	{
-		OverrideRegistry = NewObject<UMovieSceneSectionChannelOverrideRegistry>(this);
+		OverrideRegistry = NewObject<UMovieSceneSectionChannelOverrideRegistry>(this, NAME_None, RF_Transactional);
 	}
 	return OverrideRegistry;
 }
diff --git a/Engine/Source/Runtime/MovieSceneTracks/Private/Sections/MovieSceneFloatSection.cpp b/Engine/Source/Runtime/MovieSceneTracks/Private/Sections/MovieSceneFloatSection.cpp
--- a/Engine/Source/Runtime/MovieSceneTracks/Private/Sections/MovieSceneFloatSection.cpp
+++ b/Engine/Source/Runtime/MovieSceneTracks/Private/Sections/MovieSceneFloatSection.cpp
@@ -76,7 +76,7 @@ UMovieSceneSectionChannelOverrideRegistry* UMovieSceneFloatSection::GetChannelOv
 {
 	if (bCreateIfMissing && OverrideRegistry == nullptr)
 	{
-		OverrideRegistry = NewObject<UMovieSceneSectionChannelOverrideRegistry>(this);
+		OverrideRegistry = NewObject<UMovieSceneSectionChannelOverrideRegistry>(this, NAME_None, RF_Transactional);
 	}
 	return OverrideRegistry;
 }
