Sequencer: gather overridable channels from tracks' and track rows' top level channel groups in addition to channel groups

diff --git a/Engine/Source/Editor/Sequencer/Private/MVVM/ViewModels/ChannelModel.cpp b/Engine/Source/Editor/Sequencer/Private/MVVM/ViewModels/ChannelModel.cpp
--- a/Engine/Source/Editor/Sequencer/Private/MVVM/ViewModels/ChannelModel.cpp
+++ b/Engine/Source/Editor/Sequencer/Private/MVVM/ViewModels/ChannelModel.cpp
@@ -19,6 +19,7 @@
 #include "MVVM/ViewModels/SequencerModelUtils.h"
 #include "MVVM/ViewModels/ViewModelIterators.h"
 #include "MVVM/ViewModels/OutlinerViewModel.h"
+#include "MVVM/ViewModels/TrackModel.h"
 #include "MVVM/Views/SOutlinerView.h"
 #include "MVVM/Views/SOutlinerItemViewBase.h"
 #include "MVVM/Views/SSequencerKeyNavigationButtons.h"
@@ -436,77 +437,110 @@ void FChannelGroupModel::BuildChannelOverrideMenu(FMenuBuilder& MenuBuilder)
 	FChannelGroupOverrideHelper::BuildChannelOverrideMenu(MenuBuilder, SequenceModel);
 }
 
-void FChannelGroupOverrideHelper::BuildChannelOverrideMenu(FMenuBuilder& MenuBuilder, TSharedPtr<FSequenceModel> SequenceModel)
+bool FChannelGroupOverrideHelper::GetChannelsFromOutlinerSelection(TSharedPtr<FSequenceModel> SequenceModel, TArray<TViewModelPtr<FChannelModel>>& OutChannelModels)
 {
+	// Sanity checks.
 	if (!SequenceModel)
 	{
-		return;
+		return false;
 	}
 
 	TSharedPtr<ISequencer> Sequencer = SequenceModel->GetSequencer();
 	if (!Sequencer)
 	{
-		return;
+		return false;
 	}
 
 	// Get all selected outliner items. We will filter out those that aren't channel groups.
 	TArray<TSharedRef<FViewModel>> SelectedItems;
 	Sequencer->GetSelection().GetSelectedOutlinerItems(SelectedItems);
 
+	// The selected items could either be channel groups, or tracks/track-rows with top-level channel groups that
+	// don't explicitly show up in the outliner (they are view-models that live under the track or row, but the
+	// user interface shows them as the track or row itself).
+	for (TSharedRef<FViewModel> SelectedItem : SelectedItems)
+	{
+		if (FChannelGroupModel* ChannelGroup = SelectedItem->CastThis<FChannelGroupModel>())
+		{
+			const int32 NumChannels = ChannelGroup->GetChannels().Num();
+			for (int32 ChannelIndex = 0; ChannelIndex < NumChannels; ++ChannelIndex)
+			{
+				if (TSharedPtr<FChannelModel> Channel = ChannelGroup->GetChannel(ChannelIndex))
+				{
+					OutChannelModels.Add(Channel);
+				}
+			}
+		}
+		else
+		{
+			const EViewModelListType ListType = FTrackModel::GetTopLevelChannelGroupType();
+			TOptional<FViewModelChildren> TopLevelChannels = SelectedItem->FindChildList(ListType);
+			if (TopLevelChannels.IsSet())
+			{
+				for (auto It(TopLevelChannels->IterateSubList<FChannelGroupModel>()); It; ++It)
+				{
+					for (TWeakViewModelPtr<FChannelModel> WeakChannel : It->GetChannels())
+					{
+						if (TViewModelPtr<FChannelModel> Channel = WeakChannel.Pin())
+						{
+							OutChannelModels.Add(Channel);
+						}
+					}
+				}
+			}
+		}
+	}
+
+	return true;
+}
+
+void FChannelGroupOverrideHelper::BuildChannelOverrideMenu(FMenuBuilder& MenuBuilder, TSharedPtr<FSequenceModel> SequenceModel)
+{
+	TArray<TViewModelPtr<FChannelModel>> Channels;
+	if (!GetChannelsFromOutlinerSelection(SequenceModel, Channels))
+	{
+		return;
+	}
+
 	// Gather up information:
 	// - Candidates for overriding each of the channel types found in each group.
 	// - Already overriden channels.
 	int32 NumOverridableChannels = 0;
 	TMap<UClass*, TArray<UMovieSceneChannelOverrideContainer*>> ExistingOverrides;
 	TArray<UMovieSceneChannelOverrideContainer::FOverrideCandidates> AllCandidateOverrides;
-	for (int32 ChannelGroupIndex = 0; ChannelGroupIndex < SelectedItems.Num(); ++ChannelGroupIndex)
+	for (TViewModelPtr Channel : Channels)
 	{
-		TSharedPtr<FChannelGroupModel> SelectedChannelGroup = SelectedItems[ChannelGroupIndex]->CastThisShared<FChannelGroupModel>();
-		if (!SelectedChannelGroup)
+		UMovieSceneSection* Section = Channel->GetSection();
+		if (!Section)
 		{
 			continue;
 		}
-		
-		for (int32 ChannelIndex = 0; ChannelIndex < SelectedChannelGroup->GetChannels().Num(); ++ChannelIndex)
-		{
-			TSharedPtr<FChannelModel> Channel = SelectedChannelGroup->GetChannel(ChannelIndex);
-			if (!Channel)
-			{
-				continue;
-			}
 
-			UMovieSceneSection* Section = Channel->GetSection();
-			if (!Section)
-			{
-				continue;
-			}
-
-			IMovieSceneChannelOverrideProvider* OverrideProvider = Cast<IMovieSceneChannelOverrideProvider>(Section);
-			if (!OverrideProvider)
-			{
-				continue;
-			}
+		IMovieSceneChannelOverrideProvider* OverrideProvider = Cast<IMovieSceneChannelOverrideProvider>(Section);
+		if (!OverrideProvider)
+		{
+			continue;
+		}
 
-			++NumOverridableChannels;
+		++NumOverridableChannels;
 
-			UMovieSceneSectionChannelOverrideRegistry* OverrideRegistry = OverrideProvider->GetChannelOverrideRegistry(false);
-			UE::MovieScene::FChannelOverrideProviderTraitsHandle OverrideTraits = OverrideProvider->GetChannelOverrideProviderTraits();
+		UMovieSceneSectionChannelOverrideRegistry* OverrideRegistry = OverrideProvider->GetChannelOverrideRegistry(false);
+		UE::MovieScene::FChannelOverrideProviderTraitsHandle OverrideTraits = OverrideProvider->GetChannelOverrideProviderTraits();
 
-			UMovieSceneChannelOverrideContainer* OverrideContainer = OverrideRegistry ? 
-				OverrideRegistry->GetChannel(Channel->GetChannelName()) : nullptr;
-			if (OverrideContainer)
-			{
-				ExistingOverrides.FindOrAdd(OverrideContainer->GetClass()).Add(OverrideContainer);
-			}
+		UMovieSceneChannelOverrideContainer* OverrideContainer = OverrideRegistry ? 
+			OverrideRegistry->GetChannel(Channel->GetChannelName()) : nullptr;
+		if (OverrideContainer)
+		{
+			ExistingOverrides.FindOrAdd(OverrideContainer->GetClass()).Add(OverrideContainer);
+		}
 
-			const FName DefaultChannelTypeName = OverrideTraits->GetDefaultChannelTypeName(Channel->GetChannelName());
-			if (!DefaultChannelTypeName.IsNone())
-			{
-				UMovieSceneChannelOverrideContainer::FOverrideCandidates OverrideCandidates;
-				UMovieSceneChannelOverrideContainer::GetOverrideCandidates(DefaultChannelTypeName, OverrideCandidates);
+		const FName DefaultChannelTypeName = OverrideTraits->GetDefaultChannelTypeName(Channel->GetChannelName());
+		if (!DefaultChannelTypeName.IsNone())
+		{
+			UMovieSceneChannelOverrideContainer::FOverrideCandidates OverrideCandidates;
+			UMovieSceneChannelOverrideContainer::GetOverrideCandidates(DefaultChannelTypeName, OverrideCandidates);
 
-				AllCandidateOverrides.Add(MoveTemp(OverrideCandidates));
-			}
+			AllCandidateOverrides.Add(MoveTemp(OverrideCandidates));
 		}
 	}
 
@@ -578,157 +612,114 @@ void FChannelGroupOverrideHelper::BuildChannelOverrideMenu(FMenuBuilder& MenuBui
 
 void FChannelGroupOverrideHelper::OverrideChannels(TSharedPtr<FSequenceModel> SequenceModel, TSubclassOf<UMovieSceneChannelOverrideContainer> OverrideClass)
 {
-	if (!SequenceModel)
-	{
-		return;
-	}
-
-	TSharedPtr<ISequencer> Sequencer = SequenceModel->GetSequencer();
-	if (!Sequencer)
+	TArray<TViewModelPtr<FChannelModel>> Channels;
+	if (!GetChannelsFromOutlinerSelection(SequenceModel, Channels))
 	{
 		return;
 	}
 
 	bool bAnyChannelsChanged = false;
+	TSet<TSharedPtr<FChannelGroupModel>> ChangedChannelGroups;
 	FScopedTransaction Transaction(LOCTEXT("OverrideChannels", "Override Channels"));
 
-	// Get all selected outliner items. We will filter out those that aren't channel groups.
-	TArray<TSharedRef<FViewModel>> SelectedItems;
-	Sequencer->GetSelection().GetSelectedOutlinerItems(SelectedItems);
-
-	for (int32 ChannelGroupIndex = 0; ChannelGroupIndex < SelectedItems.Num(); ++ChannelGroupIndex)
+	for (TViewModelPtr<FChannelModel> Channel : Channels)
 	{
-		TSharedPtr<FChannelGroupModel> SelectedChannelGroup = SelectedItems[ChannelGroupIndex]->CastThisShared<FChannelGroupModel>();
-		if (!SelectedChannelGroup)
+		UMovieSceneSection* Section = Channel->GetSection();
+		IMovieSceneChannelOverrideProvider* OverrideProvider = Cast<IMovieSceneChannelOverrideProvider>(Section);
+		if (!OverrideProvider)
 		{
 			continue;
 		}
 
-		bool bChannelsChanged = false;
-
-		for (TWeakViewModelPtr<FChannelModel> WeakChannel : SelectedChannelGroup->GetChannels())
+		if (!Section->TryModify())
 		{
-			TSharedPtr<FChannelModel> Channel = WeakChannel.Pin();
-			if (!Channel || !Channel->GetSection())
-			{
-				continue;
-			}
-
-			UMovieSceneSection* Section = Channel->GetSection();
-			IMovieSceneChannelOverrideProvider* OverrideProvider = Cast<IMovieSceneChannelOverrideProvider>(Section);
-			if (!OverrideProvider)
-			{
-				continue;
-			}
-
-			if (!Section->TryModify())
-			{
-				continue;
-			}
+			continue;
+		}
 
-			UMovieSceneSectionChannelOverrideRegistry* OverrideRegistry = OverrideProvider->GetChannelOverrideRegistry(true);
-			check(OverrideRegistry);
+		UMovieSceneSectionChannelOverrideRegistry* OverrideRegistry = OverrideProvider->GetChannelOverrideRegistry(true);
+		check(OverrideRegistry);
 
-			UMovieSceneChannelOverrideContainer* OverrideChannelContainer = NewObject<UMovieSceneChannelOverrideContainer>(
-					OverrideRegistry, OverrideClass, NAME_None, RF_Transactional);
-			OverrideChannelContainer->InitializeOverride(Channel->GetChannel());
-			OverrideRegistry->AddChannel(Channel->GetChannelName(), OverrideChannelContainer);	
-			
-			OverrideProvider->OnChannelOverridesChanged();
+		UMovieSceneChannelOverrideContainer* OverrideChannelContainer = NewObject<UMovieSceneChannelOverrideContainer>(
+				OverrideRegistry, OverrideClass, NAME_None, RF_Transactional);
+		OverrideChannelContainer->InitializeOverride(Channel->GetChannel());
+		OverrideRegistry->AddChannel(Channel->GetChannelName(), OverrideChannelContainer);	
 
-			bAnyChannelsChanged = true;
-			bChannelsChanged = true;
-		}
+		OverrideProvider->OnChannelOverridesChanged();
 
-		if (bChannelsChanged)
+		bAnyChannelsChanged = true;
+		if (TSharedPtr<FChannelGroupModel> ParentGroupModel = Channel->CastParent<FChannelGroupModel>())
 		{
-			SelectedChannelGroup->OnChannelOverridesChanged();
+			ChangedChannelGroups.Add(ParentGroupModel);
 		}
 	}
 
+	for (TSharedPtr<FChannelGroupModel> ChangedChannelGroup : ChangedChannelGroups)
+	{
+		ChangedChannelGroup->OnChannelOverridesChanged();
+	}
+
 	if (!bAnyChannelsChanged)
 	{
 		Transaction.Cancel();
 	}
 
-	Sequencer->NotifyMovieSceneDataChanged(EMovieSceneDataChangeType::MovieSceneStructureItemsChanged);
+	SequenceModel->GetSequencer()->NotifyMovieSceneDataChanged(EMovieSceneDataChangeType::MovieSceneStructureItemsChanged);
 }
 
 void FChannelGroupOverrideHelper::RemoveChannelOverrides(TSharedPtr<FSequenceModel> SequenceModel)
 {
-	if (!SequenceModel)
-	{
-		return;
-	}
-
-	TSharedPtr<ISequencer> Sequencer = SequenceModel->GetSequencer();
-	if (!Sequencer)
+	TArray<TViewModelPtr<FChannelModel>> Channels;
+	if (!GetChannelsFromOutlinerSelection(SequenceModel, Channels))
 	{
 		return;
 	}
 
 	bool bAnyChannelsChanged = false;
+	TSet<TSharedPtr<FChannelGroupModel>> ChangedChannelGroups;
 	FScopedTransaction Transaction(LOCTEXT("RemoveChannelOverrides", "Remove Override Channels"));
 
-	// Get all selected outliner items. We will filter out those that aren't channel groups.
-	TArray<TSharedRef<FViewModel>> SelectedItems;
-	Sequencer->GetSelection().GetSelectedOutlinerItems(SelectedItems);
-
-	for (int32 ChannelGroupIndex = 0; ChannelGroupIndex < SelectedItems.Num(); ++ChannelGroupIndex)
+	for (TViewModelPtr<FChannelModel> Channel : Channels)
 	{
-		TSharedPtr<FChannelGroupModel> SelectedChannelGroup = SelectedItems[ChannelGroupIndex]->CastThisShared<FChannelGroupModel>();
-		if (!SelectedChannelGroup)
+		UMovieSceneSection* Section = Channel->GetSection();
+		IMovieSceneChannelOverrideProvider* OverrideProvider = Cast<IMovieSceneChannelOverrideProvider>(Section);
+		if (!OverrideProvider)
 		{
 			continue;
 		}
 
-		bool bChannelsChanged = false;
-
-		for (TWeakViewModelPtr<FChannelModel> WeakChannel : SelectedChannelGroup->GetChannels())
+		UMovieSceneSectionChannelOverrideRegistry* OverrideRegistry = OverrideProvider->GetChannelOverrideRegistry(false);
+		if (!OverrideRegistry)
 		{
-			TSharedPtr<FChannelModel> Channel = WeakChannel.Pin();
-			if (!Channel || !Channel->GetSection())
-			{
-				continue;
-			}
-
-			IMovieSceneChannelOverrideProvider* OverrideProvider = Cast<IMovieSceneChannelOverrideProvider>(Channel->GetSection());
-			if (!OverrideProvider)
-			{
-				continue;
-			}
-
-			UMovieSceneSectionChannelOverrideRegistry* OverrideRegistry = OverrideProvider->GetChannelOverrideRegistry(false);
-			if (!OverrideRegistry)
-			{
-				continue;
-			}
-
-			if (!Channel->GetSection()->TryModify())
-			{
-				continue;
-			}
+			continue;
+		}
 
-			OverrideRegistry->RemoveChannel(Channel->GetChannelName());
+		if (!Section->TryModify())
+		{
+			continue;
+		}
 
-			OverrideProvider->OnChannelOverridesChanged();
+		OverrideRegistry->RemoveChannel(Channel->GetChannelName());
 
-			bAnyChannelsChanged = true;
-			bChannelsChanged = true;
-		}
+		OverrideProvider->OnChannelOverridesChanged();
 
-		if (bChannelsChanged)
+		bAnyChannelsChanged = true;
+		if (TSharedPtr<FChannelGroupModel> ParentGroupModel = Channel->CastParent<FChannelGroupModel>())
 		{
-			SelectedChannelGroup->OnChannelOverridesChanged();
+			ChangedChannelGroups.Add(ParentGroupModel);
 		}
 	}
 
+	for (TSharedPtr<FChannelGroupModel> ChangedChannelGroup : ChangedChannelGroups)
+	{
+		ChangedChannelGroup->OnChannelOverridesChanged();
+	}
+
 	if (!bAnyChannelsChanged)
 	{
 		Transaction.Cancel();
 	}
 
-	Sequencer->NotifyMovieSceneDataChanged(EMovieSceneDataChangeType::MovieSceneStructureItemsChanged);
+	SequenceModel->GetSequencer()->NotifyMovieSceneDataChanged(EMovieSceneDataChangeType::MovieSceneStructureItemsChanged);
 }
 
 FChannelGroupOutlinerModel::FChannelGroupOutlinerModel(FName InChannelName, const FText& InDisplayText)
diff --git a/Engine/Source/Editor/Sequencer/Public/MVVM/ViewModels/ChannelModel.h b/Engine/Source/Editor/Sequencer/Public/MVVM/ViewModels/ChannelModel.h
--- a/Engine/Source/Editor/Sequencer/Public/MVVM/ViewModels/ChannelModel.h
+++ b/Engine/Source/Editor/Sequencer/Public/MVVM/ViewModels/ChannelModel.h
@@ -252,6 +252,7 @@ public:
 
 private:
 
+	static bool GetChannelsFromOutlinerSelection(TSharedPtr<FSequenceModel> SequenceModel, TArray<TViewModelPtr<FChannelModel>>& OutChannelModels);
 	static void BuildChannelOverrideMenu(FMenuBuilder& MenuBuilder, TSharedPtr<FSequenceModel> SequenceModel, UMovieSceneChannelOverrideContainer::FOverrideCandidates OverrideCandidates);
 	static void OverrideChannels(TSharedPtr<FSequenceModel> SequenceModel, TSubclassOf<UMovieSceneChannelOverrideContainer> OverrideClass);
 	static void RemoveChannelOverrides(TSharedPtr<FSequenceModel> SequenceModel);
