Sequencer: Added a display for the Frame Count and Duration

diff --git a/Engine/Source/Editor/Sequencer/Private/SSequencer.cpp b/Engine/Source/Editor/Sequencer/Private/SSequencer.cpp
--- a/Engine/Source/Editor/Sequencer/Private/SSequencer.cpp
+++ b/Engine/Source/Editor/Sequencer/Private/SSequencer.cpp
@@ -591,6 +591,34 @@ void SSequencer::Construct(const FArguments& InArgs, TSharedRef<FSequencer> InSe
 									})
 								]
 							]
+							+ SHorizontalBox::Slot()
+							.AutoWidth()
+							.VAlign(VAlign_Center)
+							.HAlign(HAlign_Right)
+							.Padding(FMargin(CommonPadding + 2.0, 0.f, 0.f, 0.f))
+							[
+								SNew(SBorder)
+								.BorderImage(nullptr)
+								[
+									// Frame count and duration
+									SNew(STextBlock)
+									.ColorAndOpacity(FAppStyle::GetSlateColor("SelectionColor_Pressed"))
+									.Text_Lambda([this]() -> FText {
+										FFrameRate TickResolution = SequencerPtr.Pin()->GetFocusedTickResolution();
+										FFrameRate DisplayRate = SequencerPtr.Pin()->GetFocusedDisplayRate();
+
+										TOptional<TRange<FFrameNumber>> SubSequenceRange = SequencerPtr.Pin()->GetSubSequenceRange();
+
+										FFrameNumber CurrentFrame = SequencerPtr.Pin()->GetLocalTime().Time.GetFrame();
+										TRange<FFrameNumber> CurrentRange = SubSequenceRange.IsSet() ? SubSequenceRange.GetValue() : SequencerPtr.Pin()->GetPlaybackRange();
+
+										FFrameNumber FrameCount = FFrameRate::TransformTime((CurrentFrame - CurrentRange.GetLowerBoundValue() + 1).Value, TickResolution, DisplayRate).CeilToFrame();
+										FFrameNumber FrameDuration = FFrameRate::TransformTime(CurrentRange.Size<FFrameNumber>().Value, TickResolution, DisplayRate).CeilToFrame();
+
+										return FText::FromString(FString::Printf(TEXT("%d of %d"), FrameCount.Value, FrameDuration.Value));
+									})
+								]
+							]
 						]
 					]
 
