Sequencer: Added commands for key context menu items - Set Key Time, Rekey, SnapToFrame, Delete

diff --git a/Engine/Source/Editor/Sequencer/Private/Sequencer.cpp b/Engine/Source/Editor/Sequencer/Private/Sequencer.cpp
--- a/Engine/Source/Editor/Sequencer/Private/Sequencer.cpp
+++ b/Engine/Source/Editor/Sequencer/Private/Sequencer.cpp
@@ -11205,6 +11205,22 @@ void FSequencer::BindCommands()
 		FExecuteAction::CreateSP(this, &FSequencer::OnAddTransformKeysForSelectedObjects, EMovieSceneTransformChannel::Scale),
 		FCanExecuteAction::CreateSP(this, &FSequencer::CanAddTransformKeysForSelectedObjects));
 
+	SequencerCommandBindings->MapAction(
+		Commands.SetKeyTime,
+		FExecuteAction::CreateSP(this, &FSequencer::SetKeyTime),
+		FCanExecuteAction::CreateSP(this, &FSequencer::CanSetKeyTime));
+	SequencerCommandBindings->MapAction(
+		Commands.Rekey,
+		FExecuteAction::CreateSP(this, &FSequencer::Rekey),
+		FCanExecuteAction::CreateSP(this, &FSequencer::CanRekey));
+	SequencerCommandBindings->MapAction(
+		Commands.SnapToFrame,
+		FExecuteAction::CreateSP(this, &FSequencer::SnapToFrame),
+		FCanExecuteAction::CreateSP(this, &FSequencer::CanSnapToFrame));
+	SequencerCommandBindings->MapAction(
+		Commands.DeleteKeys,
+		FExecuteAction::CreateSP(this, &FSequencer::DeleteSelectedKeys));
+
 	SequencerCommandBindings->MapAction(
 		Commands.TogglePilotCamera,
 		FExecuteAction::CreateSP(this, &FSequencer::OnTogglePilotCamera),
diff --git a/Engine/Source/Editor/Sequencer/Private/SequencerCommands.cpp b/Engine/Source/Editor/Sequencer/Private/SequencerCommands.cpp
--- a/Engine/Source/Editor/Sequencer/Private/SequencerCommands.cpp
+++ b/Engine/Source/Editor/Sequencer/Private/SequencerCommands.cpp
@@ -185,6 +185,11 @@ void FSequencerCommands::RegisterCommands()
 	UI_COMMAND(AddRotationKey, "Add Rotation Key", "Add a rotation key at the current time for the selected actor.", EUserInterfaceActionType::Button, FInputChord(EModifierKey::Shift, EKeys::E));
 	UI_COMMAND(AddScaleKey, "Add Scale Key", "Add a scale key at the current time for the selected actor.", EUserInterfaceActionType::Button, FInputChord(EModifierKey::Shift, EKeys::R));
 
+	UI_COMMAND( SetKeyTime, "Set Key Time", "Set the key to a specified time", EUserInterfaceActionType::Button, FInputChord() );
+	UI_COMMAND( Rekey, "Rekey", "Set the selected key's time to the current time", EUserInterfaceActionType::Button, FInputChord() );
+	UI_COMMAND( SnapToFrame, "Snap To Frame", "Snap selected keys to frame", EUserInterfaceActionType::Button, FInputChord() );
+	UI_COMMAND( DeleteKeys, "Delete Keys", "Deletes the selected keys", EUserInterfaceActionType::Button, FInputChord() );
+
 	UI_COMMAND(TogglePilotCamera, "Pilot Camera", "Toggle piloting the last camera or the camera cut camera.", EUserInterfaceActionType::ToggleButton, FInputChord(EModifierKey::Shift, EKeys::P));
 
 	UI_COMMAND(RefreshUI, "Refresh UI", "Forcibly refresh the UI from source data.", EUserInterfaceActionType::Button, FInputChord(EKeys::F5));
diff --git a/Engine/Source/Editor/Sequencer/Private/SequencerContextMenus.cpp b/Engine/Source/Editor/Sequencer/Private/SequencerContextMenus.cpp
--- a/Engine/Source/Editor/Sequencer/Private/SequencerContextMenus.cpp
+++ b/Engine/Source/Editor/Sequencer/Private/SequencerContextMenus.cpp
@@ -168,39 +168,12 @@ void FKeyContextMenu::PopulateMenu(FMenuBuilder& MenuBuilder, TSharedPtr<FExtend
 	}
 	MenuBuilder.EndSection(); // SequencerKeyEdit
 
-
-
 	MenuBuilder.BeginSection("SequencerKeys", LOCTEXT("KeysMenu", "Keys"));
 	{
-		MenuBuilder.AddMenuEntry(LOCTEXT("SetKeyTime", "Set Key Time"), LOCTEXT("SetKeyTimeTooltip", "Set the key to a specified time"),
-			FSlateIcon(),
-			FUIAction(
-				FExecuteAction::CreateSP(SequencerPtr, &FSequencer::SetKeyTime),
-				FCanExecuteAction::CreateSP(SequencerPtr, &FSequencer::CanSetKeyTime))
-		);
-
-		MenuBuilder.AddMenuEntry(LOCTEXT("Rekey", "Rekey"), LOCTEXT("RekeyTooltip", "Set the selected key's time to the current time"),
-			FSlateIcon(),
-			FUIAction(
-				FExecuteAction::CreateSP(SequencerPtr, &FSequencer::Rekey),
-				FCanExecuteAction::CreateSP(SequencerPtr, &FSequencer::CanRekey))
-		);
-
-		MenuBuilder.AddMenuEntry(
-			LOCTEXT("SnapToFrame", "Snap to Frame"),
-			LOCTEXT("SnapToFrameToolTip", "Snap selected keys to frame"),
-			FSlateIcon(),
-			FUIAction(
-				FExecuteAction::CreateSP(SequencerPtr, &FSequencer::SnapToFrame),
-				FCanExecuteAction::CreateSP(SequencerPtr, &FSequencer::CanSnapToFrame))
-		);
-
-		MenuBuilder.AddMenuEntry(
-			LOCTEXT("DeleteKey", "Delete"),
-			LOCTEXT("DeleteKeyToolTip", "Deletes the selected keys"),
-			FSlateIcon(),
-			FUIAction(FExecuteAction::CreateSP(SequencerPtr, &FSequencer::DeleteSelectedKeys))
-		);
+		MenuBuilder.AddMenuEntry(FSequencerCommands::Get().SetKeyTime);
+		MenuBuilder.AddMenuEntry(FSequencerCommands::Get().Rekey);
+		MenuBuilder.AddMenuEntry(FSequencerCommands::Get().SnapToFrame);
+		MenuBuilder.AddMenuEntry(FSequencerCommands::Get().DeleteKeys);
 	}
 	MenuBuilder.EndSection(); // SequencerKeys
 }
diff --git a/Engine/Source/Editor/Sequencer/Public/SequencerCommands.h b/Engine/Source/Editor/Sequencer/Public/SequencerCommands.h
--- a/Engine/Source/Editor/Sequencer/Public/SequencerCommands.h
+++ b/Engine/Source/Editor/Sequencer/Public/SequencerCommands.h
@@ -445,6 +445,18 @@ public:
 	/** Sets a scale key at the current time for the selected actor */
 	TSharedPtr< FUICommandInfo > AddScaleKey;
 
+	/** Set the key to a specified time */
+	TSharedPtr< FUICommandInfo > SetKeyTime;
+
+	/** Set the selected key's time to the current time */
+	TSharedPtr< FUICommandInfo > Rekey;
+
+	/** Snap selected keys to frame */
+	TSharedPtr< FUICommandInfo > SnapToFrame;
+
+	/** Deletes the selected keys */
+	TSharedPtr< FUICommandInfo > DeleteKeys;
+
 	/** Toggle piloting the last camera or the camera cut camera */
 	TSharedPtr< FUICommandInfo > TogglePilotCamera;
 
