Sequencer: Fix key selection for overlapping keys so that click select only selects one of the keys for that particular channel.

diff --git a/Engine/Source/Editor/Sequencer/Private/SequencerHotspots.cpp b/Engine/Source/Editor/Sequencer/Private/SequencerHotspots.cpp
--- a/Engine/Source/Editor/Sequencer/Private/SequencerHotspots.cpp
+++ b/Engine/Source/Editor/Sequencer/Private/SequencerHotspots.cpp
@@ -162,14 +162,36 @@ FKeyHotspot::FKeyHotspot(const TArray<FSequencerSelectedKey>& InKeys, TWeakPtr<F
 
 void FKeyHotspot::HandleMouseSelection(FHotspotSelectionManager& SelectionManager)
 {
-	TArray<FSequencerSelectedKey> KeysArray = Keys.Array();
 	if (SelectionManager.MouseEvent->GetEffectingButton() == EKeys::RightMouseButton)
 	{
-		SelectionManager.SelectKeysExclusive(KeysArray);
+		SelectionManager.SelectKeysExclusive(Keys.Array());
 	}
 	else
 	{
-		SelectionManager.ToggleKeys(KeysArray);
+		// On mouse click, select only keys that are unique to a section and channel, ie. don't select multiple keys on the same channel that have the same time
+		TArray<FSequencerSelectedKey> UniqueKeysArray;
+		for (const FSequencerSelectedKey& Key : Keys)
+		{
+			if (Key.IsValid())
+			{
+				bool bFoundKey = false;
+				for (const FSequencerSelectedKey& UniqueKey : UniqueKeysArray)
+				{
+					if (Key.Section == UniqueKey.Section && Key.WeakChannel == UniqueKey.WeakChannel)
+					{
+						bFoundKey = true;
+						break;
+					}
+				}
+
+				if (!bFoundKey)
+				{
+					UniqueKeysArray.Add(Key);
+				}
+			}
+		}
+
+		SelectionManager.ToggleKeys(UniqueKeysArray);
 	}
 }
 
