Sequencer: Deleted FSectionLayout

diff --git a/Engine/Source/Editor/Sequencer/Private/SectionLayout.cpp b/Engine/Source/Editor/Sequencer/Private/SectionLayout.cpp
deleted file mode 100644
diff --git a/Engine/Source/Editor/Sequencer/Private/SectionLayout.h b/Engine/Source/Editor/Sequencer/Private/SectionLayout.h
deleted file mode 100644
diff --git a/Engine/Source/Editor/Sequencer/Private/SequencerContextMenus.cpp b/Engine/Source/Editor/Sequencer/Private/SequencerContextMenus.cpp
--- a/Engine/Source/Editor/Sequencer/Private/SequencerContextMenus.cpp
+++ b/Engine/Source/Editor/Sequencer/Private/SequencerContextMenus.cpp
@@ -6,7 +6,6 @@
 #include "SequencerCommonHelpers.h"
 #include "SequencerCommands.h"
 #include "SSequencer.h"
-#include "SectionLayout.h"
 #include "IKeyArea.h"
 #include "SSequencerSection.h"
 #include "SequencerSettings.h"
@@ -611,29 +610,34 @@ void FSectionContextMenu::SelectAllKeys()
 {
 	using namespace UE::Sequencer;
 
-	for (const TWeakObjectPtr<UMovieSceneSection>& WeakSection : Sequencer->GetSelection().GetSelectedSections())
-	{
-		UMovieSceneSection* Section = WeakSection.Get();
-		TSharedPtr<FSectionModel> SectionHandle = Sequencer->GetNodeTree()->GetSectionModel(Section);
-		if (!SectionHandle)
-		{
-			continue;
-		}
+	TArray<FKeyHandle> HandlesScratch;
 
-		FSectionLayout Layout(SectionHandle);
-		for (const FSectionLayoutElement& Element : Layout.GetElements())
+	FSequencerSelection& Selection = Sequencer->GetSelection();
+	for (const TWeakPtr<FViewModel>& WeakModel : Selection.GetSelectedTrackAreaItems())
+	{
+		TViewModelPtr<FSectionModel> Section = WeakModel.Pin();
+		if (Section)
 		{
-			for (const TWeakPtr<FChannelModel>& WeakChannel : Element.GetChannels())
+			UMovieSceneSection* SectionObject = Section->GetSection();
+			if (SectionObject)
 			{
-				if (TSharedPtr<FChannelModel> Channel = WeakChannel.Pin())
-				{
-					TArray<FKeyHandle> Handles;
-					Channel->GetKeyArea()->GetKeyHandles(Handles);
+				FSequencerSelectedKey SelectedKey(*SectionObject, nullptr, FKeyHandle::Invalid());
 
-					for (FKeyHandle KeyHandle : Handles)
+				// Iterate all channels
+				for (const TViewModelPtr<FChannelModel>& Channel : Section->GetDescendantsOfType<FChannelModel>())
+				{
+					if (Channel->GetLinkedOutlinerItem() && !Channel->GetLinkedOutlinerItem()->IsFilteredOut())
 					{
-						FSequencerSelectedKey SelectKey(*Section, Channel, KeyHandle);
-						Sequencer->GetSelection().AddToSelection(SelectKey);
+						SelectedKey.WeakChannel = Channel;
+
+						HandlesScratch.Reset();
+						Channel->GetKeyArea()->GetKeyHandles(HandlesScratch);
+
+						for (FKeyHandle KeyHandle : HandlesScratch)
+						{
+							SelectedKey.KeyHandle = KeyHandle;
+							Selection.AddToSelection(SelectedKey);
+						}
 					}
 				}
 			}
