Sequencer: Fix tooltip for channel colors so that it refers to key bars instead of track area

diff --git a/Engine/Source/Editor/Sequencer/Private/SequencerCommands.cpp b/Engine/Source/Editor/Sequencer/Private/SequencerCommands.cpp
--- a/Engine/Source/Editor/Sequencer/Private/SequencerCommands.cpp
+++ b/Engine/Source/Editor/Sequencer/Private/SequencerCommands.cpp
@@ -148,7 +148,7 @@ void FSequencerCommands::RegisterCommands()
 	UI_COMMAND( FindInContentBrowser, "Find in Content Browser", "Find the viewed sequence asset in the content browser", EUserInterfaceActionType::Button, FInputChord() );
 	UI_COMMAND( ToggleLayerBars, "Layer Bars", "Show/hide the layer bars to edit keyframes in bulk", EUserInterfaceActionType::ToggleButton, FInputChord() );
 	UI_COMMAND( ToggleKeyBars, "Key Bars", "Show/hide key bar connectors for quickly retiming pairs of keys", EUserInterfaceActionType::ToggleButton, FInputChord() );
-	UI_COMMAND( ToggleChannelColors, "Channel Colors", "Show/hide the channel colors in the track area", EUserInterfaceActionType::ToggleButton, FInputChord() );
+	UI_COMMAND( ToggleChannelColors, "Channel Colors", "Show/hide the channel colors for the key bars", EUserInterfaceActionType::ToggleButton, FInputChord() );
 	UI_COMMAND( ToggleShowSelectedNodesOnly, "Selected Nodes Only", "Show selected nodes only", EUserInterfaceActionType::ToggleButton, FInputChord() );
 	
 	UI_COMMAND( ToggleShowCurveEditor, "Curve Editor", "Show the animation keys in a curve editor", EUserInterfaceActionType::ToggleButton, FInputChord() );
diff --git a/Engine/Source/Editor/Sequencer/Public/SequencerSettings.h b/Engine/Source/Editor/Sequencer/Public/SequencerSettings.h
--- a/Engine/Source/Editor/Sequencer/Public/SequencerSettings.h
+++ b/Engine/Source/Editor/Sequencer/Public/SequencerSettings.h
@@ -289,9 +289,9 @@ public:
 	/** Set whether to show channel colors */
 	void SetInfiniteKeyAreas(bool bInInfiniteKeyAreas);
 
-	/** @return true if showing channel colors */
+	/** @return true if showing channel colors for the key bars */
 	bool GetShowChannelColors() const;
-	/** Set whether to show channel colors */
+	/** Set whether to show channel colors for the key bars */
 	void SetShowChannelColors(bool bInShowChannelColors);
 
 	/** @return true if showing status bar */
@@ -567,7 +567,7 @@ protected:
 	UPROPERTY( config, EditAnywhere, Category=Timeline )
 	bool bInfiniteKeyAreas;
 
-	/** Enable or disable displaying channel bar colors for vector properties. */
+	/** Enable or disable displaying channel bar colors for the key bars. */
 	UPROPERTY(config, EditAnywhere, Category = Timeline)
 	bool bShowChannelColors;
 
