Curve Editor: Performance: Don't modify when setting color on init.
#trivial

diff --git a/Engine/Source/Editor/Sequencer/Private/IKeyArea.cpp b/Engine/Source/Editor/Sequencer/Private/IKeyArea.cpp
--- a/Engine/Source/Editor/Sequencer/Private/IKeyArea.cpp
+++ b/Engine/Source/Editor/Sequencer/Private/IKeyArea.cpp
@@ -343,7 +343,7 @@ TUniquePtr<FCurveModel> IKeyArea::CreateCurveEditorModel(TSharedRef<ISequencer>
 				TOptional<FLinearColor> SettingColor = Settings->GetCustomColor(Object->GetClass(), Name);
 				if (SettingColor.IsSet())
 				{
-					CurveModel->SetColor(SettingColor.GetValue());
+					CurveModel->SetColor(SettingColor.GetValue(),false);
 				}
 			}
 		}
