Sequencer: Always respect autoscroll setting. ie. don't autoscroll when stepping frames unless it's enabled.

diff --git a/Engine/Source/Editor/Sequencer/Private/Sequencer.cpp b/Engine/Source/Editor/Sequencer/Private/Sequencer.cpp
--- a/Engine/Source/Editor/Sequencer/Private/Sequencer.cpp
+++ b/Engine/Source/Editor/Sequencer/Private/Sequencer.cpp
@@ -2985,7 +2985,7 @@ void FSequencer::SetLocalTime( FFrameTime NewTime, ESnapTimeMode SnapTimeMode, b
 	FFrameRate LocalResolution = GetFocusedTickResolution();
 
 	// Ensure the time is in the current view
-	if (IsAutoScrollEnabled() || GetPlaybackStatus() != EMovieScenePlayerStatus::Playing)
+	if (IsAutoScrollEnabled())
 	{
 		ScrollIntoView(NewTime / LocalResolution);
 	}
