Sequencer: Put autoscrub behind a cvar - Sequencer.AutoScrub - autoscrubbing animates the current time to the desired current time.

Stepping between marked frames no longer autoscrubs, it just sets the time directly.

diff --git a/Engine/Source/Editor/Sequencer/Private/Sequencer.cpp b/Engine/Source/Editor/Sequencer/Private/Sequencer.cpp
--- a/Engine/Source/Editor/Sequencer/Private/Sequencer.cpp
+++ b/Engine/Source/Editor/Sequencer/Private/Sequencer.cpp
@@ -167,6 +167,11 @@
 
 DEFINE_LOG_CATEGORY(LogSequencer);
 
+static TAutoConsoleVariable<bool> CVarAutoScrub(
+	TEXT("Sequencer.AutoScrub"),
+	false,
+	TEXT("Enable/disable auto-scrubbing"));
+
 static TAutoConsoleVariable<float> CVarAutoScrubSpeed(
 	TEXT("Sequencer.AutoScrubSpeed"),
 	6.0f,
@@ -5216,7 +5221,7 @@ void FSequencer::OnScrubPositionChanged( FFrameTime NewScrubPosition, bool bScru
 		}
 	}
 
-	if (!bScrubbing && FSlateApplication::Get().GetModifierKeys().IsShiftDown())
+	if (!bScrubbing && CVarAutoScrub->GetBool() && FSlateApplication::Get().GetModifierKeys().IsShiftDown())
 	{
 		AutoScrubToTime(NewScrubPosition);
 	}
@@ -9766,7 +9771,7 @@ void FSequencer::StepToNextMark()
 	int32 MarkedIndex = FocusedMovieScene->FindNextMarkedFrame(GetLocalTime().Time.FloorToFrame(), bForwards);
 	if (MarkedIndex != INDEX_NONE)
 	{
-		AutoScrubToTime(FocusedMovieScene->GetMarkedFrames()[MarkedIndex].FrameNumber.Value);
+		SetLocalTimeDirectly(FocusedMovieScene->GetMarkedFrames()[MarkedIndex].FrameNumber.Value);
 	}
 }
 
@@ -9788,7 +9793,7 @@ void FSequencer::StepToPreviousMark()
 	int32 MarkedIndex = FocusedMovieScene->FindNextMarkedFrame(GetLocalTime().Time.FloorToFrame(), bForwards);
 	if (MarkedIndex != INDEX_NONE)
 	{
-		AutoScrubToTime(FocusedMovieScene->GetMarkedFrames()[MarkedIndex].FrameNumber.Value);
+		SetLocalTimeDirectly(FocusedMovieScene->GetMarkedFrames()[MarkedIndex].FrameNumber.Value);
 	}
 }
 
