Sequencer: Group properties in the +Track menu by their defined Category (if it exists)

diff --git a/Engine/Source/Editor/Sequencer/Private/MVVM/ViewModels/ObjectBindingModel.cpp b/Engine/Source/Editor/Sequencer/Private/MVVM/ViewModels/ObjectBindingModel.cpp
--- a/Engine/Source/Editor/Sequencer/Private/MVVM/ViewModels/ObjectBindingModel.cpp
+++ b/Engine/Source/Editor/Sequencer/Private/MVVM/ViewModels/ObjectBindingModel.cpp
@@ -595,8 +595,30 @@ TSharedRef<SWidget> FObjectBindingModel::GetAddTrackMenuContent()
 	// [PostProcess Settings] [ColorGrading]
 	// [Ortho View]
 
+	static const FString DefaultPropertyCategory = TEXT("Default");
+
+	// Properties with the category "Default" have no category and should be sorted to the top
+	struct FCategorySortPredicate
+	{
+		bool operator()(const FString& A, const FString& B) const
+		{
+			if (A == DefaultPropertyCategory)
+			{
+				return true;
+			}
+			else if (B == DefaultPropertyCategory)
+			{
+				return false;
+			}
+			else
+			{
+				return A.Compare(B) < 0;
+			}
+		}
+	};
+
 	// Create property menu data based on keyable property paths
-	TArray<PropertyMenuData> KeyablePropertyMenuData;
+	TSortedMap<FString, TArray<PropertyMenuData>, FDefaultAllocator, FCategorySortPredicate> KeyablePropertyMenuData;
 	for (const FPropertyPath& KeyablePropertyPath : KeyablePropertyPaths)
 	{
 		FProperty* Property = KeyablePropertyPath.GetRootProperty().Property.Get();
@@ -612,57 +634,73 @@ TSharedRef<SWidget> FObjectBindingModel::GetAddTrackMenuContent()
 			{
 				KeyableMenuData.MenuName = Property->GetDisplayNameText().ToString();
 			}
-			KeyablePropertyMenuData.Add(KeyableMenuData);
+
+			FString CategoryText = FObjectEditorUtils::GetCategory(Property);
+
+			KeyablePropertyMenuData.FindOrAdd(CategoryText).Add(KeyableMenuData);
 		}
 	}
 
-	// Sort on the menu name
-	KeyablePropertyMenuData.Sort([](const PropertyMenuData& A, const PropertyMenuData& B)
-	{
-		int32 CompareResult = A.MenuName.Compare(B.MenuName);
-		return CompareResult < 0;
-	});
-	
-
 	// Add menu items
-	AddTrackMenuBuilder.BeginSection( SequencerMenuExtensionPoints::AddTrackMenu_PropertiesSection, LOCTEXT("PropertiesMenuHeader" , "Properties"));
-	for (int32 MenuDataIndex = 0; MenuDataIndex < KeyablePropertyMenuData.Num(); )
+	for (TPair<FString, TArray<PropertyMenuData>>& Pair : KeyablePropertyMenuData)
 	{
-		TArray<FPropertyPath> KeyableSubMenuPropertyPaths;
-
-		KeyableSubMenuPropertyPaths.Add(KeyablePropertyMenuData[MenuDataIndex].PropertyPath);
+		// Sort on the property name
+		Pair.Value.Sort([](const PropertyMenuData& A, const PropertyMenuData& B)
+		{
+			int32 CompareResult = A.MenuName.Compare(B.MenuName);
+			return CompareResult < 0;
+		});
 
-		// If this menu data only has one property name, add the menu item
-		if (KeyablePropertyMenuData[MenuDataIndex].PropertyPath.GetNumProperties() == 1 || !bUseSubMenus)
+		FString CategoryText = Pair.Key;
+		
+		if (CategoryText == DefaultPropertyCategory)
 		{
-			AddPropertyMenuItems(AddTrackMenuBuilder, KeyableSubMenuPropertyPaths, 0, -1);
-			++MenuDataIndex;
+			AddTrackMenuBuilder.BeginSection(SequencerMenuExtensionPoints::AddTrackMenu_PropertiesSection, LOCTEXT("PropertiesMenuHeader", "Properties"));
 		}
-		// Otherwise, look to the next menu data to gather up new data
 		else
 		{
-			for (; MenuDataIndex < KeyablePropertyMenuData.Num()-1; )
+			AddTrackMenuBuilder.BeginSection(NAME_None, FText::FromString(CategoryText));
+		}
+	
+		for (int32 MenuDataIndex = 0; MenuDataIndex < Pair.Value.Num(); )
+		{
+			TArray<FPropertyPath> KeyableSubMenuPropertyPaths;
+
+			KeyableSubMenuPropertyPaths.Add(Pair.Value[MenuDataIndex].PropertyPath);
+
+			// If this menu data only has one property name, add the menu item
+			if (Pair.Value[MenuDataIndex].PropertyPath.GetNumProperties() == 1 || !bUseSubMenus)
 			{
-				if (KeyablePropertyMenuData[MenuDataIndex].MenuName == KeyablePropertyMenuData[MenuDataIndex+1].MenuName)
-				{	
-					++MenuDataIndex;
-					KeyableSubMenuPropertyPaths.Add(KeyablePropertyMenuData[MenuDataIndex].PropertyPath);
-				}
-				else
+				AddPropertyMenuItems(AddTrackMenuBuilder, KeyableSubMenuPropertyPaths, 0, -1);
+				++MenuDataIndex;
+			}
+			// Otherwise, look to the next menu data to gather up new data
+			else
+			{
+				for (; MenuDataIndex < Pair.Value.Num()-1; )
 				{
-					break;
+					if (Pair.Value[MenuDataIndex].MenuName == Pair.Value[MenuDataIndex+1].MenuName)
+					{	
+						++MenuDataIndex;
+						KeyableSubMenuPropertyPaths.Add(Pair.Value[MenuDataIndex].PropertyPath);
+					}
+					else
+					{
+						break;
+					}
 				}
-			}
 
-			AddTrackMenuBuilder.AddSubMenu(
-				FText::FromString(KeyablePropertyMenuData[MenuDataIndex].MenuName),
-				FText::GetEmpty(), 
-				FNewMenuDelegate::CreateSP(this, &FObjectBindingModel::HandleAddTrackSubMenuNew, KeyableSubMenuPropertyPaths, 0));
+				AddTrackMenuBuilder.AddSubMenu(
+					FText::FromString(Pair.Value[MenuDataIndex].MenuName),
+					FText::GetEmpty(), 
+					FNewMenuDelegate::CreateSP(this, &FObjectBindingModel::HandleAddTrackSubMenuNew, KeyableSubMenuPropertyPaths, 0));
 
-			++MenuDataIndex;
+				++MenuDataIndex;
+			}
 		}
+
+		AddTrackMenuBuilder.EndSection();
 	}
-	AddTrackMenuBuilder.EndSection();
 
 	if (AddTrackMenuBuilder.GetMultiBox()->GetBlocks().Num() == NumStartingBlocks)
 	{
