Sequencer: Add the ability to specify a section color tint which will be blended with the track color

- Added a menu to set the section color tint from a list of user-defined color tints
- Skeletal Track Editor's skeleton color for the animation clips are now based on the section color if defined. Otherwise, it defaults to the original hardcoded colors

diff --git a/Engine/Source/Editor/MovieSceneTools/Private/EditModes/SkeletalAnimationTrackEditMode.cpp b/Engine/Source/Editor/MovieSceneTools/Private/EditModes/SkeletalAnimationTrackEditMode.cpp
--- a/Engine/Source/Editor/MovieSceneTools/Private/EditModes/SkeletalAnimationTrackEditMode.cpp
+++ b/Engine/Source/Editor/MovieSceneTools/Private/EditModes/SkeletalAnimationTrackEditMode.cpp
@@ -332,8 +332,15 @@ void FSkeletalAnimationTrackEditMode::Render(const FSceneView* View, FViewport*
 											FCompactPoseBoneIndex PoseIndex = AnimationPoseData.GetPose().GetBoneContainer().GetCompactPoseIndexFromSkeletonIndex(Index);
 											OutPose[PoseIndex] = RootMotionAtStart;
 
+											FLinearColor SectionColor = FLinearColor(AnimSection->GetColorTint());
+
+											const float Alpha = SectionColor.A;
+											SectionColor.A = 1.f;
+
+											FLinearColor BoneColor = Colors[SectionIndex % 5] * (1.f - Alpha) + SectionColor * Alpha;
+
 											FVector RootLocation;
-											DrawBonesFromCompactPose(OutPose, SkelMeshComp, PDI, Colors[SectionIndex % 5], PoseIndex, RootLocation);
+											DrawBonesFromCompactPose(OutPose, SkelMeshComp, PDI, BoneColor, PoseIndex, RootLocation);
 											
 											if (IsRootSelected(AnimSection))
 											{
diff --git a/Engine/Source/Editor/Sequencer/Private/MVVM/ViewModels/OutlinerItemModel.cpp b/Engine/Source/Editor/Sequencer/Private/MVVM/ViewModels/OutlinerItemModel.cpp
--- a/Engine/Source/Editor/Sequencer/Private/MVVM/ViewModels/OutlinerItemModel.cpp
+++ b/Engine/Source/Editor/Sequencer/Private/MVVM/ViewModels/OutlinerItemModel.cpp
@@ -14,10 +14,12 @@
 #include "Framework/Commands/GenericCommands.h"
 #include "Framework/Commands/UIAction.h"
 #include "Framework/MultiBox/MultiBoxBuilder.h"
+#include "ISettingsModule.h"
 #include "MovieSceneSequence.h"
 #include "ScopedTransaction.h"
 #include "Sequencer.h"
 #include "SequencerSelectionCurveFilter.h"
+#include "SequencerSettings.h"
 #include "Styling/AppStyle.h"
 #include "Tree/CurveEditorTreeFilter.h"
 #include "Tree/SCurveEditorTreePin.h"
@@ -518,6 +520,11 @@ void FOutlinerItemModelMixin::BuildContextMenu(FMenuBuilder& MenuBuilder)
 
 		MenuBuilder.BeginSection("TrackDisplayOptions", NSLOCTEXT("Sequencer", "TrackNodeDisplayOptions", "Display Options"));
 		{
+			MenuBuilder.AddSubMenu(
+				LOCTEXT("SetColorTint", "Set Color Tint"),
+				LOCTEXT("SetColorTintTooltip", "Set color tint from the preferences for the selected sections or the track's sections"),
+				FNewMenuDelegate::CreateSP(SharedThis, &FOutlinerItemModelMixin::BuildSectionColorTintsContextMenu));
+
 			UStruct* DisplayOptionsStruct = FMovieSceneTrackDisplayOptions::StaticStruct();
 
 			const FBoolProperty* ShowVerticalFramesProperty = CastField<FBoolProperty>(DisplayOptionsStruct->FindPropertyByName(GET_MEMBER_NAME_CHECKED(FMovieSceneTrackDisplayOptions, bShowVerticalFrames)));
@@ -588,6 +595,87 @@ void FOutlinerItemModelMixin::BuildOrganizeContextMenu(FMenuBuilder& MenuBuilder
 	}
 }
 
+void FOutlinerItemModelMixin::BuildSectionColorTintsContextMenu(FMenuBuilder& MenuBuilder)
+{
+	TSharedPtr<FSequencer> Sequencer = GetEditor()->GetSequencerImpl();
+
+	TArray<UMovieSceneSection*> Sections;
+	for (TWeakObjectPtr<UMovieSceneSection> WeakSection : Sequencer->GetSelection().GetSelectedSections())
+	{
+		if (UMovieSceneSection* Section = WeakSection.Get())
+		{
+			Sections.Add(Section);
+		}
+	}
+
+	if (!Sections.Num())
+	{
+		for (TWeakPtr<FViewModel> WeakNode : Sequencer->GetSelection().GetSelectedOutlinerItems())
+		{
+			TSharedPtr<FViewModel> Node = WeakNode.Pin();
+
+			if (ITrackExtension* TrackExtension = Node->CastThis<ITrackExtension>())
+			{
+				for (UMovieSceneSection* Section : TrackExtension->GetSections())
+				{
+					Sections.Add(Section);
+				}
+			}
+		}
+	}
+
+	if (!Sections.Num())
+	{
+		return;
+	}
+
+	const bool bIsReadOnly = Sequencer->IsReadOnly();
+	FCanExecuteAction CanExecute = FCanExecuteAction::CreateLambda([bIsReadOnly] { return !bIsReadOnly; });
+
+	TArray<FColor> SectionColorTints = Sequencer->GetSequencerSettings()->GetSectionColorTints();
+
+	for (const FColor& SectionColorTint : SectionColorTints)
+	{
+		TSharedPtr<SBox> ColorWidget = SNew(SBox)
+			.WidthOverride(70.f)
+			.HeightOverride(20.f)
+		[
+			SNew(SBorder)
+				.BorderImage(FAppStyle::GetBrush("WhiteBrush"))
+				.BorderBackgroundColor(FLinearColor::FromSRGBColor(SectionColorTint))
+		];
+
+		MenuBuilder.AddMenuEntry(
+			FUIAction(
+				FExecuteAction::CreateSP(Sequencer.Get(), &FSequencer::SetSectionColorTint, Sections, SectionColorTint),
+				CanExecute),
+			ColorWidget.ToSharedRef());
+	}
+
+	MenuBuilder.AddSeparator();
+
+	// Clear any assigned color tints
+	MenuBuilder.AddMenuEntry(
+		LOCTEXT("ClearColorTintLabel", "Clear"),
+		LOCTEXT("ClearColorTintTooltip", "Clear any assigned color tints"),
+		FSlateIcon(),
+		FUIAction(
+			FExecuteAction::CreateSP(Sequencer.Get(), &FSequencer::SetSectionColorTint, Sections, FColor(0, 0, 0, 0)),
+			CanExecute));
+
+	// Pop up preferences to edit custom color tints
+	MenuBuilder.AddMenuEntry(
+		LOCTEXT("EditColorTintLabel", "Edit Color Tints..."),
+		LOCTEXT("EditColorTintTooltip", "Edit the custom color tints"),
+		FSlateIcon(),
+		FUIAction(
+			FExecuteAction::CreateLambda([=] { 
+				FString SettingsName = Sequencer->GetSequencerSettings()->GetName();
+				return FModuleManager::LoadModuleChecked<ISettingsModule>("Settings").ShowViewer("Editor", "ContentEditors", *SettingsName); 
+			})
+		));
+}
+
 bool FOutlinerItemModelMixin::HasCurves() const
 {
 	return false;
diff --git a/Engine/Source/Editor/Sequencer/Private/Sequencer.cpp b/Engine/Source/Editor/Sequencer/Private/Sequencer.cpp
--- a/Engine/Source/Editor/Sequencer/Private/Sequencer.cpp
+++ b/Engine/Source/Editor/Sequencer/Private/Sequencer.cpp
@@ -11701,6 +11701,17 @@ void FSequencer::BuildCustomContextMenuForGuid(FMenuBuilder& MenuBuilder, FGuid
 	OutlinerViewModel->BuildCustomContextMenuForGuid(MenuBuilder, ObjectBinding);
 }
 
+void FSequencer::SetSectionColorTint(TArray<UMovieSceneSection*> Sections, FColor ColorTint)
+{
+	FScopedTransaction Transaction(LOCTEXT("SetSectionColorTint", "Set Section Color Tint"));
+
+	for (UMovieSceneSection* Section : Sections)
+	{
+		Section->Modify();
+		Section->SetColorTint(ColorTint);
+	}
+}
+
 bool FSequencer::GetGridMetrics(const float PhysicalWidth, const double InViewStart, const double InViewEnd, double& OutMajorInterval, int32& OutMinorDivisions) const
 {
 	FSlateFontInfo SmallLayoutFont = FCoreStyle::GetDefaultFontStyle("Regular", 8);
diff --git a/Engine/Source/Editor/Sequencer/Private/Sequencer.h b/Engine/Source/Editor/Sequencer/Private/Sequencer.h
--- a/Engine/Source/Editor/Sequencer/Private/Sequencer.h
+++ b/Engine/Source/Editor/Sequencer/Private/Sequencer.h
@@ -599,6 +599,9 @@ public:
 	/** Will create a custom menu if FSequencerViewParams::OnBuildCustomContextMenuForGuid is specified. */
 	void BuildCustomContextMenuForGuid(FMenuBuilder& MenuBuilder, FGuid ObjectBinding);
 
+	/** Set the color tint for the requested sections */
+	void SetSectionColorTint(TArray<UMovieSceneSection*> Sections, FColor ColorTint);
+
 public:
 
 	/** Copy the selection, whether it's keys or tracks */
diff --git a/Engine/Source/Editor/Sequencer/Private/SequencerSectionPainter.cpp b/Engine/Source/Editor/Sequencer/Private/SequencerSectionPainter.cpp
--- a/Engine/Source/Editor/Sequencer/Private/SequencerSectionPainter.cpp
+++ b/Engine/Source/Editor/Sequencer/Private/SequencerSectionPainter.cpp
@@ -24,7 +24,14 @@ FSequencerSectionPainter::~FSequencerSectionPainter()
 
 int32 FSequencerSectionPainter::PaintSectionBackground()
 {
-	return PaintSectionBackground(GetTrack()->GetColorTint());
+	FLinearColor TrackColor = FLinearColor(GetTrack()->GetColorTint());
+	FLinearColor SectionColor = FLinearColor(SectionModel->GetSection()->GetColorTint());
+
+	const float Alpha = SectionColor.A;
+	SectionColor.A = 1.f;
+
+	FLinearColor BackgroundColor = TrackColor * (1.f - Alpha) + SectionColor * Alpha;
+	return PaintSectionBackground(BackgroundColor);
 }
 
 UMovieSceneTrack* FSequencerSectionPainter::GetTrack() const
diff --git a/Engine/Source/Editor/Sequencer/Private/SequencerSettings.cpp b/Engine/Source/Editor/Sequencer/Private/SequencerSettings.cpp
--- a/Engine/Source/Editor/Sequencer/Private/SequencerSettings.cpp
+++ b/Engine/Source/Editor/Sequencer/Private/SequencerSettings.cpp
@@ -65,6 +65,14 @@ USequencerSettings::USequencerSettings( const FObjectInitializer& ObjectInitiali
 	TreeViewWidth = 0.3f;
 	bShowTickLines = true;
 	bShowSequencerToolbar = true;
+
+	SectionColorTints.Add(FColor(88, 102, 142, 255)); // blue
+	SectionColorTints.Add(FColor(99, 137, 132, 255)); // blue-green
+	SectionColorTints.Add(FColor(110, 127, 92, 255)); // green
+	SectionColorTints.Add(FColor(151, 142, 102, 255)); // yellow
+	SectionColorTints.Add(FColor(147, 119, 101, 255)); // orange
+	SectionColorTints.Add(FColor(139, 95, 108, 255)); // red 
+	SectionColorTints.Add(FColor(109, 74, 121, 255)); // purple
 }
 
 void USequencerSettings::PostEditChangeProperty(FPropertyChangedEvent& PropertyChangedEvent)
@@ -784,6 +792,20 @@ void USequencerSettings::SetDisableSectionsAfterBaking(bool bInDisableSectionsAf
 	}
 }
 
+TArray<FColor> USequencerSettings::GetSectionColorTints() const
+{
+	return SectionColorTints;
+}
+
+void USequencerSettings::SetSectionColorTints(const TArray<FColor>& InSectionColorTints)
+{
+	if (SectionColorTints != InSectionColorTints)
+	{
+		SectionColorTints = InSectionColorTints;
+		SaveConfig();
+	}
+}
+
 bool USequencerSettings::GetCleanPlaybackMode() const
 {
 	return bCleanPlaybackMode;
diff --git a/Engine/Source/Editor/Sequencer/Public/MVVM/ViewModels/OutlinerItemModel.h b/Engine/Source/Editor/Sequencer/Public/MVVM/ViewModels/OutlinerItemModel.h
--- a/Engine/Source/Editor/Sequencer/Public/MVVM/ViewModels/OutlinerItemModel.h
+++ b/Engine/Source/Editor/Sequencer/Public/MVVM/ViewModels/OutlinerItemModel.h
@@ -87,6 +87,8 @@ protected:
 	/** Set expansion state without saving it in the movie-scene data */
 	void SetExpansionWithoutSaving(bool bInIsExpanded);
 
+	void BuildSectionColorTintsContextMenu(FMenuBuilder& MenuBuilder);
+
 private:
 
 	virtual FViewModel* AsViewModel() = 0;
diff --git a/Engine/Source/Editor/Sequencer/Public/SequencerSettings.h b/Engine/Source/Editor/Sequencer/Public/SequencerSettings.h
--- a/Engine/Source/Editor/Sequencer/Public/SequencerSettings.h
+++ b/Engine/Source/Editor/Sequencer/Public/SequencerSettings.h
@@ -338,6 +338,11 @@ public:
 	/** Set whether to disable sections when baking, as opposed to deleting */
 	void SetDisableSectionsAfterBaking(bool bInDisableSectionsAfterBaking);
 
+	/** @return the section color tints */
+	TArray<FColor> GetSectionColorTints() const;
+	/** Set the section color tints */
+	void SetSectionColorTints(const TArray<FColor>& InSectionColorTints);
+
 	/** @return Whether to playback in clean mode (game view, hide viewport UI) */
 	bool GetCleanPlaybackMode() const;
 	/** Toggle whether to playback in clean mode */
@@ -603,6 +608,10 @@ protected:
 	UPROPERTY(config, EditAnywhere, Category = Timeline)
 	bool bDisableSectionsAfterBaking;
 
+	/** Section color tints */
+	UPROPERTY(config, EditAnywhere, Category = General)
+	TArray<FColor> SectionColorTints;
+
 	/** When enabled, sequencer will playback in clean mode (game view, hide viewport UI) */
 	UPROPERTY(config, EditAnywhere, Category = General)
 	bool bCleanPlaybackMode;
diff --git a/Engine/Source/Runtime/MovieScene/Private/MovieSceneSection.cpp b/Engine/Source/Runtime/MovieScene/Private/MovieSceneSection.cpp
--- a/Engine/Source/Runtime/MovieScene/Private/MovieSceneSection.cpp
+++ b/Engine/Source/Runtime/MovieScene/Private/MovieSceneSection.cpp
@@ -49,6 +49,10 @@ UMovieSceneSection::UMovieSceneSection(const FObjectInitializer& ObjectInitializ
 	Easing.EaseOut = DefaultEaseOut;
 
 	ChannelProxyType = EMovieSceneChannelProxyType::Static;
+
+#if WITH_EDITORONLY_DATA
+	ColorTint = FColor(0, 0, 0, 0);
+#endif
 }
 
 
@@ -553,6 +557,24 @@ void UMovieSceneSection::InitialPlacementOnRow(const TArray<UMovieSceneSection*>
 	}
 }
 
+void UMovieSceneSection::SetColorTint(const FColor& InColorTint)
+{
+#if WITH_EDITORONLY_DATA
+	if (TryModify())
+	{
+		ColorTint = InColorTint;
+	}
+#endif
+}
+
+FColor UMovieSceneSection::GetColorTint() const
+{
+#if WITH_EDITORONLY_DATA
+	return ColorTint;
+#endif
+	return FColor(0, 0, 0, 0);
+}
+
 UMovieSceneSection* UMovieSceneSection::SplitSection(FQualifiedFrameTime SplitTime, bool bDeleteKeys)
 {
 	if (!SectionRange.Value.Contains(SplitTime.Time.GetFrame()))
diff --git a/Engine/Source/Runtime/MovieScene/Public/MovieSceneSection.h b/Engine/Source/Runtime/MovieScene/Public/MovieSceneSection.h
--- a/Engine/Source/Runtime/MovieScene/Public/MovieSceneSection.h
+++ b/Engine/Source/Runtime/MovieScene/Public/MovieSceneSection.h
@@ -592,6 +592,13 @@ public:
 	UFUNCTION(BlueprintPure, Category = "Sequencer|Section")
 	int32 GetPostRollFrames() const { return PostRollFrames.Value; }
 
+	/** Set this section's color tint. */
+	UFUNCTION(BlueprintCallable, Category = "Sequencer|Section")
+	MOVIESCENE_API void SetColorTint(const FColor& InColorTint);
+	/** Get this section's color tint. */
+	UFUNCTION(BlueprintPure, Category = "Sequencer|Section")
+	MOVIESCENE_API FColor GetColorTint() const;
+
 	/** The optional offset time of this section */
 	virtual TOptional<FFrameTime> GetOffsetTime() const { return TOptional<FFrameTime>(); }
 
@@ -754,6 +761,12 @@ private:
 	UPROPERTY(EditAnywhere, Category="Section")
 	uint32 bIsLocked : 1;
 
+#if WITH_EDITORONLY_DATA
+	/** The color tint for this section */
+	UPROPERTY(EditAnywhere, Category = "Section")
+	FColor ColorTint;
+#endif
+
 protected:
 
 	/** The start time of the section */
