Sequencer: Add default extension point if none exists. Follow up to 23483638

#rnx

diff --git a/Engine/Source/Editor/Sequencer/Private/MVVM/ViewModels/ObjectBindingModel.cpp b/Engine/Source/Editor/Sequencer/Private/MVVM/ViewModels/ObjectBindingModel.cpp
--- a/Engine/Source/Editor/Sequencer/Private/MVVM/ViewModels/ObjectBindingModel.cpp
+++ b/Engine/Source/Editor/Sequencer/Private/MVVM/ViewModels/ObjectBindingModel.cpp
@@ -617,6 +617,8 @@ TSharedRef<SWidget> FObjectBindingModel::GetAddTrackMenuContent()
 		}
 	};
 
+	bool bDefaultCategoryFound = false;
+
 	// Create property menu data based on keyable property paths
 	TSortedMap<FString, TArray<PropertyMenuData>, FDefaultAllocator, FCategorySortPredicate> KeyablePropertyMenuData;
 	for (const FPropertyPath& KeyablePropertyPath : KeyablePropertyPaths)
@@ -637,10 +639,22 @@ TSharedRef<SWidget> FObjectBindingModel::GetAddTrackMenuContent()
 
 			FString CategoryText = FObjectEditorUtils::GetCategory(Property);
 
+			if (CategoryText == DefaultPropertyCategory)
+			{
+				bDefaultCategoryFound = true;
+			}
+
 			KeyablePropertyMenuData.FindOrAdd(CategoryText).Add(KeyableMenuData);
 		}
 	}
 
+	// Always add an extension point for Properties section even if none are found (Components rely on this) 
+	if (!bDefaultCategoryFound)
+	{
+		AddTrackMenuBuilder.BeginSection(SequencerMenuExtensionPoints::AddTrackMenu_PropertiesSection, LOCTEXT("PropertiesMenuHeader", "Properties"));
+		AddTrackMenuBuilder.EndSection();
+	}
+
 	// Add menu items
 	for (TPair<FString, TArray<PropertyMenuData>>& Pair : KeyablePropertyMenuData)
 	{
