Sequencer: Clear the section hotspot when deleting the section.

This fixes an ensure when trying to drag in the track area after a section has been deleted because the section hotspot still exists.

#rnx

diff --git a/Engine/Source/Editor/Sequencer/Private/SSequencerSection.cpp b/Engine/Source/Editor/Sequencer/Private/SSequencerSection.cpp
--- a/Engine/Source/Editor/Sequencer/Private/SSequencerSection.cpp
+++ b/Engine/Source/Editor/Sequencer/Private/SSequencerSection.cpp
@@ -883,6 +883,10 @@ void SSequencerSection::Construct( const FArguments& InArgs, TSharedPtr<FSequenc
 
 SSequencerSection::~SSequencerSection()
 {
+	if (Sequencer.IsValid() && GetSequencer().GetViewModel() && GetSequencer().GetViewModel()->GetTrackArea())
+	{
+		GetSequencer().GetViewModel()->GetTrackArea()->SetHotspot(nullptr);
+	}
 }
 
 EVisibility SSequencerSection::GetTopLevelChannelGroupVisibility() const
