Sequencer: Change check to ensure if track editor is not found. This can happen if you create a level sequence with a track type that is not supported in restricted mode, ie. event track.

Added a few additional checks for a valid TrackEditor before dereferencing.

diff --git a/Engine/Source/Editor/Sequencer/Private/MVVM/ViewModels/TrackModel.cpp b/Engine/Source/Editor/Sequencer/Private/MVVM/ViewModels/TrackModel.cpp
--- a/Engine/Source/Editor/Sequencer/Private/MVVM/ViewModels/TrackModel.cpp
+++ b/Engine/Source/Editor/Sequencer/Private/MVVM/ViewModels/TrackModel.cpp
@@ -208,7 +208,7 @@ void FTrackModel::ForceUpdate()
 		for (UMovieSceneSection* Section : Track->GetAllSections())
 		{
 			TSharedPtr<FSectionModel> SectionModel = SectionModelStorage->FindModelForSection(Section);
-			if (!SectionModel)
+			if (!SectionModel && TrackEditor)
 			{
 				TSharedRef<ISequencerSection> SectionInterface = TrackEditor->MakeSectionInterface(*Section, *Track, ObjectBinding);
 				SectionModel = SectionModelStorage->CreateModelForSection(Section, SectionInterface);
@@ -296,7 +296,7 @@ void FTrackModel::ForceUpdate()
 			const int32 RowIndex = Section->GetRowIndex();
 
 			TSharedPtr<FSectionModel> SectionModel = SectionModelStorage->FindModelForSection(Section);
-			if (!SectionModel)
+			if (!SectionModel && TrackEditor)
 			{
 				TSharedRef<ISequencerSection> SectionInterface = TrackEditor->MakeSectionInterface(*Section, *Track, ObjectBinding);
 				SectionModel = SectionModelStorage->CreateModelForSection(Section, SectionInterface);
@@ -646,7 +646,10 @@ void FTrackModel::BuildContextMenu(FMenuBuilder& MenuBuilder)
 
 	const int32 TrackRowIndex = GetRowIndex();
 
-	TrackEditor->BuildTrackContextMenu(MenuBuilder, Track);
+	if (TrackEditor)
+	{
+		TrackEditor->BuildTrackContextMenu(MenuBuilder, Track);
+	}
 
 	if (Track && Track->GetSupportedBlendTypes().Num() > 0)
 	{
diff --git a/Engine/Source/Editor/Sequencer/Private/MVVM/Views/SOutlinerTrackView.cpp b/Engine/Source/Editor/Sequencer/Private/MVVM/Views/SOutlinerTrackView.cpp
--- a/Engine/Source/Editor/Sequencer/Private/MVVM/Views/SOutlinerTrackView.cpp
+++ b/Engine/Source/Editor/Sequencer/Private/MVVM/Views/SOutlinerTrackView.cpp
@@ -78,7 +78,7 @@ void SOutlinerTrackView::Construct(
 		Params.TrackInsertRowIndex = Track->GetMaxRowIndex()+1;
 	}
 
-	TSharedPtr<SWidget> CustomWidget = TrackEditor->BuildOutlinerEditWidget(ObjectBinding ? ObjectBinding->GetObjectGuid() : FGuid(), Track, Params);
+	TSharedPtr<SWidget> CustomWidget = TrackEditor ? TrackEditor->BuildOutlinerEditWidget(ObjectBinding ? ObjectBinding->GetObjectGuid() : FGuid(), Track, Params) : nullptr;
 
 	TSharedPtr<FChannelGroupModel> TopLevelChannel = TopLevelChannels ? TopLevelChannels->FindFirstChildOfType<FChannelGroupModel>() : nullptr;
 	if (TopLevelChannel)
diff --git a/Engine/Source/Editor/Sequencer/Private/Sequencer.cpp b/Engine/Source/Editor/Sequencer/Private/Sequencer.cpp
--- a/Engine/Source/Editor/Sequencer/Private/Sequencer.cpp
+++ b/Engine/Source/Editor/Sequencer/Private/Sequencer.cpp
@@ -679,7 +679,7 @@ TSharedPtr<ISequencerTrackEditor> FSequencer::GetTrackEditor(UMovieSceneTrack* I
 		}
 	}
 
-	checkf(TrackEditor, TEXT("Unable to find a track editor for track type %s"), *TrackClass->GetName());
+	ensureAlwaysMsgf(TrackEditor, TEXT("Unable to find a track editor for track type %s"), *TrackClass->GetName());
 	return TrackEditor;
 }
 
