Fixed dirtied sequence assets when setting TODM by disabling RenameBindingRecursive on actor label changed delegate for the DaySequence preview editor.

diff --git a/Engine/Source/Editor/Sequencer/Private/LevelEditorSequencerIntegration.cpp b/Engine/Source/Editor/Sequencer/Private/LevelEditorSequencerIntegration.cpp
--- a/Engine/Source/Editor/Sequencer/Private/LevelEditorSequencerIntegration.cpp
+++ b/Engine/Source/Editor/Sequencer/Private/LevelEditorSequencerIntegration.cpp
@@ -333,6 +333,11 @@ void FLevelEditorSequencerIntegration::OnActorLabelChanged(AActor* ChangedActor)
 {
 	for (const FSequencerAndOptions& SequencerAndOptions : BoundSequencers)
 	{
+		if (!SequencerAndOptions.Options.bSyncBindingsToActorLabels)
+		{
+			continue;
+		}
+		
 		TSharedPtr<FSequencer> Pinned = SequencerAndOptions.Sequencer.Pin();
 		if (Pinned.IsValid())
 		{
diff --git a/Engine/Source/Editor/Sequencer/Public/LevelEditorSequencerIntegration.h b/Engine/Source/Editor/Sequencer/Public/LevelEditorSequencerIntegration.h
--- a/Engine/Source/Editor/Sequencer/Public/LevelEditorSequencerIntegration.h
+++ b/Engine/Source/Editor/Sequencer/Public/LevelEditorSequencerIntegration.h
@@ -31,6 +31,7 @@ struct FLevelEditorSequencerIntegrationOptions
 		, bForceRefreshDetails(true)
 		, bAttachOutlinerColumns(true)
 		, bActivateSequencerEdMode(true)
+		, bSyncBindingsToActorLabels(true)
 	{}
 
 	bool bRequiresLevelEvents : 1;
@@ -38,6 +39,7 @@ struct FLevelEditorSequencerIntegrationOptions
 	bool bForceRefreshDetails : 1;
 	bool bAttachOutlinerColumns : 1;
 	bool bActivateSequencerEdMode : 1;
+	bool bSyncBindingsToActorLabels : 1;
 };
 
 
