Sequencer: Check whether autoscroll is enabled inside of UpdateAutoScroll

diff --git a/Engine/Source/Editor/Sequencer/Private/Sequencer.cpp b/Engine/Source/Editor/Sequencer/Private/Sequencer.cpp
--- a/Engine/Source/Editor/Sequencer/Private/Sequencer.cpp
+++ b/Engine/Source/Editor/Sequencer/Private/Sequencer.cpp
@@ -819,7 +819,7 @@ void FSequencer::Tick(float InDeltaTime)
 		// Put the time into local space
  		SetLocalTimeLooped(NewGlobalTime * RootToLocalTransform);
 
-		if (IsAutoScrollEnabled() && GetPlaybackStatus() == EMovieScenePlayerStatus::Playing)
+		if (GetPlaybackStatus() == EMovieScenePlayerStatus::Playing)
 		{
 			const float ThresholdPercentage = 0.15f;
 			UpdateAutoScroll(GetLocalTime().Time / GetFocusedTickResolution(), ThresholdPercentage);
@@ -3287,6 +3287,11 @@ void FSequencer::ScrollIntoView(float InLocalTime)
 
 void FSequencer::UpdateAutoScroll(double NewTime, float ThresholdPercentage)
 {
+	if (!IsAutoScrollEnabled())
+	{
+		return;
+	}
+
 	AutoscrollOffset = CalculateAutoscrollEncroachment(NewTime, ThresholdPercentage);
 
 	if (!AutoscrollOffset.IsSet())
@@ -5217,7 +5222,7 @@ void FSequencer::OnScrubPositionChanged( FFrameTime NewScrubPosition, bool bScru
 		{
 			OnEndScrubbing();
 		}
-		else if (IsAutoScrollEnabled())
+		else
 		{
 			UpdateAutoScroll(NewScrubPosition / GetFocusedTickResolution());
 			
