Sequencer: marked frames improvements

- Marked frames can now be selected, which allows the user to move them along with keys and sections
- This had to be implemented as a special case in FSequencerSelection. Ideally marked frames would have been moved over to the MVVM design so that their view-models could be added to the selection without much change, but this opened up giant cans of worms.
- Move operations for keys and sections now always move any selected marked frames too
- Being selectable, marked frames now have a "selected" state that makes them drawn white and bold
- To make interaction easier, marked frames labels can be used to drag and select, and a hover state has been added
- A "comment" field has been added on marked frames

diff --git a/Engine/Source/Editor/EditorStyle/Private/StarshipStyle.cpp b/Engine/Source/Editor/EditorStyle/Private/StarshipStyle.cpp
--- a/Engine/Source/Editor/EditorStyle/Private/StarshipStyle.cpp
+++ b/Engine/Source/Editor/EditorStyle/Private/StarshipStyle.cpp
@@ -2079,6 +2079,9 @@ void FStarshipEditorStyle::FStyle::SetupSequencerStyles()
 		Set("Sequencer.Section.BackgroundText", DEFAULT_FONT("Bold", 24));
 		Set("Sequencer.Section.EmptySpace", new BOX_BRUSH(TEXT("Sequencer/EmptySpace"), FMargin(0.f, 7.f / 14.f)));
 
+		Set("Sequencer.MarkedFrame.LabelRight", new FSlateRoundedBoxBrush(FLinearColor::White, FVector4(0, 0, 4, 0)));
+		Set("Sequencer.MarkedFrame.LabelLeft", new FSlateRoundedBoxBrush(FLinearColor::White, FVector4(0, 0, 0, 4)));
+
 		Set("Sequencer.LayerBar.Background",  new BOX_BRUSH(TEXT("Sequencer/LayerBarBackground"), FMargin(4.f / 16.f)));
 		Set("Sequencer.LayerBar.HandleLeft",  new FSlateColorBrush(FStyleColors::White)); // Maybe use a rounded box if we decide to round-out the layer bars
 		Set("Sequencer.LayerBar.HandleRight", new FSlateColorBrush(FStyleColors::White)); // Maybe use a rounded box in we decide to round-out the layer bars
diff --git a/Engine/Source/Editor/Sequencer/Private/SequencerSelection.cpp b/Engine/Source/Editor/Sequencer/Private/SequencerSelection.cpp
--- a/Engine/Source/Editor/Sequencer/Private/SequencerSelection.cpp
+++ b/Engine/Source/Editor/Sequencer/Private/SequencerSelection.cpp
@@ -94,6 +94,11 @@ const TSet<TWeakPtr<UE::Sequencer::FViewModel>>& FSequencerSelection::GetSelecte
 	return SelectedTrackAreaItems;
 }
 
+const TSet<int32> FSequencerSelection::GetSelectedMarkedFrames() const
+{
+	return SelectedMarkedFrames;
+}
+
 const TSet<TWeakPtr<UE::Sequencer::FViewModel>>& FSequencerSelection::GetNodesWithSelectedKeysOrSections() const
 {
 	using namespace UE::Sequencer;
@@ -321,6 +326,23 @@ void FSequencerSelection::AddToOutlinerSelection(TSharedPtr<UE::Sequencer::FView
 	EmptySelectedTrackAreaItems();
 }
 
+void FSequencerSelection::AddToSelection(int32 InMarkedFrameIndex)
+{
+	if (!ensureAlwaysMsgf(InMarkedFrameIndex != INDEX_NONE, TEXT("AddToSelection with invalid marked frame index")))
+	{
+		return;
+	}
+
+	++SerialNumber;
+
+	SelectedMarkedFrames.Add(InMarkedFrameIndex);
+
+	if (IsBroadcasting())
+	{
+		OnMarkedFramesSelectionChanged.Broadcast();
+	}
+}
+
 void FSequencerSelection::AddToSelection(const TArrayView<TSharedPtr<UE::Sequencer::FViewModel>>& InModels)
 {
 	using namespace UE::Sequencer;
@@ -433,6 +455,18 @@ void FSequencerSelection::RemoveFromSelection(TSharedRef<UE::Sequencer::FViewMod
 	}
 }
 
+void FSequencerSelection::RemoveFromSelection(int32 InMarkedFrameIndex)
+{
+	++SerialNumber;
+
+	SelectedMarkedFrames.Remove(InMarkedFrameIndex);
+
+	if (IsBroadcasting())
+	{
+		OnMarkedFramesSelectionChanged.Broadcast();
+	}
+}
+
 bool FSequencerSelection::IsSelected(const FSequencerSelectedKey& Key) const
 {
 	return SelectedKeys.Contains(Key);
@@ -443,6 +477,11 @@ bool FSequencerSelection::IsSelected(TWeakPtr<UE::Sequencer::FViewModel> InModel
 	return SelectedOutlinerItems.Contains(InModel) || SelectedTrackAreaItems.Contains(InModel);
 }
 
+bool FSequencerSelection::IsSelected(int32 InMarkedFrameIndex) const
+{
+	return SelectedMarkedFrames.Contains(InMarkedFrameIndex);
+}
+
 bool FSequencerSelection::NodeHasSelectedKeysOrSections(TWeakPtr<UE::Sequencer::FViewModel> InModel) const
 {
 	return GetNodesWithSelectedKeysOrSections().Contains(InModel);
@@ -490,6 +529,7 @@ void FSequencerSelection::Empty()
 	EmptySelectedKeys();
 	EmptySelectedOutlinerNodes();
 	EmptySelectedTrackAreaItems();
+	EmptySelectedMarkedFrames();
 }
 
 void FSequencerSelection::EmptySelectedKeys()
@@ -547,6 +587,18 @@ void FSequencerSelection::EmptySelectedTrackAreaItems()
 			OnSectionSelectionChanged.Broadcast();
 		}
 	}
+
+	if (SelectedMarkedFrames.Num() > 0)
+	{
+		++SerialNumber;
+
+		SelectedMarkedFrames.Empty();
+
+		if (IsBroadcasting())
+		{
+			OnMarkedFramesSelectionChanged.Broadcast();
+		}
+	}
 }
 
 void FSequencerSelection::EmptySelectedOutlinerNodes()
@@ -578,6 +630,22 @@ void FSequencerSelection::EmptySelectedOutlinerNodes()
 	}
 }
 
+void FSequencerSelection::EmptySelectedMarkedFrames()
+{
+	if (!SelectedMarkedFrames.Num())
+	{
+		return;
+	}
+	
+	++SerialNumber;
+	SelectedMarkedFrames.Empty();
+
+	if (IsBroadcasting())
+	{
+		OnMarkedFramesSelectionChanged.Broadcast();
+	}
+}
+
 /** Suspend or resume broadcast of selection changing  */
 void FSequencerSelection::SuspendBroadcast()
 {
diff --git a/Engine/Source/Editor/Sequencer/Private/SequencerSelection.h b/Engine/Source/Editor/Sequencer/Private/SequencerSelection.h
--- a/Engine/Source/Editor/Sequencer/Private/SequencerSelection.h
+++ b/Engine/Source/Editor/Sequencer/Private/SequencerSelection.h
@@ -80,6 +80,9 @@ public:
 	/** Gets a set of the selected models */
 	const TSet<TWeakPtr<FViewModel>>& GetSelectedTrackAreaItems() const;
 
+	/** Gets the indexes of the selected marked frames */
+	const TSet<int32> GetSelectedMarkedFrames() const;
+
 	/** Gets a set of the outliner nodes that have selected keys or sections */
 	const TSet<TWeakPtr<FViewModel>>& GetNodesWithSelectedKeysOrSections() const;
 
@@ -118,6 +121,9 @@ public:
 	void AddToTrackAreaSelection(TSharedPtr<FViewModel> InModel);
 	void AddToOutlinerSelection(TSharedPtr<FViewModel> InModel);
 
+	/** Adds a marked frame to the selection */
+	void AddToSelection(int32 InMarkedFrameIndex);
+
 	/** Removes a key from the selection */
 	void RemoveFromSelection(const FSequencerSelectedKey& Key);
 
@@ -130,6 +136,9 @@ public:
 	/** Removes an outliner node from the selection */
 	void RemoveFromSelection(TSharedRef<FViewModel> OutlinerNode);
 
+	/** Removes a marked frame from the selection */
+	void RemoveFromSelection(int32 InMarkedFrameIndex);
+
 	/** Removes any outliner nodes from the selection that do not relate to the given section */
 	void EmptySelectedOutlinerNodesWithoutSections(const TArray<UMovieSceneSection*>& Sections);
 
@@ -139,6 +148,9 @@ public:
 	/** Returns whether or not the model is selected. */
 	bool IsSelected(TWeakPtr<FViewModel> InModel) const;
 
+	/** Returns whether the given marked frame is selected */
+	bool IsSelected(int32 InMarkedFrameIndex) const;
+
 	/** Returns whether or not the outliner node has keys or sections selected. */
 	bool NodeHasSelectedKeysOrSections(TWeakPtr<FViewModel> Model) const;
 
@@ -148,12 +160,15 @@ public:
 	/** Empties the key selection. */
 	void EmptySelectedKeys();
 
-	/** Empties the Track Area selection. */
+	/** Empties the Track Area selection (including Time Slider selection). */
 	void EmptySelectedTrackAreaItems();
 
 	/** Empties the outliner node selection. */
 	void EmptySelectedOutlinerNodes();
 
+	/** Empties the marked frames selection. */
+	void EmptySelectedMarkedFrames();
+
 	/** Gets a multicast delegate which is called when the key selection changes. */
 	FOnSelectionChanged& GetOnKeySelectionChanged();
 
@@ -166,6 +181,9 @@ public:
 	/** Gets a multicast delegate with an array of FGuid of bound objects which is called when the outliner node selection changes. */
 	FOnSelectionChangedObjectGuids& GetOnOutlinerNodeSelectionChangedObjectGuids();
 
+	/** Gets a multicast delegate which is called when the marked frames selection changes. */
+	FOnSelectionChanged& GetOnMarkedFramesSelectionChanged();
+
 	/** Helper function to get an array of FGuid of bound objects on return */
 	TArray<FGuid> GetBoundObjectsGuids();
 
@@ -206,11 +224,13 @@ private:
 	TSet<FKeyHandle> RawSelectedKeys;
 	TSet<TWeakPtr<FViewModel>> SelectedTrackAreaItems;
 	TSet<TWeakPtr<FViewModel>> SelectedOutlinerItems;
+	TSet<int32> SelectedMarkedFrames;
 	mutable TSet<TWeakPtr<FViewModel>> NodesWithSelectedKeysOrSections;
 
 	FOnSelectionChanged OnKeySelectionChanged;
 	FOnSelectionChanged OnSectionSelectionChanged;
 	FOnSelectionChanged OnOutlinerNodeSelectionChanged;
+	FOnSelectionChanged OnMarkedFramesSelectionChanged;
 
 	FOnSelectionChangedObjectGuids OnOutlinerNodeSelectionChangedObjectGuids;
 
diff --git a/Engine/Source/Editor/Sequencer/Private/SequencerTimeSliderController.cpp b/Engine/Source/Editor/Sequencer/Private/SequencerTimeSliderController.cpp
--- a/Engine/Source/Editor/Sequencer/Private/SequencerTimeSliderController.cpp
+++ b/Engine/Source/Editor/Sequencer/Private/SequencerTimeSliderController.cpp
@@ -46,6 +46,12 @@ namespace ScrubConstants
 
 	/**The fraction of the current view range to scroll per unit delta  */
 	const float ScrollPanFraction = 0.1f;
+
+	/** Marked frame label box margin */
+	const int32 MarkLabelBoxMargin = 2;
+
+	/** Marked frame label box margin on the opposite side of the marker time */
+	const int32 MarkLabelBoxWideMargin = 4;
 }
 
 FSequencerTimeSliderController::FSequencerTimeSliderController( const FTimeSliderArgs& InArgs, TWeakPtr<FSequencer> InWeakSequencer )
@@ -55,14 +61,23 @@ FSequencerTimeSliderController::FSequencerTimeSliderController( const FTimeSlide
 	, MouseDragType( DRAG_NONE )
 	, bMouseDownInRegion(false)
 	, bPanning( false )
+	, HoverMarkIndex( INDEX_NONE )
 	, DragMarkIndex( INDEX_NONE )
+	, DragMarkReferenceFrameNumber( 0 )
+	, DragMarkCurrentFrameNumber( 0 )
 {
 	ScrubFillBrush              = FAppStyle::GetBrush( TEXT( "Sequencer.Timeline.ScrubFill" ) );
 	FrameBlockScrubHandleUpBrush   = FAppStyle::GetBrush( TEXT( "Sequencer.Timeline.FrameBlockScrubHandleUp" ) ); 
 	FrameBlockScrubHandleDownBrush = FAppStyle::GetBrush( TEXT( "Sequencer.Timeline.FrameBlockScrubHandleDown" ) );
 	VanillaScrubHandleUpBrush      = FAppStyle::GetBrush( TEXT( "Sequencer.Timeline.VanillaScrubHandleUp" ) ); 
 	VanillaScrubHandleDownBrush    = FAppStyle::GetBrush( TEXT( "Sequencer.Timeline.VanillaScrubHandleDown" ) );
+
+	FontMeasureService = FSlateApplication::Get().GetRenderer()->GetFontMeasureService().ToSharedPtr();
+	SmallLayoutFont = FCoreStyle::GetDefaultFontStyle("Regular", 10);
+	SmallBoldLayoutFont = FCoreStyle::GetDefaultFontStyle("Bold", 10);
+
 	ContextMenuSuppression = 0;
+
 	TSharedPtr<FSequencer> Sequencer = WeakSequencer.Pin();
 	if (Sequencer.IsValid())
 	{
@@ -233,7 +248,6 @@ void FSequencerTimeSliderController::DrawTicks( FSlateWindowElementList& OutDraw
 	FFrameRate     TickResolution  = GetTickResolution();
 	FFrameRate     DisplayRate     = GetDisplayRate();
 	FPaintGeometry PaintGeometry   = InArgs.AllottedGeometry.ToPaintGeometry();
-	FSlateFontInfo SmallLayoutFont = FCoreStyle::GetDefaultFontStyle("Regular", 8);
 
 	double MajorGridStep  = 0.0;
 	int32  MinorDivisions = 0;
@@ -315,7 +329,7 @@ void FSequencerTimeSliderController::DrawTicks( FSlateWindowElementList& OutDraw
 	}
 }
 
-int32 FSequencerTimeSliderController::DrawMarkedFrames( const FGeometry& AllottedGeometry, const FScrubRangeToScreen& RangeToScreen, FSlateWindowElementList& OutDrawElements, int32 LayerId, const ESlateDrawEffect& DrawEffects, bool bDrawLabels ) const
+int32 FSequencerTimeSliderController::DrawMarkedFrames( const FGeometry& AllottedGeometry, const FScrubRangeToScreen& RangeToScreen, FSlateWindowElementList& OutDrawElements, int32 LayerId, const ESlateDrawEffect& DrawEffects, const FWidgetStyle& InWidgetStyle, bool bDrawLabels ) const
 {
 	const TArray<FMovieSceneMarkedFrame> & MarkedFrames = TimeSliderArgs.MarkedFrames.Get();
 	const TArray<FMovieSceneMarkedFrame> & GlobalMarkedFrames = TimeSliderArgs.GlobalMarkedFrames.Get();
@@ -324,18 +338,33 @@ int32 FSequencerTimeSliderController::DrawMarkedFrames( const FGeometry& Allotte
 		return LayerId;
 	}
 
-	const TSharedRef< FSlateFontMeasure > FontMeasureService = FSlateApplication::Get().GetRenderer()->GetFontMeasureService();
-	FSlateFontInfo SmallLayoutFont = FCoreStyle::GetDefaultFontStyle("Regular", 10);
-
-	FQualifiedFrameTime ScrubPosition = FQualifiedFrameTime(TimeSliderArgs.ScrubPosition.Get(), GetTickResolution());
-	FScrubberMetrics    ScrubMetrics = GetScrubPixelMetrics(ScrubPosition, RangeToScreen);
-	auto DrawFrameMarkers = ([=](const TArray<FMovieSceneMarkedFrame> & InMarkedFrames, FSlateWindowElementList& DrawElements, bool bFade) {
+	const TSet<int32>& SelectedMarkedFrames = WeakSequencer.Pin()->GetSelection().GetSelectedMarkedFrames();
+	//const FLinearColor SelectedColor = FLinearColor::White; //FAppStyle::GetSlateColor("SelectionColor").GetColor(InWidgetStyle);
+	const FLinearColor WhiteColorHSV = FLinearColor::White.LinearRGBToHSV();
 
-		for (const FMovieSceneMarkedFrame& MarkedFrame : InMarkedFrames)
+	auto DrawFrameMarkers = ([=](const TArray<FMovieSceneMarkedFrame> & InMarkedFrames, FSlateWindowElementList& DrawElements, bool bIsGlobal)
+	{
+		for (int32 MarkIndex = 0; MarkIndex < InMarkedFrames.Num(); ++MarkIndex)
 		{
+			const FMovieSceneMarkedFrame& MarkedFrame(InMarkedFrames[MarkIndex]);
 			double Seconds = MarkedFrame.FrameNumber / GetTickResolution();
 
-			FLinearColor DrawColor = bFade ? MarkedFrame.Color.Desaturate(0.25f) : MarkedFrame.Color;
+			const bool bIsHovered = (!bIsGlobal && HoverMarkIndex == MarkIndex);
+			const bool bIsSelected = (!bIsGlobal && SelectedMarkedFrames.Contains(MarkIndex));
+
+			// Get a selected color that's the marked frame color but at full opacity, full brightness, and a bit desaturated if it's
+			// already bright to begin with.
+			const FLinearColor MarkedFrameColorHSV = MarkedFrame.Color.CopyWithNewOpacity(1.f).LinearRGBToHSV();
+			const FLinearColor SelectedColor = FLinearColor::LerpUsingHSV(
+				MarkedFrameColorHSV, 
+				FLinearColor(
+					MarkedFrameColorHSV.R, 
+					MarkedFrameColorHSV.G * (1.f - MarkedFrameColorHSV.B),
+					1.f),
+				0.5f)
+				.HSVToLinearRGB();
+
+			FLinearColor DrawColor = bIsGlobal ? MarkedFrame.Color.Desaturate(0.25f) : (bIsSelected ? SelectedColor : MarkedFrame.Color);
 			const float  LinePos = RangeToScreen.InputToLocalX(Seconds);
 			TArray<FVector2D> LinePoints;
 			LinePoints.AddUninitialized(2);
@@ -349,32 +378,47 @@ int32 FSequencerTimeSliderController::DrawMarkedFrames( const FGeometry& Allotte
 				LinePoints,
 				DrawEffects,
 				DrawColor,
-				false
+				false,
+				(bIsHovered || bIsSelected) ? 2.f : 1.f
 			);
 
 			FString LabelString = MarkedFrame.Label;
 			if (bDrawLabels && !LabelString.IsEmpty())
 			{
 				// Draw the label next to the marked frame line
-				FVector2D TextSize = FontMeasureService->Measure(LabelString, SmallLayoutFont);
+				bool bDrawLeft;
+				FVector2D TextPosition, TextSize;
+				GetMarkLabelGeometry(AllottedGeometry, RangeToScreen, MarkedFrame, TextPosition, TextSize, bDrawLeft);
 
-				// Flip the text position if getting near the end of the view range
-				static const float TextOffsetPx = 2.f;
-				bool  bDrawLeft = (AllottedGeometry.Size.X - LinePos) < (TextSize.X + 14.f) - TextOffsetPx;
-				float TextPosition = bDrawLeft ? LinePos - TextSize.X - TextOffsetPx : LinePos + TextOffsetPx;
+				const FSlateBrush* LabelBrush = bDrawLeft ?
+					FAppStyle::GetBrush("Sequencer.MarkedFrame.LabelLeft") :
+					FAppStyle::GetBrush("Sequencer.MarkedFrame.LabelRight");
+
+				if (bIsHovered || bIsSelected)
+				{
+					FSlateDrawElement::MakeBox(
+						DrawElements,
+						LayerId + 1,
+						AllottedGeometry.ToPaintGeometry(
+							TextSize + ScrubConstants::MarkLabelBoxMargin + ScrubConstants::MarkLabelBoxWideMargin,
+							FSlateLayoutTransform(FVector2D(LinePos, 0.f))),
+						LabelBrush,
+						DrawEffects,
+						DrawColor.CopyWithNewOpacity(0.3f)
+					);
+				}
 
 				FSlateDrawElement::MakeText(
 					DrawElements,
 					LayerId + 1,
-					AllottedGeometry.ToPaintGeometry(TextSize, FSlateLayoutTransform(FVector2f(TextPosition, 0.f))),
+					AllottedGeometry.ToPaintGeometry(TextSize, FSlateLayoutTransform(TextPosition)),
 					LabelString,
-					SmallLayoutFont,
+					(bIsHovered || bIsSelected) ? SmallBoldLayoutFont : SmallLayoutFont,
 					DrawEffects,
 					DrawColor
 				);
 			}
 		}
-
 	});
 	
 	DrawFrameMarkers(GlobalMarkedFrames, OutDrawElements, true);
@@ -504,7 +548,7 @@ int32 FSequencerTimeSliderController::OnPaintTimeSlider( bool bMirrorLabels, con
 			ScrubColor
 		);
 		
-		LayerId = DrawMarkedFrames(AllottedGeometry, RangeToScreen, OutDrawElements, LayerId, DrawEffects, true);
+		LayerId = DrawMarkedFrames(AllottedGeometry, RangeToScreen, OutDrawElements, LayerId, DrawEffects, InWidgetStyle, true);
 
 		{
 			// Draw the current time next to the scrub handle
@@ -524,9 +568,6 @@ int32 FSequencerTimeSliderController::OnPaintTimeSlider( bool bMirrorLabels, con
 				TextColor = FLinearColor::Yellow;
 			}
 
-			FSlateFontInfo SmallLayoutFont = FCoreStyle::GetDefaultFontStyle("Regular", 10);
-
-			const TSharedRef< FSlateFontMeasure > FontMeasureService = FSlateApplication::Get().GetRenderer()->GetFontMeasureService();
 			FVector2D TextSize = FontMeasureService->Measure(FrameString, SmallLayoutFont);
 
 			// Flip the text position if getting near the end of the view range
@@ -858,6 +899,7 @@ FReply FSequencerTimeSliderController::OnMouseButtonUp( SWidget& WidgetOwner, co
 		else if (MouseDragType == DRAG_MARK)
 		{
 			TimeSliderArgs.OnMarkEndDrag.ExecuteIfBound();
+			UpdateMarkSelection(DragMarkIndex, DragMarkCurrentFrameNumber);
 		}
 		else if (MouseDragType == DRAG_SETTING_RANGE)
 		{
@@ -892,27 +934,53 @@ FReply FSequencerTimeSliderController::OnMouseButtonUp( SWidget& WidgetOwner, co
 		}
 		else if (bMouseDownInRegion)
 		{
-			TimeSliderArgs.OnEndScrubberMovement.ExecuteIfBound();
-
-			FFrameTime ScrubTime = MouseTime;
-			FVector2D CursorPos  = MouseEvent.GetScreenSpacePosition();
-
 			TSharedPtr<FSequencer> Sequencer = WeakSequencer.Pin();
 
-			if (MouseDragType == DRAG_SCRUBBING_TIME)
+			if (HoverMarkIndex == INDEX_NONE)
 			{
-				ScrubTime = ComputeScrubTimeFromMouse(MyGeometry, MouseEvent, RangeToScreen);
+				// Teleport the playhead to the clicked time.
+				TimeSliderArgs.OnEndScrubberMovement.ExecuteIfBound();
+
+				FFrameTime ScrubTime = MouseTime;
+				FVector2D CursorPos  = MouseEvent.GetScreenSpacePosition();
+
+				if (MouseDragType == DRAG_SCRUBBING_TIME)
+				{
+					ScrubTime = ComputeScrubTimeFromMouse(MyGeometry, MouseEvent, RangeToScreen);
+				}
+				else if (Sequencer.IsValid())
+				{
+					ScrubTime = SnapTimeToNearestKey(MouseEvent, RangeToScreen, CursorPos.X, ScrubTime);
+
+					// We weren't scrubbing, so this is a single click in the time slider. This should
+					// also act as a deselection of any selected markers.
+					Sequencer->GetSelection().EmptySelectedMarkedFrames();
+				}
+
+				// If middle mouse button down we don't evaluate on the time change
+				CommitScrubPosition( ScrubTime, /*bIsScrubbing=*/false , /*bEvaluate*/ !bHandleMiddleMouseButton);
 			}
 			else if (Sequencer.IsValid())
 			{
-				ScrubTime = SnapTimeToNearestKey(MouseEvent, RangeToScreen, CursorPos.X, ScrubTime);
-			}
+				// Select the clicked marker, and optionally teleport the playhead.
+				HandleMarkSelection(HoverMarkIndex);
 
-			CommitScrubPosition( ScrubTime, /*bIsScrubbing=*/false , /*bEvaluate*/ !bHandleMiddleMouseButton); //if middle mouse button down we don't evaluate on the time change
+				FModifierKeysState ModifierKeys = FSlateApplication::Get().GetModifierKeys();
+				const bool bSnapPlayTimeToMark = ModifierKeys.AreModifersDown(EModifierKey::Shift);
+				if (bSnapPlayTimeToMark || Sequencer->GetSequencerSettings()->GetSnapPlayTimeToPressedKey())
+				{
+					const TArray<FMovieSceneMarkedFrame>& MarkedFrames = TimeSliderArgs.MarkedFrames.Get();
+					const FFrameNumber MarkTime = MarkedFrames[HoverMarkIndex].FrameNumber;
+
+					// If middle mouse button down we don't evaluate on the time change
+					CommitScrubPosition(MarkTime, /*bIsScrubbing=*/ false , /*bEvaluate*/ !bHandleMiddleMouseButton);
+				}
+			}
 		}
 
 		MouseDragType = DRAG_NONE;
 		DistanceDragged = 0.f;
+		DragMarkIndex = INDEX_NONE;
 		bMouseDownInRegion = false;
 
 		return FReply::Handled().ReleaseMouseCapture();
@@ -922,8 +990,30 @@ FReply FSequencerTimeSliderController::OnMouseButtonUp( SWidget& WidgetOwner, co
 	return FReply::Unhandled();
 }
 
+FReply FSequencerTimeSliderController::OnMouseButtonDoubleClick(TSharedRef<const SWidget> WidgetOwner, const FGeometry& MyGeometry, const FPointerEvent& MouseEvent)
+{
+	if (HoverMarkIndex != INDEX_NONE)
+	{
+		// Teleport the playhead to the double-clicked marker.
+		bool bHandleMiddleMouseButton = MouseEvent.GetEffectingButton() == EKeys::MiddleMouseButton && WidgetOwner->HasMouseCapture();
+		const TArray<FMovieSceneMarkedFrame>& MarkedFrames = TimeSliderArgs.MarkedFrames.Get();
+		const FFrameNumber MarkTime = MarkedFrames[HoverMarkIndex].FrameNumber;
+		CommitScrubPosition(MarkTime, /*bIsScrubbing=*/ false , /*bEvaluate*/ !bHandleMiddleMouseButton);
+	}
+	return FReply::Unhandled();
+}
 
 FReply FSequencerTimeSliderController::OnMouseMove( SWidget& WidgetOwner, const FGeometry& MyGeometry, const FPointerEvent& MouseEvent )
+{
+	return OnMouseMoveImpl(WidgetOwner, MyGeometry, MouseEvent, false);
+}
+
+FReply FSequencerTimeSliderController::OnTimeSliderMouseMove(SWidget& WidgetOwner, const FGeometry& MyGeometry, const FPointerEvent& MouseEvent)
+{
+	return OnMouseMoveImpl(WidgetOwner, MyGeometry, MouseEvent, true);
+}
+
+FReply FSequencerTimeSliderController::OnMouseMoveImpl( SWidget& WidgetOwner, const FGeometry& MyGeometry, const FPointerEvent& MouseEvent, bool bFromTimeSlider )
 {
 	TSharedPtr<FSequencer> Sequencer = WeakSequencer.Pin();
 	if (!Sequencer.IsValid())
@@ -935,6 +1025,8 @@ FReply FSequencerTimeSliderController::OnMouseMove( SWidget& WidgetOwner, const
 	bool bHandleRightMouseButton = MouseEvent.IsMouseButtonDown( EKeys::RightMouseButton ) && TimeSliderArgs.AllowZoom;
 	bool bHandleMiddleMouseButton = MouseEvent.IsMouseButtonDown(EKeys::MiddleMouseButton);
 
+	HoverMarkIndex = INDEX_NONE;
+
 	if (bHandleRightMouseButton)
 	{
 		if (!bPanning)
@@ -1028,9 +1120,11 @@ FReply FSequencerTimeSliderController::OnMouseMove( SWidget& WidgetOwner, const
 					MouseDragType = DRAG_PLAYBACK_START;
 					TimeSliderArgs.OnPlaybackRangeBeginDrag.ExecuteIfBound();
 				}
-				else if (!bReadOnly && !bHitScrubber && HitTestMark(RangeToScreen, MouseDownPixel, DragMarkIndex) && bHandleMiddleMouseButton == false)
+				else if (!bReadOnly && !bHitScrubber && HitTestMark(MyGeometry, RangeToScreen, MouseDownPixel, bFromTimeSlider, &DragMarkIndex, &DragMarkReferenceFrameNumber) && bHandleMiddleMouseButton == false)
 				{
 					MouseDragType = DRAG_MARK;
+					DragMarkCurrentFrameNumber = DragMarkReferenceFrameNumber;
+					HandleMarkSelection(DragMarkIndex);
 					TimeSliderArgs.OnMarkBeginDrag.ExecuteIfBound();
 				}
 				else if (FSlateApplication::Get().GetModifierKeys().AreModifersDown(EModifierKey::Control) && bHandleMiddleMouseButton == false)
@@ -1104,7 +1198,8 @@ FReply FSequencerTimeSliderController::OnMouseMove( SWidget& WidgetOwner, const
 			}
 			else if (MouseDragType == DRAG_MARK)
 			{
-				SetMark(DragMarkIndex, MouseTime.FrameNumber);
+				DragMarkCurrentFrameNumber = DragMarkReferenceFrameNumber + DiffFrame;
+				SetMark(DragMarkIndex, DragMarkCurrentFrameNumber);
 			}
 			else if (MouseDragType == DRAG_SCRUBBING_TIME)
 			{
@@ -1117,6 +1212,14 @@ FReply FSequencerTimeSliderController::OnMouseMove( SWidget& WidgetOwner, const
 			}
 		}
 	}
+	else if (bFromTimeSlider && DragMarkIndex == INDEX_NONE)
+	{
+		// Update hover state of marked frames.
+		TRange<double> LocalViewRange = GetViewRange();
+		FScrubRangeToScreen RangeToScreen(LocalViewRange, MyGeometry.Size);
+		const FVector2D LocalMousePostion = MyGeometry.AbsoluteToLocal(MouseEvent.GetScreenSpacePosition());
+		HitTestMark(MyGeometry, RangeToScreen, LocalMousePostion.X, bFromTimeSlider, &HoverMarkIndex);
+	}
 
 	if ( DistanceDragged != 0.f && (bHandleLeftMouseButton || bHandleRightMouseButton) )
 	{
@@ -1235,8 +1338,8 @@ FCursorReply FSequencerTimeSliderController::OnCursorQuery( TSharedRef<const SWi
 		return FCursorReply::Cursor(EMouseCursor::ResizeLeftRight);
 	}
 
-	int32 DummyMarkIndex = INDEX_NONE;
-	if (MouseDragType == DRAG_MARK || (!bReadOnly && !bHitScrubber && HitTestMark(RangeToScreen, HitTestPixel, DummyMarkIndex)))
+	const bool bFromTimeSlider = true;
+	if (MouseDragType == DRAG_MARK || (!bReadOnly && !bHitScrubber && HitTestMark(MyGeometry, RangeToScreen, HitTestPixel, bFromTimeSlider)))
 	{
 		return FCursorReply::Cursor(EMouseCursor::CardinalCross);
 	}
@@ -1291,7 +1394,7 @@ int32 FSequencerTimeSliderController::OnPaintViewArea( const FGeometry& Allotted
 
 	if (Args.bDisplayMarkedFrames)
 	{
-		LayerId = DrawMarkedFrames(AllottedGeometry, RangeToScreen, OutDrawElements, LayerId, DrawEffects, false);
+		LayerId = DrawMarkedFrames(AllottedGeometry, RangeToScreen, OutDrawElements, LayerId, DrawEffects, FWidgetStyle(), false);
 	}
 
 	LayerId = DrawVerticalFrames(AllottedGeometry, RangeToScreen, OutDrawElements, LayerId, DrawEffects);
@@ -1340,7 +1443,7 @@ int32 FSequencerTimeSliderController::OnPaintViewArea( const FGeometry& Allotted
 
 TSharedRef<SWidget> FSequencerTimeSliderController::OpenSetPlaybackRangeMenu(const FGeometry& MyGeometry, const FPointerEvent& MouseEvent)
 {
-	TSharedPtr<ISequencer> Sequencer = WeakSequencer.Pin();
+	TSharedPtr<FSequencer> Sequencer = WeakSequencer.Pin();
 	const bool bReadOnly = Sequencer && Sequencer->IsReadOnly();
 	
 	FScrubRangeToScreen RangeToScreen = FScrubRangeToScreen(GetViewRange(), MyGeometry.Size);
@@ -1427,8 +1530,8 @@ TSharedRef<SWidget> FSequencerTimeSliderController::OpenSetPlaybackRangeMenu(con
 	}
 	MenuBuilder.EndSection(); // SequencerPlaybackRangeMenu
 
-	UMovieSceneCompiledDataManager* CompiledDataManager = WeakSequencer.Pin()->GetEvaluationTemplate().GetCompiledDataManager();
-	const FMovieSceneSequenceHierarchy* Hierarchy = CompiledDataManager->FindHierarchy(WeakSequencer.Pin()->GetEvaluationTemplate().GetCompiledDataID());
+	UMovieSceneCompiledDataManager* CompiledDataManager = Sequencer->GetEvaluationTemplate().GetCompiledDataManager();
+	const FMovieSceneSequenceHierarchy* Hierarchy = CompiledDataManager->FindHierarchy(Sequencer->GetEvaluationTemplate().GetCompiledDataID());
 
 	if (TimeSliderArgs.ScrubPositionParentChain.IsSet() && Hierarchy)
 	{
@@ -1437,7 +1540,7 @@ TSharedRef<SWidget> FSequencerTimeSliderController::OpenSetPlaybackRangeMenu(con
 			TArray<FMovieSceneSequenceID> ParentChain = TimeSliderArgs.ScrubPositionParentChain.Get();
 			for (FMovieSceneSequenceID ParentID : ParentChain)
 			{
-				FText ParentText = WeakSequencer.Pin()->GetRootMovieSceneSequence()->GetDisplayName();
+				FText ParentText = Sequencer->GetRootMovieSceneSequence()->GetDisplayName();
 
 				for (const TTuple<FMovieSceneSequenceID, FMovieSceneSubSequenceData>& Pair : Hierarchy->AllSubSequenceData())
 				{
@@ -1467,16 +1570,21 @@ TSharedRef<SWidget> FSequencerTimeSliderController::OpenSetPlaybackRangeMenu(con
 
 	MenuBuilder.BeginSection("SequencerMarkMenu", FText::Format(LOCTEXT("MarkTextFormat", "Mark ({0}):"), CurrentTimeText));
 	{
-		FFrameNumber DisplayFrameNumber = GetDisplayRate().AsFrameNumber(FrameNumber / GetTickResolution());
-
-		UMovieScene* MovieScene = WeakSequencer.Pin()->GetFocusedMovieSceneSequence()->GetMovieScene();
+		UMovieScene* MovieScene = Sequencer->GetFocusedMovieSceneSequence()->GetMovieScene();
 		bool bHasMarks = MovieScene->GetMarkedFrames().Num() > 0;
 
 		int32 MarkedIndex = INDEX_NONE;
-		HitTestMark(RangeToScreen, MousePixel, MarkedIndex);
+		const bool bTestLabelBox = true;
+		HitTestMark(MyGeometry, RangeToScreen, MousePixel, bTestLabelBox, &MarkedIndex);
+
+		FSequencerSelection& SequencerSelection = Sequencer->GetSelection();
+		SequencerSelection.Empty();
 
 		if (MarkedIndex != INDEX_NONE)
 		{
+			SequencerSelection.EmptySelectedMarkedFrames();
+			SequencerSelection.AddToSelection(MarkedIndex);
+
 			class SMarkedFramePropertyWidget : public SCompoundWidget, public FNotifyHook
 			{
 			public:
@@ -1533,28 +1641,23 @@ TSharedRef<SWidget> FSequencerTimeSliderController::OpenSetPlaybackRangeMenu(con
 			MenuBuilder.AddWidget(Widget, FText::GetEmpty(), false);
 		}
 
-		if (MarkedIndex == INDEX_NONE)
-		{
-			MenuBuilder.AddMenuEntry( 
-				LOCTEXT("AddMark", "Add Mark"),
-				FText(),
-				FSlateIcon(),
-				FUIAction(
-					FExecuteAction::CreateLambda( [=]{ AddMarkAtFrame(FrameNumber); }),
-					FCanExecuteAction::CreateLambda([=]{ return !bReadOnly; }))
-			);
-		}
-		else 
-		{
-			MenuBuilder.AddMenuEntry( 
-				LOCTEXT("DeleteMark", "Delete Mark"),
-				FText(),
-				FSlateIcon(),
-				FUIAction(
-					FExecuteAction::CreateLambda([=]{ DeleteMarkAtIndex(MarkedIndex); }),
-					FCanExecuteAction::CreateLambda([=]{ return !bReadOnly; }))
-			);
-		}
+		MenuBuilder.AddMenuEntry( 
+			LOCTEXT("AddMark", "Add Mark"),
+			FText(),
+			FSlateIcon(),
+			FUIAction(
+				FExecuteAction::CreateLambda( [=]{ AddMarkAtFrame(FrameNumber); }),
+				FCanExecuteAction::CreateLambda([=]{ return !bReadOnly && MarkedIndex == INDEX_NONE; }))
+		);
+
+		MenuBuilder.AddMenuEntry( 
+			LOCTEXT("DeleteMark", "Delete Mark"),
+			FText(),
+			FSlateIcon(),
+			FUIAction(
+				FExecuteAction::CreateLambda([=]{ DeleteMarkAtIndex(MarkedIndex); }),
+				FCanExecuteAction::CreateLambda([=]{ return !bReadOnly && MarkedIndex != INDEX_NONE; }))
+		);
 
 		MenuBuilder.AddMenuEntry( 
 			LOCTEXT("Delete All Marks", "Delete All Marks"),
@@ -1715,8 +1818,17 @@ bool FSequencerTimeSliderController::HitTestRangeEnd(const FScrubRangeToScreen&
 		HitPixel <= RangeEndPixel + MouseTolerance + DragToleranceSlateUnits;
 }
 
-bool FSequencerTimeSliderController::HitTestMark(const FScrubRangeToScreen& RangeToScreen, float HitPixel, int32& OutMarkIndex) const
+bool FSequencerTimeSliderController::HitTestMark(const FGeometry& AllottedGeometry, const FScrubRangeToScreen& RangeToScreen, float HitPixel, bool bTestLabelBox, int32* OutMarkIndex, FFrameNumber* OutMarkFrameNumber) const
 {
+	if (OutMarkIndex)
+	{
+		*OutMarkIndex = INDEX_NONE;
+	}
+	if (OutMarkFrameNumber)
+	{
+		*OutMarkFrameNumber = 0;
+	}
+
 	const TArray<FMovieSceneMarkedFrame> & MarkedFrames = TimeSliderArgs.MarkedFrames.Get();
 	if (MarkedFrames.Num() < 1)
 	{
@@ -1737,14 +1849,61 @@ bool FSequencerTimeSliderController::HitTestMark(const FScrubRangeToScreen& Rang
 			(HitPixel >= MarkPixel - MouseTolerance - BrushSizeInStateUnits - DragToleranceSlateUnits &&
 			 HitPixel <= MarkPixel + MouseTolerance + DragToleranceSlateUnits))
 		{
-			OutMarkIndex = MarkIndex;
+			if (OutMarkIndex)
+			{
+				*OutMarkIndex = MarkIndex;
+			}
+			if (OutMarkFrameNumber)
+			{
+				*OutMarkFrameNumber = MarkedFrames[MarkIndex].FrameNumber;
+			}
 			return true;
 		}
+
+		// Hit test against the label box, which is the text size offset and grown by the margin.
+		if (bTestLabelBox)
+		{
+			FVector2D TextPosition, TextSize;
+			bool bDrawLeft;
+			GetMarkLabelGeometry(AllottedGeometry, RangeToScreen, MarkedFrames[MarkIndex], TextPosition, TextSize, bDrawLeft);
+			FVector2D TextBoxSize = TextSize + ScrubConstants::MarkLabelBoxMargin + ScrubConstants::MarkLabelBoxWideMargin;
+			if (HitPixel >= MarkPixel && HitPixel <= MarkPixel + TextBoxSize.X)
+			{
+				if (OutMarkIndex)
+				{
+					*OutMarkIndex = MarkIndex;
+				}
+				if (OutMarkFrameNumber)
+				{
+					*OutMarkFrameNumber = MarkedFrames[MarkIndex].FrameNumber;
+				}
+				return true;
+			}
+		}
 	}
 
 	return false;
 }
 
+void FSequencerTimeSliderController::GetMarkLabelGeometry(const FGeometry& AllottedGeometry, const FScrubRangeToScreen& RangeToScreen, const FMovieSceneMarkedFrame& MarkedFrame, FVector2D& OutPosition, FVector2D& OutSize, bool& bIsDrawLeft) const
+{
+	double Seconds = MarkedFrame.FrameNumber / GetTickResolution();
+	float MarkPixel = RangeToScreen.InputToLocalX(Seconds);
+
+	const FString& LabelString = MarkedFrame.Label;
+	FVector2D TextSize = FontMeasureService->Measure(LabelString, SmallLayoutFont);
+
+	// Flip the text position if getting near the end of the view range
+	bool bDrawLeft = (AllottedGeometry.Size.X - MarkPixel) < (TextSize.X + ScrubConstants::MarkLabelBoxMargin + ScrubConstants::MarkLabelBoxWideMargin);
+	float TextPosition = bDrawLeft ? 
+		MarkPixel - TextSize.X - ScrubConstants::MarkLabelBoxMargin : 
+		MarkPixel + ScrubConstants::MarkLabelBoxMargin;
+
+	OutPosition = FVector2D(TextPosition, 0);
+	OutSize = TextSize;
+	bIsDrawLeft = bDrawLeft;
+}
+
 FFrameTime FSequencerTimeSliderController::SnapTimeToNearestKey(const FPointerEvent& MouseEvent, const FScrubRangeToScreen& RangeToScreen, float CursorPos, FFrameTime InTime) const
 {
 	using namespace UE::Sequencer;
@@ -1843,6 +2002,44 @@ void FSequencerTimeSliderController::SetSelectionRangeEnd(FFrameNumber NewEnd)
 	}
 }
 
+void FSequencerTimeSliderController::HandleMarkSelection(int32 InMarkIndex)
+{
+	TSharedPtr<FSequencer> Sequencer = WeakSequencer.Pin();
+	FSequencerSelection& SequencerSelection = Sequencer->GetSelection();
+
+	FModifierKeysState ModifierKeys = FSlateApplication::Get().GetModifierKeys();
+	const bool bToggleSelection = ModifierKeys.AreModifersDown(EModifierKey::Control);
+	const bool bAddToSelection = ModifierKeys.AreModifersDown(EModifierKey::Shift);
+
+	if (!bToggleSelection && !bAddToSelection)
+	{
+		SequencerSelection.Empty();
+		SequencerSelection.AddToSelection(InMarkIndex);
+	}
+	else if (bAddToSelection || !SequencerSelection.IsSelected(InMarkIndex))
+	{
+		SequencerSelection.AddToSelection(InMarkIndex);
+	}
+	else
+	{
+		SequencerSelection.RemoveFromSelection(InMarkIndex);
+	}
+}
+
+void FSequencerTimeSliderController::UpdateMarkSelection(int32 InOldMarkIndex, FFrameNumber InMarkFrameNumber)
+{
+	TSharedPtr<FSequencer> Sequencer = WeakSequencer.Pin();
+	FSequencerSelection& SequencerSelection = Sequencer->GetSelection();
+	SequencerSelection.RemoveFromSelection(InOldMarkIndex);
+
+	UMovieScene* MovieScene = Sequencer->GetFocusedMovieSceneSequence()->GetMovieScene();
+	int32 NewMarkIndex = MovieScene->FindMarkedFrameByFrameNumber(InMarkFrameNumber);
+	if (ensure(NewMarkIndex != INDEX_NONE))
+	{
+		SequencerSelection.AddToSelection(NewMarkIndex);
+	}
+}
+
 void FSequencerTimeSliderController::SetMark(int32 InMarkIndex, FFrameNumber FrameNumber)
 {
 	TimeSliderArgs.OnSetMarkedFrame.ExecuteIfBound(InMarkIndex, FrameNumber);
diff --git a/Engine/Source/Editor/Sequencer/Private/Tools/EditToolDragOperations.cpp b/Engine/Source/Editor/Sequencer/Private/Tools/EditToolDragOperations.cpp
--- a/Engine/Source/Editor/Sequencer/Private/Tools/EditToolDragOperations.cpp
+++ b/Engine/Source/Editor/Sequencer/Private/Tools/EditToolDragOperations.cpp
@@ -634,7 +634,6 @@ FMoveKeysAndSections::FMoveKeysAndSections(FSequencer& InSequencer, ESequencerMo
 		KeysAsArray = Keys.Array();
 	}
 
-
 	if (EnumHasAnyFlags(MoveType, ESequencerMoveOperationType::MoveSections))
 	{
 		for (TWeakPtr<FViewModel> WeakSelectedModel : Sequencer.GetSelection().GetSelectedTrackAreaItems())
@@ -645,6 +644,9 @@ FMoveKeysAndSections::FMoveKeysAndSections(FSequencer& InSequencer, ESequencerMo
 			}
 		}
 	}
+
+	// Always move selected marked frames along with keys and/or sections.
+	MarkedFrames = Sequencer.GetSelection().GetSelectedMarkedFrames();
 }
 
 void FMoveKeysAndSections::AddSnapTime(FFrameNumber SnapTime)
@@ -751,11 +753,6 @@ void FMoveKeysAndSections::OnDrag(const FPointerEvent& MouseEvent, FVector2D Loc
 		return;
 	}
 
-	ESequencerScrubberStyle ScrubStyle = Sequencer.GetScrubStyle();
-
-	FFrameRate TickResolution = Sequencer.GetFocusedTickResolution();
-	FFrameRate DisplayRate = Sequencer.GetFocusedDisplayRate();
-
 	// Convert the current mouse position to a time
 	FVector2D  VirtualMousePos = VirtualTrackArea.PhysicalToVirtual(LocalMousePos);
 	FFrameTime MouseTime = VirtualTrackArea.PixelToFrame(LocalMousePos.X);
@@ -820,6 +817,9 @@ void FMoveKeysAndSections::OnDrag(const FPointerEvent& MouseEvent, FVector2D Loc
 	// Update our key times by moving them by our delta.
 	HandleKeyMovement(MaxDeltaX, MouseDeltaTime);
 
+	// Update our marked frames by moving them by our delta.
+	HandleMarkedFrameMovement(MaxDeltaX, MouseDeltaTime);
+
 	// Get a list of the unique tracks in this selection and update their easing so previews draw interactively as you drag.
 	TSet<UMovieSceneTrack*> Tracks;
 	FMovieSceneSectionMovedParams SectionMovedParams(EPropertyChangeType::Interactive);
@@ -1434,3 +1434,26 @@ void FMoveKeysAndSections::HandleKeyMovement(TOptional<FFrameNumber> MaxDeltaX,
 		}
 	}
 }
+
+void FMoveKeysAndSections::HandleMarkedFrameMovement(TOptional<FFrameNumber> MaxDeltaX, FFrameNumber DesiredDeltaX)
+{
+	if (MarkedFrames.Num() == 0)
+	{
+		return;
+	}
+
+	const FFrameNumber EffectiveDelta = MaxDeltaX.Get(DesiredDeltaX);
+	UMovieScene* FocusedMovieScene = Sequencer.GetFocusedMovieSceneSequence()->GetMovieScene();
+	const TArray<FMovieSceneMarkedFrame>& AllMarkedFrames = FocusedMovieScene->GetMarkedFrames();
+
+	for (int32 MarkIndex = 0; MarkIndex < MarkedFrames.Num(); ++MarkIndex)
+	{
+		const FMovieSceneMarkedFrame& MarkedFrame = AllMarkedFrames[MarkIndex];
+		const FFrameNumber NewMarkTime = MarkedFrame.FrameNumber + EffectiveDelta;
+		FocusedMovieScene->SetMarkedFrame(MarkIndex, NewMarkTime);
+	}
+
+	FocusedMovieScene->MarkAsChanged();
+	FocusedMovieScene->BroadcastChanged();
+}
+
diff --git a/Engine/Source/Editor/Sequencer/Private/Tools/EditToolDragOperations.h b/Engine/Source/Editor/Sequencer/Private/Tools/EditToolDragOperations.h
--- a/Engine/Source/Editor/Sequencer/Private/Tools/EditToolDragOperations.h
+++ b/Engine/Source/Editor/Sequencer/Private/Tools/EditToolDragOperations.h
@@ -179,6 +179,8 @@ protected:
 	bool HandleSectionMovement(FFrameTime MouseTime, FVector2D VirtualMousePos, FVector2D LocalMousePos, TOptional<FFrameNumber> MaxDeltaX, FFrameNumber DesiredDeltaX);
 	/** Move selected keys, if any. */
 	void HandleKeyMovement(TOptional<FFrameNumber> MaxDeltaX, FFrameNumber DesiredDeltaX);
+	/** Move selected marked frames, if any. */
+	void HandleMarkedFrameMovement(TOptional<FFrameNumber> MaxDeltaX, FFrameNumber DesiredDeltaX);
 
 	void OnSequencerNodeTreeUpdated();
 
@@ -196,6 +198,9 @@ protected:
 	TSet<FSequencerSelectedKey> Keys;
 	TArray<FSequencerSelectedKey> KeysAsArray;
 
+	/** Set of marked frames that are being moved. */
+	TSet<int32> MarkedFrames;
+
 	/** What was the time of the mouse for the previous frame? Used to calculate a per-frame delta. */
 	FFrameTime MouseTimePrev;
 
diff --git a/Engine/Source/Editor/Sequencer/Public/SequencerTimeSliderController.h b/Engine/Source/Editor/Sequencer/Public/SequencerTimeSliderController.h
--- a/Engine/Source/Editor/Sequencer/Public/SequencerTimeSliderController.h
+++ b/Engine/Source/Editor/Sequencer/Public/SequencerTimeSliderController.h
@@ -12,6 +12,7 @@
 class FSlateWindowElementList;
 struct FContextMenuSuppressor;
 struct FSlateBrush;
+class FSlateFontMeasure;
 class FSequencer;
 class IPropertyTypeCustomization;
 
@@ -42,7 +43,9 @@ public:
 	virtual FReply OnMouseButtonUp( SWidget& WidgetOwner, const FGeometry& MyGeometry, const FPointerEvent& MouseEvent ) override;
 	virtual FReply OnMouseMove( SWidget& WidgetOwner, const FGeometry& MyGeometry, const FPointerEvent& MouseEvent ) override;
 	virtual FReply OnMouseWheel( SWidget& WidgetOwner, const FGeometry& MyGeometry, const FPointerEvent& MouseEvent ) override;
+	virtual FReply OnTimeSliderMouseMove( SWidget& OwnerWidget, const FGeometry& MyGeometry, const FPointerEvent& MouseEvent ) override;
 	virtual FCursorReply OnCursorQuery( TSharedRef<const SWidget> WidgetOwner, const FGeometry& MyGeometry, const FPointerEvent& CursorEvent ) const override;
+	virtual FReply OnMouseButtonDoubleClick( TSharedRef<const SWidget> WidgetOwner, const FGeometry& MyGeometry, const FPointerEvent& MouseEvent ) override;
 	/** End ITimeSliderController Interface */
 
 	/** Get the current play rate for this controller */
@@ -165,6 +168,12 @@ public:
 			return PixelsPerInput > 0 ? (ScreenX / PixelsPerInput) + ViewStart : ViewStart;
 		}
 
+		/** Local Widget Space -> Curve Input domain. */
+		double LocalDeltaXToDeltaInput(float ScreenDeltaX) const
+		{
+			return PixelsPerInput > 0 ? (ScreenDeltaX / PixelsPerInput) : 0;
+		}
+
 		/** Curve Input domain -> local Widget Space */
 		float InputToLocalX(double Input) const
 		{
@@ -178,6 +187,10 @@ public:
 		bIsEvaluating = true;
 	}
 
+private:
+
+	FReply OnMouseMoveImpl( SWidget& WidgetOwner, const FGeometry& MyGeometry, const FPointerEvent& MouseEvent, bool bFromTimeSlider );
+
 private:
 	// forward declared as class members to prevent name collision with similar types defined in other units
 	struct FDrawTickArgs;
@@ -234,7 +247,7 @@ private:
 	 *
 	 * @return the new layer ID
 	 */
-	int32 DrawMarkedFrames(const FGeometry& AllottedGeometry, const FScrubRangeToScreen& RangeToScreen, FSlateWindowElementList& OutDrawElements, int32 LayerId, const ESlateDrawEffect& DrawEffects, bool bDrawLabels) const;
+	int32 DrawMarkedFrames(const FGeometry& AllottedGeometry, const FScrubRangeToScreen& RangeToScreen, FSlateWindowElementList& OutDrawElements, int32 LayerId, const ESlateDrawEffect& DrawEffects, const FWidgetStyle& InWidgetStyle, bool bDrawLabels) const;
 
 private:
 
@@ -253,7 +266,12 @@ private:
 	 *
 	 * @return The mark index hit
 	 */
-	bool HitTestMark(const FScrubRangeToScreen& RangeToScreen, float HitPixel, int32& OutMarkIndex) const;
+	bool HitTestMark(const FGeometry& AllottedGeometry, const FScrubRangeToScreen& RangeToScreen, float HitPixel, bool bTestLabelBox, int32* OutMarkIndex = nullptr, FFrameNumber* OutMarkFrameNumber = nullptr) const;
+
+	/**
+	 * Get marked frame label box size
+	 */
+	void GetMarkLabelGeometry(const FGeometry& AllottedGeometry, const FScrubRangeToScreen& RangeToScreen, const FMovieSceneMarkedFrame& MarkedFrame, FVector2D& OutPosition, FVector2D& OutSize, bool& bIsDrawLeft) const;
 
 	FFrameTime SnapTimeToNearestKey(const FPointerEvent& MouseEvent, const FScrubRangeToScreen& RangeToScreen, float CursorPos, FFrameTime InTime) const;
 
@@ -269,6 +287,8 @@ private:
 	FFrameTime ComputeScrubTimeFromMouse(const FGeometry& Geometry, const FPointerEvent& MouseEvent, FScrubRangeToScreen RangeToScreen) const;
 	FFrameTime ComputeFrameTimeFromMouse(const FGeometry& Geometry, FVector2D ScreenSpacePosition, FScrubRangeToScreen RangeToScreen, bool CheckSnapping = true) const;
 
+	void HandleMarkSelection(int32 InMarkIndex);
+	void UpdateMarkSelection(int32 InOldMarkIndex, FFrameNumber InMarkFrameNumber);
 	void AddMarkAtFrame(FFrameNumber FrameNumber);
 	void DeleteMarkAtIndex(int32 InMarkIndex);
 	void DeleteAllMarks();
@@ -301,6 +321,13 @@ private:
 	
 	/** Brush for drawing a downwards facing scrub handle */
 	const FSlateBrush* FrameBlockScrubHandleDownBrush, *VanillaScrubHandleDownBrush;
+
+	/** Font measure service */
+	TSharedPtr<FSlateFontMeasure> FontMeasureService;
+
+	/** Font info for the marked frames labels */
+	FSlateFontInfo SmallLayoutFont;
+	FSlateFontInfo SmallBoldLayoutFont;
 	
 	/** Total mouse delta during dragging **/
 	float DistanceDragged;
@@ -341,9 +368,18 @@ private:
 	/** Range stack */
 	TArray<TRange<double>> ViewRangeStack;
 
+	/** Index of mark being hovered */
+	int32 HoverMarkIndex;
+
 	/** Index of mark being edited */
 	int32 DragMarkIndex;
 
+	/** Marked frame time when the mouse is first pressed down */
+	FFrameNumber DragMarkReferenceFrameNumber;
+
+	/** Time of the mark being edited */
+	FFrameNumber DragMarkCurrentFrameNumber;
+
 	/** When > 0, we should not show context menus */
 	int32 ContextMenuSuppression;
 	
diff --git a/Engine/Source/Editor/SequencerWidgets/Private/SSequencerTimeSlider.cpp b/Engine/Source/Editor/SequencerWidgets/Private/SSequencerTimeSlider.cpp
--- a/Engine/Source/Editor/SequencerWidgets/Private/SSequencerTimeSlider.cpp
+++ b/Engine/Source/Editor/SequencerWidgets/Private/SSequencerTimeSlider.cpp
@@ -28,18 +28,23 @@ int32 SSequencerTimeSlider::OnPaint( const FPaintArgs& Args, const FGeometry& Al
 
 FReply SSequencerTimeSlider::OnMouseButtonDown( const FGeometry& MyGeometry, const FPointerEvent& MouseEvent )
 {
-	TimeSliderController->OnMouseButtonDown( *this, MyGeometry, MouseEvent );
+	TimeSliderController->OnTimeSliderMouseButtonDown( *this, MyGeometry, MouseEvent );
 	return FReply::Handled().CaptureMouse(AsShared()).PreventThrottling();
 }
 
 FReply SSequencerTimeSlider::OnMouseButtonUp( const FGeometry& MyGeometry, const FPointerEvent& MouseEvent )
 {
-	return TimeSliderController->OnMouseButtonUp( *this,  MyGeometry, MouseEvent );
+	return TimeSliderController->OnTimeSliderMouseButtonUp( *this,  MyGeometry, MouseEvent );
 }
 
 FReply SSequencerTimeSlider::OnMouseMove( const FGeometry& MyGeometry, const FPointerEvent& MouseEvent )
 {
-	return TimeSliderController->OnMouseMove( *this, MyGeometry, MouseEvent );
+	return TimeSliderController->OnTimeSliderMouseMove( *this, MyGeometry, MouseEvent );
+}
+
+FReply SSequencerTimeSlider::OnMouseButtonDoubleClick(const FGeometry& MyGeometry, const FPointerEvent& MouseEvent)
+{
+	return TimeSliderController->OnMouseButtonDoubleClick( SharedThis(this), MyGeometry, MouseEvent );
 }
 
 FVector2D SSequencerTimeSlider::ComputeDesiredSize( float ) const
@@ -49,7 +54,7 @@ FVector2D SSequencerTimeSlider::ComputeDesiredSize( float ) const
 
 FReply SSequencerTimeSlider::OnMouseWheel( const FGeometry& MyGeometry, const FPointerEvent& MouseEvent )
 {
-	return TimeSliderController->OnMouseWheel( *this, MyGeometry, MouseEvent );
+	return TimeSliderController->OnTimeSliderMouseWheel( *this, MyGeometry, MouseEvent );
 }
 
 FCursorReply SSequencerTimeSlider::OnCursorQuery(const FGeometry& MyGeometry, const FPointerEvent& CursorEvent) const
diff --git a/Engine/Source/Editor/SequencerWidgets/Private/SSequencerTimeSlider.h b/Engine/Source/Editor/SequencerWidgets/Private/SSequencerTimeSlider.h
--- a/Engine/Source/Editor/SequencerWidgets/Private/SSequencerTimeSlider.h
+++ b/Engine/Source/Editor/SequencerWidgets/Private/SSequencerTimeSlider.h
@@ -42,6 +42,7 @@ protected:
 	virtual int32 OnPaint( const FPaintArgs& Args, const FGeometry& AllottedGeometry, const FSlateRect& MyCullingRect, FSlateWindowElementList& OutDrawElements, int32 LayerId, const FWidgetStyle& InWidgetStyle, bool bParentEnabled ) const override;
 	virtual FReply OnMouseButtonDown( const FGeometry& MyGeometry, const FPointerEvent& MouseEvent ) override;
 	virtual FReply OnMouseButtonUp( const FGeometry& MyGeometry, const FPointerEvent& MouseEvent ) override;
+	virtual FReply OnMouseButtonDoubleClick(const FGeometry& InMyGeometry, const FPointerEvent& InMouseEvent) override;
 	virtual FReply OnMouseMove( const FGeometry& MyGeometry, const FPointerEvent& MouseEvent ) override;
 	virtual FReply OnMouseWheel( const FGeometry& MyGeometry, const FPointerEvent& MouseEvent ) override;
 	virtual FCursorReply OnCursorQuery(const FGeometry& MyGeometry, const FPointerEvent& CursorEvent) const override;
diff --git a/Engine/Source/Editor/SequencerWidgets/Public/ITimeSlider.h b/Engine/Source/Editor/SequencerWidgets/Public/ITimeSlider.h
--- a/Engine/Source/Editor/SequencerWidgets/Public/ITimeSlider.h
+++ b/Engine/Source/Editor/SequencerWidgets/Public/ITimeSlider.h
@@ -259,7 +259,28 @@ public:
 
 	virtual int32 OnPaintViewArea( const FGeometry& AllottedGeometry, const FSlateRect& MyCullingRect, FSlateWindowElementList& OutDrawElements, int32 LayerId, bool bEnabled, const FPaintViewAreaArgs& Args ) const = 0;
 
-	virtual FCursorReply OnCursorQuery( TSharedRef<const SWidget> WidgetOwner, const FGeometry& MyGeometry, const FPointerEvent& CursorEvent ) const = 0;
+	virtual FReply OnMouseButtonDoubleClick( TSharedRef<const SWidget> WidgetOwner, const FGeometry& MyGeometry, const FPointerEvent& MouseEvent ) { return FReply::Unhandled(); }
+	virtual FCursorReply OnCursorQuery( TSharedRef<const SWidget> WidgetOwner, const FGeometry& MyGeometry, const FPointerEvent& CursorEvent ) const { return FCursorReply::Unhandled(); }
+
+	virtual FReply OnTimeSliderMouseButtonDown(SWidget& OwnerWidget, const FGeometry& MyGeometry, const FPointerEvent& MouseEvent)
+	{
+		return OnMouseButtonDown(OwnerWidget, MyGeometry, MouseEvent);
+	}
+
+	virtual FReply OnTimeSliderMouseButtonUp(SWidget& OwnerWidget, const FGeometry& MyGeometry, const FPointerEvent& MouseEvent)
+	{
+		return OnMouseButtonUp(OwnerWidget, MyGeometry, MouseEvent);
+	}
+
+	virtual FReply OnTimeSliderMouseMove(SWidget& OwnerWidget, const FGeometry& MyGeometry, const FPointerEvent& MouseEvent)
+	{
+		return OnMouseMove(OwnerWidget, MyGeometry, MouseEvent);
+	}
+
+	virtual FReply OnTimeSliderMouseWheel(SWidget& OwnerWidget, const FGeometry& MyGeometry, const FPointerEvent& MouseEvent)
+	{
+		return OnMouseWheel(OwnerWidget, MyGeometry, MouseEvent);
+	}
 
 	/** Get the current play rate for this controller */
 	virtual FFrameRate GetDisplayRate() const = 0;
diff --git a/Engine/Source/Runtime/MovieScene/Public/MovieScene.h b/Engine/Source/Runtime/MovieScene/Public/MovieScene.h
--- a/Engine/Source/Runtime/MovieScene/Public/MovieScene.h
+++ b/Engine/Source/Runtime/MovieScene/Public/MovieScene.h
@@ -80,6 +80,7 @@ struct FMovieSceneMarkedFrame
 	FMovieSceneMarkedFrame()
 		: Label(FString())
 #if WITH_EDITORONLY_DATA
+		, Comment(FString())
 		, Color(0.f, 1.f, 1.f, 0.4f)
 #endif
 		, bIsDeterminismFence(false)
@@ -89,6 +90,7 @@ struct FMovieSceneMarkedFrame
 		: FrameNumber(InFrameNumber)
 		, Label(FString())
 #if WITH_EDITORONLY_DATA
+		, Comment(FString())
 		, Color(0.f, 1.f, 1.f, 0.4f)
 #endif
 		, bIsDeterminismFence(false)
@@ -101,6 +103,9 @@ struct FMovieSceneMarkedFrame
 	FString Label;
 
 #if WITH_EDITORONLY_DATA
+	UPROPERTY(EditAnywhere, Category = "Marked Frame")
+	FString Comment;
+
 	UPROPERTY(EditAnywhere, Category = "Marked Frame")
 	FLinearColor Color;
 #endif
