Sequencer: Change key visitor so that it explicitly goes through the models.

This fixes issues where keys in collapsed tracks would be erroneously selected.

diff --git a/Engine/Source/Editor/Sequencer/Private/Tools/SequencerEntityVisitor.cpp b/Engine/Source/Editor/Sequencer/Private/Tools/SequencerEntityVisitor.cpp
--- a/Engine/Source/Editor/Sequencer/Private/Tools/SequencerEntityVisitor.cpp
+++ b/Engine/Source/Editor/Sequencer/Private/Tools/SequencerEntityVisitor.cpp
@@ -9,6 +9,7 @@
 #include "MVVM/Extensions/IOutlinerExtension.h"
 #include "MVVM/Extensions/ITrackAreaExtension.h"
 #include "MVVM/Extensions/ITrackLaneExtension.h"
+#include "MVVM/ViewModels/CategoryModel.h"
 #include "MVVM/ViewModels/ChannelModel.h"
 #include "MVVM/ViewModels/TrackModel.h"
 #include "MVVM/ViewModels/ViewModel.h"
@@ -157,41 +158,63 @@ void FSequencerEntityWalker::VisitAnyChannels(const ISequencerEntityVisitor& Vis
 	FChannelGroupModel* ChannelGroup = nullptr;
 	if (FChannelGroupModel* ChannelGroupNode = InNode->CastThis<FChannelGroupModel>())
 	{
-		ChannelGroup = ChannelGroupNode;
-	}
-	if (FTrackModel* TrackModel = InNode->CastThis<FTrackModel>())
-	{
-		// TODO: Top-level channel group -- there used to be only one at most but this assumption isn't necessarily true anymore
-		for (const TViewModelPtr<FChannelGroupModel>& TopLevelChannel : TrackModel->GetTopLevelChannels().IterateSubList<FChannelGroupModel>())
+		const IOutlinerExtension* OutlinerItem = ChannelGroupNode->CastThis<IOutlinerExtension>();
+		if (!OutlinerItem || OutlinerItem->IsFilteredOut() == false)
 		{
-			ChannelGroup = TopLevelChannel.Get();
-			break;
+			for (const TWeakViewModelPtr<FChannelModel>& WeakChannel : ChannelGroupNode->GetChannels())
+			{
+				if (TViewModelPtr<FChannelModel> Channel = WeakChannel.Pin())
+				{
+					VisitChannel(Visitor, Channel);
+				}
+			}
 		}
 	}
-
-	if (ChannelGroup)
+	else if (FCategoryGroupModel* CategoryGroupNode = InNode->CastThis<FCategoryGroupModel>())
 	{
-		const IOutlinerExtension* OutlinerItem = ChannelGroup->CastThis<IOutlinerExtension>();
-		if (!OutlinerItem || OutlinerItem->IsFilteredOut() == false)
+		if (!CategoryGroupNode->IsExpanded())
 		{
-			for (const TWeakViewModelPtr<FChannelModel>& WeakChannel : ChannelGroup->GetChannels())
+			for (const TWeakViewModelPtr<FCategoryModel>& WeakCategory : CategoryGroupNode->GetCategories())
 			{
-				if (TViewModelPtr<FChannelModel> Channel = WeakChannel.Pin())
+				if (TViewModelPtr<FCategoryModel> Category = WeakCategory.Pin())
 				{
-					VisitChannel(Visitor, Channel);
+					TSharedPtr<FViewModel> Model = Category->GetLinkedOutlinerItem().AsModel();
+					for (TSharedPtr<FChannelGroupModel> DescandantChannelGroup : Model->GetDescendantsOfType<FChannelGroupModel>())
+					{
+						const IOutlinerExtension* OutlinerItem = DescandantChannelGroup->CastThis<IOutlinerExtension>();
+						if (!OutlinerItem || OutlinerItem->IsFilteredOut() == false)
+						{
+							for (const TWeakViewModelPtr<FChannelModel>& WeakChannel : DescandantChannelGroup->GetChannels())
+							{
+								if (TViewModelPtr<FChannelModel> Channel = WeakChannel.Pin())
+								{
+									VisitChannel(Visitor, Channel);
+								}
+							}
+						}
+					}
 				}
 			}
 		}
 	}
-	// Otherwise it might be a collapsed node that contains key areas as children. If so we visit them as if they were a part of this track so that key groupings are visited properly.
-	else if (bAnyParentCollapsed)
+	else if (FTrackModel* TrackModel = InNode->CastThis<FTrackModel>())
 	{
-		const IOutlinerExtension* OutlinerItem = InNode->CastThis<IOutlinerExtension>();
-		const bool bIsExpanded = !OutlinerItem || OutlinerItem->IsExpanded();
-
-		for (TSharedPtr<FViewModel> ChildNode : InNode->GetChildren(EViewModelListType::Outliner))
+		if (!TrackModel->IsExpanded())
 		{
-			VisitAnyChannels(Visitor, ChildNode.ToSharedRef(), bAnyParentCollapsed || !bIsExpanded);
+			for (TSharedPtr<FChannelGroupModel> DescandantChannelGroup : TrackModel->GetDescendantsOfType<FChannelGroupModel>())
+			{
+				const IOutlinerExtension* OutlinerItem = DescandantChannelGroup->CastThis<IOutlinerExtension>();
+				if (!OutlinerItem || OutlinerItem->IsFilteredOut() == false)
+				{
+					for (const TWeakViewModelPtr<FChannelModel>& WeakChannel : DescandantChannelGroup->GetChannels())
+					{
+						if (TViewModelPtr<FChannelModel> Channel = WeakChannel.Pin())
+						{
+							VisitChannel(Visitor, Channel);
+						}
+					}
+				}
+			}
 		}
 	}
 }
