Sequencer: Make Time Change undoable when setting key.

diff --git a/Engine/Source/Editor/Sequencer/Private/Sequencer.cpp b/Engine/Source/Editor/Sequencer/Private/Sequencer.cpp
--- a/Engine/Source/Editor/Sequencer/Private/Sequencer.cpp
+++ b/Engine/Source/Editor/Sequencer/Private/Sequencer.cpp
@@ -128,6 +128,7 @@
 #include "ISequencerChannelInterface.h"
 #include "IMovieRendererInterface.h"
 #include "SequencerKeyCollection.h"
+#include "SequencerTimeChangeUndoRedoProxy.h"
 #include "CurveEditor.h"
 #include "CurveEditorScreenSpace.h"
 #include "CurveDataAbstraction.h"
@@ -522,7 +523,12 @@ void FSequencer::InitSequencer(const FSequencerInitParams& InitParams, const TSh
 	OverlayCurve = OverlayAnimation.AddCurve(0.f, 0.2f, ECurveEaseFunction::QuadIn);
 	RecordingAnimation = FCurveSequence();
 	RecordingAnimation.AddCurve(0.f, 1.5f, ECurveEaseFunction::Linear);
-
+	
+	{
+		UndoRedoProxy = NewObject<USequencerTimeChangeUndoRedoProxy>(GetTransientPackage(), NAME_None);
+		UndoRedoProxy->SetSequencer(SharedThis(this));
+		UndoRedoProxy->SetFlags(RF_Transactional);
+	}
 	// Update initial movie scene data
 	NotifyMovieSceneDataChanged( EMovieSceneDataChangeType::ActiveMovieSceneChanged );
 
@@ -4248,7 +4254,7 @@ void FSequencer::AddReferencedObjects( FReferenceCollector& Collector )
 {
 	Collector.AddReferencedObject( CompiledDataManager );
 	Collector.AddReferencedObject( Settings );
-
+	Collector.AddReferencedObject( UndoRedoProxy );
 	if (UMovieSceneSequence* RootSequencePtr = RootSequence.Get())
 	{
 		Collector.AddReferencedObject( RootSequencePtr );
diff --git a/Engine/Source/Editor/Sequencer/Private/Sequencer.h b/Engine/Source/Editor/Sequencer/Private/Sequencer.h
--- a/Engine/Source/Editor/Sequencer/Private/Sequencer.h
+++ b/Engine/Source/Editor/Sequencer/Private/Sequencer.h
@@ -69,6 +69,7 @@ struct FMovieSceneSequencePlaybackParams;
 struct FMovieScenePossessable;
 struct FTransformData;
 struct FKeyAttributes;
+class USequencerTimeChangeUndoRedoProxy;
 struct FNotificationInfo;
 struct FEditorViewportViewModifierParams;
 
@@ -667,6 +668,10 @@ public:
 	virtual void AddReferencedObjects(FReferenceCollector& Collector) override;
 	virtual FString GetReferencerName() const override;
 
+private:
+	//gco object for undo redo support for changing time
+	USequencerTimeChangeUndoRedoProxy* UndoRedoProxy;
+
 public:
 
 	//~ FTickableEditorObject Interface
diff --git a/Engine/Source/Editor/Sequencer/Private/SequencerTimeChangeUndoRedoProxy.cpp b/Engine/Source/Editor/Sequencer/Private/SequencerTimeChangeUndoRedoProxy.cpp
new file mode 100644
--- /dev/null
+++ b/Engine/Source/Editor/Sequencer/Private/SequencerTimeChangeUndoRedoProxy.cpp
@@ -0,0 +1,67 @@
+// Copyright Epic Games, Inc. All Rights Reserved.
+
+#include "SequencerTimeChangeUndoRedoProxy.h"
+#include "Sequencer.h"
+#include "ScopedTransaction.h"
+#include "CoreGlobals.h"
+
+#define LOCTEXT_NAMESPACE "Sequencer"
+
+USequencerTimeChangeUndoRedoProxy::~USequencerTimeChangeUndoRedoProxy()
+{
+	if (OnActivateSequenceChangedHandle.IsValid() && WeakSequencer.IsValid())
+	{
+		WeakSequencer.Pin()->OnActivateSequence().Remove(OnActivateSequenceChangedHandle);
+	}
+}
+
+void USequencerTimeChangeUndoRedoProxy::SetSequencer(TSharedRef<FSequencer> InSequencer)
+{
+	WeakSequencer = InSequencer;
+	OnActivateSequenceChangedHandle = InSequencer->OnActivateSequence().AddUObject(this, &USequencerTimeChangeUndoRedoProxy::OnActivateSequenceChanged);
+	OnActivateSequenceChanged(InSequencer->GetFocusedTemplateID());
+}
+
+void USequencerTimeChangeUndoRedoProxy::OnActivateSequenceChanged(FMovieSceneSequenceIDRef ID)
+{
+	if (TSharedPtr<ISequencer> SequencerPtr = WeakSequencer.Pin())
+	{
+		if (SequencerPtr->GetFocusedMovieSceneSequence() && SequencerPtr->GetFocusedMovieSceneSequence()->GetMovieScene())
+		{
+			using namespace UE::MovieScene;
+			bTimeWasSet = false;
+			MovieSceneModified.Unlink();
+			SequencerPtr->GetFocusedMovieSceneSequence()->GetMovieScene()->UMovieSceneSignedObject::EventHandlers.Link(MovieSceneModified, this);
+		}
+	}
+}
+
+void USequencerTimeChangeUndoRedoProxy::OnModifiedIndirectly(UMovieSceneSignedObject* Object)
+{
+	if (Object->IsA<UMovieSceneSection>() && WeakSequencer.IsValid())
+	{
+		FQualifiedFrameTime InTime = WeakSequencer.Pin()->GetGlobalTime();
+
+		if (bTimeWasSet)
+		{
+			if (!GIsTransacting && (InTime.Time != Time.Time || InTime.Rate != Time.Rate))
+			{
+				const FScopedTransaction Transaction(LOCTEXT("TimeChanged", "Time Changed"));
+				Modify();
+			}
+		}
+		bTimeWasSet = true;
+		Time = InTime;
+	}
+}
+
+void USequencerTimeChangeUndoRedoProxy::PostEditUndo()
+{
+	if (WeakSequencer.IsValid())
+	{
+		WeakSequencer.Pin()->SetGlobalTime(Time.Time, true);
+	}
+}
+
+#undef LOCTEXT_NAMESPACE
+
diff --git a/Engine/Source/Editor/Sequencer/Private/SequencerTimeChangeUndoRedoProxy.h b/Engine/Source/Editor/Sequencer/Private/SequencerTimeChangeUndoRedoProxy.h
new file mode 100644
--- /dev/null
+++ b/Engine/Source/Editor/Sequencer/Private/SequencerTimeChangeUndoRedoProxy.h
@@ -0,0 +1,45 @@
+// Copyright Epic Games, Inc. All Rights Reserved.
+
+#pragma once
+
+#include "CoreMinimal.h"
+#include "CoreTypes.h"
+#include "UObject/ObjectMacros.h"
+#include "EventHandlers/ISignedObjectEventHandler.h"
+#include "Delegates/IDelegateInstance.h"
+#include "Misc/QualifiedFrameTime.h"
+#include "SequencerTimeChangeUndoRedoProxy.generated.h"
+
+class FSequencer;
+
+UCLASS()
+class  USequencerTimeChangeUndoRedoProxy : public UObject, public UE::MovieScene::ISignedObjectEventHandler
+{
+public:
+	GENERATED_BODY()
+	USequencerTimeChangeUndoRedoProxy() :bTimeWasSet(false), WeakSequencer(nullptr) {  };
+	~USequencerTimeChangeUndoRedoProxy();
+
+	void SetSequencer(TSharedRef<FSequencer> InSequencer);
+	void OnActivateSequenceChanged(FMovieSceneSequenceIDRef ID);
+
+	/*~ UObject */
+	virtual void PostEditUndo() override;
+
+	/*~ ISignedObjectEventHandler Interface */
+	virtual void  OnModifiedIndirectly(UMovieSceneSignedObject*) override;
+
+private:
+	UPROPERTY(Transient)
+	FQualifiedFrameTime Time;
+	
+	//no TOptional UPROPERTY so use this instead
+	UPROPERTY(Transient)
+	bool bTimeWasSet = false;
+	
+	FDelegateHandle OnActivateSequenceChangedHandle;
+	TWeakPtr<FSequencer> WeakSequencer;
+
+	UE::MovieScene::TNonIntrusiveEventHandler<UE::MovieScene::ISignedObjectEventHandler> MovieSceneModified;
+};
+
