Fixed crash in sequencer due to modification of the array while it's iterated on

diff --git a/Engine/Source/Editor/Sequencer/Private/SequencerModule.cpp b/Engine/Source/Editor/Sequencer/Private/SequencerModule.cpp
--- a/Engine/Source/Editor/Sequencer/Private/SequencerModule.cpp
+++ b/Engine/Source/Editor/Sequencer/Private/SequencerModule.cpp
@@ -112,14 +112,16 @@ struct FDeferredSignedObjectChangeHandler : UE::MovieScene::IDeferredSignedObjec
 
 	void Flush() override
 	{
-		for (TWeakObjectPtr<UMovieSceneSignedObject> WeakObject : SignedObjects)
+		TSet<TWeakObjectPtr<UMovieSceneSignedObject>> SignedObjectsTmp = SignedObjects;
+		SignedObjects.Empty();
+		// we operate on a copy of the signed objects in case the delegates would modify the array
+		for (TWeakObjectPtr<UMovieSceneSignedObject> WeakObject : SignedObjectsTmp)
 		{
 			if (UMovieSceneSignedObject* Object = WeakObject.Get())
 			{
 				Object->BroadcastChanged();
 			}
 		}
-		SignedObjects.Empty();
 	}
 
 	void DeferMarkAsChanged(UMovieSceneSignedObject* SignedObject) override
