Sequencer: Do not allow autokey changes if sequencer is ReadOnly

diff --git a/Engine/Source/Editor/Sequencer/Public/ISequencer.h b/Engine/Source/Editor/Sequencer/Public/ISequencer.h
--- a/Engine/Source/Editor/Sequencer/Public/ISequencer.h
+++ b/Engine/Source/Editor/Sequencer/Public/ISequencer.h
@@ -376,7 +376,7 @@ public:
 	/** @return Returns whether sequencer will respond to changes and possibly create a key or track */
 	virtual bool IsAllowedToChange() const 
 	{
-		if (GetAllowEditsMode() == EAllowEditsMode::AllowLevelEditsOnly)
+		if (IsReadOnly() || GetAllowEditsMode() == EAllowEditsMode::AllowLevelEditsOnly)
 		{
 			return false;
 		}
