Sequencer: Fix track row index 0 muting all track rows

diff --git a/Engine/Source/Editor/Sequencer/Private/SSequencer.cpp b/Engine/Source/Editor/Sequencer/Private/SSequencer.cpp
--- a/Engine/Source/Editor/Sequencer/Private/SSequencer.cpp
+++ b/Engine/Source/Editor/Sequencer/Private/SSequencer.cpp
@@ -10,6 +10,7 @@
 #include "MVVM/Extensions/IRenameableExtension.h"
 #include "MVVM/ViewModels/SequencerEditorViewModel.h"
 #include "MVVM/ViewModels/SequencerOutlinerViewModel.h"
+#include "MVVM/ViewModels/TrackRowModel.h"
 #include "Sections/MovieSceneSubSection.h"
 #include "MovieScene.h"
 #include "Framework/MultiBox/MultiBoxDefs.h"
@@ -2864,7 +2865,9 @@ void SSequencer::UpdateLayoutTree()
 				{
 					UMovieSceneTrack* Track = TrackNode->GetTrack();
 					bool bDisableEval = NodeTree->IsNodeMute(OutlinerItem) || (bHasSoloNodes && !NodeTree->IsNodeSolo(OutlinerItem));
-					if (TrackNode->GetRowIndex() > 0)
+
+					TViewModelPtr<FTrackRowModel> TrackRowModel = OutlinerItem.ImplicitCast();
+					if (TrackRowModel)
 					{
 						if (bDisableEval != Track->IsRowEvalDisabled(TrackNode->GetRowIndex()))
 						{
