Sequencer: Fix restriction on paste bindings which was looking for a valid world. For UMG, this is not necessary

diff --git a/Engine/Source/Editor/Sequencer/Private/SequencerUtilities.cpp b/Engine/Source/Editor/Sequencer/Private/SequencerUtilities.cpp
--- a/Engine/Source/Editor/Sequencer/Private/SequencerUtilities.cpp
+++ b/Engine/Source/Editor/Sequencer/Private/SequencerUtilities.cpp
@@ -2085,10 +2085,6 @@ bool FSequencerUtilities::PasteBindings(const FString& TextToImport, TSharedRef<
 	}
 
 	UWorld* World = Cast<UWorld>(Sequencer->GetPlaybackContext());
-	if (!World)
-	{
-		return false;
-	}
 
 	const FScopedTransaction Transaction(LOCTEXT("PasteBindings", "Paste Bindings"));
 
@@ -2160,19 +2156,22 @@ bool FSequencerUtilities::PasteBindings(const FString& TextToImport, TSharedRef<
 
 				// Find the actors that this pasted binding should bind to
 				TArray<AActor*> ActorsToRebind;
-				for (TActorIterator<AActor> ActorItr(World); ActorItr; ++ActorItr)
+				if (World)
 				{
-					AActor* Actor = *ActorItr;
-					if (Actor && CopyableBinding->BoundObjectNames.Contains(Actor->GetName()))
+					for (TActorIterator<AActor> ActorItr(World); ActorItr; ++ActorItr)
 					{
-						// If this actor is already bound and we're not duplicating actors, don't bind to anything
-						if (!PasteBindingsParams.bDuplicateExistingActors && Sequencer->FindObjectId(*Actor, Sequencer->GetFocusedTemplateID()).IsValid())
+						AActor* Actor = *ActorItr;
+						if (Actor && CopyableBinding->BoundObjectNames.Contains(Actor->GetName()))
 						{
-							continue;
-						}
+							// If this actor is already bound and we're not duplicating actors, don't bind to anything
+							if (!PasteBindingsParams.bDuplicateExistingActors && Sequencer->FindObjectId(*Actor, Sequencer->GetFocusedTemplateID()).IsValid())
+							{
+								continue;
+							}
 
-						ActorsToRebind.Add(Actor);
-						CopyableBinding->BoundObjectNames.Remove(Actor->GetName());
+							ActorsToRebind.Add(Actor);
+							CopyableBinding->BoundObjectNames.Remove(Actor->GetName());
+						}
 					}
 				}
 
