Sequencer: Fix tick frame font size

diff --git a/Engine/Source/Editor/Sequencer/Private/SequencerTimeSliderController.cpp b/Engine/Source/Editor/Sequencer/Private/SequencerTimeSliderController.cpp
--- a/Engine/Source/Editor/Sequencer/Private/SequencerTimeSliderController.cpp
+++ b/Engine/Source/Editor/Sequencer/Private/SequencerTimeSliderController.cpp
@@ -248,6 +248,7 @@ void FSequencerTimeSliderController::DrawTicks( FSlateWindowElementList& OutDraw
 	FFrameRate     TickResolution  = GetTickResolution();
 	FFrameRate     DisplayRate     = GetDisplayRate();
 	FPaintGeometry PaintGeometry   = InArgs.AllottedGeometry.ToPaintGeometry();
+	FSlateFontInfo TickFrameFont = FCoreStyle::GetDefaultFontStyle("Regular", 8);
 
 	double MajorGridStep  = 0.0;
 	int32  MinorDivisions = 0;
@@ -300,7 +301,7 @@ void FSequencerTimeSliderController::DrawTicks( FSlateWindowElementList& OutDraw
 				InArgs.StartLayer+1, 
 				InArgs.AllottedGeometry.ToPaintGeometry( InArgs.AllottedGeometry.Size, FSlateLayoutTransform(TextOffset) ), 
 				FrameString, 
-				SmallLayoutFont,
+				TickFrameFont,
 				InArgs.DrawEffects,
 				InArgs.TickColor*0.65f 
 			);
