Sequencer: Track filter saved in preferences

Control Rig Controls filter is no longer enabled by default. But if the user enables it, it will be saved between sessions

diff --git a/Engine/Plugins/Animation/ControlRig/Source/ControlRigEditor/Private/ControlRigBlueprintActions.cpp b/Engine/Plugins/Animation/ControlRig/Source/ControlRigEditor/Private/ControlRigBlueprintActions.cpp
--- a/Engine/Plugins/Animation/ControlRig/Source/ControlRigEditor/Private/ControlRigBlueprintActions.cpp
+++ b/Engine/Plugins/Animation/ControlRig/Source/ControlRigEditor/Private/ControlRigBlueprintActions.cpp
@@ -424,8 +424,6 @@ void FControlRigBlueprintActions::OnSpawnedSkeletalMeshActorChanged(UObject* InO
 				WeakSequencer.Pin()->ObjectImplicitlyAdded(ControlRig);
 			}
 			
-			FText Name = LOCTEXT("SequenceTrackFilter_ControlRigControls", "Control Rig Controls");
-			WeakSequencer.Pin()->SetTrackFilterEnabled(Name, true);
 			WeakSequencer.Pin()->NotifyMovieSceneDataChanged(EMovieSceneDataChangeType::MovieSceneStructureItemAdded);
 			if (!ControlRigEditMode)
 			{
diff --git a/Engine/Source/Editor/Sequencer/Private/SSequencer.cpp b/Engine/Source/Editor/Sequencer/Private/SSequencer.cpp
--- a/Engine/Source/Editor/Sequencer/Private/SSequencer.cpp
+++ b/Engine/Source/Editor/Sequencer/Private/SSequencer.cpp
@@ -1056,6 +1056,12 @@ void SSequencer::InitializeTrackFilters()
 		}
 	}
 
+	for (TSharedRef<FSequencerTrackFilter> TrackFilter : AllTrackFilters)
+	{
+		const bool bEnabled = GetSequencerSettings()->IsTrackFilterEnabled(TrackFilter->GetDisplayName().ToString());
+		SetTrackFilterEnabled(TrackFilter->GetDisplayName(), bEnabled);
+	}
+
 	// Sort by display name
 	AllTrackFilters.Sort([](const TSharedRef<FSequencerTrackFilter>& LHS, const TSharedRef<FSequencerTrackFilter>& RHS) { return LHS->GetDisplayName().ToString() < RHS->GetDisplayName().ToString(); });
 }
diff --git a/Engine/Source/Editor/Sequencer/Private/SequencerNodeTree.cpp b/Engine/Source/Editor/Sequencer/Private/SequencerNodeTree.cpp
--- a/Engine/Source/Editor/Sequencer/Private/SequencerNodeTree.cpp
+++ b/Engine/Source/Editor/Sequencer/Private/SequencerNodeTree.cpp
@@ -271,16 +271,25 @@ void FSequencerNodeTree::SortAllNodesAndDescendants()
 
 void FSequencerNodeTree::AddFilter(TSharedPtr<FSequencerTrackFilter> TrackFilter)
 {
+	GetSequencer().GetSequencerSettings()->SetTrackFilterEnabled(TrackFilter->GetDisplayName().ToString(), true);
+
 	TrackFilters->Add(TrackFilter);
 }
 
 int32 FSequencerNodeTree::RemoveFilter(TSharedPtr<FSequencerTrackFilter> TrackFilter)
 {
+	GetSequencer().GetSequencerSettings()->SetTrackFilterEnabled(TrackFilter->GetDisplayName().ToString(), false);
+
 	return TrackFilters->Remove(TrackFilter);
 }
 
 void FSequencerNodeTree::RemoveAllFilters()
 {
+	for (TSharedPtr< FSequencerTrackFilter > TrackFilter : *TrackFilters)
+	{
+		GetSequencer().GetSequencerSettings()->SetTrackFilterEnabled(TrackFilter->GetDisplayName().ToString(), false);
+	}
+
 	TrackFilters->RemoveAll();
 	TrackFilterLevelFilter->ResetFilter();
 }
diff --git a/Engine/Source/Editor/Sequencer/Private/SequencerSettings.cpp b/Engine/Source/Editor/Sequencer/Private/SequencerSettings.cpp
--- a/Engine/Source/Editor/Sequencer/Private/SequencerSettings.cpp
+++ b/Engine/Source/Editor/Sequencer/Private/SequencerSettings.cpp
@@ -961,4 +961,29 @@ void USequencerSettings::SetTreeViewWidth(float InTreeViewWidth)
 		TreeViewWidth = InTreeViewWidth;
 		SaveConfig();
 	}
-}
\ No newline at end of file
+}
+
+bool USequencerSettings::IsTrackFilterEnabled(const FString& TrackFilter) const
+{
+	return TrackFilters.Contains(TrackFilter);
+}
+
+void USequencerSettings::SetTrackFilterEnabled(const FString & TrackFilter, bool bEnabled)
+{
+	if (bEnabled)
+	{
+		if (!TrackFilters.Contains(TrackFilter))
+		{
+			TrackFilters.Add(TrackFilter);
+			SaveConfig();
+		}
+	}
+	else
+	{
+		if (TrackFilters.Contains(TrackFilter))
+		{
+			TrackFilters.Remove(TrackFilter);
+			SaveConfig();
+		}
+	}
+}
diff --git a/Engine/Source/Editor/Sequencer/Public/SequencerSettings.h b/Engine/Source/Editor/Sequencer/Public/SequencerSettings.h
--- a/Engine/Source/Editor/Sequencer/Public/SequencerSettings.h
+++ b/Engine/Source/Editor/Sequencer/Public/SequencerSettings.h
@@ -412,6 +412,11 @@ public:
 	/** Sets the tree view width percentage */
 	void SetTreeViewWidth(float InTreeViewWidth);
 
+	/** Gets whether the given track filter is enabled */
+	bool IsTrackFilterEnabled(const FString& TrackFilter) const;
+	/** Sets whether the track filter should be enabled/disabled */
+	void SetTrackFilterEnabled(const FString& TrackFilter, bool bEnabled);
+
 protected:
 
 	/** The auto change mode (auto-key, auto-track or none). */
@@ -660,6 +665,10 @@ protected:
 	UPROPERTY(config, EditAnywhere, Category = General)
 	float TreeViewWidth;
 
+	/** The track filters that are enabled */
+	UPROPERTY(config, EditAnywhere, Category = General)
+	TArray<FString> TrackFilters;
+
 	FOnEvaluateSubSequencesInIsolationChanged OnEvaluateSubSequencesInIsolationChangedEvent;
 	FOnShowSelectedNodesOnlyChanged OnShowSelectedNodesOnlyChangedEvent;
 	FOnAllowEditsModeChanged OnAllowEditsModeChangedEvent;
