Sequencer: Use Property->GetOwnerStruct to find the property container. This fixes getting the correct container for determining property permissions when calling FPropertyEditorPermissionList::DoesPropertyPassFilter where the owner struct is necessary for property structs, ie. FCameraFocusSettings for ManualFocusDistance.

diff --git a/Engine/Source/Editor/Sequencer/Private/KeyPropertyParams.cpp b/Engine/Source/Editor/Sequencer/Private/KeyPropertyParams.cpp
--- a/Engine/Source/Editor/Sequencer/Private/KeyPropertyParams.cpp
+++ b/Engine/Source/Editor/Sequencer/Private/KeyPropertyParams.cpp
@@ -70,7 +70,7 @@ const UStruct* FCanKeyPropertyParams::FindPropertyContainer(const FProperty* For
 			bFoundProperty = Property == ForProperty;
 			if (bFoundProperty)
 			{
-				return ObjectClass;
+				return Property->GetOwnerStruct();
 			}
 		}
 		else if (const FStructProperty* StructProperty = CastField<const FStructProperty>(Property))
