Editor: use FLinearColor for screen fades

Both sequencer and scene views use linear colors for screen fades, but level viewports use FColor, forcing conversion and clamping for no good reason when it sits between the other two. This change therefore makes level viewports use FLinearColor so we can get the full range of supported colors.

Since FLinearColor is constructible from FColor, this should be a fairly transparent change.

diff --git a/Engine/Source/Editor/Sequencer/Private/Sequencer.cpp b/Engine/Source/Editor/Sequencer/Private/Sequencer.cpp
--- a/Engine/Source/Editor/Sequencer/Private/Sequencer.cpp
+++ b/Engine/Source/Editor/Sequencer/Private/Sequencer.cpp
@@ -4134,7 +4134,7 @@ void FSequencer::SetViewportSettings(const TMap<FViewportClient*, EMovieSceneVie
 					}
 					if (ViewportParams->SetWhichViewportParam & EMovieSceneViewportParams::SVP_FadeColor)
 					{
-						LevelVC->FadeColor = ViewportParams->FadeColor.ToFColor(/*bSRGB=*/ true);
+						LevelVC->FadeColor = ViewportParams->FadeColor;
 						LevelVC->bEnableFading = true;
 					}
 					if (ViewportParams->SetWhichViewportParam & EMovieSceneViewportParams::SVP_ColorScaling)
diff --git a/Engine/Source/Editor/UnrealEd/Public/LevelEditorViewport.h b/Engine/Source/Editor/UnrealEd/Public/LevelEditorViewport.h
--- a/Engine/Source/Editor/UnrealEd/Public/LevelEditorViewport.h
+++ b/Engine/Source/Editor/UnrealEd/Public/LevelEditorViewport.h
@@ -837,7 +837,7 @@ public:
 
 	FVector					ColorScale;
 
-	FColor					FadeColor;
+	FLinearColor			FadeColor;
 
 	float					FadeAmount;
 
