Sequencer: ASAN issue do to UE::MovieScene::ISignedObjectHandler needing to not live beyond the lifeftime of the signed object it's on. So moved the undo proxy to always outered under the active movie scene versus the transient package.
Note I had issues rerunning the horde test this is coming from, the linux machine is always waiting...

diff --git a/Engine/Source/Editor/Sequencer/Private/Sequencer.cpp b/Engine/Source/Editor/Sequencer/Private/Sequencer.cpp
--- a/Engine/Source/Editor/Sequencer/Private/Sequencer.cpp
+++ b/Engine/Source/Editor/Sequencer/Private/Sequencer.cpp
@@ -525,13 +525,16 @@ void FSequencer::InitSequencer(const FSequencerInitParams& InitParams, const TSh
 	RecordingAnimation = FCurveSequence();
 	RecordingAnimation.AddCurve(0.f, 1.5f, ECurveEaseFunction::Linear);
 	
+	// Update initial movie scene data
+	NotifyMovieSceneDataChanged( EMovieSceneDataChangeType::ActiveMovieSceneChanged );
+
+	//create proxy and set it as child of focused movie scene sequence.
+	if(GetRootMovieSceneSequence())
 	{
-		UndoRedoProxy = NewObject<USequencerTimeChangeUndoRedoProxy>(GetTransientPackage(), NAME_None);
+		UndoRedoProxy = NewObject<USequencerTimeChangeUndoRedoProxy>(GetRootMovieSceneSequence(), NAME_None);
 		UndoRedoProxy->SetSequencer(SharedThis(this));
-		UndoRedoProxy->SetFlags(RF_Transactional);
+		UndoRedoProxy->SetFlags(RF_Transactional | RF_Transient);
 	}
-	// Update initial movie scene data
-	NotifyMovieSceneDataChanged( EMovieSceneDataChangeType::ActiveMovieSceneChanged );
 
 	// Update the view range to the new current time
 	UpdateTimeBoundsToFocusedMovieScene();
