Hide details panel AddKey buttons for read only sequencers.

diff --git a/Engine/Source/Editor/Sequencer/Private/LevelEditorSequencerIntegration.cpp b/Engine/Source/Editor/Sequencer/Private/LevelEditorSequencerIntegration.cpp
--- a/Engine/Source/Editor/Sequencer/Private/LevelEditorSequencerIntegration.cpp
+++ b/Engine/Source/Editor/Sequencer/Private/LevelEditorSequencerIntegration.cpp
@@ -77,7 +77,7 @@ public:
 		for (const TWeakPtr<ISequencer>& WeakSequencer : Sequencers)
 		{
 			TSharedPtr<ISequencer> Sequencer = WeakSequencer.Pin();
-			if (Sequencer.IsValid() && Sequencer->CanKeyProperty(CanKeyPropertyParams))
+			if (Sequencer.IsValid() && Sequencer->CanKeyProperty(CanKeyPropertyParams) && !Sequencer->IsReadOnly())
 			{
 				return true;
 			}
