Sequencer: Change the folder color for all selected folders and not just the single folder

diff --git a/Engine/Source/Editor/Sequencer/Private/MVVM/ViewModels/FolderModel.cpp b/Engine/Source/Editor/Sequencer/Private/MVVM/ViewModels/FolderModel.cpp
--- a/Engine/Source/Editor/Sequencer/Private/MVVM/ViewModels/FolderModel.cpp
+++ b/Engine/Source/Editor/Sequencer/Private/MVVM/ViewModels/FolderModel.cpp
@@ -39,7 +39,7 @@ namespace UE
 namespace Sequencer
 {
 
-FColor FFolderModel::InitialFolderColor;
+TMap<TWeakObjectPtr<UMovieSceneFolder>, FColor> FFolderModel::InitialFolderColors;
 bool FFolderModel::bFolderPickerWasCancelled = false;
 
 FFolderModel::FFolderModel(UMovieSceneFolder* Folder)
@@ -311,7 +311,7 @@ void FFolderModel::BuildContextMenu(FMenuBuilder& MenuBuilder)
 	{
 		MenuBuilder.AddMenuEntry(
 			LOCTEXT("SetColor", "Set Color"),
-			LOCTEXT("SetColorTooltip", "Set the folder color"),
+			LOCTEXT("SetColorTooltip", "Set the color for the selected folders"),
 			FSlateIcon(),
 			FUIAction(
 				FExecuteAction::CreateSP(this, &FFolderModel::SetFolderColor))
@@ -348,19 +348,29 @@ TSharedRef<SWidget> FFolderModel::CreateOutlinerView(const FCreateOutlinerViewPa
 
 void FFolderModel::SetFolderColor()
 {
-	UMovieSceneFolder* MovieSceneFolder = GetFolder();
-	if (!MovieSceneFolder)
+	bFolderPickerWasCancelled = false;
+	InitialFolderColors.Empty();
+
+	TSharedPtr<FSequencerEditorViewModel> EditorViewModel = GetEditor();
+	TSharedPtr<FSequencer> Sequencer = EditorViewModel->GetSequencerImpl();
+
+	TArray<UMovieSceneFolder*> MovieSceneFolders;
+	Sequencer->GetSelectedFolders(MovieSceneFolders);
+
+	if (MovieSceneFolders.Num() == 0)
 	{
 		return;
 	}
 
-	InitialFolderColor = MovieSceneFolder->GetFolderColor();
-	bFolderPickerWasCancelled = false;
+	for (UMovieSceneFolder* MovieSceneFolder : MovieSceneFolders)
+	{
+		InitialFolderColors.Add(MovieSceneFolder, MovieSceneFolder->GetFolderColor());
+	}
 
 	FColorPickerArgs PickerArgs;
 	PickerArgs.bUseAlpha = false;
 	PickerArgs.DisplayGamma = TAttribute<float>::Create( TAttribute<float>::FGetter::CreateUObject(GEngine, &UEngine::GetDisplayGamma) );
-	PickerArgs.InitialColor = InitialFolderColor.ReinterpretAsLinear();
+	PickerArgs.InitialColor = (*InitialFolderColors.CreateIterator()).Value.ReinterpretAsLinear();
 	PickerArgs.OnColorCommitted = FOnLinearColorValueChanged::CreateSP(this, &FFolderModel::OnColorPickerPicked);
 	PickerArgs.OnColorPickerWindowClosed = FOnWindowClosed::CreateSP(this, &FFolderModel::OnColorPickerClosed);
 	PickerArgs.OnColorPickerCancelled  = FOnColorPickerCancelled::CreateSP(this, &FFolderModel::OnColorPickerCancelled );
@@ -370,7 +380,13 @@ void FFolderModel::SetFolderColor()
 
 void FFolderModel::OnColorPickerPicked(FLinearColor NewFolderColor)
 {
-	if (UMovieSceneFolder* MovieSceneFolder = GetFolder())
+	TSharedPtr<FSequencerEditorViewModel> EditorViewModel = GetEditor();
+	TSharedPtr<FSequencer> Sequencer = EditorViewModel->GetSequencerImpl();
+
+	TArray<UMovieSceneFolder*> MovieSceneFolders;
+	Sequencer->GetSelectedFolders(MovieSceneFolders);
+
+	for (UMovieSceneFolder* MovieSceneFolder : MovieSceneFolders)
 	{
 		MovieSceneFolder->SetFolderColor(NewFolderColor.ToFColor(false));
 	}
@@ -378,25 +394,47 @@ void FFolderModel::OnColorPickerPicked(FLinearColor NewFolderColor)
 
 void FFolderModel::OnColorPickerClosed(const TSharedRef<SWindow>& Window)
 {
-	UMovieSceneFolder* MovieSceneFolder = GetFolder();
-	if (MovieSceneFolder && !bFolderPickerWasCancelled)
+	if (!bFolderPickerWasCancelled)
 	{
-		const FScopedTransaction Transaction( NSLOCTEXT( "SequencerFolderNode", "SetFolderColor", "Set Folder Color" ) );
+		TSharedPtr<FSequencerEditorViewModel> EditorViewModel = GetEditor();
+		TSharedPtr<FSequencer> Sequencer = EditorViewModel->GetSequencerImpl();
 
-		FColor CurrentColor = MovieSceneFolder->GetFolderColor();
-		MovieSceneFolder->SetFolderColor(InitialFolderColor);
-		MovieSceneFolder->Modify();
-		MovieSceneFolder->SetFolderColor(CurrentColor);
+		TArray<UMovieSceneFolder*> MovieSceneFolders;
+		Sequencer->GetSelectedFolders(MovieSceneFolders);
+
+		const FScopedTransaction Transaction(NSLOCTEXT("SequencerFolderNode", "SetFolderColor", "Set Folder Color"));
+
+		for (UMovieSceneFolder* MovieSceneFolder : MovieSceneFolders)
+		{
+			FColor CurrentColor = MovieSceneFolder->GetFolderColor();
+			FColor InitialFolderColor = InitialFolderColors.Contains(MovieSceneFolder) ? InitialFolderColors[MovieSceneFolder] : FColor();
+			MovieSceneFolder->SetFolderColor(InitialFolderColor);
+			MovieSceneFolder->Modify();
+			MovieSceneFolder->SetFolderColor(CurrentColor);
+		}
 	}
+	InitialFolderColors.Empty();
 }
 
 void FFolderModel::OnColorPickerCancelled(FLinearColor NewFolderColor)
 {
 	bFolderPickerWasCancelled = true;
-	if (UMovieSceneFolder* MovieSceneFolder = GetFolder())
+
+	TSharedPtr<FSequencerEditorViewModel> EditorViewModel = GetEditor();
+	TSharedPtr<FSequencer> Sequencer = EditorViewModel->GetSequencerImpl();
+
+	TArray<UMovieSceneFolder*> MovieSceneFolders;
+	Sequencer->GetSelectedFolders(MovieSceneFolders);
+
+	for (UMovieSceneFolder* MovieSceneFolder : MovieSceneFolders)
 	{
-		MovieSceneFolder->SetFolderColor(InitialFolderColor);
+		if (InitialFolderColors.Contains(MovieSceneFolder))
+		{
+			MovieSceneFolder->SetFolderColor(InitialFolderColors[MovieSceneFolder]);
+		}
 	}
+
+	InitialFolderColors.Empty();
 }
 
 FText FFolderModel::GetLabel() const
diff --git a/Engine/Source/Editor/Sequencer/Public/MVVM/ViewModels/FolderModel.h b/Engine/Source/Editor/Sequencer/Public/MVVM/ViewModels/FolderModel.h
--- a/Engine/Source/Editor/Sequencer/Public/MVVM/ViewModels/FolderModel.h
+++ b/Engine/Source/Editor/Sequencer/Public/MVVM/ViewModels/FolderModel.h
@@ -108,7 +108,7 @@ private:
 	/*~ FViewModel interface */
 	void OnConstruct() override;
 
-	static FColor InitialFolderColor;
+	static TMap<TWeakObjectPtr<UMovieSceneFolder>, FColor> InitialFolderColors;
 	static bool bFolderPickerWasCancelled;
 
 	void SetFolderColor();
