Sequencer:  Break up undo time proxy into two objects so the IMovieSceneSignedObjectHandler can get destroyed when FSEquencer goes out of scope like the other handlers in order to fix an ASAN issue.

diff --git a/Engine/Source/Editor/Sequencer/Private/Sequencer.cpp b/Engine/Source/Editor/Sequencer/Private/Sequencer.cpp
--- a/Engine/Source/Editor/Sequencer/Private/Sequencer.cpp
+++ b/Engine/Source/Editor/Sequencer/Private/Sequencer.cpp
@@ -129,7 +129,6 @@
 #include "ISequencerChannelInterface.h"
 #include "IMovieRendererInterface.h"
 #include "SequencerKeyCollection.h"
-#include "SequencerTimeChangeUndoRedoProxy.h"
 #include "CurveEditor.h"
 #include "CurveEditorScreenSpace.h"
 #include "CurveDataAbstraction.h"
@@ -528,13 +527,11 @@ void FSequencer::InitSequencer(const FSequencerInitParams& InitParams, const TSh
 	// Update initial movie scene data
 	NotifyMovieSceneDataChanged( EMovieSceneDataChangeType::ActiveMovieSceneChanged );
 
-	//create proxy and set it as child of focused movie scene sequence.
-	if(GetRootMovieSceneSequence())
-	{
-		UndoRedoProxy = NewObject<USequencerTimeChangeUndoRedoProxy>(GetRootMovieSceneSequence(), NAME_None);
-		UndoRedoProxy->SetSequencer(SharedThis(this));
-		UndoRedoProxy->SetFlags(RF_Transactional | RF_Transient);
-	}
+
+	TimeUndoRedoHandler.UndoRedoProxy = NewObject<USequencerTimeChangeUndoRedoProxy>(GetTransientPackage(), NAME_None);
+	TimeUndoRedoHandler.SetSequencer(SharedThis(this));
+	TimeUndoRedoHandler.UndoRedoProxy->SetFlags(RF_Transactional | RF_Transient);
+	
 
 	// Update the view range to the new current time
 	UpdateTimeBoundsToFocusedMovieScene();
@@ -4258,7 +4255,7 @@ void FSequencer::AddReferencedObjects( FReferenceCollector& Collector )
 {
 	Collector.AddReferencedObject( CompiledDataManager );
 	Collector.AddReferencedObject( Settings );
-	Collector.AddReferencedObject( UndoRedoProxy );
+	Collector.AddReferencedObject(TimeUndoRedoHandler.UndoRedoProxy );
 	if (UMovieSceneSequence* RootSequencePtr = RootSequence.Get())
 	{
 		Collector.AddReferencedObject( RootSequencePtr );
diff --git a/Engine/Source/Editor/Sequencer/Private/Sequencer.h b/Engine/Source/Editor/Sequencer/Private/Sequencer.h
--- a/Engine/Source/Editor/Sequencer/Private/Sequencer.h
+++ b/Engine/Source/Editor/Sequencer/Private/Sequencer.h
@@ -39,6 +39,7 @@
 #include "SequencerSettings.h"
 #include "Curves/RichCurve.h"
 #include "Sections/MovieScene3DTransformSection.h"
+#include "SequencerTimeChangeUndoRedoProxy.h"
 
 class AActor;
 class ACameraActor;
@@ -72,7 +73,6 @@ struct FMovieSceneSequencePlaybackParams;
 struct FMovieScenePossessable;
 struct FTransformData;
 struct FKeyAttributes;
-class USequencerTimeChangeUndoRedoProxy;
 struct FNotificationInfo;
 struct FEditorViewportViewModifierParams;
 struct FMovieSceneMarkedFrame;
@@ -674,8 +674,8 @@ public:
 	virtual FString GetReferencerName() const override;
 
 private:
-	//gco object for undo redo support for changing time
-	USequencerTimeChangeUndoRedoProxy* UndoRedoProxy;
+	//object for undo redo support for changing time
+	FSequencerTimeChangedHandler TimeUndoRedoHandler;
 
 public:
 
diff --git a/Engine/Source/Editor/Sequencer/Private/SequencerTimeChangeUndoRedoProxy.cpp b/Engine/Source/Editor/Sequencer/Private/SequencerTimeChangeUndoRedoProxy.cpp
--- a/Engine/Source/Editor/Sequencer/Private/SequencerTimeChangeUndoRedoProxy.cpp
+++ b/Engine/Source/Editor/Sequencer/Private/SequencerTimeChangeUndoRedoProxy.cpp
@@ -3,11 +3,12 @@
 #include "SequencerTimeChangeUndoRedoProxy.h"
 #include "Sequencer.h"
 #include "ScopedTransaction.h"
+#include "MovieScene.h"
 #include "CoreGlobals.h"
 
 #define LOCTEXT_NAMESPACE "Sequencer"
 
-USequencerTimeChangeUndoRedoProxy::~USequencerTimeChangeUndoRedoProxy()
+FSequencerTimeChangedHandler::~FSequencerTimeChangedHandler()
 {
 	if (OnActivateSequenceChangedHandle.IsValid() && WeakSequencer.IsValid())
 	{
@@ -15,43 +16,54 @@ USequencerTimeChangeUndoRedoProxy::~USequencerTimeChangeUndoRedoProxy()
 	}
 }
 
-void USequencerTimeChangeUndoRedoProxy::SetSequencer(TSharedRef<FSequencer> InSequencer)
+void FSequencerTimeChangedHandler::SetSequencer(TSharedRef<FSequencer> InSequencer)
 {
 	WeakSequencer = InSequencer;
-	OnActivateSequenceChangedHandle = InSequencer->OnActivateSequence().AddUObject(this, &USequencerTimeChangeUndoRedoProxy::OnActivateSequenceChanged);
+	if (UndoRedoProxy)
+	{
+		UndoRedoProxy->WeakSequencer = InSequencer;
+	}
+	OnActivateSequenceChangedHandle = InSequencer->OnActivateSequence().AddRaw(this, &FSequencerTimeChangedHandler::OnActivateSequenceChanged);
 	OnActivateSequenceChanged(InSequencer->GetFocusedTemplateID());
 }
 
-void USequencerTimeChangeUndoRedoProxy::OnActivateSequenceChanged(FMovieSceneSequenceIDRef ID)
+void FSequencerTimeChangedHandler::OnActivateSequenceChanged(FMovieSceneSequenceIDRef ID)
 {
 	if (TSharedPtr<ISequencer> SequencerPtr = WeakSequencer.Pin())
 	{
 		if (SequencerPtr->GetFocusedMovieSceneSequence() && SequencerPtr->GetFocusedMovieSceneSequence()->GetMovieScene())
 		{
 			using namespace UE::MovieScene;
-			bTimeWasSet = false;
+			if (UndoRedoProxy)
+			{
+				UndoRedoProxy->bTimeWasSet = false;
+			}
 			MovieSceneModified.Unlink();
 			SequencerPtr->GetFocusedMovieSceneSequence()->GetMovieScene()->UMovieSceneSignedObject::EventHandlers.Link(MovieSceneModified, this);
 		}
 	}
 }
 
-void USequencerTimeChangeUndoRedoProxy::OnModifiedIndirectly(UMovieSceneSignedObject* Object)
+void FSequencerTimeChangedHandler::OnModifiedIndirectly(UMovieSceneSignedObject* Object)
 {
-	if (Object->IsA<UMovieSceneSection>() && WeakSequencer.IsValid())
+	if (UndoRedoProxy == nullptr)
 	{
-		FQualifiedFrameTime InTime = WeakSequencer.Pin()->GetGlobalTime();
+		return;
+	}
+	if (Object->IsA<UMovieSceneSection>() && UndoRedoProxy->WeakSequencer.IsValid())
+	{
+		FQualifiedFrameTime InTime = UndoRedoProxy->WeakSequencer.Pin()->GetGlobalTime();
 
-		if (bTimeWasSet)
+		if (UndoRedoProxy->bTimeWasSet)
 		{
-			if (!GIsTransacting && (InTime.Time != Time.Time || InTime.Rate != Time.Rate))
+			if (!GIsTransacting && (InTime.Time != UndoRedoProxy->Time.Time || InTime.Rate != UndoRedoProxy->Time.Rate))
 			{
 				const FScopedTransaction Transaction(LOCTEXT("TimeChanged", "Time Changed"));
-				Modify();
+				UndoRedoProxy->Modify();
 			}
 		}
-		bTimeWasSet = true;
-		Time = InTime;
+		UndoRedoProxy->bTimeWasSet = true;
+		UndoRedoProxy->Time = InTime;
 	}
 }
 
@@ -63,5 +75,7 @@ void USequencerTimeChangeUndoRedoProxy::PostEditUndo()
 	}
 }
 
+
+
 #undef LOCTEXT_NAMESPACE
 
diff --git a/Engine/Source/Editor/Sequencer/Private/SequencerTimeChangeUndoRedoProxy.h b/Engine/Source/Editor/Sequencer/Private/SequencerTimeChangeUndoRedoProxy.h
--- a/Engine/Source/Editor/Sequencer/Private/SequencerTimeChangeUndoRedoProxy.h
+++ b/Engine/Source/Editor/Sequencer/Private/SequencerTimeChangeUndoRedoProxy.h
@@ -12,24 +12,19 @@
 
 class FSequencer;
 
+
+
 UCLASS()
-class  USequencerTimeChangeUndoRedoProxy : public UObject, public UE::MovieScene::ISignedObjectEventHandler
+class  USequencerTimeChangeUndoRedoProxy : public UObject
 {
 public:
 	GENERATED_BODY()
 	USequencerTimeChangeUndoRedoProxy() :bTimeWasSet(false), WeakSequencer(nullptr) {  };
-	~USequencerTimeChangeUndoRedoProxy();
-
-	void SetSequencer(TSharedRef<FSequencer> InSequencer);
-	void OnActivateSequenceChanged(FMovieSceneSequenceIDRef ID);
+	~USequencerTimeChangeUndoRedoProxy() {};
 
 	/*~ UObject */
 	virtual void PostEditUndo() override;
 
-	/*~ ISignedObjectEventHandler Interface */
-	virtual void  OnModifiedIndirectly(UMovieSceneSignedObject*) override;
-
-private:
 	UPROPERTY(Transient)
 	FQualifiedFrameTime Time;
 	
@@ -37,9 +32,27 @@ private:
 	UPROPERTY(Transient)
 	bool bTimeWasSet = false;
 	
-	FDelegateHandle OnActivateSequenceChangedHandle;
 	TWeakPtr<FSequencer> WeakSequencer;
 
+};
+
+
+class FSequencerTimeChangedHandler: public UE::MovieScene::ISignedObjectEventHandler
+{
+public:
+	FSequencerTimeChangedHandler() : UndoRedoProxy(nullptr), WeakSequencer(nullptr){};
+	virtual ~FSequencerTimeChangedHandler();
+
+	/*~ ISignedObjectEventHandler Interface */
+	virtual void  OnModifiedIndirectly(UMovieSceneSignedObject*) override;
+
+	void SetSequencer(TSharedRef<FSequencer> InSequencer);
+	void OnActivateSequenceChanged(FMovieSceneSequenceIDRef ID);
+
+	USequencerTimeChangeUndoRedoProxy* UndoRedoProxy;
 	UE::MovieScene::TNonIntrusiveEventHandler<UE::MovieScene::ISignedObjectEventHandler> MovieSceneModified;
+	TWeakPtr<FSequencer> WeakSequencer;
+	FDelegateHandle OnActivateSequenceChangedHandle;
+
 };
 
