Sequencer: Fix regression with adding a spawnable to an existing binding

This part was lost in the refactor to SequencerUtilities

diff --git a/Engine/Source/Editor/Sequencer/Private/SequencerUtilities.cpp b/Engine/Source/Editor/Sequencer/Private/SequencerUtilities.cpp
--- a/Engine/Source/Editor/Sequencer/Private/SequencerUtilities.cpp
+++ b/Engine/Source/Editor/Sequencer/Private/SequencerUtilities.cpp
@@ -2789,7 +2789,10 @@ void FSequencerUtilities::AddActorsToBinding(TSharedRef<ISequencer> Sequencer, c
 				}
 
 				ActorToAdd->Modify();
-				Sequence->BindPossessableObject(Guid, *ActorToAdd, Sequencer->GetPlaybackContext());
+				if (!MovieScene->FindPossessable(Guid)->BindSpawnableObject(Sequencer->GetFocusedTemplateID(), ActorToAdd, &Sequencer.Get()))
+				{
+					Sequence->BindPossessableObject(Guid, *ActorToAdd, Sequencer->GetPlaybackContext());
+				}
 				++NumObjectsAdded;
 			}
 			else
