Sequencer: When copying actors ID them with their full path name in order to handle them living in sublevels, since GetName() can return actor in different level with same FName.

diff --git a/Engine/Source/Editor/Sequencer/Private/Sequencer.cpp b/Engine/Source/Editor/Sequencer/Private/Sequencer.cpp
--- a/Engine/Source/Editor/Sequencer/Private/Sequencer.cpp
+++ b/Engine/Source/Editor/Sequencer/Private/Sequencer.cpp
@@ -8070,7 +8070,7 @@ bool FSequencer::PasteObjectBindings(const FString& TextToImport, UMovieSceneFol
 			{
 				AActor* Actor = Cast<AActor>(WeakObject.Get());
 
-				if (Actor && ObjectNames.Contains(Actor->GetName()))
+				if (Actor && ObjectNames.Contains(Actor->GetPathName()))
 				{
 					FText DuplicateActorsMsg = FText::Format(LOCTEXT("DuplicateActorsForPastedBinding", "Attempting to paste a binding that is already bound to {0}.\nShould the existing actor be duplicated for the pasted binding?"), FText::FromString(Actor->GetActorLabel()));
 
diff --git a/Engine/Source/Editor/Sequencer/Private/SequencerUtilities.cpp b/Engine/Source/Editor/Sequencer/Private/SequencerUtilities.cpp
--- a/Engine/Source/Editor/Sequencer/Private/SequencerUtilities.cpp
+++ b/Engine/Source/Editor/Sequencer/Private/SequencerUtilities.cpp
@@ -1954,7 +1954,7 @@ void FSequencerUtilities::CopyBindings(TSharedRef<ISequencer> Sequencer, const T
 			{
 				if (AActor* Actor = Cast<AActor>(RuntimeObject.Get()))
 				{
-					CopyableBinding->BoundObjectNames.Add(Actor->GetName());
+					CopyableBinding->BoundObjectNames.Add(Actor->GetPathName());
 				}
 			}
 		}
@@ -2162,7 +2162,7 @@ bool FSequencerUtilities::PasteBindings(const FString& TextToImport, TSharedRef<
 					for (TActorIterator<AActor> ActorItr(World); ActorItr; ++ActorItr)
 					{
 						AActor* Actor = *ActorItr;
-						if (Actor && CopyableBinding->BoundObjectNames.Contains(Actor->GetName()))
+						if (Actor && CopyableBinding->BoundObjectNames.Contains(Actor->GetPathName()))
 						{
 							// If this actor is already bound and we're not duplicating actors, don't bind to anything
 							if (!PasteBindingsParams.bDuplicateExistingActors && Sequencer->FindObjectId(*Actor, Sequencer->GetFocusedTemplateID()).IsValid())
@@ -2173,6 +2173,10 @@ bool FSequencerUtilities::PasteBindings(const FString& TextToImport, TSharedRef<
 							ActorsToRebind.Add(Actor);
 							CopyableBinding->BoundObjectNames.Remove(Actor->GetName());
 						}
+
+						ActorsToRebind.Add(Actor);
+						CopyableBinding->BoundObjectNames.Remove(Actor->GetPathName());
+
 					}
 				}
 
@@ -2203,7 +2207,7 @@ bool FSequencerUtilities::PasteBindings(const FString& TextToImport, TSharedRef<
 							{
 								ActorsToRebind.Add(Actor);
 
-								CopyableBinding->BoundObjectNames.Add(Actor->GetName());
+								CopyableBinding->BoundObjectNames.Add(Actor->GetPathName());
 							}
 						}
 					}
@@ -2797,7 +2801,7 @@ void FSequencerUtilities::AddActorsToBinding(TSharedRef<ISequencer> Sequencer, c
 			}
 			else
 			{
-				const FText NotificationText = FText::Format(LOCTEXT("UnableToAssignObject", "Cannot assign object {0}. Expected class {1}"), FText::FromString(ActorToAdd->GetName()), FText::FromString(ActorClass->GetName()));
+				const FText NotificationText = FText::Format(LOCTEXT("UnableToAssignObject", "Cannot assign object {0}. Expected class {1}"), FText::FromString(ActorToAdd->GetPathName()), FText::FromString(ActorClass->GetName()));
 				FNotificationInfo Info(NotificationText);
 				Info.ExpireDuration = 3.f;
 				Info.bUseLargeFont = false;
