Sequencer: fix incorrect setup and design for track areas

The pinned track area was shared with the normal track area, resulting in much painting and user interaction happening simultaneously in both areas. Now we properly create a pinned track area.

This change highlighted a problem which existed before the mvvm refactor, but became worse with that refactor: a lot of code needing access to the track area would simply get the "normal" track area from the sequencer editor, instead of getting the actual track area they are related to. This mean a lot of code was getting the "normal" track area when they were instead operating on the pinned one. Regressions have been fixed, but incorrect code from before the refactor has also been fixed. There might still be more to fix.

Before the mvvm refactor, the active hotspot was stored on the sequencer itself. Now, it's stored on the track area, which means there's no one place to figure out where the current hotspot is -- we would have to know which track area the mouse cursor is currently over. To avoid that, we bubble up the current hotspot from any of our track areas up to the sequencer editor view model. So now any code needing to simply know what is the current hotspot can look that up on the view model.

diff --git a/Engine/Source/Editor/Sequencer/Private/MVVM/ViewModels/LayerBarModel.cpp b/Engine/Source/Editor/Sequencer/Private/MVVM/ViewModels/LayerBarModel.cpp
--- a/Engine/Source/Editor/Sequencer/Private/MVVM/ViewModels/LayerBarModel.cpp
+++ b/Engine/Source/Editor/Sequencer/Private/MVVM/ViewModels/LayerBarModel.cpp
@@ -9,6 +9,7 @@
 #include "MVVM/ViewModels/EditorViewModel.h"
 #include "MVVM/ViewModels/SequencerEditorViewModel.h"
 #include "MVVM/ViewModels/ViewModelIterators.h"
+#include "MVVM/Views/STrackLane.h"
 #include "Math/RangeBound.h"
 #include "Misc/AssertionMacros.h"
 #include "Misc/FrameNumber.h"
@@ -77,7 +78,7 @@ void FLayerBarModel::OnExtensionsDirtied()
 
 TSharedPtr<ITrackLaneWidget> FLayerBarModel::CreateTrackLaneView(const FCreateTrackLaneViewParams& InParams)
 {
-	return SNew(SSequencerLayerBar, InParams.Editor->CastThisSharedChecked<FSequencerEditorViewModel>(), SharedThis(this));
+	return SNew(SSequencerLayerBar, InParams.OwningTrackLane->GetTrackAreaView(), InParams.Editor->CastThisSharedChecked<FSequencerEditorViewModel>(), SharedThis(this));
 }
 
 FTrackLaneVirtualAlignment FLayerBarModel::ArrangeVirtualTrackLaneView() const
diff --git a/Engine/Source/Editor/Sequencer/Private/MVVM/ViewModels/SequencerEditorViewModel.cpp b/Engine/Source/Editor/Sequencer/Private/MVVM/ViewModels/SequencerEditorViewModel.cpp
--- a/Engine/Source/Editor/Sequencer/Private/MVVM/ViewModels/SequencerEditorViewModel.cpp
+++ b/Engine/Source/Editor/Sequencer/Private/MVVM/ViewModels/SequencerEditorViewModel.cpp
@@ -43,7 +43,22 @@ TSharedPtr<FTrackAreaViewModel> FSequencerEditorViewModel::CreateTrackAreaImpl()
 {
 	TSharedPtr<ISequencer> Sequencer = WeakSequencer.Pin();
 	check(Sequencer.IsValid());
-	return MakeShared<FSequencerTrackAreaViewModel>(Sequencer.ToSharedRef());
+
+	TSharedRef<FSequencerTrackAreaViewModel> NewTrackArea = MakeShared<FSequencerTrackAreaViewModel>(Sequencer.ToSharedRef());
+	NewTrackArea->GetOnHotspotChangedDelegate().AddSP(SharedThis(this), &FSequencerEditorViewModel::OnTrackAreaHotspotChanged);
+	return NewTrackArea;
+}
+
+void FSequencerEditorViewModel::InitializeEditorImpl()
+{
+	PinnedTrackArea = CreateTrackAreaImpl();
+	PinnedTrackArea->GetOnHotspotChangedDelegate().AddSP(SharedThis(this), &FSequencerEditorViewModel::OnTrackAreaHotspotChanged);
+	GetEditorPanels().AddChild(PinnedTrackArea);
+}
+
+TSharedPtr<FTrackAreaViewModel> FSequencerEditorViewModel::GetPinnedTrackArea() const
+{
+	return PinnedTrackArea;
 }
 
 TSharedPtr<ISequencer> FSequencerEditorViewModel::GetSequencer() const
@@ -69,6 +84,16 @@ bool FSequencerEditorViewModel::IsReadOnly() const
 	return !Sequencer || Sequencer->IsReadOnly();
 }
 
+TSharedPtr<ITrackAreaHotspot> FSequencerEditorViewModel::GetHotspot() const
+{
+	return CurrentHotspot;
+}
+
+void FSequencerEditorViewModel::OnTrackAreaHotspotChanged(TSharedPtr<ITrackAreaHotspot> NewHotspot)
+{
+	CurrentHotspot = NewHotspot;
+}
+
 } // namespace Sequencer
 } // namespace UE
 
diff --git a/Engine/Source/Editor/Sequencer/Private/MVVM/ViewModels/SequencerTrackAreaViewModel.cpp b/Engine/Source/Editor/Sequencer/Private/MVVM/ViewModels/SequencerTrackAreaViewModel.cpp
--- a/Engine/Source/Editor/Sequencer/Private/MVVM/ViewModels/SequencerTrackAreaViewModel.cpp
+++ b/Engine/Source/Editor/Sequencer/Private/MVVM/ViewModels/SequencerTrackAreaViewModel.cpp
@@ -1,10 +1,13 @@
 // Copyright Epic Games, Inc. All Rights Reserved.
 
 #include "MVVM/ViewModels/SequencerTrackAreaViewModel.h"
-#include "MVVM/ViewModels/SequencerEditorViewModel.h"
-#include "MVVM/ViewModels/SequenceModel.h"
+
 #include "AnimatedRange.h"
+#include "MVVM/ViewModels/SequenceModel.h"
+#include "MVVM/ViewModels/SequencerEditorViewModel.h"
 #include "Sequencer.h"
+#include "Tools/SequencerEditTool_Movement.h"
+#include "Tools/SequencerEditTool_Selection.h"
 
 #include "Framework/MultiBox/MultiBoxBuilder.h"
 #include "Framework/Commands/GenericCommands.h"
@@ -19,6 +22,16 @@ FSequencerTrackAreaViewModel::FSequencerTrackAreaViewModel(TSharedRef<ISequencer
 {
 }
 
+void FSequencerTrackAreaViewModel::InitializeDefaultEditTools(UE::Sequencer::STrackAreaView& InTrackArea)
+{
+	TSharedPtr<FSequencer> Sequencer = StaticCastSharedPtr<FSequencer>(WeakSequencer.Pin());
+	if (ensure(Sequencer))
+	{
+		AddEditTool(MakeShared<FSequencerEditTool_Selection>(*Sequencer, InTrackArea));
+		AddEditTool(MakeShared<FSequencerEditTool_Movement>(*Sequencer, InTrackArea));
+	}
+}
+
 FFrameRate FSequencerTrackAreaViewModel::GetTickResolution() const
 {
 	if (TSharedPtr<ISequencer> Sequencer = WeakSequencer.Pin())
diff --git a/Engine/Source/Editor/Sequencer/Private/SSequencer.cpp b/Engine/Source/Editor/Sequencer/Private/SSequencer.cpp
--- a/Engine/Source/Editor/Sequencer/Private/SSequencer.cpp
+++ b/Engine/Source/Editor/Sequencer/Private/SSequencer.cpp
@@ -9,6 +9,7 @@
 #include "MVVM/ViewModels/ViewModelIterators.h"
 #include "MVVM/Extensions/IRenameableExtension.h"
 #include "MVVM/ViewModels/SequencerEditorViewModel.h"
+#include "MVVM/ViewModels/SequencerTrackAreaViewModel.h"
 #include "MVVM/ViewModels/SequencerOutlinerViewModel.h"
 #include "MVVM/ViewModels/TrackRowModel.h"
 #include "Sections/MovieSceneSubSection.h"
@@ -281,8 +282,9 @@ void SSequencer::Construct(const FArguments& InArgs, TSharedRef<FSequencer> InSe
 	TSharedRef<SScrollBar> PinnedAreaScrollBar = SNew(SScrollBar)
 		.Thickness(FVector2D(9.0f, 9.0f));
 
-	SAssignNew(PinnedTrackArea, SSequencerTrackAreaView, InSequencer->GetViewModel()->GetTrackArea(), TimeSliderControllerRef);
-	SAssignNew(PinnedTreeView, SOutlinerView, InSequencer->GetViewModel()->GetOutliner(), PinnedTrackArea.ToSharedRef())
+	TSharedPtr<FSequencerEditorViewModel> SequencerViewModel = InSequencer->GetViewModel()->CastThisShared<FSequencerEditorViewModel>();
+	SAssignNew(PinnedTrackArea, SSequencerTrackAreaView, SequencerViewModel->GetPinnedTrackArea(), TimeSliderControllerRef);
+	SAssignNew(PinnedTreeView, SOutlinerView, SequencerViewModel->GetOutliner(), PinnedTrackArea.ToSharedRef())
 		.SelectionHandler(InArgs._SelectionHandler)
 		.ExternalScrollbar(PinnedAreaScrollBar)
 		.Clipping(EWidgetClipping::ClipToBounds);
@@ -292,8 +294,8 @@ void SSequencer::Construct(const FArguments& InArgs, TSharedRef<FSequencer> InSe
 	PinnedTrackArea->SetIsPinned(true);
 	PinnedTreeView->SetShowPinned(true);
 
-	SAssignNew(TrackArea, SSequencerTrackAreaView, InSequencer->GetViewModel()->GetTrackArea(), TimeSliderControllerRef);
-	SAssignNew(TreeView, SOutlinerView, InSequencer->GetViewModel()->GetOutliner(), TrackArea.ToSharedRef())
+	SAssignNew(TrackArea, SSequencerTrackAreaView, SequencerViewModel->GetTrackArea(), TimeSliderControllerRef);
+	SAssignNew(TreeView, SOutlinerView, SequencerViewModel->GetOutliner(), TrackArea.ToSharedRef())
 		.SelectionHandler(InArgs._SelectionHandler)
 		.ExternalScrollbar(ScrollBar)
 		.Clipping(EWidgetClipping::ClipToBounds);
@@ -302,12 +304,14 @@ void SSequencer::Construct(const FArguments& InArgs, TSharedRef<FSequencer> InSe
 
 	TreeView->AddPinnedTreeView(PinnedTreeView);
 
+	SequencerViewModel->GetTrackArea()->CastThisChecked<FSequencerTrackAreaViewModel>()->InitializeDefaultEditTools(*TrackArea);
+	SequencerViewModel->GetPinnedTrackArea()->CastThisChecked<FSequencerTrackAreaViewModel>()->InitializeDefaultEditTools(*PinnedTrackArea);
+
 	TAttribute<FAnimatedRange> ViewRangeAttribute = InArgs._ViewRange;
 
 	if (InSequencer->GetHostCapabilities().bSupportsCurveEditor)
 	{
-		TSharedPtr<FEditorViewModel> RootModel = InSequencer->GetViewModel();
-		FCurveEditorExtension* CurveEditorExtension = RootModel->CastDynamicChecked<FCurveEditorExtension>();
+		FCurveEditorExtension* CurveEditorExtension = SequencerViewModel->CastDynamicChecked<FCurveEditorExtension>();
 		CurveEditorExtension->CreateCurveEditor(TimeSliderArgs);
 	}
 
@@ -3750,19 +3754,18 @@ UE::Sequencer::FVirtualTrackArea SSequencer::GetVirtualTrackArea(const UE::Seque
 {
 	using namespace UE::Sequencer;
 
-	TSharedPtr<UE::Sequencer::FEditorViewModel> EditorViewModel = SequencerPtr.Pin()->GetViewModel();
-	FTrackAreaViewModel& TrackAreaViewModel = *EditorViewModel->GetTrackArea().Get();
-
 	const STrackAreaView* TargetTrackArea = TrackArea.Get();
 	TSharedPtr<SOutlinerView> TargetTreeView = TreeView;
 	
 	if (InTrackArea != nullptr)
 	{
 		TargetTrackArea = InTrackArea;
-		//TargetTreeView = TargetTrackArea->GetTreeView().Pin();
+		TargetTreeView = TargetTrackArea->GetOutliner().Pin();
 	}
 
-	return FVirtualTrackArea(TrackAreaViewModel, *TargetTreeView.Get(), TargetTrackArea->GetCachedGeometry());
+	TSharedPtr<FTrackAreaViewModel> TargetTrackAreaViewModel = TargetTrackArea->GetViewModel();
+
+	return FVirtualTrackArea(*TargetTrackAreaViewModel, *TargetTreeView.Get(), TargetTrackArea->GetCachedGeometry());
 }
 
 FPasteContextMenuArgs SSequencer::GeneratePasteArgs(FFrameNumber PasteAtTime, TSharedPtr<FMovieSceneClipboard> Clipboard)
diff --git a/Engine/Source/Editor/Sequencer/Private/SSequencerSection.cpp b/Engine/Source/Editor/Sequencer/Private/SSequencerSection.cpp
--- a/Engine/Source/Editor/Sequencer/Private/SSequencerSection.cpp
+++ b/Engine/Source/Editor/Sequencer/Private/SSequencerSection.cpp
@@ -73,20 +73,21 @@ FTimeToPixel ConstructTimeConverterForSection(const FGeometry& InSectionGeometry
 
 struct FSequencerSectionPainterImpl : FSequencerSectionPainter
 {
-	FSequencerSectionPainterImpl(FSequencer& InSequencer, TSharedPtr<FSectionModel> InSection, FSlateWindowElementList& _OutDrawElements, const FGeometry& InSectionGeometry, const SSequencerSection& InSectionWidget)
+	FSequencerSectionPainterImpl(FSequencer& InSequencer, TSharedPtr<FTrackAreaViewModel> InTrackAreaViewModel, TSharedPtr<FSectionModel> InSection, FSlateWindowElementList& _OutDrawElements, const FGeometry& InSectionGeometry, const SSequencerSection& InSectionWidget)
 		: FSequencerSectionPainter(_OutDrawElements, InSectionGeometry, InSection)
 		, Sequencer(InSequencer)
 		, SectionWidget(InSectionWidget)
 		, TimeToPixelConverter(ConstructTimeConverterForSection(SectionGeometry, *InSection->GetSection(), Sequencer))
+		, TrackAreaViewModel(InTrackAreaViewModel)
 		, bClipRectEnabled(false)
 	{
 		CalculateSelectionColor();
 
-		const ISequencerEditTool* EditTool = InSequencer.GetViewModel()->GetTrackArea()->GetEditTool();
+		const ISequencerEditTool* EditTool = TrackAreaViewModel->GetEditTool();
 		Hotspot = EditTool ? EditTool->GetDragHotspot() : nullptr;
 		if (!Hotspot)
 		{
-			Hotspot = InSequencer.GetViewModel()->GetTrackArea()->GetHotspot();
+			Hotspot = TrackAreaViewModel->GetHotspot();
 		}
 	}
 
@@ -823,6 +824,7 @@ struct FSequencerSectionPainterImpl : FSequencerSectionPainter
 	FSequencer& Sequencer;
 	const SSequencerSection& SectionWidget;
 	FTimeToPixel TimeToPixelConverter;
+	TSharedPtr<FTrackAreaViewModel> TrackAreaViewModel;
 	TSharedPtr<ITrackAreaHotspot> Hotspot;
 
 	/** The clipping rectangle of the parent widget */
@@ -884,9 +886,10 @@ void SSequencerSection::Construct( const FArguments& InArgs, TSharedPtr<FSequenc
 
 SSequencerSection::~SSequencerSection()
 {
-	if (Sequencer.IsValid() && GetSequencer().GetViewModel() && GetSequencer().GetViewModel()->GetTrackArea())
+	TSharedPtr<FTrackAreaViewModel> TrackAreaViewModel = GetTrackAreaViewModel();
+	if (TrackAreaViewModel.IsValid())
 	{
-		GetSequencer().GetViewModel()->GetTrackArea()->SetHotspot(nullptr);
+		TrackAreaViewModel->SetHotspot(nullptr);
 	}
 }
 
@@ -1097,6 +1100,7 @@ bool SSequencerSection::CheckForEasingHandleInteraction( const FPointerEvent& Mo
 		}
 	}
 
+	TSharedPtr<FTrackAreaViewModel> TrackAreaViewModel = GetTrackAreaViewModel();
 	for (const TSharedPtr<FSectionModel>& SectionModel : AllUnderlappingSections)
 	{
 		UMovieSceneSection* EasingSectionObj = SectionModel->GetSection();
@@ -1108,7 +1112,7 @@ bool SSequencerSection::CheckForEasingHandleInteraction( const FPointerEvent& Mo
 
 			if (FMath::IsNearlyEqual(MouseTime, HandlePositionIn, HalfHandleSizeX))
 			{
-				GetSequencer().GetViewModel()->GetTrackArea()->SetHotspot(MakeShared<FSectionEasingHandleHotspot>(ESequencerEasingType::In, SectionModel, Sequencer));
+				TrackAreaViewModel->SetHotspot(MakeShared<FSectionEasingHandleHotspot>(ESequencerEasingType::In, SectionModel, Sequencer));
 				return true;
 			}
 		}
@@ -1120,7 +1124,7 @@ bool SSequencerSection::CheckForEasingHandleInteraction( const FPointerEvent& Mo
 
 			if (FMath::IsNearlyEqual(MouseTime, HandlePositionOut, HalfHandleSizeX))
 			{
-				GetSequencer().GetViewModel()->GetTrackArea()->SetHotspot(MakeShared<FSectionEasingHandleHotspot>(ESequencerEasingType::Out, SectionModel, Sequencer));
+				TrackAreaViewModel->SetHotspot(MakeShared<FSectionEasingHandleHotspot>(ESequencerEasingType::Out, SectionModel, Sequencer));
 				return true;
 			}
 		}
@@ -1138,11 +1142,12 @@ bool SSequencerSection::CheckForEdgeInteraction( const FPointerEvent& MouseEvent
 		return false;
 	}
 
-	const ISequencerEditTool* EditTool = GetSequencer().GetViewModel()->GetTrackArea()->GetEditTool();
+	TSharedPtr<FTrackAreaViewModel> TrackAreaViewModel = GetTrackAreaViewModel();
+	const ISequencerEditTool* EditTool = TrackAreaViewModel->GetEditTool();
 	TSharedPtr<ITrackAreaHotspot> Hotspot = EditTool ? EditTool->GetDragHotspot() : nullptr;
 	if (!Hotspot)
 	{
-		Hotspot = GetSequencer().GetViewModel()->GetTrackArea()->GetHotspot();
+		Hotspot = TrackAreaViewModel->GetHotspot();
 	}
 
 	if (Hotspot && !Hotspot->IsA<FSectionHotspotBase>())
@@ -1185,7 +1190,7 @@ bool SSequencerSection::CheckForEdgeInteraction( const FPointerEvent& MouseEvent
 
 			if( SectionRectLeft.IsUnderLocation( MouseEvent.GetScreenSpacePosition() ) )
 			{
-				GetSequencer().GetViewModel()->GetTrackArea()->SetHotspot(MakeShareable( new FSectionResizeHotspot(FSectionResizeHotspot::Left, UnderlappingSection, Sequencer)) );
+				TrackAreaViewModel->SetHotspot(MakeShareable( new FSectionResizeHotspot(FSectionResizeHotspot::Left, UnderlappingSection, Sequencer)) );
 				return true;
 			}
 		}
@@ -1199,7 +1204,7 @@ bool SSequencerSection::CheckForEdgeInteraction( const FPointerEvent& MouseEvent
 
 			if( SectionRectRight.IsUnderLocation( MouseEvent.GetScreenSpacePosition() ) )
 			{
-				GetSequencer().GetViewModel()->GetTrackArea()->SetHotspot(MakeShareable( new FSectionResizeHotspot(FSectionResizeHotspot::Right, UnderlappingSection, Sequencer)) );
+				TrackAreaViewModel->SetHotspot(MakeShareable( new FSectionResizeHotspot(FSectionResizeHotspot::Right, UnderlappingSection, Sequencer)) );
 				return true;
 			}
 		}
@@ -1230,6 +1235,7 @@ bool SSequencerSection::CheckForEasingAreaInteraction( const FPointerEvent& Mous
 	const FFrameNumber MouseTime = TimeToPixelConverter.PixelToFrame(LocalMousePos.X).FrameNumber;
 
 	// First off, set the hotspot to an easing area if necessary
+	TSharedPtr<FTrackAreaViewModel> TrackAreaViewModel = GetTrackAreaViewModel();
 	for (const FOverlappingSections& Segment : UnderlappingEasingSegments)
 	{
 		if (!Segment.Range.Contains(MouseTime))
@@ -1253,7 +1259,7 @@ bool SSequencerSection::CheckForEasingAreaInteraction( const FPointerEvent& Mous
 
 		if (EasingAreas.Num())
 		{
-			GetSequencer().GetViewModel()->GetTrackArea()->SetHotspot(MakeShared<FSectionEasingAreaHotspot>(EasingAreas, WeakSectionModel, Sequencer));
+			TrackAreaViewModel->SetHotspot(MakeShared<FSectionEasingAreaHotspot>(EasingAreas, WeakSectionModel, Sequencer));
 			return true;
 		}
 	}
@@ -1265,6 +1271,24 @@ FSequencer& SSequencerSection::GetSequencer() const
 	return *Sequencer.Pin().Get();
 }
 
+TSharedPtr<STrackAreaView> SSequencerSection::GetTrackAreaView() const
+{
+	if (TSharedPtr<STrackLane> OwningTrackLane = WeakOwningTrackLane.Pin())
+	{
+		return OwningTrackLane->GetTrackAreaView();
+	}
+	return nullptr;
+}
+
+TSharedPtr<FTrackAreaViewModel> SSequencerSection::GetTrackAreaViewModel() const
+{
+	if (TSharedPtr<STrackAreaView> TrackAreaView = GetTrackAreaView())
+	{
+		return TrackAreaView->GetViewModel();
+	}
+	return nullptr;
+}
+
 int32 SSequencerSection::OnPaint( const FPaintArgs& Args, const FGeometry& AllottedGeometry, const FSlateRect& MyCullingRect, FSlateWindowElementList& OutDrawElements, int32 LayerId, const FWidgetStyle& InWidgetStyle, bool bParentEnabled ) const
 {
 	TSharedPtr<FSectionModel> SectionModel = WeakSectionModel.Pin();
@@ -1274,11 +1298,12 @@ int32 SSequencerSection::OnPaint( const FPaintArgs& Args, const FGeometry& Allot
 		return LayerId;
 	}
 
-	const ISequencerEditTool* EditTool = GetSequencer().GetViewModel()->GetTrackArea()->GetEditTool();
+	TSharedPtr<FTrackAreaViewModel> TrackAreaViewModel = GetTrackAreaViewModel();
+	const ISequencerEditTool* EditTool = TrackAreaViewModel->GetEditTool();
 	TSharedPtr<ITrackAreaHotspot> Hotspot = EditTool ? EditTool->GetDragHotspot() : nullptr;
 	if (!Hotspot)
 	{
-		Hotspot = GetSequencer().GetViewModel()->GetTrackArea()->GetHotspot();
+		Hotspot = TrackAreaViewModel->GetHotspot();
 	}
 
 	UMovieSceneTrack* Track = SectionObject->GetTypedOuter<UMovieSceneTrack>();
@@ -1289,7 +1314,7 @@ int32 SSequencerSection::OnPaint( const FPaintArgs& Args, const FGeometry& Allot
 
 	FGeometry SectionGeometry = MakeSectionGeometryWithoutHandles( AllottedGeometry, SectionInterface );
 
-	FSequencerSectionPainterImpl Painter(GetSequencer(), SectionModel, OutDrawElements, SectionGeometry, *this);
+	FSequencerSectionPainterImpl Painter(GetSequencer(), TrackAreaViewModel, SectionModel, OutDrawElements, SectionGeometry, *this);
 
 	FGeometry PaintSpaceParentGeometry = ParentGeometry;
 	PaintSpaceParentGeometry.AppendTransform(FSlateLayoutTransform(Inverse(Args.GetWindowToDesktopTransform())));
@@ -1787,7 +1812,8 @@ FReply SSequencerSection::OnMouseMove( const FGeometry& MyGeometry, const FPoint
 		!CheckForEasingAreaInteraction(MouseEvent, MyGeometry))
 	{
 		// If nothing was hit, we just hit the section
-		GetSequencer().GetViewModel()->GetTrackArea()->SetHotspot( MakeShareable( new FSectionHotspot(WeakSectionModel, Sequencer)) );
+		TSharedPtr<FTrackAreaViewModel> TrackAreaViewModel = GetTrackAreaViewModel();
+		TrackAreaViewModel->SetHotspot( MakeShareable( new FSectionHotspot(WeakSectionModel, Sequencer)) );
 	}
 
 	return FReply::Unhandled();
@@ -1802,7 +1828,8 @@ void SSequencerSection::OnMouseLeave( const FPointerEvent& MouseEvent )
 {
 	SCompoundWidget::OnMouseLeave( MouseEvent );
 
-	GetSequencer().GetViewModel()->GetTrackArea()->SetHotspot(nullptr);
+	TSharedPtr<FTrackAreaViewModel> TrackAreaViewModel = GetTrackAreaViewModel();
+	TrackAreaViewModel->SetHotspot(nullptr);
 }
 
 static float SectionThrobDurationSeconds = 1.f;
diff --git a/Engine/Source/Editor/Sequencer/Private/SSequencerSection.h b/Engine/Source/Editor/Sequencer/Private/SSequencerSection.h
--- a/Engine/Source/Editor/Sequencer/Private/SSequencerSection.h
+++ b/Engine/Source/Editor/Sequencer/Private/SSequencerSection.h
@@ -24,7 +24,9 @@ namespace Sequencer
 {
 
 class SCompoundTrackLaneView;
+class STrackAreaView;
 class FSectionModel;
+class FTrackAreaViewModel;
 struct ITrackAreaHotspot;
 
 class SSequencerSection : public SCompoundWidget, public ITrackLaneWidget, public UE::MovieScene::ISignedObjectEventHandler
@@ -85,6 +87,12 @@ private:
 	/** @return the sequencer interface */
 	FSequencer& GetSequencer() const;
 
+	/** @return the track area view that this section belongs to */
+	TSharedPtr<STrackAreaView> GetTrackAreaView() const;
+
+	/** @return the track area view-model that this section belongs to */
+	TSharedPtr<FTrackAreaViewModel> GetTrackAreaViewModel() const;
+
 	/** 
 	 * Creates geometry for a section without space for the handles
 	 */
@@ -133,6 +141,7 @@ public:
 	static bool IsSectionHighlighted(UMovieSceneSection* InSection, TSharedPtr<ITrackAreaHotspot> Hotspot);
 
 private:
+
 	TWeakPtr<FSequencer> Sequencer;
 	/** Interface to section data */
 	TSharedPtr<ISequencerSection> SectionInterface;
@@ -157,7 +166,6 @@ private:
 
 	MovieScene::TNonIntrusiveEventHandler<MovieScene::ISignedObjectEventHandler> TrackModifiedBinding;
 
-
 	friend struct FSequencerSectionPainterImpl;
 };
 
diff --git a/Engine/Source/Editor/Sequencer/Private/Sequencer.cpp b/Engine/Source/Editor/Sequencer/Private/Sequencer.cpp
--- a/Engine/Source/Editor/Sequencer/Private/Sequencer.cpp
+++ b/Engine/Source/Editor/Sequencer/Private/Sequencer.cpp
@@ -154,8 +154,6 @@
 #include "ActorTreeItem.h"
 #include "Widgets/Layout/SSpacer.h"
 #include "CurveEditorCommands.h"
-#include "Tools/SequencerEditTool_Movement.h"
-#include "Tools/SequencerEditTool_Selection.h"
 #include "EntitySystem/MovieSceneEntitySystemRunner.h"
 #include "EntitySystem/MovieSceneEntitySystemLinker.h"
 #include "EntitySystem/MovieScenePreAnimatedStateSystem.h"
@@ -504,9 +502,6 @@ void FSequencer::InitSequencer(const FSequencerInitParams& InitParams, const TSh
 	// When undo occurs, get a notification so we can make sure our view is up to date
 	GEditor->RegisterForUndo(this);
 
-	ViewModel->GetTrackArea()->AddEditTool(MakeShared<FSequencerEditTool_Selection>(*this, *SequencerWidget->GetTrackAreaWidget().Get()));
-	ViewModel->GetTrackArea()->AddEditTool(MakeShared<FSequencerEditTool_Movement>(*this));
-
 	for (int32 DelegateIndex = 0; DelegateIndex < EditorObjectBindingDelegates.Num(); ++DelegateIndex)
 	{
 		check(EditorObjectBindingDelegates[DelegateIndex].IsBound());
@@ -9992,7 +9987,7 @@ void FSequencer::CopySelectedKeys()
 	TOptional<FFrameNumber> CopyRelativeTo;
 
 	// Copy relative to the current key hotspot, if applicable
-	if (TSharedPtr<FKeyHotspot> KeyHotspot = HotspotCast<FKeyHotspot>(GetViewModel()->GetTrackArea()->GetHotspot()))
+	if (TSharedPtr<FKeyHotspot> KeyHotspot = HotspotCast<FKeyHotspot>(ViewModel->GetHotspot()))
 	{
 		CopyRelativeTo = KeyHotspot->GetTime();
 	}
diff --git a/Engine/Source/Editor/Sequencer/Private/SequencerCommonHelpers.cpp b/Engine/Source/Editor/Sequencer/Private/SequencerCommonHelpers.cpp
--- a/Engine/Source/Editor/Sequencer/Private/SequencerCommonHelpers.cpp
+++ b/Engine/Source/Editor/Sequencer/Private/SequencerCommonHelpers.cpp
@@ -209,7 +209,8 @@ void SequencerHelpers::PerformDefaultSelection(FSequencer& Sequencer, const FPoi
 
 	// @todo: selection in transactions
 	FHotspotSelectionManager SelectionManager(&MouseEvent, &Sequencer);
-	TSharedPtr<ITrackAreaHotspot> Hotspot = Sequencer.GetViewModel()->GetTrackArea()->GetHotspot();
+	TSharedPtr<FSequencerEditorViewModel> SequencerViewModel = Sequencer.GetViewModel()->CastThisShared<FSequencerEditorViewModel>();
+	TSharedPtr<ITrackAreaHotspot> Hotspot = SequencerViewModel->GetHotspot();
 	if (TSharedPtr<IMouseHandlerHotspot> MouseHandler = HotspotCast<IMouseHandlerHotspot>(Hotspot))
 	{
 		MouseHandler->HandleMouseSelection(SelectionManager);
@@ -243,7 +244,8 @@ TSharedPtr<SWidget> SequencerHelpers::SummonContextMenu(FSequencer& Sequencer, c
 
 	FMenuBuilder MenuBuilder(bShouldCloseWindowAfterMenuSelection, Sequencer.GetCommandBindings(), MenuExtender, false, &FCoreStyle::Get(), true, NAME_None, bInRecursivelySearchable);
 
-	TSharedPtr<ITrackAreaHotspot> Hotspot = Sequencer.GetViewModel()->GetTrackArea()->GetHotspot();
+	TSharedPtr<FSequencerEditorViewModel> SequencerViewModel = Sequencer.GetViewModel()->CastThisShared<FSequencerEditorViewModel>();
+	TSharedPtr<ITrackAreaHotspot> Hotspot = SequencerViewModel->GetHotspot();
 
 	if (Hotspot.IsValid() && Hotspot->PopulateContextMenu(MenuBuilder, MenuExtender, PasteAtTime))
 	{
diff --git a/Engine/Source/Editor/Sequencer/Private/SequencerContextMenus.cpp b/Engine/Source/Editor/Sequencer/Private/SequencerContextMenus.cpp
--- a/Engine/Source/Editor/Sequencer/Private/SequencerContextMenus.cpp
+++ b/Engine/Source/Editor/Sequencer/Private/SequencerContextMenus.cpp
@@ -158,7 +158,8 @@ void FKeyContextMenu::PopulateMenu(FMenuBuilder& MenuBuilder, TSharedPtr<FExtend
 
 	MenuBuilder.BeginSection("SequencerKeyEdit", LOCTEXT("EditMenu", "Edit"));
 	{
-		if (HotspotCast<FKeyHotspot>(SequencerPtr->GetViewModel()->GetTrackArea()->GetHotspot()))
+		TSharedPtr<FSequencerEditorViewModel> SequencerViewModel = SequencerPtr->GetViewModel()->CastThisShared<FSequencerEditorViewModel>();
+		if (HotspotCast<FKeyHotspot>(SequencerViewModel->GetHotspot()))
 		{
 			MenuBuilder.AddMenuEntry(FGenericCommands::Get().Cut);
 			MenuBuilder.AddMenuEntry(FGenericCommands::Get().Copy);
diff --git a/Engine/Source/Editor/Sequencer/Private/Tools/EditToolDragOperations.cpp b/Engine/Source/Editor/Sequencer/Private/Tools/EditToolDragOperations.cpp
--- a/Engine/Source/Editor/Sequencer/Private/Tools/EditToolDragOperations.cpp
+++ b/Engine/Source/Editor/Sequencer/Private/Tools/EditToolDragOperations.cpp
@@ -619,7 +619,8 @@ FMoveKeysAndSections::FMoveKeysAndSections(FSequencer& InSequencer, ESequencerMo
 {
 	using namespace UE::Sequencer;
 
-	bAllowVerticalMovement = HotspotCast<FSectionHotspotBase>(InSequencer.GetViewModel()->GetTrackArea()->GetHotspot()) != nullptr;
+	TSharedPtr<FSequencerEditorViewModel> SequencerViewModel = InSequencer.GetViewModel()->CastThisShared<FSequencerEditorViewModel>();
+	bAllowVerticalMovement = HotspotCast<FSectionHotspotBase>(SequencerViewModel->GetHotspot()) != nullptr;
 
 	if (EnumHasAnyFlags(MoveType, ESequencerMoveOperationType::MoveKeys))
 	{
diff --git a/Engine/Source/Editor/Sequencer/Private/Tools/SequencerEditTool_Movement.cpp b/Engine/Source/Editor/Sequencer/Private/Tools/SequencerEditTool_Movement.cpp
--- a/Engine/Source/Editor/Sequencer/Private/Tools/SequencerEditTool_Movement.cpp
+++ b/Engine/Source/Editor/Sequencer/Private/Tools/SequencerEditTool_Movement.cpp
@@ -5,6 +5,7 @@
 #include "MVVM/ViewModels/SequencerEditorViewModel.h"
 #include "Editor.h"
 #include "Fonts/FontMeasure.h"
+#include "MVVM/Views/STrackAreaView.h"
 #include "Styling/CoreStyle.h"
 #include "Styling/AppStyle.h"
 #include "SequencerCommonHelpers.h"
@@ -24,8 +25,9 @@
 const FName FSequencerEditTool_Movement::Identifier = "Movement";
 
 
-FSequencerEditTool_Movement::FSequencerEditTool_Movement(FSequencer& InSequencer)
+FSequencerEditTool_Movement::FSequencerEditTool_Movement(FSequencer& InSequencer, UE::Sequencer::STrackAreaView& InTrackArea)
 	: FSequencerEditTool(InSequencer)
+	, TrackArea(InTrackArea)
 	, CursorDecorator(nullptr)
 { }
 
@@ -36,13 +38,14 @@ FReply FSequencerEditTool_Movement::OnMouseButtonDown(SWidget& OwnerWidget, cons
 
 	TSharedRef<SSequencer> SequencerWidget = StaticCastSharedRef<SSequencer>(Sequencer.GetSequencerWidget());
 
-	TSharedPtr<ITrackAreaHotspot> Hotspot = Sequencer.GetViewModel()->GetTrackArea()->GetHotspot();
+	TSharedPtr<FTrackAreaViewModel> TrackAreaViewModel = TrackArea.GetViewModel();
+	TSharedPtr<ITrackAreaHotspot> Hotspot = TrackAreaViewModel->GetHotspot();
 
 	DelayedDrag.Reset();
 
 	if (MouseEvent.GetEffectingButton() == EKeys::LeftMouseButton || MouseEvent.GetEffectingButton() == EKeys::MiddleMouseButton)
 	{
-		const FVirtualTrackArea VirtualTrackArea = SequencerWidget->GetVirtualTrackArea();
+		const FVirtualTrackArea VirtualTrackArea = SequencerWidget->GetVirtualTrackArea(&TrackArea);
 
 		DelayedDrag = FDelayedDrag_Hotspot(VirtualTrackArea.CachedTrackAreaGeometry().AbsoluteToLocal(MouseEvent.GetScreenSpacePosition()), MouseEvent.GetEffectingButton(), Hotspot);
 
@@ -75,7 +78,7 @@ FReply FSequencerEditTool_Movement::OnMouseMove(SWidget& OwnerWidget, const FGeo
 	if (DelayedDrag.IsSet())
 	{
 		TSharedRef<SSequencer> SequencerWidget = StaticCastSharedRef<SSequencer>(Sequencer.GetSequencerWidget());
-		const FVirtualTrackArea VirtualTrackArea = SequencerWidget->GetVirtualTrackArea();
+		const FVirtualTrackArea VirtualTrackArea = SequencerWidget->GetVirtualTrackArea(&TrackArea);
 
 		FReply Reply = FReply::Handled();
 
@@ -286,13 +289,13 @@ FReply FSequencerEditTool_Movement::OnMouseButtonUp(SWidget& OwnerWidget, const
 
 	DelayedDrag.Reset();
 
-	TSharedPtr<FTrackAreaViewModel> TrackArea = Sequencer.GetViewModel()->GetTrackArea();
+	TSharedPtr<FTrackAreaViewModel> TrackAreaViewModel = TrackArea.GetViewModel();
 
 	if (DragOperation.IsValid())
 	{
 		TSharedRef<SSequencer> SequencerWidget = StaticCastSharedRef<SSequencer>(Sequencer.GetSequencerWidget());
 
-		DragOperation->OnEndDrag(MouseEvent, MyGeometry.AbsoluteToLocal(MouseEvent.GetScreenSpacePosition()), SequencerWidget->GetVirtualTrackArea());
+		DragOperation->OnEndDrag(MouseEvent, MyGeometry.AbsoluteToLocal(MouseEvent.GetScreenSpacePosition()), SequencerWidget->GetVirtualTrackArea(&TrackArea));
 		DragOperation = nullptr;
 
 		if (MouseEvent.GetEffectingButton() == EKeys::MiddleMouseButton)
@@ -324,18 +327,18 @@ FReply FSequencerEditTool_Movement::OnMouseButtonUp(SWidget& OwnerWidget, const
 				);
 
 			// Lock the hotspot while the menu is open
-			TrackArea->LockHotspot(true);
+			TrackAreaViewModel->LockHotspot(true);
 
 			// Unlock and reset the hotspot when the menu closes
 			{
-				TSharedPtr<ITrackAreaHotspot> ExistingHotspot = TrackArea->GetHotspot();
+				TSharedPtr<ITrackAreaHotspot> ExistingHotspot = TrackAreaViewModel->GetHotspot();
 				Menu->GetOnMenuDismissed().AddLambda(
-					[TrackArea, ExistingHotspot](TSharedRef<IMenu>)
+					[TrackAreaViewModel, ExistingHotspot](TSharedRef<IMenu>)
 					{
-						TrackArea->LockHotspot(false);
-						if (TrackArea->GetHotspot() == ExistingHotspot)
+						TrackAreaViewModel->LockHotspot(false);
+						if (TrackAreaViewModel->GetHotspot() == ExistingHotspot)
 						{
-							TrackArea->SetHotspot(nullptr);
+							TrackAreaViewModel->SetHotspot(nullptr);
 						}
 					}
 				);
@@ -398,7 +401,7 @@ int32 FSequencerEditTool_Movement::OnPaint(const FGeometry& AllottedGeometry, co
 				const float MousePadding = 20.0f;
 
 				// calculate draw position
-				const FVirtualTrackArea VirtualTrackArea = SequencerWidget->GetVirtualTrackArea();
+				const FVirtualTrackArea VirtualTrackArea = SequencerWidget->GetVirtualTrackArea(&TrackArea);
 				const float HorizontalDelta = DragPosition.X - DelayedDrag->GetInitialPosition().X;
 				const float InitialY = DelayedDrag->GetInitialPosition().Y;
 
@@ -494,10 +497,10 @@ void FSequencerEditTool_Movement::UpdateCursorDecorator(const FGeometry& MyGeome
 	using namespace UE::Sequencer;
 
 	DragPosition = MyGeometry.AbsoluteToLocal(CursorEvent.GetScreenSpacePosition());
-	TSharedPtr<FTrackAreaViewModel> TrackArea = Sequencer.GetViewModel()->GetTrackArea();
+	TSharedPtr<FTrackAreaViewModel> TrackAreaViewModel = TrackArea.GetViewModel();
 	TSharedPtr<ITrackAreaHotspot> Hotspot = DelayedDrag.IsSet()
 		? DelayedDrag->Hotspot
-		: TrackArea->GetHotspot();
+		: TrackAreaViewModel->GetHotspot();
 
 	if (Hotspot.IsValid())
 	{
diff --git a/Engine/Source/Editor/Sequencer/Private/Tools/SequencerEditTool_Movement.h b/Engine/Source/Editor/Sequencer/Private/Tools/SequencerEditTool_Movement.h
--- a/Engine/Source/Editor/Sequencer/Private/Tools/SequencerEditTool_Movement.h
+++ b/Engine/Source/Editor/Sequencer/Private/Tools/SequencerEditTool_Movement.h
@@ -10,6 +10,14 @@
 
 class SSequencer;
 
+namespace UE
+{
+namespace Sequencer
+{
+	class STrackAreaView;
+}
+}
+
 class FSequencerEditTool_Movement
 	: public FSequencerEditTool
 {
@@ -19,7 +27,7 @@ public:
 	static const FName Identifier;
 
 	/** Create and initialize a new instance. */
-	FSequencerEditTool_Movement(FSequencer& InSequencer);
+	FSequencerEditTool_Movement(FSequencer& InSequencer, UE::Sequencer::STrackAreaView& InTrackArea);
 
 public:
 
@@ -59,6 +67,9 @@ private:
 	/** The hotspot's time before dragging started. */
 	FFrameTime OriginalHotspotTime;
 
+	/** TrackArea this object belongs to */
+	UE::Sequencer::STrackAreaView& TrackArea;
+
 	/** Software cursor decorator brush */
 	const FSlateBrush* CursorDecorator;
 };
diff --git a/Engine/Source/Editor/Sequencer/Private/Tools/SequencerEditTool_Selection.cpp b/Engine/Source/Editor/Sequencer/Private/Tools/SequencerEditTool_Selection.cpp
--- a/Engine/Source/Editor/Sequencer/Private/Tools/SequencerEditTool_Selection.cpp
+++ b/Engine/Source/Editor/Sequencer/Private/Tools/SequencerEditTool_Selection.cpp
@@ -359,7 +359,8 @@ FReply FSequencerEditTool_Selection::OnMouseButtonDown(SWidget& OwnerWidget, con
 
 	if (MouseEvent.GetEffectingButton() == EKeys::LeftMouseButton)
 	{
-		TSharedPtr<ITrackAreaHotspot> Hotspot = Sequencer.GetViewModel()->GetTrackArea()->GetHotspot();
+		TSharedPtr<FSequencerEditorViewModel> SequencerViewModel = Sequencer.GetViewModel()->CastThisShared<FSequencerEditorViewModel>();
+		TSharedPtr<ITrackAreaHotspot> Hotspot = SequencerViewModel->GetHotspot();
 		DelayedDrag = FDelayedDrag_Hotspot(MyGeometry.AbsoluteToLocal(MouseEvent.GetScreenSpacePosition()), EKeys::LeftMouseButton, Hotspot);
 		return FReply::Handled();
 	}
diff --git a/Engine/Source/Editor/Sequencer/Private/Widgets/SSequencerLayerBar.cpp b/Engine/Source/Editor/Sequencer/Private/Widgets/SSequencerLayerBar.cpp
--- a/Engine/Source/Editor/Sequencer/Private/Widgets/SSequencerLayerBar.cpp
+++ b/Engine/Source/Editor/Sequencer/Private/Widgets/SSequencerLayerBar.cpp
@@ -6,13 +6,14 @@
 #include "Tools/DragOperation_Stretch.h"
 #include "Tools/SequencerEditTool_Movement.h"
 
-#include "MVVM/ViewModels/ViewModel.h"
-#include "MVVM/ViewModels/SequenceModel.h"
-#include "MVVM/ViewModels/ObjectBindingModel.h"
-#include "MVVM/ViewModels/LayerBarModel.h"
 #include "MVVM/ObjectBindingModelStorageExtension.h"
+#include "MVVM/ViewModels/LayerBarModel.h"
+#include "MVVM/ViewModels/ObjectBindingModel.h"
+#include "MVVM/ViewModels/SequenceModel.h"
 #include "MVVM/ViewModels/SequencerEditorViewModel.h"
 #include "MVVM/ViewModels/TrackAreaViewModel.h"
+#include "MVVM/ViewModels/ViewModel.h"
+#include "MVVM/Views/STrackAreaView.h"
 
 #include "CommonMovieSceneTools.h"
 #include "SequencerCommonHelpers.h"
@@ -162,8 +163,9 @@ public:
 		SLATE_ATTRIBUTE(const FSlateBrush*, BackgroundBrush)
 	SLATE_END_ARGS()
 
-	void Construct(const FArguments& InArgs, TWeakPtr<SSequencerLayerBar> InWeakParent, EStretchConstraint InHandleType)
+	void Construct(const FArguments& InArgs, TWeakPtr<FTrackAreaViewModel> InWeakTrackArea, TWeakPtr<SSequencerLayerBar> InWeakParent, EStretchConstraint InHandleType)
 	{
+		WeakTrackArea = InWeakTrackArea;
 		WeakParent = InWeakParent;
 		HandleType = InHandleType;
 
@@ -181,9 +183,9 @@ public:
 		if (LayerBarWidget)
 		{
 			TSharedPtr<FLayerBarModel> LayerBar  = LayerBarWidget->GetLayerBarModel();
-			TSharedPtr<FTrackAreaViewModel> TrackArea = Sequencer->GetViewModel()->GetTrackArea();
 			if (LayerBar)
 			{
+				TSharedPtr<FTrackAreaViewModel> TrackArea = WeakTrackArea.Pin();
 				TrackArea->AddHotspot(MakeShared<FLayerBarStretchHotspot>(LayerBar, HandleType, Sequencer));
 			}
 		}
@@ -197,21 +199,25 @@ public:
 		TSharedPtr<FLayerBarModel> LayerBar  = LayerBarWidget->GetLayerBarModel();
 		if (LayerBarWidget)
 		{
-			TSharedPtr<FTrackAreaViewModel> TrackArea = LayerBarWidget->GetSequencer()->GetViewModel()->GetTrackArea();
+			TSharedPtr<FTrackAreaViewModel> TrackArea = WeakTrackArea.Pin();
 			TrackArea->RemoveHotspot(FLayerBarStretchHotspot::ID);
 		}
 
 		SCompoundWidget::OnMouseLeave(MouseEvent);
 	}
 
+	TWeakPtr<FTrackAreaViewModel> WeakTrackArea;
 	TWeakPtr<SSequencerLayerBar> WeakParent;
 	EStretchConstraint HandleType;
 };
 
-void SSequencerLayerBar::Construct(const FArguments& InArgs, TWeakPtr<FSequencerEditorViewModel> InWeakEditor, TWeakPtr<FLayerBarModel> InWeakLayerBar)
+void SSequencerLayerBar::Construct(const FArguments& InArgs, TWeakPtr<STrackAreaView> InWeakTrackArea, TWeakPtr<FSequencerEditorViewModel> InWeakEditor, TWeakPtr<FLayerBarModel> InWeakLayerBar)
 {
 	WeakLayerBar   = InWeakLayerBar;
 	WeakEditor     = InWeakEditor;
+	WeakTrackArea  = InWeakTrackArea;
+
+	TSharedPtr<STrackAreaView> TrackArea = WeakTrackArea.Pin();
 
 	ChildSlot
 	[
@@ -225,7 +231,7 @@ void SSequencerLayerBar::Construct(const FArguments& InArgs, TWeakPtr<FSequencer
 			SNew(SBox)
 			.WidthOverride(8.f)
 			[
-				SNew(SSequencerLayerBarHandle, SharedThis(this), EStretchConstraint::AnchorToEnd)
+				SNew(SSequencerLayerBarHandle, TrackArea->GetViewModel(), SharedThis(this), EStretchConstraint::AnchorToEnd)
 				.BackgroundBrush(FAppStyle::Get().GetBrush("Sequencer.LayerBar.HandleLeft"))
 			]
 		]
@@ -242,7 +248,7 @@ void SSequencerLayerBar::Construct(const FArguments& InArgs, TWeakPtr<FSequencer
 			SNew(SBox)
 			.WidthOverride(8.f)
 			[
-				SNew(SSequencerLayerBarHandle, SharedThis(this), EStretchConstraint::AnchorToStart)
+				SNew(SSequencerLayerBarHandle, TrackArea->GetViewModel(), SharedThis(this), EStretchConstraint::AnchorToStart)
 				.BackgroundBrush(FAppStyle::Get().GetBrush("Sequencer.LayerBar.HandleRight"))
 			]
 		]
@@ -350,7 +356,7 @@ void SSequencerLayerBar::OnMouseEnter(const FGeometry& MyGeometry, const FPointe
 	if (LayerBar)
 	{
 		TSharedPtr<FSequencer> Sequencer = GetSequencer();
-		TSharedPtr<FTrackAreaViewModel> TrackArea = Sequencer->GetViewModel()->GetTrackArea();
+		TSharedPtr<FTrackAreaViewModel> TrackArea = WeakTrackArea.Pin()->GetViewModel();
 		TrackArea->AddHotspot(MakeShared<FLayerBarHotspot>(LayerBar, Sequencer));
 	}
 
@@ -367,7 +373,7 @@ void SSequencerLayerBar::OnMouseLeave(const FPointerEvent& MouseEvent)
 	TSharedPtr<FLayerBarModel> LayerBar = WeakLayerBar.Pin();
 	if (LayerBar)
 	{
-		TSharedPtr<FTrackAreaViewModel> TrackArea = GetSequencer()->GetViewModel()->GetTrackArea();
+		TSharedPtr<FTrackAreaViewModel> TrackArea = WeakTrackArea.Pin()->GetViewModel();
 		TrackArea->RemoveHotspot(FLayerBarHotspot::ID);
 	}
 
diff --git a/Engine/Source/Editor/Sequencer/Private/Widgets/SSequencerLayerBar.h b/Engine/Source/Editor/Sequencer/Private/Widgets/SSequencerLayerBar.h
--- a/Engine/Source/Editor/Sequencer/Private/Widgets/SSequencerLayerBar.h
+++ b/Engine/Source/Editor/Sequencer/Private/Widgets/SSequencerLayerBar.h
@@ -26,6 +26,7 @@ class FLayerBarModel;
 class FSequencerEditorViewModel;
 class FViewModel;
 class ILayerBarExtension;
+class STrackAreaView;
 
 class SSequencerLayerBar
 	: public SCompoundWidget
@@ -35,7 +36,7 @@ public:
 	SLATE_BEGIN_ARGS(SSequencerLayerBar) {}
 	SLATE_END_ARGS()
 
-	void Construct(const FArguments& InArgs, TWeakPtr<FSequencerEditorViewModel> InWeakEditor, TWeakPtr<FLayerBarModel> InWeakLayerBar);
+	void Construct(const FArguments& InArgs, TWeakPtr<STrackAreaView> InWeakTrackArea, TWeakPtr<FSequencerEditorViewModel> InWeakEditor, TWeakPtr<FLayerBarModel> InWeakLayerBar);
 
 	TSharedPtr<FLayerBarModel> GetLayerBarModel() const;
 	TSharedPtr<FSequencer> GetSequencer() const;
@@ -61,6 +62,7 @@ private:
 
 	TWeakPtr<FLayerBarModel> WeakLayerBar;
 	TWeakPtr<FSequencerEditorViewModel> WeakEditor;
+	TWeakPtr<STrackAreaView> WeakTrackArea;
 };
 
 } // namespace Sequencer
diff --git a/Engine/Source/Editor/Sequencer/Public/MVVM/ViewModels/SequencerEditorViewModel.h b/Engine/Source/Editor/Sequencer/Public/MVVM/ViewModels/SequencerEditorViewModel.h
--- a/Engine/Source/Editor/Sequencer/Public/MVVM/ViewModels/SequencerEditorViewModel.h
+++ b/Engine/Source/Editor/Sequencer/Public/MVVM/ViewModels/SequencerEditorViewModel.h
@@ -17,6 +17,8 @@ namespace Sequencer
 {
 
 class FSequenceModel;
+class STrackAreaView;
+struct ITrackAreaHotspot;
 
 /**
  * Main view-model for the Sequencer editor.
@@ -37,18 +39,31 @@ public:
 	// @todo_sequencer_mvvm move this to the root view-model
 	void SetSequence(UMovieSceneSequence* InRootSequence);
 
+	/** Gets the pinned track area view-model. */
+	TSharedPtr<FTrackAreaViewModel> GetPinnedTrackArea() const;
+
+	/** Gets the current hotspots across any of our track areas */
+	TSharedPtr<ITrackAreaHotspot> GetHotspot() const;
+
 protected:
 
 	virtual void PreInitializeEditorImpl() override;
+	virtual void InitializeEditorImpl() override;
 	virtual TSharedPtr<FViewModel> CreateRootModelImpl() override;
 	virtual TSharedPtr<FOutlinerViewModel> CreateOutlinerImpl() override;
 	virtual TSharedPtr<FTrackAreaViewModel> CreateTrackAreaImpl() override;
 	virtual bool IsReadOnly() const override;
 
+	void OnTrackAreaHotspotChanged(TSharedPtr<ITrackAreaHotspot> NewHotspot);
+
 protected:
 
 	TWeakPtr<ISequencer> WeakSequencer;
+	TSharedPtr<FTrackAreaViewModel> PinnedTrackArea;
 	bool bSupportsCurveEditor;
+
+	/** The current hotspot, from any of our track areas */
+	TSharedPtr<ITrackAreaHotspot> CurrentHotspot;
 };
 
 } // namespace Sequencer
diff --git a/Engine/Source/Editor/Sequencer/Public/MVVM/ViewModels/SequencerTrackAreaViewModel.h b/Engine/Source/Editor/Sequencer/Public/MVVM/ViewModels/SequencerTrackAreaViewModel.h
--- a/Engine/Source/Editor/Sequencer/Public/MVVM/ViewModels/SequencerTrackAreaViewModel.h
+++ b/Engine/Source/Editor/Sequencer/Public/MVVM/ViewModels/SequencerTrackAreaViewModel.h
@@ -14,6 +14,8 @@ namespace UE
 namespace Sequencer
 {
 
+class STrackAreaView;
+
 class FSequencerTrackAreaViewModel
 	: public FTrackAreaViewModel
 {
@@ -23,6 +25,8 @@ public:
 
 	FSequencerTrackAreaViewModel(TSharedRef<ISequencer> InSequencer);
 
+	void InitializeDefaultEditTools(UE::Sequencer::STrackAreaView& InTrackArea);
+
 	virtual FFrameRate GetTickResolution() const override;
 	virtual TRange<double> GetViewRange() const override;
 
diff --git a/Engine/Source/Editor/SequencerCore/Private/MVVM/ViewModels/TrackAreaViewModel.cpp b/Engine/Source/Editor/SequencerCore/Private/MVVM/ViewModels/TrackAreaViewModel.cpp
--- a/Engine/Source/Editor/SequencerCore/Private/MVVM/ViewModels/TrackAreaViewModel.cpp
+++ b/Engine/Source/Editor/SequencerCore/Private/MVVM/ViewModels/TrackAreaViewModel.cpp
@@ -72,6 +72,8 @@ void FTrackAreaViewModel::SetHotspot(TSharedPtr<ITrackAreaHotspot> NewHotspot)
 			// Simulate an update-on-hover for the new hotspot to ensure that any hover behavior doesn't have to wait until the next frame
 			NewHotspot->UpdateOnHover(*this);
 		}
+
+		OnHotspotChangedDelegate.Broadcast(NewHotspot);
 	}
 }
 
@@ -115,6 +117,11 @@ void FTrackAreaViewModel::LockHotspot(bool bIsLocked)
 	bHotspotLocked = bIsLocked;
 }
 
+FOnTrackAreaHotspotChanged& FTrackAreaViewModel::GetOnHotspotChangedDelegate()
+{
+	return OnHotspotChangedDelegate;
+}
+
 bool FTrackAreaViewModel::CanActivateEditTool(FName Identifier) const
 {
 	auto IdentifierIsFound = [=](TSharedPtr<ISequencerEditTool> InEditTool){ return InEditTool->GetIdentifier() == Identifier; };
diff --git a/Engine/Source/Editor/SequencerCore/Public/MVVM/ViewModels/TrackAreaViewModel.h b/Engine/Source/Editor/SequencerCore/Public/MVVM/ViewModels/TrackAreaViewModel.h
--- a/Engine/Source/Editor/SequencerCore/Public/MVVM/ViewModels/TrackAreaViewModel.h
+++ b/Engine/Source/Editor/SequencerCore/Public/MVVM/ViewModels/TrackAreaViewModel.h
@@ -25,6 +25,8 @@ class FEditorViewModel;
 class ISequencerEditTool;
 struct ITrackAreaHotspot;
 
+DECLARE_MULTICAST_DELEGATE_OneParam(FOnTrackAreaHotspotChanged, TSharedPtr<ITrackAreaHotspot>);
+
 class SEQUENCERCORE_API FTrackAreaViewModel
 	: public FViewModel
 {
@@ -60,6 +62,9 @@ public:
 	/** Set whether the hotspot is locked and cannot be changed (ie when a menu is open) */
 	void LockHotspot(bool bIsLocked);
 
+	/** Get the callback that is fired when the current hotspot changes */
+	FOnTrackAreaHotspotChanged& GetOnHotspotChangedDelegate();
+
 	void AddEditTool(TSharedPtr<ISequencerEditTool> InNewTool);
 
 	/** Access the currently active track area edit tool */
@@ -78,6 +83,9 @@ protected:
 	/** The current hotspot that can be set from anywhere to initiate drags */
 	TArray<TSharedPtr<ITrackAreaHotspot>> HotspotStack;
 
+	/** The delegate that is fired when the current hotspot changes */
+	FOnTrackAreaHotspotChanged OnHotspotChangedDelegate;
+
 	/** The currently active edit tools on this track area */
 	TArray<TSharedPtr<ISequencerEditTool>> EditTools;
 
