Fixed not calling the NotifyMovieSceneDataChanged() function when pasting tracks.
This fixes an issue with Niagara not setting up its data correctly after pasting a track.

diff --git a/Engine/Source/Editor/Sequencer/Private/Sequencer.cpp b/Engine/Source/Editor/Sequencer/Private/Sequencer.cpp
--- a/Engine/Source/Editor/Sequencer/Private/Sequencer.cpp
+++ b/Engine/Source/Editor/Sequencer/Private/Sequencer.cpp
@@ -8115,7 +8115,14 @@ bool FSequencer::PasteTracks(const FString& TextToImport, UMovieSceneFolder* InP
 	}
 
 	TArray<UMovieSceneTrack*> OutTracks;
-	return FSequencerUtilities::PasteTracks(TextToImport, FMovieScenePasteTracksParams(GetFocusedMovieSceneSequence(), TargetBindings, InParentFolder, InFolders), OutTracks, PasteErrors);
+	bool bSuccess = FSequencerUtilities::PasteTracks(TextToImport, FMovieScenePasteTracksParams(GetFocusedMovieSceneSequence(), TargetBindings, InParentFolder, InFolders), OutTracks, PasteErrors);
+
+	if(bSuccess)
+	{
+		NotifyMovieSceneDataChanged(EMovieSceneDataChangeType::MovieSceneStructureItemAdded);
+	}
+	
+	return bSuccess;
 }
 
 bool FSequencer::PasteSections(const FString& TextToImport, TArray<FNotificationInfo>& PasteErrors)
