Sequencer: Don't allow infinite sections to be selectable and then movable

diff --git a/Engine/Source/Editor/Sequencer/Private/MVVM/ViewModels/SectionModel.cpp b/Engine/Source/Editor/Sequencer/Private/MVVM/ViewModels/SectionModel.cpp
--- a/Engine/Source/Editor/Sequencer/Private/MVVM/ViewModels/SectionModel.cpp
+++ b/Engine/Source/Editor/Sequencer/Private/MVVM/ViewModels/SectionModel.cpp
@@ -189,8 +189,8 @@ ESelectionIntent FSectionModel::IsSelectable() const
 {
 	if (GetRange() == TRange<FFrameNumber>::All())
 	{
-		// Infinite sections also now selectable in all contexts
-		return ESelectionIntent::Any;
+		// Infinite sections are only selectable through the context menu
+		return ESelectionIntent::ContextMenu;
 	}
 	return ESelectionIntent::Any;
 }
diff --git a/Engine/Source/Editor/Sequencer/Private/SequencerHotspots.cpp b/Engine/Source/Editor/Sequencer/Private/SequencerHotspots.cpp
--- a/Engine/Source/Editor/Sequencer/Private/SequencerHotspots.cpp
+++ b/Engine/Source/Editor/Sequencer/Private/SequencerHotspots.cpp
@@ -481,12 +481,7 @@ void FSectionHotspotBase::UpdateOnHover(FTrackAreaViewModel& InTrackArea) const
 			}
 		}
 
-		// Activate selection mode if the section is infinite, otherwise just move it
-		if (ThisSection->GetRange() == TRange<FFrameNumber>::All())
-		{
-			InTrackArea.AttemptToActivateTool(FSequencerEditTool_Selection::Identifier);
-		}
-		else
+		if (ThisSection->GetRange() != TRange<FFrameNumber>::All())
 		{
 			InTrackArea.AttemptToActivateTool(FSequencerEditTool_Movement::Identifier);
 		}
