Sequencer: Fix bad merge

diff --git a/Engine/Source/Editor/Sequencer/Private/SequencerUtilities.cpp b/Engine/Source/Editor/Sequencer/Private/SequencerUtilities.cpp
--- a/Engine/Source/Editor/Sequencer/Private/SequencerUtilities.cpp
+++ b/Engine/Source/Editor/Sequencer/Private/SequencerUtilities.cpp
@@ -2171,12 +2171,8 @@ bool FSequencerUtilities::PasteBindings(const FString& TextToImport, TSharedRef<
 							}
 
 							ActorsToRebind.Add(Actor);
-							CopyableBinding->BoundObjectNames.Remove(Actor->GetName());
+							CopyableBinding->BoundObjectNames.Remove(Actor->GetPathName());
 						}
-
-						ActorsToRebind.Add(Actor);
-						CopyableBinding->BoundObjectNames.Remove(Actor->GetPathName());
-
 					}
 				}
 
