Sequencer: Fix BringForward crash when iterating through the sections backwards. Looks to be accessing SectionIndex + 1

diff --git a/Engine/Source/Editor/Sequencer/Private/SequencerContextMenus.cpp b/Engine/Source/Editor/Sequencer/Private/SequencerContextMenus.cpp
--- a/Engine/Source/Editor/Sequencer/Private/SequencerContextMenus.cpp
+++ b/Engine/Source/Editor/Sequencer/Private/SequencerContextMenus.cpp
@@ -1228,7 +1228,7 @@ void FSectionContextMenu::BringForward()
 				return A.GetOverlapPriority() < B.GetOverlapPriority();
 			});
 
-			for (int32 SectionIndex = Row.Sections.Num() - 1; SectionIndex > 0; --SectionIndex)
+			for (int32 SectionIndex = Row.Sections.Num() - 2; SectionIndex > 0; --SectionIndex)
 			{
 				UMovieSceneSection* ThisSection = Row.Sections[SectionIndex];
 				if (Row.SectionToReOrder.Contains(ThisSection))
@@ -1293,7 +1293,6 @@ void FSectionContextMenu::SendBackward()
 	Sequencer->SetLocalTimeDirectly(Sequencer->GetLocalTime().Time);
 }
 
-
 bool FPasteContextMenu::BuildMenu(FMenuBuilder& MenuBuilder, TSharedPtr<FExtender> MenuExtender, FSequencer& InSequencer, const FPasteContextMenuArgs& Args)
 {
 	TSharedRef<FPasteContextMenu> Menu = MakeShareable(new FPasteContextMenu(InSequencer, Args));
