Sequencer: Remove rollover visibility of Add button text. Now the text is just in the tooltip so that there's no overlap with key editor buttons

diff --git a/Engine/Source/Editor/MovieSceneTools/Private/Channels/BuiltInChannelEditors.cpp b/Engine/Source/Editor/MovieSceneTools/Private/Channels/BuiltInChannelEditors.cpp
--- a/Engine/Source/Editor/MovieSceneTools/Private/Channels/BuiltInChannelEditors.cpp
+++ b/Engine/Source/Editor/MovieSceneTools/Private/Channels/BuiltInChannelEditors.cpp
@@ -8,6 +8,7 @@
 #include "Sections/MovieSceneEventSectionBase.h"
 #include "ISequencerChannelInterface.h"
 #include "Widgets/Input/SComboButton.h"
+#include "Widgets/Layout/SSpacer.h"
 #include "Widgets/SNullWidget.h"
 #include "IKeyArea.h"
 #include "ISequencer.h"
@@ -423,13 +424,26 @@ TSharedRef<SWidget> CreateKeyEditor(const TMovieSceneChannelHandle<FMovieSceneOb
 			AssetDataArray.Add((FAssetData)Sequence);
 		}
 
-		return SNew(SObjectPropertyEntryBox)
-		.DisplayBrowse(true)
-		.DisplayUseSelected(false)
-		.ObjectPath_Lambda(GetObjectPathLambda)
-		.AllowedClass(RawChannel->GetPropertyClass())
-		.OnObjectChanged_Lambda(OnSetObjectLambda)
-		.OwnerAssetDataArray(AssetDataArray);
+		return SNew(SHorizontalBox)
+
+			+ SHorizontalBox::Slot()
+			.AutoWidth()
+			[
+				SNew(SObjectPropertyEntryBox)
+				.DisplayBrowse(true)
+				.DisplayUseSelected(false)
+				.ObjectPath_Lambda(GetObjectPathLambda)
+				.AllowedClass(RawChannel->GetPropertyClass())
+				.OnObjectChanged_Lambda(OnSetObjectLambda)
+				.OwnerAssetDataArray(AssetDataArray)
+			]
+
+			+ SHorizontalBox::Slot()
+			.AutoWidth()
+			.Padding(8.0f, 0.0f)
+			[
+				SNew(SSpacer)
+			];
 	}
 
 	return SNullWidget::NullWidget;
diff --git a/Engine/Source/Editor/Sequencer/Private/SequencerUtilities.cpp b/Engine/Source/Editor/Sequencer/Private/SequencerUtilities.cpp
--- a/Engine/Source/Editor/Sequencer/Private/SequencerUtilities.cpp
+++ b/Engine/Source/Editor/Sequencer/Private/SequencerUtilities.cpp
@@ -62,19 +62,6 @@
 
 #define LOCTEXT_NAMESPACE "FSequencerUtilities"
 
-static EVisibility GetRolloverVisibility(TAttribute<bool> HoverState, TWeakPtr<SButton> WeakButton)
-{
-	TSharedPtr<SButton> Button = WeakButton.Pin();
-	if (HoverState.Get())
-	{
-		return EVisibility::SelfHitTestInvisible;
-	}
-	else
-	{
-		return EVisibility::Collapsed;
-	}
-}
-
 static void ResetCopiedTracksFlags(UMovieSceneTrack* Track)
 {
 	TArray<UObject*> SectionSubObjects;
@@ -94,13 +81,6 @@ TSharedRef<SWidget> FSequencerUtilities::MakeAddButton(FText HoverText, FOnGetCo
 
 TSharedRef<SWidget> FSequencerUtilities::MakeAddButton(FText HoverText, FOnClicked OnClicked, const TAttribute<bool>& HoverState, TWeakPtr<ISequencer> InSequencer)
 {
-	FSlateFontInfo SmallLayoutFont = FCoreStyle::GetDefaultFontStyle("Regular", 8);
-
-	TSharedRef<STextBlock> ButtonText = SNew(STextBlock)
-		.Text(HoverText)
-		.Font(SmallLayoutFont)
-		.ColorAndOpacity(FSlateColor::UseForeground());
-
 	TSharedRef<SButton> Button =
 
 		SNew(SButton)
@@ -124,19 +104,10 @@ TSharedRef<SWidget> FSequencerUtilities::MakeAddButton(FText HoverText, FOnClick
 				SNew(SImage)
 				.ColorAndOpacity(FSlateColor::UseForeground())
 				.Image(FAppStyle::GetBrush("Plus"))
-			]
-
-			+ SHorizontalBox::Slot()
-			.VAlign(VAlign_Center)
-			.AutoWidth()
-			[
-				ButtonText
+				.ToolTipText(HoverText)
 			]
 		];
 
-	TAttribute<EVisibility> Visibility = TAttribute<EVisibility>::Create(TAttribute<EVisibility>::FGetter::CreateStatic(GetRolloverVisibility, HoverState, TWeakPtr<SButton>(Button)));
-	ButtonText->SetVisibility(Visibility);
-
 	return Button;
 }
 
@@ -1983,6 +1954,17 @@ void FSequencerUtilities::CopyBindings(TSharedRef<ISequencer> Sequencer, const T
 
 				CopyableBinding->Tracks.Add(DuplicatedTrack);
 			}
+
+			for (const TTuple<FName, FMovieSceneObjectBindingIDs>& Pair : MovieScene->AllTaggedBindings())
+			{
+				for (const FMovieSceneObjectBindingID& ID : Pair.Value.IDs)
+				{
+					if (ID.GetGuid() == ObjectBinding.BindingID)
+					{
+						CopyableBinding->Tags.AddUnique(Pair.Key);
+					}
+				}
+			}
 		}
 
 		UMovieSceneFolder* Folder = nullptr;
diff --git a/Engine/Source/Editor/SequencerCore/Private/MVVM/Views/ViewUtilities.cpp b/Engine/Source/Editor/SequencerCore/Private/MVVM/Views/ViewUtilities.cpp
--- a/Engine/Source/Editor/SequencerCore/Private/MVVM/Views/ViewUtilities.cpp
+++ b/Engine/Source/Editor/SequencerCore/Private/MVVM/Views/ViewUtilities.cpp
@@ -29,26 +29,6 @@ namespace Sequencer
 
 TSharedRef<SWidget> MakeAddButton(FText HoverText, FOnGetContent MenuContent, const TAttribute<bool>& HoverState, const TAttribute<bool>& IsEnabled)
 {
-	auto GetRolloverVisibility = [](TAttribute<bool> InHoverState, TWeakPtr<SComboButton> WeakComboButton)
-	{
-		TSharedPtr<SComboButton> TmpComboButton = WeakComboButton.Pin();
-		if (InHoverState.Get() || TmpComboButton->IsOpen())
-		{
-			return EVisibility::SelfHitTestInvisible;
-		}
-		else
-		{
-			return EVisibility::Collapsed;
-		}
-	};
-
-	FSlateFontInfo SmallLayoutFont = FCoreStyle::GetDefaultFontStyle("Regular", 8);
-
-	TSharedRef<STextBlock> ComboButtonText = SNew(STextBlock)
-		.Text(HoverText)
-		.Font(SmallLayoutFont)
-		.ColorAndOpacity( FSlateColor::UseForeground() );
-
 	TSharedRef<SComboButton> ComboButton =
 
 		SNew(SComboButton)
@@ -73,18 +53,10 @@ TSharedRef<SWidget> MakeAddButton(FText HoverText, FOnGetContent MenuContent, co
 				SNew(SImage)
 				.ColorAndOpacity( FSlateColor::UseForeground() )
 				.Image(FAppStyle::GetBrush("Plus"))
-			]
-
-			+ SHorizontalBox::Slot()
-			.VAlign(VAlign_Center)
-			.AutoWidth()
-			[
-				ComboButtonText
+				.ToolTipText(HoverText)
 			]
 		];
 
-	ComboButtonText->SetVisibility(MakeAttributeLambda(GetRolloverVisibility, HoverState, TWeakPtr<SComboButton>(ComboButton)));
-
 	return ComboButton;
 }
 
