Remove unintended

diff --git a/Engine/Source/Editor/Sequencer/Private/SequencerUtilities.cpp b/Engine/Source/Editor/Sequencer/Private/SequencerUtilities.cpp
--- a/Engine/Source/Editor/Sequencer/Private/SequencerUtilities.cpp
+++ b/Engine/Source/Editor/Sequencer/Private/SequencerUtilities.cpp
@@ -1954,17 +1954,6 @@ void FSequencerUtilities::CopyBindings(TSharedRef<ISequencer> Sequencer, const T
 
 				CopyableBinding->Tracks.Add(DuplicatedTrack);
 			}
-
-			for (const TTuple<FName, FMovieSceneObjectBindingIDs>& Pair : MovieScene->AllTaggedBindings())
-			{
-				for (const FMovieSceneObjectBindingID& ID : Pair.Value.IDs)
-				{
-					if (ID.GetGuid() == ObjectBinding.BindingID)
-					{
-						CopyableBinding->Tags.AddUnique(Pair.Key);
-					}
-				}
-			}
 		}
 
 		UMovieSceneFolder* Folder = nullptr;
