Sequencer: Ensure the cine camera actor's name is unique when creating as possessable. Normally this is done in the actor factory but here we spawn the actor directly.

diff --git a/Engine/Source/Editor/Sequencer/Private/SequencerUtilities.cpp b/Engine/Source/Editor/Sequencer/Private/SequencerUtilities.cpp
--- a/Engine/Source/Editor/Sequencer/Private/SequencerUtilities.cpp
+++ b/Engine/Source/Editor/Sequencer/Private/SequencerUtilities.cpp
@@ -7,6 +7,7 @@
 #include "CameraRig_Crane.h"
 #include "Components/SplineComponent.h"
 #include "Containers/ArrayBuilder.h"
+#include "Editor/EditorEngine.h"
 #include "Engine/Selection.h"
 #include "EngineUtils.h"
 #include "Exporters/Exporter.h"
@@ -794,6 +795,8 @@ FGuid FSequencerUtilities::CreateCamera(TSharedRef<ISequencer> Sequencer, const
 	}
 	else
 	{
+		FActorLabelUtilities::SetActorLabelUnique(OutActor, ACineCameraActor::StaticClass()->GetName());
+
 		CameraGuid = CreateBinding(Sequencer, *OutActor, OutActor->GetActorLabel());
 	}
 
@@ -885,6 +888,8 @@ FGuid FSequencerUtilities::CreateCameraWithRig(TSharedRef<ISequencer> Sequencer,
 	}
 	else
 	{
+		FActorLabelUtilities::SetActorLabelUnique(OutActor, ACineCameraActor::StaticClass()->GetName());
+
 		// Parent it
 		OutActor->AttachToActor(Actor, FAttachmentTransformRules::KeepRelativeTransform);
 	}
