Sequencer: Don't cache view state when already cached.

Starting playback twice (spacebar and then down arrow), can result in the latter view state being cached incorrectly (already hidden) . From there on, the state is reversed - when sthe sequence is stopped, the viewport UI is hidden, when the sequence is playing, the viewport UI is shown.

diff --git a/Engine/Source/Editor/Sequencer/Private/Sequencer.cpp b/Engine/Source/Editor/Sequencer/Private/Sequencer.cpp
--- a/Engine/Source/Editor/Sequencer/Private/Sequencer.cpp
+++ b/Engine/Source/Editor/Sequencer/Private/Sequencer.cpp
@@ -4575,6 +4575,11 @@ UObject* FSequencer::FindSpawnedObjectOrTemplate(const FGuid& BindingId)
 void
 FSequencer::FCachedViewState::StoreViewState()
 {
+	if (bValid)
+	{
+		return;
+	}
+
 	bValid = true;
 	bIsViewportUIHidden = GLevelEditorModeTools().IsViewportUIHidden();
 	GLevelEditorModeTools().SetHideViewportUI(!GLevelEditorModeTools().IsViewportUIHidden());
