Sequencer: Fixed a regression in marquee selecting keys on collapsed rows

Roll back the change to explicitly go through the models and instead, stop recursing through if hitting a folder model.

diff --git a/Engine/Source/Editor/Sequencer/Private/Tools/SequencerEntityVisitor.cpp b/Engine/Source/Editor/Sequencer/Private/Tools/SequencerEntityVisitor.cpp
--- a/Engine/Source/Editor/Sequencer/Private/Tools/SequencerEntityVisitor.cpp
+++ b/Engine/Source/Editor/Sequencer/Private/Tools/SequencerEntityVisitor.cpp
@@ -11,6 +11,8 @@
 #include "MVVM/Extensions/ITrackLaneExtension.h"
 #include "MVVM/ViewModels/CategoryModel.h"
 #include "MVVM/ViewModels/ChannelModel.h"
+#include "MVVM/ViewModels/FolderModel.h"
+#include "MVVM/ViewModels/ObjectBindingModel.h"
 #include "MVVM/ViewModels/TrackModel.h"
 #include "MVVM/ViewModels/ViewModel.h"
 #include "MVVM/ViewModels/ViewModelHierarchy.h"
@@ -158,10 +160,24 @@ void FSequencerEntityWalker::VisitAnyChannels(const ISequencerEntityVisitor& Vis
 	FChannelGroupModel* ChannelGroup = nullptr;
 	if (FChannelGroupModel* ChannelGroupNode = InNode->CastThis<FChannelGroupModel>())
 	{
-		const IOutlinerExtension* OutlinerItem = ChannelGroupNode->CastThis<IOutlinerExtension>();
+		ChannelGroup = ChannelGroupNode;
+	}
+	if (FTrackModel* TrackModel = InNode->CastThis<FTrackModel>())
+	{
+		// TODO: Top-level channel group -- there used to be only one at most but this assumption isn't necessarily true anymore
+		for (const TViewModelPtr<FChannelGroupModel>& TopLevelChannel : TrackModel->GetTopLevelChannels().IterateSubList<FChannelGroupModel>())
+		{
+			ChannelGroup = TopLevelChannel.Get();
+			break;
+		}
+	}
+
+	if (ChannelGroup)
+	{
+		const IOutlinerExtension* OutlinerItem = ChannelGroup->CastThis<IOutlinerExtension>();
 		if (!OutlinerItem || OutlinerItem->IsFilteredOut() == false)
 		{
-			for (const TWeakViewModelPtr<FChannelModel>& WeakChannel : ChannelGroupNode->GetChannels())
+			for (const TWeakViewModelPtr<FChannelModel>& WeakChannel : ChannelGroup->GetChannels())
 			{
 				if (TViewModelPtr<FChannelModel> Channel = WeakChannel.Pin())
 				{
@@ -170,53 +186,23 @@ void FSequencerEntityWalker::VisitAnyChannels(const ISequencerEntityVisitor& Vis
 			}
 		}
 	}
-	else if (FCategoryGroupModel* CategoryGroupNode = InNode->CastThis<FCategoryGroupModel>())
+	// Otherwise it might be a collapsed node that contains key areas as children. If so we visit them as if they were a part of this track so that key groupings are visited properly.
+	else if (bAnyParentCollapsed)
 	{
-		if (!CategoryGroupNode->IsExpanded())
+		const IOutlinerExtension* OutlinerItem = InNode->CastThis<IOutlinerExtension>();
+		const bool bIsExpanded = !OutlinerItem || OutlinerItem->IsExpanded();
+
+		if ((InNode->IsA<FFolderModel>() || InNode->IsA<FObjectBindingModel>()) && !bIsExpanded)
 		{
-			for (const TWeakViewModelPtr<FCategoryModel>& WeakCategory : CategoryGroupNode->GetCategories())
-			{
-				if (TViewModelPtr<FCategoryModel> Category = WeakCategory.Pin())
-				{
-					TSharedPtr<FViewModel> Model = Category->GetLinkedOutlinerItem().AsModel();
-					for (TSharedPtr<FChannelGroupModel> DescandantChannelGroup : Model->GetDescendantsOfType<FChannelGroupModel>())
-					{
-						const IOutlinerExtension* OutlinerItem = DescandantChannelGroup->CastThis<IOutlinerExtension>();
-						if (!OutlinerItem || OutlinerItem->IsFilteredOut() == false)
-						{
-							for (const TWeakViewModelPtr<FChannelModel>& WeakChannel : DescandantChannelGroup->GetChannels())
-							{
-								if (TViewModelPtr<FChannelModel> Channel = WeakChannel.Pin())
-								{
-									VisitChannel(Visitor, Channel);
-								}
-							}
-						}
-					}
-				}
-			}
+			return;
 		}
-	}
-	else if (FTrackModel* TrackModel = InNode->CastThis<FTrackModel>())
-	{
-		if (!TrackModel->IsExpanded())
+
+		for (TSharedPtr<FViewModel> ChildNode : InNode->GetChildren(EViewModelListType::Outliner))
 		{
-			for (TSharedPtr<FChannelGroupModel> DescandantChannelGroup : TrackModel->GetDescendantsOfType<FChannelGroupModel>())
-			{
-				const IOutlinerExtension* OutlinerItem = DescandantChannelGroup->CastThis<IOutlinerExtension>();
-				if (!OutlinerItem || OutlinerItem->IsFilteredOut() == false)
-				{
-					for (const TWeakViewModelPtr<FChannelModel>& WeakChannel : DescandantChannelGroup->GetChannels())
-					{
-						if (TViewModelPtr<FChannelModel> Channel = WeakChannel.Pin())
-						{
-							VisitChannel(Visitor, Channel);
-						}
-					}
-				}
-			}
+			VisitAnyChannels(Visitor, ChildNode.ToSharedRef(), bAnyParentCollapsed || !bIsExpanded);
 		}
 	}
+
 }
 
 void FSequencerEntityWalker::VisitChannel(const ISequencerEntityVisitor& Visitor, const UE::Sequencer::TViewModelPtr<UE::Sequencer::FChannelModel>& Channel)
