Sequencer: Fix track row names when row indices have been remapped or a new row has been added

diff --git a/Engine/Source/Editor/Sequencer/Private/SequencerUtilities.cpp b/Engine/Source/Editor/Sequencer/Private/SequencerUtilities.cpp
--- a/Engine/Source/Editor/Sequencer/Private/SequencerUtilities.cpp
+++ b/Engine/Source/Editor/Sequencer/Private/SequencerUtilities.cpp
@@ -30,6 +30,7 @@
 #include "EntitySystem/MovieSceneBlenderSystem.h"
 #include "EntitySystem/IMovieSceneBlenderSystemSupport.h"
 #include "MovieSceneFolder.h"
+#include "MovieSceneNameableTrack.h"
 #include "MovieSceneSection.h"
 #include "MovieSceneSpawnRegister.h"
 #include "MovieSceneTimeHelpers.h"
@@ -181,19 +182,26 @@ void FSequencerUtilities::PopulateMenu_CreateNewSection(FMenuBuilder& MenuBuilde
 		if (UMovieSceneSection* NewSection = Track->CreateNewSection())
 		{
 			int32 OverlapPriority = 0;
+			TMap<int32, int32> NewToOldRowIndices;
 			for (UMovieSceneSection* Section : Track->GetAllSections())
 			{
-				OverlapPriority = FMath::Max(Section->GetOverlapPriority() + 1, OverlapPriority);
+				OverlapPriority = FMath::Max(Section->GetOverlapPriority() + 1, OverlapPriority);				
 
 				// Move existing sections on the same row or beyond so that they don't overlap with the new section
 				if (Section != NewSection && Section->GetRowIndex() >= RowIndex)
 				{
-					Section->SetRowIndex(Section->GetRowIndex() + 1);
+					int32 OldRowIndex = Section->GetRowIndex();
+					int32 NewRowIndex = Section->GetRowIndex() + 1;
+					NewToOldRowIndices.FindOrAdd(NewRowIndex, OldRowIndex);
+					Section->Modify();
+					Section->SetRowIndex(NewRowIndex);
 				}
 			}
 
 			Track->Modify();
 
+			Track->OnRowIndicesChanged(NewToOldRowIndices);
+
 			FFrameNumber NewSectionRangeEnd = PlaybackEnd;
 			if (PlaybackEnd <= CurrentTime.Time.FrameNumber)
 			{
@@ -204,12 +212,17 @@ void FSequencerUtilities::PopulateMenu_CreateNewSection(FMenuBuilder& MenuBuilde
 			
 			NewSection->SetRange(TRange<FFrameNumber>(CurrentTime.Time.FrameNumber, NewSectionRangeEnd));
 			NewSection->SetOverlapPriority(OverlapPriority);
-			NewSection->SetRowIndex(RowIndex);
+			NewSection->SetRowIndex(RowIndex);			
 			NewSection->SetBlendType(BlendType);
 
 			Track->AddSection(*NewSection);
 			Track->UpdateEasing();
 
+			if (UMovieSceneNameableTrack* NameableTrack = Cast<UMovieSceneNameableTrack>(Track))
+			{
+				NameableTrack->SetTrackRowDisplayName(FText::GetEmpty(), RowIndex);
+			}
+
 			Sequencer->NotifyMovieSceneDataChanged(EMovieSceneDataChangeType::MovieSceneStructureItemAdded);
 			Sequencer->EmptySelection();
 			Sequencer->SelectSection(NewSection);
diff --git a/Engine/Source/Editor/Sequencer/Private/Tools/EditToolDragOperations.cpp b/Engine/Source/Editor/Sequencer/Private/Tools/EditToolDragOperations.cpp
--- a/Engine/Source/Editor/Sequencer/Private/Tools/EditToolDragOperations.cpp
+++ b/Engine/Source/Editor/Sequencer/Private/Tools/EditToolDragOperations.cpp
@@ -1086,17 +1086,22 @@ bool FMoveKeysAndSections::HandleSectionMovement(FFrameTime MouseTime, FVector2D
 
 	// If sections are all on different rows or from different tracks, don't set row indices for anything because it leads to odd behavior.
 	bool bSectionsAreOnDifferentRows = false;
-	int32 FirstRowIndex = -1;
+	TOptional<int32> LowestRowIndex;
+	TOptional<int32> HighestRowIndex;
 	UMovieSceneTrack* FirstTrack = nullptr;
 
 	for (UMovieSceneSection* Section : Sections)
 	{
 		UMovieSceneTrack* Track = Section->GetTypedOuter<UMovieSceneTrack>();
-		if (FirstRowIndex == -1)
+		if (!LowestRowIndex.IsSet() || LowestRowIndex.GetValue() < Section->GetRowIndex())
 		{
-			FirstRowIndex = Section->GetRowIndex();
+			LowestRowIndex = Section->GetRowIndex();
 		}
-		if (FirstRowIndex != Section->GetRowIndex())
+		if (!HighestRowIndex.IsSet() || HighestRowIndex.GetValue() > Section->GetRowIndex())
+		{
+			HighestRowIndex = Section->GetRowIndex();
+		}
+		if (LowestRowIndex.IsSet() && LowestRowIndex.GetValue() != Section->GetRowIndex())
 		{
 			bSectionsAreOnDifferentRows = true;
 		}
@@ -1296,7 +1301,7 @@ bool FMoveKeysAndSections::HandleSectionMovement(FFrameTime MouseTime, FVector2D
 					{
 						if (!Sections.Contains(InitialRowIndex.Section))
 						{
-							if (InitialRowIndex.RowIndex <= FirstRowIndex)
+							if (LowestRowIndex.IsSet() && InitialRowIndex.RowIndex <= LowestRowIndex.GetValue())
 							{
 								bSectionsBeingMovedAreAtTop = false;
 								break;
@@ -1320,9 +1325,30 @@ bool FMoveKeysAndSections::HandleSectionMovement(FFrameTime MouseTime, FVector2D
 			}
 			else
 			{
-				Section->Modify();
-				Section->SetRowIndex(TargetRowIndex);
-				bRowIndexChanged = true;
+				if (!bSectionsAreOnDifferentRows)
+				{
+					// If the sections being moved are all at the bottom, and all others are aove it, do nothing
+					bool bSectionsBeingMovedAreAtBottom = true;
+
+					for (const FInitialRowIndex& InitialRowIndex : InitialSectionRowIndicies)
+					{
+						if (!Sections.Contains(InitialRowIndex.Section))
+						{
+							if (HighestRowIndex.IsSet() && InitialRowIndex.RowIndex >= HighestRowIndex.GetValue())
+							{
+								bSectionsBeingMovedAreAtBottom = false;
+								break;
+							}
+						}
+					}
+
+					if (!bSectionsBeingMovedAreAtBottom || TargetRowIndex < Section->GetRowIndex())
+					{
+						Section->Modify();
+						Section->SetRowIndex(TargetRowIndex);
+						bRowIndexChanged = true;
+					}
+				}
 			}
 		}
 	}
diff --git a/Engine/Source/Runtime/MovieScene/Private/MovieSceneNameableTrack.cpp b/Engine/Source/Runtime/MovieScene/Private/MovieSceneNameableTrack.cpp
--- a/Engine/Source/Runtime/MovieScene/Private/MovieSceneNameableTrack.cpp
+++ b/Engine/Source/Runtime/MovieScene/Private/MovieSceneNameableTrack.cpp
@@ -93,11 +93,27 @@ FText UMovieSceneNameableTrack::GetTrackRowDisplayName(int32 TrackRowIndex) cons
 }
 
 FText UMovieSceneNameableTrack::GetDefaultDisplayName() const
-
 { 
 	return LOCTEXT("UnnamedTrackName", "Unnamed Track"); 
 }
 
+void UMovieSceneNameableTrack::OnRowIndicesChanged(const TMap<int32, int32>& NewToOldRowIndices)
+{
+	TArray<FText> OriginalTrackRowDisplayNames = TrackRowDisplayNames;
+
+	Modify();
+	TrackRowDisplayNames.Empty();
+	for (int32 RowIndex = 0, RowMax = GetMaxRowIndex(); RowIndex <= RowMax; ++RowIndex)
+	{
+		const int32* OldIndex = NewToOldRowIndices.Find(RowIndex);
+		int32 RemappedIndex = OldIndex ? *OldIndex : RowIndex;
+		if (RemappedIndex < OriginalTrackRowDisplayNames.Num())
+		{
+			TrackRowDisplayNames.Add(OriginalTrackRowDisplayNames[RemappedIndex]);
+		}
+	}
+}
+
 #endif
 
 
diff --git a/Engine/Source/Runtime/MovieScene/Private/MovieSceneTrack.cpp b/Engine/Source/Runtime/MovieScene/Private/MovieSceneTrack.cpp
--- a/Engine/Source/Runtime/MovieScene/Private/MovieSceneTrack.cpp
+++ b/Engine/Source/Runtime/MovieScene/Private/MovieSceneTrack.cpp
@@ -237,6 +237,8 @@ int32 UMovieSceneTrack::GetMaxRowIndex() const
 
 bool UMovieSceneTrack::FixRowIndices()
 {
+	TMap<int32, int32> NewToOldRowIndices;
+
 	bool bFixesMade = false;
 	TArray<UMovieSceneSection*> Sections = GetAllSections();
 	if (SupportsMultipleRows())
@@ -259,8 +261,10 @@ bool UMovieSceneTrack::FixRowIndices()
 				{
 					if (SectionForIndex->GetRowIndex() != NewIndex)
 					{
+						int32 OldIndex = SectionForIndex->GetRowIndex();
 						SectionForIndex->Modify();
 						SectionForIndex->SetRowIndex(NewIndex);
+						NewToOldRowIndices.FindOrAdd(NewIndex, OldIndex);
 						bFixesMade = true;
 					}
 				}
@@ -287,6 +291,11 @@ bool UMovieSceneTrack::FixRowIndices()
 			}
 		}
 	}
+
+	if (NewToOldRowIndices.Num())
+	{
+		OnRowIndicesChanged(NewToOldRowIndices);
+	}
 	return bFixesMade;
 }
 
diff --git a/Engine/Source/Runtime/MovieScene/Public/MovieSceneNameableTrack.h b/Engine/Source/Runtime/MovieScene/Public/MovieSceneNameableTrack.h
--- a/Engine/Source/Runtime/MovieScene/Public/MovieSceneNameableTrack.h
+++ b/Engine/Source/Runtime/MovieScene/Public/MovieSceneNameableTrack.h
@@ -65,6 +65,8 @@ public:
 	MOVIESCENE_API virtual FText GetTrackRowDisplayName(int32 RowIndex) const override;
 	MOVIESCENE_API virtual FText GetDefaultDisplayName() const;
 
+	MOVIESCENE_API virtual void OnRowIndicesChanged(const TMap<int32, int32>& NewToOldRowIndices) override;
+
 private:
 
 	/** The track's human readable display name. */
diff --git a/Engine/Source/Runtime/MovieScene/Public/MovieSceneTrack.h b/Engine/Source/Runtime/MovieScene/Public/MovieSceneTrack.h
--- a/Engine/Source/Runtime/MovieScene/Public/MovieSceneTrack.h
+++ b/Engine/Source/Runtime/MovieScene/Public/MovieSceneTrack.h
@@ -345,6 +345,9 @@ public:
 	 */
 	MOVIESCENE_API bool FixRowIndices();
 
+	/** Called when row indices have been fixed up */
+	MOVIESCENE_API virtual void OnRowIndicesChanged(const TMap<int32, int32>& NewToOldRowIndices) {}
+
 	/**
 	* @return Whether evaluation of this track should be disabled due to mute/solo settings
 	*/
