Sequencer: Disable some right click menus if the sequence is read only (Add Section, Blending Algorithm, Edit/Batch Edit Sections)

diff --git a/Engine/Source/Editor/Sequencer/Private/SequencerCommonHelpers.cpp b/Engine/Source/Editor/Sequencer/Private/SequencerCommonHelpers.cpp
--- a/Engine/Source/Editor/Sequencer/Private/SequencerCommonHelpers.cpp
+++ b/Engine/Source/Editor/Sequencer/Private/SequencerCommonHelpers.cpp
@@ -352,5 +352,6 @@ void SequencerHelpers::AddPropertiesMenu(FSequencer& Sequencer, FMenuBuilder& Me
 	DetailsView->SetObjects(Sections);
 
 	DetailsNotifyWrapper->SetDetailsAndSequencer(DetailsView, Sequencer.AsShared());
+	DetailsNotifyWrapper->SetEnabled(!Sequencer.IsReadOnly());
 	MenuBuilder.AddWidget(DetailsNotifyWrapper, FText::GetEmpty(), true);
 }
diff --git a/Engine/Source/Editor/Sequencer/Private/SequencerUtilities.cpp b/Engine/Source/Editor/Sequencer/Private/SequencerUtilities.cpp
--- a/Engine/Source/Editor/Sequencer/Private/SequencerUtilities.cpp
+++ b/Engine/Source/Editor/Sequencer/Private/SequencerUtilities.cpp
@@ -246,7 +246,10 @@ void FSequencerUtilities::PopulateMenu_CreateNewSection(FMenuBuilder& MenuBuilde
 			NameOverride.IsEmpty() ? DisplayName : NameOverride,
 			TooltipOverride.IsEmpty() ? FText::Format(LOCTEXT("AddSectionFormatToolTip", "Adds a new {0} section at the current time"), DisplayName) : TooltipOverride,
 			FSlateIcon(FAppStyle::GetAppStyleSetName(), EnumValueName),
-			FUIAction(FExecuteAction::CreateLambda(CreateNewSection, BlendType))
+			FUIAction(
+				FExecuteAction::CreateLambda(CreateNewSection, BlendType),
+				FCanExecuteAction::CreateLambda([InSequencer] { return InSequencer.IsValid() && !InSequencer.Pin()->IsReadOnly(); })
+			)
 		);
 	}
 }
@@ -284,7 +287,7 @@ void FSequencerUtilities::PopulateMenu_BlenderSubMenu(FMenuBuilder& MenuBuilder,
 					Track->Modify();
 					BlenderSystemSupport->SetBlenderSystem(SystemClass);
 				}),
-				FCanExecuteAction(),
+				FCanExecuteAction::CreateLambda([InSequencer] { return InSequencer.IsValid() && !InSequencer.Pin()->IsReadOnly(); }),
 				FIsActionChecked::CreateLambda([BlenderSystemSupport, SystemClass]
 				{
 					return BlenderSystemSupport->GetBlenderSystem() == SystemClass;
@@ -399,7 +402,7 @@ void FSequencerUtilities::PopulateMenu_SetBlendType(FMenuBuilder& MenuBuilder, c
 			FSlateIcon(FAppStyle::GetAppStyleSetName(), EnumValueName),
 			FUIAction(
 				FExecuteAction::CreateLambda(Execute, BlendType),
-				FCanExecuteAction(),
+				FCanExecuteAction::CreateLambda([InSequencer] { return InSequencer.IsValid() && !InSequencer.Pin()->IsReadOnly(); }),
 				FIsActionChecked::CreateLambda([InSections, BlendType]
 				{
 					int32 NumActiveBlendTypes = 0;
