Sequencer: Remove obsolete Evaluate When Not Spawned

#rnx

diff --git a/Engine/Source/Editor/Sequencer/Private/MVVM/ViewModels/ObjectBindingModel.cpp b/Engine/Source/Editor/Sequencer/Private/MVVM/ViewModels/ObjectBindingModel.cpp
--- a/Engine/Source/Editor/Sequencer/Private/MVVM/ViewModels/ObjectBindingModel.cpp
+++ b/Engine/Source/Editor/Sequencer/Private/MVVM/ViewModels/ObjectBindingModel.cpp
@@ -974,19 +974,6 @@ void FObjectBindingModel::BuildContextMenu(FMenuBuilder& MenuBuilder)
 				EUserInterfaceActionType::ToggleButton
 			);
 
-			MenuBuilder.AddMenuEntry(
-				LOCTEXT("EvaluateTracksWhenNotSpawned", "Evaluate Tracks When Not Spawned"),
-				LOCTEXT("EvaluateTracksWhenNotSpawnedTooltip", "When enabled, any tracks on this object binding or its children will still be evaluated even when the object is not spawned."),
-				FSlateIcon(),
-				FUIAction(
-					FExecuteAction::CreateLambda(FMovieSceneSpawnableFlagToggler{Sequencer, MovieScene, &FMovieSceneSpawnable::bEvaluateTracksWhenNotSpawned, LOCTEXT("EvaluateTracksWhenNotSpawned_Transaction", "Evaluate Tracks When Not Spawned")}),
-					FCanExecuteAction(),
-					FGetActionCheckState::CreateLambda(FMovieSceneSpawnableFlagCheckState{Sequencer, MovieScene, &FMovieSceneSpawnable::bEvaluateTracksWhenNotSpawned})
-				),
-				NAME_None,
-				EUserInterfaceActionType::ToggleButton
-			);
-
 			MenuBuilder.AddMenuEntry(
 				LOCTEXT("NetAddressable", "Net Addressable"),
 				LOCTEXT("NetAddressableTooltip", "When enabled, this spawnable will be spawned using a unique name that allows it to be addressed by the server and client (useful for relative movement calculations on spawned props)"),
diff --git a/Engine/Source/Runtime/MovieScene/Public/MovieSceneSpawnable.h b/Engine/Source/Runtime/MovieScene/Public/MovieSceneSpawnable.h
--- a/Engine/Source/Runtime/MovieScene/Public/MovieSceneSpawnable.h
+++ b/Engine/Source/Runtime/MovieScene/Public/MovieSceneSpawnable.h
@@ -36,7 +36,6 @@ struct FMovieSceneSpawnable
 	FMovieSceneSpawnable()
 		: bContinuouslyRespawn(false)
 		, bNetAddressableName(false)
-		, bEvaluateTracksWhenNotSpawned(false)
 		, ObjectTemplate(nullptr)
 		, Ownership(ESpawnOwnership::InnerSequence)
 #if WITH_EDITORONLY_DATA
@@ -49,7 +48,6 @@ struct FMovieSceneSpawnable
 	FMovieSceneSpawnable(const FString& InitName, UObject& InObjectTemplate)
 		: bContinuouslyRespawn(false)
 		, bNetAddressableName(false)
-		, bEvaluateTracksWhenNotSpawned(false)
 		, Guid(FGuid::NewGuid())
 		, Name(InitName)
 		, ObjectTemplate(&InObjectTemplate)
@@ -249,10 +247,6 @@ public:
 	UPROPERTY(EditAnywhere, Category=Actor)
 	bool bNetAddressableName;
 
-	/** When enabled, any tracks on this object binding or its children will still be evaluated even when the object is not spawned. */
-	UPROPERTY(EditAnywhere, Category=Actor)
-	bool bEvaluateTracksWhenNotSpawned;
-
 private:
 
 	/** Unique identifier of the spawnable object. */
