Sequencer: Put Time Undo in CVAR based on animators feedback, default false.

diff --git a/Engine/Source/Editor/Sequencer/Private/Sequencer.cpp b/Engine/Source/Editor/Sequencer/Private/Sequencer.cpp
--- a/Engine/Source/Editor/Sequencer/Private/Sequencer.cpp
+++ b/Engine/Source/Editor/Sequencer/Private/Sequencer.cpp
@@ -181,6 +181,11 @@ static TAutoConsoleVariable<float> CVarAutoScrubCurveExponent(
 	2.0f,
 	TEXT("How much to ramp in and out the scrub speed when auto-scrubbing"));
 
+static TAutoConsoleVariable<bool> CVarTimeUndo(
+	TEXT("Sequencer.TimeUndo"),
+	false,
+	TEXT("Enable/disable ability to undo time when making changes"));
+
 class FSequencerOutlinerSelectionHandler : 
 	public UE::Sequencer::IOutlinerSelectionHandler
 {
@@ -522,10 +527,12 @@ void FSequencer::InitSequencer(const FSequencerInitParams& InitParams, const TSh
 	// Update initial movie scene data
 	NotifyMovieSceneDataChanged( EMovieSceneDataChangeType::ActiveMovieSceneChanged );
 
-
-	TimeUndoRedoHandler.UndoRedoProxy = NewObject<USequencerTimeChangeUndoRedoProxy>(GetTransientPackage(), NAME_None);
-	TimeUndoRedoHandler.SetSequencer(SharedThis(this));
-	TimeUndoRedoHandler.UndoRedoProxy->SetFlags(RF_Transactional | RF_Transient);
+	if (CVarTimeUndo->GetBool())
+	{
+		TimeUndoRedoHandler.UndoRedoProxy = NewObject<USequencerTimeChangeUndoRedoProxy>(GetTransientPackage(), NAME_None);
+		TimeUndoRedoHandler.SetSequencer(SharedThis(this));
+		TimeUndoRedoHandler.UndoRedoProxy->SetFlags(RF_Transactional | RF_Transient);
+	}
 	
 
 	// Update the view range to the new current time
@@ -4250,7 +4257,10 @@ void FSequencer::AddReferencedObjects( FReferenceCollector& Collector )
 {
 	Collector.AddReferencedObject( CompiledDataManager );
 	Collector.AddReferencedObject( Settings );
-	Collector.AddReferencedObject(TimeUndoRedoHandler.UndoRedoProxy );
+	if (TimeUndoRedoHandler.UndoRedoProxy)
+	{
+		Collector.AddReferencedObject(TimeUndoRedoHandler.UndoRedoProxy);
+	}
 	if (UMovieSceneSequence* RootSequencePtr = RootSequence.Get())
 	{
 		Collector.AddReferencedObject( RootSequencePtr );
