Sequencer: Rename to playhead

#rnx

diff --git a/Engine/Source/Editor/Sequencer/Private/SequencerSettings.cpp b/Engine/Source/Editor/Sequencer/Private/SequencerSettings.cpp
--- a/Engine/Source/Editor/Sequencer/Private/SequencerSettings.cpp
+++ b/Engine/Source/Editor/Sequencer/Private/SequencerSettings.cpp
@@ -31,7 +31,7 @@ USequencerSettings::USequencerSettings( const FObjectInitializer& ObjectInitiali
 	GridSpacing = TOptional<float>();
 	bSnapCurveValueToInterval = false;
 	bShowSelectedNodesOnly = false;
-	ZoomPosition = ESequencerZoomPosition::SZP_CurrentTime;
+	ZoomPosition = ESequencerZoomPosition::SZP_Playhead;
 	bAutoScrollEnabled = false;
 	bLinkCurveEditorTimeRange = false;
 	bSynchronizeCurveEditorSelection = true;
diff --git a/Engine/Source/Editor/Sequencer/Private/SequencerTimeSliderController.cpp b/Engine/Source/Editor/Sequencer/Private/SequencerTimeSliderController.cpp
--- a/Engine/Source/Editor/Sequencer/Private/SequencerTimeSliderController.cpp
+++ b/Engine/Source/Editor/Sequencer/Private/SequencerTimeSliderController.cpp
@@ -1042,8 +1042,8 @@ FReply FSequencerTimeSliderController::OnMouseMoveImpl( SWidget& WidgetOwner, co
 		{
 			float MouseFractionX = MyGeometry.AbsoluteToLocal(MouseEvent.GetScreenSpacePosition()).X / MyGeometry.GetLocalSize().X;
 
-			// If zooming on the current time, adjust mouse fractionX
-			if (Sequencer->GetSequencerSettings()->GetZoomPosition() == ESequencerZoomPosition::SZP_CurrentTime)
+			// If zooming on the playhead, adjust mouse fractionX
+			if (Sequencer->GetSequencerSettings()->GetZoomPosition() == ESequencerZoomPosition::SZP_Playhead)
 			{
 				const double ScrubPosition = TimeSliderArgs.ScrubPosition.Get() / GetTickResolution();
 				if (GetViewRange().Contains(ScrubPosition))
@@ -1276,7 +1276,7 @@ FReply FSequencerTimeSliderController::OnMouseWheel( SWidget& WidgetOwner, const
 		TSharedPtr<FSequencer> Sequencer = WeakSequencer.Pin();
 			
 		// If zooming on the current time, adjust mouse fractionX
-		if (Sequencer.IsValid() && Sequencer->GetSequencerSettings()->GetZoomPosition() == ESequencerZoomPosition::SZP_CurrentTime)
+		if (Sequencer.IsValid() && Sequencer->GetSequencerSettings()->GetZoomPosition() == ESequencerZoomPosition::SZP_Playhead)
 		{
 			const double ScrubPosition = TimeSliderArgs.ScrubPosition.Get() / GetTickResolution();
 			if (GetViewRange().Contains(ScrubPosition))
diff --git a/Engine/Source/Editor/Sequencer/Public/SequencerSettings.h b/Engine/Source/Editor/Sequencer/Public/SequencerSettings.h
--- a/Engine/Source/Editor/Sequencer/Public/SequencerSettings.h
+++ b/Engine/Source/Editor/Sequencer/Public/SequencerSettings.h
@@ -28,8 +28,8 @@ enum ESequencerSpawnPosition : int
 UENUM()
 enum ESequencerZoomPosition : int
 {
-	/** Current Time. */
-	SZP_CurrentTime UMETA(DisplayName="Playhead"),
+	/** Playhead. */
+	SZP_Playhead UMETA(DisplayName="Playhead"),
 
 	/** Mouse Position. */
 	SZP_MousePosition UMETA(DisplayName="Mouse Position"),
