Sequencer: Captialize "Remove Override"

#rnx

diff --git a/Engine/Source/Editor/Sequencer/Private/MVVM/ViewModels/ChannelModel.cpp b/Engine/Source/Editor/Sequencer/Private/MVVM/ViewModels/ChannelModel.cpp
--- a/Engine/Source/Editor/Sequencer/Private/MVVM/ViewModels/ChannelModel.cpp
+++ b/Engine/Source/Editor/Sequencer/Private/MVVM/ViewModels/ChannelModel.cpp
@@ -582,7 +582,7 @@ void FChannelGroupOverrideHelper::BuildChannelOverrideMenu(FMenuBuilder& MenuBui
 		if (ExistingOverrides.Num() > 0)
 		{
 			MenuBuilder.AddMenuEntry(
-				LOCTEXT("RemoveOverrideChannelLabel", "Remove override"),
+				LOCTEXT("RemoveOverrideChannelLabel", "Remove Override"),
 				LOCTEXT("RemoveOverrideChannelTooltip", "Removes the channel override and revert back to the default channel type"),
 				FSlateIcon(),
 				FUIAction(
