Sequencer: Change FDeferredSignedObjectChangeHandler so that it's a SharedPtr in FSequencer and created in InitSequencer if it doesn't exist.

Removed workaround in Interchange
Tested that opening multiple sequencers (Level Sequence Editor, Actor Sequence, Niagara, UMG) only creates a single FDeferredSignedObjectChangeHandler which is destroyed when they're all closed.

#rnx

diff --git a/Engine/Plugins/Interchange/Runtime/Source/Import/Private/Animation/InterchangeAnimationTrackSetFactory.cpp b/Engine/Plugins/Interchange/Runtime/Source/Import/Private/Animation/InterchangeAnimationTrackSetFactory.cpp
--- a/Engine/Plugins/Interchange/Runtime/Source/Import/Private/Animation/InterchangeAnimationTrackSetFactory.cpp
+++ b/Engine/Plugins/Interchange/Runtime/Source/Import/Private/Animation/InterchangeAnimationTrackSetFactory.cpp
@@ -819,19 +819,4 @@ void UInterchangeAnimationTrackSetFactory::SetupObject_GameThread(const FSetupOb
 	// TODO: Talk with sequence team about adding AssetImportData to ULevelSequence for re-import 
 }
 
-void UInterchangeAnimationTrackSetFactory::FinalizeObject_GameThread(const FSetupObjectParams& Arguments)
-{
-	check(IsInGameThread());
-	if (ULevelSequence* LevelSequence = Cast<ULevelSequence>(Arguments.ImportedObject))
-	{
-		// Ugly temporary workaround to Sequencer's design which has an implicit deferring system which
-		// holds onto UMovieSceneSignedObject objects when PreEditChange is called on such object (i.e. ULevelSequence)
-		// Only a call to open the Sequencer, which calls UMovieSceneSignedObject::ResetImplicitScopedModifyDefer, will release the hold
-		// See FDeferredSignedObjectChangeHandler, FSequencer::Tick
-		// Note: Sequencer team is aware of the issue and will look their design in a future release
-		UMovieSceneSignedObject::ResetImplicitScopedModifyDefer();
-	}
-	Super::FinalizeObject_GameThread(Arguments);
-}
-
-#undef LOCTEXT_NAMESPACE
\ No newline at end of file
+#undef LOCTEXT_NAMESPACE
diff --git a/Engine/Plugins/Interchange/Runtime/Source/Import/Public/Animation/InterchangeAnimationTrackSetFactory.h b/Engine/Plugins/Interchange/Runtime/Source/Import/Public/Animation/InterchangeAnimationTrackSetFactory.h
--- a/Engine/Plugins/Interchange/Runtime/Source/Import/Public/Animation/InterchangeAnimationTrackSetFactory.h
+++ b/Engine/Plugins/Interchange/Runtime/Source/Import/Public/Animation/InterchangeAnimationTrackSetFactory.h
@@ -22,7 +22,6 @@ public:
 	virtual EInterchangeFactoryAssetType GetFactoryAssetType() override { return EInterchangeFactoryAssetType::Animations; }
 	virtual UObject* ImportAssetObject_GameThread(const FImportAssetObjectParams& Arguments) override;
 	virtual void SetupObject_GameThread(const FSetupObjectParams& Arguments) override;
-	virtual void FinalizeObject_GameThread(const FSetupObjectParams& Arguments) override;
 
 	// Interchange factory base interface end
 	//////////////////////////////////////////////////////////////////////////
diff --git a/Engine/Source/Editor/Sequencer/Private/Sequencer.cpp b/Engine/Source/Editor/Sequencer/Private/Sequencer.cpp
--- a/Engine/Source/Editor/Sequencer/Private/Sequencer.cpp
+++ b/Engine/Source/Editor/Sequencer/Private/Sequencer.cpp
@@ -159,6 +159,7 @@
 #include "EntitySystem/MovieScenePreAnimatedStateSystem.h"
 #include "Systems/MovieSceneMotionVectorSimulationSystem.h"
 #include "IKeyArea.h"
+#include "Editor/TransBuffer.h"
 
 #include "EngineModule.h"
 
@@ -267,6 +268,103 @@ private:
 	TWeakPtr<UE::Sequencer::SOutlinerView> WeakPinnedTreeView;
 };
 
+namespace UE
+{
+	namespace Sequencer
+	{
+
+		struct FDeferredSignedObjectChangeHandler : UE::MovieScene::IDeferredSignedObjectChangeHandler
+		{
+			FDeferredSignedObjectChangeHandler()
+			{
+				Init();
+			}
+
+			~FDeferredSignedObjectChangeHandler()
+			{
+				if (UTransBuffer* TransBuffer = WeakBuffer.Get())
+				{
+					TransBuffer->OnTransactionStateChanged().RemoveAll(this);
+				}
+			}
+
+			void Init()
+			{
+				UTransBuffer* TransBuffer = GUnrealEd ? Cast<UTransBuffer>(GUnrealEd->Trans) : nullptr;
+				if (TransBuffer)
+				{
+					WeakBuffer = TransBuffer;
+					TransBuffer->OnTransactionStateChanged().AddRaw(this, &FDeferredSignedObjectChangeHandler::OnTransactionStateChanged);
+					if (TransBuffer->IsActive())
+					{
+						DeferTransactionChanges.Emplace();
+					}
+				}
+				else
+				{
+					FCoreDelegates::OnPostEngineInit.AddLambda([this] { this->Init(); });
+				}
+			}
+
+			void OnTransactionStateChanged(const FTransactionContext& TransactionContext, ETransactionStateEventType TransactionState)
+			{
+				/** A transaction has been started. This will be followed by a TransactionCanceled or TransactionFinalized event. */
+				switch (TransactionState)
+				{
+				case ETransactionStateEventType::TransactionStarted:
+				case ETransactionStateEventType::UndoRedoStarted:
+					DeferTransactionChanges.Emplace();
+					break;
+
+				case ETransactionStateEventType::TransactionCanceled:
+				case ETransactionStateEventType::PreTransactionFinalized:
+				case ETransactionStateEventType::UndoRedoFinalized:
+					DeferTransactionChanges.Reset();
+					break;
+				}
+			}
+
+			void Flush() override
+			{
+				TSet<TWeakObjectPtr<UMovieSceneSignedObject>> SignedObjectsTmp = SignedObjects;
+				SignedObjects.Empty();
+				// we operate on a copy of the signed objects in case the delegates would modify the array
+				for (TWeakObjectPtr<UMovieSceneSignedObject> WeakObject : SignedObjectsTmp)
+				{
+					if (UMovieSceneSignedObject* Object = WeakObject.Get())
+					{
+						Object->BroadcastChanged();
+					}
+				}
+			}
+
+			void DeferMarkAsChanged(UMovieSceneSignedObject* SignedObject) override
+			{
+				SignedObjects.Add(SignedObject);
+			}
+
+			bool CreateImplicitScopedModifyDefer() override
+			{
+				ensure(!DeferImplicitChanges.IsSet());
+				DeferImplicitChanges.Emplace();
+				return true;
+			}
+
+			void ResetImplicitScopedModifyDefer() override
+			{
+				DeferImplicitChanges.Reset();
+			}
+
+			TSet<TWeakObjectPtr<UMovieSceneSignedObject>> SignedObjects;
+			TWeakObjectPtr<UTransBuffer>   WeakBuffer;
+			TOptional<UE::MovieScene::FScopedSignedObjectModifyDefer> DeferTransactionChanges;
+			TOptional<UE::MovieScene::FScopedSignedObjectModifyDefer> DeferImplicitChanges;
+		};
+
+	} // namespace Sequencer
+} // namespace UE
+
+
 void FSequencer::InitSequencer(const FSequencerInitParams& InitParams, const TSharedRef<ISequencerObjectChangeListener>& InObjectChangeListener, const TArray<FOnCreateTrackEditor>& TrackEditorDelegates, const TArray<FOnCreateEditorObjectBinding>& EditorObjectBindingDelegates)
 {
 	using namespace UE::Sequencer;
@@ -429,6 +527,17 @@ void FSequencer::InitSequencer(const FSequencerInitParams& InitParams, const TSh
 		}
 	}
 
+	TWeakPtr<UE::MovieScene::IDeferredSignedObjectChangeHandler> WeakDeferredSignedObjectChangeHandler = UMovieSceneSignedObject::GetDeferredHandler();
+	if (WeakDeferredSignedObjectChangeHandler.IsValid())
+	{
+		DeferredSignedObjectChangeHandler = WeakDeferredSignedObjectChangeHandler.Pin();
+	}
+	else
+	{
+		DeferredSignedObjectChangeHandler = MakeShared<FDeferredSignedObjectChangeHandler>();
+		UMovieSceneSignedObject::SetDeferredHandler(DeferredSignedObjectChangeHandler);
+	}
+
 	ViewModel = MakeShared<FSequencerEditorViewModel>(SharedThis(this), GetHostCapabilities());
 	ViewModel->InitializeEditor();
 	ViewModel->SetSequence(InitParams.RootSequence);
diff --git a/Engine/Source/Editor/Sequencer/Private/Sequencer.h b/Engine/Source/Editor/Sequencer/Private/Sequencer.h
--- a/Engine/Source/Editor/Sequencer/Private/Sequencer.h
+++ b/Engine/Source/Editor/Sequencer/Private/Sequencer.h
@@ -1418,6 +1418,8 @@ private:
 	UMovieSceneCompiledDataManager* CompiledDataManager;
 	TSharedPtr<FMovieSceneEntitySystemRunner> Runner;
 
+	TSharedPtr<UE::MovieScene::IDeferredSignedObjectChangeHandler> DeferredSignedObjectChangeHandler;
+
 	TMap<FName, TFunction<void()>> CleanupFunctions;
 
 	/** Transient collection of keys that is used for jumping between keys contained within the current selection */
diff --git a/Engine/Source/Editor/Sequencer/Private/SequencerModule.cpp b/Engine/Source/Editor/Sequencer/Private/SequencerModule.cpp
--- a/Engine/Source/Editor/Sequencer/Private/SequencerModule.cpp
+++ b/Engine/Source/Editor/Sequencer/Private/SequencerModule.cpp
@@ -41,7 +41,6 @@
 #include "Misc/CoreDelegates.h"
 #include "UnrealEdGlobals.h"
 #include "Editor/UnrealEdEngine.h"
-#include "Editor/TransBuffer.h"
 #include "PropertyEditorModule.h"
 #include "PropertyHandle.h"
 
@@ -53,104 +52,6 @@
 
 #define LOCTEXT_NAMESPACE "SequencerEditor"
 
-
-namespace UE
-{
-namespace Sequencer
-{
-
-struct FDeferredSignedObjectChangeHandler : UE::MovieScene::IDeferredSignedObjectChangeHandler
-{
-	FDeferredSignedObjectChangeHandler()
-	{
-		Init();
-	}
-
-	~FDeferredSignedObjectChangeHandler()
-	{
-		if (UTransBuffer* TransBuffer = WeakBuffer.Get())
-		{
-			TransBuffer->OnTransactionStateChanged().RemoveAll(this);
-		}
-	}
-
-	void Init()
-	{
-		UTransBuffer* TransBuffer = GUnrealEd ? Cast<UTransBuffer>(GUnrealEd->Trans) : nullptr;
-		if (TransBuffer)
-		{
-			WeakBuffer = TransBuffer;
-			TransBuffer->OnTransactionStateChanged().AddRaw(this, &FDeferredSignedObjectChangeHandler::OnTransactionStateChanged);
-			if (TransBuffer->IsActive())
-			{
-				DeferTransactionChanges.Emplace();
-			}
-		}
-		else
-		{
-			FCoreDelegates::OnPostEngineInit.AddLambda([this]{ this->Init(); });
-		}
-	}
-
-	void OnTransactionStateChanged(const FTransactionContext& TransactionContext, ETransactionStateEventType TransactionState)
-	{
-		/** A transaction has been started. This will be followed by a TransactionCanceled or TransactionFinalized event. */
-		switch (TransactionState)
-		{
-		case ETransactionStateEventType::TransactionStarted:
-		case ETransactionStateEventType::UndoRedoStarted:
-			DeferTransactionChanges.Emplace();
-			break;
-
-		case ETransactionStateEventType::TransactionCanceled:
-		case ETransactionStateEventType::PreTransactionFinalized:
-		case ETransactionStateEventType::UndoRedoFinalized:
-			DeferTransactionChanges.Reset();
-			break;
-		}
-	}
-
-	void Flush() override
-	{
-		TSet<TWeakObjectPtr<UMovieSceneSignedObject>> SignedObjectsTmp = SignedObjects;
-		SignedObjects.Empty();
-		// we operate on a copy of the signed objects in case the delegates would modify the array
-		for (TWeakObjectPtr<UMovieSceneSignedObject> WeakObject : SignedObjectsTmp)
-		{
-			if (UMovieSceneSignedObject* Object = WeakObject.Get())
-			{
-				Object->BroadcastChanged();
-			}
-		}
-	}
-
-	void DeferMarkAsChanged(UMovieSceneSignedObject* SignedObject) override
-	{
-		SignedObjects.Add(SignedObject);
-	}
-
-	bool CreateImplicitScopedModifyDefer() override
-	{
-		ensure(!DeferImplicitChanges.IsSet());
-		DeferImplicitChanges.Emplace();
-		return true;
-	}
-
-	void ResetImplicitScopedModifyDefer() override
-	{
-		DeferImplicitChanges.Reset();
-	}
-
-	TSet<TWeakObjectPtr<UMovieSceneSignedObject>> SignedObjects;
-	TWeakObjectPtr<UTransBuffer>   WeakBuffer;
-	TOptional<UE::MovieScene::FScopedSignedObjectModifyDefer> DeferTransactionChanges;
-	TOptional<UE::MovieScene::FScopedSignedObjectModifyDefer> DeferImplicitChanges;
-};
-
-} // namespace Sequencer
-} // namespace UE
-
-
 // Destructor defined in CPP to avoid having to #include SequencerChannelInterface.h in the main module definition
 ISequencerModule::~ISequencerModule()
 {
@@ -255,6 +156,8 @@ public:
 
 	virtual TSharedRef<ISequencer> CreateSequencer(const FSequencerInitParams& InitParams) override
 	{
+		using namespace UE::Sequencer;
+
 		TSharedRef<FSequencer> Sequencer = MakeShared<FSequencer>();
 		TSharedRef<ISequencerObjectChangeListener> ObjectChangeListener = MakeShared<FSequencerObjectChangeListener>(Sequencer);
 
@@ -427,13 +330,6 @@ public:
 				FCoreDelegates::OnPostEngineInit.AddRaw(this, &FSequencerModule::RegisterMenus);
 			}
 
-			// Set the deferred handler for use by Sequence editors. It is not necessary in commandlets.
-			// TODO: Create/Delete the handler when Sequence editors are opened/closed.
-			if (!IsRunningCommandlet())
-			{
-				UMovieSceneSignedObject::SetDeferredHandler(MakeUnique<FDeferredSignedObjectChangeHandler>());
-			}
-
 			FPropertyEditorModule& EditModule = FModuleManager::Get().GetModuleChecked<FPropertyEditorModule>("PropertyEditor");
 			OnGetGlobalRowExtensionHandle = EditModule.GetGlobalRowExtensionDelegate().AddStatic(&RegisterKeyframeExtensionHandler);
 		}
diff --git a/Engine/Source/Runtime/MovieScene/Private/MovieSceneSignedObject.cpp b/Engine/Source/Runtime/MovieScene/Private/MovieSceneSignedObject.cpp
--- a/Engine/Source/Runtime/MovieScene/Private/MovieSceneSignedObject.cpp
+++ b/Engine/Source/Runtime/MovieScene/Private/MovieSceneSignedObject.cpp
@@ -14,7 +14,7 @@ namespace MovieScene
 {
 
 static TSet<TWeakPtr<IDeferredSignedObjectFlushSignal>> GDeferredSignedObjectFlushSignals;
-static TUniquePtr<IDeferredSignedObjectChangeHandler> GDeferredSignedObjectChangeHandler;
+static TWeakPtr<IDeferredSignedObjectChangeHandler> GDeferredSignedObjectChangeHandler;
 static uint32 GScopedSignedObjectDeferCount = 0;
 
 void SignalScopedSignedObjectModifyFlush()
@@ -44,9 +44,9 @@ FScopedSignedObjectModifyDefer::~FScopedSignedObjectModifyDefer()
 
 	if (GScopedSignedObjectDeferCount == 0 || bForceFlush)
 	{
-		if (GDeferredSignedObjectChangeHandler)
+		if (GDeferredSignedObjectChangeHandler.IsValid())
 		{
-			GDeferredSignedObjectChangeHandler->Flush();
+			GDeferredSignedObjectChangeHandler.Pin()->Flush();
 		}
 		SignalScopedSignedObjectModifyFlush();
 	}
@@ -61,17 +61,22 @@ UMovieSceneSignedObject::UMovieSceneSignedObject(const FObjectInitializer& Init)
 {
 	using namespace UE::MovieScene;
 
- 	if (HasAnyFlags(RF_Transactional) && GDeferredSignedObjectChangeHandler 
+ 	if (HasAnyFlags(RF_Transactional) && GDeferredSignedObjectChangeHandler.IsValid() 
 #if WITH_EDITOR
 		&& GIsTransacting
 #endif
 	)
 	{
-		GDeferredSignedObjectChangeHandler->DeferMarkAsChanged(this);
+		GDeferredSignedObjectChangeHandler.Pin()->DeferMarkAsChanged(this);
 	}
 }
 
-void UMovieSceneSignedObject::SetDeferredHandler(TUniquePtr<UE::MovieScene::IDeferredSignedObjectChangeHandler>&& InHandler)
+TWeakPtr<UE::MovieScene::IDeferredSignedObjectChangeHandler> UMovieSceneSignedObject::GetDeferredHandler()
+{
+	return UE::MovieScene::GDeferredSignedObjectChangeHandler;
+}
+
+void UMovieSceneSignedObject::SetDeferredHandler(TWeakPtr<UE::MovieScene::IDeferredSignedObjectChangeHandler>&& InHandler)
 {
 	UE::MovieScene::GDeferredSignedObjectChangeHandler = MoveTemp(InHandler);
 }
@@ -84,9 +89,9 @@ void UMovieSceneSignedObject::AddFlushSignal(TWeakPtr<UE::MovieScene::IDeferredS
 void UMovieSceneSignedObject::ResetImplicitScopedModifyDefer()
 {
 	using namespace UE::MovieScene;
-	if (GDeferredSignedObjectChangeHandler)
+	if (GDeferredSignedObjectChangeHandler.IsValid())
 	{
-		GDeferredSignedObjectChangeHandler->ResetImplicitScopedModifyDefer();
+		GDeferredSignedObjectChangeHandler.Pin()->ResetImplicitScopedModifyDefer();
 	}
 }
 
@@ -130,19 +135,19 @@ void UMovieSceneSignedObject::MarkAsChanged()
 
 	// Give the change handler an opportunity to create an implicit scope if there's
 	// no explicit one active.
-	if (GScopedSignedObjectDeferCount == 0 && GDeferredSignedObjectChangeHandler)
+	if (GScopedSignedObjectDeferCount == 0 && GDeferredSignedObjectChangeHandler.IsValid())
 	{
-		GDeferredSignedObjectChangeHandler->CreateImplicitScopedModifyDefer();
+		GDeferredSignedObjectChangeHandler.Pin()->CreateImplicitScopedModifyDefer();
 	}
 
-	if (GScopedSignedObjectDeferCount == 0 || !GDeferredSignedObjectChangeHandler)
+	if (GScopedSignedObjectDeferCount == 0 || !GDeferredSignedObjectChangeHandler.IsValid())
 	{
 		BroadcastChanged();
 		SignalScopedSignedObjectModifyFlush();
 	}
 	else
 	{
-		GDeferredSignedObjectChangeHandler->DeferMarkAsChanged(this);
+		GDeferredSignedObjectChangeHandler.Pin()->DeferMarkAsChanged(this);
 	}
 }
 
diff --git a/Engine/Source/Runtime/MovieScene/Public/MovieSceneSignedObject.h b/Engine/Source/Runtime/MovieScene/Public/MovieSceneSignedObject.h
--- a/Engine/Source/Runtime/MovieScene/Public/MovieSceneSignedObject.h
+++ b/Engine/Source/Runtime/MovieScene/Public/MovieSceneSignedObject.h
@@ -67,7 +67,8 @@ public:
 
 	UMovieSceneSignedObject(const FObjectInitializer& Init);
 
-	static void SetDeferredHandler(TUniquePtr<UE::MovieScene::IDeferredSignedObjectChangeHandler>&& InHandler);
+	static TWeakPtr<UE::MovieScene::IDeferredSignedObjectChangeHandler> GetDeferredHandler();
+	static void SetDeferredHandler(TWeakPtr<UE::MovieScene::IDeferredSignedObjectChangeHandler>&& InHandler);
 	static void AddFlushSignal(TWeakPtr<UE::MovieScene::IDeferredSignedObjectFlushSignal> Signal);
 	static void ResetImplicitScopedModifyDefer();
 
