Sequencer: Fix deleting tag so that it happens on the root sequence

diff --git a/Engine/Source/Editor/Sequencer/Private/MVVM/ViewModels/ObjectBindingModel.cpp b/Engine/Source/Editor/Sequencer/Private/MVVM/ViewModels/ObjectBindingModel.cpp
--- a/Engine/Source/Editor/Sequencer/Private/MVVM/ViewModels/ObjectBindingModel.cpp
+++ b/Engine/Source/Editor/Sequencer/Private/MVVM/ViewModels/ObjectBindingModel.cpp
@@ -1350,7 +1350,7 @@ void FObjectBindingModel::HandleDeleteTag(FName TagName)
 	FScopedTransaction Transaction(FText::Format(LOCTEXT("RemoveBindingTag", "Remove tag '{0}' from binding(s)"), FText::FromName(TagName)));
 
 	TSharedPtr<ISequencer> Sequencer = OwnerModel->GetSequencer();
-	UMovieScene* MovieScene = OwnerModel->GetMovieScene();
+	UMovieScene* MovieScene = Sequencer->GetRootMovieSceneSequence()->GetMovieScene();
 	MovieScene->Modify();
 
 	FMovieSceneSequenceID SequenceID = OwnerModel->GetSequenceID();
@@ -1444,7 +1444,10 @@ bool FObjectBindingModel::CanDelete(FText* OutErrorMessage) const
 
 void FObjectBindingModel::Delete()
 {
-	UMovieScene* MovieScene = OwnerModel->GetMovieScene();
+	TSharedPtr<ISequencer> Sequencer = OwnerModel->GetSequencer();
+	UMovieScene* MovieScene = Sequencer->GetRootMovieSceneSequence()->GetMovieScene();
+
+	MovieScene->Modify();
 
 	// Untag this binding
 	UE::MovieScene::FFixedObjectBindingID BindingID(ObjectBindingID, OwnerModel->GetSequenceID());
