Sequencer: Add Browse to Actor in Sequencer back to the actor's right click menu

Refactored to use Tool Menus

diff --git a/Engine/Source/Editor/Sequencer/Private/LevelEditorSequencerIntegration.cpp b/Engine/Source/Editor/Sequencer/Private/LevelEditorSequencerIntegration.cpp
--- a/Engine/Source/Editor/Sequencer/Private/LevelEditorSequencerIntegration.cpp
+++ b/Engine/Source/Editor/Sequencer/Private/LevelEditorSequencerIntegration.cpp
@@ -52,7 +52,7 @@
 #include "Editor/UnrealEdEngine.h"
 #include "EditorSupportDelegates.h"
 #include "Subsystems/UnrealEditorSubsystem.h"
-
+#include "ToolMenus.h"
 
 #define LOCTEXT_NAMESPACE "LevelEditorSequencerIntegration"
 
@@ -246,7 +246,16 @@ void FLevelEditorSequencerIntegration::Initialize(const FLevelEditorSequencerInt
 		AcquiredResources.Add([=]{ FCoreDelegates::OnActorLabelChanged.Remove(Handle); });
 	}
 
-	AddLevelViewportMenuExtender();
+	// Menus need to be registered in a callback to make sure the system is ready for them.
+	{
+		UToolMenus::RegisterStartupCallback(
+			FSimpleMulticastDelegate::FDelegate::CreateRaw(this, &FLevelEditorSequencerIntegration::RegisterMenus));
+		AcquiredResources.Add([this]{ 	// Clean up menu things
+			UToolMenus::UnRegisterStartupCallback(this);
+			UToolMenus::UnregisterOwner(this);
+		});
+	}
+
 	ActivateDetailHandler(Options);
 
 	{
@@ -709,26 +718,6 @@ void FLevelEditorSequencerIntegration::OnEndPIE(bool bIsSimulating)
 	OnEndPlayMap();
 }
 
-void FLevelEditorSequencerIntegration::AddLevelViewportMenuExtender()
-{
-	typedef FLevelEditorModule::FLevelViewportMenuExtender_SelectedActors DelegateType;
-
-	FLevelEditorModule& LevelEditorModule = FModuleManager::Get().LoadModuleChecked<FLevelEditorModule>("LevelEditor");
-	auto& MenuExtenders = LevelEditorModule.GetAllLevelViewportContextMenuExtenders();
-	MenuExtenders.Add(DelegateType::CreateRaw(this, &FLevelEditorSequencerIntegration::GetLevelViewportExtender));
-
-	FDelegateHandle LevelViewportExtenderHandle = MenuExtenders.Last().GetHandle();
-	AcquiredResources.Add(
-		[LevelViewportExtenderHandle]{
-			FLevelEditorModule* LevelEditorModulePtr = FModuleManager::Get().GetModulePtr<FLevelEditorModule>("LevelEditor");
-			if (LevelEditorModulePtr)
-			{
-				LevelEditorModulePtr->GetAllLevelViewportContextMenuExtenders().RemoveAll([=](const DelegateType& In){ return In.GetHandle() == LevelViewportExtenderHandle; });
-			}
-		}
-	);
-}
-
 void FindActorInSequencesRecursive(AActor* InActor, FSequencer& Sequencer, FMovieSceneSequenceIDRef SequenceID, TArray<TPair<FMovieSceneSequenceID, FSequencer*> >& FoundInSequences)
 {
 	FMovieSceneRootEvaluationTemplateInstance& RootInstance = Sequencer.GetEvaluationTemplate();
@@ -788,49 +777,67 @@ void FindActorInSequencesRecursive(AActor* InActor, FSequencer& Sequencer, FMovi
 	}
 }
 
-TSharedRef<FExtender> FLevelEditorSequencerIntegration::GetLevelViewportExtender(const TSharedRef<FUICommandList> CommandList, const TArray<AActor*> InActors)
+void FLevelEditorSequencerIntegration::RegisterMenus()
 {
-	TSharedRef<FExtender> Extender = MakeShareable(new FExtender);
+	// Allows cleanup when module unloads.
+	FToolMenuOwnerScoped OwnerScoped(this);
 
-	FText ActorName;
-	if (InActors.Num() == 1)
 	{
-		ActorName = FText::Format(LOCTEXT("ActorNameSingular", "\"{0}\""), FText::FromString(InActors[0]->GetActorLabel()));
-	}
-	else if (InActors.Num() > 1)
-	{
-		ActorName = FText::Format(LOCTEXT("ActorNamePlural", "{0} Actors"), FText::AsNumber(InActors.Num()));
+		UToolMenu* LevelEditorMenu = UToolMenus::Get()->ExtendMenu("LevelEditor");
+		FToolMenuSection& UEToolsSection = LevelEditorMenu->FindOrAddSection("ActorUETools");
+		UEToolsSection.AddSubMenu(
+			"BrowseToActorSubMenu",
+			LOCTEXT("BrowseToActorInSequencer", "Browse to Actor in Sequencer"),
+			FText(),
+			FNewToolMenuDelegate::CreateRaw(this, &FLevelEditorSequencerIntegration::MakeBrowseToSelectedActorSubMenu),
+			false,
+			FSlateIcon("LevelSequenceEditorStyle", "LevelSequenceEditor.Tabs.Sequencer")
+		);
 	}
 
-	TArray<TPair<FMovieSceneSequenceID, FSequencer*> > FoundInSequences;
-	IterateAllSequencers(
-		[&](FSequencer& In, const FLevelEditorSequencerIntegrationOptions& Options)
 	{
-		FindActorInSequencesRecursive(InActors[0], In, MovieSceneSequenceID::Root, FoundInSequences);
-	});
-
-	TSharedRef<FUICommandList> LevelEditorCommandBindings  = FModuleManager::GetModuleChecked<FLevelEditorModule>(TEXT("LevelEditor")).GetGlobalLevelEditorActions();
-	Extender->AddMenuExtension("ActorUETools", EExtensionHook::After, LevelEditorCommandBindings, FMenuExtensionDelegate::CreateLambda(
-		[this, ActorName, Actor = InActors[0], FoundInSequences](FMenuBuilder& MenuBuilder) {
-		bool bCanBrowse = FoundInSequences.Num() > 0;
-		MenuBuilder.AddSubMenu(
+		UToolMenu* ActorContextMenu = UToolMenus::Get()->ExtendMenu("LevelEditor.ActorContextMenu");
+		FToolMenuSection& ActorTypeToolsSection = ActorContextMenu->FindOrAddSection("ActorTypeTools");
+		ActorTypeToolsSection.AddSubMenu(
+			"BrowseToActorSubMenu",
 			LOCTEXT("BrowseToActorInSequencer", "Browse to Actor in Sequencer"),
 			FText(),
-			FNewMenuDelegate::CreateRaw(this, &FLevelEditorSequencerIntegration::MakeBrowseToSelectedActorSubMenu, Actor, FoundInSequences),
-			FUIAction(FExecuteAction(), FCanExecuteAction::CreateLambda([bCanBrowse]() { return bCanBrowse; })),
-			NAME_None,
-			EUserInterfaceActionType::Button,
+			FNewToolMenuDelegate::CreateRaw(this, &FLevelEditorSequencerIntegration::MakeBrowseToSelectedActorSubMenu),
 			false,
-			FSlateIcon()
+			FSlateIcon("LevelSequenceEditorStyle", "LevelSequenceEditor.Tabs.Sequencer")
 		);
 	}
-	));
-
-	return Extender;
 }
 
-void FLevelEditorSequencerIntegration::MakeBrowseToSelectedActorSubMenu(FMenuBuilder& MenuBuilder, AActor* Actor, const TArray<TPair<FMovieSceneSequenceID, FSequencer*> > FoundInSequences)
+void FLevelEditorSequencerIntegration::MakeBrowseToSelectedActorSubMenu(UToolMenu* Menu)
 {
+	AActor* Actor = nullptr;
+	TArray<TPair<FMovieSceneSequenceID, FSequencer*> > FoundInSequences;
+
+	for (FSelectionIterator It(GEditor->GetSelectedActorIterator()); It; ++It)
+	{
+		// We are interested in the (unique) assets backing the actor, or else the actor
+		// itself if it is not asset backed (such as UDynamicMesh).
+		Actor = static_cast<AActor*>(*It);
+
+		IterateAllSequencers(
+			[&](FSequencer& In, const FLevelEditorSequencerIntegrationOptions& Options)
+			{
+				FindActorInSequencesRecursive(Actor, In, MovieSceneSequenceID::Root, FoundInSequences);
+			});
+
+		if (Actor)
+		{
+			break;
+		}
+	};
+
+	if (!Actor || !FoundInSequences.Num())
+	{
+		return;
+	}
+
+	FToolMenuSection& Section = Menu->AddSection("BrowseToActorSection");
 	for (const TPair<FMovieSceneSequenceID, FSequencer*>& Sequence : FoundInSequences)
 	{
 		UMovieSceneSequence* MovieSceneSequence = nullptr;
@@ -848,7 +855,8 @@ void FLevelEditorSequencerIntegration::MakeBrowseToSelectedActorSubMenu(FMenuBui
 		{
 			FText ActorName = FText::Format(LOCTEXT("ActorNameSingular", "\"{0}\""), FText::FromString(Actor->GetActorLabel()));
 			FUIAction AddMenuAction(FExecuteAction::CreateLambda([this, Actor, Sequence]() {this->BrowseToSelectedActor(Actor, Sequence.Value, Sequence.Key); }));
-			MenuBuilder.AddMenuEntry(FText::Format(LOCTEXT("BrowseToSelectedActorText", "Browse to {0} in {1}"), ActorName, MovieSceneSequence->GetDisplayName()), FText(), FSlateIcon(), AddMenuAction);
+			FText MenuName = FText::Format(LOCTEXT("BrowseToSelectedActorText", "Browse to {0} in {1}"), ActorName, MovieSceneSequence->GetDisplayName());
+			Section.AddMenuEntry(FName(*MenuName.ToString()), MenuName, FText(), FSlateIcon(), AddMenuAction);
 		}
 	}
 }
@@ -1201,13 +1209,6 @@ void FLevelEditorSequencerIntegration::ResetToAnimatedState(UWorld* World)
 	);
 }
 
-TSharedRef<FExtender> FLevelEditorSequencerIntegration::OnExtendLevelEditorViewMenu(const TSharedRef<FUICommandList> CommandList)
-{
-	// This would be where you added an extension to the Level Editor dropdown menus if needed.
-	TSharedRef<FExtender> Extender(new FExtender());
-	return Extender;
-}
-
 void FLevelEditorSequencerIntegration::OnTabContentChanged()
 {
 
diff --git a/Engine/Source/Editor/Sequencer/Public/LevelEditorSequencerIntegration.h b/Engine/Source/Editor/Sequencer/Public/LevelEditorSequencerIntegration.h
--- a/Engine/Source/Editor/Sequencer/Public/LevelEditorSequencerIntegration.h
+++ b/Engine/Source/Editor/Sequencer/Public/LevelEditorSequencerIntegration.h
@@ -20,6 +20,7 @@ class FUICommandList;
 class IAssetViewport;
 class ISequencer;
 class ULevel;
+class UToolMenu;
 struct FPropertyAndParent;
 
 
@@ -172,11 +173,9 @@ private:
 
 	void OnPropertyEditorOpened();
 
-	TSharedRef<FExtender> GetLevelViewportExtender(const TSharedRef<FUICommandList> CommandList, const TArray<AActor*> InActors);
+	void RegisterMenus();
 
-	TSharedRef<FExtender> OnExtendLevelEditorViewMenu(const TSharedRef<FUICommandList> CommandList);
-
-	void MakeBrowseToSelectedActorSubMenu(FMenuBuilder& MenuBuilder, AActor* Actor, const TArray<TPair<FMovieSceneSequenceID, FSequencer*> > FoundInSequences);
+	void MakeBrowseToSelectedActorSubMenu(UToolMenu* Menu);
 	void BrowseToSelectedActor(AActor* Actor, FSequencer* Sequencer, FMovieSceneSequenceID SequenceId);
 
 	bool IsPropertyReadOnly(const FPropertyAndParent& InPropertyAndParent);
