Sequencer: Request group manager tree view refresh when updating

diff --git a/Engine/Source/Editor/Sequencer/Private/SSequencerGroupManager.cpp b/Engine/Source/Editor/Sequencer/Private/SSequencerGroupManager.cpp
--- a/Engine/Source/Editor/Sequencer/Private/SSequencerGroupManager.cpp
+++ b/Engine/Source/Editor/Sequencer/Private/SSequencerGroupManager.cpp
@@ -432,6 +432,9 @@ void SSequencerGroupManager::UpdateTree()
 		}
 	}
 
+
+	TreeView->RequestTreeRefresh();
+
 	bNodeGroupsDirty = false;
 }
 
@@ -572,6 +575,8 @@ void SSequencerGroupManager::CreateNodeGroup()
 
 	const FScopedTransaction Transaction(LOCTEXT("CreateNewGroupTransaction", "Create New Group"));
 
+	MovieScene->Modify();
+
 	UMovieSceneNodeGroup* NewNodeGroup = NewObject<UMovieSceneNodeGroup>(&MovieScene->GetNodeGroups(), NAME_None, RF_Transactional);
 	NewNodeGroup->SetName(FSequencerUtilities::GetUniqueName(FName("Group"), ExistingGroupNames));
 
