Sequencer: Refresh node groups when adding/removing items

diff --git a/Engine/Source/Editor/Sequencer/Private/SSequencerGroupManager.cpp b/Engine/Source/Editor/Sequencer/Private/SSequencerGroupManager.cpp
--- a/Engine/Source/Editor/Sequencer/Private/SSequencerGroupManager.cpp
+++ b/Engine/Source/Editor/Sequencer/Private/SSequencerGroupManager.cpp
@@ -432,7 +432,6 @@ void SSequencerGroupManager::UpdateTree()
 		}
 	}
 
-
 	TreeView->RequestTreeRefresh();
 
 	bNodeGroupsDirty = false;
@@ -551,6 +550,7 @@ void SSequencerGroupManager::RemoveSelectedItemsFromNodeGroup()
 		Item.Key->RemoveNode(Item.Value);
 	}
 	
+	RefreshNodeGroups();
 }
 
 void SSequencerGroupManager::CreateNodeGroup()
diff --git a/Engine/Source/Editor/Sequencer/Private/Sequencer.cpp b/Engine/Source/Editor/Sequencer/Private/Sequencer.cpp
--- a/Engine/Source/Editor/Sequencer/Private/Sequencer.cpp
+++ b/Engine/Source/Editor/Sequencer/Private/Sequencer.cpp
@@ -6279,6 +6279,11 @@ void FSequencer::AddNodesToExistingNodeGroup(const TArray<TWeakPtr<UE::Sequencer
 		}
 	}
 
+	TSharedPtr<SSequencerGroupManager> NodeGroupManager = SequencerWidget->GetNodeGroupsManager();
+	if (NodeGroupManager)
+	{
+		NodeGroupManager->RefreshNodeGroups();
+	}
 }
 
 void FSequencer::ClearFilters()
