Fix crash accessing null pointer on shutdown of editor.

diff --git a/Engine/Source/Editor/Sequencer/Private/LevelEditorSequencerIntegration.cpp b/Engine/Source/Editor/Sequencer/Private/LevelEditorSequencerIntegration.cpp
--- a/Engine/Source/Editor/Sequencer/Private/LevelEditorSequencerIntegration.cpp
+++ b/Engine/Source/Editor/Sequencer/Private/LevelEditorSequencerIntegration.cpp
@@ -1159,14 +1159,17 @@ void FLevelEditorSequencerIntegration::RestoreRealtimeViewports()
 	// it anymore. It could happen that a viewport doesn't have it if that viewport is an actual Cinematic Viewport, for instance.
 	const bool bCheckMissingOverride = false;
 
-	for(FLevelEditorViewportClient* LevelVC : GEditor->GetLevelViewportClients())
+	if (GEditor)
 	{
-		if (LevelVC)
+		for(FLevelEditorViewportClient* LevelVC : GEditor->GetLevelViewportClients())
 		{
-			// Turn off realtime when exiting.
-			if( LevelVC->IsPerspective() && LevelVC->AllowsCinematicControl() )
-			{				
-				LevelVC->RemoveRealtimeOverride(LOCTEXT("RealtimeOverrideMessage_Sequencer", "Sequencer"), bCheckMissingOverride);
+			if (LevelVC)
+			{
+				// Turn off realtime when exiting.
+				if( LevelVC->IsPerspective() && LevelVC->AllowsCinematicControl() )
+				{				
+					LevelVC->RemoveRealtimeOverride(LOCTEXT("RealtimeOverrideMessage_Sequencer", "Sequencer"), bCheckMissingOverride);
+				}
 			}
 		}
 	}
