Sequencer: Optimized creation of the snap field when dealing with millions of snap points

  - Changed ISequencerEntityVisitor::VisitKey is now VisitKeys - calling a single virtual function for each key was non-sensical and slow
  - Snap points are now gathered into a chunked array (TPagedArray) so as to prevent excessive reallocation and memmoving duing population of the structure
  - ISnapCandidate can now bypass its virtual functions by specifying capabilities and default states for each snap type

These are spot fixes really - we should probably be hashing snap points to a sensible quantization based on the current zoom level to avoid accumulating huge arrays with millions of entries

diff --git a/Engine/Source/Editor/Sequencer/Private/Tools/SequencerEditTool_Selection.cpp b/Engine/Source/Editor/Sequencer/Private/Tools/SequencerEditTool_Selection.cpp
--- a/Engine/Source/Editor/Sequencer/Private/Tools/SequencerEditTool_Selection.cpp
+++ b/Engine/Source/Editor/Sequencer/Private/Tools/SequencerEditTool_Selection.cpp
@@ -7,6 +7,7 @@
 #include "ISequencerEditTool.h"
 #include "ISequencerSection.h"
 #include "SequencerHotspots.h"
+#include "IKeyArea.h"
 #include "Tools/SequencerEntityVisitor.h"
 #include "MVVM/Views/ISequencerTreeView.h"
 #include "MVVM/Views/ITrackAreaHotspot.h"
@@ -32,7 +33,7 @@ struct FSelectionPreviewVisitor final
 		, bPinned(bInPinned)
 	{}
 
-	virtual void VisitKey(FKeyHandle KeyHandle, FFrameNumber KeyTime, const UE::Sequencer::TViewModelPtr<UE::Sequencer::FChannelModel>& Channel, UMovieSceneSection* Section) const override
+	virtual void VisitKeys(const UE::Sequencer::TViewModelPtr<UE::Sequencer::FChannelModel>& Channel, const TRange<FFrameNumber>& VisitRangeFrames) const override
 	{
 		using namespace UE::Sequencer;
 
@@ -42,21 +43,33 @@ struct FSelectionPreviewVisitor final
 			return;
 		}
 
-		// Under default behavior keys have priority, so if a key is changing selection state then we remove any sections from the selection. The user can bypass this
-		// by holding down the control key which will allow selecting both keys and sections.
-		bool bKeySelectionHasPriority = !FSlateApplication::Get().GetModifierKeys().IsControlDown();
-		bool bKeyIsSelected = ExistingSelection.Contains(KeyHandle);
-
-		if (bKeySelectionHasPriority && 
-			((bKeyIsSelected && SetStateTo == ESelectionPreviewState::NotSelected) ||
-			(!bKeyIsSelected && SetStateTo == ESelectionPreviewState::Selected)))
+		UMovieSceneSection* Section = Channel->GetSection();
+		if (!Section)
 		{
-			// Clear selected models
-			SelectionPreview.EmptyDefinedModelStates();
+			return;
 		}
 
-		FSequencerSelectedKey Key(*Section, TSharedPtr<FChannelModel>(Channel), KeyHandle);
-		SelectionPreview.SetSelectionState(Key, SetStateTo);
+		KeyHandlesScratch.Reset();
+		Channel->GetKeyArea()->GetKeyInfo(&KeyHandlesScratch, nullptr, VisitRangeFrames);
+
+		for (int32 Index = 0; Index < KeyHandlesScratch.Num(); ++Index)
+		{
+			// Under default behavior keys have priority, so if a key is changing selection state then we remove any sections from the selection. The user can bypass this
+			// by holding down the control key which will allow selecting both keys and sections.
+			bool bKeySelectionHasPriority = !FSlateApplication::Get().GetModifierKeys().IsControlDown();
+			bool bKeyIsSelected = ExistingSelection.Contains(KeyHandlesScratch[Index]);
+
+			if (bKeySelectionHasPriority && 
+				((bKeyIsSelected && SetStateTo == ESelectionPreviewState::NotSelected) ||
+				(!bKeyIsSelected && SetStateTo == ESelectionPreviewState::Selected)))
+			{
+				// Clear selected models
+				SelectionPreview.EmptyDefinedModelStates();
+			}
+
+			FSequencerSelectedKey Key(*Section, Channel, KeyHandlesScratch[Index]);
+			SelectionPreview.SetSelectionState(Key, SetStateTo);
+		}
 	}
 
 	virtual void VisitDataModel(UE::Sequencer::FViewModel* DataModel) const override
@@ -84,6 +97,7 @@ private:
 
 	FSequencerSelectionPreview& SelectionPreview;
 	const TSet<FKeyHandle>& ExistingSelection;
+	mutable TArray<FKeyHandle> KeyHandlesScratch;
 	ESelectionPreviewState SetStateTo;
 	bool bPinned;
 };
diff --git a/Engine/Source/Editor/Sequencer/Private/Tools/SequencerEntityVisitor.cpp b/Engine/Source/Editor/Sequencer/Private/Tools/SequencerEntityVisitor.cpp
--- a/Engine/Source/Editor/Sequencer/Private/Tools/SequencerEntityVisitor.cpp
+++ b/Engine/Source/Editor/Sequencer/Private/Tools/SequencerEntityVisitor.cpp
@@ -209,28 +209,15 @@ void FSequencerEntityWalker::VisitChannel(const ISequencerEntityVisitor& Visitor
 {
 	using namespace UE::Sequencer;
 
-	TSharedPtr<IKeyArea> KeyArea = Channel->GetKeyArea();
-	UMovieSceneSection* Section = KeyArea->GetOwningSection();
-	if (!Section)
+	UMovieSceneSection* Section = Channel->GetSection();
+	if (Section)
 	{
-		return;
-	}
-
-	TArray<FKeyHandle> Handles;
-	TArray<FFrameNumber> Times;
-
-	const FFrameTime HalfKeySizeFrames   = (VirtualKeySize.X*.5f) * Range.TickResolution;
-	const FFrameTime RangeStartFrame     = Range.Range.GetLowerBoundValue() * Range.TickResolution;
-	const FFrameTime RangeEndFrame       = Range.Range.GetUpperBoundValue() * Range.TickResolution;
+		const FFrameTime HalfKeySizeFrames   = (VirtualKeySize.X*.5f) * Range.TickResolution;
+		const FFrameTime RangeStartFrame     = Range.Range.GetLowerBoundValue() * Range.TickResolution;
+		const FFrameTime RangeEndFrame       = Range.Range.GetUpperBoundValue() * Range.TickResolution;
 
-	TRange<FFrameNumber> VisitRangeFrames( (RangeStartFrame-HalfKeySizeFrames).CeilToFrame(), (RangeEndFrame+HalfKeySizeFrames).FloorToFrame() );
-
-	VisitRangeFrames = TRange<FFrameNumber>::Intersection(Section->GetRange(), VisitRangeFrames);
-
-	KeyArea->GetKeyInfo(&Handles, &Times, VisitRangeFrames);
-
-	for (int32 Index = 0; Index < Times.Num(); ++Index)
-	{
-		Visitor.VisitKey(Handles[Index], Times[Index], Channel, Section);
+		TRange<FFrameNumber> VisitRangeFrames( (RangeStartFrame-HalfKeySizeFrames).CeilToFrame(), (RangeEndFrame+HalfKeySizeFrames).FloorToFrame() );
+		VisitRangeFrames = TRange<FFrameNumber>::Intersection(Section->GetRange(), VisitRangeFrames);
+		Visitor.VisitKeys(Channel, VisitRangeFrames);
 	}
 }
diff --git a/Engine/Source/Editor/Sequencer/Private/Tools/SequencerEntityVisitor.h b/Engine/Source/Editor/Sequencer/Private/Tools/SequencerEntityVisitor.h
--- a/Engine/Source/Editor/Sequencer/Private/Tools/SequencerEntityVisitor.h
+++ b/Engine/Source/Editor/Sequencer/Private/Tools/SequencerEntityVisitor.h
@@ -41,7 +41,8 @@ struct ISequencerEntityVisitor
 {
 	ISequencerEntityVisitor(uint32 InEntityMask = ESequencerEntity::Everything) : EntityMask(InEntityMask) {}
 
-	virtual void VisitKey(FKeyHandle KeyHandle, FFrameNumber KeyTime, const UE::Sequencer::TViewModelPtr<UE::Sequencer::FChannelModel>& Channel, UMovieSceneSection* Section) const { }
+	virtual void VisitKeys(const UE::Sequencer::TViewModelPtr<UE::Sequencer::FChannelModel>& Channel, const TRange<FFrameNumber>& VisitRangeFrames) const { }
+
 	virtual void VisitDataModel(UE::Sequencer::FViewModel* Item) const { }
 
 	/** Check if the specified type of entity is applicable to this visitor */
diff --git a/Engine/Source/Editor/Sequencer/Private/Tools/SequencerSnapField.cpp b/Engine/Source/Editor/Sequencer/Private/Tools/SequencerSnapField.cpp
--- a/Engine/Source/Editor/Sequencer/Private/Tools/SequencerSnapField.cpp
+++ b/Engine/Source/Editor/Sequencer/Private/Tools/SequencerSnapField.cpp
@@ -1,11 +1,13 @@
 // Copyright Epic Games, Inc. All Rights Reserved.
 
 #include "Tools/SequencerSnapField.h"
+#include "Containers/PagedArray.h"
 #include "MovieScene.h"
 #include "SSequencer.h"
 #include "MovieSceneTimeHelpers.h"
 #include "MovieSceneSequence.h"
 #include "ISequencerSection.h"
+#include "IKeyArea.h"
 #include "MVVM/ViewModels/ViewModel.h"
 #include "MVVM/Views/SOutlinerView.h"
 #include "MVVM/Extensions/ISnappableExtension.h"
@@ -15,17 +17,56 @@ struct FSnapGridVisitor : ISequencerEntityVisitor, UE::Sequencer::ISnapField
 	FSnapGridVisitor(UE::Sequencer::ISnapCandidate& InCandidate, uint32 EntityMask)
 		: ISequencerEntityVisitor(EntityMask)
 		, Candidate(InCandidate)
-	{}
+	{
+		InCandidate.GetCapabilities(Applicability, Capabilities);
+	}
 
-	virtual void VisitKey(FKeyHandle KeyHandle, FFrameNumber KeyTime, const UE::Sequencer::TViewModelPtr<UE::Sequencer::FChannelModel>& Channel, UMovieSceneSection* Section) const
+	virtual void VisitKeys(const UE::Sequencer::TViewModelPtr<UE::Sequencer::FChannelModel>& Channel, const TRange<FFrameNumber>& VisitRangeFrames) const override
 	{
 		using namespace UE::Sequencer;
 
-		if (Candidate.IsKeyApplicable(KeyHandle, Channel))
+		const bool bCallFunction = EnumHasAnyFlags(Capabilities, ISnapCandidate::ESnapCapabilities::IsKeyApplicable);
+		const bool bDefaultState = EnumHasAnyFlags(Applicability, ISnapCandidate::EDefaultApplicability::Keys);
+
+		if (!bCallFunction)
+		{
+			if (bDefaultState)
+			{
+				// Faster implementation if we need to add all the keys.
+				// Only need to allocate space for the times in this case
+				TimesScratch.Reset();
+				Channel->GetKeyArea()->GetKeyInfo(nullptr, &TimesScratch, VisitRangeFrames);
+
+				const int32 StartNum = Snaps.Num();
+				const int32 NumKeys = TimesScratch.Num();
+				Snaps.SetNum(Snaps.Num() + NumKeys);
+
+				for (int32 Index = 0; Index < NumKeys; ++Index)
+				{
+					Snaps[StartNum + Index] = FSnapPoint(FSnapPoint::Key, TimesScratch[Index]);
+				}
+			}
+		}
+		else
 		{
-			Snaps.Add(FSnapPoint{ FSnapPoint::Key, KeyTime });
+			// Call the function for each key individually. Much slower.
+			// Need to allocate handles as well
+			HandlesScratch.Reset();
+			TimesScratch.Reset();
+			Channel->GetKeyArea()->GetKeyInfo(&HandlesScratch, &TimesScratch, VisitRangeFrames);
+
+			for (int32 Index = 0; Index < HandlesScratch.Num(); ++Index)
+			{
+				FKeyHandle KeyHandle = HandlesScratch[Index];
+
+				if (Candidate.IsKeyApplicable(KeyHandle, Channel))
+				{
+					Snaps.Add(FSnapPoint{ FSnapPoint::Key, TimesScratch[Index] });
+				}
+			}
 		}
 	}
+
 	virtual void VisitDataModel(UE::Sequencer::FViewModel* DataModel) const
 	{
 		using namespace UE::Sequencer;
@@ -38,11 +79,17 @@ struct FSnapGridVisitor : ISequencerEntityVisitor, UE::Sequencer::ISnapField
 
 	virtual void AddSnapPoint(const UE::Sequencer::FSnapPoint& SnapPoint) override
 	{
-		Snaps.Add(SnapPoint);
+		Snaps.Emplace(SnapPoint);
 	}
 
 	UE::Sequencer::ISnapCandidate& Candidate;
-	mutable TArray<UE::Sequencer::FSnapPoint> Snaps;
+	/** utilize a chunked array to reduce the number of resize/grow memcpies for very large data sets (ie, > 1million keys) */
+	mutable TPagedArray<UE::Sequencer::FSnapPoint> Snaps;
+	mutable TArray<FKeyHandle> HandlesScratch;
+	mutable TArray<FFrameNumber> TimesScratch;
+
+	UE::Sequencer::ISnapCandidate::ESnapCapabilities Capabilities;
+	UE::Sequencer::ISnapCandidate::EDefaultApplicability Applicability;
 };
 
 FSequencerSnapField::FSequencerSnapField(const FSequencer& InSequencer, UE::Sequencer::ISnapCandidate& Candidate, uint32 EntityMask)
@@ -85,42 +132,35 @@ void FSequencerSnapField::Initialize(const FSequencer& InSequencer, UE::Sequence
 	TRange<FFrameNumber> PlaybackRange = InSequencer.GetFocusedMovieSceneSequence()->GetMovieScene()->GetPlaybackRange();
 	if(UE::MovieScene::DiscreteSize(PlaybackRange) > 0)
 	{
-		Visitor.Snaps.Add(FSnapPoint{ FSnapPoint::PlaybackRange, UE::MovieScene::DiscreteInclusiveLower(PlaybackRange)});
-		Visitor.Snaps.Add(FSnapPoint{ FSnapPoint::PlaybackRange, UE::MovieScene::DiscreteExclusiveUpper(PlaybackRange)});
+		Visitor.Snaps.Emplace(FSnapPoint{ FSnapPoint::PlaybackRange, UE::MovieScene::DiscreteInclusiveLower(PlaybackRange)});
+		Visitor.Snaps.Emplace(FSnapPoint{ FSnapPoint::PlaybackRange, UE::MovieScene::DiscreteExclusiveUpper(PlaybackRange)});
 	}
 
 	// Add the current time as a potential snap candidate
-	Visitor.Snaps.Add(FSnapPoint{ FSnapPoint::CurrentTime, InSequencer.GetLocalTime().Time.FrameNumber });
+	Visitor.Snaps.Emplace(FSnapPoint{ FSnapPoint::CurrentTime, InSequencer.GetLocalTime().Time.FrameNumber });
 
 	// Add the selection range bounds as a potential snap candidate
 	TRange<FFrameNumber> SelectionRange = InSequencer.GetFocusedMovieSceneSequence()->GetMovieScene()->GetSelectionRange();
 	if (UE::MovieScene::DiscreteSize(SelectionRange) > 0)
 	{
-		Visitor.Snaps.Add(FSnapPoint{ FSnapPoint::InOutRange, UE::MovieScene::DiscreteInclusiveLower(SelectionRange)});
-		Visitor.Snaps.Add(FSnapPoint{ FSnapPoint::InOutRange, UE::MovieScene::DiscreteExclusiveUpper(SelectionRange) - 1});
+		Visitor.Snaps.Emplace(FSnapPoint{ FSnapPoint::InOutRange, UE::MovieScene::DiscreteInclusiveLower(SelectionRange)});
+		Visitor.Snaps.Emplace(FSnapPoint{ FSnapPoint::InOutRange, UE::MovieScene::DiscreteExclusiveUpper(SelectionRange) - 1});
 	}
 
 	// Add in the marked frames
 	for (const FMovieSceneMarkedFrame& MarkedFrame : InSequencer.GetFocusedMovieSceneSequence()->GetMovieScene()->GetMarkedFrames())
 	{
-		Visitor.Snaps.Add( FSnapPoint{ FSnapPoint::Mark, MarkedFrame.FrameNumber } );
+		Visitor.Snaps.Emplace( FSnapPoint{ FSnapPoint::Mark, MarkedFrame.FrameNumber } );
 	}
 
 	// Add in the global marked frames
 	for (const FMovieSceneMarkedFrame& MarkedFrame : InSequencer.GetGlobalMarkedFrames())
 	{
-		Visitor.Snaps.Add(FSnapPoint{ FSnapPoint::Mark, MarkedFrame.FrameNumber });
+		Visitor.Snaps.Emplace(FSnapPoint{ FSnapPoint::Mark, MarkedFrame.FrameNumber });
 	}
 
-
-	if (SortedSnaps.Num() == 0)
-	{
-		SortedSnaps = MoveTemp(Visitor.Snaps);
-	}
-	else
-	{
-		SortedSnaps.Append(Visitor.Snaps);
-	}
+	// Copy the paged array to our linear array ready for final sorting
+	Visitor.Snaps.ToArray(SortedSnaps);
 
 	TickResolution = InSequencer.GetFocusedTickResolution();
 	DisplayRate = InSequencer.GetFocusedDisplayRate();
@@ -136,29 +176,33 @@ void FSequencerSnapField::Finalize()
 		return A.Time < B.Time;
 	});
 
+
+	const int32 NumSnaps = SortedSnaps.Num();
+
+	TArray<FSnapPoint> FinalSnaps;
+	FinalSnaps.Reserve(NumSnaps);
+
+	const FSnapPoint* const RESTRICT Snaps = SortedSnaps.GetData();
+
 	// Remove duplicates
-	for (int32 Index = 0; Index < SortedSnaps.Num(); ++Index)
+	for (int32 Index = 0; Index < NumSnaps; /* incremented inside inner loop */)
 	{
 		const FFrameNumber CurrentTime = SortedSnaps[Index].Time;
 
-		float FinalWeight = SortedSnaps[Index].Weighting;
-		int32 NumToMerge = 0;
-		for (int32 DuplIndex = Index + 1; DuplIndex < SortedSnaps.Num(); ++DuplIndex)
-		{
-			if (CurrentTime != SortedSnaps[DuplIndex].Time)
-			{
-				break;
-			}
-			++NumToMerge;
-			FinalWeight += SortedSnaps[DuplIndex].Weighting;
-		}
+		FSnapPoint FinalSnap = Snaps[Index];
+		FinalSnap.Weighting = 0.f;
 
-		if (NumToMerge)
+		// Add up all weights of the same time
+		for ( ; Index < NumSnaps && Snaps[Index].Time == CurrentTime; ++Index)
 		{
-			SortedSnaps[Index].Weighting = FinalWeight;
-			SortedSnaps.RemoveAt(Index + 1, NumToMerge, false);
+			FinalSnap.Weighting += Snaps[Index].Weighting;
 		}
+
+		FinalSnaps.Add(FinalSnap);
 	}
+
+	FinalSnaps.Shrink();
+	Swap(FinalSnaps, SortedSnaps);
 }
 
 TOptional<FSequencerSnapField::FSnapResult> FSequencerSnapField::Snap(const FFrameTime& InTime, const FFrameTime& Threshold) const
diff --git a/Engine/Source/Editor/Sequencer/Public/MVVM/Extensions/ISnappableExtension.h b/Engine/Source/Editor/Sequencer/Public/MVVM/Extensions/ISnappableExtension.h
--- a/Engine/Source/Editor/Sequencer/Public/MVVM/Extensions/ISnappableExtension.h
+++ b/Engine/Source/Editor/Sequencer/Public/MVVM/Extensions/ISnappableExtension.h
@@ -5,6 +5,7 @@
 #include "SequencerCoreFwd.h"
 #include "MVVM/ViewModelTypeID.h"
 #include "Curves/KeyHandle.h"
+#include "Templates/MemoryOps.h"
 
 struct FGuid;
 
@@ -46,6 +47,14 @@ struct SEQUENCER_API FSnapPoint
 
 	/** The type of snap */
 	ESnapType Type;
+
+private:
+
+	friend void ::DefaultConstructItems<FSnapPoint, int32>(void*, int32);
+
+	// uninitialized default constructor - only for use by containers
+	FSnapPoint()
+	{}
 };
 
 class SEQUENCER_API ISnapField
@@ -59,8 +68,33 @@ public:
 /** Interface that defines how to construct an FSequencerSnapField */
 struct SEQUENCER_API ISnapCandidate
 {
+	enum class EDefaultApplicability
+	{
+		None                   = 0,
+
+		Keys                   = 1 << 0,
+		SectionBounds          = 1 << 1,
+		CustomSnaps            = 1 << 2,
+
+		All = Keys | SectionBounds | CustomSnaps,
+	};
+
+	enum class ESnapCapabilities
+	{
+		None                            = 0,
+
+		IsKeyApplicable                 = 1 << 0,
+		AreSectionBoundsApplicable      = 1 << 1,
+		AreSectionCustomSnapsApplicable = 1 << 2,
+
+		All = IsKeyApplicable | AreSectionBoundsApplicable | AreSectionCustomSnapsApplicable,
+	};
+
 	virtual ~ISnapCandidate() { }
 
+	/** Retrieve the capabilities of this snap candidate instance as a bitmask. Zero bits will prevent virtual functions being called as an optimziation*/
+	virtual void GetCapabilities(EDefaultApplicability& OutDefaultApplicability, ESnapCapabilities& OutCapabilities) const { OutDefaultApplicability = EDefaultApplicability::All; OutCapabilities = ESnapCapabilities::None; }
+
 	/** Return true to include the specified key in the snap field */
 	virtual bool IsKeyApplicable(FKeyHandle KeyHandle, const UE::Sequencer::FViewModelPtr& Owner) const { return true; }
 
@@ -82,6 +116,8 @@ public:
 	virtual void AddToSnapField(const ISnapCandidate& Candidate, ISnapField& SnapField) const = 0;
 };
 
+ENUM_CLASS_FLAGS(ISnapCandidate::ESnapCapabilities);
+
 } // namespace Sequencer
 } // namespace UE
 
