Sequencer: Piloted Spawnables are now preserved across save boundaries

diff --git a/Engine/Source/Editor/Sequencer/Private/LevelEditorSequencerIntegration.cpp b/Engine/Source/Editor/Sequencer/Private/LevelEditorSequencerIntegration.cpp
--- a/Engine/Source/Editor/Sequencer/Private/LevelEditorSequencerIntegration.cpp
+++ b/Engine/Source/Editor/Sequencer/Private/LevelEditorSequencerIntegration.cpp
@@ -13,6 +13,8 @@
 #include "ILevelEditor.h"
 #include "IAssetViewport.h"
 #include "LevelEditor.h"
+#include "SLevelViewport.h"
+#include "MovieSceneSpawnableAnnotation.h"
 #include "Framework/Application/SlateApplication.h"
 #include "IDetailsView.h"
 #include "ISequencer.h"
@@ -56,6 +58,14 @@
 
 #define LOCTEXT_NAMESPACE "LevelEditorSequencerIntegration"
 
+
+struct FPilotedSpawnable
+{
+	TWeakPtr<SLevelViewport> WeakLevelViewport;
+	FLevelViewportActorLock PreviousActorLock;
+	FMovieSceneSpawnableAnnotation Annotation;
+};
+
 class FDetailKeyframeHandlerWrapper : public IDetailKeyframeHandler
 {
 public:
@@ -157,6 +167,8 @@ FLevelEditorSequencerIntegration::FLevelEditorSequencerIntegration()
 	KeyFrameHandler = MakeShared<FDetailKeyframeHandlerWrapper>();
 }
 
+FLevelEditorSequencerIntegration::~FLevelEditorSequencerIntegration() = default;
+
 FLevelEditorSequencerIntegration& FLevelEditorSequencerIntegration::Get()
 {
 	static FLevelEditorSequencerIntegration Singleton;
@@ -480,6 +492,8 @@ void FLevelEditorSequencerIntegration::OnNewActorsDropped(const TArray<UObject*>
 
 void FLevelEditorSequencerIntegration::OnSequencerEvaluated()
 {
+	RestoreSpawnablePilotData();
+
 	// Redraw if not in PIE/simulate
 	const bool bIsInPIEOrSimulate = GEditor->PlayWorld != NULL || GEditor->bIsSimulatingInEditor;
 	if (bIsInPIEOrSimulate)
@@ -1177,6 +1191,8 @@ void FLevelEditorSequencerIntegration::RestoreRealtimeViewports()
 
 void FLevelEditorSequencerIntegration::RestoreToSavedState(UWorld* World)
 {
+	BackupSpawnablePilotData();
+
 	// Restore the saved state so that the level save can save that instead of the animated state.
 	IterateAllSequencers(
 		[World](FSequencer& In, const FLevelEditorSequencerIntegrationOptions& Options)
@@ -1238,6 +1254,70 @@ void FLevelEditorSequencerIntegration::OnMapChanged(UWorld* World, EMapChangeTyp
 	}
 }
 
+void FLevelEditorSequencerIntegration::BackupSpawnablePilotData()
+{
+	if (PilotedSpawnables.Num() != 0)
+	{
+		return;
+	}
+
+	FLevelEditorModule& LevelEditorModule = FModuleManager::Get().GetModuleChecked<FLevelEditorModule>("LevelEditor");
+	TSharedPtr<ILevelEditor> LevelEditor = LevelEditorModule.GetFirstLevelEditor();
+	if (LevelEditor.IsValid())
+	{
+		for (TSharedPtr<SLevelViewport> LevelViewport : LevelEditor->GetViewports())
+		{
+			const FLevelViewportActorLock& ActorLock = static_cast<FLevelEditorViewportClient&>(LevelViewport->GetAssetViewportClient()).GetActorLock();
+			if (AActor* LockedActor = ActorLock.GetLockedActor())
+			{
+				TOptional<FMovieSceneSpawnableAnnotation> SpawnableAnnotation = FMovieSceneSpawnableAnnotation::Find(ActorLock.GetLockedActor());
+				if (SpawnableAnnotation)
+				{
+					FPilotedSpawnable Pilot;
+					Pilot.WeakLevelViewport = LevelViewport;
+					Pilot.PreviousActorLock = ActorLock;
+					Pilot.Annotation = SpawnableAnnotation.GetValue();
+					PilotedSpawnables.Add(Pilot);
+				}
+			}
+		}
+	}
+}
+
+void FLevelEditorSequencerIntegration::RestoreSpawnablePilotData()
+{
+	if (PilotedSpawnables.Num() == 0)
+	{
+		return;
+	}
+
+	for (const FPilotedSpawnable& PilotData : PilotedSpawnables)
+	{
+		TSharedPtr<SLevelViewport> LevelViewport = PilotData.WeakLevelViewport.Pin();
+		if (LevelViewport && static_cast<FLevelEditorViewportClient&>(LevelViewport->GetAssetViewportClient()).GetActorLock().GetLockedActor() == nullptr)
+		{
+			// Find the new spawnable
+			IterateAllSequencers(
+				[&PilotData, LevelViewport](FSequencer& In, const FLevelEditorSequencerIntegrationOptions& Options)
+				{
+					for (TWeakObjectPtr<> WeakObject : In.FindBoundObjects(PilotData.Annotation.ObjectBindingID, PilotData.Annotation.SequenceID))
+					{
+						if (AActor* Actor = Cast<AActor>(WeakObject.Get()))
+						{
+							// Update the actor lock using the previous settings, but with the new actor
+							FLevelViewportActorLock ActorLock = PilotData.PreviousActorLock;
+							ActorLock.LockedActor = Actor;
+							static_cast<FLevelEditorViewportClient&>(LevelViewport->GetAssetViewportClient()).SetActorLock(ActorLock);
+							break;
+						}
+					}
+				}
+			);
+		}
+	}
+
+	PilotedSpawnables.Empty();
+}
 
 
 void FLevelEditorSequencerIntegration::AddSequencer(TSharedRef<ISequencer> InSequencer, const FLevelEditorSequencerIntegrationOptions& Options)
diff --git a/Engine/Source/Editor/Sequencer/Public/LevelEditorSequencerIntegration.h b/Engine/Source/Editor/Sequencer/Public/LevelEditorSequencerIntegration.h
--- a/Engine/Source/Editor/Sequencer/Public/LevelEditorSequencerIntegration.h
+++ b/Engine/Source/Editor/Sequencer/Public/LevelEditorSequencerIntegration.h
@@ -16,13 +16,14 @@ class FMenuBuilder;
 class FSequencer;
 class FObjectPostSaveContext;
 class FObjectPreSaveContext;
+class SLevelViewport;
 class FUICommandList;
 class IAssetViewport;
 class ISequencer;
 class ULevel;
 class UToolMenu;
 struct FPropertyAndParent;
-
+struct FPilotedSpawnable;
 
 struct FLevelEditorSequencerIntegrationOptions
 {
@@ -193,6 +194,9 @@ private:
 	void RestoreToSavedState(UWorld* World);
 	void ResetToAnimatedState(UWorld* World);
 
+	void BackupSpawnablePilotData();
+	void RestoreSpawnablePilotData();
+
 	struct FSequencerAndOptions
 	{
 		TWeakPtr<FSequencer> Sequencer;
@@ -211,12 +215,15 @@ private:
 	void UpdateDetails(bool bForceRefresh = false);
 
 	FLevelEditorSequencerIntegration();
+	~FLevelEditorSequencerIntegration();
 
 private:
 	FAcquiredResources AcquiredResources;
 
 	TSharedPtr<class FDetailKeyframeHandlerWrapper> KeyFrameHandler;
 
+	TArray<FPilotedSpawnable> PilotedSpawnables;
+
 	bool bDeferUpdates;
 
 	FOnSequencersChanged OnSequencersChanged;
