Sequencer: Fix mute/solo nodes when hierarchy changes by tracking old and new nodes paths

diff --git a/Engine/Source/Editor/Sequencer/Private/MVVM/ViewModels/ObjectBindingModel.cpp b/Engine/Source/Editor/Sequencer/Private/MVVM/ViewModels/ObjectBindingModel.cpp
--- a/Engine/Source/Editor/Sequencer/Private/MVVM/ViewModels/ObjectBindingModel.cpp
+++ b/Engine/Source/Editor/Sequencer/Private/MVVM/ViewModels/ObjectBindingModel.cpp
@@ -362,8 +362,6 @@ void FObjectBindingModel::Rename(const FText& NewName)
 		{
 			MovieScene->SetObjectDisplayName(ObjectBindingID, NewName);
 		}
-
-		SetIdentifier(FName(*NewName.ToString()));
 	}
 }
 
diff --git a/Engine/Source/Editor/Sequencer/Private/MVVM/ViewModels/OutlinerItemModel.cpp b/Engine/Source/Editor/Sequencer/Private/MVVM/ViewModels/OutlinerItemModel.cpp
--- a/Engine/Source/Editor/Sequencer/Private/MVVM/ViewModels/OutlinerItemModel.cpp
+++ b/Engine/Source/Editor/Sequencer/Private/MVVM/ViewModels/OutlinerItemModel.cpp
@@ -150,6 +150,16 @@ FName FOutlinerItemModelMixin::GetIdentifier() const
 void FOutlinerItemModelMixin::SetIdentifier(FName InNewIdentifier)
 {
 	TreeItemIdentifier = InNewIdentifier;
+
+	const FViewModel* ViewModel = AsViewModel();
+	if (ViewModel && ViewModel->IsConstructed())
+	{
+		TSharedPtr<FSequencerEditorViewModel> EditorViewModel = GetEditor();
+		if (EditorViewModel)
+		{
+			EditorViewModel->HandleDataHierarchyChanged();
+		}
+	}
 }
 
 bool FOutlinerItemModelMixin::IsExpanded() const
diff --git a/Engine/Source/Editor/Sequencer/Private/MVVM/ViewModels/SequencerEditorViewModel.cpp b/Engine/Source/Editor/Sequencer/Private/MVVM/ViewModels/SequencerEditorViewModel.cpp
--- a/Engine/Source/Editor/Sequencer/Private/MVVM/ViewModels/SequencerEditorViewModel.cpp
+++ b/Engine/Source/Editor/Sequencer/Private/MVVM/ViewModels/SequencerEditorViewModel.cpp
@@ -4,6 +4,8 @@
 #include "MVVM/ViewModels/SequencerOutlinerViewModel.h"
 #include "MVVM/ViewModels/SequencerTrackAreaViewModel.h"
 #include "MVVM/CurveEditorExtension.h"
+#include "MVVM/Extensions/IOutlinerExtension.h"
+#include "MVVM/SharedViewModelData.h"
 #include "ISequencerModule.h"
 #include "Sequencer.h"
 #include "MovieSceneSequenceID.h"
@@ -13,6 +15,24 @@ namespace UE
 namespace Sequencer
 {
 
+TMap<TWeakPtr<FViewModel>, FString> GetNodePaths(FViewModelPtr RootModel)
+{
+	TMap<TWeakPtr<FViewModel>, FString> NodePaths;
+	constexpr bool bIncludeThis = true;
+	for (FParentFirstChildIterator ChildIt(RootModel, bIncludeThis); ChildIt; ++ChildIt)
+	{
+		FViewModelPtr CurrentViewModel = *ChildIt;
+		if (const IOutlinerExtension* TreeItem = CurrentViewModel->CastThis<IOutlinerExtension>())
+		{
+			const FString NodePath = IOutlinerExtension::GetPathName(CurrentViewModel);
+
+			NodePaths.Add(FWeakViewModelPtr(CurrentViewModel), NodePath);
+		}
+	}
+
+	return NodePaths;
+}
+	
 FSequencerEditorViewModel::FSequencerEditorViewModel(TSharedRef<ISequencer> InSequencer, const FSequencerHostCapabilities& InHostCapabilities)
 	: WeakSequencer(InSequencer)
 	, bSupportsCurveEditor(InHostCapabilities.bSupportsCurveEditor)
@@ -54,6 +74,15 @@ void FSequencerEditorViewModel::InitializeEditorImpl()
 	PinnedTrackArea = CreateTrackAreaImpl();
 	PinnedTrackArea->GetOnHotspotChangedDelegate().AddSP(SharedThis(this), &FSequencerEditorViewModel::OnTrackAreaHotspotChanged);
 	GetEditorPanels().AddChild(PinnedTrackArea);
+
+	if (FViewModelPtr RootModel = GetRootModel())
+	{
+		TSharedPtr<FSharedViewModelData> RootSharedData = RootModel->GetSharedData();
+		RootSharedData->SubscribeToHierarchyChanged(RootModel)
+			.AddSP(this, &FSequencerEditorViewModel::HandleDataHierarchyChanged);
+
+		NodePaths = GetNodePaths(RootModel);
+	}
 }
 
 TSharedPtr<FTrackAreaViewModel> FSequencerEditorViewModel::GetPinnedTrackArea() const
@@ -84,6 +113,30 @@ bool FSequencerEditorViewModel::IsReadOnly() const
 	return !Sequencer || Sequencer->IsReadOnly();
 }
 
+void FSequencerEditorViewModel::HandleDataHierarchyChanged()
+{
+	if (FViewModelPtr RootModel = GetRootModel())
+	{
+		TMap<TWeakPtr<FViewModel>, FString> NewNodePaths = GetNodePaths(RootModel);
+
+		TSharedPtr<FSequencer> Sequencer = GetSequencerImpl();
+		if (Sequencer)
+		{
+			for (TMap<TWeakPtr<FViewModel>, FString>::TConstIterator It = NewNodePaths.CreateConstIterator(); It; ++It)
+			{
+				if (NodePaths.Contains(It.Key()))
+				{
+					FString OldPath = NodePaths[It.Key()];
+					FString NewPath = It.Value();
+					Sequencer->OnNodePathChanged(OldPath, NewPath);
+				}
+			}
+		}
+
+		NodePaths = NewNodePaths;
+	}
+}
+
 TSharedPtr<ITrackAreaHotspot> FSequencerEditorViewModel::GetHotspot() const
 {
 	return CurrentHotspot;
diff --git a/Engine/Source/Editor/Sequencer/Public/MVVM/ViewModels/SequencerEditorViewModel.h b/Engine/Source/Editor/Sequencer/Public/MVVM/ViewModels/SequencerEditorViewModel.h
--- a/Engine/Source/Editor/Sequencer/Public/MVVM/ViewModels/SequencerEditorViewModel.h
+++ b/Engine/Source/Editor/Sequencer/Public/MVVM/ViewModels/SequencerEditorViewModel.h
@@ -45,6 +45,8 @@ public:
 	/** Gets the current hotspots across any of our track areas */
 	TSharedPtr<ITrackAreaHotspot> GetHotspot() const;
 
+	void HandleDataHierarchyChanged();
+
 protected:
 
 	virtual void PreInitializeEditorImpl() override;
@@ -64,6 +66,9 @@ protected:
 
 	/** The current hotspot, from any of our track areas */
 	TSharedPtr<ITrackAreaHotspot> CurrentHotspot;
+
+	// Cached node paths to be used to compare when the hierarchy changes
+	TMap<TWeakPtr<FViewModel>, FString> NodePaths;
 };
 
 } // namespace Sequencer
