Made FSequencerSectionPainter non-copyable to prevent deprecation warnings in compiler-generated special member functions.

diff --git a/Engine/Source/Editor/Sequencer/Public/SequencerSectionPainter.h b/Engine/Source/Editor/Sequencer/Public/SequencerSectionPainter.h
--- a/Engine/Source/Editor/Sequencer/Public/SequencerSectionPainter.h
+++ b/Engine/Source/Editor/Sequencer/Public/SequencerSectionPainter.h
@@ -40,6 +40,8 @@ struct FKeyAreaElement
 class SEQUENCER_API FSequencerSectionPainter
 {
 public:
+	UE_NONCOPYABLE(FSequencerSectionPainter)
+
 	/** Constructor */
 	FSequencerSectionPainter(FSlateWindowElementList& OutDrawElements, const FGeometry& InSectionGeometry, TSharedPtr<UE::Sequencer::FSectionModel> Section);
 
