Sequencer: Fix crash in GetSelectedNodesToMove that was removing duplicate indices

diff --git a/Engine/Source/Editor/Sequencer/Private/Sequencer.cpp b/Engine/Source/Editor/Sequencer/Private/Sequencer.cpp
--- a/Engine/Source/Editor/Sequencer/Private/Sequencer.cpp
+++ b/Engine/Source/Editor/Sequencer/Private/Sequencer.cpp
@@ -7551,7 +7551,7 @@ TArray<TSharedRef<UE::Sequencer::FViewModel>> FSequencer::GetSelectedNodesToMove
 			{
 				if (Child == Node)
 				{
-					NodesToRemove.Add(NodeIndex);
+					NodesToRemove.AddUnique(NodeIndex);
 					break;
 				}
 			}
