Sequencer: add ability to lock marked frames on the time slider

diff --git a/Engine/Plugins/MovieScene/SequencerScripting/Source/SequencerScripting/Private/ExtensionLibraries/MovieSceneSequenceExtensions.cpp b/Engine/Plugins/MovieScene/SequencerScripting/Source/SequencerScripting/Private/ExtensionLibraries/MovieSceneSequenceExtensions.cpp
--- a/Engine/Plugins/MovieScene/SequencerScripting/Source/SequencerScripting/Private/ExtensionLibraries/MovieSceneSequenceExtensions.cpp
+++ b/Engine/Plugins/MovieScene/SequencerScripting/Source/SequencerScripting/Private/ExtensionLibraries/MovieSceneSequenceExtensions.cpp
@@ -1217,3 +1217,77 @@ bool UMovieSceneSequenceExtensions::IsReadOnly(UMovieSceneSequence* Sequence)
 	return false;
 }
 
+void UMovieSceneSequenceExtensions::SetPlaybackRangeLocked(UMovieSceneSequence* Sequence, bool bInLocked)
+{
+	if (!Sequence)
+	{
+		FFrame::KismetExecutionMessage(TEXT("Cannot call SetPlaybackRangeLocked on a null sequence"), ELogVerbosity::Error);
+		return;
+	}
+
+#if WITH_EDITORONLY_DATA
+	UMovieScene* MovieScene = Sequence->GetMovieScene();
+	if (MovieScene)
+	{
+		MovieScene->SetPlaybackRangeLocked(bInLocked);
+	}
+#endif
+}
+
+bool UMovieSceneSequenceExtensions::IsPlaybackRangeLocked(UMovieSceneSequence* Sequence)
+{
+	if (!Sequence)
+	{
+		FFrame::KismetExecutionMessage(TEXT("Cannot call IsPlaybackRangeLocked on a null sequence"), ELogVerbosity::Error);
+		return false;
+	}
+
+#if WITH_EDITORONLY_DATA
+	UMovieScene* MovieScene = Sequence->GetMovieScene();
+	if (MovieScene)
+	{
+
+		return MovieScene->IsPlaybackRangeLocked();
+	}
+#endif
+
+	return false;
+}
+
+void UMovieSceneSequenceExtensions::SetMarkedFramesLocked(UMovieSceneSequence* Sequence, bool bInLocked)
+{
+	if (!Sequence)
+	{
+		FFrame::KismetExecutionMessage(TEXT("Cannot call SetMarkedFramesLocked on a null sequence"), ELogVerbosity::Error);
+		return;
+	}
+
+#if WITH_EDITORONLY_DATA
+	UMovieScene* MovieScene = Sequence->GetMovieScene();
+	if (MovieScene)
+	{
+		MovieScene->SetMarkedFramesLocked(bInLocked);
+	}
+#endif
+}
+
+bool UMovieSceneSequenceExtensions::AreMarkedFramesLocked(UMovieSceneSequence* Sequence)
+{
+	if (!Sequence)
+	{
+		FFrame::KismetExecutionMessage(TEXT("Cannot call AreMarkedFramesLocked on a null sequence"), ELogVerbosity::Error);
+		return false;
+	}
+
+#if WITH_EDITORONLY_DATA
+	UMovieScene* MovieScene = Sequence->GetMovieScene();
+	if (MovieScene)
+	{
+
+		return MovieScene->AreMarkedFramesLocked();
+	}
+#endif
+
+	return false;
+}
+
diff --git a/Engine/Plugins/MovieScene/SequencerScripting/Source/SequencerScripting/Public/ExtensionLibraries/MovieSceneSequenceExtensions.h b/Engine/Plugins/MovieScene/SequencerScripting/Source/SequencerScripting/Public/ExtensionLibraries/MovieSceneSequenceExtensions.h
--- a/Engine/Plugins/MovieScene/SequencerScripting/Source/SequencerScripting/Public/ExtensionLibraries/MovieSceneSequenceExtensions.h
+++ b/Engine/Plugins/MovieScene/SequencerScripting/Source/SequencerScripting/Public/ExtensionLibraries/MovieSceneSequenceExtensions.h
@@ -612,6 +612,38 @@ public:
 	 */
 	UFUNCTION(BlueprintCallable, Category = "Sequencer|Sequence", meta = (ScriptMethod))
 	static bool IsReadOnly(UMovieSceneSequence* Sequence);
+
+	/*
+	 * Set playback range locked
+	 *
+	 * @bInLocked Whether the movie scene playback range should be locked
+	 */
+	UFUNCTION(BlueprintCallable, Category = "Sequencer|Sequence", meta = (ScriptMethod))
+	static void SetPlaybackRangeLocked(UMovieSceneSequence* Sequence, bool bInLocked);
+
+	/*
+	 * Is playback ranged locked
+	 *
+	 * @return Whether the movie scene playback range is locked
+	 */
+	UFUNCTION(BlueprintCallable, Category = "Sequencer|Sequence", meta = (ScriptMethod))
+	static bool IsPlaybackRangeLocked(UMovieSceneSequence* Sequence);
+
+	/*
+	 * Set marked frames locked
+	 *
+	 * @bInLocked Whether the movie scene marked frames should be locked
+	 */
+	UFUNCTION(BlueprintCallable, Category = "Sequencer|Sequence", meta = (ScriptMethod))
+	static void SetMarkedFramesLocked(UMovieSceneSequence* Sequence, bool bInLocked);
+
+	/*
+	 * Are marked frames locked
+	 *
+	 * @return Whether the movie scene marked frames are locked
+	 */
+	UFUNCTION(BlueprintCallable, Category = "Sequencer|Sequence", meta = (ScriptMethod))
+	static bool AreMarkedFramesLocked(UMovieSceneSequence* Sequence);
 };
 
 #if UE_ENABLE_INCLUDE_ORDER_DEPRECATED_IN_5_2
diff --git a/Engine/Source/Editor/Sequencer/Private/SSequencer.cpp b/Engine/Source/Editor/Sequencer/Private/SSequencer.cpp
--- a/Engine/Source/Editor/Sequencer/Private/SSequencer.cpp
+++ b/Engine/Source/Editor/Sequencer/Private/SSequencer.cpp
@@ -224,6 +224,8 @@ void SSequencer::Construct(const FArguments& InArgs, TSharedRef<FSequencer> InSe
 		TimeSliderArgs.OnAddMarkedFrame = InArgs._OnAddMarkedFrame;
 		TimeSliderArgs.OnDeleteMarkedFrame = InArgs._OnDeleteMarkedFrame;
 		TimeSliderArgs.OnDeleteAllMarkedFrames = InArgs._OnDeleteAllMarkedFrames;
+		TimeSliderArgs.AreMarkedFramesLocked = InArgs._AreMarkedFramesLocked;
+		TimeSliderArgs.OnToggleMarkedFramesLocked = InArgs._OnToggleMarkedFramesLocked;
 
 		TimeSliderArgs.NumericTypeInterface = GetNumericTypeInterface();
 	}
diff --git a/Engine/Source/Editor/Sequencer/Private/SSequencer.h b/Engine/Source/Editor/Sequencer/Private/SSequencer.h
--- a/Engine/Source/Editor/Sequencer/Private/SSequencer.h
+++ b/Engine/Source/Editor/Sequencer/Private/SSequencer.h
@@ -237,6 +237,12 @@ public:
 		/** Called when all marked frames should be deleted */
 		SLATE_EVENT( FSimpleDelegate, OnDeleteAllMarkedFrames)
 
+		/** Whether marked frames are locked */
+		SLATE_ATTRIBUTE( bool, AreMarkedFramesLocked )
+
+		/** Called when the user toggles the marked frames lock */
+		SLATE_EVENT( FSimpleDelegate, OnToggleMarkedFramesLocked )
+
 		/** Called when the user changes the clamp range */
 		SLATE_EVENT( FOnTimeRangeChanged, OnClampRangeChanged )
 
diff --git a/Engine/Source/Editor/Sequencer/Private/Sequencer.cpp b/Engine/Source/Editor/Sequencer/Private/Sequencer.cpp
--- a/Engine/Source/Editor/Sequencer/Private/Sequencer.cpp
+++ b/Engine/Source/Editor/Sequencer/Private/Sequencer.cpp
@@ -571,7 +571,9 @@ void FSequencer::InitSequencer(const FSequencerInitParams& InitParams, const TSh
 		.OnSetMarkedFrame(this, &FSequencer::SetMarkedFrame)
 		.OnAddMarkedFrame(this, &FSequencer::AddMarkedFrame)
 		.OnDeleteMarkedFrame(this, &FSequencer::DeleteMarkedFrame)
-		.OnDeleteAllMarkedFrames(this, &FSequencer::DeleteAllMarkedFrames )
+		.OnDeleteAllMarkedFrames(this, &FSequencer::DeleteAllMarkedFrames)
+		.AreMarkedFramesLocked(this, &FSequencer::AreMarkedFramesLocked)
+		.OnToggleMarkedFramesLocked(this, &FSequencer::ToggleMarkedFramesLocked)
 		.SubSequenceRange( this, &FSequencer::GetSubSequenceRange )
 		.OnPlaybackRangeChanged( this, &FSequencer::SetPlaybackRange )
 		.OnPlaybackRangeBeginDrag( this, &FSequencer::OnPlaybackRangeBeginDrag )
@@ -9954,6 +9956,48 @@ void FSequencer::StepToPreviousMark()
 	}
 }
 
+bool FSequencer::AreMarkedFramesLocked() const
+{
+	if (IsReadOnly())
+	{
+		return true;
+	}
+	
+	UMovieSceneSequence* FocusedMovieSceneSequence = GetFocusedMovieSceneSequence();
+	if (FocusedMovieSceneSequence != nullptr)
+	{
+		UMovieScene* MovieScene = FocusedMovieSceneSequence->GetMovieScene();
+
+		if (MovieScene->IsReadOnly())
+		{
+			return true;
+		}
+	
+		return MovieScene->AreMarkedFramesLocked();
+	}
+
+	return false;
+}
+
+void FSequencer::ToggleMarkedFramesLocked()
+{
+	UMovieSceneSequence* FocusedMovieSceneSequence = GetFocusedMovieSceneSequence();
+	if ( FocusedMovieSceneSequence != nullptr )
+	{
+		UMovieScene* MovieScene = FocusedMovieSceneSequence->GetMovieScene();
+
+		if (MovieScene->IsReadOnly())
+		{
+			FSequencerUtilities::ShowReadOnlyError();
+			return;
+		}
+
+		FScopedTransaction ToggleMarkedFramesLockTransaction( NSLOCTEXT( "Sequencer", "ToggleMarkedFramesLocked", "Toggle marked frames lock" ) );
+		MovieScene->Modify();
+		MovieScene->SetMarkedFramesLocked( !MovieScene->AreMarkedFramesLocked() );
+	}
+}
+
 void GatherTracksAndObjectsToCopy(TSharedRef<UE::Sequencer::FViewModel> Node, TArray<TSharedPtr<UE::Sequencer::FViewModel>>& TracksToCopy, TArray<TSharedPtr<UE::Sequencer::FObjectBindingModel>>& ObjectsToCopy, TArray<UMovieSceneFolder*>& FoldersToCopy)
 {
 	using namespace UE::Sequencer;
@@ -10932,6 +10976,12 @@ void FSequencer::BindCommands()
 		Commands.StepToPreviousMark,
 		FExecuteAction::CreateSP( this, &FSequencer::StepToPreviousMark));
 
+	SequencerCommandBindings->MapAction(
+		Commands.ToggleMarksLocked,
+		FExecuteAction::CreateSP(this, &FSequencer::ToggleMarkedFramesLocked),
+		FCanExecuteAction::CreateLambda( [this] { return GetFocusedMovieSceneSequence() != nullptr;	} ),
+		FIsActionChecked::CreateSP(this, &FSequencer::AreMarkedFramesLocked));
+
 	SequencerCommandBindings->MapAction(
 		Commands.ToggleAutoScroll,
 		FExecuteAction::CreateLambda( [this]{ Settings->SetAutoScrollEnabled( !Settings->GetAutoScrollEnabled() ); } ),
diff --git a/Engine/Source/Editor/Sequencer/Private/Sequencer.h b/Engine/Source/Editor/Sequencer/Private/Sequencer.h
--- a/Engine/Source/Editor/Sequencer/Private/Sequencer.h
+++ b/Engine/Source/Editor/Sequencer/Private/Sequencer.h
@@ -311,6 +311,8 @@ protected:
 	void ToggleMarkAtPlayPosition();
 	void StepToNextMark();
 	void StepToPreviousMark();
+	bool AreMarkedFramesLocked() const;
+	void ToggleMarkedFramesLocked();
 
 	/**
 	 * @param InMarkIndex The marked frame index to set
diff --git a/Engine/Source/Editor/Sequencer/Private/SequencerCommands.cpp b/Engine/Source/Editor/Sequencer/Private/SequencerCommands.cpp
--- a/Engine/Source/Editor/Sequencer/Private/SequencerCommands.cpp
+++ b/Engine/Source/Editor/Sequencer/Private/SequencerCommands.cpp
@@ -115,6 +115,7 @@ void FSequencerCommands::RegisterCommands()
 	UI_COMMAND( ToggleMarkAtPlayPosition, "Toggle Mark", "Sets or clears a mark at the current play position.", EUserInterfaceActionType::Button, FInputChord(EKeys::M) );
 	UI_COMMAND( StepToNextMark, "Step to Next Marked Frame", "Step to the next marked frame", EUserInterfaceActionType::Button, FInputChord(EModifierKey::Control | EModifierKey::Shift, EKeys::Period) );
 	UI_COMMAND( StepToPreviousMark, "Step to Previous Marked Frame", "Step to the previous marked frame", EUserInterfaceActionType::Button, FInputChord(EModifierKey::Control | EModifierKey::Shift, EKeys::Comma) );
+	UI_COMMAND( ToggleMarksLocked, "Marked Frames Locked", "Prevent editing the marked frames.", EUserInterfaceActionType::ToggleButton, FInputChord() );
 
 	UI_COMMAND( ToggleAutoScroll, "Auto Scroll", "Toggle auto-scroll: When enabled, automatically scrolls the sequencer view to keep the current time visible", EUserInterfaceActionType::ToggleButton, FInputChord(EModifierKey::Shift, EKeys::S) );
 
diff --git a/Engine/Source/Editor/Sequencer/Private/SequencerTimeSliderController.cpp b/Engine/Source/Editor/Sequencer/Private/SequencerTimeSliderController.cpp
--- a/Engine/Source/Editor/Sequencer/Private/SequencerTimeSliderController.cpp
+++ b/Engine/Source/Editor/Sequencer/Private/SequencerTimeSliderController.cpp
@@ -1026,6 +1026,8 @@ FReply FSequencerTimeSliderController::OnMouseMoveImpl( SWidget& WidgetOwner, co
 	bool bHandleRightMouseButton = MouseEvent.IsMouseButtonDown( EKeys::RightMouseButton ) && TimeSliderArgs.AllowZoom;
 	bool bHandleMiddleMouseButton = MouseEvent.IsMouseButtonDown(EKeys::MiddleMouseButton);
 
+	const bool bLockedMarkedFrames = TimeSliderArgs.AreMarkedFramesLocked.Get();
+
 	HoverMarkIndex = INDEX_NONE;
 
 	if (bHandleRightMouseButton)
@@ -1121,7 +1123,7 @@ FReply FSequencerTimeSliderController::OnMouseMoveImpl( SWidget& WidgetOwner, co
 					MouseDragType = DRAG_PLAYBACK_START;
 					TimeSliderArgs.OnPlaybackRangeBeginDrag.ExecuteIfBound();
 				}
-				else if (!bReadOnly && !bHitScrubber && HitTestMark(MyGeometry, RangeToScreen, MouseDownPixel, bFromTimeSlider, &DragMarkIndex, &DragMarkReferenceFrameNumber) && bHandleMiddleMouseButton == false)
+				else if (!bLockedMarkedFrames && !bHitScrubber && HitTestMark(MyGeometry, RangeToScreen, MouseDownPixel, bFromTimeSlider, &DragMarkIndex, &DragMarkReferenceFrameNumber) && bHandleMiddleMouseButton == false)
 				{
 					MouseDragType = DRAG_MARK;
 					DragMarkCurrentFrameNumber = DragMarkReferenceFrameNumber;
@@ -1213,7 +1215,7 @@ FReply FSequencerTimeSliderController::OnMouseMoveImpl( SWidget& WidgetOwner, co
 			}
 		}
 	}
-	else if (bFromTimeSlider && DragMarkIndex == INDEX_NONE)
+	else if (bFromTimeSlider && !bLockedMarkedFrames && DragMarkIndex == INDEX_NONE)
 	{
 		// Update hover state of marked frames.
 		TRange<double> LocalViewRange = GetViewRange();
@@ -1315,6 +1317,7 @@ FCursorReply FSequencerTimeSliderController::OnCursorQuery( TSharedRef<const SWi
 	const bool       bReadOnly        = Sequencer->IsReadOnly();
 	const FFrameRate TickResolution   = GetTickResolution();
 	const bool       bLockedPlayRange = TimeSliderArgs.IsPlaybackRangeLocked.Get();
+	const bool       bLockedMarkedFrames = TimeSliderArgs.AreMarkedFramesLocked.Get();
 	const float      HitTestPixel     = MyGeometry.AbsoluteToLocal(CursorEvent.GetScreenSpacePosition()).X;
 	const bool       bHitScrubber     = GetHitTestScrubPixelMetrics(RangeToScreen).HandleRangePx.Contains(HitTestPixel);
 
@@ -1340,7 +1343,7 @@ FCursorReply FSequencerTimeSliderController::OnCursorQuery( TSharedRef<const SWi
 	}
 
 	const bool bFromTimeSlider = true;
-	if (MouseDragType == DRAG_MARK || (!bReadOnly && !bHitScrubber && HitTestMark(MyGeometry, RangeToScreen, HitTestPixel, bFromTimeSlider)))
+	if (MouseDragType == DRAG_MARK || (!bLockedMarkedFrames && !bHitScrubber && HitTestMark(MyGeometry, RangeToScreen, HitTestPixel, bFromTimeSlider)))
 	{
 		return FCursorReply::Cursor(EMouseCursor::CardinalCross);
 	}
@@ -1638,7 +1641,10 @@ TSharedRef<SWidget> FSequencerTimeSliderController::OpenSetPlaybackRangeMenu(con
 				}
 			};
 
+			const bool bLockedMarkedFrames = TimeSliderArgs.AreMarkedFramesLocked.Get();
+
 			TSharedRef<SMarkedFramePropertyWidget> Widget = SNew(SMarkedFramePropertyWidget, MovieScene, MarkedIndex, WeakSequencer);
+			Widget->SetEnabled(!bLockedMarkedFrames);
 			MenuBuilder.AddWidget(Widget, FText::GetEmpty(), false);
 		}
 
@@ -1648,7 +1654,7 @@ TSharedRef<SWidget> FSequencerTimeSliderController::OpenSetPlaybackRangeMenu(con
 			FSlateIcon(),
 			FUIAction(
 				FExecuteAction::CreateLambda( [=]{ AddMarkAtFrame(FrameNumber); }),
-				FCanExecuteAction::CreateLambda([=]{ return !bReadOnly && MarkedIndex == INDEX_NONE; }))
+				FCanExecuteAction::CreateLambda([=]{ return !TimeSliderArgs.AreMarkedFramesLocked.Get() && MarkedIndex == INDEX_NONE; }))
 		);
 
 		MenuBuilder.AddMenuEntry( 
@@ -1657,16 +1663,28 @@ TSharedRef<SWidget> FSequencerTimeSliderController::OpenSetPlaybackRangeMenu(con
 			FSlateIcon(),
 			FUIAction(
 				FExecuteAction::CreateLambda([=]{ DeleteMarkAtIndex(MarkedIndex); }),
-				FCanExecuteAction::CreateLambda([=]{ return !bReadOnly && MarkedIndex != INDEX_NONE; }))
+				FCanExecuteAction::CreateLambda([=]{ return !TimeSliderArgs.AreMarkedFramesLocked.Get() && MarkedIndex != INDEX_NONE; }))
 		);
 
 		MenuBuilder.AddMenuEntry( 
-			LOCTEXT("Delete All Marks", "Delete All Marks"),
+			LOCTEXT("DeleteAllMarks", "Delete All Marks"),
 			FText(),
 			FSlateIcon(),
 			FUIAction(
 				FExecuteAction::CreateLambda([=]{ DeleteAllMarks(); }),
-				FCanExecuteAction::CreateLambda([=]{ return !bReadOnly && bHasMarks; }))
+				FCanExecuteAction::CreateLambda([=]{ return !TimeSliderArgs.AreMarkedFramesLocked.Get() && bHasMarks; }))
+			);
+
+		MenuBuilder.AddMenuEntry(
+			LOCTEXT("ToggleLockedMarks", "Locked"),
+			LOCTEXT("ToggleLockedMarksTooltip", "Lock/Unlock all marked frames"),
+			FSlateIcon(),
+			FUIAction(
+				FExecuteAction::CreateLambda([=]{ TimeSliderArgs.OnToggleMarkedFramesLocked.ExecuteIfBound(); }),
+				FCanExecuteAction::CreateLambda([=]{ return !bReadOnly; }),
+				FIsActionChecked::CreateLambda([=]{ return TimeSliderArgs.AreMarkedFramesLocked.Get(); })),
+			NAME_None,
+			EUserInterfaceActionType::ToggleButton
 			);
 	}
 	MenuBuilder.EndSection(); // SequencerMarkMenu
diff --git a/Engine/Source/Editor/Sequencer/Public/SequencerCommands.h b/Engine/Source/Editor/Sequencer/Public/SequencerCommands.h
--- a/Engine/Source/Editor/Sequencer/Public/SequencerCommands.h
+++ b/Engine/Source/Editor/Sequencer/Public/SequencerCommands.h
@@ -289,6 +289,9 @@ public:
 	/** Step to previous mark */
 	TSharedPtr< FUICommandInfo > StepToPreviousMark;
 
+	/** Toggle locking marks */
+	TSharedPtr< FUICommandInfo > ToggleMarksLocked;
+
 	/** Rotates through the supported formats for displaying times/frames/timecode. */
 	TSharedPtr< FUICommandInfo > ChangeTimeDisplayFormat;
 
diff --git a/Engine/Source/Editor/SequencerWidgets/Public/TimeSliderArgs.h b/Engine/Source/Editor/SequencerWidgets/Public/TimeSliderArgs.h
--- a/Engine/Source/Editor/SequencerWidgets/Public/TimeSliderArgs.h
+++ b/Engine/Source/Editor/SequencerWidgets/Public/TimeSliderArgs.h
@@ -151,12 +151,18 @@ struct FTimeSliderArgs
 	/** Attribute defining whether the playback range is locked */
 	TAttribute<bool> IsPlaybackRangeLocked;
 
+	/** Attribute defining whether the marked frames are locked */
+	TAttribute<bool> AreMarkedFramesLocked;
+
 	/** Attribute defining the time snap interval */
 	TAttribute<float> TimeSnapInterval;
 
 	/** Called when toggling the playback range lock */
 	FSimpleDelegate OnTogglePlaybackRangeLocked;
 
+	/** Called when toggling the marked frames lock */
+	FSimpleDelegate OnToggleMarkedFramesLocked;
+
 	/** If we are allowed to zoom */
 	bool AllowZoom;
 
diff --git a/Engine/Source/Runtime/MovieScene/Private/MovieScene.cpp b/Engine/Source/Runtime/MovieScene/Private/MovieScene.cpp
--- a/Engine/Source/Runtime/MovieScene/Private/MovieScene.cpp
+++ b/Engine/Source/Runtime/MovieScene/Private/MovieScene.cpp
@@ -72,6 +72,7 @@ UMovieScene::UMovieScene(const FObjectInitializer& ObjectInitializer)
 #if WITH_EDITORONLY_DATA
 	bReadOnly = false;
 	bPlaybackRangeLocked = false;
+	bMarkedFramesLocked = false;
 	PlaybackRange.MigrationDefault = FFloatRange::Empty();
 	EditorData.WorkingRange_DEPRECATED = EditorData.ViewRange_DEPRECATED = TRange<float>::Empty();
 
@@ -651,6 +652,16 @@ void UMovieScene::SetPlaybackRangeLocked(bool bLocked)
 {
 	bPlaybackRangeLocked = bLocked;
 }
+
+bool UMovieScene::AreMarkedFramesLocked() const
+{
+	return bMarkedFramesLocked;
+}
+
+void UMovieScene::SetMarkedFramesLocked(bool bLocked)
+{
+	bMarkedFramesLocked = bLocked;
+}
 #endif
 
 
diff --git a/Engine/Source/Runtime/MovieScene/Public/MovieScene.h b/Engine/Source/Runtime/MovieScene/Public/MovieScene.h
--- a/Engine/Source/Runtime/MovieScene/Public/MovieScene.h
+++ b/Engine/Source/Runtime/MovieScene/Public/MovieScene.h
@@ -1021,6 +1021,16 @@ public:
 	*/
 	void SetPlaybackRangeLocked(bool bLocked);
 
+	/**
+	 * Return whether marked frames are locked.
+	 */
+	bool AreMarkedFramesLocked() const;
+
+	/**
+	 * Set whether marked frames are locked.
+	 */
+	void SetMarkedFramesLocked(bool bLocked);
+
 	/**
 	 * @return The editor only data for use with this movie scene
 	 */
@@ -1269,6 +1279,10 @@ private:
 	UPROPERTY()
 	bool bPlaybackRangeLocked;
 
+	/** User-defined marked frames are locked. */
+	UPROPERTY()
+	bool bMarkedFramesLocked;
+
 	/** Maps object GUIDs to user defined display names. */
 	UPROPERTY()
 	TMap<FString, FText> ObjectsToDisplayNames;
