Sequencer: correctly restore playback speed on stop/pause

diff --git a/Engine/Source/Editor/Sequencer/Private/Sequencer.cpp b/Engine/Source/Editor/Sequencer/Private/Sequencer.cpp
--- a/Engine/Source/Editor/Sequencer/Private/Sequencer.cpp
+++ b/Engine/Source/Editor/Sequencer/Private/Sequencer.cpp
@@ -4767,6 +4767,9 @@ FReply FSequencer::OnPlay(bool bTogglePlay)
 			SetGlobalTime(NewGlobalTime);
 		}
 
+		SpeedIndexBeforePlay = CurrentSpeedIndex;
+		PlaybackSpeedBeforePlay = PlaybackSpeed;
+
 		SetPlaybackStatus(EMovieScenePlayerStatus::Playing);
 
 		// Make sure Slate ticks during playback
@@ -4986,7 +4989,6 @@ ESequencerLoopMode FSequencer::GetLoopMode() const
 	return Settings->GetLoopMode();
 }
 
-
 void FSequencer::SetLocalTimeLooped(FFrameTime NewLocalTime)
 {
 	TOptional<EMovieScenePlayerStatus::Type> NewPlaybackStatus;
@@ -5078,11 +5080,29 @@ void FSequencer::SetLocalTimeLooped(FFrameTime NewLocalTime)
 
 		if (NewPlaybackStatus.GetValue() == EMovieScenePlayerStatus::Stopped)
 		{
+			RestorePlaybackSpeedAfterPlay();
+
 			OnStopDelegate.Broadcast();
 		}
 	}
 }
 
+void FSequencer::SetPlaybackSpeed(float InPlaybackSpeed)
+{
+	PlaybackSpeed = InPlaybackSpeed;
+
+	const bool bExactOnly = true;
+	CurrentSpeedIndex = FindClosestPlaybackSpeed(InPlaybackSpeed, bExactOnly);
+}
+
+void FSequencer::RestorePlaybackSpeedAfterPlay()
+{
+	// Reset the speed to what it was before we started playing, in case the user increased/decreased the speed while it
+	// was playing.
+	CurrentSpeedIndex = SpeedIndexBeforePlay;
+	PlaybackSpeed = PlaybackSpeedBeforePlay;
+}
+
 EPlaybackMode::Type FSequencer::GetPlaybackMode() const
 {
 	if (PlaybackState == EMovieScenePlayerStatus::Playing)
@@ -9217,18 +9237,16 @@ void FSequencer::ShuttleBackward()
 	}
 }
 
-void FSequencer::SnapToClosestPlaybackSpeed()
+int32 FSequencer::FindClosestPlaybackSpeed(float InPlaybackSpeed, bool bExactOnly) const
 {
 	TArray<float> PlaybackSpeeds = GetPlaybackSpeeds.Execute();
 
-	float CurrentSpeed = GetPlaybackSpeed();
-
+	int32 NewSpeedIndex = INDEX_NONE;
 	float Delta = TNumericLimits<float>::Max();
 
-	int32 NewSpeedIndex = INDEX_NONE;
 	for (int32 Idx = 0; Idx < PlaybackSpeeds.Num(); Idx++)
 	{
-		float NewDelta = FMath::Abs(CurrentSpeed - PlaybackSpeeds[Idx]);
+		float NewDelta = FMath::Abs(InPlaybackSpeed - PlaybackSpeeds[Idx]);
 		if (NewDelta < Delta)
 		{
 			Delta = NewDelta;
@@ -9236,8 +9254,24 @@ void FSequencer::SnapToClosestPlaybackSpeed()
 		}
 	}
 
+	if (bExactOnly && Delta > UE_KINDA_SMALL_NUMBER)
+	{
+		return INDEX_NONE;
+	}
+
+	return NewSpeedIndex;
+}
+
+void FSequencer::SnapToClosestPlaybackSpeed()
+{
+	float CurrentSpeed = GetPlaybackSpeed();
+	int32 NewSpeedIndex = FindClosestPlaybackSpeed(CurrentSpeed);
+
 	if (NewSpeedIndex != INDEX_NONE)
 	{
+		TArray<float> PlaybackSpeeds = GetPlaybackSpeeds.Execute();
+
+		CurrentSpeedIndex = NewSpeedIndex;
 		PlaybackSpeed = PlaybackSpeeds[NewSpeedIndex];
 	}	
 }
@@ -9266,13 +9300,8 @@ void FSequencer::Pause()
 		EvaluateInternal(Range);
 	}
 
-	// reset the speed to 1. We have to update the speed index as well.
-	TArray<float> PlaybackSpeeds = GetPlaybackSpeeds.Execute();
+	RestorePlaybackSpeedAfterPlay();
 
-	CurrentSpeedIndex = PlaybackSpeeds.Find(1.f);
-	check(CurrentSpeedIndex != INDEX_NONE);
-	PlaybackSpeed = PlaybackSpeeds[CurrentSpeedIndex];
-	
 	OnStopDelegate.Broadcast();
 }
 
diff --git a/Engine/Source/Editor/Sequencer/Private/Sequencer.h b/Engine/Source/Editor/Sequencer/Private/Sequencer.h
--- a/Engine/Source/Editor/Sequencer/Private/Sequencer.h
+++ b/Engine/Source/Editor/Sequencer/Private/Sequencer.h
@@ -745,7 +745,7 @@ public:
 	virtual void NotifyMovieSceneDataChanged( EMovieSceneDataChangeType DataChangeType ) override;
 	virtual void RefreshTree() override;
 	virtual void UpdatePlaybackRange() override;
-	virtual void SetPlaybackSpeed(float InPlaybackSpeed) override { PlaybackSpeed = InPlaybackSpeed; }
+	virtual void SetPlaybackSpeed(float InPlaybackSpeed) override;
 	virtual float GetPlaybackSpeed() const override { return PlaybackSpeed; }
 	virtual TArray<FGuid> AddActors(const TArray<TWeakObjectPtr<AActor> >& InActors, bool bSelectActors = true) override;
 	virtual TArray<FGuid> ConvertToSpawnable(FGuid Guid) override;
@@ -1128,6 +1128,9 @@ private:
 
 	void UpdateCachedCameraActors();
 
+	int32 FindClosestPlaybackSpeed(float InPlaybackSpeed, bool bExactOnly = false) const;
+	void RestorePlaybackSpeedAfterPlay();
+
 public:
 
 	/** Helper function which returns how many frames (in tick resolution) one display rate frame represents. */
@@ -1222,7 +1225,11 @@ private:
 	/** The last time range that was viewed */
 	TRange<double> LastViewRange;
 
+	/** The index of the playback speed in the supported playback speeds */
 	int32 CurrentSpeedIndex;
+
+	/** The index of the playback speed before we started playing */
+	int32 SpeedIndexBeforePlay;
 	
 	/** The view range before zooming */
 	TRange<double> ViewRangeBeforeZoom;
@@ -1271,6 +1278,9 @@ private:
 	/** The playback speed */
 	float PlaybackSpeed;
 
+	/** The playback speed before we started playing */
+	float PlaybackSpeedBeforePlay;
+
 	/** The shuttle multiplier */
 	float ShuttleMultiplier;
 
