Sequencer: Check for valid track before dereferencing

#rnx

diff --git a/Engine/Source/Editor/Sequencer/Private/SSequencer.cpp b/Engine/Source/Editor/Sequencer/Private/SSequencer.cpp
--- a/Engine/Source/Editor/Sequencer/Private/SSequencer.cpp
+++ b/Engine/Source/Editor/Sequencer/Private/SSequencer.cpp
@@ -2884,7 +2884,7 @@ void SSequencer::UpdateLayoutTree()
 				: Sequencer->GetNodeTree()->GetRootNode()->GetDescendantsOfType<IOutlinerExtension>())
 			{
 				TViewModelPtr<ITrackExtension> TrackNode = OutlinerItem.ImplicitCast();
-				if (TrackNode)
+				if (TrackNode && TrackNode->GetTrack())
 				{
 					UMovieSceneTrack* Track = TrackNode->GetTrack();
 					bool bDisableEval = NodeTree->IsNodeMute(OutlinerItem) || (bHasSoloNodes && !NodeTree->IsNodeSolo(OutlinerItem));
