Fixed missing initialization for cached speed & index, which would set the playback speed to 0 if clicking on Pause, and Playing again wouldn't work due to it

diff --git a/Engine/Source/Editor/Sequencer/Private/Sequencer.cpp b/Engine/Source/Editor/Sequencer/Private/Sequencer.cpp
--- a/Engine/Source/Editor/Sequencer/Private/Sequencer.cpp
+++ b/Engine/Source/Editor/Sequencer/Private/Sequencer.cpp
@@ -381,6 +381,7 @@ void FSequencer::InitSequencer(const FSequencerInitParams& InitParams, const TSh
 	const int32 IndexOfOne = GetPlaybackSpeeds.Execute().Find(1.f);
 	check(IndexOfOne != INDEX_NONE);
 	CurrentSpeedIndex = IndexOfOne;
+	SpeedIndexBeforePlay = CurrentSpeedIndex;
 
 	if (InitParams.SpawnRegister.IsValid())
 	{
@@ -488,6 +489,7 @@ void FSequencer::InitSequencer(const FSequencerInitParams& InitParams, const TSh
 	ToolkitHost = InitParams.ToolkitHost;
 
 	PlaybackSpeed = 1.f;
+	PlaybackSpeedBeforePlay = PlaybackSpeed;
 	ShuttleMultiplier = 0;
 	ObjectChangeListener = InObjectChangeListener;
 
