Sequencer: Show Marked Frames Globally should be marked as a UPROPERTY so that it gets serialized. Also, added a transaction for toggling Show Marked Frames Globally.

diff --git a/Engine/Source/Editor/Sequencer/Private/SSequencer.cpp b/Engine/Source/Editor/Sequencer/Private/SSequencer.cpp
--- a/Engine/Source/Editor/Sequencer/Private/SSequencer.cpp
+++ b/Engine/Source/Editor/Sequencer/Private/SSequencer.cpp
@@ -2169,6 +2169,8 @@ TSharedRef<SWidget> SSequencer::MakeViewMenu()
 				FExecuteAction::CreateLambda([FocusedMovieScene, this]() {
 					if (FocusedMovieScene != nullptr)
 					{
+						const FScopedTransaction Transaction(LOCTEXT("ToggleShowMarkedFramesGlobally", "Toggle Show Marked Frames Globally"));
+						FocusedMovieScene->Modify();
 						FocusedMovieScene->ToggleGloballyShowMarkedFrames();
 						SequencerPtr.Pin()->InvalidateGlobalMarkedFramesCache();
 					}
diff --git a/Engine/Source/Runtime/MovieScene/Public/MovieScene.h b/Engine/Source/Runtime/MovieScene/Public/MovieScene.h
--- a/Engine/Source/Runtime/MovieScene/Public/MovieScene.h
+++ b/Engine/Source/Runtime/MovieScene/Public/MovieScene.h
@@ -1316,6 +1316,7 @@ private:
 	TObjectPtr<UMovieSceneNodeGroupCollection> NodeGroupCollection;
 
 	/** Whether this scene's marked frames should be shown globally */
+	UPROPERTY()
 	bool bGloballyShowMarkedFrames;
 
 private:
