Sequencer: Change property menu to a TMap rather than an TSortedMap to fix duplicate entries.

diff --git a/Engine/Source/Editor/Sequencer/Private/MVVM/ViewModels/ObjectBindingModel.cpp b/Engine/Source/Editor/Sequencer/Private/MVVM/ViewModels/ObjectBindingModel.cpp
--- a/Engine/Source/Editor/Sequencer/Private/MVVM/ViewModels/ObjectBindingModel.cpp
+++ b/Engine/Source/Editor/Sequencer/Private/MVVM/ViewModels/ObjectBindingModel.cpp
@@ -623,7 +623,7 @@ TSharedRef<SWidget> FObjectBindingModel::GetAddTrackMenuContent()
 	bool bDefaultCategoryFound = false;
 
 	// Create property menu data based on keyable property paths
-	TSortedMap<FString, TArray<PropertyMenuData>, FDefaultAllocator, FCategorySortPredicate> KeyablePropertyMenuData;
+	TMap<FString, TArray<PropertyMenuData>> KeyablePropertyMenuData;
 	for (const FPropertyPath& KeyablePropertyPath : KeyablePropertyPaths)
 	{
 		FProperty* Property = KeyablePropertyPath.GetRootProperty().Property.Get();
@@ -651,6 +651,8 @@ TSharedRef<SWidget> FObjectBindingModel::GetAddTrackMenuContent()
 		}
 	}
 
+	KeyablePropertyMenuData.KeySort(FCategorySortPredicate());
+
 	// Always add an extension point for Properties section even if none are found (Components rely on this) 
 	if (!bDefaultCategoryFound)
 	{
