Sequencer:
- Fix issue with Object Handles being created when calling Details Keyframe Handler IsPropertyAnimated

diff --git a/Engine/Source/Editor/Sequencer/Private/LevelEditorSequencerIntegration.cpp b/Engine/Source/Editor/Sequencer/Private/LevelEditorSequencerIntegration.cpp
--- a/Engine/Source/Editor/Sequencer/Private/LevelEditorSequencerIntegration.cpp
+++ b/Engine/Source/Editor/Sequencer/Private/LevelEditorSequencerIntegration.cpp
@@ -116,7 +116,8 @@ public:
 			TSharedPtr<ISequencer> Sequencer = WeakSequencer.Pin();
 			if (Sequencer.IsValid() && Sequencer->GetFocusedMovieSceneSequence())
 			{
-				FGuid ObjectHandle = Sequencer->GetHandleToObject(ParentObject);
+				constexpr bool bCreateHandleIfMissing = false;
+				FGuid ObjectHandle = Sequencer->GetHandleToObject(ParentObject, bCreateHandleIfMissing);
 				if (ObjectHandle.IsValid()) 
 				{
 					UMovieScene* MovieScene = Sequencer->GetFocusedMovieSceneSequence()->GetMovieScene();
