Sequencer: Add time scrubbing(b) to sequencer and curve editor to match viewport behavior. Had to add playback status to the time controller, add key up/down events  to sequencer input handlers, and for both created new drag operations.

diff --git a/Engine/Source/Editor/CurveEditor/Private/CurveEditorCommands.cpp b/Engine/Source/Editor/CurveEditor/Private/CurveEditorCommands.cpp
--- a/Engine/Source/Editor/CurveEditor/Private/CurveEditorCommands.cpp
+++ b/Engine/Source/Editor/CurveEditor/Private/CurveEditorCommands.cpp
@@ -112,6 +112,7 @@ void FCurveEditorCommands::RegisterCommands()
 	UI_COMMAND(StepBackward, "Step Backward", "Step the timeline backward", EUserInterfaceActionType::Button, FInputChord(EKeys::Left));
 	UI_COMMAND(JumpToStart, "Jump to Start", "Jump to the start of the playback range", EUserInterfaceActionType::Button, FInputChord(EKeys::Up));
 	UI_COMMAND(JumpToEnd, "Jump to End", "Jump to the end of the playback range", EUserInterfaceActionType::Button, FInputChord(EModifierKey::Control, EKeys::Up));
+	UI_COMMAND(ScrubTime, "Scrub Time", "Scrub mouse left and right to change time", EUserInterfaceActionType::Button, FInputChord(EKeys::B));
 
 	// Selection Range
 	UI_COMMAND(SetSelectionRangeStart, "Set Selection Start", "Sets the start of the selection range", EUserInterfaceActionType::Button, FInputChord(EKeys::I) );
diff --git a/Engine/Source/Editor/CurveEditor/Private/DragOperations/CurveEditorDragOperation_ScrubTime.cpp b/Engine/Source/Editor/CurveEditor/Private/DragOperations/CurveEditorDragOperation_ScrubTime.cpp
new file mode 100644
--- /dev/null
+++ b/Engine/Source/Editor/CurveEditor/Private/DragOperations/CurveEditorDragOperation_ScrubTime.cpp
@@ -0,0 +1,72 @@
+// Copyright Epic Games, Inc. All Rights Reserved.
+
+#include "CurveEditorDragOperation_ScrubTime.h"
+
+#include "CurveEditor.h"
+#include "CurveEditorScreenSpace.h"
+#include "ICurveEditorBounds.h"
+#include "Input/Events.h"
+#include "Math/UnrealMathUtility.h"
+#include "SCurveEditorPanel.h"
+#include "SCurveEditorView.h"
+
+FCurveEditorDragOperation_ScrubTime::FCurveEditorDragOperation_ScrubTime(FCurveEditor* InCurveEditor)
+	: CurveEditor(InCurveEditor) 
+{}
+
+void FCurveEditorDragOperation_ScrubTime::OnBeginDrag(FVector2D InitialPosition, FVector2D CurrentPosition, const FPointerEvent& MouseEvent)
+{
+	TSharedPtr<ITimeSliderController> TimeSliderController = CurveEditor->GetTimeSliderController();
+	if (TimeSliderController)
+	{
+		SnappingState.Reset();
+		LastMousePosition = CurrentPosition;
+		const FFrameRate TickResolution = TimeSliderController->GetTickResolution();
+
+		InitialTime = TickResolution.AsSeconds(TimeSliderController->GetScrubPosition());
+		TimeSliderController->SetPlaybackStatus(ETimeSliderPlaybackStatus::Scrubbing);
+	}
+}
+
+void FCurveEditorDragOperation_ScrubTime::OnDrag(FVector2D InitialPosition, FVector2D CurrentPosition, const FPointerEvent& MouseEvent)
+{
+	TSharedPtr<ITimeSliderController> TimeSliderController = CurveEditor->GetTimeSliderController();
+	if (TimeSliderController)
+	{
+		FVector2D PixelDelta = CurveEditor->GetAxisSnap().GetSnappedPosition(InitialPosition, LastMousePosition, CurrentPosition, MouseEvent, SnappingState, true) - InitialPosition;
+		LastMousePosition = CurrentPosition;
+
+		FCurveEditorScreenSpaceH InputSpace = CurveEditor->GetPanelInputSpace();
+		const double InitialMouseTime = InputSpace.ScreenToSeconds(InitialPosition.X);
+		const double CurrentMouseTime = InputSpace.ScreenToSeconds(LastMousePosition.X);
+		const double DiffSeconds = CurrentMouseTime - InitialMouseTime;
+		const FFrameRate TickResolution = TimeSliderController->GetTickResolution();
+		double TimeToSet = (InitialTime + DiffSeconds);
+		if (CurveEditor->InputSnapEnabledAttribute.Get())
+		{
+			FCurveSnapMetrics SnapMetrics;
+			SnapMetrics.bSnapInputValues = true;
+			SnapMetrics.InputSnapRate = CurveEditor->InputSnapRateAttribute.Get();
+			TimeToSet = SnapMetrics.SnapInputSeconds(TimeToSet);
+		}
+		TimeSliderController->SetScrubPosition((TimeToSet)*TickResolution, /*bEvaluate*/ true);
+	}
+}
+
+void FCurveEditorDragOperation_ScrubTime::OnEndDrag(FVector2D InitialPosition, FVector2D CurrentPosition, const FPointerEvent& MouseEvent)
+{
+	TSharedPtr<ITimeSliderController> TimeSliderController = CurveEditor->GetTimeSliderController();
+	if (TimeSliderController)
+	{
+		TimeSliderController->SetPlaybackStatus(ETimeSliderPlaybackStatus::Stopped);
+	}
+}
+
+void FCurveEditorDragOperation_ScrubTime::OnCancelDrag()
+{
+	TSharedPtr<ITimeSliderController> TimeSliderController = CurveEditor->GetTimeSliderController();
+	if (TimeSliderController)
+	{
+		TimeSliderController->SetPlaybackStatus(ETimeSliderPlaybackStatus::Stopped);
+	}
+}
\ No newline at end of file
diff --git a/Engine/Source/Editor/CurveEditor/Private/DragOperations/CurveEditorDragOperation_ScrubTime.h b/Engine/Source/Editor/CurveEditor/Private/DragOperations/CurveEditorDragOperation_ScrubTime.h
new file mode 100644
--- /dev/null
+++ b/Engine/Source/Editor/CurveEditor/Private/DragOperations/CurveEditorDragOperation_ScrubTime.h
@@ -0,0 +1,32 @@
+// Copyright Epic Games, Inc. All Rights Reserved.
+
+#pragma once
+
+#include "CurveEditorSnapMetrics.h"
+#include "ICurveEditorDragOperation.h"
+#include "Input/Reply.h"
+#include "Math/Vector2D.h"
+#include "Templates/SharedPointer.h"
+
+class FCurveEditor;
+struct FPointerEvent;
+class ITimeSliderController;
+
+class FCurveEditorDragOperation_ScrubTime : public ICurveEditorDragOperation
+{
+public:
+	FCurveEditorDragOperation_ScrubTime(FCurveEditor* CurveEditor);
+
+	virtual void OnBeginDrag(FVector2D InitialPosition, FVector2D CurrentPosition, const FPointerEvent& MouseEvent) override;
+	virtual void OnDrag(FVector2D InitialPosition, FVector2D CurrentPosition, const FPointerEvent& MouseEvent) override;
+	virtual void OnEndDrag(FVector2D InitialPosition, FVector2D CurrentPosition, const FPointerEvent& MouseEvent) override;
+
+	virtual void OnCancelDrag() override;
+
+private:
+
+	FCurveEditor* CurveEditor;
+	FCurveEditorAxisSnap::FSnapState SnappingState;
+	FVector2D LastMousePosition;
+	double InitialTime;
+};
\ No newline at end of file
diff --git a/Engine/Source/Editor/CurveEditor/Private/SCurveEditorViewContainer.cpp b/Engine/Source/Editor/CurveEditor/Private/SCurveEditorViewContainer.cpp
--- a/Engine/Source/Editor/CurveEditor/Private/SCurveEditorViewContainer.cpp
+++ b/Engine/Source/Editor/CurveEditor/Private/SCurveEditorViewContainer.cpp
@@ -8,6 +8,7 @@
 #include "DragOperations/CurveEditorDragOperation_Marquee.h"
 #include "DragOperations/CurveEditorDragOperation_Pan.h"
 #include "DragOperations/CurveEditorDragOperation_Zoom.h"
+#include "DragOperations/CurveEditorDragOperation_ScrubTime.h"
 #include "GenericPlatform/ICursor.h"
 #include "HAL/PlatformCrt.h"
 #include "ICurveEditorBounds.h"
@@ -36,6 +37,7 @@
 #include "Views/SInteractiveCurveEditorView.h"
 #include "Widgets/Layout/SBox.h"
 #include "Widgets/SWidget.h"
+#include "CurveEditorCommands.h"
 
 class FPaintArgs;
 class FSlateRect;
@@ -155,6 +157,23 @@ int32 SCurveEditorViewContainer::OnPaint(const FPaintArgs& Args, const FGeometry
 	return LayerId + CurveViewConstants::ELayerOffset::Last;
 }
 
+bool SCurveEditorViewContainer::IsScrubTimeKeyEvent(const FKeyEvent& InKeyEvent)
+{
+	const FCurveEditorCommands& Commands = FCurveEditorCommands::Get();
+	// Need to iterate through primary and secondary to make sure they are all pressed.
+	for (uint32 i = 0; i < static_cast<uint8>(EMultipleKeyBindingIndex::NumChords); ++i)
+	{
+		EMultipleKeyBindingIndex ChordIndex = static_cast<EMultipleKeyBindingIndex>(i);
+		const FInputChord& Chord = *Commands.ScrubTime->GetActiveChord(ChordIndex);
+		const bool bIsMovingTimeSlider = Chord.IsValidChord() && InKeyEvent.GetKey() == Chord.Key;
+		if (bIsMovingTimeSlider)
+		{
+			return true;
+		}
+	}
+	return false;
+}
+
 FReply SCurveEditorViewContainer::OnKeyDown(const FGeometry& MyGeometry, const FKeyEvent& InKeyEvent)
 {
 	if (InKeyEvent.GetKey() == EKeys::Escape && DragOperation.IsSet())
@@ -163,6 +182,28 @@ FReply SCurveEditorViewContainer::OnKeyDown(const FGeometry& MyGeometry, const F
 		DragOperation.Reset();
 		return FReply::Handled();
 	}
+
+	if (IsScrubTimeKeyEvent(InKeyEvent))
+	{
+		bIsScrubbingTime = true;
+		return FReply::Handled();
+	}
+	return FReply::Unhandled();
+}
+
+FReply SCurveEditorViewContainer::OnKeyUp(const FGeometry& MyGeometry, const FKeyEvent& InKeyEvent)
+{
+	if (IsScrubTimeKeyEvent(InKeyEvent) && bIsScrubbingTime )
+	{
+		bIsScrubbingTime = false;
+		if (DragOperation.IsSet())
+		{
+			DragOperation->DragImpl->CancelDrag();
+			DragOperation.Reset();
+			return FReply::Handled();
+		}
+	}
+	
 	return FReply::Unhandled();
 }
 
@@ -182,11 +223,19 @@ FReply SCurveEditorViewContainer::OnMouseButtonDown(const FGeometry& MyGeometry,
 {
 	FVector2D MousePixel = MyGeometry.AbsoluteToLocal(MouseEvent.GetScreenSpacePosition());
 
-	// Marquee Selection
+	// Marquee Selection or time scrub
 	if (MouseEvent.GetEffectingButton() == EKeys::LeftMouseButton)
 	{
-		DragOperation = FCurveEditorDelayedDrag(MousePixel, MouseEvent.GetEffectingButton());
-		DragOperation->DragImpl = MakeUnique<FCurveEditorDragOperation_Marquee>(CurveEditor.Get());
+		if (bIsScrubbingTime)
+		{
+			DragOperation = FCurveEditorDelayedDrag(MousePixel, MouseEvent.GetEffectingButton());
+			DragOperation->DragImpl = MakeUnique<FCurveEditorDragOperation_ScrubTime>(CurveEditor.Get());
+		}
+		else
+		{
+			DragOperation = FCurveEditorDelayedDrag(MousePixel, MouseEvent.GetEffectingButton());
+			DragOperation->DragImpl = MakeUnique<FCurveEditorDragOperation_Marquee>(CurveEditor.Get());
+		}
 		return FReply::Handled();
 	}
 	// Middle Click + Alt Pan
diff --git a/Engine/Source/Editor/CurveEditor/Private/SCurveEditorViewContainer.h b/Engine/Source/Editor/CurveEditor/Private/SCurveEditorViewContainer.h
--- a/Engine/Source/Editor/CurveEditor/Private/SCurveEditorViewContainer.h
+++ b/Engine/Source/Editor/CurveEditor/Private/SCurveEditorViewContainer.h
@@ -76,6 +76,7 @@ private:
 
 	virtual bool SupportsKeyboardFocus() const override { return true; }
 	virtual FReply OnKeyDown(const FGeometry& MyGeometry, const FKeyEvent& InKeyEvent) override;
+	virtual FReply OnKeyUp(const FGeometry& MyGeometry, const FKeyEvent& InKeyEvent) override;
 
 private:
 
@@ -85,6 +86,8 @@ private:
 
 	void ExpandInputBounds(float NewWidth);
 
+	bool IsScrubTimeKeyEvent(const FKeyEvent& InKeyEvent);
+
 private:
 
 	/** The curve editor pointer */
@@ -110,4 +113,7 @@ private:
 
 	/** The minimum height for this container panel. **/
 	float MinimumPanelHeight;
+
+	/** Whether or not we are scrubbing time*/
+	bool bIsScrubbingTime = false;
 };
\ No newline at end of file
diff --git a/Engine/Source/Editor/CurveEditor/Public/CurveEditorCommands.h b/Engine/Source/Editor/CurveEditor/Public/CurveEditorCommands.h
--- a/Engine/Source/Editor/CurveEditor/Public/CurveEditorCommands.h
+++ b/Engine/Source/Editor/CurveEditor/Public/CurveEditorCommands.h
@@ -120,6 +120,7 @@ public:
 	TSharedPtr<FUICommandInfo> StepBackward;
 	TSharedPtr<FUICommandInfo> JumpToStart;
 	TSharedPtr<FUICommandInfo> JumpToEnd;
+	TSharedPtr<FUICommandInfo> ScrubTime;
 
 	// Selection Range
 	TSharedPtr< FUICommandInfo > SetSelectionRangeStart;
diff --git a/Engine/Source/Editor/Sequencer/Private/Sequencer.cpp b/Engine/Source/Editor/Sequencer/Private/Sequencer.cpp
--- a/Engine/Source/Editor/Sequencer/Private/Sequencer.cpp
+++ b/Engine/Source/Editor/Sequencer/Private/Sequencer.cpp
@@ -3171,6 +3171,34 @@ void FSequencer::PlayTo(FMovieSceneSequencePlaybackParams PlaybackParams)
 	PauseOnFrame = PlayToTime;
 }
 
+void FSequencer::SnapSequencerTime(FFrameTime& ScrubTime)
+{
+	// Clamp first, snap to frame last
+	if (GetSequencerSettings()->ShouldKeepCursorInPlayRangeWhileScrubbing())
+	{
+		TRange<FFrameNumber> PlaybackRange =GetFocusedMovieSceneSequence()->GetMovieScene()->GetPlaybackRange();
+		ScrubTime = UE::MovieScene::ClampToDiscreteRange(ScrubTime, PlaybackRange);
+	}
+
+	if (GetSequencerSettings()->GetIsSnapEnabled())
+	{
+		FFrameRate TickResolution = GetFocusedTickResolution();
+		FFrameRate DisplayRate = GetFocusedDisplayRate();
+
+		// Set the style of the scrub handle
+		if (GetScrubStyle() == ESequencerScrubberStyle::FrameBlock)
+		{
+			// Floor to the display frame
+			ScrubTime = ConvertFrameTime(ConvertFrameTime(ScrubTime, TickResolution, DisplayRate).FloorToFrame(), DisplayRate, TickResolution);
+		}
+		else
+		{
+			// Snap (round) to display rate
+			ScrubTime = FFrameRate::Snap(ScrubTime, TickResolution, DisplayRate);
+		}
+	}
+}
+
 void FSequencer::ForceEvaluate()
 {
 	EvaluateInternal(PlayPosition.GetCurrentPositionAsRange());
diff --git a/Engine/Source/Editor/Sequencer/Private/Sequencer.h b/Engine/Source/Editor/Sequencer/Private/Sequencer.h
--- a/Engine/Source/Editor/Sequencer/Private/Sequencer.h
+++ b/Engine/Source/Editor/Sequencer/Private/Sequencer.h
@@ -725,6 +725,7 @@ public:
 	virtual FFrameTime GetLastEvaluatedLocalTime() const override;
 	virtual void SetGlobalTime(FFrameTime Time, bool bEvaluate = true) override;
 	virtual void PlayTo(FMovieSceneSequencePlaybackParams PlaybackParams) override;
+	virtual void SnapSequencerTime(FFrameTime& InOutScrubTime) override;
 	virtual void RestorePlaybackSpeed() override;
 	virtual void SnapToClosestPlaybackSpeed() override;
 	virtual void RequestInvalidateCachedData() override { bNeedsInvalidateCachedData = true; }
diff --git a/Engine/Source/Editor/Sequencer/Private/SequencerEdMode.cpp b/Engine/Source/Editor/Sequencer/Private/SequencerEdMode.cpp
--- a/Engine/Source/Editor/Sequencer/Private/SequencerEdMode.cpp
+++ b/Engine/Source/Editor/Sequencer/Private/SequencerEdMode.cpp
@@ -352,34 +352,6 @@ bool FSequencerEdMode::InputDelta(FEditorViewportClient* InViewportClient, FView
 	return FEdMode::InputDelta(InViewportClient, InViewport, InDrag, InRot, InScale);
 }
 
-static void SnapSequencerTime(TSharedPtr<FSequencer>& Sequencer, FFrameTime& ScrubTime)
-{
-	// Clamp first, snap to frame last
-	if (Sequencer->GetSequencerSettings()->ShouldKeepCursorInPlayRangeWhileScrubbing())
-	{
-		TRange<FFrameNumber> PlaybackRange = Sequencer->GetFocusedMovieSceneSequence()->GetMovieScene()->GetPlaybackRange();
-		ScrubTime = UE::MovieScene::ClampToDiscreteRange(ScrubTime, PlaybackRange);
-	}
-
-	if (Sequencer->GetSequencerSettings()->GetIsSnapEnabled())
-	{
-		FFrameRate TickResolution =  Sequencer->GetFocusedTickResolution();
-		FFrameRate DisplayRate = Sequencer->GetFocusedDisplayRate();
-
-		// Set the style of the scrub handle
-		if (Sequencer->GetScrubStyle() == ESequencerScrubberStyle::FrameBlock)
-		{
-			// Floor to the display frame
-			ScrubTime = ConvertFrameTime(ConvertFrameTime(ScrubTime, TickResolution, DisplayRate).FloorToFrame(), DisplayRate, TickResolution);
-		}
-		else
-		{
-			// Snap (round) to display rate
-			ScrubTime = FFrameRate::Snap(ScrubTime, TickResolution, DisplayRate);
-		}
-	}
-}
-
 bool FSequencerEdMode::ProcessCapturedMouseMoves(FEditorViewportClient* InViewportClient, FViewport* InViewport, const TArrayView<FIntPoint>& CapturedMouseMoves)
 {
 	TSharedPtr<FSequencer> ActiveSequencer;
@@ -422,7 +394,7 @@ bool FSequencerEdMode::ProcessCapturedMouseMoves(FEditorViewportClient* InViewpo
 							FFrameNumber FrameDiff = ViewRange.Value - ViewRange.Key;
 							FrameDiff = FrameDiff * FloatViewDiff;
 							FFrameTime ScrubTime = StartFrameNumber + FrameDiff;
-							SnapSequencerTime(ActiveSequencer, ScrubTime);
+							ActiveSequencer->SnapSequencerTime(ScrubTime);
 							ActiveSequencer->SetLocalTime(ScrubTime.GetFrame());
 						}
 					}
diff --git a/Engine/Source/Editor/Sequencer/Private/SequencerTimeSliderController.cpp b/Engine/Source/Editor/Sequencer/Private/SequencerTimeSliderController.cpp
--- a/Engine/Source/Editor/Sequencer/Private/SequencerTimeSliderController.cpp
+++ b/Engine/Source/Editor/Sequencer/Private/SequencerTimeSliderController.cpp
@@ -1445,6 +1445,89 @@ int32 FSequencerTimeSliderController::OnPaintViewArea( const FGeometry& Allotted
 	return LayerId;
 }
 
+void FSequencerTimeSliderController::SetPlaybackStatus(ETimeSliderPlaybackStatus InStatus)
+{
+	using namespace UE::Sequencer;
+	if (!WeakSequencer.IsValid())
+	{
+		return;
+	}
+	switch (InStatus)
+	{
+		case ETimeSliderPlaybackStatus::Jumping:
+		{
+			WeakSequencer.Pin()->SetPlaybackStatus(EMovieScenePlayerStatus::Jumping);
+			break;
+		}
+		case ETimeSliderPlaybackStatus::Paused:
+		{
+			WeakSequencer.Pin()->SetPlaybackStatus(EMovieScenePlayerStatus::Paused);
+			break;
+		}
+		case ETimeSliderPlaybackStatus::Playing:
+		{
+			WeakSequencer.Pin()->SetPlaybackStatus(EMovieScenePlayerStatus::Playing);
+			break;
+		}
+		case ETimeSliderPlaybackStatus::Scrubbing:
+		{
+			WeakSequencer.Pin()->SetPlaybackStatus(EMovieScenePlayerStatus::Scrubbing);
+			break;
+		}
+		case ETimeSliderPlaybackStatus::Stepping:
+		{
+			WeakSequencer.Pin()->SetPlaybackStatus(EMovieScenePlayerStatus::Stepping);
+			break;
+		}
+		default:
+		case ETimeSliderPlaybackStatus::Stopped:
+		{
+			WeakSequencer.Pin()->SetPlaybackStatus(EMovieScenePlayerStatus::Stopped);
+			break;
+		}
+	}
+}
+
+ETimeSliderPlaybackStatus FSequencerTimeSliderController::GetPlaybackStatus() const
+{
+	using namespace UE::Sequencer;
+	if (!WeakSequencer.IsValid())
+	{
+		return ETimeSliderPlaybackStatus::Stopped;
+	}
+	EMovieScenePlayerStatus::Type Status = WeakSequencer.Pin()->GetPlaybackStatus();
+	switch (Status)
+	{
+		case EMovieScenePlayerStatus::Jumping:
+		{
+			return ETimeSliderPlaybackStatus::Jumping;
+		}
+		case EMovieScenePlayerStatus::Paused:
+		{
+			return ETimeSliderPlaybackStatus::Paused;
+		}
+		case EMovieScenePlayerStatus::Playing:
+		{
+			return ETimeSliderPlaybackStatus::Playing;
+		}
+		case EMovieScenePlayerStatus::Scrubbing:
+		{
+			return ETimeSliderPlaybackStatus::Scrubbing;
+		}
+		case EMovieScenePlayerStatus::Stepping:
+		{
+			return ETimeSliderPlaybackStatus::Stepping;
+		}
+		case EMovieScenePlayerStatus::Stopped:
+		{
+			return ETimeSliderPlaybackStatus::Stopped;
+		}
+
+	}
+	return ETimeSliderPlaybackStatus::Stopped;
+
+}
+
 TSharedRef<SWidget> FSequencerTimeSliderController::OpenSetPlaybackRangeMenu(const FGeometry& MyGeometry, const FPointerEvent& MouseEvent)
 {
 	TSharedPtr<FSequencer> Sequencer = WeakSequencer.Pin();
diff --git a/Engine/Source/Editor/Sequencer/Private/Tools/SequencerEditTool_Selection.cpp b/Engine/Source/Editor/Sequencer/Private/Tools/SequencerEditTool_Selection.cpp
--- a/Engine/Source/Editor/Sequencer/Private/Tools/SequencerEditTool_Selection.cpp
+++ b/Engine/Source/Editor/Sequencer/Private/Tools/SequencerEditTool_Selection.cpp
@@ -9,6 +9,7 @@
 #include "SequencerHotspots.h"
 #include "IKeyArea.h"
 #include "Tools/SequencerEntityVisitor.h"
+#include "SequencerCommands.h"
 #include "MVVM/Views/ISequencerTreeView.h"
 #include "MVVM/Views/ITrackAreaHotspot.h"
 #include "MVVM/ViewModels/SequencerEditorViewModel.h"
@@ -102,6 +103,75 @@ private:
 	bool bPinned;
 };
 
+class FScrubTimeDragOperation
+	: public UE::Sequencer::ISequencerEditToolDragOperation
+{
+public:
+
+	FScrubTimeDragOperation(FSequencer& InSequencer, UE::Sequencer::STrackAreaView& InTrackArea)
+		: Sequencer(InSequencer)
+	{}
+
+public:
+
+	// ISequencerEditToolDragOperation interface
+
+	virtual FCursorReply GetCursor() const override
+	{
+		return FCursorReply::Cursor(EMouseCursor::Default);
+	}
+
+	virtual void OnBeginDrag(const FPointerEvent& MouseEvent, FVector2D LocalMousePos, const UE::Sequencer::FVirtualTrackArea& VirtualTrackArea) override
+	{
+		// Start a new marquee selection
+		InitialPosition = CurrentPosition = VirtualTrackArea.PhysicalToVirtual(LocalMousePos);
+		CurrentMousePos = LocalMousePos;
+
+		Sequencer.SetPlaybackStatus(EMovieScenePlayerStatus::Scrubbing);
+		SequencerStartTime = Sequencer.GetLocalTime();
+		InitialTime = VirtualTrackArea.PixelToSeconds(LocalMousePos.X);
+	}
+
+	virtual void OnDrag(const FPointerEvent& MouseEvent, FVector2D LocalMousePos, const UE::Sequencer::FVirtualTrackArea& VirtualTrackArea) override
+	{
+		CurrentTime = VirtualTrackArea.PixelToSeconds(LocalMousePos.X);
+		FFrameTime FrameTime = CalculateScrubTime();
+
+		Sequencer.SnapSequencerTime(FrameTime);
+		Sequencer.SetLocalTimeDirectly(FrameTime, true);
+	}
+
+	virtual void OnEndDrag(const FPointerEvent& MouseEvent, FVector2D LocalMousePos, const UE::Sequencer::FVirtualTrackArea& VirtualTrackArea) override
+	{
+		Sequencer.SetPlaybackStatus(EMovieScenePlayerStatus::Stopped);
+	}
+
+	virtual int32 OnPaint(const FGeometry& AllottedGeometry, const FSlateRect& MyCullingRect, FSlateWindowElementList& OutDrawElements, int32 LayerId) const override
+	{
+		return LayerId;
+	}
+
+private:
+
+	FFrameTime CalculateScrubTime()
+	{
+		double Diff = CurrentTime - InitialTime;
+		FFrameTime FrameTime= SequencerStartTime.Rate.AsFrameTime(Diff);
+		FrameTime += SequencerStartTime.Time;
+		return FrameTime;
+	}
+
+	/** The sequencer itself */
+	FSequencer& Sequencer;
+
+	FVector2D InitialPosition;
+	FVector2D CurrentPosition;
+	FVector2D CurrentMousePos;
+
+	FQualifiedFrameTime SequencerStartTime;
+	double InitialTime;
+	double CurrentTime;
+};
 
 class FMarqueeDragOperation
 	: public UE::Sequencer::ISequencerEditToolDragOperation
@@ -413,7 +483,14 @@ FReply FSequencerEditTool_Selection::OnMouseMove(SWidget& OwnerWidget, const FGe
 
 			if (!DragOperation.IsValid())
 			{
-				DragOperation = MakeShareable( new FMarqueeDragOperation(Sequencer, TrackArea) );
+				if (bIsScrubbingTime)
+				{
+					DragOperation = MakeShareable(new FScrubTimeDragOperation(Sequencer, TrackArea));
+				}
+				else
+				{
+					DragOperation = MakeShareable(new FMarqueeDragOperation(Sequencer, TrackArea));
+				}
 			}
 
 			if (DragOperation.IsValid())
@@ -482,6 +559,7 @@ void FSequencerEditTool_Selection::OnMouseLeave(SWidget& OwnerWidget, const FPoi
 	{
 		CursorDecorator = nullptr;
 	}
+	bIsScrubbingTime = false;
 }
 
 
@@ -493,9 +571,47 @@ void FSequencerEditTool_Selection::OnMouseCaptureLost()
 		DelayedDrag.Reset();
 		DragOperation = nullptr;
 		CursorDecorator = nullptr;
+		bIsScrubbingTime = false;
 	});
 }
 
+bool FSequencerEditTool_Selection::IsScrubTimeKeyEvent(const FKeyEvent& InKeyEvent)
+{
+	const FSequencerCommands& Commands = FSequencerCommands::Get();
+	// Need to iterate through primary and secondary to make sure they are all pressed.
+	for (uint32 i = 0; i < static_cast<uint8>(EMultipleKeyBindingIndex::NumChords); ++i)
+	{
+		EMultipleKeyBindingIndex ChordIndex = static_cast<EMultipleKeyBindingIndex>(i);
+		const FInputChord& Chord = *Commands.ScrubTimeViewport->GetActiveChord(ChordIndex);
+		const bool bIsMovingTimeSlider = Chord.IsValidChord() && InKeyEvent.GetKey() == Chord.Key;
+		if (bIsMovingTimeSlider)
+		{
+			return true;
+		}
+	}
+	return false;
+}
+
+FReply FSequencerEditTool_Selection::OnKeyDown(SWidget& OwnerWidget, const FGeometry& MyGeometry, const FKeyEvent& InKeyEvent)
+{
+	if (IsScrubTimeKeyEvent(InKeyEvent))
+	{
+		bIsScrubbingTime = true;
+		return FReply::Handled();
+	}
+	return FReply::Unhandled();
+}
+
+FReply FSequencerEditTool_Selection::OnKeyUp(SWidget& OwnerWidget, const FGeometry& MyGeometry, const FKeyEvent& InKeyEvent)
+{
+	if (IsScrubTimeKeyEvent(InKeyEvent) && bIsScrubbingTime)
+	{
+		bIsScrubbingTime = false;
+		//would be nice to cancle the drag but doesn't seem like we can like we can with curve editor drag handlers.
+		return FReply::Handled();
+	}
+	return FReply::Unhandled();
+}
 
 FName FSequencerEditTool_Selection::GetIdentifier() const
 {
diff --git a/Engine/Source/Editor/Sequencer/Private/Tools/SequencerEditTool_Selection.h b/Engine/Source/Editor/Sequencer/Private/Tools/SequencerEditTool_Selection.h
--- a/Engine/Source/Editor/Sequencer/Private/Tools/SequencerEditTool_Selection.h
+++ b/Engine/Source/Editor/Sequencer/Private/Tools/SequencerEditTool_Selection.h
@@ -42,12 +42,16 @@ public:
 	virtual FName GetIdentifier() const override;
 	virtual bool CanDeactivate() const override;
 	virtual TSharedPtr<UE::Sequencer::ITrackAreaHotspot> GetDragHotspot() const override;
-	
+	virtual FReply OnKeyDown(SWidget& OwnerWidget, const FGeometry& MyGeometry, const FKeyEvent& InKeyEvent) override;
+	virtual FReply OnKeyUp(SWidget& OwnerWidget, const FGeometry& MyGeometry, const FKeyEvent& InKeyEvent) override;
+
 private:
 
 	/** Update the software cursor */
 	void UpdateCursor(const FGeometry& MyGeometry, const FPointerEvent& MouseEvent);
 
+	/** Are we scrubbing time*/
+	bool IsScrubTimeKeyEvent(const FKeyEvent& InKeyEvent);
 private:
 
 	/** Helper class responsible for handling delayed dragging */
@@ -64,4 +68,7 @@ private:
 
 	/** Software cursor decorator brush */
 	const FSlateBrush* CursorDecorator;
+
+	/** Whether or not we are scrubbing time*/
+	bool bIsScrubbingTime = false;
 };
diff --git a/Engine/Source/Editor/Sequencer/Public/ISequencer.h b/Engine/Source/Editor/Sequencer/Public/ISequencer.h
--- a/Engine/Source/Editor/Sequencer/Public/ISequencer.h
+++ b/Engine/Source/Editor/Sequencer/Public/ISequencer.h
@@ -428,6 +428,9 @@ public:
 	/** Play from the current time to the requested time */
 	virtual void PlayTo(FMovieSceneSequencePlaybackParams PlaybackParams) = 0;
 
+	/*Modify the Sequencer time by any snap settings */
+	virtual void SnapSequencerTime(FFrameTime& InOutScrubTime) = 0;
+
 	/** Invalidate cached data so that it will be reevaluated on the next frame */
 	virtual void RequestInvalidateCachedData() = 0;
 
diff --git a/Engine/Source/Editor/Sequencer/Public/SequencerTimeSliderController.h b/Engine/Source/Editor/Sequencer/Public/SequencerTimeSliderController.h
--- a/Engine/Source/Editor/Sequencer/Public/SequencerTimeSliderController.h
+++ b/Engine/Source/Editor/Sequencer/Public/SequencerTimeSliderController.h
@@ -71,7 +71,13 @@ public:
 	virtual FFrameTime GetScrubPosition() const override { return TimeSliderArgs.ScrubPosition.Get(FFrameTime()); }
 
 	/** Get the current time for the Scrub handle which indicates what range is being evaluated. */
-	virtual void SetScrubPosition(FFrameTime InTime, bool bEvaluate) override { CommitScrubPosition(InTime, false, bEvaluate); }
+	virtual void SetScrubPosition(FFrameTime InTime, bool bEvaluate) override { CommitScrubPosition(InTime, GetPlaybackStatus() == ETimeSliderPlaybackStatus::Scrubbing, bEvaluate); }
+
+	/** Set the playback status for the controller*/
+	virtual void SetPlaybackStatus(ETimeSliderPlaybackStatus InStatus) override;
+
+	/** Get the playback status for the controller, by default it is ETimeSliderPlaybackStatus::Stopped */
+	virtual ETimeSliderPlaybackStatus GetPlaybackStatus() const override;
 
 	/**
 	 * Clamp the given range to the clamp range 
diff --git a/Engine/Source/Editor/SequencerCore/Private/MVVM/Views/STrackAreaView.cpp b/Engine/Source/Editor/SequencerCore/Private/MVVM/Views/STrackAreaView.cpp
--- a/Engine/Source/Editor/SequencerCore/Private/MVVM/Views/STrackAreaView.cpp
+++ b/Engine/Source/Editor/SequencerCore/Private/MVVM/Views/STrackAreaView.cpp
@@ -317,6 +317,31 @@ int32 STrackAreaView::OnPaint(const FPaintArgs& Args, const FGeometry& AllottedG
 	return LayerId;
 }
 
+FReply STrackAreaView::OnKeyDown(const FGeometry& MyGeometry, const FKeyEvent& InKeyEvent)
+{
+	TSharedPtr<FTrackAreaViewModel> TrackArea = WeakViewModel.Pin();
+	ISequencerEditTool* EditTool = TrackArea ? TrackArea->GetEditTool() : nullptr;
+
+	if (EditTool)
+	{
+		// Always ensure the edit tool is set up
+		InputStack.SetHandlerAt(0, EditTool);
+	}
+	return InputStack.HandleKeyDown(*this, MyGeometry, InKeyEvent);
+}
+
+FReply STrackAreaView::OnKeyUp(const FGeometry& MyGeometry, const FKeyEvent& InKeyEvent)
+{
+	TSharedPtr<FTrackAreaViewModel> TrackArea = WeakViewModel.Pin();
+	ISequencerEditTool* EditTool = TrackArea ? TrackArea->GetEditTool() : nullptr;
+
+	if (EditTool)
+	{
+		// Always ensure the edit tool is set up
+		InputStack.SetHandlerAt(0, EditTool);
+	}
+	return InputStack.HandleKeyUp(*this, MyGeometry, InKeyEvent);
+}
 
 FReply STrackAreaView::OnMouseButtonDown( const FGeometry& MyGeometry, const FPointerEvent& MouseEvent )
 {
diff --git a/Engine/Source/Editor/SequencerCore/Public/ISequencerEditTool.h b/Engine/Source/Editor/SequencerCore/Public/ISequencerEditTool.h
--- a/Engine/Source/Editor/SequencerCore/Public/ISequencerEditTool.h
+++ b/Engine/Source/Editor/SequencerCore/Public/ISequencerEditTool.h
@@ -76,6 +76,8 @@ public:
 	virtual FCursorReply OnCursorQuery(const FGeometry& MyGeometry, const FPointerEvent& CursorEvent) const = 0;
 	virtual void OnMouseEnter(SWidget& OwnerWidget, const FGeometry& MyGeometry, const FPointerEvent& MouseEvent) = 0;
 	virtual void OnMouseLeave(SWidget& OwnerWidget, const FPointerEvent& MouseEvent) = 0;
+	virtual FReply OnKeyDown(SWidget& OwnerWidget, const FGeometry& MyGeometry, const FKeyEvent& InKeyEvent) {return FReply::Unhandled();}
+	virtual FReply OnKeyUp(SWidget& OwnerWidget, const FGeometry& MyGeometry, const FKeyEvent& InKeyEvent) {return FReply::Unhandled();}
 	virtual FName GetIdentifier() const = 0;
 	virtual bool CanDeactivate() const = 0;
 	virtual TSharedPtr<ITrackAreaHotspot> GetDragHotspot() const = 0;
diff --git a/Engine/Source/Editor/SequencerCore/Public/ISequencerInputHandler.h b/Engine/Source/Editor/SequencerCore/Public/ISequencerInputHandler.h
--- a/Engine/Source/Editor/SequencerCore/Public/ISequencerInputHandler.h
+++ b/Engine/Source/Editor/SequencerCore/Public/ISequencerInputHandler.h
@@ -16,4 +16,6 @@ struct ISequencerInputHandler
 	virtual FReply OnMouseButtonUp(SWidget& OwnerWidget, const FGeometry& MyGeometry, const FPointerEvent& MouseEvent) = 0;
 	virtual FReply OnMouseMove(SWidget& OwnerWidget, const FGeometry& MyGeometry, const FPointerEvent& MouseEvent) = 0;
 	virtual FReply OnMouseWheel(SWidget& OwnerWidget, const FGeometry& MyGeometry, const FPointerEvent& MouseEvent) = 0;
+	virtual FReply OnKeyDown(SWidget& OwnerWidget, const FGeometry& MyGeometry, const FKeyEvent& InKeyEvent) { return FReply::Unhandled(); }
+	virtual FReply OnKeyUp(SWidget& OwnerWidget, const FGeometry& MyGeometry, const FKeyEvent& InKeyEvent) { return FReply::Unhandled(); }
 };
diff --git a/Engine/Source/Editor/SequencerCore/Public/MVVM/Views/STrackAreaView.h b/Engine/Source/Editor/SequencerCore/Public/MVVM/Views/STrackAreaView.h
--- a/Engine/Source/Editor/SequencerCore/Public/MVVM/Views/STrackAreaView.h
+++ b/Engine/Source/Editor/SequencerCore/Public/MVVM/Views/STrackAreaView.h
@@ -141,6 +141,9 @@ public:
 	void OnDragLeave(const FDragDropEvent& DragDropEvent) override;
 	FReply OnDragOver(const FGeometry& MyGeometry, const FDragDropEvent& DragDropEvent) override;
 	FReply OnDrop(const FGeometry& MyGeometry, const FDragDropEvent& DragDropEvent) override;
+	virtual bool SupportsKeyboardFocus() const override { return true; }
+	virtual FReply OnKeyDown(const FGeometry& MyGeometry, const FKeyEvent& InKeyEvent) override;
+	virtual FReply OnKeyUp(const FGeometry& MyGeometry, const FKeyEvent& InKeyEvent) override;
 
 protected:
 
diff --git a/Engine/Source/Editor/SequencerCore/Public/MVVM/Views/SequencerInputHandlerStack.h b/Engine/Source/Editor/SequencerCore/Public/MVVM/Views/SequencerInputHandlerStack.h
--- a/Engine/Source/Editor/SequencerCore/Public/MVVM/Views/SequencerInputHandlerStack.h
+++ b/Engine/Source/Editor/SequencerCore/Public/MVVM/Views/SequencerInputHandlerStack.h
@@ -85,10 +85,18 @@ public:
 		return ProcessEvent(&ISequencerInputHandler::OnMouseWheel, OwnerWidget, MyGeometry, MouseEvent);
 	}
 
-
+	FReply HandleKeyDown(SWidget& OwnerWidget, const FGeometry& MyGeometry, const FKeyEvent& InKeyEvent)
+	{
+		return ProcessKeyEvent(&ISequencerInputHandler::OnKeyDown, OwnerWidget, MyGeometry, InKeyEvent);
+	}
+	FReply HandleKeyUp(SWidget& OwnerWidget, const FGeometry& MyGeometry, const FKeyEvent& InKeyEvent)
+	{
+		return ProcessKeyEvent(&ISequencerInputHandler::OnKeyUp, OwnerWidget, MyGeometry, InKeyEvent);
+	}
 private:
 
 	typedef FReply(ISequencerInputHandler::*InputHandlerFunction)(SWidget&, const FGeometry&, const FPointerEvent&);
+	typedef FReply(ISequencerInputHandler::* KeyInputHandlerFunction)(SWidget&, const FGeometry&, const FKeyEvent&);
 
 	FReply ProcessEvent(InputHandlerFunction Function, SWidget& OwnerWidget, const FGeometry& MyGeometry, const FPointerEvent& MouseEvent)
 	{
@@ -136,6 +144,38 @@ private:
 		return Reply;
 	}
 
+	//use mouse capture for the key event
+	FReply ProcessKeyEvent(KeyInputHandlerFunction Function, SWidget& OwnerWidget, const FGeometry& MyGeometry, const FKeyEvent& KeyEvent)
+	{
+		FReply Reply = FReply::Unhandled();
+
+		// Give the captured index priority over everything
+		if (CapturedIndex != INDEX_NONE && Handlers[CapturedIndex])
+		{
+			Reply = (Handlers[CapturedIndex]->*Function)(OwnerWidget, MyGeometry, KeyEvent);
+			if (Reply.IsEventHandled())
+			{
+				return Reply;
+			}
+		}
+
+		for (int32 Index = 0; Index < Handlers.Num(); ++Index)
+		{
+			if (!Handlers[Index] || CapturedIndex == Index)
+			{
+				continue;
+			}
+
+			Reply = (Handlers[Index]->*Function)(OwnerWidget, MyGeometry, KeyEvent);
+			if (Reply.IsEventHandled())
+			{
+				return Reply;
+			}
+		}
+
+		return Reply;
+	}
+
 	/** Index of the handler that currently has the mouse captured */
 	int32 CapturedIndex;
 
diff --git a/Engine/Source/Editor/SequencerWidgets/Private/ITimeSlider.cpp b/Engine/Source/Editor/SequencerWidgets/Private/ITimeSlider.cpp
--- a/Engine/Source/Editor/SequencerWidgets/Private/ITimeSlider.cpp
+++ b/Engine/Source/Editor/SequencerWidgets/Private/ITimeSlider.cpp
@@ -48,3 +48,12 @@ void ITimeSliderController::SetPlayRange(FFrameNumber RangeStart, int32 RangeDur
 void ITimeSliderController::SetSelectionRange(const TRange<FFrameNumber>& NewRange)
 {
 }
+
+void ITimeSliderController::SetPlaybackStatus(ETimeSliderPlaybackStatus InStatus)
+{
+}
+
+ETimeSliderPlaybackStatus ITimeSliderController::GetPlaybackStatus() const
+{
+	return ETimeSliderPlaybackStatus::Stopped;
+}
\ No newline at end of file
diff --git a/Engine/Source/Editor/SequencerWidgets/Public/ITimeSlider.h b/Engine/Source/Editor/SequencerWidgets/Public/ITimeSlider.h
--- a/Engine/Source/Editor/SequencerWidgets/Public/ITimeSlider.h
+++ b/Engine/Source/Editor/SequencerWidgets/Public/ITimeSlider.h
@@ -66,6 +66,17 @@ struct FPaintViewAreaArgs
 	TOptional<FPaintPlaybackRangeArgs> PlaybackRangeArgs;
 };
 
+enum ETimeSliderPlaybackStatus 
+{
+	Stopped,
+	Playing,
+	Scrubbing,
+	Jumping,
+	Stepping,
+	Paused,
+	MAX
+};
+
 class ITimeSliderController : public ISequencerInputHandler
 {
 public:
@@ -121,6 +132,13 @@ public:
 
 	/** Set the current time for the Scrub handle which indicates what range is being evaluated. */
 	SEQUENCERWIDGETS_API virtual void SetScrubPosition(FFrameTime InTime, bool bEvaluate);
+	
+	/** Set the playback status for the controller*/
+	SEQUENCERWIDGETS_API virtual void SetPlaybackStatus(ETimeSliderPlaybackStatus InStatus);
+
+	/** Get the playback status for the controller, by default it is ETimeSliderPlaybackStatus::Stopped */
+	SEQUENCERWIDGETS_API virtual ETimeSliderPlaybackStatus GetPlaybackStatus() const;
+
 
 	/**
 	 * Set a new range based on a min, max and an interpolation mode
