Create a callback in NumbericTypeInterface and use it to broadcast changes in time format in sequencer settings. Register spinbox and NumericBoxEntry to the callback to ensure their format gets updated even when their value doesn't change.

diff --git a/Engine/Source/Editor/Sequencer/Private/SSequencer.cpp b/Engine/Source/Editor/Sequencer/Private/SSequencer.cpp
--- a/Engine/Source/Editor/Sequencer/Private/SSequencer.cpp
+++ b/Engine/Source/Editor/Sequencer/Private/SSequencer.cpp
@@ -184,7 +184,13 @@ void SSequencer::Construct(const FArguments& InArgs, TSharedRef<FSequencer> InSe
 	TAttribute<FFrameRate> GetDisplayRateAttr    = TAttribute<FFrameRate>(InSequencer, &FSequencer::GetFocusedDisplayRate);
 
 	// Create our numeric type interface so we can pass it to the time slider below.
-	NumericTypeInterface = MakeShareable(new FFrameNumberInterface(GetDisplayFormatAttr, GetZeroPadFramesAttr, GetTickResolutionAttr, GetDisplayRateAttr));
+	TSharedPtr<FFrameNumberInterface> FrameNumberInterface = MakeShareable(new FFrameNumberInterface(GetDisplayFormatAttr, GetZeroPadFramesAttr, GetTickResolutionAttr, GetDisplayRateAttr));
+	NumericTypeInterface = FrameNumberInterface;
+	
+	if (USequencerSettings* Settings = GetSequencerSettings())
+	{
+		Settings->GetOnTimeDisplayFormatChanged().AddRaw(FrameNumberInterface.Get(), &FFrameNumberInterface::DisplayFormatChanged);
+	}
 
 	FTimeSliderArgs TimeSliderArgs;
 	{
@@ -2809,6 +2815,11 @@ SSequencer::~SSequencer()
 		Window->DestroyWindowImmediately();
 		NodeGroupManager.Reset();
 	}
+
+	if (USequencerSettings* Settings = GetSequencerSettings())
+	{
+		Settings->GetOnTimeDisplayFormatChanged().RemoveAll(NumericTypeInterface.Get());
+	}
 }
 
 
diff --git a/Engine/Source/Editor/Sequencer/Private/Sequencer.cpp b/Engine/Source/Editor/Sequencer/Private/Sequencer.cpp
--- a/Engine/Source/Editor/Sequencer/Private/Sequencer.cpp
+++ b/Engine/Source/Editor/Sequencer/Private/Sequencer.cpp
@@ -872,6 +872,11 @@ void FSequencer::Tick(float InDeltaTime)
 			}
 
 			SuppressAutoEvalSignature.Reset();
+
+			if (GetNumericTypeInterface()->GetOnSettingChanged() != nullptr)
+			{
+				GetNumericTypeInterface()->GetOnSettingChanged()->Broadcast();
+			}
 		}
 	}
 
@@ -4361,6 +4366,10 @@ void FSequencer::ResetPerMovieSceneData()
 
 	SuppressAutoEvalSignature.Reset();
 
+	if (GetNumericTypeInterface()->GetOnSettingChanged() != nullptr)
+	{
+		GetNumericTypeInterface()->GetOnSettingChanged()->Broadcast();
+	}
 	// @todo run through all tracks for new movie scene changes
 	//  needed for audio track decompression
 }
diff --git a/Engine/Source/Editor/Sequencer/Private/SequencerSettings.cpp b/Engine/Source/Editor/Sequencer/Private/SequencerSettings.cpp
--- a/Engine/Source/Editor/Sequencer/Private/SequencerSettings.cpp
+++ b/Engine/Source/Editor/Sequencer/Private/SequencerSettings.cpp
@@ -548,6 +548,7 @@ void USequencerSettings::SetZeroPadFrames(uint8 InZeroPadFrames)
 	{
 		ZeroPadFrames = InZeroPadFrames;
 		SaveConfig();
+		OnTimeDisplayFormatChangedEvent.Broadcast();
 	}
 }
 
@@ -926,6 +927,10 @@ USequencerSettings::FOnLoopStateChanged& USequencerSettings::GetOnLoopStateChang
 	return OnLoopStateChangedEvent;
 }
 
+USequencerSettings::FOnTimeDisplayFormatChanged& USequencerSettings::GetOnTimeDisplayFormatChanged()
+{
+	return OnTimeDisplayFormatChangedEvent;
+}
 
 void USequencerSettings::SetTimeDisplayFormat(EFrameNumberDisplayFormats InFormat)
 {
@@ -933,6 +938,7 @@ void USequencerSettings::SetTimeDisplayFormat(EFrameNumberDisplayFormats InForma
 	{
 		FrameNumberDisplayFormat = InFormat;
 		SaveConfig();
+		OnTimeDisplayFormatChangedEvent.Broadcast();
 	}
 }
 
diff --git a/Engine/Source/Editor/Sequencer/Public/SequencerSettings.h b/Engine/Source/Editor/Sequencer/Public/SequencerSettings.h
--- a/Engine/Source/Editor/Sequencer/Public/SequencerSettings.h
+++ b/Engine/Source/Editor/Sequencer/Public/SequencerSettings.h
@@ -93,6 +93,7 @@ public:
 	DECLARE_MULTICAST_DELEGATE( FOnShowSelectedNodesOnlyChanged );
 	DECLARE_MULTICAST_DELEGATE_OneParam( FOnAllowEditsModeChanged, EAllowEditsMode );
 	DECLARE_MULTICAST_DELEGATE(FOnLoopStateChanged);
+	DECLARE_MULTICAST_DELEGATE(FOnTimeDisplayFormatChanged);
 
 	virtual void PostEditChangeProperty(FPropertyChangedEvent& PropertyChangedEvent) override;
 
@@ -392,6 +393,8 @@ public:
 
 	FOnLoopStateChanged& GetOnLoopStateChanged();
 
+	FOnTimeDisplayFormatChanged& GetOnTimeDisplayFormatChanged();
+
 	/** What format should we display the UI controls in when representing time in a sequence? */
 	EFrameNumberDisplayFormats GetTimeDisplayFormat() const { return FrameNumberDisplayFormat; }
 	/** Sets the time display format to the specified type. */
@@ -673,4 +676,5 @@ protected:
 	FOnShowSelectedNodesOnlyChanged OnShowSelectedNodesOnlyChangedEvent;
 	FOnAllowEditsModeChanged OnAllowEditsModeChangedEvent;
 	FOnLoopStateChanged OnLoopStateChangedEvent;
+	FOnTimeDisplayFormatChanged OnTimeDisplayFormatChangedEvent;
 };
diff --git a/Engine/Source/Runtime/Slate/Public/Widgets/Input/NumericTypeInterface.h b/Engine/Source/Runtime/Slate/Public/Widgets/Input/NumericTypeInterface.h
--- a/Engine/Source/Runtime/Slate/Public/Widgets/Input/NumericTypeInterface.h
+++ b/Engine/Source/Runtime/Slate/Public/Widgets/Input/NumericTypeInterface.h
@@ -13,6 +13,8 @@
 
 enum class EUnit : uint8;
 
+DECLARE_MULTICAST_DELEGATE(FOnSettingChanged);
+
 /** Interface to provide specific functionality for dealing with a numeric type. Currently includes string conversion functionality. */
 template<typename NumericType>
 struct INumericTypeInterface
@@ -33,6 +35,9 @@ struct INumericTypeInterface
 
 	/** Check whether the typed character is valid */
 	virtual bool IsCharacterValid(TCHAR InChar) const = 0;
+
+	/** Optional callback to broadcast when a setting in the type interface changes */
+	virtual FOnSettingChanged* GetOnSettingChanged() { return nullptr;}
 };
 
 /** Default numeric type interface */
diff --git a/Engine/Source/Runtime/Slate/Public/Widgets/Input/SNumericEntryBox.h b/Engine/Source/Runtime/Slate/Public/Widgets/Input/SNumericEntryBox.h
--- a/Engine/Source/Runtime/Slate/Public/Widgets/Input/SNumericEntryBox.h
+++ b/Engine/Source/Runtime/Slate/Public/Widgets/Input/SNumericEntryBox.h
@@ -210,6 +210,11 @@ public:
 		BorderImageFocused = &InArgs._EditableTextBoxStyle->BackgroundImageFocused;
 		Interface = InArgs._TypeInterface.IsValid() ? InArgs._TypeInterface : MakeShareable( new TDefaultNumericTypeInterface<NumericType> );
 
+		if (InArgs._TypeInterface.IsValid() && Interface->GetOnSettingChanged())
+		{
+			Interface->GetOnSettingChanged()->AddSP(this, &SNumericEntryBox::ResetCachedValueString);
+		}
+
 		MinFractionalDigits = (InArgs._MinFractionalDigits.Get().IsSet()) ? InArgs._MinFractionalDigits : DefaultMinFractionalDigits;
 		MaxFractionalDigits = (InArgs._MaxFractionalDigits.Get().IsSet()) ? InArgs._MaxFractionalDigits : DefaultMaxFractionalDigits;
 		SetMinFractionalDigits(MinFractionalDigits);
@@ -738,6 +743,14 @@ private:
 		}
 	}
 
+	void ResetCachedValueString() 
+	{ 
+		if (ValueAttribute.Get().IsSet())
+		{
+			bCachedValueStringDirty = true;
+		}
+	}
+
 private:
 
 	/** The default minimum fractional digits */
diff --git a/Engine/Source/Runtime/Slate/Public/Widgets/Input/SSpinBox.h b/Engine/Source/Runtime/Slate/Public/Widgets/Input/SSpinBox.h
--- a/Engine/Source/Runtime/Slate/Public/Widgets/Input/SSpinBox.h
+++ b/Engine/Source/Runtime/Slate/Public/Widgets/Input/SSpinBox.h
@@ -182,6 +182,11 @@ public:
 		SetForegroundColor(InArgs._Style->ForegroundColor);
 		Interface = InArgs._TypeInterface.IsValid() ? InArgs._TypeInterface : MakeShareable(new TDefaultNumericTypeInterface<NumericType>);
 
+		if (InArgs._TypeInterface.IsValid() && Interface->GetOnSettingChanged())
+		{
+			Interface->GetOnSettingChanged()->AddSP(this, &SSpinBox::ResetCachedValueString);
+		}
+
 		ValueAttribute = InArgs._Value;
 		OnValueChanged = InArgs._OnValueChanged;
 		OnValueCommitted = InArgs._OnValueCommitted;
@@ -1207,6 +1212,10 @@ private:
 
 	/** Used to prevent per-frame re-conversion of the cached numeric value to a string. */
 	FString CachedValueString;
+	void ResetCachedValueString() 
+	{ 
+		CachedValueString = Interface->ToString(ValueAttribute.Get()); 
+	}
 
 	/** Whetever the interfaced setting changed and the CachedValueString needs to be recomputed. */
 	mutable bool bCachedValueStringDirty;
diff --git a/Engine/Source/Runtime/TimeManagement/Public/FrameNumberNumericInterface.h b/Engine/Source/Runtime/TimeManagement/Public/FrameNumberNumericInterface.h
--- a/Engine/Source/Runtime/TimeManagement/Public/FrameNumberNumericInterface.h
+++ b/Engine/Source/Runtime/TimeManagement/Public/FrameNumberNumericInterface.h
@@ -223,6 +223,16 @@ struct FFrameNumberInterface : public INumericTypeInterface<double>
 		SubframeIndicator = InSubframeIndicator;
 	}
 
+	void DisplayFormatChanged()
+	{
+		OnSettingChangedEvent.Broadcast();
+	}
+
+	virtual FOnSettingChanged* GetOnSettingChanged() override
+	{
+		return &OnSettingChangedEvent;
+	}
+
 private:
 	TAttribute<EFrameNumberDisplayFormats> DisplayFormatAttr;
 	TAttribute<FFrameRate> TickResolutionAttr;
@@ -230,4 +240,5 @@ private:
 	TAttribute<uint8> ZeroPadFramesAttr;
 
 	FString SubframeIndicator;
+	FOnSettingChanged OnSettingChangedEvent;
 };
