Sequencer: New selection API

diff --git a/Engine/Plugins/Animation/ControlRig/Source/ControlRigEditor/Private/Sequencer/ControlRigParameterTrackEditor.cpp b/Engine/Plugins/Animation/ControlRig/Source/ControlRigEditor/Private/Sequencer/ControlRigParameterTrackEditor.cpp
--- a/Engine/Plugins/Animation/ControlRig/Source/ControlRigEditor/Private/Sequencer/ControlRigParameterTrackEditor.cpp
+++ b/Engine/Plugins/Animation/ControlRig/Source/ControlRigEditor/Private/Sequencer/ControlRigParameterTrackEditor.cpp
@@ -1,6 +1,8 @@
 // Copyright Epic Games, Inc. All Rights Reserved.
 
 #include "Sequencer/ControlRigParameterTrackEditor.h"
+#include "MVVM/Selection/Selection.h"
+#include "MVVM/ViewModels/SequencerEditorViewModel.h"
 #include "Animation/AnimMontage.h"
 #include "Sequencer/MovieSceneControlRigParameterTrack.h"
 #include "Sequencer/MovieSceneControlRigParameterSection.h"
@@ -2923,6 +2925,8 @@ void FControlRigParameterTrackEditor::SetUpEditModeIfNeeded(UControlRig* Control
 
 void FControlRigParameterTrackEditor::HandleControlSelected(UControlRig* Subject, FRigControlElement* ControlElement, bool bSelected)
 {
+	using namespace UE::Sequencer;
+
 	if(ControlElement == nullptr)
 	{
 		return;
@@ -3023,7 +3027,8 @@ void FControlRigParameterTrackEditor::HandleControlSelected(UControlRig* Subject
 		UMovieSceneControlRigParameterTrack* Track = CastChecked<UMovieSceneControlRigParameterTrack>(TrackResult.Track, ECastCheckedType::NullAllowed);
 		if (Track)
 		{
-			GetSequencer()->SuspendSelectionBroadcast();
+			FSelectionEventSuppressor SuppressSelectionEvents = GetSequencer()->GetViewModel()->GetSelection()->SuppressEvents();
+
 			//Just select in section to key, if deselecting makes sure deselected everywhere
 			if (bSelected == true)
 			{
@@ -3041,8 +3046,6 @@ void FControlRigParameterTrackEditor::HandleControlSelected(UControlRig* Subject
 					}
 				}
 			}
-			
-			GetSequencer()->ResumeSelectionBroadcast();
 
 			SetUpEditModeIfNeeded(Subject);
 
diff --git a/Engine/Plugins/Animation/ControlRig/Source/ControlRigEditor/Private/Tools/ControlRigTweener.cpp b/Engine/Plugins/Animation/ControlRig/Source/ControlRigEditor/Private/Tools/ControlRigTweener.cpp
--- a/Engine/Plugins/Animation/ControlRig/Source/ControlRigEditor/Private/Tools/ControlRigTweener.cpp
+++ b/Engine/Plugins/Animation/ControlRig/Source/ControlRigEditor/Private/Tools/ControlRigTweener.cpp
@@ -8,7 +8,7 @@
 #include "MovieSceneSequence.h"
 #include "MovieScene.h"
 #include "MVVM/CurveEditorExtension.h"
-#include "MVVM/ViewModels/EditorViewModel.h"
+#include "MVVM/ViewModels/SequencerEditorViewModel.h"
 #include "Rigs/RigControlHierarchy.h"
 #include "Sequencer/MovieSceneControlRigParameterTrack.h"
 #include "Sequencer/MovieSceneControlRigParameterSection.h"
@@ -95,7 +95,7 @@ bool FAnimSliderKeySelection::Setup(TWeakPtr<ISequencer>& InSequencer)
 	KeyMap.Reset();
 	if (InSequencer.IsValid())
 	{
-		TSharedPtr<FEditorViewModel> SequencerViewModel = InSequencer.Pin()->GetViewModel();
+		TSharedPtr<FSequencerEditorViewModel> SequencerViewModel = InSequencer.Pin()->GetViewModel();
 		FCurveEditorExtension* CurveEditorExtension = SequencerViewModel->CastDynamic<FCurveEditorExtension>();
 		check(CurveEditorExtension);
 		CurveEditor = CurveEditorExtension->GetCurveEditor();
diff --git a/Engine/Plugins/MovieScene/LevelSequenceEditor/Source/LevelSequenceEditor/Private/LevelSequenceEditorSubsystem.cpp b/Engine/Plugins/MovieScene/LevelSequenceEditor/Source/LevelSequenceEditor/Private/LevelSequenceEditorSubsystem.cpp
--- a/Engine/Plugins/MovieScene/LevelSequenceEditor/Source/LevelSequenceEditor/Private/LevelSequenceEditorSubsystem.cpp
+++ b/Engine/Plugins/MovieScene/LevelSequenceEditor/Source/LevelSequenceEditor/Private/LevelSequenceEditorSubsystem.cpp
@@ -1,6 +1,8 @@
 // Copyright Epic Games, Inc. All Rights Reserved.
 
 #include "LevelSequenceEditorSubsystem.h"
+#include "MVVM/ViewModels/SequencerEditorViewModel.h"
+#include "Scripting/SequencerScriptingLayer.h"
 
 #include "Evaluation/MovieScenePlayback.h"
 #include "ISequencerModule.h"
@@ -208,6 +210,16 @@ TSharedPtr<ISequencer> ULevelSequenceEditorSubsystem::GetActiveSequencer()
 	return nullptr;
 }
 
+USequencerScriptingLayer* ULevelSequenceEditorSubsystem::GetScriptingLayer()
+{
+	TSharedPtr<ISequencer> Sequencer = GetActiveSequencer();
+	if (Sequencer)
+	{
+		return Sequencer->GetViewModel()->GetScriptingLayer();
+	}
+	return nullptr;
+}
+
 TArray<FMovieSceneBindingProxy> ULevelSequenceEditorSubsystem::AddActors(const TArray<AActor*>& InActors)
 {
 	TArray<FMovieSceneBindingProxy> BindingProxies;
diff --git a/Engine/Plugins/MovieScene/LevelSequenceEditor/Source/LevelSequenceEditor/Public/LevelSequenceEditorSubsystem.h b/Engine/Plugins/MovieScene/LevelSequenceEditor/Source/LevelSequenceEditor/Public/LevelSequenceEditorSubsystem.h
--- a/Engine/Plugins/MovieScene/LevelSequenceEditor/Source/LevelSequenceEditor/Public/LevelSequenceEditorSubsystem.h
+++ b/Engine/Plugins/MovieScene/LevelSequenceEditor/Source/LevelSequenceEditor/Public/LevelSequenceEditorSubsystem.h
@@ -26,6 +26,7 @@ class FMenuBuilder;
 class UMovieSceneCompiledDataManager;
 class UMovieSceneFolder;
 class UMovieSceneSection;
+class USequencerScriptingLayer;
 
 USTRUCT(BlueprintType)
 struct FMovieSceneScriptingParams
@@ -54,6 +55,10 @@ public:
 
 	void OnSequencerCreated(TSharedRef<ISequencer> InSequencer);
 
+	/** Retrieve the outliner */
+	UFUNCTION(BlueprintPure, Category = "Level Sequence Editor")
+	USequencerScriptingLayer* GetScriptingLayer();
+
 	/** Add existing actors to Sequencer. Tracks will be automatically added based on default track settings. */
 	UFUNCTION(BlueprintCallable, Category = "Level Sequence Editor")
 	TArray<FMovieSceneBindingProxy> AddActors(const TArray<AActor*>& Actors);
diff --git a/Engine/Source/Editor/MovieSceneTools/Private/EditModes/SkeletalAnimationTrackEditMode.cpp b/Engine/Source/Editor/MovieSceneTools/Private/EditModes/SkeletalAnimationTrackEditMode.cpp
--- a/Engine/Source/Editor/MovieSceneTools/Private/EditModes/SkeletalAnimationTrackEditMode.cpp
+++ b/Engine/Source/Editor/MovieSceneTools/Private/EditModes/SkeletalAnimationTrackEditMode.cpp
@@ -259,10 +259,8 @@ void FSkeletalAnimationTrackEditMode::Render(const FSceneView* View, FViewport*
 	// Gather a map of object bindings to their implict selection state
 	TMap<const FMovieSceneBinding*, bool> ObjectBindingNodesSelectionMap;
 
-	const FSequencerSelection& Selection = Sequencer->GetSelection();
 	for (const FMovieSceneBinding& Binding : Sequence->GetMovieScene()->GetBindings())
 	{
-
 		TArrayView<TWeakObjectPtr<UObject>> BoundObjects = Sequencer->FindObjectsInCurrentSequence(Binding.GetObjectGuid());
 		for (UMovieSceneTrack* Track : Binding.GetTracks())
 		{
diff --git a/Engine/Source/Editor/MovieSceneTools/Private/MovieSceneToolHelpers.cpp b/Engine/Source/Editor/MovieSceneTools/Private/MovieSceneToolHelpers.cpp
--- a/Engine/Source/Editor/MovieSceneTools/Private/MovieSceneToolHelpers.cpp
+++ b/Engine/Source/Editor/MovieSceneTools/Private/MovieSceneToolHelpers.cpp
@@ -123,11 +123,11 @@ void FSkelMeshRecorderState::FinishRecording()
 /* MovieSceneToolHelpers
  *****************************************************************************/
 
-void MovieSceneToolHelpers::TrimSection(const TSet<TWeakObjectPtr<UMovieSceneSection>>& Sections, FQualifiedFrameTime Time, bool bTrimLeft, bool bDeleteKeys)
+void MovieSceneToolHelpers::TrimSection(const TSet<UMovieSceneSection*>& Sections, FQualifiedFrameTime Time, bool bTrimLeft, bool bDeleteKeys)
 {
-	for (auto Section : Sections)
+	for (UMovieSceneSection* Section : Sections)
 	{
-		if (Section.IsValid())
+		if (Section)
 		{
 			Section->TrimSection(Time, bTrimLeft, bDeleteKeys);
 		}
@@ -217,11 +217,11 @@ void MovieSceneToolHelpers::TrimOrExtendSection(UMovieSceneTrack* Track, TOption
 }
 
 
-void MovieSceneToolHelpers::SplitSection(const TSet<TWeakObjectPtr<UMovieSceneSection>>& Sections, FQualifiedFrameTime Time, bool bDeleteKeys)
+void MovieSceneToolHelpers::SplitSection(const TSet<UMovieSceneSection*>& Sections, FQualifiedFrameTime Time, bool bDeleteKeys)
 {
-	for (auto Section : Sections)
+	for (UMovieSceneSection* Section : Sections)
 	{
-		if (Section.IsValid())
+		if (Section)
 		{
 			Section->SplitSection(Time, bDeleteKeys);
 		}
diff --git a/Engine/Source/Editor/MovieSceneTools/Public/MovieSceneToolHelpers.h b/Engine/Source/Editor/MovieSceneTools/Public/MovieSceneToolHelpers.h
--- a/Engine/Source/Editor/MovieSceneTools/Public/MovieSceneToolHelpers.h
+++ b/Engine/Source/Editor/MovieSceneTools/Public/MovieSceneToolHelpers.h
@@ -106,7 +106,7 @@ public:
 	 * @param bTrimLeft Trim left or trim right
 	 * @param bDeleteKeys Delete keys outside the split ranges
 	 */
-	static void TrimSection(const TSet<TWeakObjectPtr<UMovieSceneSection>>& Sections, FQualifiedFrameTime Time, bool bTrimLeft, bool bDeleteKeys);
+	static void TrimSection(const TSet<UMovieSceneSection*>& Sections, FQualifiedFrameTime Time, bool bTrimLeft, bool bDeleteKeys);
 
 	/**
 	 * Trim or extend section at the given time
@@ -126,7 +126,7 @@ public:
 	 * @param Time	The time at which to split
 	 * @param bDeleteKeys Delete keys outside the split ranges
 	 */
-	static void SplitSection(const TSet<TWeakObjectPtr<UMovieSceneSection>>& Sections, FQualifiedFrameTime Time, bool bDeleteKeys);
+	static void SplitSection(const TSet<UMovieSceneSection*>& Sections, FQualifiedFrameTime Time, bool bDeleteKeys);
 
 	/**
 	 * Parse a shot name into its components.
diff --git a/Engine/Source/Editor/Sequencer/Private/MVVM/CurveEditorExtension.cpp b/Engine/Source/Editor/Sequencer/Private/MVVM/CurveEditorExtension.cpp
--- a/Engine/Source/Editor/Sequencer/Private/MVVM/CurveEditorExtension.cpp
+++ b/Engine/Source/Editor/Sequencer/Private/MVVM/CurveEditorExtension.cpp
@@ -1,6 +1,7 @@
 // Copyright Epic Games, Inc. All Rights Reserved.
 
 #include "MVVM/CurveEditorExtension.h"
+#include "MVVM/Selection/Selection.h"
 
 #include "FrameNumberDetailsCustomization.h"
 #include "Framework/Docking/TabManager.h"
@@ -404,10 +405,9 @@ void FCurveEditorExtension::SyncSelection()
 	CurveEditorTreeView->ClearSelection();
 
 	FCurveEditorTreeItemID FirstCurveEditorTreeItemID;
-	const TSet<TWeakPtr<FViewModel>>& SelectedItems = Sequencer->GetSelection().GetSelectedOutlinerItems();
-	for (TWeakPtr<FViewModel> SelectedItem : SelectedItems)
+	for (TViewModelPtr<IOutlinerExtension> SelectedItem : OwnerModel->GetSelection()->Outliner)
 	{
-		if (ICurveEditorTreeItemExtension* CurveEditorItem = ICastable::CastWeakPtr<ICurveEditorTreeItemExtension>(SelectedItem))
+		if (TViewModelPtr<ICurveEditorTreeItemExtension> CurveEditorItem = SelectedItem.ImplicitCast())
 		{
 			FCurveEditorTreeItemID CurveEditorTreeItem = CurveEditorItem->GetCurveEditorItemID();
 			if (CurveEditorTreeItem != FCurveEditorTreeItemID::Invalid())
diff --git a/Engine/Source/Editor/Sequencer/Private/MVVM/ViewModels/ChannelModel.cpp b/Engine/Source/Editor/Sequencer/Private/MVVM/ViewModels/ChannelModel.cpp
--- a/Engine/Source/Editor/Sequencer/Private/MVVM/ViewModels/ChannelModel.cpp
+++ b/Engine/Source/Editor/Sequencer/Private/MVVM/ViewModels/ChannelModel.cpp
@@ -20,6 +20,7 @@
 #include "MVVM/ViewModels/ViewModelIterators.h"
 #include "MVVM/ViewModels/OutlinerViewModel.h"
 #include "MVVM/ViewModels/TrackModel.h"
+#include "MVVM/Selection/Selection.h"
 #include "MVVM/Views/SOutlinerView.h"
 #include "MVVM/Views/SOutlinerItemViewBase.h"
 #include "MVVM/Views/SSequencerKeyNavigationButtons.h"
@@ -33,7 +34,6 @@
 #include "SKeyAreaEditorSwitcher.h"
 #include "SequencerCoreFwd.h"
 #include "SequencerSectionPainter.h"
-#include "SequencerSelection.h"
 #include "SequencerSettings.h"
 #include "SlotBase.h"
 #include "Styling/AppStyle.h"
@@ -451,16 +451,12 @@ bool FChannelGroupOverrideHelper::GetChannelsFromOutlinerSelection(TSharedPtr<FS
 		return false;
 	}
 
-	// Get all selected outliner items. We will filter out those that aren't channel groups.
-	TArray<TSharedRef<FViewModel>> SelectedItems;
-	Sequencer->GetSelection().GetSelectedOutlinerItems(SelectedItems);
-
 	// The selected items could either be channel groups, or tracks/track-rows with top-level channel groups that
 	// don't explicitly show up in the outliner (they are view-models that live under the track or row, but the
 	// user interface shows them as the track or row itself).
-	for (TSharedRef<FViewModel> SelectedItem : SelectedItems)
+	for (FViewModelPtr ViewModel : Sequencer->GetViewModel()->GetSelection()->Outliner)
 	{
-		if (FChannelGroupModel* ChannelGroup = SelectedItem->CastThis<FChannelGroupModel>())
+		if (TViewModelPtr<FChannelGroupModel> ChannelGroup = ViewModel.ImplicitCast())
 		{
 			const int32 NumChannels = ChannelGroup->GetChannels().Num();
 			for (int32 ChannelIndex = 0; ChannelIndex < NumChannels; ++ChannelIndex)
@@ -474,7 +470,7 @@ bool FChannelGroupOverrideHelper::GetChannelsFromOutlinerSelection(TSharedPtr<FS
 		else
 		{
 			const EViewModelListType ListType = FTrackModel::GetTopLevelChannelGroupType();
-			TOptional<FViewModelChildren> TopLevelChannels = SelectedItem->FindChildList(ListType);
+			TOptional<FViewModelChildren> TopLevelChannels = ViewModel->FindChildList(ListType);
 			if (TopLevelChannels.IsSet())
 			{
 				for (auto It(TopLevelChannels->IterateSubList<FChannelGroupModel>()); It; ++It)
diff --git a/Engine/Source/Editor/Sequencer/Private/MVVM/ViewModels/ObjectBindingModel.cpp b/Engine/Source/Editor/Sequencer/Private/MVVM/ViewModels/ObjectBindingModel.cpp
--- a/Engine/Source/Editor/Sequencer/Private/MVVM/ViewModels/ObjectBindingModel.cpp
+++ b/Engine/Source/Editor/Sequencer/Private/MVVM/ViewModels/ObjectBindingModel.cpp
@@ -11,6 +11,7 @@
 #include "MVVM/ViewModels/SequencerEditorViewModel.h"
 #include "MVVM/ViewModels/OutlinerViewModelDragDropOp.h"
 #include "MVVM/Views/SOutlinerObjectBindingView.h"
+#include "MVVM/Selection/Selection.h"
 #include "MVVM/Extensions/IRecyclableExtension.h"
 #include "Algo/Sort.h"
 #include "ClassViewerModule.h"
@@ -128,19 +129,16 @@ struct FMovieSceneSpawnableFlagCheckState
 		using namespace UE::Sequencer;
 
 		ECheckBoxState CheckState = ECheckBoxState::Undetermined;
-		for (TWeakPtr<FViewModel> WeakItem : Sequencer->GetSelection().GetSelectedOutlinerItems())
+		for (TViewModelPtr<IObjectBindingExtension> ObjectBinding : Sequencer->GetViewModel()->GetSelection()->Outliner.Filter<IObjectBindingExtension>())
 		{
-			if (IObjectBindingExtension* ObjectBindingID = ICastable::CastWeakPtr<IObjectBindingExtension>(WeakItem))
+			FMovieSceneSpawnable* SelectedSpawnable = MovieScene->FindSpawnable(ObjectBinding->GetObjectGuid());
+			if (SelectedSpawnable)
 			{
-				FMovieSceneSpawnable* SelectedSpawnable = MovieScene->FindSpawnable(ObjectBindingID->GetObjectGuid());
-				if (SelectedSpawnable)
+				if (CheckState != ECheckBoxState::Undetermined && SelectedSpawnable->*PtrToFlag != ( CheckState == ECheckBoxState::Checked ))
 				{
-					if (CheckState != ECheckBoxState::Undetermined && SelectedSpawnable->*PtrToFlag != ( CheckState == ECheckBoxState::Checked ))
-					{
-						return ECheckBoxState::Undetermined;
-					}
-					CheckState = SelectedSpawnable->*PtrToFlag ? ECheckBoxState::Checked : ECheckBoxState::Unchecked;
+					return ECheckBoxState::Undetermined;
 				}
+				CheckState = SelectedSpawnable->*PtrToFlag ? ECheckBoxState::Checked : ECheckBoxState::Unchecked;
 			}
 		}
 		return CheckState;
@@ -163,15 +161,12 @@ struct FMovieSceneSpawnableFlagToggler
 		const ECheckBoxState CheckState = FMovieSceneSpawnableFlagCheckState{Sequencer, MovieScene, PtrToFlag}();
 
 		MovieScene->Modify();
-		for (TWeakPtr<FViewModel> WeakItem : Sequencer->GetSelection().GetSelectedOutlinerItems())
+		for (TViewModelPtr<IObjectBindingExtension> ObjectBinding : Sequencer->GetViewModel()->GetSelection()->Outliner.Filter<IObjectBindingExtension>())
 		{
-			if (IObjectBindingExtension* ObjectBinding = ICastable::CastWeakPtr<IObjectBindingExtension>(WeakItem))
+			FMovieSceneSpawnable* SelectedSpawnable = MovieScene->FindSpawnable(ObjectBinding->GetObjectGuid());
+			if (SelectedSpawnable)
 			{
-				FMovieSceneSpawnable* SelectedSpawnable = MovieScene->FindSpawnable(ObjectBinding->GetObjectGuid());
-				if (SelectedSpawnable)
-				{
-					SelectedSpawnable->*PtrToFlag = (CheckState == ECheckBoxState::Unchecked);
-				}
+				SelectedSpawnable->*PtrToFlag = (CheckState == ECheckBoxState::Unchecked);
 			}
 		}
 	}
@@ -535,16 +530,10 @@ TSharedRef<SWidget> FObjectBindingModel::GetAddTrackMenuContent()
 
 	// Only include other selected object bindings if this binding is selected. Otherwise, this will lead to 
 	// confusion with multiple tracks being added to possibly unrelated objects
-	if (Sequencer->GetSelection().IsSelected(AsShared()))
+	if (Sequencer->GetViewModel()->GetSelection()->Outliner.IsSelected(SharedThis(this)))
 	{
-		for (const TWeakPtr<FViewModel>& Node : Sequencer->GetSelection().GetSelectedOutlinerItems())
+		for (TViewModelPtr<FObjectBindingModel> ObjectBindingNode : Sequencer->GetViewModel()->GetSelection()->Outliner.Filter<FObjectBindingModel>())
 		{
-			const FObjectBindingModel* ObjectBindingNode = Node.Pin()->CastThis<FObjectBindingModel>();
-			if (!ObjectBindingNode)
-			{
-				continue;
-			}
-
 			const FGuid Guid = ObjectBindingNode->GetObjectGuid();
 			for (auto RuntimeObject : Sequencer->FindBoundObjects(Guid, OwnerModel->GetSequenceID()))
 			{
@@ -854,18 +843,14 @@ void FObjectBindingModel::HandlePropertyMenuItemExecute(FPropertyPath PropertyPa
 		}
 	}
 
-	for (const TWeakPtr<FViewModel>& Node : Sequencer->GetSelection().GetSelectedOutlinerItems())
+	for (TViewModelPtr<FObjectBindingModel> ObjectBindingNode : Sequencer->GetViewModel()->GetSelection()->Outliner.Filter<FObjectBindingModel>())
 	{
-		FObjectBindingModel* ObjectBindingNode = ICastable::CastWeakPtr<FObjectBindingModel>(Node);
-		if (ObjectBindingNode)
+		FGuid Guid = ObjectBindingNode->GetObjectGuid();
+		for (auto RuntimeObject : Sequencer->FindBoundObjects(Guid, OwnerModel->GetSequenceID()))
 		{
-			FGuid Guid = ObjectBindingNode->GetObjectGuid();
-			for (auto RuntimeObject : Sequencer->FindBoundObjects(Guid, OwnerModel->GetSequenceID()))
+			if (Sequencer->CanKeyProperty(FCanKeyPropertyParams(RuntimeObject->GetClass(), PropertyPath)))
 			{
-				if (Sequencer->CanKeyProperty(FCanKeyPropertyParams(RuntimeObject->GetClass(), PropertyPath)))
-				{
-					KeyableBoundObjects.AddUnique(RuntimeObject.Get());
-				}
+				KeyableBoundObjects.AddUnique(RuntimeObject.Get());
 			}
 		}
 	}
@@ -1070,21 +1055,30 @@ void FObjectBindingModel::AddSpawnOwnershipMenu(FMenuBuilder& MenuBuilder)
 	}
 	auto Callback = [=](ESpawnOwnership NewOwnership){
 
-		FScopedTransaction Transaction(LOCTEXT("SetSpawnOwnership", "Set Spawnable Ownership"));
+		using namespace UE::Sequencer;
 
-		Spawnable->SetSpawnOwnership(NewOwnership);
+		FScopedTransaction Transaction(LOCTEXT("SetSpawnOwnership", "Set Spawnable Ownership"));
 
-		// Overwrite the completion state for all spawn sections to ensure the expected behaviour.
-		EMovieSceneCompletionMode NewCompletionMode = NewOwnership == ESpawnOwnership::InnerSequence ? EMovieSceneCompletionMode::RestoreState : EMovieSceneCompletionMode::KeepState;
+		MovieScene->Modify();
 
-		// Make all spawn sections retain state
-		UMovieSceneSpawnTrack* SpawnTrack = MovieScene->FindTrack<UMovieSceneSpawnTrack>(ObjectBindingID);
-		if (SpawnTrack)
+		for (TViewModelPtr<IObjectBindingExtension> ObjectBinding : Sequencer->GetViewModel()->GetSelection()->Outliner.Filter<IObjectBindingExtension>())
 		{
-			for (UMovieSceneSection* Section : SpawnTrack->GetAllSections())
+			FMovieSceneSpawnable* Spawnable = MovieScene->FindSpawnable(ObjectBinding->GetObjectGuid());
+		
+			Spawnable->SetSpawnOwnership(NewOwnership);
+
+			// Overwrite the completion state for all spawn sections to ensure the expected behaviour.
+			EMovieSceneCompletionMode NewCompletionMode = NewOwnership == ESpawnOwnership::InnerSequence ? EMovieSceneCompletionMode::RestoreState : EMovieSceneCompletionMode::KeepState;
+
+			// Make all spawn sections retain state
+			UMovieSceneSpawnTrack* SpawnTrack = MovieScene->FindTrack<UMovieSceneSpawnTrack>(ObjectBinding->GetObjectGuid());
+			if (SpawnTrack)
 			{
-				Section->Modify();
-				Section->EvalOptions.CompletionMode = NewCompletionMode;
+				for (UMovieSceneSection* Section : SpawnTrack->GetAllSections())
+				{
+					Section->Modify();
+					Section->EvalOptions.CompletionMode = NewCompletionMode;
+				}
 			}
 		}
 	};
@@ -1354,17 +1348,14 @@ void FObjectBindingModel::AddTagMenu(FMenuBuilder& MenuBuilder)
 
 		// Gather all the tags on all currently selected object binding IDs
 		FMovieSceneSequenceID SequenceID = OwnerModel->GetSequenceID();
-		for (const TWeakPtr<FViewModel>& Node : Sequencer->GetSelection().GetSelectedOutlinerItems())
+		for (TViewModelPtr<FObjectBindingModel> ObjectBindingNode : Sequencer->GetViewModel()->GetSelection()->Outliner.Filter<FObjectBindingModel>())
 		{
-			if (FObjectBindingModel* ObjectBindingNode = ICastable::CastWeakPtr<FObjectBindingModel>(Node))
-			{
-				const FGuid& ObjectID = ObjectBindingNode->GetObjectGuid();
+			const FGuid& ObjectID = ObjectBindingNode->GetObjectGuid();
 
-				UE::MovieScene::FFixedObjectBindingID BindingID(ObjectID, SequenceID);
-				for (auto It = Sequencer->GetObjectBindingTagCache()->IterateTags(BindingID); It; ++It)
-				{
-					AllTags.Add(It.Value());
-				}
+			UE::MovieScene::FFixedObjectBindingID BindingID(ObjectID, SequenceID);
+			for (auto It = Sequencer->GetObjectBindingTagCache()->IterateTags(BindingID); It; ++It)
+			{
+				AllTags.Add(It.Value());
 			}
 		}
 
@@ -1408,25 +1399,22 @@ ECheckBoxState FObjectBindingModel::GetTagCheckState(FName TagName)
 	TSharedPtr<FSequencer> Sequencer = OwnerModel->GetSequencerImpl();
 	FMovieSceneSequenceID SequenceID = OwnerModel->GetSequenceID();
 
-	for (const TWeakPtr<FViewModel>& Node : Sequencer->GetSelection().GetSelectedOutlinerItems())
+	for (TViewModelPtr<FObjectBindingModel> ObjectBindingNode : Sequencer->GetViewModel()->GetSelection()->Outliner.Filter<FObjectBindingModel>())
 	{
-		if (FObjectBindingModel* ObjectBindingNode = ICastable::CastWeakPtr<FObjectBindingModel>(Node))
-		{
-			const FGuid& ObjectID = ObjectBindingNode->GetObjectGuid();
+		const FGuid& ObjectID = ObjectBindingNode->GetObjectGuid();
 
-			UE::MovieScene::FFixedObjectBindingID BindingID(ObjectID, SequenceID);
-			ECheckBoxState ThisCheckState = Sequencer->GetObjectBindingTagCache()->HasTag(BindingID, TagName)
-				? ECheckBoxState::Checked
-				: ECheckBoxState::Unchecked;
+		UE::MovieScene::FFixedObjectBindingID BindingID(ObjectID, SequenceID);
+		ECheckBoxState ThisCheckState = Sequencer->GetObjectBindingTagCache()->HasTag(BindingID, TagName)
+			? ECheckBoxState::Checked
+			: ECheckBoxState::Unchecked;
 
-			if (CheckBoxState == ECheckBoxState::Undetermined)
-			{
-				CheckBoxState = ThisCheckState;
-			}
-			else if (CheckBoxState != ThisCheckState)
-			{
-				return ECheckBoxState::Undetermined;
-			}
+		if (CheckBoxState == ECheckBoxState::Undetermined)
+		{
+			CheckBoxState = ThisCheckState;
+		}
+		else if (CheckBoxState != ThisCheckState)
+		{
+			return ECheckBoxState::Undetermined;
 		}
 	}
 
@@ -1438,18 +1426,15 @@ void FObjectBindingModel::ToggleTag(FName TagName)
 	TSharedPtr<FSequencer> Sequencer = OwnerModel->GetSequencerImpl();
 	FMovieSceneSequenceID SequenceID = OwnerModel->GetSequenceID();
 
-	for (TWeakPtr<FViewModel> Node : Sequencer->GetSelection().GetSelectedOutlinerItems())
+	for (TViewModelPtr<FObjectBindingModel> ObjectBindingNode : Sequencer->GetViewModel()->GetSelection()->Outliner.Filter<FObjectBindingModel>())
 	{
-		if (FObjectBindingModel* ObjectBindingNode = ICastable::CastWeakPtr<FObjectBindingModel>(Node))
-		{
-			const FGuid& ObjectID = ObjectBindingNode->GetObjectGuid();
+		const FGuid& ObjectID = ObjectBindingNode->GetObjectGuid();
 
-			UE::MovieScene::FFixedObjectBindingID BindingID(ObjectID, SequenceID);
-			if (!Sequencer->GetObjectBindingTagCache()->HasTag(BindingID, TagName))
-			{
-				HandleAddTag(TagName);
-				return;
-			}
+		UE::MovieScene::FFixedObjectBindingID BindingID(ObjectID, SequenceID);
+		if (!Sequencer->GetObjectBindingTagCache()->HasTag(BindingID, TagName))
+		{
+			HandleAddTag(TagName);
+			return;
 		}
 	}
 
@@ -1465,14 +1450,10 @@ void FObjectBindingModel::HandleDeleteTag(FName TagName)
 	MovieScene->Modify();
 
 	FMovieSceneSequenceID SequenceID = OwnerModel->GetSequenceID();
-	for (TWeakPtr<FViewModel> Node : Sequencer->GetSelection().GetSelectedOutlinerItems())
+	for (TViewModelPtr<FObjectBindingModel> ObjectBindingNode : Sequencer->GetViewModel()->GetSelection()->Outliner.Filter<FObjectBindingModel>())
 	{
-		if (FObjectBindingModel* ObjectBindingNode = ICastable::CastWeakPtr<FObjectBindingModel>(Node))
-		{
-			const FGuid& ObjectID = ObjectBindingNode->GetObjectGuid();
-
-			MovieScene->UntagBinding(TagName, UE::MovieScene::FFixedObjectBindingID(ObjectID, SequenceID));
-		}
+		const FGuid& ObjectID = ObjectBindingNode->GetObjectGuid();
+		MovieScene->UntagBinding(TagName, UE::MovieScene::FFixedObjectBindingID(ObjectID, SequenceID));
 	}
 }
 
@@ -1485,14 +1466,10 @@ void FObjectBindingModel::HandleAddTag(FName TagName)
 	MovieScene->Modify();
 
 	FMovieSceneSequenceID SequenceID = OwnerModel->GetSequenceID();
-	for (TWeakPtr<FViewModel> Node : Sequencer->GetSelection().GetSelectedOutlinerItems())
+	for (TViewModelPtr<FObjectBindingModel> ObjectBindingNode : Sequencer->GetViewModel()->GetSelection()->Outliner.Filter<FObjectBindingModel>())
 	{
-		if (FObjectBindingModel* ObjectBindingNode = ICastable::CastWeakPtr<FObjectBindingModel>(Node))
-		{
-			const FGuid& ObjectID = ObjectBindingNode->GetObjectGuid();
-
-			MovieScene->TagBinding(TagName, UE::MovieScene::FFixedObjectBindingID(ObjectID, SequenceID));
-		}
+		const FGuid& ObjectID = ObjectBindingNode->GetObjectGuid();
+		MovieScene->TagBinding(TagName, UE::MovieScene::FFixedObjectBindingID(ObjectID, SequenceID));
 	}
 }
 
diff --git a/Engine/Source/Editor/Sequencer/Private/MVVM/ViewModels/OutlinerItemModel.cpp b/Engine/Source/Editor/Sequencer/Private/MVVM/ViewModels/OutlinerItemModel.cpp
--- a/Engine/Source/Editor/Sequencer/Private/MVVM/ViewModels/OutlinerItemModel.cpp
+++ b/Engine/Source/Editor/Sequencer/Private/MVVM/ViewModels/OutlinerItemModel.cpp
@@ -9,6 +9,7 @@
 #include "MVVM/Extensions/IObjectBindingExtension.h"
 #include "MVVM/ViewModels/OutlinerViewModel.h"
 #include "MVVM/ViewModels/SequencerEditorViewModel.h"
+#include "MVVM/Selection/Selection.h"
 
 #include "CurveEditor.h"
 #include "Framework/Commands/GenericCommands.h"
@@ -485,15 +486,12 @@ void FOutlinerItemModelMixin::BuildContextMenu(FMenuBuilder& MenuBuilder)
 	MenuBuilder.EndSection();
 
 	TArray<UMovieSceneTrack*> AllTracks;
-	for (TWeakPtr<FViewModel> Node : Sequencer->GetSelection().GetSelectedOutlinerItems())
+	for (TViewModelPtr<ITrackExtension> TrackExtension : Sequencer->GetViewModel()->GetSelection()->Outliner.Filter<ITrackExtension>())
 	{
-		if (ITrackExtension* TrackExtension = ICastable::CastWeakPtr<ITrackExtension>(Node))
+		UMovieSceneTrack* Track = TrackExtension->GetTrack();
+		if (Track)
 		{
-			UMovieSceneTrack* Track = TrackExtension->GetTrack();
-			if (Track)
-			{
-				AllTracks.Add(Track);
-			}
+			AllTracks.Add(Track);
 		}
 	}
 
@@ -558,10 +556,8 @@ void FOutlinerItemModelMixin::BuildOrganizeContextMenu(FMenuBuilder& MenuBuilder
 	
 	TArray<UMovieSceneTrack*> AllTracks;
 	TArray<TSharedPtr<FViewModel> > DraggableNodes;
-	for (TWeakPtr<FViewModel> WeakNode : Sequencer->GetSelection().GetSelectedOutlinerItems())
+	for (FViewModelPtr Node : Sequencer->GetViewModel()->GetSelection()->Outliner)
 	{
-		TSharedPtr<FViewModel> Node = WeakNode.Pin();
-
 		if (ITrackExtension* TrackExtension = Node->CastThis<ITrackExtension>())
 		{
 			UMovieSceneTrack* Track = TrackExtension->GetTrack();
@@ -609,11 +605,12 @@ void FOutlinerItemModelMixin::BuildOrganizeContextMenu(FMenuBuilder& MenuBuilder
 void FOutlinerItemModelMixin::BuildSectionColorTintsContextMenu(FMenuBuilder& MenuBuilder)
 {
 	TSharedPtr<FSequencer> Sequencer = GetEditor()->GetSequencerImpl();
+	TSharedPtr<FSequencerSelection> Selection = Sequencer->GetViewModel()->GetSelection();
 
 	TArray<UMovieSceneSection*> Sections;
-	for (TWeakObjectPtr<UMovieSceneSection> WeakSection : Sequencer->GetSelection().GetSelectedSections())
+	for (TViewModelPtr<FSectionModel> SectionModel : Selection->Outliner.Filter<FSectionModel>())
 	{
-		if (UMovieSceneSection* Section = WeakSection.Get())
+		if (UMovieSceneSection* Section = SectionModel->GetSection())
 		{
 			Sections.Add(Section);
 		}
@@ -621,16 +618,11 @@ void FOutlinerItemModelMixin::BuildSectionColorTintsContextMenu(FMenuBuilder& Me
 
 	if (!Sections.Num())
 	{
-		for (TWeakPtr<FViewModel> WeakNode : Sequencer->GetSelection().GetSelectedOutlinerItems())
+		for (TViewModelPtr<ITrackExtension> TrackExtension : Selection->Outliner.Filter<ITrackExtension>())
 		{
-			TSharedPtr<FViewModel> Node = WeakNode.Pin();
-
-			if (ITrackExtension* TrackExtension = Node->CastThis<ITrackExtension>())
+			for (UMovieSceneSection* Section : TrackExtension->GetSections())
 			{
-				for (UMovieSceneSection* Section : TrackExtension->GetSections())
-				{
-					Sections.Add(Section);
-				}
+				Sections.Add(Section);
 			}
 		}
 	}
diff --git a/Engine/Source/Editor/Sequencer/Private/MVVM/ViewModels/SequencerEditorViewModel.cpp b/Engine/Source/Editor/Sequencer/Private/MVVM/ViewModels/SequencerEditorViewModel.cpp
--- a/Engine/Source/Editor/Sequencer/Private/MVVM/ViewModels/SequencerEditorViewModel.cpp
+++ b/Engine/Source/Editor/Sequencer/Private/MVVM/ViewModels/SequencerEditorViewModel.cpp
@@ -3,6 +3,7 @@
 #include "MVVM/ViewModels/SequencerEditorViewModel.h"
 #include "MVVM/ViewModels/SequencerOutlinerViewModel.h"
 #include "MVVM/ViewModels/SequencerTrackAreaViewModel.h"
+#include "MVVM/Selection/Selection.h"
 #include "MVVM/CurveEditorExtension.h"
 #include "MVVM/Extensions/IOutlinerExtension.h"
 #include "MVVM/SharedViewModelData.h"
@@ -85,6 +86,13 @@ void FSequencerEditorViewModel::InitializeEditorImpl()
 	}
 }
 
+TSharedPtr<FSequencerCoreSelection> FSequencerEditorViewModel::CreateSelectionImpl()
+{
+	TSharedPtr<FSequencerSelection> SelectionInst = MakeShared<FSequencerSelection>();
+	SelectionInst->Initialize(SharedThis(this));
+	return SelectionInst;
+}
+
 TSharedPtr<FTrackAreaViewModel> FSequencerEditorViewModel::GetPinnedTrackArea() const
 {
 	return PinnedTrackArea;
@@ -113,6 +121,11 @@ bool FSequencerEditorViewModel::IsReadOnly() const
 	return !Sequencer || Sequencer->IsReadOnly();
 }
 
+TSharedPtr<FSequencerSelection> FSequencerEditorViewModel::GetSelection() const
+{
+	return StaticCastSharedPtr<FSequencerSelection>(FEditorViewModel::GetSelection());
+}
+
 void FSequencerEditorViewModel::HandleDataHierarchyChanged()
 {
 	if (FViewModelPtr RootModel = GetRootModel())
diff --git a/Engine/Source/Editor/Sequencer/Private/MVVM/ViewModels/TrackModel.cpp b/Engine/Source/Editor/Sequencer/Private/MVVM/ViewModels/TrackModel.cpp
--- a/Engine/Source/Editor/Sequencer/Private/MVVM/ViewModels/TrackModel.cpp
+++ b/Engine/Source/Editor/Sequencer/Private/MVVM/ViewModels/TrackModel.cpp
@@ -16,6 +16,7 @@
 #include "MVVM/TrackRowModelStorageExtension.h"
 #include "MVVM/Views/SOutlinerTrackView.h"
 #include "MVVM/ViewModels/SequencerEditorViewModel.h"
+#include "MVVM/Selection/Selection.h"
 
 #include "MovieScene.h"
 #include "MovieSceneFolder.h"
@@ -688,14 +689,11 @@ void FTrackModel::BuildContextMenu(FMenuBuilder& MenuBuilder)
 	// Find sections in the track to add batch properties for
 	TArray<TWeakObjectPtr<UObject>> TrackSections;
 
-	for (TWeakPtr<FViewModel> Node : GetEditor()->GetSequencer()->GetSelection().GetSelectedOutlinerItems())
+	for (TViewModelPtr<ITrackExtension> TrackExtension : WeakSequencer.Pin()->GetViewModel()->GetSelection()->Outliner.Filter<ITrackExtension>())
 	{
-		if (ITrackExtension* TrackExtension = ICastable::CastWeakPtr<ITrackExtension>(Node))
+		for (UMovieSceneSection* Section : TrackExtension->GetSections())
 		{
-			for (UMovieSceneSection* Section : TrackExtension->GetSections())
-			{
-				TrackSections.Add(Section);
-			}
+			TrackSections.Add(Section);
 		}
 	}
 
diff --git a/Engine/Source/Editor/Sequencer/Private/MVVM/ViewModels/TrackRowModel.cpp b/Engine/Source/Editor/Sequencer/Private/MVVM/ViewModels/TrackRowModel.cpp
--- a/Engine/Source/Editor/Sequencer/Private/MVVM/ViewModels/TrackRowModel.cpp
+++ b/Engine/Source/Editor/Sequencer/Private/MVVM/ViewModels/TrackRowModel.cpp
@@ -11,6 +11,7 @@
 #include "MVVM/ViewModels/TrackModelLayoutBuilder.h"
 #include "MVVM/Views/SOutlinerTrackView.h"
 #include "MVVM/ViewModels/SequencerEditorViewModel.h"
+#include "MVVM/Selection/Selection.h"
 
 #include "SequencerNodeTree.h"
 #include "ISequencer.h"
@@ -298,7 +299,7 @@ void FTrackRowModel::BuildContextMenu(FMenuBuilder& MenuBuilder)
 		return;
 	}
 
-	TSharedPtr<ISequencer> WeakSequencer = GetEditor()->GetSequencer();
+	TSharedPtr<ISequencer> Sequencer = GetEditor()->GetSequencer();
 
 	const int32 TrackRowIndex = GetRowIndex();
 
@@ -311,7 +312,7 @@ void FTrackRowModel::BuildContextMenu(FMenuBuilder& MenuBuilder)
 			LOCTEXT("AddSection", "Add Section"),
 			FText(),
 			FNewMenuDelegate::CreateLambda([=](FMenuBuilder& SubMenuBuilder){
-				FSequencerUtilities::PopulateMenu_CreateNewSection(SubMenuBuilder, TrackRowIndex + 1, Track, WeakSequencer);
+				FSequencerUtilities::PopulateMenu_CreateNewSection(SubMenuBuilder, TrackRowIndex + 1, Track, Sequencer);
 			})
 		);	
 	}
@@ -329,7 +330,7 @@ void FTrackRowModel::BuildContextMenu(FMenuBuilder& MenuBuilder)
 				LOCTEXT("BlendingAlgorithmSubMenu", "Blending Algorithm"),
 				FText(),
 				FNewMenuDelegate::CreateLambda([=](FMenuBuilder& SubMenuBuilder){
-					FSequencerUtilities::PopulateMenu_BlenderSubMenu(SubMenuBuilder, Track, WeakSequencer);
+					FSequencerUtilities::PopulateMenu_BlenderSubMenu(SubMenuBuilder, Track, Sequencer);
 				})
 			);
 		}
@@ -338,14 +339,11 @@ void FTrackRowModel::BuildContextMenu(FMenuBuilder& MenuBuilder)
 	// Find sections in the track to add batch properties for
 	TArray<TWeakObjectPtr<UObject>> TrackSections;
 	
-	for (TWeakPtr<FViewModel> Node : GetEditor()->GetSequencer()->GetSelection().GetSelectedOutlinerItems())
+	for (TViewModelPtr<ITrackExtension> TrackExtension : Sequencer->GetViewModel()->GetSelection()->Outliner.Filter<ITrackExtension>())
 	{
-		if (ITrackExtension* TrackExtension = ICastable::CastWeakPtr<ITrackExtension>(Node))
+		for (UMovieSceneSection* Section : TrackExtension->GetSections())
 		{
-			for (UMovieSceneSection* Section : TrackExtension->GetSections())
-			{
-				TrackSections.Add(Section);
-			}
+			TrackSections.Add(Section);
 		}
 	}
 	
@@ -363,12 +361,11 @@ void FTrackRowModel::BuildContextMenu(FMenuBuilder& MenuBuilder)
 		
 	if (TrackSections.Num())
 	{
-		TSharedPtr<FSequencer> Sequencer = GetEditor()->GetSequencerImpl();
 		MenuBuilder.AddSubMenu(
 			TrackSections.Num() > 1 ? LOCTEXT("BatchEditSections", "Batch Edit Sections") : LOCTEXT("EditSection", "Edit Section"),
 			FText(),
 			FNewMenuDelegate::CreateLambda([=](FMenuBuilder& SubMenuBuilder){
-				SequencerHelpers::AddPropertiesMenu(*Sequencer, SubMenuBuilder, TrackSections);
+				SequencerHelpers::AddPropertiesMenu(static_cast<FSequencer&>(*Sequencer), SubMenuBuilder, TrackSections);
 			})
 		);
 	}
diff --git a/Engine/Source/Editor/Sequencer/Private/MVVM/Views/SChannelView.cpp b/Engine/Source/Editor/Sequencer/Private/MVVM/Views/SChannelView.cpp
--- a/Engine/Source/Editor/Sequencer/Private/MVVM/Views/SChannelView.cpp
+++ b/Engine/Source/Editor/Sequencer/Private/MVVM/Views/SChannelView.cpp
@@ -5,7 +5,9 @@
 #include "MVVM/ViewModels/SectionModel.h"
 #include "MVVM/ViewModels/SequenceModel.h"
 #include "MVVM/ViewModels/TrackAreaViewModel.h"
+#include "MVVM/ViewModels/SequencerEditorViewModel.h"
 #include "MVVM/Views/STrackAreaView.h"
+#include "MVVM/Selection/Selection.h"
 
 #include "IKeyArea.h"
 #include "MovieScene.h"
@@ -57,8 +59,8 @@ struct FSequencerKeyRendererInterface : IKeyRendererInterface
 	{
 		TSharedPtr<FKeyHotspot> KeyHotspot = HotspotCast<FKeyHotspot>(Hotspot);
 
-		SelectedKeys     = &Sequencer->GetSelection().GetRawSelectedKeys();
-		SelectionPreview = &Sequencer->GetSelectionPreview().GetRawDefinedKeyStates();
+		SelectedKeys     = &Sequencer->GetViewModel()->GetSelection()->KeySelection.GetSelected();
+		SelectionPreview = &Sequencer->GetSelectionPreview().GetDefinedKeyStates();
 		HoveredKeys      = KeyHotspot ? &KeyHotspot->RawKeys : nullptr;
 	}
 
@@ -126,7 +128,7 @@ FChannelViewKeyCachedState::FChannelViewKeyCachedState(TRange<FFrameTime> InVisi
 	ValidPlayRangeMax = UE::MovieScene::DiscreteExclusiveUpper(ValidKeyRange);
 	VisibleRange = InVisibleRange;
 	KeySizePx = FVector2D(12, 12);
-	SelectionSerial = Sequencer->GetSelection().GetSerialNumber();
+	SelectionSerial = Sequencer->GetViewModel()->GetSelection()->GetSerialNumber();
 	SelectionPreviewHash = Sequencer->GetSelectionPreview().GetSelectionHash();
 	bShowCurve = Channel && Channel->GetKeyArea()->ShouldShowCurve();
 	bShowKeyBars = Sequencer->GetSequencerSettings()->GetShowKeyBars();
@@ -251,10 +253,15 @@ FReply SChannelView::OnMouseButtonDown(const FGeometry& MyGeometry, const FPoint
 		if (NewKeys.Num())
 		{
 			TSharedPtr<FSequencer> Sequencer = LegacyGetSequencer();
-			Sequencer->GetSelection().EmptySelectedKeys();
+			FKeySelection& KeySelection = Sequencer->GetViewModel()->GetSelection()->KeySelection;
+
+			KeySelection.Empty();
 			for (const FSequencerSelectedKey& NewKey : NewKeys)
 			{
-				Sequencer->GetSelection().AddToSelection(NewKey);
+				if (TSharedPtr<FChannelModel> Channel = NewKey.WeakChannel.Pin())
+				{
+					KeySelection.Select(Channel, NewKey.KeyHandle);
+				}
 			}
 
 			// Pass the event to the tool to copy the hovered key and move it
@@ -376,11 +383,16 @@ FReply SChannelView::OnMouseButtonUp(const FGeometry& MyGeometry, const FPointer
 		{
 			Sequencer->SnapToFrame();
 
-			for (const FSequencerSelectedKey& SelectedKey : Sequencer->GetSelection().GetSelectedKeys())
+			FSequencerSelection& Selection = *Sequencer->GetViewModel()->GetSelection();
+			for (FKeyHandle Key : Selection.KeySelection)
 			{
-				const FFrameNumber CurrentTime = SelectedKey.WeakChannel.Pin()->GetKeyArea()->GetKeyTime(SelectedKey.KeyHandle);
-				Sequencer->SetLocalTime(CurrentTime, ESnapTimeMode::STM_Interval);
-				break;
+				TSharedPtr<FChannelModel> Channel = Selection.KeySelection.GetModelForKey(Key);
+				if (Channel)
+				{
+					const FFrameNumber CurrentTime = Channel->GetKeyArea()->GetKeyTime(Key);
+					Sequencer->SetLocalTime(CurrentTime, ESnapTimeMode::STM_Interval);
+					break;
+				}
 			}
 		}
 		GEditor->EndTransaction();
@@ -589,7 +601,7 @@ int32 SChannelView::DrawLane(const FPaintArgs& Args, const FGeometry& AllottedGe
 
 	// --------------------------------------------
 	// Draw outliner selection tint
-	if (Selection.IsSelected(OutlinerItem.AsModel()))
+	if (Selection.Outliner.IsSelected(OutlinerItem))
 	{
 		static const FName SelectionColorName("SelectionColor");
 		static const FName SelectedTrackTintBrushName("Sequencer.Section.SelectedTrackTint");
diff --git a/Engine/Source/Editor/Sequencer/Private/SSequencer.cpp b/Engine/Source/Editor/Sequencer/Private/SSequencer.cpp
--- a/Engine/Source/Editor/Sequencer/Private/SSequencer.cpp
+++ b/Engine/Source/Editor/Sequencer/Private/SSequencer.cpp
@@ -12,6 +12,7 @@
 #include "MVVM/ViewModels/SequencerTrackAreaViewModel.h"
 #include "MVVM/ViewModels/SequencerOutlinerViewModel.h"
 #include "MVVM/ViewModels/TrackRowModel.h"
+#include "MVVM/Selection/Selection.h"
 #include "Sections/MovieSceneSubSection.h"
 #include "MovieScene.h"
 #include "Framework/MultiBox/MultiBoxDefs.h"
@@ -293,7 +294,7 @@ void SSequencer::Construct(const FArguments& InArgs, TSharedRef<FSequencer> InSe
 	TSharedPtr<FSequencerEditorViewModel> SequencerViewModel = InSequencer->GetViewModel()->CastThisShared<FSequencerEditorViewModel>();
 	SAssignNew(PinnedTrackArea, SSequencerTrackAreaView, SequencerViewModel->GetPinnedTrackArea(), TimeSliderControllerRef);
 	SAssignNew(PinnedTreeView, SOutlinerView, SequencerViewModel->GetOutliner(), PinnedTrackArea.ToSharedRef())
-		.SelectionHandler(InArgs._SelectionHandler)
+		.Selection(InSequencer->GetViewModel()->GetSelection())
 		.ExternalScrollbar(PinnedAreaScrollBar)
 		.Clipping(EWidgetClipping::ClipToBounds);
 
@@ -304,7 +305,7 @@ void SSequencer::Construct(const FArguments& InArgs, TSharedRef<FSequencer> InSe
 
 	SAssignNew(TrackArea, SSequencerTrackAreaView, SequencerViewModel->GetTrackArea(), TimeSliderControllerRef);
 	SAssignNew(TreeView, SOutlinerView, SequencerViewModel->GetOutliner(), TrackArea.ToSharedRef())
-		.SelectionHandler(InArgs._SelectionHandler)
+		.Selection(InSequencer->GetViewModel()->GetSelection())
 		.ExternalScrollbar(ScrollBar)
 		.Clipping(EWidgetClipping::ClipToBounds);
 
@@ -892,8 +893,8 @@ void SSequencer::Construct(const FArguments& InArgs, TSharedRef<FSequencer> InSe
 
 	ApplySequencerCustomization(RootCustomization);
 
-	InSequencer->GetSelection().GetOnKeySelectionChanged().AddSP(this, &SSequencer::HandleKeySelectionChanged);
-	InSequencer->GetSelection().GetOnOutlinerNodeSelectionChanged().AddSP(this, &SSequencer::HandleOutlinerNodeSelectionChanged);
+	InSequencer->GetViewModel()->GetSelection()->KeySelection.OnChanged.AddSP(this, &SSequencer::HandleKeySelectionChanged);
+	InSequencer->GetViewModel()->GetSelection()->Outliner.OnChanged.AddSP(this, &SSequencer::HandleOutlinerNodeSelectionChanged);
 
 	ResetBreadcrumbs();
 }
@@ -1103,21 +1104,19 @@ void SSequencer::HandleOutlinerNodeSelectionChanged()
 		return;
 	}
 
-	const TSet<TWeakPtr<FViewModel>>& SelectedDisplayNodes = Sequencer->GetSelection().GetSelectedOutlinerItems();
-
 	TSharedPtr<FCurveEditor> CurveEditor = CurveEditorExtension->GetCurveEditor();
 	if (CurveEditor)
 	{
 		const USequencerSettings* SequencerSettings = GetSequencerSettings();
 		// If we're isolating to the selection and there is one, add the filter
-		if (SequencerSettings->ShouldIsolateToCurveEditorSelection() && SelectedDisplayNodes.Num() != 0)
+		if (SequencerSettings->ShouldIsolateToCurveEditorSelection() && Sequencer->GetViewModel()->GetSelection()->Outliner.Num() != 0)
 		{
 			if (!SequencerSelectionCurveEditorFilter)
 			{
 				SequencerSelectionCurveEditorFilter = MakeShared<FSequencerSelectionCurveFilter>();
 			}
 
-			SequencerSelectionCurveEditorFilter->Update(Sequencer->GetSelection().GetSelectedOutlinerItems(), SequencerSettings->GetAutoExpandNodesOnSelection());
+			SequencerSelectionCurveEditorFilter->Update(Sequencer->GetViewModel()->GetSelection(), SequencerSettings->GetAutoExpandNodesOnSelection());
 
 			CurveEditor->GetTree()->AddFilter(SequencerSelectionCurveEditorFilter);
 		}
@@ -3189,17 +3188,11 @@ void SSequencer::OnAssetsDropped( const FAssetDragDropOp& DragDropOp )
 		GWarn->EndSlowTask();
 	}
 
-	const TSet<TWeakPtr<FViewModel>>& SelectedNodes = SequencerPtr.Pin()->GetSelection().GetSelectedOutlinerItems();
 	FGuid TargetObjectGuid;
 	// if exactly one object node is selected, we have a target object guid
-	TSharedPtr<const FViewModel> DisplayNode;
-	if (SelectedNodes.Num() == 1)
+	if (SequencerPtr.Pin()->GetViewModel()->GetSelection()->Outliner.Num() == 1)
 	{
-		for (TWeakPtr<const FViewModel> SelectedNode : SelectedNodes )
-		{
-			DisplayNode = SelectedNode.Pin();
-		}
-		if (const FObjectBindingModel* ObjectBindingNode = DisplayNode->CastThis<FObjectBindingModel>())
+		for (TViewModelPtr<FObjectBindingModel> ObjectBindingNode : SequencerPtr.Pin()->GetViewModel()->GetSelection()->Outliner.Filter<FObjectBindingModel>())
 		{
 			TargetObjectGuid = ObjectBindingNode->GetObjectGuid();
 		}
@@ -3476,14 +3469,10 @@ void SSequencer::StepToKey(bool bStepToNextKey, bool bCameraOnly)
 		}
 		else
 		{
-			Nodes.Reserve(Sequencer->GetSelection().GetSelectedOutlinerItems().Num());
-			for (const TWeakPtr<FViewModel>& WeakModel : Sequencer->GetSelection().GetSelectedOutlinerItems())
+			Nodes.Reserve(Sequencer->GetViewModel()->GetSelection()->Outliner.Num());
+			for (TViewModelPtr<IOutlinerExtension> Model : Sequencer->GetViewModel()->GetSelection()->Outliner.Filter<IOutlinerExtension>())
 			{
-				TViewModelPtr<IOutlinerExtension> Model = CastViewModel<IOutlinerExtension>(WeakModel.Pin());
-				if (Model)
-				{
-					Nodes.Add(Model);
-				}
+				Nodes.Add(Model);
 			}
 			if ( Nodes.Num() == 0 )
 			{
@@ -3817,16 +3806,23 @@ FPasteContextMenuArgs SSequencer::GeneratePasteArgs(FFrameNumber PasteAtTime, TS
 	// Paste into the currently selected sections, or hit test the mouse position as a last resort
 	TArray<TViewModelPtr<IOutlinerExtension>> PasteIntoNodes;
 	{
-		TSet<TWeakObjectPtr<UMovieSceneSection>> Sections = Sequencer->GetSelection().GetSelectedSections();
-		
-		for (const FSequencerSelectedKey& Key : Sequencer->GetSelection().GetSelectedKeys())
+		TSharedPtr<FSequencerSelection> Selection = Sequencer->GetViewModel()->GetSelection();
+
+		TSet<UMovieSceneSection*> Sections = Selection->GetSelectedSections();
+
+		for (FKeyHandle Key : Selection->KeySelection)
 		{
-			Sections.Add(Key.Section);
+			TSharedPtr<FChannelModel> Channel = Selection->KeySelection.GetModelForKey(Key);
+			UMovieSceneSection*       Section = Channel ? Channel->GetSection() : nullptr;
+			if (Channel)
+			{
+				Sections.Add(Section);
+			}
 		}
 
-		for (TWeakObjectPtr<UMovieSceneSection> WeakSection : Sections)
+		for (UMovieSceneSection* Section : Sections)
 		{
-			if (TSharedPtr<FSectionModel> Handle = Sequencer->GetNodeTree()->GetSectionModel(WeakSection.Get()))
+			if (TSharedPtr<FSectionModel> Handle = Sequencer->GetNodeTree()->GetSectionModel(Section))
 			{
 				TViewModelPtr<IOutlinerExtension> TrackModel = Handle->GetParentTrackModel().ImplicitCast();
 				if (TrackModel)
@@ -3836,13 +3832,11 @@ FPasteContextMenuArgs SSequencer::GeneratePasteArgs(FFrameNumber PasteAtTime, TS
 			}
 		}
 
-		TArray<TSharedRef<FViewModel>> SelectedNodes;
-		Sequencer->GetSelection().GetSelectedOutlinerItems(SelectedNodes);
-		for (const TSharedRef<FViewModel>& SelectedNode : SelectedNodes)
+		for (FViewModelPtr SelectedNode : Sequencer->GetViewModel()->GetSelection()->Outliner)
 		{
 			if (SelectedNode->IsA<FCategoryGroupModel>() || SelectedNode->IsA<ITrackExtension>() || SelectedNode->IsA<FChannelGroupModel>())
 			{
-				TViewModelPtr<IOutlinerExtension> TrackModel =  CastViewModel<IOutlinerExtension>(SelectedNode);
+				TViewModelPtr<IOutlinerExtension> TrackModel = CastViewModelChecked<IOutlinerExtension>(SelectedNode);
 				if (TrackModel)
 				{
 					PasteIntoNodes.Add(TrackModel);
@@ -3868,8 +3862,7 @@ void SSequencer::OnPaste()
 	using namespace UE::Sequencer;
 
 	TSharedPtr<FSequencer> Sequencer = SequencerPtr.Pin();
-	const TSet<TWeakPtr<FViewModel>>& SelectedNodes = Sequencer->GetSelection().GetSelectedOutlinerItems();
-	if (SelectedNodes.Num() == 0)
+	if (Sequencer->GetViewModel()->GetSelection()->Outliner.Num() == 0)
 	{
 		if (!OpenPasteMenu())
 		{
diff --git a/Engine/Source/Editor/Sequencer/Private/SSequencer.h b/Engine/Source/Editor/Sequencer/Private/SSequencer.h
--- a/Engine/Source/Editor/Sequencer/Private/SSequencer.h
+++ b/Engine/Source/Editor/Sequencer/Private/SSequencer.h
@@ -40,17 +40,16 @@ class SSequencerGroupManager;
 class SSequencerTreeFilterStatusBar;
 struct FPaintPlaybackRangeArgs;
 struct FSequencerCustomizationInfo;
-struct FSequencerSelectionCurveFilter;
 
 namespace UE
 {
 namespace Sequencer
 {
 
-	class IOutlinerSelectionHandler;
 	class SOutlinerView;
 	class STrackAreaView;
 	class FVirtualTrackArea;
+	struct FSequencerSelectionCurveFilter;
 
 } // namespace Sequencer
 } // namespace UE
@@ -288,9 +287,6 @@ public:
 		/** Extender to use for the toolbar. */
 		SLATE_ARGUMENT( TSharedPtr<FExtender>, ToolbarExtender )
 
-		/** Selection handler to pass to outliner view */
-		SLATE_ATTRIBUTE( TSharedPtr<UE::Sequencer::IOutlinerSelectionHandler>, SelectionHandler )
-
 		/** Whether to display the playback range spin box in time range slider */
 		SLATE_ARGUMENT( bool, ShowPlaybackRangeInTimeSlider )
 
@@ -681,7 +677,7 @@ private:
 	TSharedPtr<UE::Sequencer::STrackAreaView> PinnedTrackArea;
 
 	/** Curve editor filter that shows only the selected nodes */
-	TSharedPtr<FSequencerSelectionCurveFilter> SequencerSelectionCurveEditorFilter;
+	TSharedPtr<UE::Sequencer::FSequencerSelectionCurveFilter> SequencerSelectionCurveEditorFilter;
 
 	/** The breadcrumb trail widget for this sequencer */
 	TSharedPtr<SBreadcrumbTrail<FSequencerBreadcrumb>> BreadcrumbTrail;
diff --git a/Engine/Source/Editor/Sequencer/Private/SSequencerDebugVisualizer.cpp b/Engine/Source/Editor/Sequencer/Private/SSequencerDebugVisualizer.cpp
--- a/Engine/Source/Editor/Sequencer/Private/SSequencerDebugVisualizer.cpp
+++ b/Engine/Source/Editor/Sequencer/Private/SSequencerDebugVisualizer.cpp
@@ -1,6 +1,8 @@
 // Copyright Epic Games, Inc. All Rights Reserved.
 
 #include "SSequencerDebugVisualizer.h"
+#include "MVVM/Selection/Selection.h"
+#include "MVVM/ViewModels/SequencerEditorViewModel.h"
 #include "CommonMovieSceneTools.h"
 #include "Styling/AppStyle.h"
 #include "Widgets/SToolTip.h"
@@ -457,6 +459,7 @@ void SSequencerEntityComponentSystemDebugVisualizer::Refresh()
 bool SSequencerEntityComponentSystemDebugVisualizer::DoRefresh()
 {
 	using namespace UE::MovieScene;
+	using namespace UE::Sequencer;
 
  	TSharedPtr<FSequencer> Sequencer = WeakSequencer.Pin();
  	if (!Sequencer.IsValid())
@@ -484,9 +487,7 @@ bool SSequencerEntityComponentSystemDebugVisualizer::DoRefresh()
 		return false;
 	}
 
-	TSet<TWeakObjectPtr<UMovieSceneSection>> SelectedSections = Sequencer->GetSelection().GetSelectedSections();
-	CachedSelection = SelectedSections;
-
+	CachedSelectionSerialNumber = Sequencer->GetViewModel()->GetSelection()->GetSerialNumber();
 	CachedSignature = CompiledDataManager->GetEntryRef(CompiledDataID).CompiledSignature;
 
  	const FFrameRate SequenceResolution = ActiveMovieScene->GetTickResolution();
@@ -498,9 +499,9 @@ bool SSequencerEntityComponentSystemDebugVisualizer::DoRefresh()
 		SectionToWidget.Add(Child->GetSection(), Child);
 	}
 
-	for (TWeakObjectPtr<UMovieSceneSection> WeakSection : SelectedSections)
+	for (TViewModelPtr<FSectionModel> SectionModel : Sequencer->GetViewModel()->GetSelection()->TrackArea.Filter<FSectionModel>())
 	{
-		UMovieSceneSection* Section = WeakSection.Get();
+		UMovieSceneSection* Section = SectionModel->GetSection();
 		if (Section != nullptr)
 		{
 			TSharedRef<SSequencerEntityComponentSystemDebugSlot>* ExistingChild = SectionToWidget.Find(Section);
@@ -575,11 +576,7 @@ void SSequencerEntityComponentSystemDebugVisualizer::Tick( const FGeometry& Allo
  	TSharedPtr<FSequencer> Sequencer = WeakSequencer.Pin();
  	if (Sequencer.IsValid())
 	{
-		TSet<TWeakObjectPtr<UMovieSceneSection>> SelectedSections = Sequencer->GetSelection().GetSelectedSections();
-		if (SelectedSections.Difference(CachedSelection).Num() > 0 || CachedSelection.Difference(SelectedSections).Num() > 0)
-		{
-			bSelectionChanged = true;
-		}
+		bSelectionChanged = CachedSelectionSerialNumber != Sequencer->GetViewModel()->GetSelection()->GetSerialNumber();
 
 		UMovieSceneSequence* ActiveSequence = Sequencer->GetFocusedMovieSceneSequence();
 		if (ActiveSequence != nullptr)
diff --git a/Engine/Source/Editor/Sequencer/Private/SSequencerDebugVisualizer.h b/Engine/Source/Editor/Sequencer/Private/SSequencerDebugVisualizer.h
--- a/Engine/Source/Editor/Sequencer/Private/SSequencerDebugVisualizer.h
+++ b/Engine/Source/Editor/Sequencer/Private/SSequencerDebugVisualizer.h
@@ -169,5 +169,5 @@ private:
  	FGuid CachedSignature;
 
 	/** Last known selection */
-	TSet<TWeakObjectPtr<UMovieSceneSection>> CachedSelection;
+	uint32 CachedSelectionSerialNumber;
 };
diff --git a/Engine/Source/Editor/Sequencer/Private/SSequencerGroupManager.cpp b/Engine/Source/Editor/Sequencer/Private/SSequencerGroupManager.cpp
--- a/Engine/Source/Editor/Sequencer/Private/SSequencerGroupManager.cpp
+++ b/Engine/Source/Editor/Sequencer/Private/SSequencerGroupManager.cpp
@@ -2,8 +2,10 @@
 
 #include "SSequencerGroupManager.h"
 
+#include "MVVM/ViewModels/SequencerEditorViewModel.h"
 #include "MVVM/Extensions/IOutlinerExtension.h"
 #include "MVVM/Extensions/IGroupableExtension.h"
+#include "MVVM/Selection/Selection.h"
 #include "Sequencer.h"
 #include "MovieSceneSequence.h"
 #include "MovieScene.h"
@@ -283,12 +285,7 @@ private:
 				if (SequencerGroupTreeNode->GetType() == FSequencerNodeGroupTreeNode::Type::GroupNode)
 				{
 					TSharedPtr<FSequencerNodeGroupNode> NodeGroupNode = StaticCastSharedPtr<FSequencerNodeGroupNode>(SequencerGroupTreeNode);
-					TArray<TWeakPtr<FViewModel>> DraggedNodes;
-					for (const TWeakViewModelPtr<IOutlinerExtension>& DraggedNode : DragDropOp->GetDraggedViewModels())
-					{
-						DraggedNodes.Add(DraggedNode);
-					}
-					Sequencer->AddNodesToExistingNodeGroup(DraggedNodes, NodeGroupNode->Group);
+					Sequencer->AddNodesToExistingNodeGroup(DragDropOp->GetDraggedViewModels(), NodeGroupNode->Group);
 
 					return FReply::Handled();
 				}
@@ -661,20 +658,12 @@ void SSequencerGroupManager::SelectItemsSelectedInSequencer()
 		return;
 	}
 
-	const TSet<TWeakPtr<FViewModel>>& SelectedModels = Sequencer->GetSelection().GetSelectedOutlinerItems();
-
 	TStringBuilder<128> TempString;
 
 	// Build a list of the nodepaths that we want to consider for selection
 	TSet<FString> NodesPathsToSelect;
-	for (TWeakPtr<FViewModel> WeakModel : SelectedModels)
+	for (FViewModelPtr Model : Sequencer->GetViewModel()->GetSelection()->Outliner)
 	{
-		TSharedPtr<FViewModel> Model = WeakModel.Pin();
-		if (!Model)
-		{
-			continue;
-		}
-
 		TViewModelPtr<IGroupableExtension> Groupable = Model->FindAncestorOfType<IGroupableExtension>(true);
 		if (Groupable)
 		{
diff --git a/Engine/Source/Editor/Sequencer/Private/SSequencerSection.cpp b/Engine/Source/Editor/Sequencer/Private/SSequencerSection.cpp
--- a/Engine/Source/Editor/Sequencer/Private/SSequencerSection.cpp
+++ b/Engine/Source/Editor/Sequencer/Private/SSequencerSection.cpp
@@ -12,6 +12,7 @@
 #include "MVVM/Views/STrackLane.h"
 #include "MVVM/Views/SChannelView.h"
 #include "MVVM/Extensions/IObjectBindingExtension.h"
+#include "MVVM/Selection/Selection.h"
 #include "AnimatedRange.h"
 #include "Rendering/DrawElements.h"
 #include "Styling/AppStyle.h"
@@ -331,7 +332,7 @@ struct FSequencerSectionPainterImpl : FSequencerSectionPainter
 
 	void CalculateSelectionColor()
 	{
-		FSequencerSelection& Selection = Sequencer.GetSelection();
+		FSequencerSelection& Selection = *Sequencer.GetViewModel()->GetSelection();
 		FSequencerSelectionPreview& SelectionPreview = Sequencer.GetSelectionPreview();
 
 		ESelectionPreviewState SelectionPreviewState = SelectionPreview.GetSelectionState(SectionWidget.WeakSectionModel);
@@ -342,7 +343,7 @@ struct FSequencerSectionPainterImpl : FSequencerSectionPainter
 			return;
 		}
 		
-		if (SelectionPreviewState == ESelectionPreviewState::Undefined && !Selection.IsSelected(SectionModel))
+		if (SelectionPreviewState == ESelectionPreviewState::Undefined && !Selection.TrackArea.IsSelected(SectionModel))
 		{
 			// No preview selection for this section, and it's not selected
 			return;
@@ -1338,7 +1339,7 @@ int32 SSequencerSection::OnPaint( const FPaintArgs& Args, const FGeometry& Allot
 		}
 	}
 
-	Painter.bIsSelected = GetSequencer().GetSelection().IsSelected(SectionModel);
+	Painter.bIsSelected = GetSequencer().GetSelection().TrackArea.IsSelected(SectionModel);
 
 	// Ask the interface to draw the section
 	LayerId = SectionInterface->OnPaintSection(Painter);
@@ -1360,7 +1361,7 @@ int32 SSequencerSection::OnPaint( const FPaintArgs& Args, const FGeometry& Allot
 	// --------------------------------------------
 	// Draw section selection tint
 	const float SectionThrobValue = GetSectionSelectionThrobValue();
-	if (SectionThrobValue != 0.f && GetSequencer().GetSelection().IsSelected(SectionModel))
+	if (SectionThrobValue != 0.f && GetSequencer().GetSelection().TrackArea.IsSelected(SectionModel))
 	{
 		static const FName BackgroundTrackTintBrushName("Sequencer.Section.BackgroundTint");
 
diff --git a/Engine/Source/Editor/Sequencer/Private/SSequencerTreeFilterStatusBar.cpp b/Engine/Source/Editor/Sequencer/Private/SSequencerTreeFilterStatusBar.cpp
--- a/Engine/Source/Editor/Sequencer/Private/SSequencerTreeFilterStatusBar.cpp
+++ b/Engine/Source/Editor/Sequencer/Private/SSequencerTreeFilterStatusBar.cpp
@@ -5,6 +5,8 @@
 
 #include "MovieScene.h"
 #include "MovieSceneSequence.h"
+#include "MVVM/Selection/Selection.h"
+#include "MVVM/ViewModels/SequencerEditorViewModel.h"
 
 #include "Algo/Count.h"
 #include "Widgets/SBoxPanel.h"
@@ -82,7 +84,7 @@ void SSequencerTreeFilterStatusBar::UpdateText()
 	FLinearColor NewColor = FLinearColor::White;
 
 	const TSharedRef<FSequencerNodeTree> NodeTree = Sequencer->GetNodeTree();
-	const TSet<TWeakPtr<FViewModel>>& SelectedOutlinerItems = Sequencer->GetSelection().GetSelectedOutlinerItems();
+	const FOutlinerSelection& SelectedOutlinerItems = Sequencer->GetViewModel()->GetSelection()->Outliner;
 
 	FFormatNamedArguments NamedArgs;
 	NamedArgs.Add("Total", NodeTree->GetTotalDisplayNodeCount());
