Sequencer: New selection API

diff --git a/Engine/Source/Editor/Sequencer/Private/Sequencer.cpp b/Engine/Source/Editor/Sequencer/Private/Sequencer.cpp
--- a/Engine/Source/Editor/Sequencer/Private/Sequencer.cpp
+++ b/Engine/Source/Editor/Sequencer/Private/Sequencer.cpp
@@ -6,6 +6,7 @@
 #include "Engine/World.h"
 #include "Algo/RemoveIf.h"
 #include "MVVM/SharedViewModelData.h"
+#include "MVVM/Selection/Selection.h"
 #include "MVVM/ViewModels/SequenceModel.h"
 #include "MVVM/ViewModels/SectionModel.h"
 #include "MVVM/ViewModels/FolderModel.h"
@@ -18,13 +19,15 @@
 #include "MVVM/ViewModels/SequencerOutlinerViewModel.h"
 #include "MVVM/ViewModels/SequencerTrackAreaViewModel.h"
 #include "MVVM/ViewModels/SequencerEditorViewModel.h"
+#include "MVVM/ViewModels/PossessableModel.h"
 #include "MVVM/Views/SOutlinerView.h"
-#include "MVVM/Views/IOutlinerSelectionHandler.h"
 #include "MVVM/Extensions/IDeletableExtension.h"
 #include "MVVM/CurveEditorExtension.h"
 #include "MVVM/CurveEditorIntegrationExtension.h"
 #include "MVVM/ObjectBindingModelStorageExtension.h"
 #include "MVVM/SectionModelStorageExtension.h"
+#include "MVVM/TrackModelStorageExtension.h"
+#include "MVVM/FolderModelStorageExtension.h"
 #include "Camera/PlayerCameraManager.h"
 #include "Misc/MessageDialog.h"
 #include "Misc/FeedbackContext.h"
@@ -187,87 +190,6 @@ static TAutoConsoleVariable<bool> CVarTimeUndo(
 	false,
 	TEXT("Enable/disable ability to undo time when making changes"));
 
-class FSequencerOutlinerSelectionHandler : 
-	public UE::Sequencer::IOutlinerSelectionHandler
-{
-public:
-
-	bool bForwardSelectionToView;
-	/** True if we are waiting for menus to be closed to synchronize external selection. */
-	bool bPendingSelectionChangedNotification;
-
-	FSequencerOutlinerSelectionHandler(FSequencerSelection& InSelection)
-		: Selection(InSelection)
-	{
-		bForwardSelectionToView = true;
-		bPendingSelectionChangedNotification = false;
-		Selection.GetOnOutlinerNodeSelectionChanged().AddRaw(this, &FSequencerOutlinerSelectionHandler::OnSelectedOutlinerNodesChanged);
-	}
-
-	void SetTreeViews(TSharedPtr<UE::Sequencer::SOutlinerView> InTreeView, TSharedPtr<UE::Sequencer::SOutlinerView> InPinnedTreeView)
-	{
-		WeakTreeView = InTreeView;
-		WeakPinnedTreeView = InPinnedTreeView;
-	}
-
-	void OnSelectedOutlinerNodesChanged()
-	{
-		using namespace UE::Sequencer;
-
-		if (bForwardSelectionToView)
-		{
-			TGuardValue<bool> Guard(bForwardSelectionToView, false);
-
-			// Called when our selection has changed programatically from code
-			TSharedPtr<SOutlinerView> TreeView = WeakTreeView.Pin();
-			if (ensure(TreeView))
-			{
-				TreeView->ForceSetSelectedItems(Selection.GetSelectedOutlinerItems());
-			}
-
-			TSharedPtr<SOutlinerView> PinnedTreeView = WeakPinnedTreeView.Pin();
-			if (ensure(PinnedTreeView))
-			{
-				PinnedTreeView->ForceSetSelectedItems(Selection.GetSelectedOutlinerItems());
-			}
-		}
-	}
-
-	void SelectOutlinerItems(const TArrayView<UE::Sequencer::TWeakViewModelPtr<UE::Sequencer::IOutlinerExtension>>& Items, bool bRightMouseButtonDown) override
-	{
-		using namespace UE::Sequencer;
-
-		TGuardValue<bool> Guard(bForwardSelectionToView, false);
-
-		Selection.SuspendBroadcast();
-		Selection.EmptySelectedOutlinerNodes();
-		for (const TWeakViewModelPtr<IOutlinerExtension>& Item : Items)
-		{
-			if (TViewModelPtr<IOutlinerExtension> ItemPtr = Item.Pin())
-			{
-				Selection.AddToOutlinerSelection(ItemPtr);
-			}
-		}
-
-		Selection.ResumeBroadcast();
-
-		if (bRightMouseButtonDown)
-		{
-			bPendingSelectionChangedNotification = true;
-		}
-		else
-		{
-			Selection.GetOnOutlinerNodeSelectionChanged().Broadcast();
-		}
-	}
-
-private:
-
-	FSequencerSelection& Selection;
-	TWeakPtr<UE::Sequencer::SOutlinerView> WeakTreeView;
-	TWeakPtr<UE::Sequencer::SOutlinerView> WeakPinnedTreeView;
-};
-
 namespace UE
 {
 	namespace Sequencer
@@ -544,6 +466,8 @@ void FSequencer::InitSequencer(const FSequencerInitParams& InitParams, const TSh
 	ViewModel->InitializeEditor();
 	ViewModel->SetSequence(InitParams.RootSequence);
 
+	ViewModel->GetSelection()->OnChanged.AddRaw(this, &FSequencer::OnSelectionChanged);
+
 	FSequencerOutlinerViewModel* OutlinerViewModel = ViewModel->GetOutliner()->CastThisChecked<FSequencerOutlinerViewModel>();
 	{
 		OutlinerViewModel->OnGetAddMenuContent = InitParams.ViewParams.OnGetAddMenuContent;
@@ -551,9 +475,6 @@ void FSequencer::InitSequencer(const FSequencerInitParams& InitParams, const TSh
 	}
 
 	NodeTree->SetRootNode(ViewModel->GetRootModel());
-	Selection.SetRootModel(ViewModel->GetRootModel());
-
-	SelectionHandler = MakeShared<FSequencerOutlinerSelectionHandler>(Selection);
 
 	ResetTimeController();
 
@@ -603,11 +524,8 @@ void FSequencer::InitSequencer(const FSequencerInitParams& InitParams, const TSh
 		.OnInitToolMenuContext(InitParams.ViewParams.OnInitToolMenuContext)
 		.AddMenuExtender(InitParams.ViewParams.AddMenuExtender)
 		.ToolbarExtender(InitParams.ViewParams.ToolbarExtender)
-		.SelectionHandler(SelectionHandler)
 		.ShowPlaybackRangeInTimeSlider(InitParams.ViewParams.bShowPlaybackRangeInTimeSlider);
 
-	SelectionHandler->SetTreeViews(SequencerWidget->GetTreeView(), SequencerWidget->GetPinnedTreeView());
-
 	if (GetHostCapabilities().bSupportsCurveEditor)
 	{
 		FCurveEditorExtension* CurveEditorExtension = ViewModel->CastDynamic<FCurveEditorExtension>();
@@ -701,10 +619,6 @@ FSequencer::FSequencer()
 	, bNeedsInvalidateCachedData(false)
 	, bHasPreAnimatedInfo(false)
 {
-	Selection.GetOnOutlinerNodeSelectionChanged().AddRaw(this, &FSequencer::OnSelectedOutlinerNodesChanged);
-	Selection.GetOnKeySelectionChanged().AddRaw(this, &FSequencer::OnSelectedOutlinerNodesChanged);
-	Selection.GetOnOutlinerNodeSelectionChangedObjectGuids().AddRaw(this, &FSequencer::OnSelectedOutlinerNodesChanged);
-
 	// Exposes the sequencer and curve editor command lists to subscribers from other systems
 	FInputBindingManager::Get().RegisterCommandList(FSequencerCommands::Get().GetContextName(), SequencerCommandBindings);
 	FInputBindingManager::Get().RegisterCommandList(FCurveEditorCommands::Get().GetContextName(), CurveEditorSharedBindings);
@@ -731,8 +645,6 @@ FSequencer::~FSequencer()
 
 	AcquiredResources.Release();
 
-	SelectionHandler.Reset();
-
 	ViewModel.Reset();
 	SequencerWidget.Reset();
 
@@ -813,12 +725,6 @@ void FSequencer::Tick(float InDeltaTime)
 
 	if (!FSlateApplication::Get().AnyMenusVisible())
 	{
-		if (SelectionHandler->bPendingSelectionChangedNotification)
-		{
-			SelectionHandler->bPendingSelectionChangedNotification = false;
-			HandleSelectedOutlinerNodesChanged();
-		}
-
 		if (bEnableRefCountCheck)
 		{
 			const int32 SequencerRefCount = AsShared().GetSharedReferenceCount() - 1;
@@ -846,8 +752,6 @@ void FSequencer::Tick(float InDeltaTime)
 	UMovieSceneSequence* RootSequencePtr = RootSequence.Get();
 	ObjectBindingTagCache->ConditionalUpdate(RootSequencePtr);
 
-	Selection.Tick();
-
 	UpdateCachedPlaybackContextAndClient();
 
 	{
@@ -1273,7 +1177,7 @@ void FSequencer::FocusSequenceInstance(UMovieSceneSubSection& InSubSection)
 	bGlobalMarkedFramesCached = false;
 }
 
-TSharedPtr<UE::Sequencer::FEditorViewModel> FSequencer::GetViewModel() const
+TSharedPtr<UE::Sequencer::FSequencerEditorViewModel> FSequencer::GetViewModel() const
 {
 	return ViewModel;
 }
@@ -1319,7 +1223,11 @@ void FSequencer::GetKeysFromSelection(TUniquePtr<FSequencerKeyCollection>& KeyCo
 	}
 
 	TArray<TSharedRef<FViewModel>> SelectedItems;
-	Selection.GetSelectedOutlinerItems(SelectedItems);
+	SelectedItems.Reserve(ViewModel->GetSelection()->Outliner.Num());
+	for (FViewModelPtr SelectedItem : ViewModel->GetSelection()->Outliner)
+	{
+		SelectedItems.Add(SelectedItem.AsModel().ToSharedRef());
+	}
 
 	int64 TotalMaxSeconds = static_cast<int64>(TNumericLimits<int32>::Max() / GetFocusedTickResolution().AsDecimal());
 
@@ -1616,16 +1524,16 @@ void FSequencer::GetConstructionScriptActors(UMovieScene* MovieScene, FMovieScen
 	}
 }
 
-void FSequencer::DeleteSections(const TSet<TWeakObjectPtr<UMovieSceneSection>>& Sections)
+void FSequencer::DeleteSections(const TSet<UMovieSceneSection*>& Sections)
 {
 	UMovieScene* MovieScene = GetFocusedMovieSceneSequence()->GetMovieScene();
 	bool bAnythingRemoved = false;
 
 	FScopedTransaction DeleteSectionTransaction( NSLOCTEXT("Sequencer", "DeleteSection_Transaction", "Delete Section") );
 
-	for (const auto& Section : Sections)
+	for (UMovieSceneSection* Section : Sections)
 	{
-		if (!Section.IsValid() || Section->IsLocked())
+		if (!Section || Section->IsLocked())
 		{
 			continue;
 		}
@@ -1648,7 +1556,7 @@ void FSequencer::DeleteSections(const TSet<TWeakObjectPtr<UMovieSceneSection>>&
 		NotifyMovieSceneDataChanged( EMovieSceneDataChangeType::MovieSceneStructureItemRemoved );
 	}
 
-	Selection.EmptySelectedTrackAreaItems();
+	ViewModel->GetSelection()->TrackArea.Empty();
 }
 
 
@@ -1657,7 +1565,7 @@ void FSequencer::DeleteSelectedKeys()
 	FScopedTransaction DeleteKeysTransaction( NSLOCTEXT("Sequencer", "DeleteSelectedKeys_Transaction", "Delete Selected Keys") );
 	bool bAnythingRemoved = false;
 
-	FSelectedKeysByChannel KeysByChannel(Selection.GetSelectedKeys().Array());
+	FSelectedKeysByChannel KeysByChannel(ViewModel->GetSelection()->KeySelection);
 	TSet<UMovieSceneSection*> ModifiedSections;
 
 	for (const FSelectedChannelInfo& ChannelInfo : KeysByChannel.SelectedChannels)
@@ -1685,15 +1593,14 @@ void FSequencer::DeleteSelectedKeys()
 	{
 		NotifyMovieSceneDataChanged( EMovieSceneDataChangeType::TrackValueChanged );
 
-		Selection.EmptySelectedKeys();
+		ViewModel->GetSelection()->KeySelection.Empty();
 	}
 }
 
 
 void FSequencer::SetInterpTangentMode(ERichCurveInterpMode InterpMode, ERichCurveTangentMode TangentMode)
 {
-	TArray<FSequencerSelectedKey> SelectedKeysArray = Selection.GetSelectedKeys().Array();
-	if (SelectedKeysArray.Num() == 0)
+	if (ViewModel->GetSelection()->KeySelection.Num() == 0)
 	{
 		return;
 	}
@@ -1701,7 +1608,7 @@ void FSequencer::SetInterpTangentMode(ERichCurveInterpMode InterpMode, ERichCurv
 	FScopedTransaction SetInterpTangentModeTransaction(NSLOCTEXT("Sequencer", "SetInterpTangentMode_Transaction", "Set Interpolation and Tangent Mode"));
 	bool bAnythingChanged = false;
 
-	FSelectedKeysByChannel KeysByChannel(SelectedKeysArray);
+	FSelectedKeysByChannel KeysByChannel(ViewModel->GetSelection()->KeySelection);
 	TSet<UMovieSceneSection*> ModifiedSections;
 
 	const FName FloatChannelTypeName = FMovieSceneFloatChannel::StaticStruct()->GetFName();
@@ -1777,9 +1684,7 @@ void FSequencer::SetInterpTangentMode(ERichCurveInterpMode InterpMode, ERichCurv
 void FSequencer::ToggleInterpTangentWeightMode()
 {
 	// @todo: sequencer-timecode: move this float-specific logic elsewhere to make it extensible for any channel type
-
-	TArray<FSequencerSelectedKey> SelectedKeysArray = Selection.GetSelectedKeys().Array();
-	if (SelectedKeysArray.Num() == 0)
+	if (ViewModel->GetSelection()->KeySelection.Num() == 0)
 	{
 		return;
 	}
@@ -1787,7 +1692,7 @@ void FSequencer::ToggleInterpTangentWeightMode()
 	FScopedTransaction SetInterpTangentWeightModeTransaction(NSLOCTEXT("Sequencer", "ToggleInterpTangentWeightMode_Transaction", "Toggle Tangent Weight Mode"));
 	bool bAnythingChanged = false;
 
-	FSelectedKeysByChannel KeysByChannel(SelectedKeysArray);
+	FSelectedKeysByChannel KeysByChannel(ViewModel->GetSelection()->KeySelection);
 	TSet<UMovieSceneSection*> ModifiedSections;
 
 	const FName FloatChannelTypeName = FMovieSceneFloatChannel::StaticStruct()->GetFName();
@@ -1913,7 +1818,7 @@ void FSequencer::SnapToFrame()
 	FScopedTransaction SnapToFrameTransaction(NSLOCTEXT("Sequencer", "SnapToFrame_Transaction", "Snap Selected Keys to Frame"));
 	bool bAnythingChanged = false;
 
-	FSelectedKeysByChannel KeysByChannel(Selection.GetSelectedKeys().Array());
+	FSelectedKeysByChannel KeysByChannel(ViewModel->GetSelection()->KeySelection);
 	TSet<UMovieSceneSection*> ModifiedSections;
 
 	TArray<FFrameNumber> KeyTimesScratch;
@@ -1960,7 +1865,7 @@ void FSequencer::SnapToFrame()
 
 bool FSequencer::CanSnapToFrame() const
 {
-	const bool bKeysSelected = Selection.GetSelectedKeys().Num() > 0;
+	const bool bKeysSelected = ViewModel->GetSelection()->KeySelection.Num() > 0;
 
 	return bKeysSelected;
 }
@@ -1970,12 +1875,11 @@ void FSequencer::TransformSelectedKeysAndSections(FFrameTime InDeltaTime, float
 	FScopedTransaction TransformKeysAndSectionsTransaction(NSLOCTEXT("Sequencer", "TransformKeysandSections_Transaction", "Transform Keys and Sections"));
 	bool bAnythingChanged = false;
 
-	TArray<FSequencerSelectedKey> SelectedKeysArray = Selection.GetSelectedKeys().Array();
-	TArray<TWeakObjectPtr<UMovieSceneSection>> SelectedSectionsArray = Selection.GetSelectedSections().Array();
+	TSet<UMovieSceneSection*> SelectedSections = ViewModel->GetSelection()->GetSelectedSections();
 
 	const FFrameTime OriginTime = GetLocalTime().Time;
 
-	FSelectedKeysByChannel KeysByChannel(SelectedKeysArray);
+	FSelectedKeysByChannel KeysByChannel(ViewModel->GetSelection()->KeySelection);
 	TMap<UMovieSceneSection*, TRange<FFrameNumber>> SectionToNewBounds;
 
 	TArray<FFrameNumber> KeyTimesScratch;
@@ -1988,7 +1892,7 @@ void FSequencer::TransformSelectedKeysAndSections(FFrameTime InDeltaTime, float
 			if (Channel)
 			{
 				// Skip any channels whose section is already selected because they'll be handled below (moving the section and the keys together)
-				if (SelectedSectionsArray.Contains(ChannelInfo.OwningSection))
+				if (SelectedSections.Contains(ChannelInfo.OwningSection))
 				{
 					continue;
 				}
@@ -2041,14 +1945,8 @@ void FSequencer::TransformSelectedKeysAndSections(FFrameTime InDeltaTime, float
 		}
 
 		// Dilate the sections
-		for (TWeakObjectPtr<UMovieSceneSection> WeakSection : SelectedSectionsArray)
+		for (UMovieSceneSection* Section : SelectedSections)
 		{
-			UMovieSceneSection* Section = WeakSection.Get();
-			if (!Section)
-			{
-				continue;
-			}
-
 			TRangeBound<FFrameNumber> LowerBound = Section->GetRange().GetLowerBound();
 			TRangeBound<FFrameNumber> UpperBound = Section->GetRange().GetUpperBound();
 
@@ -2160,11 +2058,13 @@ void FSequencer::StretchTime(FFrameTime InDeltaTime)
 
 	// Return state
 	FocusedMovieScene->SetSelectionRange(CachedSelectionRange);
-	Selection.Empty(); //todo restore key and section selection
+	ViewModel->GetSelection()->Empty(); //todo restore key and section selection
 }
 
 void FSequencer::ShrinkTime(FFrameTime InDeltaTime)
 {
+	using namespace UE::Sequencer;
+
 	// From the current time, find all the keys and sections to the right and move them by -InDeltaTime
 	UMovieScene* FocusedMovieScene = GetFocusedMovieSceneSequence()->GetMovieScene();
 	if (!FocusedMovieScene)
@@ -2200,27 +2100,28 @@ void FSequencer::ShrinkTime(FFrameTime InDeltaTime)
 	FocusedMovieScene->SetSelectionRange(CheckRange);
 	SelectInSelectionRange(true, true);
 
-	if (Selection.GetSelectedKeys().Num() > 0)
+	TSharedPtr<FSequencerSelection> Selection = ViewModel->GetSelection();
+	if (Selection->KeySelection.Num() > 0)
 	{
-		FNotificationInfo Info(FText::Format(NSLOCTEXT("Sequencer", "ShrinkTimeFailedKeys", "Shrink failed. There are {0} keys in between"), Selection.GetSelectedKeys().Num()));
+		FNotificationInfo Info(FText::Format(NSLOCTEXT("Sequencer", "ShrinkTimeFailedKeys", "Shrink failed. There are {0} keys in between"), Selection->KeySelection.Num()));
 		Info.ExpireDuration = 5.0f;
 		FSlateNotificationManager::Get().AddNotification(Info)->SetCompletionState(SNotificationItem::CS_Fail);
 
 		// Return state
 		FocusedMovieScene->SetSelectionRange(CachedSelectionRange);
-		Selection.Empty(); //todo restore key and section selection
+		Selection->Empty(); //todo restore key and section selection
 		return;
 	}
 
-	if (Selection.GetSelectedSections().Num() > 0)
+	if (Selection->GetSelectedSections().Num() > 0)
 	{
-		FNotificationInfo Info(FText::Format(NSLOCTEXT("Sequencer", "ShrinkTimeFailedSections", "Shrink failed. There are {0} sections in between"), Selection.GetSelectedSections().Num()));
+		FNotificationInfo Info(FText::Format(NSLOCTEXT("Sequencer", "ShrinkTimeFailedSections", "Shrink failed. There are {0} sections in between"), Selection->GetSelectedSections().Num()));
 		Info.ExpireDuration = 5.0f;
 		FSlateNotificationManager::Get().AddNotification(Info)->SetCompletionState(SNotificationItem::CS_Fail);
 
 		// Return state
 		FocusedMovieScene->SetSelectionRange(CachedSelectionRange);
-		Selection.Empty(); //todo restore key and section selection
+		Selection->Empty(); //todo restore key and section selection
 		return;
 	}
 
@@ -2243,7 +2144,7 @@ void FSequencer::ShrinkTime(FFrameTime InDeltaTime)
 
 	// Return state
 	FocusedMovieScene->SetSelectionRange(CachedSelectionRange);
-	Selection.Empty(); //todo restore key and section selection
+	Selection->Empty(); //todo restore key and section selection
 }
 
 bool FSequencer::CanAddTransformKeysForSelectedObjects() const
@@ -2465,9 +2366,6 @@ void FSequencer::RefreshTree()
 		SequencerWidget->UpdateLayoutTree();
 		bNeedTreeRefresh = false;
 		OnTreeViewChangedDelegate.Broadcast();
-
-		// Force a broadcast of selection changed after the tree view has been updated, in the event that selection was suppressed while the tree was refreshing
-		Selection.Tick();
 	}
 }
 
@@ -2588,6 +2486,8 @@ void FSequencer::SelectInSelectionRange(const TSharedPtr<UE::Sequencer::FViewMod
 		return;
 	}
 
+	TSharedPtr<FSequencerSelection> Selection = ViewModel->GetSelection();
+
 	if (bSelectKeys)
 	{
 		TArray<FKeyHandle> HandlesScratch;
@@ -2604,7 +2504,7 @@ void FSequencer::SelectInSelectionRange(const TSharedPtr<UE::Sequencer::FViewMod
 
 				for (int32 Index = 0; Index < HandlesScratch.Num(); ++Index)
 				{
-					Selection.AddToSelection(FSequencerSelectedKey(*Section, Channel, HandlesScratch[Index]));
+					Selection->KeySelection.Select(Channel, HandlesScratch[Index]);
 				}
 			}
 		}
@@ -2617,7 +2517,7 @@ void FSequencer::SelectInSelectionRange(const TSharedPtr<UE::Sequencer::FViewMod
 			TRange<FFrameNumber> SectionRange = SectionModel->GetRange();
 			if (SectionRange.Overlaps(SelectionRange) && SectionRange.GetLowerBound().IsClosed() && SectionRange.GetUpperBound().IsClosed())
 			{
-				Selection.AddToSelection(SectionModel);
+				Selection->TrackArea.Select(SectionModel);
 			}
 		}
 	}
@@ -2641,16 +2541,18 @@ void FSequencer::SelectInSelectionRange(bool bSelectKeys, bool bSelectSections)
 	UMovieScene* MovieScene = Sequence->GetMovieScene();
 	TRange<FFrameNumber> SelectionRange = MovieScene->GetSelectionRange();
 
+	TSharedPtr<FSequencerSelection> Selection = ViewModel->GetSelection();
+
+	FSelectionEventSuppressor EventSuppressor = Selection->SuppressEvents();
+
 	// Don't empty all selection, just keys and sections
-	Selection.SuspendBroadcast();
-	Selection.EmptySelectedKeys();
-	Selection.EmptySelectedTrackAreaItems();
+	Selection->KeySelection.Empty();
+	Selection->TrackArea.Empty();
 
 	for (const TViewModelPtr<IOutlinerExtension>& DisplayNode : NodeTree->GetRootNodes())
 	{
 		SelectInSelectionRange(DisplayNode, SelectionRange, bSelectKeys, bSelectSections);
 	}
-	Selection.ResumeBroadcast();
 }
 
 void FSequencer::SelectForward()
@@ -2661,34 +2563,30 @@ void FSequencer::SelectForward()
 	FFrameNumber CurrentFrame = GetLocalTime().ConvertTo(TickResolution).CeilToFrame();
 	TRange<FFrameNumber> SelectionRange(CurrentFrame, TNumericLimits<FFrameNumber>::Max());
 
-	TSet<TWeakPtr<FViewModel>> SelectedItems = Selection.GetNodesWithSelectedKeysOrSections();
-	if (SelectedItems.Num() == 0)
-	{
-		SelectedItems = Selection.GetSelectedOutlinerItems();
-	}
-	if (SelectedItems.Num() == 0)
+	TSharedPtr<FSequencerSelection> Selection = ViewModel->GetSelection();
+
+	FSelectionEventSuppressor EventSuppressor = Selection->SuppressEvents();
+
+	Selection->KeySelection.Empty();
+	Selection->TrackArea.Empty();
+
+	if (Selection->GetNodesWithSelectedKeysOrSections().Num() != 0)
 	{
-		TArray<TSharedPtr<FViewModel>> AllNodes;
-		NodeTree->GetAllNodes(AllNodes);
-		for (TSharedPtr<FViewModel> Node : AllNodes)
+		for (TViewModelPtr<IOutlinerExtension> Item : Selection->IterateIndirectOutlinerSelection())
 		{
-			SelectedItems.Add(Node);
+			SelectInSelectionRange(Item, SelectionRange, true, true);
 		}
 	}
-
-	if (SelectedItems.Num() > 0)
+	else if (Selection->Outliner.Num() != 0)
 	{
-		Selection.SuspendBroadcast();
-		Selection.EmptySelectedKeys();
-		Selection.EmptySelectedTrackAreaItems();
-		for (const TWeakPtr<FViewModel>& WeakItem : SelectedItems)
+		for (TViewModelPtr<IOutlinerExtension> Item : Selection->Outliner)
 		{
-			if (TSharedPtr<FViewModel> Item = WeakItem.Pin())
-			{
-				SelectInSelectionRange(Item, SelectionRange, true, true);
-			}
+			SelectInSelectionRange(Item, SelectionRange, true, true);
 		}
-		Selection.ResumeBroadcast();
+	}
+	else
+	{
+		SelectInSelectionRange(ViewModel->GetRootModel(), SelectionRange, true, true);
 	}
 }
 
@@ -2701,34 +2599,30 @@ void FSequencer::SelectBackward()
 	FFrameNumber CurrentFrame = GetLocalTime().ConvertTo(TickResolution).CeilToFrame();
 	TRange<FFrameNumber> SelectionRange(TNumericLimits<FFrameNumber>::Min(), CurrentFrame);
 
-	TSet<TWeakPtr<FViewModel>> SelectedItems = Selection.GetNodesWithSelectedKeysOrSections();
-	if (SelectedItems.Num() == 0)
-	{
-		SelectedItems = Selection.GetSelectedOutlinerItems();
-	}
-	if (SelectedItems.Num() == 0)
+	TSharedPtr<FSequencerSelection> Selection = ViewModel->GetSelection();
+
+	FSelectionEventSuppressor EventSuppressor = Selection->SuppressEvents();
+
+	Selection->KeySelection.Empty();
+	Selection->TrackArea.Empty();
+
+	if (Selection->GetNodesWithSelectedKeysOrSections().Num() != 0)
 	{
-		TArray<TSharedPtr<FViewModel>> AllNodes;
-		NodeTree->GetAllNodes(AllNodes);
-		for (const TSharedPtr<FViewModel>& Node : AllNodes)
+		for (TViewModelPtr<IOutlinerExtension> Item : Selection->IterateIndirectOutlinerSelection())
 		{
-			SelectedItems.Add(Node);
+			SelectInSelectionRange(Item, SelectionRange, true, true);
 		}
 	}
-
-	if (SelectedItems.Num() > 0)
+	else if (Selection->Outliner.Num() != 0)
 	{
-		Selection.SuspendBroadcast();
-		Selection.EmptySelectedKeys();
-		Selection.EmptySelectedTrackAreaItems();
-		for (const TWeakPtr<FViewModel>& WeakItem : SelectedItems)
+		for (TViewModelPtr<IOutlinerExtension> Item : Selection->Outliner)
 		{
-			if (TSharedPtr<FViewModel> Item = WeakItem.Pin())
-			{
-				SelectInSelectionRange(Item, SelectionRange, true, true);
-			}
+			SelectInSelectionRange(Item, SelectionRange, true, true);
 		}
-		Selection.ResumeBroadcast();
+	}
+	else
+	{
+		SelectInSelectionRange(ViewModel->GetRootModel(), SelectionRange, true, true);
 	}
 }
 
@@ -4353,7 +4247,7 @@ void FSequencer::ResetPerMovieSceneData()
 	using namespace UE::Sequencer;
 
 	//@todo Sequencer - We may want to preserve selections when moving between movie scenes
-	Selection.Empty();
+	ViewModel->GetSelection()->Empty();
 
 	// This will reinitialize all extensions to rebuild the view-model hierarchy
 	UMovieSceneSequence* CurrentSequence = GetFocusedMovieSceneSequence();
@@ -4378,7 +4272,7 @@ void FSequencer::RefreshUI()
 {
 	using namespace UE::Sequencer;
 
-	Selection.Empty();
+	ViewModel->GetSelection()->Empty();
 
 	UMovieSceneSequence* FocusedSequence = GetFocusedMovieSceneSequence();
 	FMovieSceneSequenceID SequenceID = GetFocusedTemplateID();
@@ -4889,7 +4783,7 @@ FReply FSequencer::SetPlaybackStart()
 
 FReply FSequencer::JumpToPreviousKey()
 {
-	if (Selection.GetSelectedOutlinerItems().Num())
+	if (ViewModel->GetSelection()->Outliner.Num())
 	{
 		GetKeysFromSelection(SelectedKeyCollection, SMALL_NUMBER);
 	}
@@ -4918,7 +4812,7 @@ FReply FSequencer::JumpToPreviousKey()
 
 FReply FSequencer::JumpToNextKey()
 {
-	if (Selection.GetSelectedOutlinerItems().Num())
+	if (ViewModel->GetSelection()->Outliner.Num())
 	{
 		GetKeysFromSelection(SelectedKeyCollection, SMALL_NUMBER);
 	}
@@ -6019,16 +5913,10 @@ TArray<FGuid> FSequencer::AddActors(const TArray<TWeakObjectPtr<AActor> >& InAct
 		// Check if a folder is selected so we can add the actors to the selected folder.
 		TArray<UMovieSceneFolder*> SelectedParentFolders;
 		FString NewNodePath;
-		if (Selection.GetSelectedOutlinerItems().Num() > 0)
+		if (ViewModel->GetSelection()->Outliner.Num() > 0)
 		{
-			for (TWeakPtr<FViewModel> WeakItem : Selection.GetSelectedOutlinerItems())
+			for (FViewModelPtr CurrentItem : ViewModel->GetSelection()->Outliner)
 			{
-				TSharedPtr<FViewModel> CurrentItem = WeakItem.Pin();
-				if (!CurrentItem)
-				{
-					continue;
-				}
-
 				if (TSharedPtr<FFolderModel> Folder = CurrentItem->FindAncestorOfType<FFolderModel>(true))
 				{
 					SelectedParentFolders.Add(Folder->GetFolder());
@@ -6097,13 +5985,15 @@ TArray<FGuid> FSequencer::AddActors(const TArray<TWeakObjectPtr<AActor> >& InAct
 	return PossessableGuids;
 }
 
-void FSequencer::OnSelectedOutlinerNodesChanged()
+void FSequencer::OnSelectionChanged()
 {
 	HandleSelectedOutlinerNodesChanged();
 }
 
 void FSequencer::HandleSelectedOutlinerNodesChanged()
 {
+	using namespace UE::Sequencer;
+
 	SynchronizeExternalSelectionWithSequencerSelection();
 
 	FSequencerEdMode* SequencerEdMode = (FSequencerEdMode*)(GLevelEditorModeTools().GetActiveMode(FSequencerEdMode::EM_SequencerMode));
@@ -6117,16 +6007,26 @@ void FSequencer::HandleSelectedOutlinerNodesChanged()
 		}
 	}
 
-	OnSelectionChangedObjectGuidsDelegate.Broadcast(Selection.GetBoundObjectsGuids());
-	OnSelectionChangedTracksDelegate.Broadcast(Selection.GetSelectedTracks());
+	TSet<UMovieSceneTrack*> SelectedTracks;
+	for (TViewModelPtr<ITrackExtension> TrackModel : ViewModel->GetSelection()->Outliner.Filter<ITrackExtension>())
+	{
+		if (UMovieSceneTrack* Track = TrackModel->GetTrack())
+		{
+			SelectedTracks.Add(Track);
+		}
+	}
+
 	TArray<UMovieSceneSection*> SelectedSections;
-	for (TWeakObjectPtr<UMovieSceneSection> SelectedSectionPtr : Selection.GetSelectedSections())
+	for (TViewModelPtr<FSectionModel> SectionModel : ViewModel->GetSelection()->TrackArea.Filter<FSectionModel>())
 	{
-		if (SelectedSectionPtr.IsValid())
+		if (UMovieSceneSection* Section = SectionModel->GetSection())
 		{
-			SelectedSections.Add(SelectedSectionPtr.Get());
+			SelectedSections.Add(Section);
 		}
 	}
+
+	OnSelectionChangedObjectGuidsDelegate.Broadcast(ViewModel->GetSelection()->GetBoundObjectsGuids());
+	OnSelectionChangedTracksDelegate.Broadcast(SelectedTracks.Array());
 	OnSelectionChangedSectionsDelegate.Broadcast(SelectedSections);
 }
 
@@ -6175,18 +6075,15 @@ void FSequencer::AddSelectedNodesToNewNodeGroup()
 		return;
 	}
 
-	const TSet<TWeakPtr<FViewModel>>& SelectedItems = GetSelection().GetSelectedOutlinerItems();
-	if (SelectedItems.Num() == 0)
+	if (ViewModel->GetSelection()->Outliner.Num() == 0)
 	{
 		return;
 	}
 
 	TSet<FString> NodesToAdd;
-	for (const TWeakPtr<FViewModel>& WeakItem : SelectedItems)
+	for (FViewModelPtr Item : ViewModel->GetSelection()->Outliner)
 	{
-		TSharedPtr<FViewModel>  Item      = WeakItem.Pin();
-		TSharedPtr<IGroupableExtension> Groupable = Item ? Item->FindAncestorOfType<IGroupableExtension>(true) : nullptr;
-
+		TSharedPtr<IGroupableExtension> Groupable = Item->FindAncestorOfType<IGroupableExtension>(true);
 		if (Groupable)
 		{
 			NodesToAdd.Add(IOutlinerExtension::GetPathName(Item));
@@ -6222,10 +6119,10 @@ void FSequencer::AddSelectedNodesToNewNodeGroup()
 
 void FSequencer::AddSelectedNodesToExistingNodeGroup(UMovieSceneNodeGroup* NodeGroup)
 {
-	AddNodesToExistingNodeGroup(GetSelection().GetSelectedOutlinerItems().Array(), NodeGroup);
+	AddNodesToExistingNodeGroup(ViewModel->GetSelection()->Outliner.GetSelected().Array(), NodeGroup);
 }
 
-void FSequencer::AddNodesToExistingNodeGroup(const TArray<TWeakPtr<UE::Sequencer::FViewModel>>& InItems, UMovieSceneNodeGroup* InNodeGroup)
+void FSequencer::AddNodesToExistingNodeGroup(TArrayView<const UE::Sequencer::TWeakViewModelPtr<UE::Sequencer::IOutlinerExtension>> InItems, UMovieSceneNodeGroup* InNodeGroup)
 {
 	using namespace UE::Sequencer;
 
@@ -6242,9 +6139,9 @@ void FSequencer::AddNodesToExistingNodeGroup(const TArray<TWeakPtr<UE::Sequencer
 	}
 
 	TSet<FString> NodesToAdd;
-	for (const TWeakPtr<FViewModel>& WeakItem : InItems)
+	for (const TWeakViewModelPtr<IOutlinerExtension>& WeakItem : InItems)
 	{
-		TSharedPtr<FViewModel>  Item      = WeakItem.Pin();
+		FViewModelPtr                   Item      = WeakItem.Pin();
 		TSharedPtr<IGroupableExtension> Groupable = Item ? Item->FindAncestorOfType<IGroupableExtension>(true) : nullptr;
 
 		if (Groupable)
@@ -6310,15 +6207,15 @@ void FSequencer::SynchronizeExternalSelectionWithSequencerSelection()
 	TSet<AActor*> SelectedSequencerActors;
 	TSet<UActorComponent*> SelectedSequencerComponents;
 
-	TSet<TWeakPtr<FViewModel>> Items = Selection.GetNodesWithSelectedKeysOrSections();
-	Items.Append(Selection.GetSelectedOutlinerItems());
+	TSet<TWeakViewModelPtr<IOutlinerExtension>> Items = ViewModel->GetSelection()->GetNodesWithSelectedKeysOrSections();
+	Items.Append(ViewModel->GetSelection()->Outliner.GetSelected());
 
 	const FName ControlRigEditModeModeName("EditMode.ControlRig");
 	const bool bHACK_ControlRigEditMode = GLevelEditorModeTools().GetActiveMode(ControlRigEditModeModeName) != nullptr;
 
-	for ( TWeakPtr<FViewModel> WeakItem : Items)
+	for ( TWeakViewModelPtr<IOutlinerExtension> WeakItem : Items)
 	{
-		TViewModelPtr<FViewModel> SelectedItem = WeakItem.Pin();
+		FViewModelPtr SelectedItem = WeakItem.Pin();
 		if (!SelectedItem)
 		{
 			continue;
@@ -6510,6 +6407,8 @@ void FSequencer::SynchronizeSequencerSelectionWithExternalSelection()
 	TArray<ASequencerKeyActor*> SelectedSequencerKeyActors;
 	ActorSelection->GetSelectedObjects<ASequencerKeyActor>(SelectedSequencerKeyActors);
 
+	TSharedPtr<FSequencerSelection> Selection = ViewModel->GetSelection();
+
 	FObjectBindingModelStorageExtension* ObjectModelStorage = ViewModel->GetRootModel()->CastDynamic<FObjectBindingModelStorageExtension>();
 	check(ObjectModelStorage);
 
@@ -6554,14 +6453,12 @@ void FSequencer::SynchronizeSequencerSelectionWithExternalSelection()
 			{
 				NodesToSelect.Add( ObjectModel );
 
-				if (bAllAlreadySelected && !Selection.IsSelected(ObjectModel))
+				if (bAllAlreadySelected && !Selection->Outliner.IsSelected(ObjectModel))
 				{
 					bool bAnyChildrenSelected = false;
-					TArray<TSharedPtr<FViewModel>> OutlinerChildren;
-					ObjectModel->GetDescendantsOfType(IOutlinerExtension::ID, OutlinerChildren);
-					for (TSharedPtr<FViewModel> Child : OutlinerChildren)
+					for (TViewModelPtr<IOutlinerExtension> Child : ObjectModel->GetDescendantsOfType<IOutlinerExtension>())
 					{
-						if (Selection.IsSelected(Child) || Selection.NodeHasSelectedKeysOrSections(Child))
+						if (Selection->Outliner.IsSelected(Child) || Selection->NodeHasSelectedKeysOrSections(Child))
 						{
 							bAnyChildrenSelected = true;
 							break;
@@ -6574,7 +6471,7 @@ void FSequencer::SynchronizeSequencerSelectionWithExternalSelection()
 					}
 				}
 			}
-			else if (Selection.IsSelected(ObjectModel))
+			else if (Selection->Outliner.IsSelected(ObjectModel))
 			{
 				bAllAlreadySelected = false;
 			}
@@ -6582,38 +6479,32 @@ void FSequencer::SynchronizeSequencerSelectionWithExternalSelection()
 	}
 	//Only test if none are selected if we are not transacting, otherwise it will clear out control rig's incorrectly.
 
-	if (!bAllAlreadySelected || (!GIsTransacting && (NodesToSelect.Num() == 0 && Selection.GetSelectedOutlinerItems().Num())))
+	if (!bAllAlreadySelected || (!GIsTransacting && (NodesToSelect.Num() == 0 && Selection->Outliner.Num())))
 	{
-		Selection.SuspendBroadcast();
-		Selection.EmptySelectedOutlinerNodes();
-		for (TSharedPtr<FViewModel> NodeToSelect : NodesToSelect)
+		FSelectionEventSuppressor EventSuppressor = Selection->SuppressEvents();
+
+		Selection->Outliner.Empty();
+		for (TSharedPtr<FObjectBindingModel> NodeToSelect : NodesToSelect)
 		{
-			Selection.AddToSelection( NodeToSelect );
+			Selection->Outliner.Select(NodeToSelect);
 		}
 		
 		TSharedPtr<SOutlinerView> TreeView = SequencerWidget->GetTreeView();
-		const TSet<TWeakPtr<FViewModel>>& OutlinerSelection = GetSelection().GetSelectedOutlinerItems();
-		if (OutlinerSelection.Num() == 1)
+		if (Selection->Outliner.Num() == 1)
 		{
-			for (TWeakPtr<FViewModel> WeakNode : OutlinerSelection)
+			for (TViewModelPtr<IOutlinerExtension> Node : Selection->Outliner)
 			{
-				if (TViewModelPtr<IOutlinerExtension> Node = CastViewModel<IOutlinerExtension>(WeakNode.Pin()))
+				for (TViewModelPtr<IOutlinerExtension> Parent : Node.AsModel()->GetAncestorsOfType<IOutlinerExtension>())
 				{
-					TViewModelPtr<IOutlinerExtension> Parent = Node.AsModel()->FindAncestorOfType<IOutlinerExtension>();
-					while (Parent.IsValid())
-					{
+					TreeView->SetItemExpansion(Parent, true);
 						TreeView->SetItemExpansion(Parent, true);
 						Parent = Parent.AsModel()->FindAncestorOfType<IOutlinerExtension>();
-					}
-
-					TreeView->RequestScrollIntoView(Node);
-					break;
 				}
+
+				TreeView->RequestScrollIntoView(Node);
+				break;
 			}
 		}
-
-		Selection.ResumeBroadcast();
-		Selection.GetOnOutlinerNodeSelectionChanged().Broadcast();
 	}
 }
 
@@ -6636,17 +6527,15 @@ void FSequencer::SelectNodesByPath(const TSet<FString>& NodePaths)
 	// nodes are not cleared and reselected, which can cause issues with the curve editor auto-fitting 
 	// based on selection.
 	bool bAllAlreadySelected = true;
-	const TSet<TWeakPtr<FViewModel>>& CurrentSelection = GetSelection().GetSelectedOutlinerItems();
+	const FOutlinerSelection& CurrentSelection = ViewModel->GetSelection()->Outliner;
 
-	TArray<TSharedPtr<FViewModel>> AllNodes;
-	NodeTree->GetAllNodes(AllNodes);
-	TSet<TSharedPtr<FViewModel>> NodesToSelect;
-	for (TSharedPtr<FViewModel> DisplayNode : AllNodes)
+	TSet<TViewModelPtr<IOutlinerExtension>> NodesToSelect;
+	for (TViewModelPtr<IOutlinerExtension> DisplayNode : ViewModel->GetRootModel()->GetDescendantsOfType<IOutlinerExtension>())
 	{
 		if (NodePaths.Contains(IOutlinerExtension::GetPathName(DisplayNode)))
 		{
 			NodesToSelect.Add(DisplayNode);
-			if (bAllAlreadySelected && !CurrentSelection.Contains(DisplayNode))
+			if (bAllAlreadySelected && !CurrentSelection.IsSelected(DisplayNode))
 			{
 				bAllAlreadySelected = false;
 			}
@@ -6655,33 +6544,29 @@ void FSequencer::SelectNodesByPath(const TSet<FString>& NodePaths)
 
 	if (!bAllAlreadySelected || (NodesToSelect.Num() != CurrentSelection.Num()))
 	{
-		Selection.SuspendBroadcast();
-		Selection.EmptySelectedOutlinerNodes();
-		for (TSharedPtr<FViewModel> NodeToSelect : NodesToSelect)
+		TSharedPtr<FSequencerSelection> Selection = ViewModel->GetSelection();
+
+		FSelectionEventSuppressor EventSuppressor = Selection->SuppressEvents();
+
+		Selection->Outliner.Empty();
+		for (TViewModelPtr<IOutlinerExtension> NodeToSelect : NodesToSelect)
 		{
-			Selection.AddToSelection( NodeToSelect );
+			Selection->Outliner.Select( NodeToSelect );
 		}
 
 		TSharedPtr<SOutlinerView> TreeView = SequencerWidget->GetTreeView();
-		const TSet<TWeakPtr<FViewModel>>& OutlinerSelection = GetSelection().GetSelectedOutlinerItems();
-		for (TWeakPtr<FViewModel> WeakNode : OutlinerSelection)
+		for (TViewModelPtr<IOutlinerExtension> Node : Selection->Outliner)
 		{
-			if (TViewModelPtr<IOutlinerExtension> Node = CastViewModel<IOutlinerExtension>(WeakNode.Pin()))
+			for (TViewModelPtr<IOutlinerExtension> Parent : Node.AsModel()->GetAncestorsOfType<IOutlinerExtension>())
 			{
-				TViewModelPtr<IOutlinerExtension> Parent = Node.AsModel()->FindAncestorOfType<IOutlinerExtension>();
-				while (Parent.IsValid())
-				{
+				TreeView->SetItemExpansion(Parent, true);
 					TreeView->SetItemExpansion(Parent, true);
 					Parent = Parent.AsModel()->FindAncestorOfType<IOutlinerExtension>();
-				}
-
-				TreeView->RequestScrollIntoView(Node);
-				break;
 			}
-		}
 
-		Selection.ResumeBroadcast();
-		Selection.RequestOutlinerNodeSelectionChangedBroadcast();
+			TreeView->RequestScrollIntoView(Node);
+			break;
+		}
 	}
 }
 
@@ -6785,15 +6670,20 @@ void FSequencer::OnSelectedNodesOnlyChanged()
 
 void FSequencer::ZoomToFit()
 {
+	using namespace UE::Sequencer;
+
 	FFrameRate TickResolution = GetFocusedTickResolution();
 
 	TRange<FFrameNumber> BoundsHull = TRange<FFrameNumber>::All();
 	
-	for ( const FSequencerSelectedKey& Key : Selection.GetSelectedKeys().Array() )
+	TSharedPtr<FSequencerSelection> Selection = ViewModel->GetSelection();
+
+	for (FKeyHandle Key : Selection->KeySelection)
 	{
-		if (Key.IsValid())
+		TSharedPtr<FChannelModel> Channel = Selection->KeySelection.GetModelForKey(Key);
+		if (Channel)
 		{
-			FFrameNumber KeyTime = Key.WeakChannel.Pin()->GetKeyArea()->GetKeyTime(Key.KeyHandle);
+			FFrameNumber KeyTime = Channel->GetKeyArea()->GetKeyTime(Key);
 			if (!BoundsHull.HasLowerBound() || BoundsHull.GetLowerBoundValue() > KeyTime)
 			{
 				BoundsHull.SetLowerBound(TRange<FFrameNumber>::BoundsType::Inclusive(KeyTime));
@@ -6805,7 +6695,7 @@ void FSequencer::ZoomToFit()
 		}
 	}
 
-	for (TWeakObjectPtr<UMovieSceneSection> SelectedSection : Selection.GetSelectedSections())
+	for (TWeakObjectPtr<UMovieSceneSection> SelectedSection : Selection->GetSelectedSections())
 	{
 		if (SelectedSection->GetRange().HasUpperBound() && SelectedSection->GetRange().HasLowerBound())
 		{
@@ -6866,35 +6756,19 @@ void FSequencer::KeyProperty(FKeyPropertyParams KeyPropertyParams)
 }
 
 
-FSequencerSelection& FSequencer::GetSelection()
-{
-	return Selection;
-}
-
-
 FSequencerSelectionPreview& FSequencer::GetSelectionPreview()
 {
 	return SelectionPreview;
 }
 
-void FSequencer::SuspendSelectionBroadcast()
-{
-	Selection.SuspendBroadcast();
-}
-
-void FSequencer::ResumeSelectionBroadcast()
-{
-	Selection.ResumeBroadcast();
-}
-
 void FSequencer::GetSelectedTracks(TArray<UMovieSceneTrack*>& OutSelectedTracks)
 {
-	OutSelectedTracks.Append(Selection.GetSelectedTracks());
+	OutSelectedTracks.Append(ViewModel->GetSelection()->GetSelectedTracks().Array());
 }
 
 void FSequencer::GetSelectedSections(TArray<UMovieSceneSection*>& OutSelectedSections)
 {
-	for (TWeakObjectPtr<UMovieSceneSection> WeakSection : Selection.GetSelectedSections())
+	for (TWeakObjectPtr<UMovieSceneSection> WeakSection : ViewModel->GetSelection()->GetSelectedSections())
 	{
 		if (UMovieSceneSection* Section = WeakSection.Get())
 		{
@@ -6921,15 +6795,9 @@ void FSequencer::GetSelectedKeyAreas(TArray<const IKeyArea*>& OutSelectedKeyArea
 	TArray<TSharedRef<FViewModel>> NodesToKey;
 	{
 		TSet<TSharedRef<FViewModel>> ChildNodes;
-		for (const TWeakPtr<FViewModel>& WeakNode : Selection.GetSelectedOutlinerItems())
+		for (FViewModelPtr Node : ViewModel->GetSelection()->Outliner)
 		{
-			TSharedPtr<FViewModel> Node = WeakNode.Pin();
-			if (!Node)
-			{
-				continue;
-			}
-
-			NodesToKey.Add(Node.ToSharedRef());
+			NodesToKey.Add(Node.AsModel().ToSharedRef());
 
 			// No need to gather key areas from binding/tracks because they have no key areas
 			if (Node->IsA<IObjectBindingExtension>() || Node->IsA<ITrackExtension>() || Node->IsA<FFolderModel>())
@@ -6938,7 +6806,7 @@ void FSequencer::GetSelectedKeyAreas(TArray<const IKeyArea*>& OutSelectedKeyArea
 			}
 
 			ChildNodes.Reset();
-			SequencerHelpers::GetDescendantNodes(Node.ToSharedRef(), ChildNodes);
+			SequencerHelpers::GetDescendantNodes(Node.AsModel().ToSharedRef(), ChildNodes);
 
 			for (TSharedRef<FViewModel> ChildNode : ChildNodes)
 			{
@@ -6961,9 +6829,13 @@ void FSequencer::GetSelectedKeyAreas(TArray<const IKeyArea*>& OutSelectedKeyArea
 
 	if (bIncludeSelectedKeys)
 	{
-		for (FSequencerSelectedKey Key : Selection.GetSelectedKeys())
+		for (FKeyHandle Key : ViewModel->GetSelection()->KeySelection)
 		{
-			KeyAreas.Add(Key.WeakChannel.Pin()->GetKeyArea()); 
+			TSharedPtr<FChannelModel> Channel = ViewModel->GetSelection()->KeySelection.GetModelForKey(Key);
+			if (Channel)
+			{
+				KeyAreas.Add(Channel->GetKeyArea()); 
+			}
 		}
 	}
 	for (TSharedPtr<IKeyArea> KeyArea : KeyAreas)
@@ -6978,30 +6850,33 @@ void FSequencer::SelectByNthCategoryNode(UMovieSceneSection* Section, int Index,
 	using namespace UE::Sequencer;
 
 	TSet<TSharedRef<FViewModel>> Nodes;
-	TArray<TSharedRef<FViewModel>> NodesToSelect;
+	TArray<TViewModelPtr<IOutlinerExtension>> NodesToSelect;
+
+	TSharedPtr<FSequencerSelection> Selection = ViewModel->GetSelection();
+	FSelectionEventSuppressor EventSuppressor = Selection->SuppressEvents();
 
 	TSharedPtr<FSectionModel> SectionHandle = NodeTree->GetSectionModel(Section);
 	int32 Count = 0;
 	if (SectionHandle)
 	{
 		TSharedPtr<FViewModel> TrackNode = SectionHandle->GetParentTrackModel();
-		for (TSharedPtr<FViewModel> Node : TrackNode->GetChildren())
+		for (TViewModelPtr<FCategoryGroupModel> Node : TrackNode->GetChildrenOfType<FCategoryGroupModel>())
 		{
-			if (Node->IsA<FCategoryGroupModel>() && Count++ == Index)
+			if (Count++ == Index)
 			{
 				bool bAlreadySelected = false;
 				if (bSelect == true)
 				{
-					bAlreadySelected = Selection.GetSelectedOutlinerItems().Contains(Node);					
+					bAlreadySelected = Selection->Outliner.IsSelected(Node);
 				}
 				if (bAlreadySelected == false)
 				{
-					NodesToSelect.Add(Node.ToSharedRef());
+					NodesToSelect.Add(Node);
 					if (bSelect == false) //make sure all children not selected
 					{
-						for (TSharedPtr<FViewModel> ChildNode : Node->GetChildren())
+						for (TViewModelPtr<IOutlinerExtension> ChildNode : Node->GetChildrenOfType<IOutlinerExtension>())
 						{
-							NodesToSelect.Add(ChildNode.ToSharedRef());
+							NodesToSelect.Add(ChildNode);
 						}
 					}
 				}
@@ -7012,15 +6887,16 @@ void FSequencer::SelectByNthCategoryNode(UMovieSceneSection* Section, int Index,
 	{
 		if (Settings->GetAutoExpandNodesOnSelection())
 		{
-			for (const TSharedRef<FViewModel>& DisplayNode : NodesToSelect)
+			for (FViewModelPtr DisplayNode : NodesToSelect)
 			{
-				if (DisplayNode->GetParent().IsValid() && DisplayNode->GetParent()->IsA<ITrackExtension>())
+				FViewModelPtr Parent = DisplayNode->GetParent();
+				if (TViewModelPtr<ITrackExtension> ParentTrack = Parent.ImplicitCast())
 				{
-					if (IOutlinerExtension* ParentOutlinerItem = DisplayNode->GetParent()->CastThis<IOutlinerExtension>())
+					if (TViewModelPtr<IOutlinerExtension> ParentOutliner = Parent.ImplicitCast())
 					{
-						if (!ParentOutlinerItem->IsExpanded())
+						if (!ParentOutliner->IsExpanded())
 						{
-							ParentOutlinerItem->SetExpansion(true);
+							ParentOutliner->SetExpansion(true);
 						}
 					}
 					break;
@@ -7032,17 +6908,15 @@ void FSequencer::SelectByNthCategoryNode(UMovieSceneSection* Section, int Index,
 		{
 			SequencerWidget->GetTreeView()->RequestScrollIntoView(CastViewModel<IOutlinerExtension>(NodesToSelect[0]));
 
-			Selection.AddToSelection(NodesToSelect);
-			Selection.RequestOutlinerNodeSelectionChangedBroadcast();
+			Selection->Outliner.SelectRange(NodesToSelect);
 		}
 	}
 	else if (NodesToSelect.Num() > 0)
 	{
-		for (const TSharedRef<FViewModel>& DisplayNode : NodesToSelect)
+		for (const TViewModelPtr<IOutlinerExtension>& DisplayNode : NodesToSelect)
 		{
-			Selection.RemoveFromSelection(DisplayNode);
+			Selection->Outliner.Deselect(DisplayNode);
 		}
-		Selection.RequestOutlinerNodeSelectionChangedBroadcast();
 	}
 }
 
@@ -7050,8 +6924,8 @@ void FSequencer::SelectByChannels(UMovieSceneSection* Section, TArrayView<const
 {
 	using namespace UE::Sequencer;
 
-	TSet<TSharedRef<FViewModel>> Nodes;
-	TArray<TSharedRef<FViewModel>> NodesToSelect;
+	TSet<TViewModelPtr<IOutlinerExtension>> Nodes;
+	TArray<TViewModelPtr<IOutlinerExtension>> NodesToSelect;
 
 	TSharedPtr<FSectionModel> SectionHandle = NodeTree->GetSectionModel(Section);
 	if (SectionHandle)
@@ -7068,47 +6942,53 @@ void FSequencer::SelectByChannels(UMovieSceneSection* Section, TArrayView<const
 					{
 						if (bSelect)
 						{
-							NodesToSelect.Add(KeyAreaNode->GetParent().ToSharedRef());
+							NodesToSelect.Add(KeyAreaNode->FindAncestorOfType<IOutlinerExtension>());
 						}
 						else
 						{
-							Nodes.Add(KeyAreaNode->GetParent().ToSharedRef());
+							Nodes.Add(KeyAreaNode->FindAncestorOfType<IOutlinerExtension>());
 						}
 					}
 					if (!bSelectParentInstead || bSelect == false)
 					{
+						// Handle top-level channels which might need to select their parent model
+						TViewModelPtr<IOutlinerExtension> OutlinerNode = CastViewModel<IOutlinerExtension>(KeyAreaNode);
+						if (!OutlinerNode)
+						{
+							OutlinerNode = KeyAreaNode->FindAncestorOfType<IOutlinerExtension>();
+						}
 						if (bSelect)
 						{
-							NodesToSelect.Add(KeyAreaNode.ToSharedRef());
+							NodesToSelect.Add(OutlinerNode);
 						}
 						else
 						{
-							Nodes.Add(KeyAreaNode.ToSharedRef());
+							Nodes.Add(OutlinerNode);
 						}
 					}
 				}
 			}
 		}
 	}
-	
+
+	TSharedPtr<FSequencerSelection> Selection = ViewModel->GetSelection();
+	FSelectionEventSuppressor EventSuppressor = Selection->SuppressEvents();
+
 	if (bSelect)
 	{
 		if (NodesToSelect.Num() > 0)
 		{
 			//todo hide behind preference 
 			SequencerWidget->GetTreeView()->RequestScrollIntoView(CastViewModel<IOutlinerExtension>(NodesToSelect[0]));
-		
-			Selection.AddToSelection(NodesToSelect);
-			Selection.RequestOutlinerNodeSelectionChangedBroadcast();
+			Selection->Outliner.SelectRange(NodesToSelect);
 		}
 	}
 	else if (Nodes.Num() > 0)
 	{
-		for (const TSharedRef<FViewModel>& DisplayNode : Nodes)
+		for (const TViewModelPtr<IOutlinerExtension>& DisplayNode : Nodes)
 		{
-			Selection.RemoveFromSelection(DisplayNode);
+			Selection->Outliner.Deselect(DisplayNode);
 		}
-		Selection.RequestOutlinerNodeSelectionChangedBroadcast();
 	}
 }
 
@@ -7116,18 +6996,19 @@ void FSequencer::SelectByChannels(UMovieSceneSection* Section, const TArray<FNam
 {
 	using namespace UE::Sequencer;
 
-	TSet<TSharedRef<FViewModel>> Nodes;
-	TArray<TSharedRef<FViewModel>> NodesToSelect;
+	TSet<TViewModelPtr<IOutlinerExtension>> Nodes;
+	TArray<TViewModelPtr<IOutlinerExtension>> NodesToSelect;
 
 	TSharedPtr<FSectionModel> SectionHandle = NodeTree->GetSectionModel(Section);
 	if (SectionHandle)
 	{
 		TSharedPtr<FViewModel> TrackNode = SectionHandle->GetParentTrackModel();
-		for (TSharedPtr<FChannelGroupModel> KeyAreaNode : TrackNode->GetDescendantsOfType<FChannelGroupModel>())
+		for (TSharedPtr<FChannelGroupOutlinerModel> KeyAreaNode : TrackNode->GetDescendantsOfType<FChannelGroupOutlinerModel>())
 		{
-			if (KeyAreaNode->GetParent().IsValid() && InChannelNames.Contains(*KeyAreaNode->GetParent()->CastThis<IOutlinerExtension>()->GetLabel().ToString()))
+			TViewModelPtr<IOutlinerExtension> Parent = KeyAreaNode->FindAncestorOfType<IOutlinerExtension>();
+			if (Parent && InChannelNames.Contains(*Parent->GetLabel().ToString()))
 			{
-				Nodes.Add(KeyAreaNode->GetParent().ToSharedRef());
+				Nodes.Add(Parent);
 			}
 
 			for (TSharedPtr<IKeyArea> KeyArea : KeyAreaNode->GetAllKeyAreas())
@@ -7140,26 +7021,30 @@ void FSequencer::SelectByChannels(UMovieSceneSection* Section, const TArray<FNam
 				{
 					if (bSelectParentInstead || bSelect == false)
 					{
-						Nodes.Add(KeyAreaNode->GetParent().ToSharedRef());
+						Nodes.Add(Parent);
 					}
 					if (!bSelectParentInstead || bSelect == false)
 					{
-						Nodes.Add(KeyAreaNode.ToSharedRef());
+						Nodes.Add(KeyAreaNode);
 					}
 				}
 			}
 		}
 	}
-	
+
+	TSharedPtr<FSequencerSelection> Selection = ViewModel->GetSelection();
+	FSelectionEventSuppressor EventSuppressor = Selection->SuppressEvents();
+
 	if (bSelect)
 	{
-		for (const TSharedRef<FViewModel>& DisplayNode : Nodes)
+		for (const TViewModelPtr<IOutlinerExtension>& DisplayNode : Nodes)
 		{
 			if (Settings->GetAutoExpandNodesOnSelection())
 			{
-				if (DisplayNode->GetParent().IsValid() && DisplayNode->GetParent()->IsA<ITrackExtension>() && !DisplayNode->GetParent()->CastThisChecked<IOutlinerExtension>()->IsExpanded())
+				TViewModelPtr<IOutlinerExtension> Parent = DisplayNode.AsModel()->FindAncestorOfType<IOutlinerExtension>();
+				if (Parent && Parent.AsModel()->IsA<ITrackExtension>() && !Parent->IsExpanded())
 				{
-					DisplayNode->GetParent()->CastThisChecked<IOutlinerExtension>()->SetExpansion(true);
+					Parent->SetExpansion(true);
 				}
 			}
 			NodesToSelect.Add(DisplayNode);
@@ -7169,17 +7054,15 @@ void FSequencer::SelectByChannels(UMovieSceneSection* Section, const TArray<FNam
 		{
 			SequencerWidget->GetTreeView()->RequestScrollIntoView(CastViewModel<IOutlinerExtension>(NodesToSelect[0]));
 
-			Selection.AddToSelection(NodesToSelect);
-			Selection.RequestOutlinerNodeSelectionChangedBroadcast();
+			Selection->Outliner.SelectRange(NodesToSelect);
 		}
 	}
 	else if (Nodes.Num() > 0)
 	{
-		for (const TSharedRef<FViewModel>& DisplayNode : Nodes)
+		for (const TViewModelPtr<IOutlinerExtension>& DisplayNode : Nodes)
 		{
-			Selection.RemoveFromSelection(DisplayNode);
+			Selection->Outliner.Deselect(DisplayNode);
 		}
-		Selection.RequestOutlinerNodeSelectionChangedBroadcast();
 	}
 }
 
@@ -7192,7 +7075,7 @@ void FSequencer::SelectObject(FGuid ObjectBinding)
 
 	if (TSharedPtr<FObjectBindingModel> Model = ObjectStorage->FindModelForObjectBinding(ObjectBinding))
 	{
-		Selection.AddToOutlinerSelection(Model);
+		ViewModel->GetSelection()->Outliner.Select(Model);
 	}
 }
 
@@ -7200,19 +7083,13 @@ void FSequencer::SelectTrack(UMovieSceneTrack* Track)
 {
 	using namespace UE::Sequencer;
 	
-	TArray<TSharedPtr<FViewModel>> AllNodes;
-	NodeTree->GetAllNodes(AllNodes);
-	for (TSharedPtr<FViewModel> Node : AllNodes)
+	FTrackModelStorageExtension* TrackStorage = ViewModel->GetRootModel()->CastDynamic<FTrackModelStorageExtension>();
+	check(TrackStorage);
+
+	TSharedPtr<FTrackModel> TrackModel = TrackStorage->FindModelForTrack(Track);
+	if (TrackModel)
 	{
-		if (ITrackExtension* TrackNode = Node->CastThis<ITrackExtension>())
-		{
-			UMovieSceneTrack* TrackForNode = TrackNode->GetTrack();
-			if (TrackForNode == Track)
-			{
-				Selection.AddToOutlinerSelection(Node);
-				break;
-			}
-		}
+		ViewModel->GetSelection()->Outliner.Select(TrackModel);
 	}
 }
 
@@ -7223,24 +7100,22 @@ void FSequencer::SelectSection(UMovieSceneSection* Section)
 	FSectionModelStorageExtension* SectionModelStorage = ViewModel->GetRootModel()->CastDynamic<FSectionModelStorageExtension>();
 	check(SectionModelStorage);
 
-	TSharedPtr<UE::Sequencer::FSectionModel> SectionModel = SectionModelStorage->FindModelForSection(Section);
+	TSharedPtr<FSectionModel> SectionModel = SectionModelStorage->FindModelForSection(Section);
 	if (SectionModel)
 	{
-		Selection.AddToSelection(SectionModel);
+		ViewModel->GetSelection()->TrackArea.Select(SectionModel);
 	}
 }
 
 void FSequencer::SelectKey(UMovieSceneSection* InSection, TSharedPtr<UE::Sequencer::FChannelModel> InChannel, FKeyHandle KeyHandle, bool bToggle)
 {
-	FSequencerSelectedKey SelectedKey(*InSection, InChannel, KeyHandle);
-
-	if (bToggle && Selection.IsSelected(SelectedKey))
+	if (bToggle && ViewModel->GetSelection()->KeySelection.IsSelected(KeyHandle))
 	{
-		Selection.RemoveFromSelection(SelectedKey);
+		ViewModel->GetSelection()->KeySelection.Deselect(KeyHandle);
 	}
 	else
 	{
-		Selection.AddToSelection(SelectedKey);
+		ViewModel->GetSelection()->KeySelection.Select(InChannel, KeyHandle);
 	}
 }
 
@@ -7248,37 +7123,26 @@ void FSequencer::SelectByPropertyPaths(const TArray<FString>& InPropertyPaths)
 {
 	using namespace UE::Sequencer;
 
-	TArray<TSharedRef<FViewModel>> AllNodes;
-	NodeTree->GetAllNodes(AllNodes);
+	TSharedPtr<FSequencerSelection> Selection = ViewModel->GetSelection();
+	FSelectionEventSuppressor EventSuppressor = Selection->SuppressEvents();
 
-	TArray<TSharedRef<FViewModel>> NodesToSelect;
-	for (const TSharedRef<FViewModel>& Node : AllNodes)
+	Selection->Empty();
+
+	for (TViewModelPtr<ITrackExtension> Track : ViewModel->GetRootModel()->GetDescendantsOfType<ITrackExtension>())
 	{
-		if (ITrackExtension* TrackNode = Node->CastThis<ITrackExtension>())
+		if (UMovieScenePropertyTrack* PropertyTrack = Cast<UMovieScenePropertyTrack>(Track->GetTrack()))
 		{
-			if (UMovieScenePropertyTrack* PropertyTrack = Cast<UMovieScenePropertyTrack>(TrackNode->GetTrack()))
+			FString Path = PropertyTrack->GetPropertyPath().ToString();
+			for (const FString& PropertyPath : InPropertyPaths)
 			{
-				FString Path = PropertyTrack->GetPropertyPath().ToString();
-				for (const FString& PropertyPath : InPropertyPaths)
+				if (Path == PropertyPath)
 				{
-					if (Path == PropertyPath)
-					{
-						NodesToSelect.Add(Node);
-						break;
-					}
+					Selection->Outliner.Select(CastViewModelChecked<IOutlinerExtension>(Track));
+					break;
 				}
 			}
 		}
 	}
-
-	Selection.SuspendBroadcast();
-	Selection.Empty();
-	Selection.ResumeBroadcast();
-
-	if (NodesToSelect.Num())
-	{
-		Selection.AddToSelection(NodesToSelect);
-	}
 }
 
 
@@ -7286,26 +7150,20 @@ void FSequencer::SelectFolder(UMovieSceneFolder* Folder)
 {
 	using namespace UE::Sequencer;
 
-	TArray<TSharedRef<FViewModel>> AllNodes;
-	NodeTree->GetAllNodes(AllNodes);
-	for (TSharedRef<FViewModel> Node : AllNodes)
+	FFolderModelStorageExtension* FolderModelStorage = ViewModel->GetRootModel()->CastDynamic<FFolderModelStorageExtension>();
+	check(FolderModelStorage);
+
+	TSharedPtr<FFolderModel> FolderModel = FolderModelStorage->FindModelForFolder(Folder);
+	if (FolderModel)
 	{
-		if (FFolderModel* FolderNode = Node->CastThis<FFolderModel>())
-		{
-			UMovieSceneFolder* FolderForNode = FolderNode->GetFolder();
-			if (FolderForNode == Folder)
-			{
-				Selection.AddToSelection(Node);
-				break;
-			}
-		}
+		ViewModel->GetSelection()->TrackArea.Select(FolderModel);
 	}
 }
 
 
 void FSequencer::EmptySelection()
 {
-	Selection.Empty();
+	ViewModel->GetSelection()->Empty();
 }
 
 void FSequencer::ThrobKeySelection()
@@ -7321,9 +7179,9 @@ void FSequencer::ThrobSectionSelection()
 
 	// Scrub to the beginning of newly created sections if they're out of view
 	TOptional<FFrameNumber> ScrubFrame;
-	for (TWeakObjectPtr<UMovieSceneSection> SelectedSectionPtr : Selection.GetSelectedSections())
+	for (UMovieSceneSection* SelectedSectionPtr : ViewModel->GetSelection()->GetSelectedSections())
 	{
-		if (SelectedSectionPtr.IsValid() && SelectedSectionPtr->HasStartFrame())
+		if (SelectedSectionPtr && SelectedSectionPtr->HasStartFrame())
 		{
 			if (!ScrubFrame.IsSet() || (ScrubFrame.GetValue() > SelectedSectionPtr->GetInclusiveStartFrame()))
 			{
@@ -7360,19 +7218,19 @@ float FSequencer::GetOverlayFadeCurve() const
 
 void FSequencer::DeleteSelectedItems()
 {
-	if (Selection.GetSelectedKeys().Num())
+	if (ViewModel->GetSelection()->KeySelection.Num())
 	{
 		FScopedTransaction DeleteKeysTransaction( NSLOCTEXT("Sequencer", "DeleteKeys_Transaction", "Delete Keys") );
 		
 		DeleteSelectedKeys();
 	}
-	else if (Selection.GetSelectedSections().Num())
+	else if (ViewModel->GetSelection()->GetSelectedSections().Num())
 	{
 		FScopedTransaction DeleteSectionsTransaction( NSLOCTEXT("Sequencer", "DeleteSections_Transaction", "Delete Sections") );
 	
-		DeleteSections(Selection.GetSelectedSections());
+		DeleteSections(ViewModel->GetSelection()->GetSelectedSections());
 	}
-	else if (Selection.GetSelectedOutlinerItems().Num())
+	else if (ViewModel->GetSelection()->Outliner.Num())
 	{
 		DeleteSelectedNodes(false);
 	}
@@ -7381,8 +7239,10 @@ void FSequencer::DeleteSelectedItems()
 
 void FSequencer::DeleteNode(TSharedRef<FViewModel> NodeToBeDeleted, const bool bKeepState)
 {
+	using namespace UE::Sequencer;
+
 	// If this node is selected, delete all selected nodes
-	if (GetSelection().IsSelected(NodeToBeDeleted))
+	if (GetSelection().Outliner.IsSelected(CastViewModelChecked<IOutlinerExtension>(NodeToBeDeleted)))
 	{
 		DeleteSelectedNodes(bKeepState);
 	}
@@ -7402,17 +7262,20 @@ void FSequencer::DeleteSelectedNodes(const bool bKeepState)
 {
 	using namespace UE::Sequencer;
 
-	TArray<TSharedPtr<FViewModel>> SelectedNodesCopy;
-	GetSelection().GetSelectedOutlinerItems(SelectedNodesCopy);
-
-	if (SelectedNodesCopy.Num() == 0)
+	if (ViewModel->GetSelection()->Outliner.Num() == 0)
 	{
 		return;
 	}
 
 	const FScopedTransaction Transaction( NSLOCTEXT("Sequencer", "UndoDeletingObject", "Delete Node") );
 
-	Selection.EmptySelectedOutlinerNodes();
+	TArray<TSharedPtr<FViewModel>> SelectedNodesCopy;
+	SelectedNodesCopy.Reserve(ViewModel->GetSelection()->Outliner.Num());
+	for (FViewModelPtr OutlinerNode : ViewModel->GetSelection()->Outliner)
+	{
+		SelectedNodesCopy.Emplace(MoveTemp(OutlinerNode));
+	}
+	ViewModel->GetSelection()->Outliner.Empty();
 
 	bool bAnythingDeleted = false;
 
@@ -7498,17 +7361,12 @@ TArray<TSharedRef<UE::Sequencer::FViewModel>> FSequencer::GetSelectedNodesToMove
 	TArray<TSharedRef<FViewModel>> NodesToMove;
 
 	// Build a list of the nodes we want to move.
-	TArray<TSharedRef<FViewModel>> SelectedItems;
-	GetSelection().GetSelectedOutlinerItems(SelectedItems);
-	for (TSharedRef<FViewModel> Node : SelectedItems)
+	for (TViewModelPtr<IDraggableOutlinerExtension> DraggableModel : GetSelection().Outliner.Filter<IDraggableOutlinerExtension>())
 	{
 		// Only nodes that can be dragged can be moved in to a folder.
-		if (IDraggableOutlinerExtension* DraggableModel = Node->CastThis<IDraggableOutlinerExtension>())
+		if (DraggableModel->CanDrag())
 		{
-			if (DraggableModel->CanDrag())
-			{
-				NodesToMove.Add(Node);
-			}
+			NodesToMove.Add(DraggableModel.AsModel().ToSharedRef());
 		}
 	}
 
@@ -7556,11 +7414,8 @@ TArray<TSharedRef<UE::Sequencer::FViewModel> > FSequencer::GetSelectedNodesInFol
 {
 	using namespace UE::Sequencer;
 
-	TArray<TSharedRef<FViewModel>> SelectedItems;
-	GetSelection().GetSelectedOutlinerItems(SelectedItems);
-
 	TArray<TSharedRef<FViewModel> > NodesToFolders;
-	for (TSharedRef<FViewModel> SelectedNode : SelectedItems)
+	for (FViewModelPtr SelectedNode : GetSelection().Outliner)
 	{
 		TSharedPtr<FFolderModel> Folder = SelectedNode->FindAncestorOfType<FFolderModel>();
 		if (Folder.IsValid())
@@ -7569,7 +7424,7 @@ TArray<TSharedRef<UE::Sequencer::FViewModel> > FSequencer::GetSelectedNodesInFol
 			{
 				if (Folder->GetFolder()->GetChildObjectBindings().Contains(ObjectBindingNode->GetObjectGuid()))
 				{
-					NodesToFolders.Add(SelectedNode);
+					NodesToFolders.Add(SelectedNode.AsModel().ToSharedRef());
 				}
 			}
 			else if (ITrackExtension* TrackNode = SelectedNode->CastThis<ITrackExtension>())
@@ -7578,7 +7433,7 @@ TArray<TSharedRef<UE::Sequencer::FViewModel> > FSequencer::GetSelectedNodesInFol
 				{
 					if (Folder->GetFolder()->GetChildTracks().Contains(TrackNode->GetTrack()))
 					{
-						NodesToFolders.Add(SelectedNode);
+						NodesToFolders.Add(SelectedNode.AsModel().ToSharedRef());
 					}
 				}
 			}
@@ -7700,7 +7555,7 @@ void FSequencer::MoveSelectedNodesToFolder(UMovieSceneFolder* DestinationFolder)
 
 	FScopedTransaction Transaction(LOCTEXT("MoveTracksToFolder", "Move to Folder"));
 
-	Selection.Empty();
+	ViewModel->GetSelection()->Empty();
 
 	// Find the path to the displaynode of our destination folder
 	FString DestinationFolderPath;
@@ -7925,7 +7780,7 @@ void FSequencer::MoveSelectedNodesToNewFolder()
 	}
 
 	// Set the newly created folder as our selection
-	Selection.Empty();
+	ViewModel->GetSelection()->Empty();
 	SequencerWidget->AddAdditionalPathToSelectionSet(NewFolderPath);
 
 
@@ -8026,7 +7881,7 @@ void FSequencer::CopySelectedTracks(TArray<TSharedPtr<FViewModel>>& TrackNodes,
 		TSharedPtr<FViewModel> ParentNode = TrackNode->GetParent();
 		while (ParentNode.IsValid() && !ParentNode->IsA<FFolderModel>())
 		{
-			if (Selection.GetSelectedOutlinerItems().Contains(ParentNode.ToSharedRef()))
+			if (ViewModel->GetSelection()->Outliner.IsSelected(CastViewModel<IOutlinerExtension>(ParentNode)))
 			{
 				bIsParentSelected = true;
 				break;
@@ -8152,18 +8007,12 @@ bool FSequencer::PasteObjectBindings(const FString& TextToImport, UMovieSceneFol
 {
 	using namespace UE::Sequencer;
 
-	TArray<TSharedRef<FViewModel>> SelectedNodes;
-	Selection.GetSelectedOutlinerItems(SelectedNodes);
-
 	TArray<FMovieSceneBindingProxy> TargetBindings;
 	if (!bClearSelection)
 	{
-		for (TSharedRef<FViewModel> Node : SelectedNodes)
+		for (TViewModelPtr<FObjectBindingModel> ObjectNode : ViewModel->GetSelection()->Outliner.Filter<FObjectBindingModel>())
 		{
-			if (FObjectBindingModel* ObjectNode = Node->CastThis<FObjectBindingModel>())
-			{
-				TargetBindings.Add(FMovieSceneBindingProxy(ObjectNode->GetObjectGuid(), GetFocusedMovieSceneSequence()));
-			}
+			TargetBindings.Add(FMovieSceneBindingProxy(ObjectNode->GetObjectGuid(), GetFocusedMovieSceneSequence()));
 		}
 	}
 
@@ -8217,18 +8066,12 @@ bool FSequencer::PasteTracks(const FString& TextToImport, UMovieSceneFolder* InP
 {
 	using namespace UE::Sequencer;
 
-	TArray<TSharedRef<FViewModel>> SelectedNodes;
-	Selection.GetSelectedOutlinerItems(SelectedNodes);
-
 	TArray<FMovieSceneBindingProxy> TargetBindings;
 	if (!bClearSelection)
 	{
-		for (TSharedRef<FViewModel> Node : SelectedNodes)
+		for (TViewModelPtr<FObjectBindingModel> ObjectNode : ViewModel->GetSelection()->Outliner.Filter<FObjectBindingModel>())
 		{
-			if (TSharedPtr<FObjectBindingModel> ObjectNode = Node->CastThisShared<FObjectBindingModel>())
-			{
-				TargetBindings.Add(FMovieSceneBindingProxy(ObjectNode->GetObjectGuid(), GetFocusedMovieSceneSequence()));
-			}
+			TargetBindings.Add(FMovieSceneBindingProxy(ObjectNode->GetObjectGuid(), GetFocusedMovieSceneSequence()));
 		}
 	}
 
@@ -8249,8 +8092,28 @@ bool FSequencer::PasteSections(const FString& TextToImport, TArray<FNotification
 
 	FScopedTransaction Transaction(FGenericCommands::Get().Paste->GetDescription());
 
+	TSharedPtr<FSequencerSelection> Selection = ViewModel->GetSelection();
+	FSelectionEventSuppressor EventSuppressor = Selection->SuppressEvents();
+
 	TArray<TSharedRef<FViewModel>> SelectedNodes;
-	Selection.GetSelectedOutlinerItems(SelectedNodes);
+	SelectedNodes.Reserve(Selection->Outliner.Num());
+	for (FViewModelPtr SelectedNode : Selection->Outliner)
+	{
+		SelectedNodes.Add(SelectedNode.AsModel().ToSharedRef());
+	}
+
+	TSet<UMovieSceneSection*> SelectedSections = Selection->GetSelectedSections();
+	if (SelectedNodes.Num() == 0)
+	{
+		if (Selection->GetNodesWithSelectedKeysOrSections().Num() != 0)
+		{
+			SelectedNodes.Reserve(Selection->GetNodesWithSelectedKeysOrSections().Num());
+			for (TViewModelPtr<IOutlinerExtension> Item : Selection->IterateIndirectOutlinerSelection())
+			{
+				SelectedNodes.Add(Item.AsModel().ToSharedRef());
+			}
+		}
+	}
 
 	if (SelectedNodes.Num() == 0)
 	{
@@ -8260,7 +8123,7 @@ bool FSequencer::PasteSections(const FString& TextToImport, TArray<FNotification
 			SequencerHelpers::GetAllSections(RootNode.AsModel(), Sections);
 			for (TWeakObjectPtr<UMovieSceneSection> Section : Sections)
 			{
-				if (Selection.GetSelectedSections().Contains(Section.Get()))
+				if (SelectedSections.Contains(Section.Get()))
 				{
 					SelectedNodes.Add(RootNode.AsModel().ToSharedRef());
 					break;
@@ -8340,18 +8203,17 @@ bool FSequencer::PasteSections(const FString& TextToImport, TArray<FNotification
 		check(SectionModelStorage);
 	}
 
+
 	EmptySelection();
 
-	Selection.SuspendBroadcast();
 	for (UMovieSceneSection* NewSection : NewSections)
 	{
 		TSharedPtr<FSectionModel> SectionModel = SectionModelStorage->FindModelForSection(NewSection);
 		if (ensure(SectionModel))
 		{
-			Selection.AddToSelection(SectionModel);
+			Selection->TrackArea.Select(SectionModel);
 		}
 	}
-	Selection.ResumeBroadcast();
 
 	ThrobSectionSelection();
 
@@ -8397,60 +8259,18 @@ TArray<FText> FSequencer::GetTrackFilterNames() const
 	return SequencerWidget->GetTrackFilterNames();
 }
 
-void FSequencer::ToggleNodeActive()
-{
-	bool bIsActive = !IsNodeActive();
-	const FScopedTransaction Transaction( NSLOCTEXT("Sequencer", "ToggleNodeActive", "Toggle Node Active") );
-
-	for (auto OutlinerNode : Selection.GetSelectedOutlinerItems())
-	{
-		TSet<TWeakObjectPtr<UMovieSceneSection> > Sections;
-		SequencerHelpers::GetAllSections(OutlinerNode.Pin(), Sections);
-
-		for (auto Section : Sections)
-		{
-			Section->Modify();
-			Section->SetIsActive(bIsActive);
-		}
-	}
-
-	NotifyMovieSceneDataChanged(EMovieSceneDataChangeType::TrackValueChanged);
-}
-
-
-bool FSequencer::IsNodeActive() const
-{
-	// Active if ONE is active, changed in 4.20
-	for (auto OutlinerNode : Selection.GetSelectedOutlinerItems())
-	{
-		TSet<TWeakObjectPtr<UMovieSceneSection> > Sections;
-		SequencerHelpers::GetAllSections(OutlinerNode.Pin(), Sections);
-		if (Sections.Num() > 0)
-		{
-			for (auto Section : Sections)
-			{
-				if (Section->IsActive())
-				{
-					return true;
-				}
-			}
-			return false;
-		}
-	}
-	return true;
-}
-
-
 void FSequencer::ToggleNodeLocked()
 {
+	using namespace UE::Sequencer;
+
 	bool bIsLocked = !IsNodeLocked();
 
 	const FScopedTransaction Transaction( NSLOCTEXT("Sequencer", "ToggleNodeLocked", "Toggle Node Locked") );
 
-	for (auto OutlinerNode : Selection.GetSelectedOutlinerItems())
+	for (FViewModelPtr OutlinerNode : ViewModel->GetSelection()->Outliner)
 	{
 		TSet<TWeakObjectPtr<UMovieSceneSection> > Sections;
-		SequencerHelpers::GetAllSections(OutlinerNode.Pin(), Sections);
+		SequencerHelpers::GetAllSections(OutlinerNode, Sections);
 
 		for (auto Section : Sections)
 		{
@@ -8463,12 +8283,14 @@ void FSequencer::ToggleNodeLocked()
 
 bool FSequencer::IsNodeLocked() const
 {
+	using namespace UE::Sequencer;
+
 	// Locked only if all are locked
 	int NumSections = 0;
-	for (auto OutlinerNode : Selection.GetSelectedOutlinerItems())
+	for (FViewModelPtr OutlinerNode : ViewModel->GetSelection()->Outliner)
 	{
 		TSet<TWeakObjectPtr<UMovieSceneSection> > Sections;
-		SequencerHelpers::GetAllSections(OutlinerNode.Pin(), Sections);
+		SequencerHelpers::GetAllSections(OutlinerNode, Sections);
 
 		for (auto Section : Sections)
 		{
@@ -8494,9 +8316,8 @@ void FSequencer::GroupSelectedSections()
 	const FScopedTransaction Transaction(LOCTEXT("GroupSelectedSections", "Group Selected Sections"));
 
 	TArray<UMovieSceneSection*> Sections;
-	for (TWeakObjectPtr<UMovieSceneSection> WeakSection : Selection.GetSelectedSections())
+	for (UMovieSceneSection* Section : ViewModel->GetSelection()->GetSelectedSections())
 	{
-		UMovieSceneSection* Section = WeakSection.Get();
 		// We do not want to group sections that are infinite, as they should not be moveable
 		if (Section && (Section->HasStartFrame() || Section->HasEndFrame()))
 		{
@@ -8510,11 +8331,10 @@ void FSequencer::GroupSelectedSections()
 bool FSequencer::CanGroupSelectedSections() const
 {
 	int32 GroupableSections = 0;
-	for (TWeakObjectPtr<UMovieSceneSection> WeakSection : Selection.GetSelectedSections())
+	for (UMovieSceneSection* Section : ViewModel->GetSelection()->GetSelectedSections())
 	{
-		UMovieSceneSection* Section = WeakSection.Get();
 		// We do not want to group sections that are infinite, as they should not be moveable
-		if (Section && (Section->HasStartFrame() || Section->HasEndFrame()))
+		if (Section->HasStartFrame() || Section->HasEndFrame())
 		{
 			if (++GroupableSections >= 2)
 			{
@@ -8536,12 +8356,9 @@ void FSequencer::UngroupSelectedSections()
 
 	const FScopedTransaction Transaction(LOCTEXT("UngroupSelectedSections", "Ungroup Selected Sections"));
 
-	for (TWeakObjectPtr<UMovieSceneSection> WeakSection : Selection.GetSelectedSections())
+	for (UMovieSceneSection* Section : ViewModel->GetSelection()->GetSelectedSections())
 	{
-		if (WeakSection.IsValid())
-		{
-			MovieScene->UngroupSection(*WeakSection.Get());
-		}
+		MovieScene->UngroupSection(*Section);
 	}
 }
 
@@ -8549,9 +8366,9 @@ bool FSequencer::CanUngroupSelectedSections() const
 {
 	UMovieScene* MovieScene = GetFocusedMovieSceneSequence()->GetMovieScene();
 
-	for (TWeakObjectPtr<UMovieSceneSection> WeakSection : Selection.GetSelectedSections())
+	for (UMovieSceneSection* Section : ViewModel->GetSelection()->GetSelectedSections())
 	{
-		if (WeakSection.IsValid() && MovieScene->IsSectionInGroup(*WeakSection.Get()))
+		if (MovieScene->IsSectionInGroup(*Section))
 		{
 			return true;
 		}
@@ -8577,15 +8394,12 @@ void FSequencer::SaveSelectedNodesSpawnableState()
 
 	TArray<FMovieSceneSpawnable*> Spawnables;
 
-	for (const TWeakPtr<FViewModel>& Node : Selection.GetSelectedOutlinerItems())
+	for (TViewModelPtr<IObjectBindingExtension> ObjectBindingNode : ViewModel->GetSelection()->Outliner.Filter<IObjectBindingExtension>())
 	{
-		if (IObjectBindingExtension* ObjectBindingNode = ICastable::CastWeakPtr<IObjectBindingExtension>(Node))
+		FMovieSceneSpawnable* Spawnable = MovieScene->FindSpawnable(ObjectBindingNode->GetObjectGuid());
+		if (Spawnable)
 		{
-			FMovieSceneSpawnable* Spawnable = MovieScene->FindSpawnable(ObjectBindingNode->GetObjectGuid());
-			if (Spawnable)
-			{
-				Spawnables.Add(Spawnable);
-			}
+			Spawnables.Add(Spawnable);
 		}
 	}
 
@@ -8609,7 +8423,7 @@ void FSequencer::SaveSelectedNodesSpawnableState()
 }
 
 void FSequencer::SetSelectedNodesSpawnableLevel(FName InLevelName)
-{						
+{
 	using namespace UE::Sequencer;
 
 	UMovieScene* MovieScene = GetFocusedMovieSceneSequence()->GetMovieScene();
@@ -8626,15 +8440,12 @@ void FSequencer::SetSelectedNodesSpawnableLevel(FName InLevelName)
 
 	TArray<FMovieSceneSpawnable*> Spawnables;
 
-	for (const TWeakPtr<FViewModel>& Node : Selection.GetSelectedOutlinerItems())
+	for (TViewModelPtr<IObjectBindingExtension> ObjectBindingNode : ViewModel->GetSelection()->Outliner.Filter<IObjectBindingExtension>())
 	{
-		if (IObjectBindingExtension* ObjectBindingNode = ICastable::CastWeakPtr<IObjectBindingExtension>(Node))
+		FMovieSceneSpawnable* Spawnable = MovieScene->FindSpawnable(ObjectBindingNode->GetObjectGuid());
+		if (Spawnable)
 		{
-			FMovieSceneSpawnable* Spawnable = MovieScene->FindSpawnable(ObjectBindingNode->GetObjectGuid());
-			if (Spawnable)
-			{
-				Spawnable->SetLevelName(InLevelName);
-			}
+			Spawnable->SetLevelName(InLevelName);
 		}
 	}
 }
@@ -8700,16 +8511,13 @@ void FSequencer::ConvertSelectedNodesToSpawnables()
 
 	TArray<IObjectBindingExtension*> ObjectBindingNodes;
 
-	for (const TWeakPtr<FViewModel>& Node : Selection.GetSelectedOutlinerItems())
+	for (TViewModelPtr<IObjectBindingExtension> ObjectBindingNode : ViewModel->GetSelection()->Outliner.Filter<IObjectBindingExtension>())
 	{
-		if (IObjectBindingExtension* ObjectBindingNode = ICastable::CastWeakPtr<IObjectBindingExtension>(Node))
+		// If we have a possessable for this node, and it has no parent, we can convert it to a spawnable
+		FMovieScenePossessable* Possessable = MovieScene->FindPossessable(ObjectBindingNode->GetObjectGuid());
+		if (Possessable && !Possessable->GetParent().IsValid())
 		{
-			// If we have a possessable for this node, and it has no parent, we can convert it to a spawnable
-			FMovieScenePossessable* Possessable = MovieScene->FindPossessable(ObjectBindingNode->GetObjectGuid());
-			if (Possessable && !Possessable->GetParent().IsValid())
-			{
-				ObjectBindingNodes.Add(ObjectBindingNode);
-			}
+			ObjectBindingNodes.Add(ObjectBindingNode.Get());
 		}
 	}
 
@@ -8803,15 +8611,12 @@ void FSequencer::ConvertSelectedNodesToPossessables()
 
 	TArray<IObjectBindingExtension*> ObjectBindingNodes;
 
-	for (const TWeakPtr<FViewModel>& Node : Selection.GetSelectedOutlinerItems())
+	for (TViewModelPtr<IObjectBindingExtension> ObjectBindingNode : ViewModel->GetSelection()->Outliner.Filter<IObjectBindingExtension>())
 	{
-		if (IObjectBindingExtension* ObjectBindingNode = ICastable::CastWeakPtr<IObjectBindingExtension>(Node))
+		FMovieSceneSpawnable* Spawnable = MovieScene->FindSpawnable(ObjectBindingNode->GetObjectGuid());
+		if (Spawnable && SpawnRegister->CanConvertSpawnableToPossessable(*Spawnable))
 		{
-			FMovieSceneSpawnable* Spawnable = MovieScene->FindSpawnable(ObjectBindingNode->GetObjectGuid());
-			if (Spawnable && SpawnRegister->CanConvertSpawnableToPossessable(*Spawnable))
-			{
-				ObjectBindingNodes.Add(ObjectBindingNode);
-			}
+			ObjectBindingNodes.Add(ObjectBindingNode.Get());
 		}
 	}
 
@@ -8994,7 +8799,7 @@ void FSequencer::OnAddFolder()
 		FocusedMovieScene->AddRootFolder( NewFolder );
 	}
 
-	Selection.Empty();
+	ViewModel->GetSelection()->Empty();
 
 	// We can't add the newly created folder to the selection set as the nodes for it don't actually exist yet.
 	// However, we can calculate the resulting path that the node will end up at and add that to the selection
@@ -9050,33 +8855,29 @@ void FSequencer::CalculateSelectedFolderAndPath(TArray<UMovieSceneFolder*>& OutS
 	using namespace UE::Sequencer;
 
 	// Check if a folder, or child of a folder is currently selected.
-	if (Selection.GetSelectedOutlinerItems().Num() > 0)
+	for (FViewModelPtr CurrentNode : ViewModel->GetSelection()->Outliner)
 	{
-		for (TWeakPtr<FViewModel> SelectedNode : Selection.GetSelectedOutlinerItems())
+		while (CurrentNode && !CurrentNode->IsA<FFolderModel>())
+		{
+			CurrentNode = CurrentNode->GetParent();
+		}
+		if (CurrentNode)
 		{
-			TSharedPtr<FViewModel> CurrentNode = SelectedNode.Pin();
-			while (CurrentNode.IsValid() && !CurrentNode->IsA<FFolderModel>())
+			UMovieSceneFolder* Folder = CurrentNode->CastThisChecked<FFolderModel>()->GetFolder();
+			if (Folder && !OutSelectedParentFolders.Contains(Folder))
 			{
-				CurrentNode = CurrentNode->GetParent();
+				OutSelectedParentFolders.Add(Folder);
 			}
-			if (CurrentNode.IsValid())
-			{
-				UMovieSceneFolder* Folder = CurrentNode->CastThisChecked<FFolderModel>()->GetFolder();
-				if (Folder && !OutSelectedParentFolders.Contains(Folder))
-				{
-					OutSelectedParentFolders.Add(Folder);
-				}
 
-				// The first valid folder we find will be used to put the new folder into, so it's the node that we
-				// want to know the path from.
-				if (OutNewNodePath.Len() == 0)
-				{
-					// Add an extra delimiter (".") as we know that the new folder will be appended onto the end of this.
-					OutNewNodePath = FString::Printf(TEXT("%s."), *IOutlinerExtension::GetPathName(CurrentNode));
+			// The first valid folder we find will be used to put the new folder into, so it's the node that we
+			// want to know the path from.
+			if (OutNewNodePath.Len() == 0)
+			{
+				// Add an extra delimiter (".") as we know that the new folder will be appended onto the end of this.
+				OutNewNodePath = FString::Printf(TEXT("%s."), *IOutlinerExtension::GetPathName(CurrentNode));
 
-					// Make sure this folder is expanded too so that adding objects to hidden folders become visible.
-					CurrentNode->CastThisChecked<IOutlinerExtension>()->SetExpansion(true);
-				}
+				// Make sure this folder is expanded too so that adding objects to hidden folders become visible.
+				CurrentNode->CastThisChecked<IOutlinerExtension>()->SetExpansion(true);
 			}
 		}
 	}
@@ -9593,7 +9394,7 @@ void FSequencer::AddSelectedActors()
 
 void FSequencer::SetKey()
 {
-	if (Selection.GetSelectedOutlinerItems().Num() > 0)
+	if (ViewModel->GetSelection()->Outliner.Num() > 0)
 	{
 		using namespace UE::Sequencer;
 
@@ -9601,27 +9402,30 @@ void FSequencer::SetKey()
 
 		const FFrameNumber KeyTime = GetLocalTime().Time.FrameNumber;
 
-		FAddKeyOperation::FromNodes(Selection.GetSelectedOutlinerItems()).Commit(KeyTime, *this);
+		FAddKeyOperation::FromNodes(ViewModel->GetSelection()->Outliner.GetSelected()).Commit(KeyTime, *this);
 	}
 }
 
 
 bool FSequencer::CanSetKeyTime() const
 {
-	return Selection.GetSelectedKeys().Num() > 0;
+	return ViewModel->GetSelection()->KeySelection.Num() > 0;
 }
 
 
 void FSequencer::SetKeyTime()
 {
-	TArray<FSequencerSelectedKey> SelectedKeysArray = Selection.GetSelectedKeys().Array();
+	using namespace UE::Sequencer;
 
 	FFrameNumber KeyTime = 0;
-	for ( const FSequencerSelectedKey& Key : SelectedKeysArray )
+
+	const FKeySelection& KeySelection = ViewModel->GetSelection()->KeySelection;
+	for (FKeyHandle Key : KeySelection)
 	{
-		if (Key.IsValid())
+		TSharedPtr<FChannelModel> Channel = KeySelection.GetModelForKey(Key);
+		if (Channel)
 		{
-			KeyTime = Key.WeakChannel.Pin()->GetKeyArea()->GetKeyTime(Key.KeyHandle);
+			KeyTime = Channel->GetKeyArea()->GetKeyTime(Key);
 			break;
 		}
 	}
@@ -9635,6 +9439,8 @@ void FSequencer::SetKeyTime()
 
 void FSequencer::OnSetKeyTimeTextCommitted(const FText& InText, ETextCommit::Type CommitInfo)
 {
+	using namespace UE::Sequencer;
+
 	bool bAnythingChanged = false;
 
 	CloseEntryPopupMenu();
@@ -9647,18 +9453,19 @@ void FSequencer::OnSetKeyTimeTextCommitted(const FText& InText, ETextCommit::Typ
 		FFrameNumber NewFrame = FFrameNumber((int32)NewFrameTime.GetValue());
 
 		FScopedTransaction SetKeyTimeTransaction(NSLOCTEXT("Sequencer", "SetKeyTime_Transaction", "Set Key Time"));
-		TArray<FSequencerSelectedKey> SelectedKeysArray = Selection.GetSelectedKeys().Array();
-	
-		for ( const FSequencerSelectedKey& Key : SelectedKeysArray )
+		const FKeySelection& KeySelection = ViewModel->GetSelection()->KeySelection;
+		for (FKeyHandle Key : KeySelection)
 		{
-			if (Key.IsValid())
+			TSharedPtr<FChannelModel> Channel = KeySelection.GetModelForKey(Key);
+			UMovieSceneSection*       Section = Channel ? Channel->GetSection() : nullptr;
+			if (Channel && Section)
 			{
-	 			if (Key.Section->TryModify())
+	 			if (Section->TryModify())
 	 			{
-	 				Key.WeakChannel.Pin()->GetKeyArea()->SetKeyTime(Key.KeyHandle, NewFrame);
+	 				Channel->GetKeyArea()->SetKeyTime(Key, NewFrame);
 	 				bAnythingChanged = true;
 
-					Key.Section->ExpandToFrame(NewFrame);
+					Section->ExpandToFrame(NewFrame);
 	 			}
 			}
 		}
@@ -9672,29 +9479,34 @@ void FSequencer::OnSetKeyTimeTextCommitted(const FText& InText, ETextCommit::Typ
 
 bool FSequencer::CanRekey() const
 {
-	return Selection.GetSelectedKeys().Num() > 0;
+	return ViewModel->GetSelection()->KeySelection.Num() > 0;
 }
 
 
 void FSequencer::Rekey()
 {
+	using namespace UE::Sequencer;
+
 	bool bAnythingChanged = false;
 
 	FQualifiedFrameTime CurrentTime = GetLocalTime();
 
 	FScopedTransaction RekeyTransaction(NSLOCTEXT("Sequencer", "Rekey_Transaction", "Rekey"));
-	TArray<FSequencerSelectedKey> SelectedKeysArray = Selection.GetSelectedKeys().Array();
-	
-	for ( const FSequencerSelectedKey& Key : SelectedKeysArray )
+	const FKeySelection& KeySelection = ViewModel->GetSelection()->KeySelection;
+
+	for (FKeyHandle Key : KeySelection)
 	{
-		if (Key.IsValid())
+		TSharedPtr<FChannelModel> Channel = KeySelection.GetModelForKey(Key);
+		UMovieSceneSection*       Section = Channel ? Channel->GetSection() : nullptr;
+
+		if (Channel && Section)
 		{
-	 		if (Key.Section->TryModify())
+	 		if (Section->TryModify())
 	 		{
-	 			Key.WeakChannel.Pin()->GetKeyArea()->SetKeyTime(Key.KeyHandle, CurrentTime.Time.FrameNumber);
+	 			Channel->GetKeyArea()->SetKeyTime(Key, CurrentTime.Time.FrameNumber);
 	 			bAnythingChanged = true;
 
-				Key.Section->ExpandToFrame(CurrentTime.Time.FrameNumber);
+				Section->ExpandToFrame(CurrentTime.Time.FrameNumber);
 	 		}
 		}
 	}
@@ -10055,11 +9867,11 @@ void FSequencer::CopySelection()
 {
 	using namespace UE::Sequencer;
 
-	if (Selection.GetSelectedKeys().Num() != 0)
+	if (ViewModel->GetSelection()->KeySelection.Num() != 0)
 	{
 		CopySelectedKeys();
 	}
-	else if (Selection.GetSelectedSections().Num() != 0)
+	else if (ViewModel->GetSelection()->GetSelectedSections().Num() != 0)
 	{
 		CopySelectedSections();
 	}
@@ -10068,15 +9880,10 @@ void FSequencer::CopySelection()
 		TArray<TSharedPtr<FViewModel>> TracksToCopy;
 		TArray<TSharedPtr<FObjectBindingModel>> ObjectsToCopy;
 		TArray<UMovieSceneFolder*> FoldersToCopy;
-		TArray<TSharedRef<FViewModel>> SelectedNodes;
-		Selection.GetNodesWithSelectedKeysOrSections(SelectedNodes);
-		if (SelectedNodes.Num() == 0)
-		{
-			Selection.GetSelectedOutlinerItems(SelectedNodes);
-		}
-		for (TSharedRef<FViewModel> Node : SelectedNodes)
+
+		for (FViewModelPtr Node : ViewModel->GetSelection()->Outliner)
 		{
-			GatherTracksAndObjectsToCopy(Node, TracksToCopy, ObjectsToCopy, FoldersToCopy);
+			GatherTracksAndObjectsToCopy(Node.AsModel().ToSharedRef(), TracksToCopy, ObjectsToCopy, FoldersToCopy);
 		}
 
 		// Make a empty clipboard if the stack is empty
@@ -10116,11 +9923,11 @@ void FSequencer::CopySelection()
 
 void FSequencer::CutSelection()
 {
-	if (Selection.GetSelectedKeys().Num() != 0)
+	if (ViewModel->GetSelection()->KeySelection.Num() != 0)
 	{
 		CutSelectedKeys();
 	}
-	else if (Selection.GetSelectedSections().Num() != 0)
+	else if (ViewModel->GetSelection()->GetSelectedSections().Num() != 0)
 	{
 		CutSelectedSections();
 	}
@@ -10134,11 +9941,15 @@ void FSequencer::CutSelection()
 
 void FSequencer::DuplicateSelection()
 {
+	using namespace UE::Sequencer;
+
 	FScopedTransaction DuplicateSelectionTransaction(LOCTEXT("DuplicateSelection_Transaction", "Duplicate Selection"));
 
 	const bool bClearSelection = true;
 
-	if (Selection.GetSelectedKeys().Num() != 0)
+	const FKeySelection& KeySelection = ViewModel->GetSelection()->KeySelection;
+
+	if (KeySelection.Num() != 0)
 	{
 		CopySelection();
 		DoPaste(bClearSelection);
@@ -10148,34 +9959,25 @@ void FSequencer::DuplicateSelection()
 		// Offset by a visible amount
 		FFrameNumber FrameOffset = FFrameNumber((int32)GetDisplayRateDeltaFrameCount());
 
-		TArray<FSequencerSelectedKey> NewSelection;
-		for (const FSequencerSelectedKey& Key : Selection.GetSelectedKeys())
+		FKeySelection NewSelection;
+		for (FKeyHandle KeyHandle : ViewModel->GetSelection()->KeySelection)
 		{
-			if (Key.IsValid())
+			TSharedPtr<FChannelModel> Channel = KeySelection.GetModelForKey(KeyHandle);
+			if (Channel)
 			{
-				TSharedPtr<IKeyArea> KeyArea = Key.WeakChannel.Pin()->GetKeyArea();
-				FKeyHandle KeyHandle = Key.KeyHandle;
+				TSharedPtr<IKeyArea> KeyArea = Channel->GetKeyArea();
 
 				FKeyHandle NewKeyHandle = KeyArea->DuplicateKey(KeyHandle);
 				KeyArea->SetKeyTime(NewKeyHandle, KeyArea->GetKeyTime(KeyHandle) + FrameOffset);
 
-				NewSelection.Add(FSequencerSelectedKey(*KeyArea->GetOwningSection(), Key.WeakChannel, NewKeyHandle));
+				NewSelection.Select(Channel, NewKeyHandle);
 			}
 		}
 
-		Selection.SuspendBroadcast();
-		Selection.EmptySelectedKeys();
-
-		for (const FSequencerSelectedKey& Key : NewSelection)
-		{
-			Selection.AddToSelection(Key);
-		}
-		Selection.ResumeBroadcast();
-		Selection.GetOnKeySelectionChanged().Broadcast();
-
+		ViewModel->GetSelection()->KeySelection.ReplaceWith(MoveTemp(NewSelection));
 		NotifyMovieSceneDataChanged(EMovieSceneDataChangeType::TrackValueChanged);
 	}
-	else if (Selection.GetSelectedSections().Num() != 0)
+	else if (ViewModel->GetSelection()->GetSelectedSections().Num() != 0)
 	{
 		CopySelection();
 		DoPaste(bClearSelection);
@@ -10205,11 +10007,12 @@ void FSequencer::CopySelectedKeys()
 
 	// Map selected keys to their key areas
 	TMap<TSharedPtr<FChannelModel>, TArray<FKeyHandle>> KeyAreaMap;
-	for (const FSequencerSelectedKey& Key : Selection.GetSelectedKeys())
+	for (FKeyHandle Key : ViewModel->GetSelection()->KeySelection)
 	{
-		if (Key.IsValid())
+		TSharedPtr<FChannelModel> Channel = ViewModel->GetSelection()->KeySelection.GetModelForKey(Key);
+		if (Channel)
 		{
-			KeyAreaMap.FindOrAdd(Key.WeakChannel.Pin()).Add(Key.KeyHandle);
+			KeyAreaMap.FindOrAdd(Channel).Add(Key);
 		}
 	}
 
@@ -10248,7 +10051,7 @@ void FSequencer::CutSelectedKeys()
 void FSequencer::CopySelectedSections()
 {
 	TArray<UMovieSceneSection*> SelectedSections;
-	for (TWeakObjectPtr<UMovieSceneSection> SelectedSectionPtr : Selection.GetSelectedSections())
+	for (TWeakObjectPtr<UMovieSceneSection> SelectedSectionPtr : ViewModel->GetSelection()->GetSelectedSections())
 	{
 		if (SelectedSectionPtr.IsValid())
 		{
@@ -10269,7 +10072,7 @@ void FSequencer::CutSelectedSections()
 {
 	FScopedTransaction CutSelectedSectionsTransaction(LOCTEXT("CutSelectedSections_Transaction", "Cut Selected sections"));
 	CopySelectedSections();
-	DeleteSections(Selection.GetSelectedSections());
+	DeleteSections(ViewModel->GetSelection()->GetSelectedSections());
 }
 
 
@@ -10423,14 +10226,10 @@ void FSequencer::ImportFBXOntoSelectedNodes()
 	// The object binding and names to match when importing from fbx
 	TMap<FGuid, FString> ObjectBindingNameMap;
 
-	for (const TWeakPtr<FViewModel>& Node : Selection.GetSelectedOutlinerItems())
+	for (TViewModelPtr<FObjectBindingModel> ObjectBindingNode : ViewModel->GetSelection()->Outliner.Filter<FObjectBindingModel>())
 	{
-		if (FObjectBindingModel* ObjectBindingNode = ICastable::CastWeakPtr<FObjectBindingModel>(Node))
-		{
-			FGuid ObjectBinding = ObjectBindingNode->GetObjectGuid();
-
-			ObjectBindingNameMap.Add(ObjectBinding, ObjectBindingNode->GetLabel().ToString());
-		}
+		FGuid ObjectBinding = ObjectBindingNode->GetObjectGuid();
+		ObjectBindingNameMap.Add(ObjectBinding, ObjectBindingNode->GetLabel().ToString());
 	}
 
 	MovieSceneToolHelpers::ImportFBXWithDialog(GetFocusedMovieSceneSequence(), *this, ObjectBindingNameMap, TOptional<bool>(false));
@@ -10500,24 +10299,22 @@ void FSequencer::ExportFBX()
 		TArray<FGuid> Bindings;
 		TArray<UMovieSceneTrack*> Tracks;
 		UMovieScene* MovieScene = GetFocusedMovieSceneSequence()->GetMovieScene();
-		for (const TWeakPtr<FViewModel>& Node : Selection.GetSelectedOutlinerItems())
+		for (FViewModelPtr Node : ViewModel->GetSelection()->Outliner)
 		{
-			if (IObjectBindingExtension* ObjectBindingNode = ICastable::CastWeakPtr<IObjectBindingExtension>(Node))
+			if (TViewModelPtr<IObjectBindingExtension> ObjectBindingNode = Node.ImplicitCast())
 			{
 				Bindings.Add(ObjectBindingNode->GetObjectGuid());
 
-				TSet<TSharedRef<FViewModel> > DescendantNodes;
-				SequencerHelpers::GetDescendantNodes(Node.Pin().ToSharedRef(), DescendantNodes);
-				for (TSharedRef<FViewModel> DescendantNode : DescendantNodes)
+				for (TViewModelPtr<IOutlinerExtension> DescendantNode : Node->GetDescendantsOfType<IOutlinerExtension>())
 				{
-					if (!Selection.IsSelected(DescendantNode) && DescendantNode->IsA<IObjectBindingExtension>())
+					if (!ViewModel->GetSelection()->Outliner.IsSelected(DescendantNode) && DescendantNode.AsModel()->IsA<IObjectBindingExtension>())
 					{
-						IObjectBindingExtension* DescendantObjectBindingNode = DescendantNode->CastThisChecked<IObjectBindingExtension>();
+						IObjectBindingExtension* DescendantObjectBindingNode = DescendantNode.AsModel()->CastThisChecked<IObjectBindingExtension>();
 						Bindings.Add(DescendantObjectBindingNode->GetObjectGuid());
 					}
 				}
 			}
-			else if (ITrackExtension* TrackNode = ICastable::CastWeakPtr<ITrackExtension>(Node))
+			else if (TViewModelPtr<ITrackExtension> TrackNode = Node.ImplicitCast())
 			{
 				UMovieSceneTrack* Track = TrackNode->GetTrack();
 				if (Track && MovieScene->ContainsTrack(*Track))
@@ -10645,7 +10442,7 @@ void FSequencer::CloseEntryPopupMenu()
 void FSequencer::TrimSection(bool bTrimLeft)
 {
 	FScopedTransaction TrimSectionTransaction( NSLOCTEXT("Sequencer", "TrimSection_Transaction", "Trim Section") );
-	MovieSceneToolHelpers::TrimSection(Selection.GetSelectedSections(), GetLocalTime(), bTrimLeft, Settings->GetDeleteKeysWhenTrimming());
+	MovieSceneToolHelpers::TrimSection(ViewModel->GetSelection()->GetSelectedSections(), GetLocalTime(), bTrimLeft, Settings->GetDeleteKeysWhenTrimming());
 	NotifyMovieSceneDataChanged( EMovieSceneDataChangeType::TrackValueChanged );
 }
 
@@ -10663,12 +10460,9 @@ void FSequencer::TrimOrExtendSection(bool bTrimOrExtendLeft)
 
 	FScopedTransaction TrimOrExtendSectionTransaction( NSLOCTEXT("Sequencer", "TrimOrExtendSection_Transaction", "Trim or Extend Section") );
 
-	if (Selection.GetSelectedOutlinerItems().Num() > 0)
+	if (ViewModel->GetSelection()->Outliner.Num() > 0)
 	{
-		TArray<TSharedRef<FViewModel>> SelectedNodes;
-		GetSelection().GetSelectedOutlinerItems(SelectedNodes);
-	
-		for (const TSharedRef<FViewModel>& Node : SelectedNodes)
+		for (FViewModelPtr Node : ViewModel->GetSelection()->Outliner)
 		{
 			if (ITrackExtension* TrackNode = Node->CastThis<ITrackExtension>())
 			{
@@ -10713,7 +10507,7 @@ void FSequencer::TrimOrExtendSection(bool bTrimOrExtendLeft)
 void FSequencer::SplitSection()
 {
 	FScopedTransaction SplitSectionTransaction( NSLOCTEXT("Sequencer", "SplitSection_Transaction", "Split Section") );
-	MovieSceneToolHelpers::SplitSection(Selection.GetSelectedSections(), GetLocalTime(), Settings->GetDeleteKeysWhenTrimming());
+	MovieSceneToolHelpers::SplitSection(ViewModel->GetSelection()->GetSelectedSections(), GetLocalTime(), Settings->GetDeleteKeysWhenTrimming());
 	NotifyMovieSceneDataChanged( EMovieSceneDataChangeType::RefreshAllImmediately );
 }
 
@@ -10808,9 +10602,9 @@ void FSequencer::BindCommands()
 		FExecuteAction::CreateSP( this, &FSequencer::TranslateSelectedKeysAndSections, false) );
 
 	auto CanTrimSection = [this]{
-		for (auto Section : Selection.GetSelectedSections())
+		for (UMovieSceneSection* Section : this->ViewModel->GetSelection()->GetSelectedSections())
 		{
-			if (Section.IsValid() && Section->IsTimeWithinSection(GetLocalTime().Time.FrameNumber))
+			if (Section && Section->IsTimeWithinSection(GetLocalTime().Time.FrameNumber))
 			{
 				return true;
 			}
@@ -10852,15 +10646,12 @@ void FSequencer::BindCommands()
 
 		UMovieScene* MovieScene = Sequence->GetMovieScene();
 
-		for (const TWeakPtr<FViewModel>& Node : Selection.GetSelectedOutlinerItems())
+		for (FViewModelPtr Possessable : ViewModel->GetSelection()->Outliner.Filter<FPossessableModel>())
 		{
-			if (IObjectBindingExtension* ObjectBindingNode = ICastable::CastWeakPtr<IObjectBindingExtension>(Node))
+			FViewModelPtr Parent = Possessable->GetParent();
+			if (!Parent || !Parent->IsA<IObjectBindingExtension>())
 			{
-				FMovieScenePossessable* Possessable = MovieScene->FindPossessable(ObjectBindingNode->GetObjectGuid());
-				if (Possessable && !Possessable->GetParent().IsValid())
-				{
-					return true;
-				}
+				return true;
 			}
 		}
 		return false;
@@ -10874,15 +10665,12 @@ void FSequencer::BindCommands()
 	auto AreConvertableSpawnablesSelected = [this] {
 		UMovieScene* MovieScene = GetFocusedMovieSceneSequence()->GetMovieScene();
 
-		for (const TWeakPtr<FViewModel>& Node : Selection.GetSelectedOutlinerItems())
+		for (TViewModelPtr<IObjectBindingExtension> ObjectBindingNode : ViewModel->GetSelection()->Outliner.Filter<IObjectBindingExtension>())
 		{
-			if (IObjectBindingExtension* ObjectBindingNode = ICastable::CastWeakPtr<IObjectBindingExtension>(Node))
+			FMovieSceneSpawnable* Spawnable = MovieScene->FindSpawnable(ObjectBindingNode->GetObjectGuid());
+			if (Spawnable && SpawnRegister->CanConvertSpawnableToPossessable(*Spawnable))
 			{
-				FMovieSceneSpawnable* Spawnable = MovieScene->FindSpawnable(ObjectBindingNode->GetObjectGuid());
-				if (Spawnable && SpawnRegister->CanConvertSpawnableToPossessable(*Spawnable))
-				{
-					return true;
-				}
+				return true;
 			}
 		}
 		return false;
@@ -10897,15 +10685,12 @@ void FSequencer::BindCommands()
 	auto AreSpawnablesSelected = [this] {
 		UMovieScene* MovieScene = GetFocusedMovieSceneSequence()->GetMovieScene();
 
-		for (const TWeakPtr<FViewModel>& Node : Selection.GetSelectedOutlinerItems())
+		for (TViewModelPtr<IObjectBindingExtension> ObjectBindingNode : ViewModel->GetSelection()->Outliner.Filter<IObjectBindingExtension>())
 		{
-			if (IObjectBindingExtension* ObjectBindingNode = ICastable::CastWeakPtr<IObjectBindingExtension>(Node))
+			FMovieSceneSpawnable* Spawnable = MovieScene->FindSpawnable(ObjectBindingNode->GetObjectGuid());
+			if (Spawnable)
 			{
-				FMovieSceneSpawnable* Spawnable = MovieScene->FindSpawnable(ObjectBindingNode->GetObjectGuid());
-				if (Spawnable)
-				{
-					return true;
-				}
+				return true;
 			}
 		}
 		return false;
@@ -11186,14 +10971,14 @@ void FSequencer::BindCommands()
 
 	auto CanCutOrCopy = [this]{
 		// For copy tracks
-		TSet<TWeakPtr<FViewModel>> SelectedNodes = Selection.GetNodesWithSelectedKeysOrSections();
+		TSet<TWeakViewModelPtr<IOutlinerExtension>> SelectedNodes = ViewModel->GetSelection()->GetNodesWithSelectedKeysOrSections();
 		// If this is empty then we are selecting display nodes
 		if (SelectedNodes.Num() == 0)
 		{
-			SelectedNodes = Selection.GetSelectedOutlinerItems();
-			for (TWeakPtr<FViewModel> WeakNode : SelectedNodes)
+			SelectedNodes = ViewModel->GetSelection()->Outliner.GetSelected();
+			for (TWeakViewModelPtr<IOutlinerExtension> WeakNode : SelectedNodes)
 			{
-				TSharedPtr<FViewModel> Node = WeakNode.Pin();
+				FViewModelPtr Node = WeakNode.Pin();
 				if (Node->IsA<ITrackExtension>() || Node->IsA<IObjectBindingExtension>() || Node->IsA<FFolderModel>())
 				{
 					// if contains one node that can be copied we allow the action
@@ -11209,13 +10994,19 @@ void FSequencer::BindCommands()
 		}
 
 		UMovieSceneTrack* Track = nullptr;
-		for (FSequencerSelectedKey Key : Selection.GetSelectedKeys())
+		for (FKeyHandle Key : ViewModel->GetSelection()->KeySelection)
 		{
+			TSharedPtr<FChannelModel> Channel = ViewModel->GetSelection()->KeySelection.GetModelForKey(Key);
+			if (!Channel)
+			{
+				continue;
+			}
+
 			if (!Track)
 			{
-				Track = Key.Section->GetTypedOuter<UMovieSceneTrack>();
+				Track = Channel->GetSection()->GetTypedOuter<UMovieSceneTrack>();
 			}
-			if (!Track || Track != Key.Section->GetTypedOuter<UMovieSceneTrack>())
+			if (!Track || Track != Channel->GetSection()->GetTypedOuter<UMovieSceneTrack>())
 			{
 				return false;
 			}
@@ -11224,29 +11015,24 @@ void FSequencer::BindCommands()
 	};
 
 	auto CanDelete = [this]{
-		return Selection.GetSelectedKeys().Num() || Selection.GetSelectedSections().Num() || Selection.GetSelectedOutlinerItems().Num();
+		TSharedPtr<FSequencerSelection> Selection = this->ViewModel->GetSelection();
+		return Selection->KeySelection.Num() || Selection->GetSelectedSections().Num() || Selection->Outliner.Num();
 	};
 
 	auto CanDuplicate = [this]{
-
-		if (Selection.GetSelectedKeys().Num() || Selection.GetSelectedSections().Num() || Selection.GetSelectedTracks().Num())
+		TSharedPtr<FSequencerSelection> Selection = this->ViewModel->GetSelection();
+		if (Selection->KeySelection.Num() || Selection->GetSelectedSections().Num() || Selection->GetSelectedTracks().Num())
 		{
 			return true;
 		}
 
 		// For duplicate object tracks
-		TArray<TSharedRef<FViewModel>> SelectedNodes;
-		Selection.GetNodesWithSelectedKeysOrSections(SelectedNodes);
-		if (SelectedNodes.Num() == 0)
+		if (Selection->GetNodesWithSelectedKeysOrSections().Num() == 0)
 		{
-			Selection.GetSelectedOutlinerItems(SelectedNodes);
-			for (TSharedRef<FViewModel> Node : SelectedNodes)
+			// if contains one node that can be copied we allow the action
+			for (TViewModelPtr<IObjectBindingExtension> Binding : Selection->Outliner.Filter<IObjectBindingExtension>())
 			{
-				if (Node->IsA<IObjectBindingExtension>())
-				{
-					// if contains one node that can be copied we allow the action
-					return true;
-				}
+				return true;
 			}
 			return false;
 		}
@@ -11267,10 +11053,11 @@ void FSequencer::BindCommands()
 		FGenericCommands::Get().Rename,
 		FExecuteAction::CreateLambda([this]
 		{
-			for (TWeakPtr<FViewModel> Item : Selection.GetSelectedOutlinerItems())
+			TSharedPtr<FSequencerSelection> Selection = this->ViewModel->GetSelection();
+
+			for (TViewModelPtr<IRenameableExtension> Renamable : Selection->Outliner.Filter<IRenameableExtension>())
 			{
-				IRenameableExtension* Renamable = ICastable::CastWeakPtr<IRenameableExtension>(Item);
-				if (Renamable && Renamable->CanRename())
+				if (Renamable->CanRename())
 				{
 					Renamable->OnRenameRequested().Broadcast();
 				}
@@ -11278,10 +11065,11 @@ void FSequencer::BindCommands()
 		}),
 		FCanExecuteAction::CreateLambda([this]
 		{
-			for (TWeakPtr<FViewModel> Item : Selection.GetSelectedOutlinerItems())
+			TSharedPtr<FSequencerSelection> Selection = this->ViewModel->GetSelection();
+
+			for (TViewModelPtr<IRenameableExtension> Renamable : Selection->Outliner.Filter<IRenameableExtension>())
 			{
-				IRenameableExtension* Renamable = ICastable::CastWeakPtr<IRenameableExtension>(Item);
-				if (Renamable && Renamable->CanRename())
+				if (Renamable->CanRename())
 				{
 					return true;
 				}
@@ -11772,14 +11560,11 @@ void FSequencer::BuildAddSelectedToFolderMenu(FMenuBuilder& MenuBuilder)
 	if (MovieScene)
 	{
 		TSharedRef<TArray<UMovieSceneFolder*>> ExcludedFolders = MakeShared<TArray<UMovieSceneFolder*> >();
-		for (TWeakPtr<FViewModel> Node : GetSelection().GetSelectedOutlinerItems())
+		for (TViewModelPtr<FFolderModel> FolderNode : ViewModel->GetSelection()->Outliner.Filter<FFolderModel>())
 		{
-			if (FFolderModel* FolderNode = ICastable::CastWeakPtr<FFolderModel>(Node))
+			if (FolderNode->CanDrag())
 			{
-				if (FolderNode->CanDrag())
-				{
-					ExcludedFolders->Add(FolderNode->GetFolder());
-				}
+				ExcludedFolders->Add(FolderNode->GetFolder());
 			}
 		}
 
@@ -12020,16 +11805,13 @@ void FSequencer::SetDisplayName(FGuid Binding, const FText& InDisplayName)
 {
 	using namespace UE::Sequencer;
 
-	for (const TWeakPtr<FViewModel>& Node : Selection.GetSelectedOutlinerItems())
+	for (TViewModelPtr<FObjectBindingModel> ObjectBindingNode : ViewModel->GetSelection()->Outliner.Filter<FObjectBindingModel>())
 	{
-		if (FObjectBindingModel* ObjectBindingNode = ICastable::CastWeakPtr<FObjectBindingModel>(Node))
+		FGuid Guid = ObjectBindingNode->GetObjectGuid();
+		if (Guid == Binding)
 		{
-			FGuid Guid = ObjectBindingNode->GetObjectGuid();
-			if (Guid == Binding)
-			{
-				ObjectBindingNode->Rename(InDisplayName);
-				break;
-			}
+			ObjectBindingNode->Rename(InDisplayName);
+			break;
 		}
 	}
 }
@@ -12038,15 +11820,12 @@ FText FSequencer::GetDisplayName(FGuid Binding)
 {
 	using namespace UE::Sequencer;
 
-	for (const TWeakPtr<FViewModel>& Node : Selection.GetSelectedOutlinerItems())
+	for (TViewModelPtr<FObjectBindingModel> ObjectBindingNode : ViewModel->GetSelection()->Outliner.Filter<FObjectBindingModel>())
 	{
-		if (FObjectBindingModel* ObjectBindingNode = ICastable::CastWeakPtr<FObjectBindingModel>(Node))
+		FGuid Guid = ObjectBindingNode->GetObjectGuid();
+		if (Guid == Binding)
 		{
-			FGuid Guid = ObjectBindingNode->GetObjectGuid();
-			if (Guid == Binding)
-			{
-				return ObjectBindingNode->GetLabel();
-			}
+			return ObjectBindingNode->GetLabel();
 		}
 	}
 	return FText();
@@ -12093,5 +11872,10 @@ bool FSequencer::UsesDynamicWeighting()
 	return EnumHasAnyFlags(GetFocusedMovieSceneSequence()->GetFlags(), EMovieSceneSequenceFlags::DynamicWeighting);
 }
 
+UE::Sequencer::FSequencerSelection& FSequencer::GetSelection()
+{
+	return *ViewModel->GetSelection();
+}
+
 #undef LOCTEXT_NAMESPACE
 
