Sequencer: New selection API

diff --git a/Engine/Source/Editor/Sequencer/Private/Sequencer.h b/Engine/Source/Editor/Sequencer/Private/Sequencer.h
--- a/Engine/Source/Editor/Sequencer/Private/Sequencer.h
+++ b/Engine/Source/Editor/Sequencer/Private/Sequencer.h
@@ -27,7 +27,6 @@
 #include "ISequencer.h"
 #include "ISequencerModule.h"
 #include "ISequencerObjectChangeListener.h"
-#include "SequencerSelection.h"
 #include "SequencerSelectionPreview.h"
 #include "SequencerCustomizationManager.h"
 #include "ITransportControl.h"
@@ -54,7 +53,6 @@ class IMenu;
 class FCurveEditor;
 class ISequencerEditTool;
 class FSequencerKeyCollection;
-class FSequencerOutlinerSelectionHandler;
 class FObjectBindingTagCache;
 class ISequencerTrackEditor;
 class ISequencerEditorObjectBinding;
@@ -96,6 +94,7 @@ namespace UE
 		class FSequenceModel;
 		class FSequencerEditorViewModel;
 		class FViewModel;
+		class FSequencerSelection;
 
 	} // namespace Sequencer
 } // namespace UE
@@ -251,7 +250,7 @@ public:
 	virtual void PopToSequenceInstance( FMovieSceneSequenceIDRef SequenceID ) override;
 
 	/** Deletes the passed in sections. */
-	void DeleteSections(const TSet<TWeakObjectPtr<UMovieSceneSection> > & Sections);
+	void DeleteSections(const TSet<UMovieSceneSection*> & Sections);
 
 	/** Deletes the currently selected in keys. */
 	void DeleteSelectedKeys();
@@ -557,10 +556,6 @@ public:
 	bool PasteSections(const FString& TextToImport, TArray<FNotificationInfo>& PasteErrors);
 	bool PasteObjectBindings(const FString& TextToImport, UMovieSceneFolder* ParentFolder, const TArray<UMovieSceneFolder*>& InFolders, TArray<FNotificationInfo>& PasteErrors, bool bClearSelection = false);
 
-	/** Called when a user executes the active node menu item */
-	void ToggleNodeActive();
-	bool IsNodeActive() const;
-
 	/** Called when a user executes the locked node menu item */
 	void ToggleNodeLocked();
 	bool IsNodeLocked() const;
@@ -704,7 +699,7 @@ public:
 	virtual FMovieSceneRootEvaluationTemplateInstance& GetEvaluationTemplate() override { return RootTemplateInstance; }
 	virtual void ResetToNewRootSequence(UMovieSceneSequence& NewSequence) override;
 	virtual void FocusSequenceInstance(UMovieSceneSubSection& InSubSection) override;
-	virtual TSharedPtr<UE::Sequencer::FEditorViewModel> GetViewModel() const override;
+	virtual TSharedPtr<UE::Sequencer::FSequencerEditorViewModel> GetViewModel() const override;
 	virtual void SuppressAutoEvaluation(UMovieSceneSequence* Sequence, const FGuid& InSequenceSignature) override;
 	virtual EAutoChangeMode GetAutoChangeMode() const override;
 	virtual void SetAutoChangeMode(EAutoChangeMode AutoChangeMode) override;
@@ -753,10 +748,6 @@ public:
 	virtual void AddSubSequence(UMovieSceneSequence* Sequence) override;
 	virtual bool CanKeyProperty(FCanKeyPropertyParams CanKeyPropertyParams) const override;
 	virtual void KeyProperty(FKeyPropertyParams KeyPropertyParams) override;
-	virtual FSequencerSelection& GetSelection() override;
-	virtual FSequencerSelectionPreview& GetSelectionPreview() override;
-	virtual void SuspendSelectionBroadcast() override;
-	virtual void ResumeSelectionBroadcast() override;
 	virtual void GetSelectedTracks(TArray<UMovieSceneTrack*>& OutSelectedTracks) override;
 	virtual void GetSelectedSections(TArray<UMovieSceneSection*>& OutSelectedSections) override;
 	virtual void GetSelectedFolders(TArray<UMovieSceneFolder*>& OutSelectedFolders) override;
@@ -833,6 +824,10 @@ public:
 	virtual bool IsPreview() const override { return SilentModeCount != 0; }
 	virtual UMovieSceneEntitySystemLinker* ConstructEntitySystemLinker() override;
 
+	/** Shortcut for GetEditorModel()->GetSelection() for backwards compat with existing code */
+	UE::Sequencer::FSequencerSelection& GetSelection();
+	FSequencerSelectionPreview& GetSelectionPreview();
+
 protected:
 
 	/** Reevaluate the sequence at the current time */
@@ -1033,8 +1028,8 @@ protected:
 	virtual void PostRedo(bool bSuccess) override { PostUndo(bSuccess); }
 	// End of FEditorUndoClient
 
-	void OnSelectedOutlinerNodesChanged();
-	// Called on Tick after OnSelectedOutlinerNodesChanged has been called
+	void OnSelectionChanged();
+	// Called on Tick after OnSelectionChanged has been called
 	void HandleSelectedOutlinerNodesChanged();
 
 	void AddNodeGroupsCollectionChangedDelegate();
@@ -1045,7 +1040,7 @@ protected:
 public:
 	void AddSelectedNodesToNewNodeGroup();
 	void AddSelectedNodesToExistingNodeGroup(UMovieSceneNodeGroup* NodeGroup);
-	void AddNodesToExistingNodeGroup(const TArray<TWeakPtr<UE::Sequencer::FViewModel>>& Nodes, UMovieSceneNodeGroup* NodeGroup);
+	void AddNodesToExistingNodeGroup(TArrayView<const UE::Sequencer::TWeakViewModelPtr<UE::Sequencer::IOutlinerExtension>> InItems, UMovieSceneNodeGroup* NodeGroup);
 
 	void ClearFilters();
 
@@ -1177,9 +1172,6 @@ private:
 	/** Main sequencer widget */
 	TSharedPtr<SSequencer> SequencerWidget;
 
-	/** Selection handler for interfacing with sequencer selection */
-	TSharedPtr<FSequencerOutlinerSelectionHandler> SelectionHandler;
-
 	/** Spawn register for keeping track of what is spawned */
 	TSharedPtr<FMovieSceneSpawnRegister> SpawnRegister;
 
@@ -1303,7 +1295,7 @@ private:
 		the MovieScene data can change many times per frame.) */
 	bool bNeedTreeRefresh;
 
-	FSequencerSelection Selection;
+	//FSequencerSelection Selection;
 	FSequencerSelectionPreview SelectionPreview;
 
 	/** Represents the tree of nodes to display in the animation outliner. */
diff --git a/Engine/Source/Editor/Sequencer/Private/SequencerAddKeyOperation.cpp b/Engine/Source/Editor/Sequencer/Private/SequencerAddKeyOperation.cpp
--- a/Engine/Source/Editor/Sequencer/Private/SequencerAddKeyOperation.cpp
+++ b/Engine/Source/Editor/Sequencer/Private/SequencerAddKeyOperation.cpp
@@ -31,19 +31,19 @@ namespace UE
 namespace Sequencer
 {
 
-FAddKeyOperation FAddKeyOperation::FromNodes(const TSet<TWeakPtr<FViewModel>>& InNodes)
+FAddKeyOperation FAddKeyOperation::FromNodes(const TSet<TWeakViewModelPtr<IOutlinerExtension>>& InNodes)
 {
 	FAddKeyOperation Operation;
 
 	TArray<TWeakPtr<FViewModel>> FilteredNodes;
 
 	// Remove any child nodes that have a parent also included in the set
-	for (const TWeakPtr<FViewModel>& ProspectiveNode : InNodes)
+	for (const TWeakViewModelPtr<IOutlinerExtension>& ProspectiveNode : InNodes)
 	{
-		TSharedPtr<FViewModel> Parent = ProspectiveNode.Pin()->GetParent();
+		TSharedPtr<FViewModel> Parent = ProspectiveNode.Pin().AsModel()->GetParent();
 		while (Parent)
 		{
-			if (InNodes.Contains(Parent.ToSharedRef()))
+			if (InNodes.Contains(CastViewModel<IOutlinerExtension>(Parent)))
 			{
 				goto Continue;
 			}
diff --git a/Engine/Source/Editor/Sequencer/Private/SequencerCommonHelpers.cpp b/Engine/Source/Editor/Sequencer/Private/SequencerCommonHelpers.cpp
--- a/Engine/Source/Editor/Sequencer/Private/SequencerCommonHelpers.cpp
+++ b/Engine/Source/Editor/Sequencer/Private/SequencerCommonHelpers.cpp
@@ -17,6 +17,7 @@
 #include "MVVM/ViewModels/SequencerEditorViewModel.h"
 #include "MVVM/ViewModels/VirtualTrackArea.h"
 #include "MVVM/Views/ITrackAreaHotspot.h"
+#include "MVVM/Selection/Selection.h"
 #include "Modules/ModuleManager.h"
 #include "MovieScene.h"
 #include "MovieSceneSectionDetailsCustomization.h"
@@ -170,6 +171,8 @@ bool IsSectionSelectedInNode(FSequencer& Sequencer, TSharedPtr<UE::Sequencer::FV
 {
 	using namespace UE::Sequencer;
 
+	const FTrackAreaSelection& Selection = Sequencer.GetViewModel()->GetSelection()->TrackArea;
+
 	if (ITrackAreaExtension* TrackArea = InModel->CastThis<ITrackAreaExtension>())
 	{
 		for (TSharedPtr<FViewModel> TrackAreaModel : TrackArea->GetTrackAreaModelList())
@@ -177,7 +180,7 @@ bool IsSectionSelectedInNode(FSequencer& Sequencer, TSharedPtr<UE::Sequencer::FV
 			constexpr bool bIncludeThis = true;
 			for (TSharedPtr<FSectionModel> Section : TParentFirstChildIterator<FSectionModel>(TrackAreaModel, bIncludeThis))
 			{
-				if (Sequencer.GetSelection().IsSelected(Section))
+				if (Selection.IsSelected(Section))
 				{
 					return true;
 				}
@@ -190,12 +193,17 @@ bool IsSectionSelectedInNode(FSequencer& Sequencer, TSharedPtr<UE::Sequencer::FV
 
 bool AreKeysSelectedInNode(FSequencer& Sequencer, TSharedPtr<UE::Sequencer::FViewModel> InModel)
 {
-	TSet<TSharedPtr<UE::Sequencer::FChannelModel>> Channels;
+	using namespace UE::Sequencer;
+
+	TSet<TSharedPtr<FChannelModel>> Channels;
 	SequencerHelpers::GetAllChannels(InModel, Channels);
 
-	for (const FSequencerSelectedKey& Key : Sequencer.GetSelection().GetSelectedKeys())
+	const FKeySelection& Selection = Sequencer.GetViewModel()->GetSelection()->KeySelection;
+
+	for (FKeyHandle Key : Selection)
 	{
-		if (Channels.Contains(Key.WeakChannel.Pin()))
+		TSharedPtr<FChannelModel> Channel = Selection.GetModelForKey(Key);
+		if (Channels.Contains(Channel))
 		{
 			return true;
 		}
diff --git a/Engine/Source/Editor/Sequencer/Private/SequencerContextMenus.cpp b/Engine/Source/Editor/Sequencer/Private/SequencerContextMenus.cpp
--- a/Engine/Source/Editor/Sequencer/Private/SequencerContextMenus.cpp
+++ b/Engine/Source/Editor/Sequencer/Private/SequencerContextMenus.cpp
@@ -46,6 +46,7 @@
 #include "MVVM/ViewModels/TrackModel.h"
 #include "MVVM/ViewModels/ViewModel.h"
 #include "MVVM/ViewModels/SequencerEditorViewModel.h"
+#include "MVVM/Selection/Selection.h"
 #include "Tracks/MovieScenePropertyTrack.h"
 #include "Algo/AnyOf.h"
 #include "IKeyArea.h"
@@ -55,16 +56,19 @@
 
 static void CreateKeyStructForSelection(TSharedPtr<ISequencer> InSequencer, TSharedPtr<FStructOnScope>& OutKeyStruct, TWeakObjectPtr<UMovieSceneSection>& OutKeyStructSection)
 {
-	const TSet<FSequencerSelectedKey>& SelectedKeys = InSequencer->GetSelection().GetSelectedKeys();
+	using namespace UE::Sequencer;
+
+	const FKeySelection& SelectedKeys = InSequencer->GetViewModel()->GetSelection()->KeySelection;
 
 	if (SelectedKeys.Num() == 1)
 	{
-		for (const FSequencerSelectedKey& Key : SelectedKeys)
+		for (FKeyHandle Key : SelectedKeys)
 		{
-			if (Key.IsValid())
+			TSharedPtr<FChannelModel> Channel = SelectedKeys.GetModelForKey(Key);
+			if (Channel)
 			{
-				OutKeyStruct = Key.WeakChannel.Pin()->GetKeyArea()->GetKeyStruct(Key.KeyHandle);
-				OutKeyStructSection = Key.WeakChannel.Pin()->GetSection();
+				OutKeyStruct = Channel->GetKeyArea()->GetKeyStruct(Key);
+				OutKeyStructSection = Channel->GetSection();
 				return;
 			}
 		}
@@ -73,17 +77,18 @@ static void CreateKeyStructForSelection(TSharedPtr<ISequencer> InSequencer, TSha
 	{
 		TArray<FKeyHandle> KeyHandles;
 		UMovieSceneSection* CommonSection = nullptr;
-		for (const FSequencerSelectedKey& Key : SelectedKeys)
+		for (FKeyHandle Key : SelectedKeys)
 		{
-			if (Key.IsValid())
+			TSharedPtr<FChannelModel> Channel = SelectedKeys.GetModelForKey(Key);
+			if (Channel)
 			{
-				KeyHandles.Add(Key.KeyHandle);
+				KeyHandles.Add(Key);
 
 				if (!CommonSection)
 				{
-					CommonSection = Key.WeakChannel.Pin()->GetSection();
+					CommonSection = Channel->GetSection();
 				}
-				else if (CommonSection != Key.WeakChannel.Pin()->GetSection())
+				else if (CommonSection != Channel->GetSection())
 				{
 					CommonSection = nullptr;
 					return;
@@ -117,7 +122,7 @@ void FKeyContextMenu::PopulateMenu(FMenuBuilder& MenuBuilder, TSharedPtr<FExtend
 	{
 		ISequencerModule& SequencerModule = FModuleManager::LoadModuleChecked<ISequencerModule>("Sequencer");
 
-		FSelectedKeysByChannel SelectedKeysByChannel(SequencerPtr->GetSelection().GetSelectedKeys().Array());
+		FSelectedKeysByChannel SelectedKeysByChannel(SequencerPtr->GetViewModel()->GetSelection()->KeySelection);
 
 		TMap<FName, TArray<FExtendKeyMenuParams>> ChannelAndHandlesByType;
 		for (FSelectedChannelInfo& ChannelInfo : SelectedKeysByChannel.SelectedChannels)
@@ -198,25 +203,22 @@ FSectionContextMenu::FSectionContextMenu(FSequencer& InSeqeuncer, FFrameTime InM
 	: Sequencer(StaticCastSharedRef<FSequencer>(InSeqeuncer.AsShared()))
 	, MouseDownTime(InMouseDownTime)
 {
-	for (TWeakObjectPtr<UMovieSceneSection> WeakSection : Sequencer->GetSelection().GetSelectedSections())
+	for (UMovieSceneSection* Section : Sequencer->GetViewModel()->GetSelection()->GetSelectedSections())
 	{
-		if (UMovieSceneSection* Section = WeakSection.Get())
+		FMovieSceneChannelProxy& ChannelProxy = Section->GetChannelProxy();
+		for (const FMovieSceneChannelEntry& Entry : ChannelProxy.GetAllEntries())
 		{
-			FMovieSceneChannelProxy& ChannelProxy = Section->GetChannelProxy();
-			for (const FMovieSceneChannelEntry& Entry : ChannelProxy.GetAllEntries())
-			{
-				FName ChannelTypeName = Entry.GetChannelTypeName();
+			FName ChannelTypeName = Entry.GetChannelTypeName();
 
-				TArray<UMovieSceneSection*>& SectionArray = SectionsByType.FindOrAdd(ChannelTypeName);
-				SectionArray.Add(Section);
+			TArray<UMovieSceneSection*>& SectionArray = SectionsByType.FindOrAdd(ChannelTypeName);
+			SectionArray.Add(Section);
 
-				TArray<FMovieSceneChannelHandle>& ChannelHandles = ChannelsByType.FindOrAdd(ChannelTypeName);
+			TArray<FMovieSceneChannelHandle>& ChannelHandles = ChannelsByType.FindOrAdd(ChannelTypeName);
 
-				const int32 NumChannels = Entry.GetChannels().Num();
-				for (int32 Index = 0; Index < NumChannels; ++Index)
-				{
-					ChannelHandles.Add(ChannelProxy.MakeHandle(ChannelTypeName, Index));
-				}
+			const int32 NumChannels = Entry.GetChannels().Num();
+			for (int32 Index = 0; Index < NumChannels; ++Index)
+			{
+				ChannelHandles.Add(ChannelProxy.MakeHandle(ChannelTypeName, Index));
 			}
 		}
 	}
@@ -231,6 +233,8 @@ void FSectionContextMenu::BuildMenu(FMenuBuilder& MenuBuilder, TSharedPtr<FExten
 
 void FSectionContextMenu::PopulateMenu(FMenuBuilder& MenuBuilder, TSharedPtr<FExtender> MenuExtender)
 {
+	using namespace UE::Sequencer;
+
 	// Copy a reference to the context menu by value into each lambda handler to ensure the type stays alive until the menu is closed
 	TSharedRef<FSectionContextMenu> Shared = AsShared();
 
@@ -260,15 +264,14 @@ void FSectionContextMenu::PopulateMenu(FMenuBuilder& MenuBuilder, TSharedPtr<FEx
 		FNewMenuDelegate::CreateLambda([=](FMenuBuilder& SubMenuBuilder)
 		{
 			TArray<TWeakObjectPtr<UObject>> Sections;
+			for (TViewModelPtr<FSectionModel> SectionModel : Sequencer->GetViewModel()->GetSelection()->TrackArea.Filter<FSectionModel>())
 			{
-				for (TWeakObjectPtr<UMovieSceneSection> Section : Sequencer->GetSelection().GetSelectedSections())
+				if (UMovieSceneSection* Section = SectionModel->GetSection())
 				{
-					if (Section.IsValid())
-					{
-						Sections.Add(Section);
-					}
+					Sections.Add(Section);
 				}
 			}
+
 			SequencerHelpers::AddPropertiesMenu(*Sequencer, SubMenuBuilder, Sections);
 		})
 	);
@@ -570,13 +573,10 @@ FMovieSceneBlendTypeField FSectionContextMenu::GetSupportedBlendTypes() const
 {
 	FMovieSceneBlendTypeField BlendTypes = FMovieSceneBlendTypeField::All();
 
-	for (TWeakObjectPtr<UMovieSceneSection> WeakSection : Sequencer->GetSelection().GetSelectedSections())
+	for (UMovieSceneSection* Section : Sequencer->GetViewModel()->GetSelection()->GetSelectedSections())
 	{
-		if (UMovieSceneSection* Section = WeakSection.Get())
-		{
-			// Remove unsupported blend types
-			BlendTypes.Remove(Section->GetSupportedBlendTypes().Invert());
-		}
+		// Remove unsupported blend types
+		BlendTypes.Remove(Section->GetSupportedBlendTypes().Invert());
 	}
 
 	return BlendTypes;
@@ -602,13 +602,14 @@ void FSectionContextMenu::AddOrderMenu(FMenuBuilder& MenuBuilder)
 
 void FSectionContextMenu::AddBlendTypeMenu(FMenuBuilder& MenuBuilder)
 {
-	TArray<TWeakObjectPtr<UMovieSceneSection>> Sections;
+	using namespace UE::Sequencer;
 
-	for (TWeakObjectPtr<UMovieSceneSection> WeakSection : Sequencer->GetSelection().GetSelectedSections())
+	TArray<TWeakObjectPtr<UMovieSceneSection>> Sections;
+	for (TViewModelPtr<FSectionModel> SectionModel : Sequencer->GetViewModel()->GetSelection()->TrackArea.Filter<FSectionModel>())
 	{
-		if (WeakSection.IsValid())
+		if (UMovieSceneSection* Section = SectionModel->GetSection())
 		{
-			Sections.Add(WeakSection);
+			Sections.Add(Section);
 		}
 	}
 
@@ -622,32 +623,23 @@ void FSectionContextMenu::SelectAllKeys()
 
 	TArray<FKeyHandle> HandlesScratch;
 
-	FSequencerSelection& Selection = Sequencer->GetSelection();
-	for (const TWeakPtr<FViewModel>& WeakModel : Selection.GetSelectedTrackAreaItems())
+	FSequencerSelection& Selection = *Sequencer->GetViewModel()->GetSelection();
+	for (TViewModelPtr<FSectionModel> Section : Selection.TrackArea.Filter<FSectionModel>())
 	{
-		TViewModelPtr<FSectionModel> Section = CastViewModel<FSectionModel>(WeakModel.Pin());
-		if (Section)
+		UMovieSceneSection* SectionObject = Section->GetSection();
+		if (SectionObject)
 		{
-			UMovieSceneSection* SectionObject = Section->GetSection();
-			if (SectionObject)
+			// Iterate all channels
+			for (const TViewModelPtr<FChannelModel>& Channel : Section->GetDescendantsOfType<FChannelModel>())
 			{
-				FSequencerSelectedKey SelectedKey(*SectionObject, nullptr, FKeyHandle::Invalid());
-
-				// Iterate all channels
-				for (const TViewModelPtr<FChannelModel>& Channel : Section->GetDescendantsOfType<FChannelModel>())
+				if (Channel->GetLinkedOutlinerItem() && !Channel->GetLinkedOutlinerItem()->IsFilteredOut())
 				{
-					if (Channel->GetLinkedOutlinerItem() && !Channel->GetLinkedOutlinerItem()->IsFilteredOut())
-					{
-						SelectedKey.WeakChannel = Channel;
+					HandlesScratch.Reset();
+					Channel->GetKeyArea()->GetKeyHandles(HandlesScratch);
 
-						HandlesScratch.Reset();
-						Channel->GetKeyArea()->GetKeyHandles(HandlesScratch);
-
-						for (FKeyHandle KeyHandle : HandlesScratch)
-						{
-							SelectedKey.KeyHandle = KeyHandle;
-							Selection.AddToSelection(SelectedKey);
-						}
+					for (FKeyHandle KeyHandle : HandlesScratch)
+					{
+						Selection.KeySelection.Select(Channel, KeyHandle);
 					}
 				}
 			}
@@ -663,39 +655,34 @@ void FSectionContextMenu::CopyAllKeys()
 
 void FSectionContextMenu::SetSectionToKey()
 {
-	if (Sequencer->GetSelection().GetSelectedSections().Num() != 1)
+	if (Sequencer->GetViewModel()->GetSelection()->GetSelectedSections().Num() != 1)
 	{
 		return;
 	}
 
 	const bool bToggle = IsSectionToKey();
-	for (TWeakObjectPtr<UMovieSceneSection> WeakSection : Sequencer->GetSelection().GetSelectedSections())
+	for (UMovieSceneSection* Section : Sequencer->GetViewModel()->GetSelection()->GetSelectedSections())
 	{
-		if (UMovieSceneSection* Section = WeakSection.Get())
+		UMovieSceneTrack* Track = Section->GetTypedOuter<UMovieSceneTrack>();
+		if (Track)
 		{
-			UMovieSceneTrack* Track = Section->GetTypedOuter<UMovieSceneTrack>();
-			if (Track)
-			{
-				FScopedTransaction Transaction(LOCTEXT("SetSectionToKey", "Set Section To Key"));
-				Track->Modify();
-				Track->SetSectionToKey(bToggle ? nullptr : Section);
-			}
+			FScopedTransaction Transaction(LOCTEXT("SetSectionToKey", "Set Section To Key"));
+			Track->Modify();
+			Track->SetSectionToKey(bToggle ? nullptr : Section);
 		}
+
 		break;
 	}
 }
 
 bool FSectionContextMenu::IsSectionToKey() const
 {
-	for (TWeakObjectPtr<UMovieSceneSection> WeakSection : Sequencer->GetSelection().GetSelectedSections())
+	for (UMovieSceneSection* Section : Sequencer->GetViewModel()->GetSelection()->GetSelectedSections())
 	{
-		if (UMovieSceneSection* Section = WeakSection.Get())
+		UMovieSceneTrack* Track = Section->GetTypedOuter<UMovieSceneTrack>();
+		if (Track && Track->GetSectionToKey() != Section)
 		{
-			UMovieSceneTrack* Track = Section->GetTypedOuter<UMovieSceneTrack>();
-			if (Track && Track->GetSectionToKey() != Section)
-			{
-				return false;
-			}
+			return false;
 		}
 	}
 	return true;
@@ -703,20 +690,17 @@ bool FSectionContextMenu::IsSectionToKey() const
 
 bool FSectionContextMenu::CanSetSectionToKey() const
 {
-	if (Sequencer->GetSelection().GetSelectedSections().Num() != 1)
+	if (Sequencer->GetViewModel()->GetSelection()->GetSelectedSections().Num() != 1)
 	{
 		return false;
 	}
 
-	for (TWeakObjectPtr<UMovieSceneSection> WeakSection : Sequencer->GetSelection().GetSelectedSections())
+	for (UMovieSceneSection* Section : Sequencer->GetViewModel()->GetSelection()->GetSelectedSections())
 	{
-		if (UMovieSceneSection* Section = WeakSection.Get())
+		UMovieSceneTrack* Track = Section->GetTypedOuter<UMovieSceneTrack>();
+		if (Track && Section->GetBlendType().IsValid() && (Section->GetBlendType().Get() == EMovieSceneBlendType::Absolute || Section->GetBlendType().Get() == EMovieSceneBlendType::Additive))
 		{
-			UMovieSceneTrack* Track = Section->GetTypedOuter<UMovieSceneTrack>();
-			if (Track && Section->GetBlendType().IsValid() && (Section->GetBlendType().Get() == EMovieSceneBlendType::Absolute || Section->GetBlendType().Get() == EMovieSceneBlendType::Additive))
-			{
-				return true;
-			}
+			return true;
 		}
 
 		break;
@@ -744,9 +728,9 @@ void FSectionContextMenu::AutoSizeSection()
 {
 	FScopedTransaction AutoSizeSectionTransaction(LOCTEXT("AutoSizeSection_Transaction", "Auto Size Section"));
 
-	for (auto Section : Sequencer->GetSelection().GetSelectedSections())
+	for (UMovieSceneSection* Section : Sequencer->GetViewModel()->GetSelection()->GetSelectedSections())
 	{
-		if (Section.IsValid() && Section->GetAutoSizeRange().IsSet())
+		if (Section && Section->GetAutoSizeRange().IsSet())
 		{
 			TOptional<TRange<FFrameNumber> > DefaultSectionLength = Section->GetAutoSizeRange();
 
@@ -768,14 +752,14 @@ void FSectionContextMenu::ReduceKeys()
 	FScopedTransaction ReduceKeysTransaction(LOCTEXT("ReduceKeys_Transaction", "Reduce Keys"));
 
 	TSet<TSharedPtr<IKeyArea> > KeyAreas;
-	for (const TWeakPtr<FViewModel>& WeakItem : Sequencer->GetSelection().GetSelectedOutlinerItems())
+	for (TViewModelPtr<IOutlinerExtension> Item : Sequencer->GetViewModel()->GetSelection()->Outliner)
 	{
-		SequencerHelpers::GetAllKeyAreas(WeakItem.Pin(), KeyAreas);
+		SequencerHelpers::GetAllKeyAreas(Item, KeyAreas);
 	}
 
 	if (KeyAreas.Num() == 0)
 	{
-		for (const TWeakPtr<FViewModel>& DisplayNode : Sequencer->GetSelection().GetNodesWithSelectedKeysOrSections())
+		for (const TWeakViewModelPtr<IOutlinerExtension>& DisplayNode : Sequencer->GetViewModel()->GetSelection()->GetNodesWithSelectedKeysOrSections())
 		{
 			SequencerHelpers::GetAllKeyAreas(DisplayNode.Pin(), KeyAreas);
 		}
@@ -807,9 +791,9 @@ void FSectionContextMenu::ReduceKeys()
 
 bool FSectionContextMenu::CanAutoSize() const
 {
-	for (auto Section : Sequencer->GetSelection().GetSelectedSections())
+	for (UMovieSceneSection* Section : Sequencer->GetViewModel()->GetSelection()->GetSelectedSections())
 	{
-		if (Section.IsValid() && Section->GetAutoSizeRange().IsSet())
+		if (Section && Section->GetAutoSizeRange().IsSet())
 		{
 			return true;
 		}
@@ -822,14 +806,14 @@ bool FSectionContextMenu::CanReduceKeys() const
 	using namespace UE::Sequencer;
 
 	TSet<TSharedPtr<IKeyArea> > KeyAreas;
-	for (const TWeakPtr<FViewModel>& WeakItem : Sequencer->GetSelection().GetSelectedOutlinerItems())
+	for (const TWeakPtr<FViewModel>& WeakItem : Sequencer->GetViewModel()->GetSelection()->Outliner)
 	{
 		SequencerHelpers::GetAllKeyAreas(WeakItem.Pin(), KeyAreas);
 	}
 
 	if (KeyAreas.Num() == 0)
 	{
-		for (const TWeakPtr<FViewModel>& DisplayNode : Sequencer->GetSelection().GetNodesWithSelectedKeysOrSections())
+		for (const TWeakViewModelPtr<IOutlinerExtension>& DisplayNode : Sequencer->GetViewModel()->GetSelection()->GetNodesWithSelectedKeysOrSections())
 		{
 			SequencerHelpers::GetAllKeyAreas(DisplayNode.Pin(), KeyAreas);
 		}
@@ -845,14 +829,14 @@ void FSectionContextMenu::SetInterpTangentMode(ERichCurveInterpMode InterpMode,
 	FScopedTransaction SetInterpTangentModeTransaction(LOCTEXT("SetInterpTangentMode_Transaction", "Set Interpolation and Tangent Mode"));
 
 	TSet<TSharedPtr<IKeyArea> > KeyAreas;
-	for (const TWeakPtr<FViewModel>& WeakItem : Sequencer->GetSelection().GetSelectedOutlinerItems())
+	for (const TWeakPtr<FViewModel>& WeakItem : Sequencer->GetViewModel()->GetSelection()->Outliner)
 	{
 		SequencerHelpers::GetAllKeyAreas(WeakItem.Pin(), KeyAreas);
 	}
 
 	if (KeyAreas.Num() == 0)
 	{
-		for (const TWeakPtr<FViewModel>& DisplayNode : Sequencer->GetSelection().GetNodesWithSelectedKeysOrSections())
+		for (const TWeakViewModelPtr<IOutlinerExtension>& DisplayNode : Sequencer->GetViewModel()->GetSelection()->GetNodesWithSelectedKeysOrSections())
 		{
 			SequencerHelpers::GetAllKeyAreas(DisplayNode.Pin(), KeyAreas);
 		}
@@ -915,14 +899,14 @@ bool FSectionContextMenu::CanSetInterpTangentMode() const
 	using namespace UE::Sequencer;
 
 	TSet<TSharedPtr<IKeyArea> > KeyAreas;
-	for (const TWeakPtr<FViewModel>& WeakItem : Sequencer->GetSelection().GetSelectedOutlinerItems())
+	for (const TWeakPtr<FViewModel>& WeakItem : Sequencer->GetViewModel()->GetSelection()->Outliner)
 	{
 		SequencerHelpers::GetAllKeyAreas(WeakItem.Pin(), KeyAreas);
 	}
 
 	if (KeyAreas.Num() == 0)
 	{
-		for (const TWeakPtr<FViewModel>& DisplayNode : Sequencer->GetSelection().GetNodesWithSelectedKeysOrSections())
+		for (const TWeakViewModelPtr<IOutlinerExtension>& DisplayNode : Sequencer->GetViewModel()->GetSelection()->GetNodesWithSelectedKeysOrSections())
 		{
 			SequencerHelpers::GetAllKeyAreas(DisplayNode.Pin(), KeyAreas);
 		}
@@ -948,14 +932,11 @@ void FSectionContextMenu::ToggleSectionActive()
 	bool bIsActive = !IsSectionActive();
 	bool bAnythingChanged = false;
 
-	for (auto Section : Sequencer->GetSelection().GetSelectedSections())
+	for (UMovieSceneSection* Section : Sequencer->GetViewModel()->GetSelection()->GetSelectedSections())
 	{
-		if (Section.IsValid())
-		{
-			bAnythingChanged = true;
-			Section->Modify();
-			Section->SetIsActive(bIsActive);
-		}
+		bAnythingChanged = true;
+		Section->Modify();
+		Section->SetIsActive(bIsActive);
 	}
 
 	if (bAnythingChanged)
@@ -971,9 +952,9 @@ void FSectionContextMenu::ToggleSectionActive()
 bool FSectionContextMenu::IsSectionActive() const
 {
 	// Active only if all are active
-	for (auto Section : Sequencer->GetSelection().GetSelectedSections())
+	for (UMovieSceneSection* Section : Sequencer->GetViewModel()->GetSelection()->GetSelectedSections())
 	{
-		if (Section.IsValid() && !Section->IsActive())
+		if (Section && !Section->IsActive())
 		{
 			return false;
 		}
@@ -989,9 +970,9 @@ void FSectionContextMenu::ToggleSectionLocked()
 	bool bIsLocked = !IsSectionLocked();
 	bool bAnythingChanged = false;
 
-	for (auto Section : Sequencer->GetSelection().GetSelectedSections())
+	for (UMovieSceneSection* Section : Sequencer->GetViewModel()->GetSelection()->GetSelectedSections())
 	{
-		if (Section.IsValid())
+		if (Section)
 		{
 			bAnythingChanged = true;
 			Section->Modify();
@@ -1013,9 +994,9 @@ void FSectionContextMenu::ToggleSectionLocked()
 bool FSectionContextMenu::IsSectionLocked() const
 {
 	// Locked only if all are locked
-	for (auto Section : Sequencer->GetSelection().GetSelectedSections())
+	for (UMovieSceneSection* Section : Sequencer->GetViewModel()->GetSelection()->GetSelectedSections())
 	{
-		if (Section.IsValid() && !Section->IsLocked())
+		if (Section && !Section->IsLocked())
 		{
 			return false;
 		}
@@ -1027,7 +1008,7 @@ bool FSectionContextMenu::IsSectionLocked() const
 
 void FSectionContextMenu::DeleteSection()
 {
-	Sequencer->DeleteSections(Sequencer->GetSelection().GetSelectedSections());
+	Sequencer->DeleteSections(Sequencer->GetViewModel()->GetSelection()->GetSelectedSections());
 }
 
 
@@ -1061,14 +1042,8 @@ TMap<UMovieSceneTrack*, TMap<int32, FTrackSectionRow>> GenerateTrackRowsFromSele
 {
 	TMap<UMovieSceneTrack*, TMap<int32, FTrackSectionRow>> TrackRows;
 
-	for (const TWeakObjectPtr<UMovieSceneSection>& SectionPtr : Sequencer.GetSelection().GetSelectedSections())
+	for (UMovieSceneSection* Section : Sequencer.GetViewModel()->GetSelection()->GetSelectedSections())
 	{
-		UMovieSceneSection* Section = SectionPtr.Get();
-		if (!Section)
-		{
-			continue;
-		}
-
 		UMovieSceneTrack* Track = Section->GetTypedOuter<UMovieSceneTrack>();
 		if (!Track)
 		{
@@ -1473,9 +1448,9 @@ void FPasteContextMenu::Setup()
 	else
 	{
 		// Use the selected sections
-		for (TWeakObjectPtr<UMovieSceneSection> WeakSection : Sequencer->GetSelection().GetSelectedSections())
+		for (UMovieSceneSection* WeakSection : Sequencer->GetViewModel()->GetSelection()->GetSelectedSections())
 		{
-			if (TSharedPtr<FSectionModel> SectionHandle = Sequencer->GetNodeTree()->GetSectionModel(WeakSection.Get()))
+			if (TSharedPtr<FSectionModel> SectionHandle = Sequencer->GetNodeTree()->GetSectionModel(WeakSection))
 			{
 				SectionModels.Add(SectionHandle);
 			}
@@ -1623,6 +1598,8 @@ void FPasteContextMenu::BeginPasteInto()
 
 void FPasteContextMenu::EndPasteInto(bool bAnythingPasted, const TSet<FSequencerSelectedKey>& NewSelection)
 {
+	using namespace UE::Sequencer;
+
 	if (!bAnythingPasted)
 	{
 		GEditor->CancelTransaction(0);
@@ -1633,17 +1610,21 @@ void FPasteContextMenu::EndPasteInto(bool bAnythingPasted, const TSet<FSequencer
 
 	UE::Sequencer::SSequencerSection::ThrobKeySelection();
 
-	FSequencerSelection& Selection = Sequencer->GetSelection();
-	Selection.SuspendBroadcast();
-	Selection.EmptySelectedTrackAreaItems();
-	Selection.EmptySelectedKeys();
-
-	for (const FSequencerSelectedKey& Key : NewSelection)
+	FSequencerSelection& Selection = *Sequencer->GetViewModel()->GetSelection();
 	{
-		Selection.AddToSelection(Key);
+		FSelectionEventSuppressor EventSuppressor = Selection.SuppressEvents();
+
+		Selection.TrackArea.Empty();
+		Selection.KeySelection.Empty();
+
+		for (const FSequencerSelectedKey& NewKey : NewSelection)
+		{
+			if (TSharedPtr<FChannelModel> Channel = NewKey.WeakChannel.Pin())
+			{
+				Selection.KeySelection.Select(Channel, NewKey.KeyHandle);
+			}
+		}
 	}
-	Selection.ResumeBroadcast();
-	Selection.GetOnKeySelectionChanged().Broadcast();
 
 	Sequencer->OnClipboardUsed(Args.Clipboard);
 	Sequencer->NotifyMovieSceneDataChanged(EMovieSceneDataChangeType::TrackValueChanged);
diff --git a/Engine/Source/Editor/Sequencer/Private/SequencerEdMode.cpp b/Engine/Source/Editor/Sequencer/Private/SequencerEdMode.cpp
--- a/Engine/Source/Editor/Sequencer/Private/SequencerEdMode.cpp
+++ b/Engine/Source/Editor/Sequencer/Private/SequencerEdMode.cpp
@@ -9,6 +9,7 @@
 #include "MVVM/ViewModels/SequencerEditorViewModel.h"
 #include "MVVM/ViewModels/ViewModelIterators.h"
 #include "MVVM/Views/STrackAreaView.h"
+#include "MVVM/Selection/Selection.h"
 #include "EditorViewportClient.h"
 #include "Curves/KeyHandle.h"
 #include "ISequencer.h"
@@ -515,23 +516,18 @@ void FSequencerEdMode::OnKeySelected(FViewport* Viewport, HMovieSceneKeyProxy* K
 
 	for (TWeakPtr<FSequencer> WeakSequencer : Sequencers)
 	{
-		bool bChangedSelection = false;
-
 		TSharedPtr<FSequencer> Sequencer = WeakSequencer.Pin();
 		if (Sequencer.IsValid())
 		{
 			Sequencer->SetLocalTimeDirectly(KeyProxy->Key.Time);
 
-			FSequencerSelection& Selection = Sequencer->GetSelection();
+			FSequencerSelection& Selection = *Sequencer->GetViewModel()->GetSelection();
+
+			FSelectionEventSuppressor EventSuppressor = Selection.SuppressEvents();
+
 			if (!bAddToSelection && !bToggleSelection)
 			{
-				if (!bChangedSelection)
-				{
-					Sequencer->GetSelection().SuspendBroadcast();
-					bChangedSelection = true;
-				}
-
-				Sequencer->GetSelection().EmptySelectedKeys();
+				Selection.KeySelection.Empty();
 			}
 
 			for (const FTrajectoryKey::FData& KeyData : KeyProxy->Key.KeyData)
@@ -546,24 +542,20 @@ void FSequencerEdMode::OnKeySelected(FViewport* Viewport, HMovieSceneKeyProxy* K
 						TSharedPtr<FChannelModel> Channel = KeyAreaNode->GetChannel(Section);
 						if (Channel && Channel->GetKeyArea()->GetName() == KeyData.ChannelName)
 						{
-							if (!bChangedSelection)
+							if (bToggleSelection && Selection.KeySelection.IsSelected(KeyData.KeyHandle.GetValue()))
 							{
-								Sequencer->GetSelection().SuspendBroadcast();
-								bChangedSelection = true;
+								Selection.KeySelection.Deselect(KeyData.KeyHandle.GetValue());
+							}
+							else
+							{
+								Selection.KeySelection.Select(Channel, KeyData.KeyHandle.GetValue());
 							}
 
-							Sequencer->SelectKey(Section, Channel, KeyData.KeyHandle.GetValue(), bToggleSelection);
 							break;
 						}
 					}
 				}
 			}
-			if (bChangedSelection)
-			{
-				Sequencer->GetSelection().ResumeBroadcast();
-				Sequencer->GetSelection().GetOnKeySelectionChanged().Broadcast();
-				Sequencer->GetSelection().GetOnOutlinerNodeSelectionChangedObjectGuids().Broadcast();
-			}
 		}
 	}
 }
@@ -815,7 +807,7 @@ void FSequencerEdMode::DrawTracks3D(FPrimitiveDrawInterface* PDI)
 		FObjectBindingModelStorageExtension* ObjectStorage = Sequencer->GetViewModel()->GetRootModel()->CastDynamic<FObjectBindingModelStorageExtension>();
 		check(ObjectStorage);
 
-		const FSequencerSelection& Selection = Sequencer->GetSelection();
+		const FSequencerSelection& Selection = *Sequencer->GetViewModel()->GetSelection();
 		const TSharedRef<FSequencerNodeTree>& NodeTree = Sequencer->GetNodeTree();
 		for (const FMovieSceneBinding& Binding : Sequence->GetMovieScene()->GetBindings())
 		{
@@ -825,12 +817,12 @@ void FSequencerEdMode::DrawTracks3D(FPrimitiveDrawInterface* PDI)
 				continue;
 			}
 
-			bool bSelected = Selection.IsSelected(ObjectBindingNode);
+			bool bSelected = Selection.Outliner.IsSelected(ObjectBindingNode);
 			if (!bSelected)
 			{
-				for (TSharedPtr<FViewModel> Child : ObjectBindingNode->GetDescendants())
+				for (TViewModelPtr<IOutlinerExtension> Child : ObjectBindingNode->GetDescendantsOfType<IOutlinerExtension>())
 				{
-					if (Selection.IsSelected(Child) || Selection.NodeHasSelectedKeysOrSections(Child))
+					if (Selection.Outliner.IsSelected(Child) || Selection.NodeHasSelectedKeysOrSections(Child))
 					{
 						bSelected = true;
 						// Stop traversing
@@ -839,9 +831,9 @@ void FSequencerEdMode::DrawTracks3D(FPrimitiveDrawInterface* PDI)
 				}
 
 				// If one of our parent is selected, we're considered selected
-				for (TSharedPtr<FViewModel> Parent : ObjectBindingNode->GetAncestors())
+				for (TViewModelPtr<IOutlinerExtension> Parent : ObjectBindingNode->GetAncestorsOfType<IOutlinerExtension>())
 				{
-					if (Selection.IsSelected(Parent) || Selection.NodeHasSelectedKeysOrSections(Parent))
+					if (Selection.Outliner.IsSelected(Parent) || Selection.NodeHasSelectedKeysOrSections(Parent))
 					{
 						bSelected = true;
 						break;
@@ -888,6 +880,8 @@ void FSequencerEdMode::DrawTracks3D(FPrimitiveDrawInterface* PDI)
 
 void FSequencerEdMode::DrawAudioTracks(FPrimitiveDrawInterface* PDI)
 {
+	using namespace UE::Sequencer;
+
 	for (TWeakPtr<ISequencer> WeakSequencer : Sequencers)
 	{
 		TSharedPtr<ISequencer> Sequencer = WeakSequencer.Pin();
@@ -904,10 +898,9 @@ void FSequencerEdMode::DrawAudioTracks(FPrimitiveDrawInterface* PDI)
 
 		FQualifiedFrameTime CurrentTime = Sequencer->GetLocalTime();
 
-		const FSequencerSelection& Selection = Sequencer->GetSelection();
-		for (UMovieSceneTrack* Track : Selection.GetSelectedTracks())
+		for (TViewModelPtr<ITrackExtension> TrackModel : Sequencer->GetViewModel()->GetSelection()->Outliner.Filter<ITrackExtension>())
 		{
-			UMovieSceneAudioTrack* AudioTrack = Cast<UMovieSceneAudioTrack>(Track);
+			UMovieSceneAudioTrack* AudioTrack = Cast<UMovieSceneAudioTrack>(TrackModel->GetTrack());
 
 			if (!AudioTrack)
 			{
@@ -920,14 +913,14 @@ void FSequencerEdMode::DrawAudioTracks(FPrimitiveDrawInterface* PDI)
 				const FMovieSceneActorReferenceData& AttachActorData = AudioSection->GetAttachActorData();
 
 				TMovieSceneChannelData<const FMovieSceneActorReferenceKey> ChannelData = AttachActorData.GetData();
-
-				TArrayView<const FFrameNumber> Times = ChannelData.GetTimes();
-				TArrayView<const FMovieSceneActorReferenceKey> Values = ChannelData.GetValues();
+				TArrayView<const FMovieSceneActorReferenceKey> Values = ChannelData.GetValues().Num() != 0
+					? ChannelData.GetValues()
+					: MakeArrayView(&AttachActorData.GetDefault(), 1);
 		
 				FMovieSceneActorReferenceKey CurrentValue;
 				AttachActorData.Evaluate(CurrentTime.Time, CurrentValue);
 
-				for (int32 Index = 0; Index < Times.Num(); ++Index)
+				for (int32 Index = 0; Index < Values.Num(); ++Index)
 				{
 					FMovieSceneObjectBindingID AttachBindingID = Values[Index].Object;
 					FName AttachSocketName = Values[Index].SocketName;
diff --git a/Engine/Source/Editor/Sequencer/Private/SequencerHotspots.cpp b/Engine/Source/Editor/Sequencer/Private/SequencerHotspots.cpp
--- a/Engine/Source/Editor/Sequencer/Private/SequencerHotspots.cpp
+++ b/Engine/Source/Editor/Sequencer/Private/SequencerHotspots.cpp
@@ -2,6 +2,8 @@
 
 #include "SequencerHotspots.h"
 #include "MVVM/Extensions/IObjectBindingExtension.h"
+#include "MVVM/ViewModels/SequencerEditorViewModel.h"
+#include "MVVM/Selection/Selection.h"
 #include "SequencerCommonHelpers.h"
 #include "SequencerSettings.h"
 #include "SSequencer.h"
@@ -43,10 +45,10 @@ UE_SEQUENCER_DEFINE_VIEW_MODEL_TYPE_ID(IMouseHandlerHotspot);
 
 FHotspotSelectionManager::FHotspotSelectionManager(const FPointerEvent* InMouseEvent, FSequencer* InSequencer)
 	: MouseEvent(InMouseEvent)
-	, Selection(&InSequencer->GetSelection())
+	, Selection(InSequencer->GetViewModel()->GetSelection())
 	, Sequencer(InSequencer)
 {
-	Selection->SuspendBroadcast();
+	EventSuppressor = Selection->SuppressEventsLongRunning();
 
 	bForceSelect = !MouseEvent->IsControlDown();
 	bAddingToSelection = MouseEvent->IsShiftDown() || MouseEvent->IsControlDown();
@@ -60,16 +62,16 @@ FHotspotSelectionManager::FHotspotSelectionManager(const FPointerEvent* InMouseE
 
 FHotspotSelectionManager::~FHotspotSelectionManager()
 {
-	Selection->ResumeBroadcast();
-	Selection->GetOnOutlinerNodeSelectionChanged().Broadcast();
+	// Broadcast selection events by destroying the suppressor
+	EventSuppressor = nullptr;
 }
 
 void FHotspotSelectionManager::ConditionallyClearSelection()
 {
 	if (!bAddingToSelection)
 	{
-		Selection->EmptySelectedTrackAreaItems();
-		Selection->EmptySelectedKeys();
+		Selection->TrackArea.Empty();
+		Selection->KeySelection.Empty();
 
 		bAddingToSelection = true;
 	}
@@ -79,7 +81,7 @@ void FHotspotSelectionManager::ToggleKeys(TArrayView<const FSequencerSelectedKey
 {
 	for (const FSequencerSelectedKey& Key : InKeys)
 	{
-		const bool bIsSelected = Selection->IsSelected(Key);
+		const bool bIsSelected = Selection->KeySelection.IsSelected(Key.KeyHandle);
 		if (bIsSelected && bForceSelect)
 		{
 			continue;
@@ -87,18 +89,18 @@ void FHotspotSelectionManager::ToggleKeys(TArrayView<const FSequencerSelectedKey
 
 		if (!bIsSelected)
 		{
-			Selection->AddToSelection(Key);
+			Selection->KeySelection.Select(Key.WeakChannel.Pin(), Key.KeyHandle);
 		}
 		else
 		{
-			Selection->RemoveFromSelection(Key);
+			Selection->KeySelection.Deselect(Key.KeyHandle);
 		}
 	}
 }
 
 void FHotspotSelectionManager::ToggleModel(TSharedPtr<FViewModel> InModel)
 {
-	const bool bIsSelected = Selection->IsSelected(InModel);
+	const bool bIsSelected = Selection->TrackArea.IsSelected(InModel);
 	if (bIsSelected && bForceSelect)
 	{
 		return;
@@ -120,11 +122,11 @@ void FHotspotSelectionManager::ToggleModel(TSharedPtr<FViewModel> InModel)
 
 	if (!bIsSelected)
 	{
-		Selection->AddToSelection(InModel);
+		Selection->TrackArea.Select(InModel);
 	}
 	else
 	{
-		Selection->RemoveFromSelection(InModel);
+		Selection->TrackArea.Deselect(InModel);
 	}
 }
 
@@ -132,20 +134,20 @@ void FHotspotSelectionManager::SelectKeysExclusive(TArrayView<const FSequencerSe
 {
 	for (const FSequencerSelectedKey& Key : InKeys)
 	{
-		if (!Selection->IsSelected(Key))
+		if (!Selection->KeySelection.IsSelected(Key.KeyHandle))
 		{
 			ConditionallyClearSelection();
-			Selection->AddToSelection(Key);
+			Selection->KeySelection.Select(Key.WeakChannel.Pin(), Key.KeyHandle);
 		}
 	}
 }
 
 void FHotspotSelectionManager::SelectModelExclusive(TSharedPtr<FViewModel> InModel)
 {
-	if (!Selection->IsSelected(InModel))
+	if (!Selection->TrackArea.IsSelected(InModel))
 	{
 		ConditionallyClearSelection();
-		Selection->AddToSelection(InModel);
+		Selection->TrackArea.Select(InModel);
 	}
 }
 
@@ -317,7 +319,7 @@ TSharedPtr<ISequencerEditToolDragOperation> FKeyBarHotspot::InitiateDrag(const F
 		{
 			AllKeys.Append(Hotspot->LeadingKeys);
 			AllKeys.Append(Hotspot->TrailingKeys);
-			AllKeys.Append(InSequencer->GetSelection().GetSelectedKeys());
+			FSequencerSelectedKey::AppendKeySelection(AllKeys, InSequencer->GetViewModel()->GetSelection()->KeySelection);
 
 			AllLinearKeys = AllKeys.Array();
 
@@ -459,7 +461,7 @@ void FSectionHotspotBase::UpdateOnHover(FTrackAreaViewModel& InTrackArea) const
 
 	// Move sections if they are selected
 	TSharedPtr<ISequencer> Sequencer = WeakSequencer.Pin();
-	if (Sequencer->GetSelection().IsSelected(WeakSectionModel))
+	if (Sequencer->GetViewModel()->GetSelection()->TrackArea.IsSelected(WeakSectionModel))
 	{
 		if (ThisSection && ThisSection->GetRange() != TRange<FFrameNumber>::All())
 		{
@@ -535,14 +537,18 @@ void FSectionResizeHotspot::UpdateOnHover(FTrackAreaViewModel& InTrackArea) cons
 TSharedPtr<ISequencerEditToolDragOperation> FSectionResizeHotspot::InitiateDrag(const FPointerEvent& MouseEvent)
 {
 	TSharedPtr<FSequencer> Sequencer = WeakSequencer.Pin();
-	if (!Sequencer->GetSelection().IsSelected(WeakSectionModel))
+	TSharedPtr<FSequencerSelection> Selection = Sequencer->GetViewModel()->GetSelection();
+
+	if (!Selection->TrackArea.IsSelected(WeakSectionModel))
 	{
-		Sequencer->GetSelection().Empty();
-		Sequencer->GetSelection().AddToSelection(WeakSectionModel.Pin());
+		FSelectionEventSuppressor SuppressEvents = Selection->SuppressEvents();
+
+		Selection->Empty();
+		Selection->TrackArea.Select(WeakSectionModel);
 	}
 
 	const bool bIsSlipping = false;
-	return MakeShareable( new FResizeSection(*Sequencer, Sequencer->GetSelection().GetSelectedSections(), HandleType == Right, bIsSlipping) );
+	return MakeShareable( new FResizeSection(*Sequencer, HandleType == Right, bIsSlipping) );
 }
 
 const FSlateBrush* FSectionResizeHotspot::GetCursorDecorator(const FGeometry& MyGeometry, const FPointerEvent& CursorEvent) const
diff --git a/Engine/Source/Editor/Sequencer/Private/SequencerHotspots.h b/Engine/Source/Editor/Sequencer/Private/SequencerHotspots.h
--- a/Engine/Source/Editor/Sequencer/Private/SequencerHotspots.h
+++ b/Engine/Source/Editor/Sequencer/Private/SequencerHotspots.h
@@ -19,14 +19,17 @@ namespace UE
 namespace Sequencer
 {
 
+struct FSelectionEventSuppressor;
+
 class FSectionModel;
 class ISequencerEditToolDragOperation;
 
 struct FHotspotSelectionManager
 {
 	const FPointerEvent* MouseEvent;
-	FSequencerSelection* Selection;
+	TSharedPtr<FSequencerSelection> Selection;
 	FSequencer* Sequencer;
+	TUniquePtr<FSelectionEventSuppressor> EventSuppressor;
 
 	bool bForceSelect;
 	bool bAddingToSelection;
diff --git a/Engine/Source/Editor/Sequencer/Private/SequencerNodeTree.cpp b/Engine/Source/Editor/Sequencer/Private/SequencerNodeTree.cpp
--- a/Engine/Source/Editor/Sequencer/Private/SequencerNodeTree.cpp
+++ b/Engine/Source/Editor/Sequencer/Private/SequencerNodeTree.cpp
@@ -2,6 +2,7 @@
 
 #include "SequencerNodeTree.h"
 #include "MVVM/ViewModels/ViewModelHierarchy.h"
+#include "MVVM/Selection/Selection.h"
 #include "MovieSceneBinding.h"
 #include "GameFramework/Actor.h"
 #include "MovieScene.h"
@@ -345,9 +346,7 @@ bool FSequencerNodeTree::IsSelectedNodesSolo() const
 {
 	using namespace UE::Sequencer;
 
-	const TSet<TWeakPtr<FViewModel>> SelectedNodes = Sequencer.GetSelection().GetSelectedOutlinerItems();
-	
-	if (SelectedNodes.Num() == 0)
+	if (Sequencer.GetViewModel()->GetSelection()->Outliner.Num() == 0)
 	{
 		return false;
 	}
@@ -355,7 +354,7 @@ bool FSequencerNodeTree::IsSelectedNodesSolo() const
 	const TArray<FString>& SoloNodes = Sequencer.GetFocusedMovieSceneSequence()->GetMovieScene()->GetSoloNodes();
 
 	bool bIsSolo = true;
-	for (const TWeakPtr<FViewModel>& Node : SelectedNodes)
+	for (FViewModelPtr Node : Sequencer.GetViewModel()->GetSelection()->Outliner)
 	{
 		if (!SoloNodes.Contains(IOutlinerExtension::GetPathName(Node)))
 		{
@@ -380,8 +379,7 @@ void FSequencerNodeTree::ToggleSelectedNodesSolo()
 	
 	TArray<FString>& SoloNodes = MovieScene->GetSoloNodes();
 
-	const TSet<TWeakPtr<FViewModel>> SelectedNodes = Sequencer.GetSelection().GetSelectedOutlinerItems();
-	if (SelectedNodes.Num() == 0)
+	if (Sequencer.GetViewModel()->GetSelection()->Outliner.Num() == 0)
 	{
 		return;
 	}
@@ -389,7 +387,7 @@ void FSequencerNodeTree::ToggleSelectedNodesSolo()
 	// First, determine if any of the selected nodes are not marked as solo
 	// If we have a mix, we should default to setting them all as solo
 	bool bIsSolo = true;
-	for (const TWeakPtr<const FViewModel> Node : SelectedNodes)
+	for (FViewModelPtr Node : Sequencer.GetViewModel()->GetSelection()->Outliner)
 	{
 		if (!SoloNodes.Contains(IOutlinerExtension::GetPathName(Node)))
 		{
@@ -402,7 +400,7 @@ void FSequencerNodeTree::ToggleSelectedNodesSolo()
 
 	MovieScene->Modify();
 
-	for (const TWeakPtr<FViewModel>& Node : SelectedNodes)
+	for (FViewModelPtr Node : Sequencer.GetViewModel()->GetSelection()->Outliner)
 	{
 		FString NodePath = IOutlinerExtension::GetPathName(Node);
 		if (bIsSolo)
@@ -446,9 +444,7 @@ bool FSequencerNodeTree::IsSelectedNodesMute() const
 {
 	using namespace UE::Sequencer;
 
-	const TSet<TWeakPtr<FViewModel>> SelectedNodes = Sequencer.GetSelection().GetSelectedOutlinerItems();
-
-	if (SelectedNodes.Num() == 0)
+	if (Sequencer.GetViewModel()->GetSelection()->Outliner.Num() == 0)
 	{
 		return false;
 	}
@@ -456,7 +452,7 @@ bool FSequencerNodeTree::IsSelectedNodesMute() const
 	const TArray<FString>& MuteNodes = Sequencer.GetFocusedMovieSceneSequence()->GetMovieScene()->GetMuteNodes();
 
 	bool bIsMute = true;
-	for (const TWeakPtr<const FViewModel> Node : SelectedNodes)
+	for (FViewModelPtr Node : Sequencer.GetViewModel()->GetSelection()->Outliner)
 	{
 		if (!MuteNodes.Contains(IOutlinerExtension::GetPathName(Node)))
 		{
@@ -481,8 +477,7 @@ void FSequencerNodeTree::ToggleSelectedNodesMute()
 
 	TArray<FString>& MuteNodes = MovieScene->GetMuteNodes();
 
-	const TSet<TWeakPtr<FViewModel>> SelectedNodes = Sequencer.GetSelection().GetSelectedOutlinerItems();
-	if (SelectedNodes.Num() == 0)
+	if (Sequencer.GetViewModel()->GetSelection()->Outliner.Num() == 0)
 	{
 		return;
 	}
@@ -490,7 +485,7 @@ void FSequencerNodeTree::ToggleSelectedNodesMute()
 	// First, determine if any of the selected nodes are not marked as Mute
 	// If we have a mix, we should default to setting them all as Mute
 	bool bIsMute = true;
-	for (const TWeakPtr<const FViewModel> Node : SelectedNodes)
+	for (FViewModelPtr Node : Sequencer.GetViewModel()->GetSelection()->Outliner)
 	{
 		if (!MuteNodes.Contains(IOutlinerExtension::GetPathName(Node)))
 		{
@@ -503,7 +498,7 @@ void FSequencerNodeTree::ToggleSelectedNodesMute()
 
 	MovieScene->Modify();
 
-	for (const TWeakPtr<FViewModel>& Node : SelectedNodes)
+	for (FViewModelPtr Node : Sequencer.GetViewModel()->GetSelection()->Outliner)
 	{
 		FString NodePath = IOutlinerExtension::GetPathName(Node);
 		if (bIsMute)
@@ -575,6 +570,9 @@ bool FSequencerNodeTree::GetSavedPinnedState(const UE::Sequencer::FViewModel& No
 bool FSequencerNodeTree::IsNodeFiltered(const TSharedPtr<UE::Sequencer::FViewModel>& Node) const
 {
 	using namespace UE::Sequencer;
+
+	TViewModelPtr<IOutlinerExtension> OutlinerItem = CastViewModel<IOutlinerExtension>(Node);
+	return OutlinerItem && FilteredNodes.Contains(OutlinerItem);
 	return FilteredNodes.Contains(CastViewModel<IOutlinerExtension>(Node));
 }
 
diff --git a/Engine/Source/Editor/Sequencer/Private/SequencerSelectedKey.cpp b/Engine/Source/Editor/Sequencer/Private/SequencerSelectedKey.cpp
--- a/Engine/Source/Editor/Sequencer/Private/SequencerSelectedKey.cpp
+++ b/Engine/Source/Editor/Sequencer/Private/SequencerSelectedKey.cpp
@@ -6,9 +6,54 @@
 #include "Containers/Map.h"
 #include "IKeyArea.h"
 #include "MVVM/ViewModels/ChannelModel.h"
+#include "MVVM/Selection/Selection.h"
 #include "Misc/AssertionMacros.h"
 #include "Misc/FrameNumber.h"
 
+void FSequencerSelectedKey::AppendKeySelection(TSet<FSequencerSelectedKey>& OutSelectedKeys, const UE::Sequencer::FKeySelection& InKeySelection)
+{
+	using namespace UE::Sequencer;
+
+	for (FKeyHandle Key : InKeySelection)
+	{
+		TSharedPtr<FChannelModel> Channel = InKeySelection.GetModelForKey(Key);
+		UMovieSceneSection*       Section = Channel ? Channel->GetSection() : nullptr;
+		if (Channel && Section)
+		{
+			OutSelectedKeys.Emplace(FSequencerSelectedKey(*Section, Channel, Key));
+		}
+	}
+}
+
+FSelectedKeysByChannel::FSelectedKeysByChannel(const UE::Sequencer::FKeySelection& KeySelection)
+{
+	using namespace UE::Sequencer;
+
+	TMap<const IKeyArea*, int32> KeyAreaToChannelIndex;
+
+	int32 Index = 0;
+	for (FKeyHandle Key : KeySelection)
+	{
+		TSharedPtr<FChannelModel> Channel = KeySelection.GetModelForKey(Key);
+
+		if (Channel)
+		{
+			const int32* ChannelArrayIndex = KeyAreaToChannelIndex.Find(Channel->GetKeyArea().Get());
+			if (!ChannelArrayIndex)
+			{
+				int32 NewIndex = SelectedChannels.Add(FSelectedChannelInfo(Channel->GetKeyArea()->GetChannel(), Channel->GetSection()));
+				ChannelArrayIndex = &KeyAreaToChannelIndex.Add(Channel->GetKeyArea().Get(), NewIndex);
+			}
+
+			FSelectedChannelInfo& ThisChannelInfo = SelectedChannels[*ChannelArrayIndex];
+			ThisChannelInfo.KeyHandles.Add(Key);
+			ThisChannelInfo.OriginalIndices.Add(Index);
+
+			++Index;
+		}
+	}
+}
+
 FSelectedKeysByChannel::FSelectedKeysByChannel(TArrayView<const FSequencerSelectedKey> InSelectedKeys)
 {
 	using namespace UE::Sequencer;
diff --git a/Engine/Source/Editor/Sequencer/Private/SequencerSelectedKey.h b/Engine/Source/Editor/Sequencer/Private/SequencerSelectedKey.h
--- a/Engine/Source/Editor/Sequencer/Private/SequencerSelectedKey.h
+++ b/Engine/Source/Editor/Sequencer/Private/SequencerSelectedKey.h
@@ -5,6 +5,7 @@
 #include "Channels/MovieSceneChannelHandle.h"
 #include "Containers/Array.h"
 #include "Containers/ArrayView.h"
+#include "Containers/Set.h"
 #include "Curves/KeyHandle.h"
 #include "HAL/Platform.h"
 #include "HAL/PlatformCrt.h"
@@ -20,6 +21,7 @@ namespace UE
 namespace Sequencer
 {
 
+struct FKeySelection;
 class FChannelModel;
 
 } // namespace Sequencer 
@@ -41,6 +43,8 @@ struct FSequencerSelectedKey
 
 public:
 
+	static void AppendKeySelection(TSet<FSequencerSelectedKey>& OutSelectedKeys, const UE::Sequencer::FKeySelection& InKeySelection);
+
 	/** Create and initialize a new instance. */
 	FSequencerSelectedKey(UMovieSceneSection& InSection, TWeakPtr<UE::Sequencer::FChannelModel> InChannel, FKeyHandle InKeyHandle)
 		: Section(&InSection)
@@ -101,6 +105,7 @@ struct FSelectedChannelInfo
  */
 struct FSelectedKeysByChannel
 {
+	explicit FSelectedKeysByChannel(const UE::Sequencer::FKeySelection& KeySelection);
 	explicit FSelectedKeysByChannel(TArrayView<const FSequencerSelectedKey> InSelectedKeys);
 
 	/** Array storing all selected keys for each channel */
diff --git a/Engine/Source/Editor/Sequencer/Private/SequencerSelectionCurveFilter.h b/Engine/Source/Editor/Sequencer/Private/SequencerSelectionCurveFilter.h
--- a/Engine/Source/Editor/Sequencer/Private/SequencerSelectionCurveFilter.h
+++ b/Engine/Source/Editor/Sequencer/Private/SequencerSelectionCurveFilter.h
@@ -8,6 +8,9 @@
 #include "MVVM/ViewModels/ViewModel.h"
 #include "Tree/CurveEditorTreeFilter.h"
 
+namespace UE::Sequencer
+{
+
 /**
  * A specialized filter for showing items in the curve editor selected from the sequencer panel.
  * This filter will store the selected nodes and all the parents of the selected nodes in a NodesToFilter set.
@@ -24,42 +27,31 @@ struct FSequencerSelectionCurveFilter : FCurveEditorTreeFilter
 	/**
 	 * Adds all selected nodes and their object parents to the NodesToFilter set
 	 */
-	void Update(const TSet<TWeakPtr<UE::Sequencer::FViewModel>>& SelectedNodes, const bool bExpandTreeToSelectedNodes = true)
+	void Update(TSharedPtr<FSequencerSelection> Selection, const bool bExpandTreeToSelectedNodes = true)
 	{
-		using namespace UE::Sequencer;
-
-		NodesToFilter.Empty(SelectedNodes.Num());
+		NodesToFilter.Empty(Selection->Outliner.Num());
 
-		for (const TWeakPtr<FViewModel>& WeakSelectedNode : SelectedNodes)
+		for (TViewModelPtr<IOutlinerExtension> SelectedNode : Selection->Outliner)
 		{
-			if (TSharedPtr<FViewModel> SelectedNode = WeakSelectedNode.Pin())
-			{
-				NodesToFilter.Add(WeakSelectedNode);
-
-				TSharedPtr<FViewModel> Parent = SelectedNode->GetParent();
-				while (Parent.IsValid())
-				{
-					if (Parent->IsA<IObjectBindingExtension>())
-					{
-						NodesToFilter.Add(Parent);
-						break;
-					}
+			NodesToFilter.Add(SelectedNode.AsModel());
 
-					Parent = Parent->GetParent();
-				}
+			for (TViewModelPtr<IObjectBindingExtension> ParentObject : SelectedNode.AsModel()->GetAncestorsOfType<IObjectBindingExtension>())
+			{
+				NodesToFilter.Add(ParentObject.AsModel());
 			}
 		}
 
 		bExpandToMatchedItems = bExpandTreeToSelectedNodes;
 	}
 
-	bool Match(TSharedRef<const UE::Sequencer::FViewModel> InNode) const
+	bool Match(TSharedRef<const FViewModel> InNode) const
 	{
 		return NodesToFilter.Contains(InNode);
 	}
 
 private:
 
-	TSet<TWeakPtr<const UE::Sequencer::FViewModel>> NodesToFilter;
+	TSet<TWeakPtr<const FViewModel>> NodesToFilter;
 };
 
+} // namespace UE::Sequencer
\ No newline at end of file
diff --git a/Engine/Source/Editor/Sequencer/Private/SequencerSelectionPreview.cpp b/Engine/Source/Editor/Sequencer/Private/SequencerSelectionPreview.cpp
--- a/Engine/Source/Editor/Sequencer/Private/SequencerSelectionPreview.cpp
+++ b/Engine/Source/Editor/Sequencer/Private/SequencerSelectionPreview.cpp
@@ -8,17 +8,17 @@
 
 namespace UE::Sequencer { class FViewModel; }
 
-void FSequencerSelectionPreview::SetSelectionState(const FSequencerSelectedKey& Key, ESelectionPreviewState InState)
+void FSequencerSelectionPreview::SetSelectionState(const UE::Sequencer::TWeakViewModelPtr<UE::Sequencer::FChannelModel>& Channel, FKeyHandle Key, ESelectionPreviewState InState)
 {
 	if (InState == ESelectionPreviewState::Undefined)
 	{
-		DefinedKeyStates.Remove(Key);
-		RawDefinedKeyStates.Remove(Key.KeyHandle);
+		RawDefinedKeyStates.Remove(Key);
+		KeyHandleToModel.Remove(Key);
 	}
 	else
 	{
-		DefinedKeyStates.Add(Key, InState);
-		RawDefinedKeyStates.Add(Key.KeyHandle, InState);
+		RawDefinedKeyStates.Add(Key, InState);
+		KeyHandleToModel.Add(Key, Channel);
 	}
 
 	CachedSelectionHash.Reset();
@@ -37,15 +37,6 @@ void FSequencerSelectionPreview::SetSelectionState(TWeakPtr<UE::Sequencer::FView
 	CachedSelectionHash.Reset();
 }
 
-ESelectionPreviewState FSequencerSelectionPreview::GetSelectionState(const FSequencerSelectedKey& Key) const
-{
-	if (const ESelectionPreviewState* State = DefinedKeyStates.Find(Key))
-	{
-		return *State;
-	}
-	return ESelectionPreviewState::Undefined;
-}
-
 ESelectionPreviewState FSequencerSelectionPreview::GetSelectionState(FKeyHandle Key) const
 {
 	if (const ESelectionPreviewState* State = RawDefinedKeyStates.Find(Key))
@@ -65,6 +56,11 @@ ESelectionPreviewState FSequencerSelectionPreview::GetSelectionState(TWeakPtr<UE
 	return ESelectionPreviewState::Undefined;
 }
 
+UE::Sequencer::TViewModelPtr<UE::Sequencer::FChannelModel> FSequencerSelectionPreview::GetChannelForKey(FKeyHandle KeyHandle) const
+{
+	return KeyHandleToModel.FindRef(KeyHandle).Pin();
+}
+
 void FSequencerSelectionPreview::Empty()
 {
 	EmptyDefinedKeyStates();
@@ -73,8 +69,8 @@ void FSequencerSelectionPreview::Empty()
 
 void FSequencerSelectionPreview::EmptyDefinedKeyStates()
 {
-	DefinedKeyStates.Reset();
 	RawDefinedKeyStates.Reset();
+	KeyHandleToModel.Reset();
 	CachedSelectionHash.Reset();
 }
 
diff --git a/Engine/Source/Editor/Sequencer/Private/SequencerSelectionPreview.h b/Engine/Source/Editor/Sequencer/Private/SequencerSelectionPreview.h
--- a/Engine/Source/Editor/Sequencer/Private/SequencerSelectionPreview.h
+++ b/Engine/Source/Editor/Sequencer/Private/SequencerSelectionPreview.h
@@ -6,8 +6,9 @@
 #include "HAL/Platform.h"
 #include "Misc/Optional.h"
 #include "Curves/KeyHandle.h"
-#include "SequencerSelectedKey.h"
 #include "Templates/SharedPointer.h"
+#include "MVVM/ViewModels/ChannelModel.h"
+#include "MVVM/ViewModelPtr.h"
 
 namespace UE::Sequencer { class FViewModel; }
 
@@ -18,6 +19,7 @@ enum class ESelectionPreviewState
 	NotSelected
 };
 
+
 /**
  * Manages the selection of keys, sections, and outliner nodes for the sequencer.
  */
@@ -26,25 +28,26 @@ class FSequencerSelectionPreview
 public:
 
 	/** Access the defined key states */
-	const TMap<FSequencerSelectedKey, ESelectionPreviewState>& GetDefinedKeyStates() const { return DefinedKeyStates; }
-	const TMap<FKeyHandle, ESelectionPreviewState>& GetRawDefinedKeyStates() const { return RawDefinedKeyStates; }
+	const TMap<FKeyHandle, ESelectionPreviewState>& GetDefinedKeyStates() const { return RawDefinedKeyStates; }
 
 	/** Access the defined model states */
 	const TMap<TWeakPtr<UE::Sequencer::FViewModel>, ESelectionPreviewState>& GetDefinedModelStates() const { return DefinedModelStates; }
 
 	/** Adds a key to the selection */
-	void SetSelectionState(const FSequencerSelectedKey& Key, ESelectionPreviewState InState);
+	void SetSelectionState(const UE::Sequencer::TWeakViewModelPtr<UE::Sequencer::FChannelModel>& Channel, FKeyHandle Key, ESelectionPreviewState InState);
 
 	/** Adds a model to the selection */
 	void SetSelectionState(TWeakPtr<UE::Sequencer::FViewModel> InModel, ESelectionPreviewState InState);
 
 	/** Returns the selection state for the the specified key. */
-	ESelectionPreviewState GetSelectionState(const FSequencerSelectedKey& Key) const;
 	ESelectionPreviewState GetSelectionState(FKeyHandle Key) const;
 
 	/** Returns the selection state for the the specified model. */
 	ESelectionPreviewState GetSelectionState(TWeakPtr<UE::Sequencer::FViewModel> InModel) const;
 
+	/** Finds the channel associated with the specified key. */
+	UE::Sequencer::TViewModelPtr<UE::Sequencer::FChannelModel> GetChannelForKey(FKeyHandle KeyHandle) const;
+
 	/** Empties all selections. */
 	void Empty();
 
@@ -59,7 +62,7 @@ public:
 
 private:
 
-	TMap<FSequencerSelectedKey, ESelectionPreviewState> DefinedKeyStates;
+	TMap<FKeyHandle, UE::Sequencer::TWeakViewModelPtr<UE::Sequencer::FChannelModel>> KeyHandleToModel;
 	TMap<FKeyHandle, ESelectionPreviewState> RawDefinedKeyStates;
 
 	TMap<TWeakPtr<UE::Sequencer::FViewModel>, ESelectionPreviewState> DefinedModelStates;
diff --git a/Engine/Source/Editor/Sequencer/Private/SequencerTimeSliderController.cpp b/Engine/Source/Editor/Sequencer/Private/SequencerTimeSliderController.cpp
--- a/Engine/Source/Editor/Sequencer/Private/SequencerTimeSliderController.cpp
+++ b/Engine/Source/Editor/Sequencer/Private/SequencerTimeSliderController.cpp
@@ -1,6 +1,8 @@
 // Copyright Epic Games, Inc. All Rights Reserved.
 
 #include "SequencerTimeSliderController.h"
+#include "MVVM/ViewModels/SequencerEditorViewModel.h"
+#include "MVVM/Selection/Selection.h"
 #include "Fonts/SlateFontInfo.h"
 #include "Rendering/DrawElements.h"
 #include "Misc/Paths.h"
@@ -332,6 +334,8 @@ void FSequencerTimeSliderController::DrawTicks( FSlateWindowElementList& OutDraw
 
 int32 FSequencerTimeSliderController::DrawMarkedFrames( const FGeometry& AllottedGeometry, const FScrubRangeToScreen& RangeToScreen, FSlateWindowElementList& OutDrawElements, int32 LayerId, const ESlateDrawEffect& DrawEffects, const FWidgetStyle& InWidgetStyle, bool bDrawLabels ) const
 {
+	using namespace UE::Sequencer;
+
 	const TArray<FMovieSceneMarkedFrame> & MarkedFrames = TimeSliderArgs.MarkedFrames.Get();
 	const TArray<FMovieSceneMarkedFrame> & GlobalMarkedFrames = TimeSliderArgs.GlobalMarkedFrames.Get();
 	if (MarkedFrames.Num() < 1 && GlobalMarkedFrames.Num() < 1)
@@ -339,7 +343,7 @@ int32 FSequencerTimeSliderController::DrawMarkedFrames( const FGeometry& Allotte
 		return LayerId;
 	}
 
-	const TSet<int32>& SelectedMarkedFrames = WeakSequencer.Pin()->GetSelection().GetSelectedMarkedFrames();
+	const FMarkedFrameSelection& SelectedMarkedFrames = WeakSequencer.Pin()->GetViewModel()->GetSelection()->MarkedFrames;
 	//const FLinearColor SelectedColor = FLinearColor::White; //FAppStyle::GetSlateColor("SelectionColor").GetColor(InWidgetStyle);
 	const FLinearColor WhiteColorHSV = FLinearColor::White.LinearRGBToHSV();
 
@@ -351,7 +355,7 @@ int32 FSequencerTimeSliderController::DrawMarkedFrames( const FGeometry& Allotte
 			double Seconds = MarkedFrame.FrameNumber / GetTickResolution();
 
 			const bool bIsHovered = (!bIsGlobal && HoverMarkIndex == MarkIndex);
-			const bool bIsSelected = (!bIsGlobal && SelectedMarkedFrames.Contains(MarkIndex));
+			const bool bIsSelected = (!bIsGlobal && SelectedMarkedFrames.IsSelected(MarkIndex));
 
 			// Get a selected color that's the marked frame color but at full opacity, full brightness, and a bit desaturated if it's
 			// already bright to begin with.
@@ -955,7 +959,7 @@ FReply FSequencerTimeSliderController::OnMouseButtonUp( SWidget& WidgetOwner, co
 
 					// We weren't scrubbing, so this is a single click in the time slider. This should
 					// also act as a deselection of any selected markers.
-					Sequencer->GetSelection().EmptySelectedMarkedFrames();
+					Sequencer->GetViewModel()->GetSelection()->MarkedFrames.Empty();
 				}
 
 				// If middle mouse button down we don't evaluate on the time change
@@ -1530,6 +1534,8 @@ ETimeSliderPlaybackStatus FSequencerTimeSliderController::GetPlaybackStatus() co
 
 TSharedRef<SWidget> FSequencerTimeSliderController::OpenSetPlaybackRangeMenu(const FGeometry& MyGeometry, const FPointerEvent& MouseEvent)
 {
+	using namespace UE::Sequencer;
+
 	TSharedPtr<FSequencer> Sequencer = WeakSequencer.Pin();
 	const bool bReadOnly = Sequencer && Sequencer->IsReadOnly();
 	
@@ -1669,8 +1675,8 @@ TSharedRef<SWidget> FSequencerTimeSliderController::OpenSetPlaybackRangeMenu(con
 
 		if (MarkedIndex != INDEX_NONE)
 		{
-			SequencerSelection.EmptySelectedMarkedFrames();
-			SequencerSelection.AddToSelection(MarkedIndex);
+			SequencerSelection.MarkedFrames.Empty();
+			SequencerSelection.MarkedFrames.Select(MarkedIndex);
 
 			class SMarkedFramePropertyWidget : public SCompoundWidget, public FNotifyHook
 			{
@@ -2021,8 +2027,7 @@ FFrameTime FSequencerTimeSliderController::SnapTimeToNearestKey(const FPointerEv
 
 		// If there are any tracks selected we'll find the nearest key only on that track. If there are no keys selected,
 		// we will try to find the nearest keys on all tracks. This mirrors the behavior of the Jump to Next Keyframe commands.
-		const TSet< TWeakPtr<FViewModel> >& SelectedNodes = WeakSequencer.Pin()->GetSelection().GetSelectedOutlinerItems();
-		if (SelectedNodes.Num() == 0)
+		if (WeakSequencer.Pin()->GetViewModel()->GetSelection()->Outliner.Num() == 0)
 		{
 			EnumAddFlags(NearestKeyOption, ENearestKeyOption::NKO_SearchAllTracks);
 		}
@@ -2106,6 +2111,8 @@ void FSequencerTimeSliderController::SetSelectionRangeEnd(FFrameNumber NewEnd)
 
 void FSequencerTimeSliderController::HandleMarkSelection(int32 InMarkIndex)
 {
+	using namespace UE::Sequencer;
+
 	TSharedPtr<FSequencer> Sequencer = WeakSequencer.Pin();
 	FSequencerSelection& SequencerSelection = Sequencer->GetSelection();
 
@@ -2116,29 +2123,31 @@ void FSequencerTimeSliderController::HandleMarkSelection(int32 InMarkIndex)
 	if (!bToggleSelection && !bAddToSelection)
 	{
 		SequencerSelection.Empty();
-		SequencerSelection.AddToSelection(InMarkIndex);
+		SequencerSelection.MarkedFrames.Select(InMarkIndex);
 	}
-	else if (bAddToSelection || !SequencerSelection.IsSelected(InMarkIndex))
+	else if (bAddToSelection || !SequencerSelection.MarkedFrames.IsSelected(InMarkIndex))
 	{
-		SequencerSelection.AddToSelection(InMarkIndex);
+		SequencerSelection.MarkedFrames.Select(InMarkIndex);
 	}
 	else
 	{
-		SequencerSelection.RemoveFromSelection(InMarkIndex);
+		SequencerSelection.MarkedFrames.Deselect(InMarkIndex);
 	}
 }
 
 void FSequencerTimeSliderController::UpdateMarkSelection(int32 InOldMarkIndex, FFrameNumber InMarkFrameNumber)
 {
+	using namespace UE::Sequencer;
+
 	TSharedPtr<FSequencer> Sequencer = WeakSequencer.Pin();
 	FSequencerSelection& SequencerSelection = Sequencer->GetSelection();
-	SequencerSelection.RemoveFromSelection(InOldMarkIndex);
+	SequencerSelection.MarkedFrames.Deselect(InOldMarkIndex);
 
 	UMovieScene* MovieScene = Sequencer->GetFocusedMovieSceneSequence()->GetMovieScene();
 	int32 NewMarkIndex = MovieScene->FindMarkedFrameByFrameNumber(InMarkFrameNumber);
 	if (ensure(NewMarkIndex != INDEX_NONE))
 	{
-		SequencerSelection.AddToSelection(NewMarkIndex);
+		SequencerSelection.MarkedFrames.Select(NewMarkIndex);
 	}
 }
 
diff --git a/Engine/Source/Editor/Sequencer/Private/Tools/DragOperation_Stretch.cpp b/Engine/Source/Editor/Sequencer/Private/Tools/DragOperation_Stretch.cpp
--- a/Engine/Source/Editor/Sequencer/Private/Tools/DragOperation_Stretch.cpp
+++ b/Engine/Source/Editor/Sequencer/Private/Tools/DragOperation_Stretch.cpp
@@ -2,11 +2,12 @@
 
 #include "DragOperation_Stretch.h"
 #include "MVVM/ViewModels/ViewModel.h"
+#include "MVVM/ViewModels/SequencerEditorViewModel.h"
+#include "MVVM/Selection/Selection.h"
 #include "SequencerSettings.h"
 
 #include "ISequencer.h"
 #include "Sequencer.h"
-#include "SequencerSelection.h"
 #include "MVVM/ViewModels/SectionModel.h"
 #include "MVVM/ViewModels/VirtualTrackArea.h"
 
@@ -61,12 +62,9 @@ void FEditToolDragOperation_Stretch::OnBeginDrag(const FPointerEvent& MouseEvent
 {
 	TArray<IStretchableExtension*> AllStretchables;
 
-	for (TWeakPtr<FViewModel> WeakSelectedModel : Sequencer->GetSelection().GetSelectedTrackAreaItems())
+	for (FViewModelPtr SelectedModel : Sequencer->GetViewModel()->GetSelection()->TrackArea)
 	{
-		TSharedPtr<FViewModel> Model       = WeakSelectedModel.Pin();
-		IStretchableExtension*         Stretchable = Model ? Model->CastThis<IStretchableExtension>() : nullptr;
-
-		if (Stretchable)
+		if (TViewModelPtr<IStretchableExtension> Stretchable = SelectedModel.ImplicitCast())
 		{
 			Stretchable->OnInitiateStretch(*this, StretchConstraint, &GlobalParameters);
 		}
diff --git a/Engine/Source/Editor/Sequencer/Private/Tools/EditToolDragOperations.cpp b/Engine/Source/Editor/Sequencer/Private/Tools/EditToolDragOperations.cpp
--- a/Engine/Source/Editor/Sequencer/Private/Tools/EditToolDragOperations.cpp
+++ b/Engine/Source/Editor/Sequencer/Private/Tools/EditToolDragOperations.cpp
@@ -8,6 +8,7 @@
 #include "Sequencer.h"
 #include "SequencerSettings.h"
 #include "SequencerCommonHelpers.h"
+#include "MVVM/Selection/Selection.h"
 #include "MVVM/ViewModels/ViewModel.h"
 #include "MVVM/SectionModelStorageExtension.h"
 #include "MVVM/Extensions/IDraggableTrackAreaExtension.h"
@@ -100,15 +101,16 @@ void FEditToolDragOperation::EndTransaction()
 	Sequencer.NotifyMovieSceneDataChanged( EMovieSceneDataChangeType::TrackValueChanged );
 }
 
-FResizeSection::FResizeSection( FSequencer& InSequencer, const TSet<TWeakObjectPtr<UMovieSceneSection>>& InSections, bool bInDraggingByEnd, bool bInIsSlipping )
+FResizeSection::FResizeSection( FSequencer& InSequencer, bool bInDraggingByEnd, bool bInIsSlipping )
 	: FEditToolDragOperation( InSequencer )
 	, bDraggingByEnd(bInDraggingByEnd)
 	, bIsSlipping(bInIsSlipping)
 	, MouseDownTime(0)
 {
-	Sections.Reserve(InSections.Num());
+	TSet<UMovieSceneSection*> SelectedSections = Sequencer.GetViewModel()->GetSelection()->GetSelectedSections();
+	Sections.Reserve(SelectedSections.Num());
 
-	for (TWeakObjectPtr<UMovieSceneSection> WeakSection : InSections)
+	for (TWeakObjectPtr<UMovieSceneSection> WeakSection : SelectedSections)
 	{
 		UMovieSceneSection* Section = WeakSection.Get();
 		if (Section)
@@ -625,11 +627,17 @@ FMoveKeysAndSections::FMoveKeysAndSections(FSequencer& InSequencer, ESequencerMo
 	if (EnumHasAnyFlags(MoveType, ESequencerMoveOperationType::MoveKeys))
 	{
 		// Filter out the keys on sections that are read only
-		for (const FSequencerSelectedKey& SelectedKey : Sequencer.GetSelection().GetSelectedKeys())
+		const FKeySelection& KeySelection = Sequencer.GetViewModel()->GetSelection()->KeySelection;
+
+		Keys.Reserve(KeySelection.Num());
+
+		for (FKeyHandle Key : KeySelection)
 		{
-			if (!SelectedKey.Section->IsReadOnly())
+			TSharedPtr<FChannelModel> Channel = KeySelection.GetModelForKey(Key);
+			UMovieSceneSection*       Section = Channel ? Channel->GetSection() : nullptr;
+			if (Section && !Section->IsReadOnly())
 			{
-				Keys.Add(SelectedKey);
+				Keys.Emplace(FSequencerSelectedKey(*Section, Channel, Key));
 			}
 		}
 
@@ -638,17 +646,14 @@ FMoveKeysAndSections::FMoveKeysAndSections(FSequencer& InSequencer, ESequencerMo
 
 	if (EnumHasAnyFlags(MoveType, ESequencerMoveOperationType::MoveSections))
 	{
-		for (TWeakPtr<FViewModel> WeakSelectedModel : Sequencer.GetSelection().GetSelectedTrackAreaItems())
+		for (TViewModelPtr<IDraggableTrackAreaExtension> DraggableItem : Sequencer.GetViewModel()->GetSelection()->TrackArea.Filter<IDraggableTrackAreaExtension>())
 		{
-			if (TSharedPtr<IDraggableTrackAreaExtension> DraggableItem = ICastable::CastWeakPtrShared<IDraggableTrackAreaExtension>(WeakSelectedModel))
-			{
-				DraggableItem->OnBeginDrag(*this);
-			}
+			DraggableItem->OnBeginDrag(*this);
 		}
 	}
 
 	// Always move selected marked frames along with keys and/or sections.
-	MarkedFrames = Sequencer.GetSelection().GetSelectedMarkedFrames();
+	MarkedFrames = Sequencer.GetViewModel()->GetSelection()->MarkedFrames.GetSelected();
 }
 
 void FMoveKeysAndSections::AddSnapTime(FFrameNumber SnapTime)
diff --git a/Engine/Source/Editor/Sequencer/Private/Tools/EditToolDragOperations.h b/Engine/Source/Editor/Sequencer/Private/Tools/EditToolDragOperations.h
--- a/Engine/Source/Editor/Sequencer/Private/Tools/EditToolDragOperations.h
+++ b/Engine/Source/Editor/Sequencer/Private/Tools/EditToolDragOperations.h
@@ -88,7 +88,7 @@ class FResizeSection
 public:
 
 	/** Create and initialize a new instance. */
-	FResizeSection( FSequencer& InSequencer, const TSet<TWeakObjectPtr<UMovieSceneSection>>& Sections, bool bInDraggingByEnd, bool bIsSlipping );
+	FResizeSection( FSequencer& InSequencer, bool bInDraggingByEnd, bool bIsSlipping );
 
 public:
 
diff --git a/Engine/Source/Editor/Sequencer/Private/Tools/SequencerEditTool_Movement.cpp b/Engine/Source/Editor/Sequencer/Private/Tools/SequencerEditTool_Movement.cpp
--- a/Engine/Source/Editor/Sequencer/Private/Tools/SequencerEditTool_Movement.cpp
+++ b/Engine/Source/Editor/Sequencer/Private/Tools/SequencerEditTool_Movement.cpp
@@ -3,6 +3,7 @@
 #include "Tools/SequencerEditTool_Movement.h"
 #include "MVVM/ViewModels/SectionModel.h"
 #include "MVVM/ViewModels/SequencerEditorViewModel.h"
+#include "MVVM/Selection/Selection.h"
 #include "Editor.h"
 #include "Fonts/FontMeasure.h"
 #include "MVVM/Views/STrackAreaView.h"
@@ -151,7 +152,7 @@ TSharedPtr<UE::Sequencer::ISequencerEditToolDragOperation> FSequencerEditTool_Mo
 {
 	using namespace UE::Sequencer;
 
-	FSequencerSelection& Selection = Sequencer.GetSelection();
+	FSequencerSelection& Selection = *Sequencer.GetViewModel()->GetSelection();
 	TSharedRef<SSequencer> SequencerWidget = StaticCastSharedRef<SSequencer>(Sequencer.GetSequencerWidget());
 
 	GetHotspotTime(OriginalHotspotTime);
@@ -166,8 +167,8 @@ TSharedPtr<UE::Sequencer::ISequencerEditToolDragOperation> FSequencerEditTool_Mo
 		}
 
 		// Gather the draggable sections from all selected models
-		const bool bModelsSelected = Selection.GetSelectedTrackAreaItems().Num() > 0;
-		const bool bKeySelected = Selection.GetSelectedKeys().Num() > 0;
+		const bool bModelsSelected = Selection.TrackArea.Num() > 0;
+		const bool bKeySelected = Selection.KeySelection.Num() > 0;
 		// @todo sequencer: Make this a customizable UI command modifier?
 		const bool bIsDuplicateEvent = MouseEvent.IsAltDown() || MouseEvent.GetEffectingButton() == EKeys::MiddleMouseButton;
 
@@ -190,7 +191,7 @@ TSharedPtr<UE::Sequencer::ISequencerEditToolDragOperation> FSequencerEditTool_Mo
 				for (const FSequencerSelectedKey& Key : HoveredKeys)
 				{
 					// If any are not selected, we'll treat this as a unique drag
-					if (!Selection.IsSelected(Key))
+					if (!Selection.KeySelection.IsSelected(Key.KeyHandle))
 					{
 						bUniqueDrag = true;
 					}
@@ -198,21 +199,21 @@ TSharedPtr<UE::Sequencer::ISequencerEditToolDragOperation> FSequencerEditTool_Mo
 
 				if (bUniqueDrag)
 				{
-					Selection.EmptySelectedKeys();
-					Selection.EmptySelectedTrackAreaItems();
+					Selection.KeySelection.Empty();
+					Selection.TrackArea.Empty();
 					for (const FSequencerSelectedKey& Key : HoveredKeys)
 					{
-						Selection.AddToSelection(Key);
+						Selection.KeySelection.Select(Key.WeakChannel.Pin(), Key.KeyHandle);
 					}
 				}
 			}
 			else if (SectionHotspot)
 			{
-				if (!Selection.IsSelected(SectionHotspot->WeakSectionModel.Pin()))
+				if (!Selection.TrackArea.IsSelected(SectionHotspot->WeakSectionModel.Pin()))
 				{
-					Selection.EmptySelectedKeys();
-					Selection.EmptySelectedTrackAreaItems();
-					Selection.AddToSelection(SectionHotspot->WeakSectionModel.Pin());
+					Selection.KeySelection.Empty();
+					Selection.TrackArea.Empty();
+					Selection.TrackArea.Select(SectionHotspot->WeakSectionModel.Pin());
 				}
 			}
 
@@ -222,18 +223,18 @@ TSharedPtr<UE::Sequencer::ISequencerEditToolDragOperation> FSequencerEditTool_Mo
 		// Moving section(s)?
 		if (SectionHotspot)
 		{
-			if (!Selection.IsSelected(SectionHotspot->WeakSectionModel.Pin()))
+			if (!Selection.TrackArea.IsSelected(SectionHotspot->WeakSectionModel.Pin()))
 			{
-				Selection.EmptySelectedKeys();
-				Selection.EmptySelectedTrackAreaItems();
-				Selection.AddToSelection(SectionHotspot->WeakSectionModel.Pin());
+				Selection.KeySelection.Empty();
+				Selection.TrackArea.Empty();
+				Selection.TrackArea.Select(SectionHotspot->WeakSectionModel.Pin());
 			}
 
 			if (MouseEvent.IsShiftDown())
 			{
 				const bool bDraggingByEnd = false;
 				const bool bIsSlipping = true;
-				return MakeShareable( new FResizeSection( Sequencer, Selection.GetSelectedSections(), bDraggingByEnd, bIsSlipping ) );
+				return MakeShareable( new FResizeSection( Sequencer, bDraggingByEnd, bIsSlipping ) );
 			}
 			else
 			{
@@ -249,7 +250,7 @@ TSharedPtr<UE::Sequencer::ISequencerEditToolDragOperation> FSequencerEditTool_Mo
 			for (const FSequencerSelectedKey& Key : HoveredKeys)
 			{
 				// If any are not selected, we'll treat this as a unique drag
-				if (!Selection.IsSelected(Key))
+				if (!Selection.KeySelection.IsSelected(Key.KeyHandle))
 				{
 					bUniqueDrag = true;
 				}
@@ -257,11 +258,11 @@ TSharedPtr<UE::Sequencer::ISequencerEditToolDragOperation> FSequencerEditTool_Mo
 
 			if (bUniqueDrag)
 			{
-				Selection.EmptySelectedKeys();
-				Selection.EmptySelectedTrackAreaItems();
+				Selection.KeySelection.Empty();
+				Selection.TrackArea.Empty();
 				for (const FSequencerSelectedKey& Key : HoveredKeys)
 				{
-					Selection.AddToSelection(Key);
+					Selection.KeySelection.Select(Key.WeakChannel.Pin(), Key.KeyHandle);
 				}
 			}
 
@@ -270,11 +271,11 @@ TSharedPtr<UE::Sequencer::ISequencerEditToolDragOperation> FSequencerEditTool_Mo
 		}
 	}
 	// If we're not dragging a hotspot, sections take precedence over keys
-	else if (Selection.GetSelectedTrackAreaItems().Num())
+	else if (Selection.TrackArea.Num())
 	{
 		return MakeShareable( new FMoveKeysAndSections( Sequencer, ESequencerMoveOperationType::MoveSections ) );
 	}
-	else if (Selection.GetSelectedKeys().Num())
+	else if (Selection.KeySelection.Num())
 	{
 		return MakeShareable( new FMoveKeysAndSections( Sequencer, ESequencerMoveOperationType::MoveKeys) );
 	}
diff --git a/Engine/Source/Editor/Sequencer/Private/Tools/SequencerEditTool_Selection.cpp b/Engine/Source/Editor/Sequencer/Private/Tools/SequencerEditTool_Selection.cpp
--- a/Engine/Source/Editor/Sequencer/Private/Tools/SequencerEditTool_Selection.cpp
+++ b/Engine/Source/Editor/Sequencer/Private/Tools/SequencerEditTool_Selection.cpp
@@ -10,6 +10,7 @@
 #include "IKeyArea.h"
 #include "Tools/SequencerEntityVisitor.h"
 #include "SequencerCommands.h"
+#include "MVVM/Selection/Selection.h"
 #include "MVVM/Views/ISequencerTreeView.h"
 #include "MVVM/Views/ITrackAreaHotspot.h"
 #include "MVVM/ViewModels/SequencerEditorViewModel.h"
@@ -27,12 +28,14 @@
 struct FSelectionPreviewVisitor final
 	: ISequencerEntityVisitor
 {
-	FSelectionPreviewVisitor(FSequencerSelectionPreview& InSelectionPreview, FSequencerSelection& InSelection, ESelectionPreviewState InSetStateTo, bool bInPinned)
+	FSelectionPreviewVisitor(FSequencerSelectionPreview& InSelectionPreview, UE::Sequencer::FSequencerSelection& InSelection, ESelectionPreviewState InSetStateTo, bool bInPinned)
 		: SelectionPreview(InSelectionPreview)
-		, ExistingSelection(InSelection.GetRawSelectedKeys())
+		, ExistingSelection(InSelection.KeySelection.GetSelected())
 		, SetStateTo(InSetStateTo)
 		, bPinned(bInPinned)
-	{}
+	{
+		bIsControlDown = FSlateApplication::Get().GetModifierKeys().IsControlDown();
+	}
 
 	virtual void VisitKeys(const UE::Sequencer::TViewModelPtr<UE::Sequencer::FChannelModel>& Channel, const TRange<FFrameNumber>& VisitRangeFrames) const override
 	{
@@ -44,12 +47,6 @@ struct FSelectionPreviewVisitor final
 			return;
 		}
 
-		UMovieSceneSection* Section = Channel->GetSection();
-		if (!Section)
-		{
-			return;
-		}
-
 		KeyHandlesScratch.Reset();
 		Channel->GetKeyArea()->GetKeyInfo(&KeyHandlesScratch, nullptr, VisitRangeFrames);
 
@@ -57,7 +54,7 @@ struct FSelectionPreviewVisitor final
 		{
 			// Under default behavior keys have priority, so if a key is changing selection state then we remove any sections from the selection. The user can bypass this
 			// by holding down the control key which will allow selecting both keys and sections.
-			bool bKeySelectionHasPriority = !FSlateApplication::Get().GetModifierKeys().IsControlDown();
+			bool bKeySelectionHasPriority = !bIsControlDown;
 			bool bKeyIsSelected = ExistingSelection.Contains(KeyHandlesScratch[Index]);
 
 			if (bKeySelectionHasPriority && 
@@ -68,8 +65,7 @@ struct FSelectionPreviewVisitor final
 				SelectionPreview.EmptyDefinedModelStates();
 			}
 
-			FSequencerSelectedKey Key(*Section, Channel, KeyHandlesScratch[Index]);
-			SelectionPreview.SetSelectionState(Key, SetStateTo);
+			SelectionPreview.SetSelectionState(Channel, KeyHandlesScratch[Index], SetStateTo);
 		}
 	}
 
@@ -101,6 +97,7 @@ private:
 	mutable TArray<FKeyHandle> KeyHandlesScratch;
 	ESelectionPreviewState SetStateTo;
 	bool bPinned;
+	bool bIsControlDown;
 };
 
 class FScrubTimeDragOperation
@@ -200,7 +197,7 @@ public:
 		InitialPosition = CurrentPosition = VirtualTrackArea.PhysicalToVirtual(LocalMousePos);
 		CurrentMousePos = LocalMousePos;
 
-		Sequencer.GetSelection().SuspendBroadcast();
+		EventSuppressor = Sequencer.GetViewModel()->GetSelection()->SuppressEventsLongRunning();
 
 		if (MouseEvent.IsShiftDown())
 		{
@@ -216,8 +213,8 @@ public:
 
 			// @todo: selection in transactions
 			//leave selections in the tree view alone so that dragging operations act similarly to click operations which don't change treeview selection state.
-			Sequencer.GetSelection().EmptySelectedKeys();
-			Sequencer.GetSelection().EmptySelectedTrackAreaItems();
+			Sequencer.GetViewModel()->GetSelection()->KeySelection.Empty();
+			Sequencer.GetViewModel()->GetSelection()->TrackArea.Empty();
 		}
 	}
 
@@ -294,44 +291,47 @@ public:
 
 		// Now walk everything within the current marquee range, setting preview selection states as we go
 		FSequencerEntityWalker Walker(FSequencerEntityRange(TopLeft(), BottomRight(), VirtualTrackArea.GetTickResolution()), VirtualKeySize);
-		Walker.Traverse(FSelectionPreviewVisitor(SelectionPreview, Sequencer.GetSelection(), PreviewState, TrackArea.ShowPinned()), Sequencer.GetViewModel()->GetRootModel());
+		Walker.Traverse(FSelectionPreviewVisitor(SelectionPreview, *Sequencer.GetViewModel()->GetSelection(), PreviewState, TrackArea.ShowPinned()), Sequencer.GetViewModel()->GetRootModel());
 	}
 
 	virtual void OnEndDrag(const FPointerEvent& MouseEvent, FVector2D LocalMousePos, const UE::Sequencer::FVirtualTrackArea& VirtualTrackArea) override
 	{
+		using namespace UE::Sequencer;
+
 		// finish dragging the marquee selection
-		auto& Selection = Sequencer.GetSelection();
-		auto& SelectionPreview = Sequencer.GetSelectionPreview();
+		FSequencerSelection& Selection = *Sequencer.GetViewModel()->GetSelection();
+		FSequencerSelectionPreview& SelectionPreview = Sequencer.GetSelectionPreview();
 
 		// Patch everything from the selection preview into the actual selection
-		for (const auto& Pair : SelectionPreview.GetDefinedKeyStates())
+		for (const TPair<FKeyHandle, ESelectionPreviewState>& Pair : SelectionPreview.GetDefinedKeyStates())
 		{
 			if (Pair.Value == ESelectionPreviewState::Selected)
 			{
-				// Select it in the main selection
-				Selection.AddToSelection(Pair.Key);
+				Selection.KeySelection.Select(SelectionPreview.GetChannelForKey(Pair.Key), Pair.Key);
 			}
 			else
 			{
-				Selection.RemoveFromSelection(Pair.Key);
+				Selection.KeySelection.Deselect(Pair.Key);
 			}
 		}
 
-		for (const auto& Pair : SelectionPreview.GetDefinedModelStates())
+		for (const TPair<TWeakPtr<FViewModel>, ESelectionPreviewState>& Pair : SelectionPreview.GetDefinedModelStates())
 		{
-			if (Pair.Value == ESelectionPreviewState::Selected)
+			if (TSharedPtr<FViewModel> Model = Pair.Key.Pin())
 			{
-				// Select it in the main selection
-				Selection.AddToSelection(Pair.Key.Pin());
-			}
-			else
-			{
-				Selection.RemoveFromSelection(Pair.Key.Pin());
+				if (TViewModelPtr<IOutlinerExtension> OutlinerItem = CastViewModel<IOutlinerExtension>(Model))
+				{
+					Selection.Outliner.Select(OutlinerItem);
+				}
+				else
+				{
+					Selection.TrackArea.Select(Model);
+				}
 			}
 		}
 
-		Selection.ResumeBroadcast();
-		Selection.RequestOutlinerNodeSelectionChangedBroadcast();
+		// Broadcast selection events
+		EventSuppressor = nullptr;
 
 		// We're done with this now
 		SelectionPreview.Empty();
@@ -380,6 +380,8 @@ private:
 
 	UE::Sequencer::STrackAreaView& TrackArea;
 
+	TUniquePtr<UE::Sequencer::FSelectionEventSuppressor> EventSuppressor;
+
 	/** Sequencer widget */
 	TSharedRef<SSequencer> SequencerWidget;
 
diff --git a/Engine/Source/Editor/Sequencer/Private/Widgets/SSequencerLayerBar.cpp b/Engine/Source/Editor/Sequencer/Private/Widgets/SSequencerLayerBar.cpp
--- a/Engine/Source/Editor/Sequencer/Private/Widgets/SSequencerLayerBar.cpp
+++ b/Engine/Source/Editor/Sequencer/Private/Widgets/SSequencerLayerBar.cpp
@@ -12,13 +12,13 @@
 #include "MVVM/ViewModels/SequenceModel.h"
 #include "MVVM/ViewModels/SequencerEditorViewModel.h"
 #include "MVVM/ViewModels/TrackAreaViewModel.h"
+#include "MVVM/Selection/Selection.h"
 #include "MVVM/ViewModels/ViewModel.h"
 #include "MVVM/Views/STrackAreaView.h"
 
 #include "CommonMovieSceneTools.h"
 #include "SequencerCommonHelpers.h"
 #include "MovieSceneTimeHelpers.h"
-#include "SequencerSelection.h"
 #include "Sequencer.h"
 #include "ISequencerSection.h"
 
@@ -63,11 +63,11 @@ struct FLayerBarHotspot
 
 		FSequencerSelection& Selection = Sequencer->GetSelection();
 
-		if (!Selection.IsSelected(LayerBar))
+		if (!Selection.TrackArea.IsSelected(LayerBar))
 		{
-			Selection.EmptySelectedKeys();
-			Selection.EmptySelectedTrackAreaItems();
-			Selection.AddToSelection(LayerBar);
+			Selection.KeySelection.Empty();
+			Selection.TrackArea.Empty();
+			Selection.TrackArea.Select(LayerBar);
 		}
 
 		if (!MouseEvent.IsShiftDown())
@@ -122,11 +122,11 @@ struct FLayerBarStretchHotspot
 		TSharedPtr<FSequencer> Sequencer = WeakSequencer.Pin();
 		FSequencerSelection& Selection = Sequencer->GetSelection();
 
-		if (!Selection.IsSelected(LayerBar))
+		if (!Selection.TrackArea.IsSelected(LayerBar))
 		{
-			Selection.EmptySelectedKeys();
-			Selection.EmptySelectedTrackAreaItems();
-			Selection.AddToSelection(LayerBar);
+			Selection.KeySelection.Empty();
+			Selection.TrackArea.Empty();
+			Selection.TrackArea.Select(LayerBar);
 		}
 
 		return MakeShared<FEditToolDragOperation_Stretch>(Sequencer.Get(), StretchConstraint, GetTime().GetValue());
@@ -314,7 +314,7 @@ int32 SSequencerLayerBar::OnPaint(const FPaintArgs& Args, const FGeometry& Allot
 	ESelectionPreviewState SelectionPreviewState = SelectionPreview.GetSelectionState(WeakLayerBar);
 
 	const bool bIsPreviewUnselected = SelectionPreviewState == ESelectionPreviewState::NotSelected;
-	const bool bIsUnSelected        = SelectionPreviewState == ESelectionPreviewState::Undefined && !Selection.IsSelected(WeakLayerBar);
+	const bool bIsUnSelected        = SelectionPreviewState == ESelectionPreviewState::Undefined && !Selection.TrackArea.IsSelected(WeakLayerBar);
 
 	if (!bIsPreviewUnselected && !bIsUnSelected)
 	{
diff --git a/Engine/Source/Editor/Sequencer/Public/ISequencer.h b/Engine/Source/Editor/Sequencer/Public/ISequencer.h
--- a/Engine/Source/Editor/Sequencer/Public/ISequencer.h
+++ b/Engine/Source/Editor/Sequencer/Public/ISequencer.h
@@ -41,7 +41,6 @@
 class AActor;
 class ACameraActor;
 class FSequencerKeyCollection;
-class FSequencerSelection;
 class FSequencerSelectionPreview;
 class FUICommandList;
 class IDetailsView;
@@ -80,7 +79,7 @@ namespace UE
 namespace Sequencer
 {
 
-class FEditorViewModel;
+class FSequencerEditorViewModel;
 
 } // namespace Sequencer
 } // namespace UE
@@ -286,7 +285,7 @@ public:
 	/**
 	 * Retrieve the top level view model for this sequence
 	 */
-	virtual TSharedPtr<UE::Sequencer::FEditorViewModel> GetViewModel() const = 0;
+	virtual TSharedPtr<UE::Sequencer::FSequencerEditorViewModel> GetViewModel() const = 0;
 
 	/**
 	 * Suppresses automatic evaluation the specified sequence and signature are the only difference that would prompt a re-evaluation
@@ -602,12 +601,6 @@ public:
 
 	virtual TArray<FMovieSceneMarkedFrame> GetMarkedFrames() const = 0;
 
-	virtual FSequencerSelection& GetSelection() = 0;
-	virtual FSequencerSelectionPreview& GetSelectionPreview() = 0;
-
-	virtual void SuspendSelectionBroadcast() = 0;
-	virtual void ResumeSelectionBroadcast() = 0;
-
 	/** Gets the currently selected tracks. */
 	virtual void GetSelectedTracks(TArray<UMovieSceneTrack*>& OutSelectedTracks) = 0;
 
diff --git a/Engine/Source/Editor/Sequencer/Public/MVVM/ViewModels/SequencerEditorViewModel.h b/Engine/Source/Editor/Sequencer/Public/MVVM/ViewModels/SequencerEditorViewModel.h
--- a/Engine/Source/Editor/Sequencer/Public/MVVM/ViewModels/SequencerEditorViewModel.h
+++ b/Engine/Source/Editor/Sequencer/Public/MVVM/ViewModels/SequencerEditorViewModel.h
@@ -17,6 +17,7 @@ namespace Sequencer
 {
 
 class FSequenceModel;
+class FSequencerSelection;
 class STrackAreaView;
 struct ITrackAreaHotspot;
 
@@ -31,8 +32,11 @@ public:
 
 	FSequencerEditorViewModel(TSharedRef<ISequencer> InSequencer, const FSequencerHostCapabilities& InHostCapabilities);
 
+	/** Retrieve this editor's selection class */
+	SEQUENCER_API TSharedPtr<FSequencerSelection> GetSelection() const;
+
 	// @todo_sequencer_mvvm remove this later
-	TSharedPtr<ISequencer> GetSequencer() const;
+	SEQUENCER_API TSharedPtr<ISequencer> GetSequencer() const;
 	// @todo_sequencer_mvvm remove this ASAP
 	TSharedPtr<FSequencer> GetSequencerImpl() const;
 
@@ -54,6 +58,7 @@ protected:
 	virtual TSharedPtr<FViewModel> CreateRootModelImpl() override;
 	virtual TSharedPtr<FOutlinerViewModel> CreateOutlinerImpl() override;
 	virtual TSharedPtr<FTrackAreaViewModel> CreateTrackAreaImpl() override;
+	virtual TSharedPtr<FSequencerCoreSelection> CreateSelectionImpl() override;
 	virtual bool IsReadOnly() const override;
 
 	void OnTrackAreaHotspotChanged(TSharedPtr<ITrackAreaHotspot> NewHotspot);
diff --git a/Engine/Source/Editor/Sequencer/Public/SequencerAddKeyOperation.h b/Engine/Source/Editor/Sequencer/Public/SequencerAddKeyOperation.h
--- a/Engine/Source/Editor/Sequencer/Public/SequencerAddKeyOperation.h
+++ b/Engine/Source/Editor/Sequencer/Public/SequencerAddKeyOperation.h
@@ -7,7 +7,7 @@
 #include "Containers/Map.h"
 #include "Containers/Set.h"
 #include "Containers/SparseArray.h"
-#include "MVVM/ViewModelPtr.h"
+#include "SequencerCoreFwd.h"
 #include "Misc/FrameNumber.h"
 #include "SequencerKeyParams.h"
 #include "Templates/SharedPointer.h"
@@ -28,6 +28,7 @@ namespace Sequencer
 class FChannelGroupModel;
 class FViewModel;
 class ITrackExtension;
+class IOutlinerExtension;
 
 /**
  * Temporary structure used for consistent add-key behavior for a set of display nodes
@@ -40,7 +41,7 @@ struct SEQUENCER_API FAddKeyOperation
 	 *
 	 * @param InModels A set of all the models to key
 	 */
-	static FAddKeyOperation FromNodes(const TSet<TWeakPtr<FViewModel>>& InModels);
+	static FAddKeyOperation FromNodes(const TSet<TWeakViewModelPtr<IOutlinerExtension>>& InModels);
 
 
 	/**
diff --git a/Engine/Source/Editor/SequencerCore/Private/MVVM/ViewModels/EditorViewModel.cpp b/Engine/Source/Editor/SequencerCore/Private/MVVM/ViewModels/EditorViewModel.cpp
--- a/Engine/Source/Editor/SequencerCore/Private/MVVM/ViewModels/EditorViewModel.cpp
+++ b/Engine/Source/Editor/SequencerCore/Private/MVVM/ViewModels/EditorViewModel.cpp
@@ -5,6 +5,8 @@
 #include "CoreTypes.h"
 #include "MVVM/ViewModels/OutlinerViewModel.h"
 #include "MVVM/ViewModels/TrackAreaViewModel.h"
+#include "MVVM/Selection/SequencerCoreSelection.h"
+#include "Scripting/SequencerScriptingLayer.h"
 #include "Misc/AssertionMacros.h"
 #include "Templates/TypeHash.h"
 
@@ -47,6 +49,13 @@ void FEditorViewModel::InitializeEditor()
 	TWeakPtr<FViewModel> WeakRootDataModel(RootDataModel);
 	Outliner->Initialize(FWeakViewModelPtr(WeakRootDataModel));
 
+	Selection = CreateSelectionImpl();
+
+	USequencerScriptingLayer* Scripting = CreateScriptingLayerImpl();
+	Scripting->Initialize(SharedThis(this));
+
+	ScriptingLayer.Reset(Scripting);
+
 	InitializeEditorImpl();
 }
 
@@ -70,6 +79,16 @@ TSharedPtr<FTrackAreaViewModel> FEditorViewModel::GetTrackArea() const
 	return TrackArea;
 }
 
+TSharedPtr<FSequencerCoreSelection> FEditorViewModel::GetSelection() const
+{
+	return Selection;
+}
+
+USequencerScriptingLayer* FEditorViewModel::GetScriptingLayer() const
+{
+	return ScriptingLayer.Get();
+}
+
 TSharedPtr<FViewModel> FEditorViewModel::CreateRootModelImpl()
 {
 	return MakeShared<FViewModel>();
@@ -85,6 +104,18 @@ TSharedPtr<FTrackAreaViewModel> FEditorViewModel::CreateTrackAreaImpl()
 	return MakeShared<FTrackAreaViewModel>();
 }
 
+TSharedPtr<FSequencerCoreSelection> FEditorViewModel::CreateSelectionImpl()
+{
+	struct FShellSelection : FSequencerCoreSelection
+	{};
+	return MakeShared<FShellSelection>();
+}
+
+USequencerScriptingLayer* FEditorViewModel::CreateScriptingLayerImpl()
+{
+	return NewObject<USequencerScriptingLayer>();
+}
+
 bool FEditorViewModel::IsReadOnly() const
 {
 	return false;
diff --git a/Engine/Source/Editor/SequencerCore/Private/MVVM/Views/SOutlinerItemViewBase.cpp b/Engine/Source/Editor/SequencerCore/Private/MVVM/Views/SOutlinerItemViewBase.cpp
--- a/Engine/Source/Editor/SequencerCore/Private/MVVM/Views/SOutlinerItemViewBase.cpp
+++ b/Engine/Source/Editor/SequencerCore/Private/MVVM/Views/SOutlinerItemViewBase.cpp
@@ -412,22 +412,33 @@ FSlateColor SOutlinerItemViewBase::GetNodeBackgroundTint() const
 	{
 		return FStyleColors::Select;
 	}
-	else if (EnumHasAnyFlags(SelectionState, EOutlinerSelectionState::HasSelectedKeys | EOutlinerSelectionState::HasSelectedTrackAreaItems))
+	if (EnumHasAnyFlags(SelectionState, EOutlinerSelectionState::HasSelectedKeys | EOutlinerSelectionState::HasSelectedTrackAreaItems))
 	{
 		return FStyleColors::Header;
 	}
-	else if (Editor->GetOutliner()->GetHoveredItem() == OutlinerItem.AsModel())
+
+	// If this is collapsed but has any children with selected keys or sections, we report that state on the parent
+	if (!OutlinerItem->IsExpanded())
 	{
-		return ItemStyle == EOutlinerItemViewBaseStyle::ContainerHeader
-			? FLinearColor(FColor(52, 52, 52, 255))
-			: FLinearColor(FColor(72, 72, 72, 255));
+		for (TViewModelPtr<IOutlinerExtension> Child : OutlinerItem.AsModel()->GetDescendantsOfType<IOutlinerExtension>())
+		{
+			if (EnumHasAnyFlags(SelectionState, EOutlinerSelectionState::HasSelectedKeys | EOutlinerSelectionState::HasSelectedTrackAreaItems))
+			{
+				return FStyleColors::Header;
+			}
+		}
 	}
-	else
+
+	if (Editor->GetOutliner()->GetHoveredItem() == OutlinerItem)
 	{
 		return ItemStyle == EOutlinerItemViewBaseStyle::ContainerHeader
-			? FLinearColor(FColor(48, 48, 48, 255))
-			: FLinearColor(FColor(62, 62, 62, 255));
+			? FLinearColor(FColor(52, 52, 52, 255))
+			: FLinearColor(FColor(72, 72, 72, 255));
 	}
+
+	return ItemStyle == EOutlinerItemViewBaseStyle::ContainerHeader
+		? FLinearColor(FColor(48, 48, 48, 255))
+		: FLinearColor(FColor(62, 62, 62, 255));
 }
 
 FSlateColor SOutlinerItemViewBase::GetNodeInnerBackgroundTint() const
@@ -445,7 +456,7 @@ FSlateColor SOutlinerItemViewBase::GetNodeInnerBackgroundTint() const
 		{
 			return FStyleColors::Header;
 		}
-		else if (Editor->GetOutliner()->GetHoveredItem() == OutlinerItem.AsModel())
+		else if (Editor->GetOutliner()->GetHoveredItem() == OutlinerItem)
 		{
 			return FLinearColor( FColor( 52, 52, 52, 255 ) );
 		}
