Sequencer: fix bug where key editor widget would not get cleaned up correctly when a track goes from 1 to 2 or more rows.

diff --git a/Engine/Source/Editor/Sequencer/Private/MVVM/ViewModels/TrackModel.cpp b/Engine/Source/Editor/Sequencer/Private/MVVM/ViewModels/TrackModel.cpp
--- a/Engine/Source/Editor/Sequencer/Private/MVVM/ViewModels/TrackModel.cpp
+++ b/Engine/Source/Editor/Sequencer/Private/MVVM/ViewModels/TrackModel.cpp
@@ -153,6 +153,7 @@ void FTrackModel::ForceUpdate()
 
 	FViewModelChildren OutlinerChildren = GetChildList(EViewModelListType::Outliner);
 	FViewModelChildren SectionChildren  = GetChildList(EViewModelListType::TrackArea);
+	FViewModelChildren TopLevelChannelChildren = GetChildList(GetTopLevelChannelGroupType());
 
 	UMovieSceneTrack* Track = WeakTrack.Get();
 	if (!Track)
@@ -160,6 +161,7 @@ void FTrackModel::ForceUpdate()
 		// Free outliner and section children, this track is gone.
 		OutlinerChildren.Empty();
 		SectionChildren.Empty();
+		TopLevelChannelChildren.Empty();
 		return;
 	}
 
@@ -197,6 +199,7 @@ void FTrackModel::ForceUpdate()
 		// Clear any left-over row models, layout models, or section models.
 		OutlinerChildren.Empty();
 		SectionChildren.Empty();
+		TopLevelChannelChildren.Empty();
 	}
 	else if (NumRows == 1)
 	{
@@ -206,6 +209,7 @@ void FTrackModel::ForceUpdate()
 		FScopedViewModelListHead RecycledModels(AsShared(), EViewModelListType::Recycled);
 		GetChildrenForList(&SectionList).MoveChildrenTo<IRecyclableExtension>(RecycledModels.GetChildren(), IRecyclableExtension::CallOnRecycle);
 		OutlinerChildren.MoveChildrenTo<IRecyclableExtension>(RecycledModels.GetChildren(), IRecyclableExtension::CallOnRecycle);
+		TopLevelChannelChildren.MoveChildrenTo<IRecyclableExtension>(RecycledModels.GetChildren(), IRecyclableExtension::CallOnRecycle);
 
 		// Add all sections directly to this track row
 		for (UMovieSceneSection* Section : Track->GetAllSections())
@@ -252,6 +256,7 @@ void FTrackModel::ForceUpdate()
 		FScopedViewModelListHead RecycledModels(AsShared(), EViewModelListType::Recycled);
 		GetChildrenForList(&SectionList).MoveChildrenTo<IRecyclableExtension>(RecycledModels.GetChildren(), IRecyclableExtension::CallOnRecycle);
 		OutlinerChildren.MoveChildrenTo<IRecyclableExtension>(RecycledModels.GetChildren(), IRecyclableExtension::CallOnRecycle);
+		TopLevelChannelChildren.MoveChildrenTo<IRecyclableExtension>(RecycledModels.GetChildren(), IRecyclableExtension::CallOnRecycle);
 
 		// We need to build row models so let's grab the storage for that
 		FTrackRowModelStorageExtension* TrackRowModelStorage = SequenceModel->CastDynamic<FTrackRowModelStorageExtension>();
@@ -324,11 +329,6 @@ void FTrackModel::ForceUpdate()
 			//       we fixup row indices
 		}
 	}
-
-	for (TSharedPtr<FChannelGroupModel> ChannelModel : GetDescendantsOfType<FChannelGroupModel>(false, EViewModelListType::Outliner))
-	{
-
-	}
 }
 
 FOutlinerSizing FTrackModel::GetOutlinerSizing() const
