Sequencer: Fix group manager's selection synchronization

When selection changes in the group manager tree, select the corresponding Sequencer items first. Then update the group manager tree's notion of what is selected. When this empties selection, the group item itself will be deselected. Any items within it will be synchronized with what is selected in Sequencer.

Also, fixed so this works with folder paths.

diff --git a/Engine/Source/Editor/Sequencer/Private/SSequencerGroupManager.cpp b/Engine/Source/Editor/Sequencer/Private/SSequencerGroupManager.cpp
--- a/Engine/Source/Editor/Sequencer/Private/SSequencerGroupManager.cpp
+++ b/Engine/Source/Editor/Sequencer/Private/SSequencerGroupManager.cpp
@@ -622,23 +622,26 @@ void SSequencerGroupManager::SelectSelectedItemsInSequencer()
 		return;
 	}
 
-	TGuardValue<bool> Guard(bSynchronizingSelection, true);
-
-	TSharedPtr<FSequencer> Sequencer = WeakSequencer.Pin();
-	if (!ensure(Sequencer))
+	// When selection changes in the group manager tree, select the corresponding Sequencer items first
 	{
-		return;
-	}
+		TGuardValue<bool> Guard(bSynchronizingSelection, true);
+
+		TSharedPtr<FSequencer> Sequencer = WeakSequencer.Pin();
+		if (!ensure(Sequencer))
+		{
+			return;
+		}
 	
-	TSet<FString> SelectedNodePaths;
-	GetSelectedItemsNodePaths(SelectedNodePaths);
+		TSet<FString> SelectedNodePaths;
+		GetSelectedItemsNodePaths(SelectedNodePaths);
 	
-	if (SelectedNodePaths.Num() < 1)
-	{
-		return;
+		Sequencer->SelectNodesByPath(SelectedNodePaths);
 	}
 
-	Sequencer->SelectNodesByPath(SelectedNodePaths);
+	// Then update the group manager tree's notion of what is selected. When this empties selection, 
+	// the group item itself will be deselected. Any items within it will be synchronized with what 
+	// is selected in Sequencer
+	SelectItemsSelectedInSequencer();
 }
 
 void SSequencerGroupManager::SelectItemsSelectedInSequencer()
@@ -670,9 +673,14 @@ void SSequencerGroupManager::SelectItemsSelectedInSequencer()
 			TempString.Reset();
 			Groupable->GetIdentifierForGrouping(TempString);
 
-			if (AllNodeGroupItems.Contains(TempString.ToString()))
+			for (const FString& NodeGroupPath : AllNodeGroupItems)
 			{
-				NodesPathsToSelect.Add(TempString.ToString());
+				// AllNodeGroupItems path is the full path (including folder) 
+				if (NodeGroupPath.Contains(TempString.ToString()))
+				{
+					NodesPathsToSelect.Add(NodeGroupPath);
+					break;
+				}
 			}
 		}
 	}
