Sequencer: Remove warnings about filters since they can span different sequencers

diff --git a/Engine/Source/Editor/Sequencer/Private/SSequencer.cpp b/Engine/Source/Editor/Sequencer/Private/SSequencer.cpp
--- a/Engine/Source/Editor/Sequencer/Private/SSequencer.cpp
+++ b/Engine/Source/Editor/Sequencer/Private/SSequencer.cpp
@@ -1749,14 +1749,12 @@ void SSequencer::SetTrackFilterEnabled(const FText& InTrackFilterName, bool bEna
 {
 	TSharedPtr<FSequencer> Sequencer = SequencerPtr.Pin();
 
-	bool bFound = false;
 	for (TSharedRef<FSequencerTrackFilter> TrackFilter : AllTrackFilters)
 	{
 		if (TrackFilter->SupportsSequence(Sequencer->GetFocusedMovieSceneSequence()))
 		{
 			if (InTrackFilterName.EqualToCaseIgnored(TrackFilter->GetDisplayName()))
 			{
-				bFound = true;
 				if (bEnabled && !IsTrackFilterActive(TrackFilter))
 				{
 					Sequencer->GetNodeTree()->AddFilter(TrackFilter);
@@ -1768,11 +1766,6 @@ void SSequencer::SetTrackFilterEnabled(const FText& InTrackFilterName, bool bEna
 			}
 		}
 	}
-
-	if (!bFound)
-	{
-		UE_LOG(LogSequencer, Warning, TEXT("'%s' track filter not found"), *InTrackFilterName.ToString());
-	}
 }
 
 bool SSequencer::IsTrackFilterEnabled(const FText& InTrackFilterName) const
@@ -1790,8 +1783,6 @@ bool SSequencer::IsTrackFilterEnabled(const FText& InTrackFilterName) const
 		}
 	}
 
-	UE_LOG(LogSequencer, Warning, TEXT("'%s' track filter not found"), *InTrackFilterName.ToString());
-
 	return false;
 }
 
