Sequencer: Add support for tags when bindings are copied/pasted

diff --git a/Engine/Source/Editor/Sequencer/Private/MovieSceneCopyableBinding.h b/Engine/Source/Editor/Sequencer/Private/MovieSceneCopyableBinding.h
--- a/Engine/Source/Editor/Sequencer/Private/MovieSceneCopyableBinding.h
+++ b/Engine/Source/Editor/Sequencer/Private/MovieSceneCopyableBinding.h
@@ -48,4 +48,7 @@ public:
 
 	UPROPERTY()
 	TArray<FName> FolderPath;
+
+	UPROPERTY()
+	TArray<FName> Tags;
 };
diff --git a/Engine/Source/Editor/Sequencer/Private/SequencerUtilities.cpp b/Engine/Source/Editor/Sequencer/Private/SequencerUtilities.cpp
--- a/Engine/Source/Editor/Sequencer/Private/SequencerUtilities.cpp
+++ b/Engine/Source/Editor/Sequencer/Private/SequencerUtilities.cpp
@@ -1987,6 +1987,14 @@ void FSequencerUtilities::CopyBindings(TSharedRef<ISequencer> Sequencer, const T
 				break;
 			}
 		}
+
+		for (TPair<FName, FMovieSceneObjectBindingIDs> TaggedBinding : Sequencer->GetRootMovieSceneSequence()->GetMovieScene()->AllTaggedBindings())
+		{
+			if (TaggedBinding.Value.IDs.Contains(FMovieSceneObjectBindingID(UE::MovieScene::FFixedObjectBindingID(ObjectBinding.BindingID, Sequencer->GetFocusedTemplateID()))))
+			{
+				CopyableBinding->Tags.Add(TaggedBinding.Key);
+			}
+		}
 	}
 
 	ExportObjectBindingsToText(Objects, /*out*/ ExportedText);
@@ -2077,6 +2085,8 @@ bool FSequencerUtilities::PasteBindings(const FString& TextToImport, TSharedRef<
 		return false;
 	}
 
+	UMovieScene* RootMovieScene = Sequencer->GetRootMovieSceneSequence()->GetMovieScene();
+
 	UWorld* World = Cast<UWorld>(Sequencer->GetPlaybackContext());
 
 	const FScopedTransaction Transaction(LOCTEXT("PasteBindings", "Paste Bindings"));
@@ -2139,6 +2149,16 @@ bool FSequencerUtilities::PasteBindings(const FString& TextToImport, TSharedRef<
 					GuidToFolderMap.Add(NewGuid, ParentFolder);
 				}
 
+				if (CopyableBinding->Tags.Num() > 0)
+				{
+					RootMovieScene->Modify();
+
+					for (const FName& Tag : CopyableBinding->Tags)
+					{
+						RootMovieScene->TagBinding(Tag, UE::MovieScene::FFixedObjectBindingID(NewGuid, Sequencer->GetFocusedTemplateID()));
+					}
+				}
+
 				if (FMovieScenePossessable* Possessable = MovieScene->FindPossessable(NewGuid))
 				{
 					if (TargetIndex < PasteBindingsParams.Bindings.Num())
@@ -2259,6 +2279,16 @@ bool FSequencerUtilities::PasteBindings(const FString& TextToImport, TSharedRef<
 				{
 					GuidToFolderMap.Add(NewGuid, ParentFolder);
 				}
+
+				if (CopyableBinding->Tags.Num() > 0)
+				{
+					RootMovieScene->Modify();
+
+					for (const FName& Tag : CopyableBinding->Tags)
+					{
+						RootMovieScene->TagBinding(Tag, UE::MovieScene::FFixedObjectBindingID(NewGuid, Sequencer->GetFocusedTemplateID()));
+					}
+				}
 			}
 		}
 	}
