Sequencer: Always show the Clock Source menu but only allow it to be editable if you're focused on the root sequence. Rather than not showing up at all if you're not focused on the root sequence.

diff --git a/Engine/Source/Editor/Sequencer/Private/SSequencerPlayRateCombo.cpp b/Engine/Source/Editor/Sequencer/Private/SSequencerPlayRateCombo.cpp
--- a/Engine/Source/Editor/Sequencer/Private/SSequencerPlayRateCombo.cpp
+++ b/Engine/Source/Editor/Sequencer/Private/SSequencerPlayRateCombo.cpp
@@ -308,14 +308,11 @@ TSharedRef<SWidget> SSequencerPlayRateCombo::OnCreateMenu()
 		})
 	);
 
-	if (Sequencer->GetRootMovieSceneSequence() == Sequencer->GetFocusedMovieSceneSequence())
-	{
-		MenuBuilder.AddSubMenu(
-			LOCTEXT("ClockSource", "Clock Source"),
-			LOCTEXT("ClockSource_Description", "Change which clock should be used when playing back this sequence"),
-			FNewMenuDelegate::CreateSP(this, &SSequencerPlayRateCombo::PopulateClockSourceMenu)
-		);
-	}
+	MenuBuilder.AddSubMenu(
+		LOCTEXT("ClockSource", "Clock Source"),
+		LOCTEXT("ClockSource_Description", "Change which clock should be used when playing back this sequence"),
+		FNewMenuDelegate::CreateSP(this, &SSequencerPlayRateCombo::PopulateClockSourceMenu)
+	);
 
 	MenuBuilder.AddMenuEntry(
 		LOCTEXT("LockPlayback", "Lock to Display Rate at Runtime"),
@@ -383,6 +380,8 @@ void SSequencerPlayRateCombo::PopulateClockSourceMenu(FMenuBuilder& MenuBuilder)
 
 	if (RootSequence)
 	{
+		const bool IsFocusedOnRootSequence = Sequencer->GetRootMovieSceneSequence() == Sequencer->GetFocusedMovieSceneSequence();
+
 		for (int32 Index = 0; Index < ClockSourceEnum->NumEnums() - 1; Index++)
 		{
 			if (!ClockSourceEnum->HasMetaData(TEXT("Hidden"), Index))
@@ -397,7 +396,7 @@ void SSequencerPlayRateCombo::PopulateClockSourceMenu(FMenuBuilder& MenuBuilder)
 						FNewMenuDelegate::CreateSP(this, &SSequencerPlayRateCombo::PopulateCustomClockSourceMenu),
 						FUIAction(
 							FExecuteAction::CreateSP(this, &SSequencerPlayRateCombo::SetClockSource, Value),
-							FCanExecuteAction::CreateLambda([this]{ return !GetIsSequenceReadOnly(); }),
+							FCanExecuteAction::CreateLambda([this, IsFocusedOnRootSequence]{ return !GetIsSequenceReadOnly() && IsFocusedOnRootSequence; }),
 							FIsActionChecked::CreateLambda([=]{ return RootSequence->GetMovieScene()->GetClockSource() == Value; })
 						),
 						NAME_None,
@@ -412,7 +411,7 @@ void SSequencerPlayRateCombo::PopulateClockSourceMenu(FMenuBuilder& MenuBuilder)
 						FSlateIcon(),
 						FUIAction(
 							FExecuteAction::CreateSP(this, &SSequencerPlayRateCombo::SetClockSource, Value),
-							FCanExecuteAction::CreateLambda([this]{ return !GetIsSequenceReadOnly(); }),
+							FCanExecuteAction::CreateLambda([this, IsFocusedOnRootSequence]{ return !GetIsSequenceReadOnly() && IsFocusedOnRootSequence; }),
 							FIsActionChecked::CreateLambda([=]{ return RootSequence->GetMovieScene()->GetClockSource() == Value; })
 						),
 						NAME_None,
