Sequencer: Fixed Snap to Interval not functioning when "Snap to Keys and Sections" was not selected.

diff --git a/Engine/Source/Editor/Sequencer/Private/Tools/EditToolDragOperations.cpp b/Engine/Source/Editor/Sequencer/Private/Tools/EditToolDragOperations.cpp
--- a/Engine/Source/Editor/Sequencer/Private/Tools/EditToolDragOperations.cpp
+++ b/Engine/Source/Editor/Sequencer/Private/Tools/EditToolDragOperations.cpp
@@ -759,7 +759,7 @@ void FMoveKeysAndSections::OnDrag(const FPointerEvent& MouseEvent, FVector2D Loc
 	{
 		return;
 	}
-
+	
 	// Convert the current mouse position to a time
 	FVector2D  VirtualMousePos = VirtualTrackArea.PhysicalToVirtual(LocalMousePos);
 	FFrameTime MouseTime = VirtualTrackArea.PixelToFrame(LocalMousePos.X);
@@ -778,9 +778,10 @@ void FMoveKeysAndSections::OnDrag(const FPointerEvent& MouseEvent, FVector2D Loc
 		const bool bSnapToLikeTypes = (KeysAsArray.Num() > 0 && Settings->GetSnapKeyTimesToKeys()) || (Sections.Num() > 0 && Settings->GetSnapSectionTimesToSections());
 
 		SnapField.GetValue().SetSnapToInterval(bSnapToInterval);
+		SnapField.GetValue().SetSnapToLikeTypes(bSnapToLikeTypes);
 
 		// RelativeSnapOffsets contains both our sections and our keys, and we add them all as potential things that can snap to stuff.
-		if (bSnapToLikeTypes)
+		if (bSnapToLikeTypes || bSnapToInterval)
 		{
 			ValidSnapMarkers.SetNumUninitialized(RelativeSnapOffsets.Num());
 			for (int32 Index = 0; Index < RelativeSnapOffsets.Num(); ++Index)
@@ -826,7 +827,7 @@ void FMoveKeysAndSections::OnDrag(const FPointerEvent& MouseEvent, FVector2D Loc
 
 	// Update our marked frames by moving them by our delta.
 	HandleMarkedFrameMovement(MaxDeltaX, MouseDeltaTime);
-
+	
 	// Get a list of the unique tracks in this selection and update their easing so previews draw interactively as you drag.
 	TSet<UMovieSceneTrack*> Tracks;
 	FMovieSceneSectionMovedParams SectionMovedParams(EPropertyChangeType::Interactive);
diff --git a/Engine/Source/Editor/Sequencer/Private/Tools/SequencerSnapField.cpp b/Engine/Source/Editor/Sequencer/Private/Tools/SequencerSnapField.cpp
--- a/Engine/Source/Editor/Sequencer/Private/Tools/SequencerSnapField.cpp
+++ b/Engine/Source/Editor/Sequencer/Private/Tools/SequencerSnapField.cpp
@@ -279,7 +279,12 @@ TOptional<FSequencerSnapField::FSnapResult> FSequencerSnapField::Snap(const TArr
 	float MaxSnapWeight = 0.f;
 	for (FFrameTime Time : InTimes)
 	{
-		TOptional<FSnapResult> ThisSnap = Snap(Time, Threshold);
+		TOptional<FSnapResult> ThisSnap;
+		
+		if (bSnapToLikeTypes)
+		{
+			ThisSnap = Snap(Time, Threshold);
+		}
 
 		if (bSnapToInterval)
 		{
@@ -317,4 +322,4 @@ TOptional<FSequencerSnapField::FSnapResult> FSequencerSnapField::Snap(const TArr
 	}
 
 	return ProspectiveSnap;
-}
+}
\ No newline at end of file
diff --git a/Engine/Source/Editor/Sequencer/Private/Tools/SequencerSnapField.h b/Engine/Source/Editor/Sequencer/Private/Tools/SequencerSnapField.h
--- a/Engine/Source/Editor/Sequencer/Private/Tools/SequencerSnapField.h
+++ b/Engine/Source/Editor/Sequencer/Private/Tools/SequencerSnapField.h
@@ -41,6 +41,8 @@ public:
 
 	void SetSnapToInterval(bool bInSnapToInterval) { bSnapToInterval = bInSnapToInterval; }
 
+	void SetSnapToLikeTypes(bool bInSnapToLikeTypes) { bSnapToLikeTypes = bInSnapToLikeTypes; }
+
 	/** Snap the specified time to this field with the given threshold */
 	TOptional<FSnapResult> Snap(const FFrameTime& InTime, const FFrameTime& Threshold) const;
 
@@ -56,4 +58,5 @@ private:
 	ESequencerScrubberStyle ScrubStyle;
 
 	bool bSnapToInterval = false;
+	bool bSnapToLikeTypes = false;
 };
