Sequencer: Speculative fix for null movie scene

diff --git a/Engine/Source/Editor/Sequencer/Private/Sequencer.cpp b/Engine/Source/Editor/Sequencer/Private/Sequencer.cpp
--- a/Engine/Source/Editor/Sequencer/Private/Sequencer.cpp
+++ b/Engine/Source/Editor/Sequencer/Private/Sequencer.cpp
@@ -2428,7 +2428,13 @@ TOptional<TRange<FFrameNumber>> FSequencer::GetSubSequenceRange() const
 
 TRange<FFrameNumber> FSequencer::GetSelectionRange() const
 {
-	return GetFocusedMovieSceneSequence()->GetMovieScene()->GetSelectionRange();
+	UMovieScene* FocusedMovieScene = GetFocusedMovieSceneSequence()->GetMovieScene();
+	if (!FocusedMovieScene)
+	{
+		return TRange<FFrameNumber>();
+	}
+
+	return FocusedMovieScene->GetSelectionRange();
 }
 
 
