Sequencer: Fix GetOnTimeDisplayFormatChanged delegate handling. Moved from SSequencer to Sequencer because the settings widget is no longer valid in ~SSequencer

This is a recent bug introduced by 25412557 which adds a callback to broadcast changes in time format.

#rnx

diff --git a/Engine/Source/Editor/Sequencer/Private/SSequencer.cpp b/Engine/Source/Editor/Sequencer/Private/SSequencer.cpp
--- a/Engine/Source/Editor/Sequencer/Private/SSequencer.cpp
+++ b/Engine/Source/Editor/Sequencer/Private/SSequencer.cpp
@@ -188,11 +188,6 @@ void SSequencer::Construct(const FArguments& InArgs, TSharedRef<FSequencer> InSe
 	TSharedPtr<FFrameNumberInterface> FrameNumberInterface = MakeShareable(new FFrameNumberInterface(GetDisplayFormatAttr, GetZeroPadFramesAttr, GetTickResolutionAttr, GetDisplayRateAttr));
 	NumericTypeInterface = FrameNumberInterface;
 	
-	if (USequencerSettings* Settings = GetSequencerSettings())
-	{
-		Settings->GetOnTimeDisplayFormatChanged().AddRaw(FrameNumberInterface.Get(), &FFrameNumberInterface::DisplayFormatChanged);
-	}
-
 	FTimeSliderArgs TimeSliderArgs;
 	{
 		TimeSliderArgs.ViewRange = InArgs._ViewRange;
@@ -2805,11 +2800,6 @@ SSequencer::~SSequencer()
 		Window->DestroyWindowImmediately();
 		NodeGroupManager.Reset();
 	}
-
-	if (USequencerSettings* Settings = GetSequencerSettings())
-	{
-		Settings->GetOnTimeDisplayFormatChanged().RemoveAll(NumericTypeInterface.Get());
-	}
 }
 
 
diff --git a/Engine/Source/Editor/Sequencer/Private/Sequencer.cpp b/Engine/Source/Editor/Sequencer/Private/Sequencer.cpp
--- a/Engine/Source/Editor/Sequencer/Private/Sequencer.cpp
+++ b/Engine/Source/Editor/Sequencer/Private/Sequencer.cpp
@@ -546,6 +546,9 @@ void FSequencer::InitSequencer(const FSequencerInitParams& InitParams, const TSh
 		ObjectBindings.Add(ObjectBinding);
 	}
 
+	TSharedPtr<FFrameNumberInterface> FrameNumberInterface = StaticCastSharedPtr<FFrameNumberInterface>(GetNumericTypeInterface().ToSharedPtr());
+	Settings->GetOnTimeDisplayFormatChanged().AddRaw(FrameNumberInterface.Get(), &FFrameNumberInterface::DisplayFormatChanged);
+
 	FMovieSceneObjectBindingIDCustomization::BindTo(AsShared());
 
 	ZoomAnimation = FCurveSequence();
@@ -654,6 +657,11 @@ FSequencer::~FSequencer()
 
 void FSequencer::Close()
 {
+	if (Settings && SequencerWidget)
+	{
+		Settings->GetOnTimeDisplayFormatChanged().RemoveAll(&SequencerWidget->GetNumericTypeInterface().Get());
+	}
+
 	for (FLevelEditorViewportClient* LevelVC : GEditor->GetLevelViewportClients())
 	{
 		if (LevelVC != nullptr)
