Curve Editor: RequestSyncSelection after the curve editor ids have been updated in FCurveEditorIntegrationExtension::UpdateCurveEditor.

RequestSyncSelection is also called from SSequencer::HandleOutlinerNodeSelectionChanged and that can't be removed because that handles when selection changes without having to add/remove items in the curve editor.

Also, guard around multiple calls to RequestSyncSelection so that SyncSelection only runs once.

diff --git a/Engine/Source/Editor/Sequencer/Private/MVVM/CurveEditorExtension.cpp b/Engine/Source/Editor/Sequencer/Private/MVVM/CurveEditorExtension.cpp
--- a/Engine/Source/Editor/Sequencer/Private/MVVM/CurveEditorExtension.cpp
+++ b/Engine/Source/Editor/Sequencer/Private/MVVM/CurveEditorExtension.cpp
@@ -355,8 +355,16 @@ FKeyAttributes FCurveEditorExtension::GetDefaultKeyAttributes() const
 	}
 }
 
+static bool bSyncSelectionRequested = false;
 void FCurveEditorExtension::RequestSyncSelection()
 {
+	if (bSyncSelectionRequested)
+	{
+		return;
+	}
+
+	bSyncSelectionRequested = true;
+
 	// We schedule selection syncing to the next editor tick because we might want to select items that
 	// have just been added to the curve editor tree this tick. If it happened after the Slate update,
 	// these items don't yet have a UI widget, and so selecting them doesn't do anything.
@@ -367,6 +375,8 @@ void FCurveEditorExtension::RequestSyncSelection()
 	TWeakPtr<FSequencerEditorViewModel> WeakRootViewModel(WeakOwnerModel);
 	GEditor->GetTimerManager()->SetTimerForNextTick([WeakRootViewModel]()
 	{
+		bSyncSelectionRequested = false;
+	
 		TSharedPtr<FSequencerEditorViewModel> RootViewModel = WeakRootViewModel.Pin();
 		if (!RootViewModel.IsValid())
 		{
diff --git a/Engine/Source/Editor/Sequencer/Private/MVVM/CurveEditorIntegrationExtension.cpp b/Engine/Source/Editor/Sequencer/Private/MVVM/CurveEditorIntegrationExtension.cpp
--- a/Engine/Source/Editor/Sequencer/Private/MVVM/CurveEditorIntegrationExtension.cpp
+++ b/Engine/Source/Editor/Sequencer/Private/MVVM/CurveEditorIntegrationExtension.cpp
@@ -133,6 +133,9 @@ void FCurveEditorIntegrationExtension::UpdateCurveEditor()
 			}
 		}
 	}
+
+	// Synchronize selection after the view models have been added to the curve editor. Synchronization depends on curve model tree item IDs
+	CurveEditorExtension->RequestSyncSelection();
 }
 
 FCurveEditorTreeItemID FCurveEditorIntegrationExtension::AddToCurveEditor(TViewModelPtr<ICurveEditorTreeItemExtension> InViewModel, TSharedPtr<FCurveEditor> InCurveEditor)
