Sequencer: fix mvvm refresh issue where view-models who already have a recycle list would not recycle their children

This meant that a re-used track-row view-model would not get its children cleared out before it is rebuilt with new underlying data.

diff --git a/Engine/Source/Editor/Sequencer/Private/MVVM/ViewModels/TrackModelLayoutBuilder.cpp b/Engine/Source/Editor/Sequencer/Private/MVVM/ViewModels/TrackModelLayoutBuilder.cpp
--- a/Engine/Source/Editor/Sequencer/Private/MVVM/ViewModels/TrackModelLayoutBuilder.cpp
+++ b/Engine/Source/Editor/Sequencer/Private/MVVM/ViewModels/TrackModelLayoutBuilder.cpp
@@ -47,7 +47,8 @@ FHierarchicalModelListRefresher::FHierarchicalModelListRefresher(TSharedPtr<FVie
 
 void FHierarchicalModelListRefresher::RecycleChildren(const TSharedPtr<FViewModel>& InModel, FViewModelChildren InExistingChildren)
 {
-	if (!InModel->FindChildList(EViewModelListType::Recycled))
+	TOptional<FViewModelChildren> RecycledChildren = InModel->FindChildList(EViewModelListType::Recycled);
+	if (!RecycledChildren)
 	{
 		if (RecycledLists.Num() == 0 || RecycledLists.Last().GetSlack() == 0)
 		{
@@ -57,9 +58,10 @@ void FHierarchicalModelListRefresher::RecycleChildren(const TSharedPtr<FViewMode
 
 		RecycledLists.Last().Emplace(InModel, EViewModelListType::Recycled);
 
-		FViewModelChildren RecycledChildren = InModel->GetChildList(EViewModelListType::Recycled);
-		InExistingChildren.MoveChildrenTo<IRecyclableExtension>(RecycledChildren, IRecyclableExtension::CallOnRecycle);
+		RecycledChildren = InModel->GetChildList(EViewModelListType::Recycled);
 	}
+	check(RecycledChildren.IsSet());
+	InExistingChildren.MoveChildrenTo<IRecyclableExtension>(RecycledChildren.GetValue(), IRecyclableExtension::CallOnRecycle);
 }
 
 bool FHierarchicalModelListRefresher::IsValid() const
