Sequencer: improve edge interaction for sections

- Draw easing handles as rounded triangles sized to match the grip size, instead of the previous bitmap.
- Easing handles are bigger, with bigger hit-testing area.
- Expose key size through the section interface in addition to grip size
- Fix bug where small sections can result in negative time-to-pixel conversion

diff --git a/Engine/Source/Editor/EditorStyle/Private/StarshipStyle.cpp b/Engine/Source/Editor/EditorStyle/Private/StarshipStyle.cpp
--- a/Engine/Source/Editor/EditorStyle/Private/StarshipStyle.cpp
+++ b/Engine/Source/Editor/EditorStyle/Private/StarshipStyle.cpp
@@ -2149,7 +2149,7 @@ void FStarshipEditorStyle::FStyle::SetupSequencerStyles()
 		Set("Sequencer.Section.FilmBorder", new IMAGE_BRUSH(TEXT("Sequencer/SectionFilmBorder"), FVector2D(10, 7), FLinearColor::White, ESlateBrushTileType::Horizontal));
 		Set("Sequencer.Section.GripLeft", new FSlateRoundedBoxBrush(FStyleColors::White, FVector4(4, 0, 0, 4)));
 		Set("Sequencer.Section.GripRight", new FSlateRoundedBoxBrush(FStyleColors::White, FVector4(0, 4, 4, 0)));
-		Set("Sequencer.Section.EasingHandle", new IMAGE_BRUSH("Sequencer/EasingHandle", FVector2D(10.f, 10.f)));
+		Set("Sequencer.Section.EasingHandle", new FSlateColorBrush(FStyleColors::White));
 
 		Set("Sequencer.Section.Background_Collapsed", new FSlateRoundedBoxBrush(FLinearColor::White, FVector4(4, 4, 4, 4)));
 		Set("Sequencer.Section.Background_Header", new FSlateRoundedBoxBrush(FLinearColor::White, FVector4(4, 4, 0, 0)));
diff --git a/Engine/Source/Editor/Sequencer/Private/MVVM/Views/SChannelView.cpp b/Engine/Source/Editor/Sequencer/Private/MVVM/Views/SChannelView.cpp
--- a/Engine/Source/Editor/Sequencer/Private/MVVM/Views/SChannelView.cpp
+++ b/Engine/Source/Editor/Sequencer/Private/MVVM/Views/SChannelView.cpp
@@ -502,9 +502,12 @@ int32 SChannelView::OnPaint(const FPaintArgs& Args, const FGeometry& AllottedGeo
 		LayerId = Channel->CustomPaint(ChannelPaintArgs, LayerId);
 	}
 
+	TSharedPtr<ISequencerSection> SectionInterface = Section->GetSectionInterface();
+
 	FSequencerKeyRendererInterface KeyRenderInterface(Sequencer.Get(), TrackArea->GetHotspot());
 	FChannelViewKeyCachedState NewCachedState(VisibleRange, TrackArea->GetHotspot(), Model, Sequencer.Get());
 	NewCachedState.bIsChannelHovered = IsHovered();
+	NewCachedState.KeySizePx = SectionInterface->GetKeySize();
 
 	FKeyBatchParameters Params(RelativeTimeToPixel);
 	Params.CacheState = KeyRendererCache.IsSet() ? NewCachedState.CompareTo(KeyRendererCache.GetValue()) : EViewDependentCacheFlags::All;
@@ -725,8 +728,6 @@ void SChannelView::CreateKeysUnderMouse(const FVector2D& MousePosition, const FG
 		TSharedPtr<FTrackModel>             TrackModel             = Section->FindAncestorOfType<FTrackModel>();
 		TSharedPtr<IObjectBindingExtension> ObjectBindingExtension = Section->FindAncestorOfType<IObjectBindingExtension>();
 
-		FGuid ObjectBinding = ObjectBindingExtension ? ObjectBindingExtension->GetObjectGuid() : FGuid();
-
 		FVector2D LocalSpaceMousePosition = AllottedGeometry.AbsoluteToLocal(MousePosition);
 		const FFrameTime CurrentTime = RelativeTimeToPixel.PixelToFrame(LocalSpaceMousePosition.X);
 
diff --git a/Engine/Source/Editor/Sequencer/Private/SSequencerSection.cpp b/Engine/Source/Editor/Sequencer/Private/SSequencerSection.cpp
--- a/Engine/Source/Editor/Sequencer/Private/SSequencerSection.cpp
+++ b/Engine/Source/Editor/Sequencer/Private/SSequencerSection.cpp
@@ -865,9 +865,10 @@ void SSequencerSection::Construct( const FArguments& InArgs, TSharedPtr<FSequenc
 		[
 			SAssignNew(ChildLaneWidgets, SCompoundTrackLaneView)
 			.TimeToPixel(FGetTimeToPixel::CreateLambda([this](const FGeometry& AllottedGeometry){
-				FGeometry SectionGeometry = MakeSectionGeometryWithoutHandles( AllottedGeometry, SectionInterface );
+				// The given allotted geometry already has the handle padding removed, since it's inside
+				// our child slot (see padding above). We can therefore build the time converter directly.
 				UMovieSceneSection& Section = *SectionInterface->GetSectionObject();
-				return ConstructTimeConverterForSection(SectionGeometry, Section, GetSequencer());
+				return ConstructTimeConverterForSection(AllottedGeometry, Section, GetSequencer());
 			}))
 		]
 
@@ -1056,7 +1057,7 @@ int32 SSequencerSection::GetOverlapPriority() const
 	return 0;
 }
 
-bool SSequencerSection::CheckForEasingHandleInteraction( const FPointerEvent& MouseEvent, const FGeometry& SectionGeometry )
+bool SSequencerSection::CheckForEasingHandleInteraction( const FPointerEvent& MouseEvent, const FGeometry& AllottedGeometry )
 {
 	UMovieSceneSection* ThisSection = SectionInterface->GetSectionObject();
 	if (!ThisSection)
@@ -1077,15 +1078,18 @@ bool SSequencerSection::CheckForEasingHandleInteraction( const FPointerEvent& Mo
 		return false;
 	}
 
-	FTimeToPixel TimeToPixelConverter = ConstructTimeConverterForSection(MakeSectionGeometryWithoutHandles(SectionGeometry, SectionInterface), *ThisSection, GetSequencer());
+	FGeometry SectionGeometry = MakeSectionGeometryWithoutHandles(AllottedGeometry);
+	FTimeToPixel TimeToPixelConverter = ConstructTimeConverterForSection(SectionGeometry, *ThisSection, GetSequencer());
 
-	const double MouseTime = TimeToPixelConverter.PixelToSeconds(SectionGeometry.AbsoluteToLocal(MouseEvent.GetScreenSpacePosition()).X);
-	// We intentionally give the handles a little more hit-test area than is visible as they are quite small
-	const double HalfHandleSizeX = TimeToPixelConverter.PixelToSeconds(8.f) - TimeToPixelConverter.PixelToSeconds(0.f);
+	const double MousePositionX = SectionGeometry.AbsoluteToLocal(MouseEvent.GetScreenSpacePosition()).X;
+
+	// Easing handle is a square in the top corner of the grip handle.
+	const double GripSizePx = SectionInterface->GetSectionGripSize();
+	const double GripHalfSizePx = GripSizePx;
 
 	// Now test individual easing handles if we're at the correct vertical position
 	float LocalMouseY = SectionGeometry.AbsoluteToLocal(MouseEvent.GetScreenSpacePosition()).Y;
-	if (LocalMouseY < 0.f || LocalMouseY > 5.f)
+	if (LocalMouseY < 0.f || LocalMouseY > GripSizePx)
 	{
 		return false;
 	}
@@ -1108,10 +1112,16 @@ bool SSequencerSection::CheckForEasingHandleInteraction( const FPointerEvent& Mo
 
 		if (EasingSectionObj->HasStartFrame() && EnumHasAllFlags(EasingFlags, EMovieSceneTrackEasingSupportFlags::ManualEaseIn))
 		{
-			TRange<FFrameNumber> EaseInRange      = EasingSectionObj->GetEaseInRange();
-			double               HandlePositionIn = ( EaseInRange.IsEmpty() ? EasingSectionObj->GetInclusiveStartFrame() : EaseInRange.GetUpperBoundValue() ) / TimeToPixelConverter.GetTickResolution();
+			// Compute the ease-in handle screen position.
+			const TRange<FFrameNumber> EaseInRange = EasingSectionObj->GetEaseInRange();
+			const FFrameNumber HandleTime = EaseInRange.IsEmpty() ? EasingSectionObj->GetInclusiveStartFrame() : EaseInRange.GetUpperBoundValue();
+			const double HandlePosition = TimeToPixelConverter.FrameToPixel(HandleTime);
+
+			// If the handle is too close to the left edge that we can't hit-test with the full grip width,
+			// offset it to the right.
+			const double HandleOffset = FMath::Max(0.f, GripHalfSizePx - HandlePosition);
 
-			if (FMath::IsNearlyEqual(MouseTime, HandlePositionIn, HalfHandleSizeX))
+			if (FMath::IsNearlyEqual(MousePositionX, HandlePosition + HandleOffset, GripHalfSizePx))
 			{
 				TrackAreaViewModel->SetHotspot(MakeShared<FSectionEasingHandleHotspot>(ESequencerEasingType::In, SectionModel, Sequencer));
 				return true;
@@ -1120,10 +1130,16 @@ bool SSequencerSection::CheckForEasingHandleInteraction( const FPointerEvent& Mo
 
 		if (EasingSectionObj->HasEndFrame() && EnumHasAllFlags(EasingFlags, EMovieSceneTrackEasingSupportFlags::ManualEaseOut))
 		{
-			TRange<FFrameNumber> EaseOutRange      = EasingSectionObj->GetEaseOutRange();
-			double               HandlePositionOut = (EaseOutRange.IsEmpty() ? EasingSectionObj->GetExclusiveEndFrame() : EaseOutRange.GetLowerBoundValue() ) / TimeToPixelConverter.GetTickResolution();
+			// Compute the ease-out handle screen position.
+			TRange<FFrameNumber> EaseOutRange = EasingSectionObj->GetEaseOutRange();
+			const FFrameNumber HandleTime = EaseOutRange.IsEmpty() ? EasingSectionObj->GetExclusiveEndFrame() : EaseOutRange.GetLowerBoundValue();
+			const double HandlePosition = TimeToPixelConverter.FrameToPixel(HandleTime);
 
-			if (FMath::IsNearlyEqual(MouseTime, HandlePositionOut, HalfHandleSizeX))
+			// If the handle is too close to the right edge that we can't hit-test with the full grip width,
+			// offset it to the left.
+			const double HandleOffset = FMath::Max(0.f, GripHalfSizePx - (SectionGeometry.Size.X - HandlePosition));
+
+			if (FMath::IsNearlyEqual(MousePositionX, HandlePosition - HandleOffset, GripHalfSizePx))
 			{
 				TrackAreaViewModel->SetHotspot(MakeShared<FSectionEasingHandleHotspot>(ESequencerEasingType::Out, SectionModel, Sequencer));
 				return true;
@@ -1135,7 +1151,7 @@ bool SSequencerSection::CheckForEasingHandleInteraction( const FPointerEvent& Mo
 }
 
 
-bool SSequencerSection::CheckForEdgeInteraction( const FPointerEvent& MouseEvent, const FGeometry& SectionGeometry )
+bool SSequencerSection::CheckForEdgeInteraction( const FPointerEvent& MouseEvent, const FGeometry& AllottedGeometry )
 {
 	UMovieSceneSection* ThisSection = SectionInterface->GetSectionObject();
 	if (!ThisSection)
@@ -1166,8 +1182,8 @@ bool SSequencerSection::CheckForEdgeInteraction( const FPointerEvent& MouseEvent
 		}
 	}
 
-	FGeometry SectionGeometryWithoutHandles = MakeSectionGeometryWithoutHandles(SectionGeometry, SectionInterface);
-	FTimeToPixel TimeToPixelConverter       = ConstructTimeConverterForSection(SectionGeometryWithoutHandles, *ThisSection, GetSequencer());
+	FGeometry SectionGeometry = MakeSectionGeometryWithoutHandles(AllottedGeometry);
+	FTimeToPixel TimeToPixelConverter = ConstructTimeConverterForSection(SectionGeometry, *ThisSection, GetSequencer());
 
 	for (const TSharedPtr<FSectionModel>& UnderlappingSection : AllUnderlappingSections)
 	{
@@ -1179,12 +1195,12 @@ bool SSequencerSection::CheckForEdgeInteraction( const FPointerEvent& MouseEvent
 		}
 
 		const float ThisHandleOffset = UnderlappingSectionObj == ThisSection ? HandleOffsetPx : 0.f;
-		FVector2D GripSize( UnderlappingSectionInterface->GetSectionGripSize(), SectionGeometry.Size.Y );
+		const FVector2D GripSize( UnderlappingSectionInterface->GetSectionGripSize(), SectionGeometry.Size.Y );
 
 		if (UnderlappingSectionObj->HasStartFrame())
 		{
 			// Make areas to the left and right of the geometry.  We will use these areas to determine if someone dragged the left or right edge of a section
-			FGeometry SectionRectLeft = SectionGeometryWithoutHandles.MakeChild(
+			FGeometry SectionRectLeft = SectionGeometry.MakeChild(
 				GripSize,
 				FSlateLayoutTransform(FVector2D( TimeToPixelConverter.FrameToPixel(UnderlappingSectionObj->GetInclusiveStartFrame()) - ThisHandleOffset, 0.f ))
 			);
@@ -1198,7 +1214,7 @@ bool SSequencerSection::CheckForEdgeInteraction( const FPointerEvent& MouseEvent
 
 		if (UnderlappingSectionObj->HasEndFrame())
 		{
-			FGeometry SectionRectRight = SectionGeometryWithoutHandles.MakeChild(
+			FGeometry SectionRectRight = SectionGeometry.MakeChild(
 				GripSize,
 				FSlateLayoutTransform(FVector2D( TimeToPixelConverter.FrameToPixel(UnderlappingSectionObj->GetExclusiveEndFrame()) - UnderlappingSectionInterface->GetSectionGripSize() + ThisHandleOffset, 0 ))
 			);
@@ -1213,7 +1229,7 @@ bool SSequencerSection::CheckForEdgeInteraction( const FPointerEvent& MouseEvent
 	return false;
 }
 
-bool SSequencerSection::CheckForEasingAreaInteraction( const FPointerEvent& MouseEvent, const FGeometry& SectionGeometry )
+bool SSequencerSection::CheckForEasingAreaInteraction( const FPointerEvent& MouseEvent, const FGeometry& AllottedGeometry )
 {
 	TSharedPtr<FSectionModel> ThisSectionModel = WeakSectionModel.Pin();
 	if (!ThisSectionModel)
@@ -1221,10 +1237,10 @@ bool SSequencerSection::CheckForEasingAreaInteraction( const FPointerEvent& Mous
 		return false;
 	}
 
-	FVector2D LocalMousePos = SectionGeometry.AbsoluteToLocal(MouseEvent.GetScreenSpacePosition());
+	FVector2D LocalMousePos = AllottedGeometry.AbsoluteToLocal(MouseEvent.GetScreenSpacePosition());
 
 	TViewModelPtr<IGeometryExtension> HeaderGeometry = ThisSectionModel->FindAncestorOfType<IOutlinerExtension>().ImplicitCast();
-	const float EasingHeight = HeaderGeometry ? HeaderGeometry->GetVirtualGeometry().GetHeight() : SectionGeometry.GetLocalSize().Y;
+	const float EasingHeight = HeaderGeometry ? HeaderGeometry->GetVirtualGeometry().GetHeight() : AllottedGeometry.GetLocalSize().Y;
 
 	if (LocalMousePos.Y < 0 || LocalMousePos.Y > EasingHeight)
 	{
@@ -1232,7 +1248,8 @@ bool SSequencerSection::CheckForEasingAreaInteraction( const FPointerEvent& Mous
 	}
 
 	UMovieSceneSection* ThisSection = SectionInterface->GetSectionObject();
-	FTimeToPixel TimeToPixelConverter = ConstructTimeConverterForSection(MakeSectionGeometryWithoutHandles(SectionGeometry, SectionInterface), *ThisSection, GetSequencer());
+	FGeometry SectionGeometry = MakeSectionGeometryWithoutHandles(AllottedGeometry);
+	FTimeToPixel TimeToPixelConverter = ConstructTimeConverterForSection(SectionGeometry, *ThisSection, GetSequencer());
 	const FFrameNumber MouseTime = TimeToPixelConverter.PixelToFrame(LocalMousePos.X).FrameNumber;
 
 	// First off, set the hotspot to an easing area if necessary
@@ -1267,6 +1284,19 @@ bool SSequencerSection::CheckForEasingAreaInteraction( const FPointerEvent& Mous
 	return false;
 }
 
+bool SSequencerSection::CheckForSectionInteraction( const FPointerEvent& MouseEvent, const FGeometry& AllottedGeometry )
+{
+	FGeometry SectionGeometry = MakeSectionGeometryWithoutHandles(AllottedGeometry);
+	if (SectionGeometry.IsUnderLocation(MouseEvent.GetScreenSpacePosition()))
+	{
+		TSharedPtr<FTrackAreaViewModel> TrackAreaViewModel = GetTrackAreaViewModel();
+		TrackAreaViewModel->SetHotspot( MakeShareable( new FSectionHotspot(WeakSectionModel, Sequencer)) );
+		return true;
+	}
+
+	return false;
+}
+
 FSequencer& SSequencerSection::GetSequencer() const
 {
 	return *Sequencer.Pin().Get();
@@ -1313,7 +1343,7 @@ int32 SSequencerSection::OnPaint( const FPaintArgs& Args, const FGeometry& Allot
 
 	const ESlateDrawEffect DrawEffects = bEnabled ? ESlateDrawEffect::None : ESlateDrawEffect::DisabledEffect;
 
-	FGeometry SectionGeometry = MakeSectionGeometryWithoutHandles( AllottedGeometry, SectionInterface );
+	FGeometry SectionGeometry = MakeSectionGeometryWithoutHandles(AllottedGeometry);
 
 	FSequencerSectionPainterImpl Painter(GetSequencer(), TrackAreaViewModel, SectionModel, OutDrawElements, SectionGeometry, *this);
 
@@ -1352,7 +1382,7 @@ int32 SSequencerSection::OnPaint( const FPaintArgs& Args, const FGeometry& Allot
 	// reported. This enables us to still draw and interact with keys that overlap section handles, even if we drew the handles chronologically later
 	DrawSectionHandles(AllottedGeometry, OutDrawElements, LayerId, DrawEffects, SelectionColor, Hotspot);
 
-	Painter.LayerId = LayerId;
+	Painter.LayerId = (++LayerId);
 	PaintEasingHandles( Painter, SelectionColor, Hotspot );
 
 	// Artificially increase the layer now to ensure that we are now painting above keys now
@@ -1423,11 +1453,14 @@ int32 SSequencerSection::OnPaint( const FPaintArgs& Args, const FGeometry& Allot
 		OutDrawElements.PopClip();
 	}
 
+	++LayerId;
+
 	if (Painter.bClipRectEnabled)
 	{
 		OutDrawElements.PopClip();
 	}
-	return LayerId + 1;
+
+	return LayerId;
 }
 
 
@@ -1516,52 +1549,102 @@ void SSequencerSection::PaintEasingHandles( FSequencerSectionPainter& InPainter,
 			continue;
 		}
 
+		// We only draw the corner radius for the easing handle when it's at the same place as the grip. Otherwise we
+		// draw a straight triangle.
+		const float HandleCornerRadius = 2.f;
+		const float MinHandleSize = 8.f;
+
 		const ESlateDrawEffect DrawEffects = InPainter.bParentEnabled ? ESlateDrawEffect::None : ESlateDrawEffect::DisabledEffect;
+		const bool bIsSectionToKey = Track && Track->GetAllSections().Num() > 1 && Track->GetSectionToKey() == UnderlappingSectionObj;
+
+		// If this is the section to key, we draw the easing handle in the same bright green color as the border
+		// outline. We also make the handle a bit bigger, otherwise that border being drawn on top makes it look
+		// smaller.
+		// If there isn't any "section to key" border, we just draw the handle white.
+		const FLinearColor InactiveHandleColor = bIsSectionToKey ? FStyleColors::Success.GetColor(FWidgetStyle()) : FStyleColors::AccentYellow.GetColor(FWidgetStyle());
+		const float HalfSectionHeight = SectionInterface->GetSectionHeight() / 2.f;
+		const float HandleSize = FMath::Max(MinHandleSize, FMath::Min(HalfSectionHeight, SectionInterface->GetSectionGripSize())) + (bIsSectionToKey ? 3.f : 0.f);
+
+		const FSlateBrush* HandleBrush = FAppStyle::GetBrush("Sequencer.Section.EasingHandle");
+		FSlateResourceHandle HandleResource = FSlateApplication::Get().GetRenderer()->GetResourceHandle(*HandleBrush);
 
-		const FSlateBrush* EasingHandle = FAppStyle::GetBrush("Sequencer.Section.EasingHandle");
-		FVector2D HandleSize(10.f, 10.f);
+		const FSlateRenderTransform& HandleRenderTransform = InPainter.SectionGeometry.GetAccumulatedRenderTransform();
 
 		if (UnderlappingSectionObj->HasStartFrame() && EnumHasAllFlags(EasingSupportFlags, EMovieSceneTrackEasingSupportFlags::ManualEaseIn))
 		{
 			TRange<FFrameNumber> EaseInRange = UnderlappingSectionObj->GetEaseInRange();
-			// Always draw handles if the section is highlighted, even if there is no range (to allow manual adjustment)
-			FFrameNumber HandleFrame = EaseInRange.IsEmpty() ? UnderlappingSectionObj->GetInclusiveStartFrame() : UE::MovieScene::DiscreteExclusiveUpper(EaseInRange);
-			FVector2f HandlePos(TimeToPixelConverter.FrameToPixel(HandleFrame), 0.f);
-			FSlateDrawElement::MakeBox(
+			const bool bHasEaseIn = !EaseInRange.IsEmpty();
+			FFrameNumber HandleFrame = bHasEaseIn ? UE::MovieScene::DiscreteExclusiveUpper(EaseInRange) : UnderlappingSectionObj->GetInclusiveStartFrame();
+			const float HandlePos(TimeToPixelConverter.FrameToPixel(HandleFrame));
+			FColor HandleColor = (bLeftHandleActive ? SelectionColor : InactiveHandleColor).ToFColorSRGB();
+
+			TArray<FSlateVertex> Verts;
+			Verts.Add(FSlateVertex::Make<ESlateVertexRounding::Disabled>(HandleRenderTransform, FVector2f(HandlePos, 0.f), FVector2f(0.f, 0.f), HandleColor));
+			Verts.Add(FSlateVertex::Make<ESlateVertexRounding::Disabled>(HandleRenderTransform, FVector2f(HandlePos, HandleCornerRadius), FVector2f(0.f, 0.1f), HandleColor));
+			Verts.Add(FSlateVertex::Make<ESlateVertexRounding::Disabled>(HandleRenderTransform, FVector2f(HandlePos + HandleCornerRadius, 0.f), FVector2f(0.1f, 0.f), HandleColor));
+			Verts.Add(FSlateVertex::Make<ESlateVertexRounding::Disabled>(HandleRenderTransform, FVector2f(HandlePos + HandleSize, 0.f), FVector2f(1.f, 0.f), HandleColor));
+			Verts.Add(FSlateVertex::Make<ESlateVertexRounding::Disabled>(HandleRenderTransform, FVector2f(HandlePos, HandleSize), FVector2f(0.f, 1.f), HandleColor));
+
+			TArray<SlateIndex> Indices;
+			if (bHasEaseIn)
+			{
+				Indices.Add(0); Indices.Add(3); Indices.Add(4);
+			}
+			else
+			{
+				Indices.Add(1); Indices.Add(2); Indices.Add(3);
+				Indices.Add(1); Indices.Add(3); Indices.Add(4);
+			}
+
+			FSlateDrawElement::MakeCustomVerts(
 				InPainter.DrawElements,
-				// always draw selected keys on top of other keys
 				InPainter.LayerId,
-				// Center the key along X.  Ensure the middle of the key is at the actual key time
-				InPainter.SectionGeometry.ToPaintGeometry(
-					HandleSize,
-					FSlateLayoutTransform(HandlePos - FVector2f(HandleSize.X*0.5f, 0.f))
-				),
-				EasingHandle,
-				DrawEffects,
-				(bLeftHandleActive ? SelectionColor : EasingHandle->GetTint(FWidgetStyle()))
+				HandleResource,
+				Verts,
+				Indices,
+				nullptr,
+				0,
+				0,
+				DrawEffects | ESlateDrawEffect::PreMultipliedAlpha
 			);
 		}
 
 		if (UnderlappingSectionObj->HasEndFrame() && EnumHasAllFlags(EasingSupportFlags, EMovieSceneTrackEasingSupportFlags::ManualEaseOut))
 		{
 			TRange<FFrameNumber> EaseOutRange = UnderlappingSectionObj->GetEaseOutRange();
+			const bool bHasEaseOut = !EaseOutRange.IsEmpty();
+			FFrameNumber HandleFrame = bHasEaseOut ? UE::MovieScene::DiscreteInclusiveLower(EaseOutRange) : UnderlappingSectionObj->GetExclusiveEndFrame();
+			const float HandlePos(TimeToPixelConverter.FrameToPixel(HandleFrame));
+			const FColor HandleColor = (bRightHandleActive ? SelectionColor : InactiveHandleColor).ToFColorSRGB();
 
-			// Always draw handles if the section is highlighted, even if there is no range (to allow manual adjustment)
-			FFrameNumber HandleFrame = EaseOutRange.IsEmpty() ? UnderlappingSectionObj->GetExclusiveEndFrame() : UE::MovieScene::DiscreteInclusiveLower(EaseOutRange);
-			FVector2f    HandlePos   = FVector2f(TimeToPixelConverter.FrameToPixel(HandleFrame), 0.f);
+			TArray<FSlateVertex> Verts;
+			Verts.Add(FSlateVertex::Make<ESlateVertexRounding::Disabled>(HandleRenderTransform, FVector2f(HandlePos, 0.f), FVector2f(1.f, 0.f), HandleColor));
+			Verts.Add(FSlateVertex::Make<ESlateVertexRounding::Disabled>(HandleRenderTransform, FVector2f(HandlePos - HandleCornerRadius, 0.f), FVector2f(0.9f, 0.f), HandleColor));
+			Verts.Add(FSlateVertex::Make<ESlateVertexRounding::Disabled>(HandleRenderTransform, FVector2f(HandlePos, HandleCornerRadius), FVector2f(1.f, 0.1f), HandleColor));
+			Verts.Add(FSlateVertex::Make<ESlateVertexRounding::Disabled>(HandleRenderTransform, FVector2f(HandlePos, HandleSize), FVector2f(1.f, 1.f), HandleColor));
+			Verts.Add(FSlateVertex::Make<ESlateVertexRounding::Disabled>(HandleRenderTransform, FVector2f(HandlePos - HandleSize, 0.f), FVector2f(0.f, 0.f), HandleColor));
 
-			FSlateDrawElement::MakeBox(
+			TArray<SlateIndex> Indices;
+			if (bHasEaseOut)
+			{
+				Indices.Add(0); Indices.Add(3); Indices.Add(4);
+			}
+			else
+			{
+				Indices.Add(4); Indices.Add(1); Indices.Add(2);
+				Indices.Add(4); Indices.Add(2); Indices.Add(3);
+			}
+
+			FSlateDrawElement::MakeCustomVerts(
 				InPainter.DrawElements,
-				// always draw selected keys on top of other keys
 				InPainter.LayerId,
-				// Center the key along X.  Ensure the middle of the key is at the actual key time
-				InPainter.SectionGeometry.ToPaintGeometry(
-					HandleSize,
-					FSlateLayoutTransform(HandlePos - FVector2f(HandleSize.X*0.5f, 0.f))
-				),
-				EasingHandle,
-				DrawEffects,
-				(bRightHandleActive ? SelectionColor : EasingHandle->GetTint(FWidgetStyle()))
+				HandleResource,
+				Verts,
+				Indices,
+				nullptr,
+				0,
+				0,
+				DrawEffects | ESlateDrawEffect::PreMultipliedAlpha
 			);
 		}
 	}
@@ -1607,8 +1690,8 @@ void SSequencerSection::DrawSectionHandles( const FGeometry& AllottedGeometry, F
 		}
 	}
 
-	FGeometry SectionGeometryWithoutHandles = MakeSectionGeometryWithoutHandles(AllottedGeometry, SectionInterface);
-	FTimeToPixel TimeToPixelConverter = ConstructTimeConverterForSection(SectionGeometryWithoutHandles, *ThisSection, GetSequencer());
+	FGeometry SectionGeometry = MakeSectionGeometryWithoutHandles(AllottedGeometry);
+	FTimeToPixel TimeToPixelConverter = ConstructTimeConverterForSection(SectionGeometry, *ThisSection, GetSequencer());
 
 	for (TSharedPtr<FSectionModel> SectionModel : AllUnderlappingSections)
 	{
@@ -1658,7 +1741,7 @@ void SSequencerSection::DrawSectionHandles( const FGeometry& AllottedGeometry, F
 		// Left Grip
 		if (UnderlappingSectionObj->HasStartFrame())
 		{
-			FGeometry SectionRectLeft = SectionGeometryWithoutHandles.MakeChild(
+			FGeometry SectionRectLeft = SectionGeometry.MakeChild(
 				GripSize,
 				FSlateLayoutTransform(FVector2D( TimeToPixelConverter.FrameToPixel(UnderlappingSectionObj->GetInclusiveStartFrame()) - ThisHandleOffset, 0.f ))
 			);
@@ -1676,7 +1759,7 @@ void SSequencerSection::DrawSectionHandles( const FGeometry& AllottedGeometry, F
 		// Right Grip
 		if (UnderlappingSectionObj->HasEndFrame())
 		{
-			FGeometry SectionRectRight = SectionGeometryWithoutHandles.MakeChild(
+			FGeometry SectionRectRight = SectionGeometry.MakeChild(
 				GripSize,
 				FSlateLayoutTransform(FVector2D( TimeToPixelConverter.FrameToPixel(UnderlappingSectionObj->GetExclusiveEndFrame()) - UnderlappingSection->GetSectionGripSize() + ThisHandleOffset, 0 ))
 			);
@@ -1701,28 +1784,27 @@ void SSequencerSection::DrawSectionHandles( const FGeometry& AllottedGeometry, F
 
 void SSequencerSection::Tick( const FGeometry& AllottedGeometry, const double InCurrentTime, const float InDeltaTime )
 {
-	if( GetVisibility() == EVisibility::Visible )
+	if (GetVisibility() == EVisibility::Visible)
 	{
 		UMovieSceneSection* Section = SectionInterface->GetSectionObject();
 		if (Section && Section->HasStartFrame() && Section->HasEndFrame())
 		{
-			FTimeToPixel TimeToPixelConverter(ParentGeometry, GetSequencer().GetViewRange(), Section->GetTypedOuter<UMovieScene>()->GetTickResolution());
+			constexpr float MinSectionWidth = 1.f;
 
-			const int32 SectionLengthPx = FMath::Max(0,
-				FMath::RoundToInt(
-					TimeToPixelConverter.FrameToPixel(Section->GetExclusiveEndFrame())) - FMath::RoundToInt(TimeToPixelConverter.FrameToPixel(Section->GetInclusiveStartFrame())
-				)
-			);
+			FTimeToPixel TimeToPixelConverter(ParentGeometry, GetSequencer().GetViewRange(), Section->GetTypedOuter<UMovieScene>()->GetTickResolution());
+			const FFrameNumber SectionLength = UE::MovieScene::DiscreteSize(Section->ComputeEffectiveRange());
+			const float SectionWidth = TimeToPixelConverter.FrameDeltaToPixel(SectionLength);
 
-			const float SectionGripSize = SectionInterface->GetSectionGripSize();
-			HandleOffsetPx = FMath::Max(FMath::RoundToFloat((2*SectionGripSize - SectionLengthPx) * .5f), 0.f);
+			const float GripSize = SectionInterface->GetSectionGripSize();
+			const float FinalSectionWidth = FMath::Max(MinSectionWidth + GripSize * 2.f, SectionWidth);
+			HandleOffsetPx = (FinalSectionWidth - SectionWidth);
 		}
 		else
 		{
 			HandleOffsetPx = 0;
 		}
 
-		FGeometry SectionGeometry = MakeSectionGeometryWithoutHandles( AllottedGeometry, SectionInterface );
+		FGeometry SectionGeometry = MakeSectionGeometryWithoutHandles(AllottedGeometry);
 		SectionInterface->Tick(SectionGeometry, ParentGeometry, InCurrentTime, InDeltaTime);
 	}
 }
@@ -1740,16 +1822,21 @@ void SSequencerSection::OnModifiedIndirectly(UMovieSceneSignedObject* Object)
 	}
 }
 
-FReply SSequencerSection::OnMouseButtonDown( const FGeometry& MyGeometry, const FPointerEvent& MouseEvent )
+FReply SSequencerSection::OnMouseButtonDown(const FGeometry& MyGeometry, const FPointerEvent& MouseEvent)
 {
 	return FReply::Unhandled();
 }
 
-FGeometry SSequencerSection::MakeSectionGeometryWithoutHandles( const FGeometry& AllottedGeometry, const TSharedPtr<ISequencerSection>& InSectionInterface ) const
+FGeometry SSequencerSection::MakeSectionGeometryWithoutHandles(const FGeometry& AllottedGeometry) const
 {
+	const FVector2f LocalSize = AllottedGeometry.GetLocalSize();
+
+	FVector2f SectionSizeWithoutHandles = LocalSize - FVector2f(HandleOffsetPx * 2.f, 0.0f);
+	SectionSizeWithoutHandles.X = FMath::Max(1.f, SectionSizeWithoutHandles.X);
+
 	return AllottedGeometry.MakeChild(
-		AllottedGeometry.GetLocalSize() - FVector2D( HandleOffsetPx*2, 0.0f ),
-		FSlateLayoutTransform(FVector2D(HandleOffsetPx, 0 ))
+		SectionSizeWithoutHandles,
+		FSlateLayoutTransform(FVector2D(HandleOffsetPx, 0))
 	);
 }
 
@@ -1805,18 +1892,26 @@ FReply SSequencerSection::OnMouseButtonDoubleClick( const FGeometry& MyGeometry,
 }
 
 
-FReply SSequencerSection::OnMouseMove( const FGeometry& MyGeometry, const FPointerEvent& MouseEvent )
+FReply SSequencerSection::OnMouseMove(const FGeometry& MyGeometry, const FPointerEvent& MouseEvent)
 {
-	if ( MouseEvent.GetCursorDelta().Size() > 0 &&
-		!CheckForEasingHandleInteraction(MouseEvent, MyGeometry) &&
-		!CheckForEdgeInteraction(MouseEvent, MyGeometry) &&
-		!CheckForEasingAreaInteraction(MouseEvent, MyGeometry))
+	// Don't fiddle with hotspots if there is an active drag
+	TSharedPtr<FTrackAreaViewModel> TrackAreaViewModel = GetTrackAreaViewModel();
+	const ISequencerEditTool* EditTool = TrackAreaViewModel->GetEditTool();
+	if (EditTool && EditTool->GetDragHotspot())
 	{
-		// If nothing was hit, we just hit the section
-		TSharedPtr<FTrackAreaViewModel> TrackAreaViewModel = GetTrackAreaViewModel();
-		TrackAreaViewModel->SetHotspot( MakeShareable( new FSectionHotspot(WeakSectionModel, Sequencer)) );
+		return FReply::Unhandled();
 	}
 
+	if (MouseEvent.GetCursorDelta().Size() > 0)
+	{
+		if (CheckForEasingHandleInteraction(MouseEvent, MyGeometry) ||
+			CheckForEdgeInteraction(MouseEvent, MyGeometry) ||
+			CheckForEasingAreaInteraction(MouseEvent, MyGeometry) ||
+			CheckForSectionInteraction(MouseEvent, MyGeometry))
+		{
+			return FReply::Handled();
+		}
+	}
 	return FReply::Unhandled();
 }
 
diff --git a/Engine/Source/Editor/Sequencer/Private/SSequencerSection.h b/Engine/Source/Editor/Sequencer/Private/SSequencerSection.h
--- a/Engine/Source/Editor/Sequencer/Private/SSequencerSection.h
+++ b/Engine/Source/Editor/Sequencer/Private/SSequencerSection.h
@@ -39,25 +39,47 @@ public:
 	void Construct( const FArguments& InArgs, TSharedPtr<FSequencer> Sequencer, TSharedPtr<FSectionModel> InSectionModel, TSharedPtr<STrackLane> InOwningTrackLane);
 	~SSequencerSection();
 
+	/**
+	 * Get the section interface
+	 */
 	TSharedPtr<ISequencerSection> GetSectionInterface() const { return SectionInterface; }
 
-	virtual FVector2D ComputeDesiredSize(float) const override;
-
 private:
 
 	/**
-	 * Checks for user interaction (via the mouse) with the left and right edge of a section
+	 * Checks for mouse interaction with the left and right edge of the section
 	 *
 	 * @param MousePosition		The current screen space position of the mouse
 	 * @param SectionGeometry	The geometry of the section
 	 */
 	bool CheckForEdgeInteraction( const FPointerEvent& MousePosition, const FGeometry& SectionGeometry );
 
+	/**
+	 * Checks for mouse interaction with the ease in/out handles of the section
+	 *
+	 * @param MousePosition		The current screen space position of the mouse
+	 * @param SectionGeometry	The geometry of the section
+	 */
 	bool CheckForEasingHandleInteraction( const FPointerEvent& MousePosition, const FGeometry& SectionGeometry );
 
+	/**
+	 * Checks for mouse interaction with the ease in/out area of the section
+	 *
+	 * @param MousePosition		The current screen space position of the mouse
+	 * @param SectionGeometry	The geometry of the section
+	 */
 	bool CheckForEasingAreaInteraction( const FPointerEvent& MousePosition, const FGeometry& SectionGeometry );
 
-	/** SWidget interface */
+	/**
+	 * Checks for mouse interaction with the left/right grip handles of the section
+	 *
+	 * @param MousePosition		The current screen space position of the mouse
+	 * @param SectionGeometry	The geometry of the section
+	 */
+	bool CheckForSectionInteraction( const FPointerEvent& MousePosition, const FGeometry& SectionGeometry );
+
+	/*~ SWidget interface */
+	virtual FVector2D ComputeDesiredSize(float) const override;
 	virtual int32 OnPaint( const FPaintArgs& Args, const FGeometry& AllottedGeometry, const FSlateRect& MyCullingRect, FSlateWindowElementList& OutDrawElements, int32 LayerId, const FWidgetStyle& InWidgetStyle, bool bParentEnabled ) const override;
 	virtual void Tick( const FGeometry& AllottedGeometry, const double InCurrentTime, const float InDeltaTime ) override;
 	virtual FReply OnMouseButtonDown( const FGeometry& MyGeometry, const FPointerEvent& MouseEvent ) override;
@@ -66,7 +88,7 @@ private:
 	virtual FReply OnMouseButtonUp(const FGeometry& MyGeometry, const FPointerEvent& MouseEvent) override;
 	virtual void OnMouseLeave( const FPointerEvent& MouseEvent ) override;
 
-	/*~ ITrackLaneWidget Interface */
+	/*~ ITrackLaneWidget interface */
 	virtual FTrackLaneScreenAlignment GetAlignment(const FTimeToPixel& TimeToPixel, const FGeometry& InParentGeometry) const override;
 	virtual int32 GetOverlapPriority() const override;
 	virtual void ReportParentGeometry(const FGeometry& InParentGeometry) override;
@@ -74,6 +96,9 @@ private:
 	virtual bool AcceptsChildren() const override { return true; }
 	virtual void AddChildView(TSharedPtr<ITrackLaneWidget> ChildWidget, TWeakPtr<STrackLane> InWeakOwningLane) override;
 
+	/*~ ISignedObjectEventHandler interface */
+	virtual void OnModifiedIndirectly(UMovieSceneSignedObject* Object) override;
+
 	/**
 	 * Paint the easing handles for this section
 	 */
@@ -93,11 +118,6 @@ private:
 	/** @return the track area view-model that this section belongs to */
 	TSharedPtr<FTrackAreaViewModel> GetTrackAreaViewModel() const;
 
-	/** 
-	 * Creates geometry for a section without space for the handles
-	 */
-	FGeometry MakeSectionGeometryWithoutHandles( const FGeometry& AllottedGeometry, const TSharedPtr<ISequencerSection>& InSectionInterface ) const;
-
 	/**
 	 * Ensure that the cached array of underlapping sections is up to date
 	 */
@@ -113,14 +133,26 @@ private:
 	 */
 	bool IsEnabled() const;
 
-	void OnModifiedIndirectly(UMovieSceneSignedObject* Object) override;
-
+	/**
+	 * Gets the visibility of this section's top-level channel
+	 */
 	EVisibility GetTopLevelChannelGroupVisibility() const;
 
+	/**
+	 * Gets the color of the top-level key bar
+	 */
 	FLinearColor GetTopLevelKeyBarColor() const;
 
+	/**
+	 * Get the padding offset around the actual section's geometry
+	 */
 	FMargin GetHandleOffsetPadding() const;
 
+	/** 
+	 * Creates geometry for a section without space for the handles
+	 */
+	FGeometry MakeSectionGeometryWithoutHandles(const FGeometry& AllottedGeometry) const;
+
 public:
 
 	/** Indicate that the current section selection should throb the specified number of times. A single throb takes 0.2s. */
@@ -135,13 +167,12 @@ public:
 	/** Get a value between 0 and 1 that indicates the amount of throb-scale to apply to the currently selected sections */
 	static float GetSectionSelectionThrobValue();
 
-	/**
-	 * Check to see whether the specified section is highlighted
-	 */
+	/** Check to see whether the specified section is highlighted */
 	static bool IsSectionHighlighted(UMovieSceneSection* InSection, TSharedPtr<ITrackAreaHotspot> Hotspot);
 
 private:
 
+	/** The parent sequencer */
 	TWeakPtr<FSequencer> Sequencer;
 	/** Interface to section data */
 	TSharedPtr<ISequencerSection> SectionInterface;
diff --git a/Engine/Source/Editor/Sequencer/Public/ISequencerSection.h b/Engine/Source/Editor/Sequencer/Public/ISequencerSection.h
--- a/Engine/Source/Editor/Sequencer/Public/ISequencerSection.h
+++ b/Engine/Source/Editor/Sequencer/Public/ISequencerSection.h
@@ -188,8 +188,16 @@ public:
 	 */
 	virtual float GetSectionHeight() const { return SequencerSectionConstants::DefaultSectionHeight; }
 	
+	/**
+	 * @return The width of the section drag handles
+	 */
 	virtual float GetSectionGripSize() const { return SequencerSectionConstants::DefaultSectionGripSize; }
 
+	/**
+	 * @ return The size of keyframe widgets
+	 */
+	virtual FVector2D GetKeySize() const { return SequencerSectionConstants::KeySize; }
+
 	/**
 	 * @return Whether or not the user can resize this section.
 	 */
diff --git a/Engine/Source/Editor/SequencerCore/Public/TimeToPixel.h b/Engine/Source/Editor/SequencerCore/Public/TimeToPixel.h
--- a/Engine/Source/Editor/SequencerCore/Public/TimeToPixel.h
+++ b/Engine/Source/Editor/SequencerCore/Public/TimeToPixel.h
@@ -23,7 +23,8 @@ public:
 		: ViewRangeStartSeconds( InLocalViewRange.GetLowerBoundValue() )
 		, TickResolution( InTickResolution )
 	{
-		double VisibleWidth = InLocalViewRange.Size<double>();
+		ensureMsgf(WidthPx > 0, TEXT("Got zero or negative pixel width!"));
+		const double VisibleWidth = InLocalViewRange.Size<double>();
 
 		const float MaxPixelsPerSecond = 1000.f;
 		PixelsPerSecond = VisibleWidth > 0 ? WidthPx / VisibleWidth : MaxPixelsPerSecond;
