Sequencer: fix panning track area view

diff --git a/Engine/Source/Editor/Sequencer/Private/SSequencerSection.cpp b/Engine/Source/Editor/Sequencer/Private/SSequencerSection.cpp
--- a/Engine/Source/Editor/Sequencer/Private/SSequencerSection.cpp
+++ b/Engine/Source/Editor/Sequencer/Private/SSequencerSection.cpp
@@ -1894,14 +1894,6 @@ FReply SSequencerSection::OnMouseButtonDoubleClick( const FGeometry& MyGeometry,
 
 FReply SSequencerSection::OnMouseMove(const FGeometry& MyGeometry, const FPointerEvent& MouseEvent)
 {
-	// Don't fiddle with hotspots if there is an active drag
-	TSharedPtr<FTrackAreaViewModel> TrackAreaViewModel = GetTrackAreaViewModel();
-	const ISequencerEditTool* EditTool = TrackAreaViewModel->GetEditTool();
-	if (EditTool && EditTool->GetDragHotspot())
-	{
-		return FReply::Unhandled();
-	}
-
 	if (MouseEvent.GetCursorDelta().Size() > 0)
 	{
 		if (CheckForEasingHandleInteraction(MouseEvent, MyGeometry) ||
@@ -1909,7 +1901,9 @@ FReply SSequencerSection::OnMouseMove(const FGeometry& MyGeometry, const FPointe
 			CheckForEasingAreaInteraction(MouseEvent, MyGeometry) ||
 			CheckForSectionInteraction(MouseEvent, MyGeometry))
 		{
-			return FReply::Handled();
+			// Do nothing: we always want to return FReply::Unhandled because we're just setting up
+			// hotspots here. STrackAreaView will actually handle the mouse movement event, including
+			// calling it on the active hotspot, doing panning, and more.
 		}
 	}
 	return FReply::Unhandled();
