Sequencer: fix regression with zero-width sections (UMG creates those by default)

diff --git a/Engine/Source/Editor/Sequencer/Private/SSequencerSection.cpp b/Engine/Source/Editor/Sequencer/Private/SSequencerSection.cpp
--- a/Engine/Source/Editor/Sequencer/Private/SSequencerSection.cpp
+++ b/Engine/Source/Editor/Sequencer/Private/SSequencerSection.cpp
@@ -1085,7 +1085,7 @@ bool SSequencerSection::CheckForEasingHandleInteraction( const FPointerEvent& Mo
 
 	// Easing handle is a square in the top corner of the grip handle.
 	const double GripSizePx = SectionInterface->GetSectionGripSize();
-	const double GripHalfSizePx = GripSizePx;
+	const double GripHalfSizePx = GripSizePx / 2.f;
 
 	// Now test individual easing handles if we're at the correct vertical position
 	float LocalMouseY = SectionGeometry.AbsoluteToLocal(MouseEvent.GetScreenSpacePosition()).Y;
@@ -1115,13 +1115,9 @@ bool SSequencerSection::CheckForEasingHandleInteraction( const FPointerEvent& Mo
 			// Compute the ease-in handle screen position.
 			const TRange<FFrameNumber> EaseInRange = EasingSectionObj->GetEaseInRange();
 			const FFrameNumber HandleTime = EaseInRange.IsEmpty() ? EasingSectionObj->GetInclusiveStartFrame() : EaseInRange.GetUpperBoundValue();
-			const double HandlePosition = TimeToPixelConverter.FrameToPixel(HandleTime);
+			const double HandlePosition = TimeToPixelConverter.FrameToPixel(HandleTime) - HandleOffsetPx;
 
-			// If the handle is too close to the left edge that we can't hit-test with the full grip width,
-			// offset it to the right.
-			const double HandleOffset = FMath::Max(0.f, GripHalfSizePx - HandlePosition);
-
-			if (FMath::IsNearlyEqual(MousePositionX, HandlePosition + HandleOffset, GripHalfSizePx))
+			if (FMath::IsNearlyEqual(MousePositionX, HandlePosition + GripHalfSizePx, GripHalfSizePx))
 			{
 				TrackAreaViewModel->SetHotspot(MakeShared<FSectionEasingHandleHotspot>(ESequencerEasingType::In, SectionModel, Sequencer));
 				return true;
@@ -1133,13 +1129,9 @@ bool SSequencerSection::CheckForEasingHandleInteraction( const FPointerEvent& Mo
 			// Compute the ease-out handle screen position.
 			TRange<FFrameNumber> EaseOutRange = EasingSectionObj->GetEaseOutRange();
 			const FFrameNumber HandleTime = EaseOutRange.IsEmpty() ? EasingSectionObj->GetExclusiveEndFrame() : EaseOutRange.GetLowerBoundValue();
-			const double HandlePosition = TimeToPixelConverter.FrameToPixel(HandleTime);
-
-			// If the handle is too close to the right edge that we can't hit-test with the full grip width,
-			// offset it to the left.
-			const double HandleOffset = FMath::Max(0.f, GripHalfSizePx - (SectionGeometry.Size.X - HandlePosition));
+			const double HandlePosition = TimeToPixelConverter.FrameToPixel(HandleTime) + HandleOffsetPx;
 
-			if (FMath::IsNearlyEqual(MousePositionX, HandlePosition - HandleOffset, GripHalfSizePx))
+			if (FMath::IsNearlyEqual(MousePositionX, HandlePosition - GripHalfSizePx, GripHalfSizePx))
 			{
 				TrackAreaViewModel->SetHotspot(MakeShared<FSectionEasingHandleHotspot>(ESequencerEasingType::Out, SectionModel, Sequencer));
 				return true;
@@ -1575,7 +1567,7 @@ void SSequencerSection::PaintEasingHandles( FSequencerSectionPainter& InPainter,
 			TRange<FFrameNumber> EaseInRange = UnderlappingSectionObj->GetEaseInRange();
 			const bool bHasEaseIn = !EaseInRange.IsEmpty();
 			FFrameNumber HandleFrame = bHasEaseIn ? UE::MovieScene::DiscreteExclusiveUpper(EaseInRange) : UnderlappingSectionObj->GetInclusiveStartFrame();
-			const float HandlePos(TimeToPixelConverter.FrameToPixel(HandleFrame));
+			const float HandlePos(TimeToPixelConverter.FrameToPixel(HandleFrame) - HandleOffsetPx);
 			FColor HandleColor = (bLeftHandleActive ? SelectionColor : InactiveHandleColor).ToFColorSRGB();
 
 			TArray<FSlateVertex> Verts;
@@ -1614,7 +1606,7 @@ void SSequencerSection::PaintEasingHandles( FSequencerSectionPainter& InPainter,
 			TRange<FFrameNumber> EaseOutRange = UnderlappingSectionObj->GetEaseOutRange();
 			const bool bHasEaseOut = !EaseOutRange.IsEmpty();
 			FFrameNumber HandleFrame = bHasEaseOut ? UE::MovieScene::DiscreteInclusiveLower(EaseOutRange) : UnderlappingSectionObj->GetExclusiveEndFrame();
-			const float HandlePos(TimeToPixelConverter.FrameToPixel(HandleFrame));
+			const float HandlePos(TimeToPixelConverter.FrameToPixel(HandleFrame) + HandleOffsetPx);
 			const FColor HandleColor = (bRightHandleActive ? SelectionColor : InactiveHandleColor).ToFColorSRGB();
 
 			TArray<FSlateVertex> Verts;
@@ -1797,7 +1789,7 @@ void SSequencerSection::Tick( const FGeometry& AllottedGeometry, const double In
 
 			const float GripSize = SectionInterface->GetSectionGripSize();
 			const float FinalSectionWidth = FMath::Max(MinSectionWidth + GripSize * 2.f, SectionWidth);
-			HandleOffsetPx = (FinalSectionWidth - SectionWidth);
+			HandleOffsetPx = FMath::RoundToFloat((FinalSectionWidth - SectionWidth) * 0.5f);
 		}
 		else
 		{
