Sequencer: move sequencer customizations to the editor view-model, so that MVVM code can access it

This lets us avoid exposing more things through ISequencer/FSequencer.

diff --git a/Engine/Source/Editor/Sequencer/Private/MVVM/ViewModels/SequencerEditorViewModel.cpp b/Engine/Source/Editor/Sequencer/Private/MVVM/ViewModels/SequencerEditorViewModel.cpp
--- a/Engine/Source/Editor/Sequencer/Private/MVVM/ViewModels/SequencerEditorViewModel.cpp
+++ b/Engine/Source/Editor/Sequencer/Private/MVVM/ViewModels/SequencerEditorViewModel.cpp
@@ -160,6 +160,47 @@ void FSequencerEditorViewModel::OnTrackAreaHotspotChanged(TSharedPtr<ITrackAreaH
 	CurrentHotspot = NewHotspot;
 }
 
+bool FSequencerEditorViewModel::UpdateSequencerCustomizations(const UMovieSceneSequence* PreviousFocusedSequence)
+{
+	ISequencerModule& SequencerModule = FModuleManager::LoadModuleChecked<ISequencerModule>("Sequencer");
+	TSharedPtr<FSequencerCustomizationManager> Manager = SequencerModule.GetSequencerCustomizationManager();
+
+	TSharedPtr<ISequencer> Sequencer = GetSequencer();
+	UMovieSceneSequence* FocusedSequence = Sequencer->GetFocusedMovieSceneSequence();
+	check(FocusedSequence != nullptr);
+
+	// See if we need to change any of our customizations.
+	if (PreviousFocusedSequence && !Manager->NeedsCustomizationChange(PreviousFocusedSequence, FocusedSequence))
+	{
+		return false;
+	}
+
+	// Get rid of previously active customizations.
+	for (const TUniquePtr<ISequencerCustomization>& Customization : ActiveCustomizations)
+	{
+		Customization->UnregisterSequencerCustomization();
+	}
+	ActiveCustomizations.Reset();
+
+	// Get the customizations for the current sequence.
+	Manager->GetSequencerCustomizations(FocusedSequence, ActiveCustomizations);
+
+	// Get the customization info.
+	FSequencerCustomizationBuilder Builder(*Sequencer, *FocusedSequence);
+	for (const TUniquePtr<ISequencerCustomization>& Customization : ActiveCustomizations)
+	{
+		Customization->RegisterSequencerCustomization(Builder);
+	}
+	ActiveCustomizationInfos = Builder.GetCustomizations();
+
+	return true;
+}
+
+TArrayView<const FSequencerCustomizationInfo> FSequencerEditorViewModel::GetActiveCustomizationInfos() const
+{
+	return ActiveCustomizationInfos;
+}
+
 } // namespace Sequencer
 } // namespace UE
 
diff --git a/Engine/Source/Editor/Sequencer/Private/SSequencer.cpp b/Engine/Source/Editor/Sequencer/Private/SSequencer.cpp
--- a/Engine/Source/Editor/Sequencer/Private/SSequencer.cpp
+++ b/Engine/Source/Editor/Sequencer/Private/SSequencer.cpp
@@ -4079,7 +4079,7 @@ void SSequencer::SetPlayTime(double Frame)
 	}
 }
 
-void SSequencer::ApplySequencerCustomizations(const TArray<FSequencerCustomizationInfo>& Customizations)
+void SSequencer::ApplySequencerCustomizations(const TArrayView<const FSequencerCustomizationInfo> Customizations)
 {
 	AddMenuExtenders.Reset();
 	ToolbarExtenders.Reset();
diff --git a/Engine/Source/Editor/Sequencer/Private/SSequencer.h b/Engine/Source/Editor/Sequencer/Private/SSequencer.h
--- a/Engine/Source/Editor/Sequencer/Private/SSequencer.h
+++ b/Engine/Source/Editor/Sequencer/Private/SSequencer.h
@@ -650,7 +650,7 @@ public:
 	void RequestRenameNode(const FString& Path) { NodePathToRename = Path; }
 
 	/** Applies dynamic sequencer customizations to this editor. */
-	void ApplySequencerCustomizations(const TArray<FSequencerCustomizationInfo>& Customizations);
+	void ApplySequencerCustomizations(const TArrayView<const FSequencerCustomizationInfo> Customizations);
 
 private:
 	/** Applies a single customization. */
diff --git a/Engine/Source/Editor/Sequencer/Private/Sequencer.cpp b/Engine/Source/Editor/Sequencer/Private/Sequencer.cpp
--- a/Engine/Source/Editor/Sequencer/Private/Sequencer.cpp
+++ b/Engine/Source/Editor/Sequencer/Private/Sequencer.cpp
@@ -152,7 +152,6 @@
 #include "Interfaces/IAnalyticsProvider.h"
 #include "Components/SkeletalMeshComponent.h"
 #include "EntitySystem/MovieSceneInitialValueCache.h"
-#include "SequencerCustomizationManager.h"
 #include "SSequencerGroupManager.h"
 #include "ActorTreeItem.h"
 #include "Widgets/Layout/SSpacer.h"
@@ -594,7 +593,7 @@ void FSequencer::InitSequencer(const FSequencerInitParams& InitParams, const TSh
 		TrackEditor->OnInitialize();
 	}
 
-	UpdateSequencerCustomizations();
+	UpdateSequencerCustomizations(nullptr);
 
 	AddNodeGroupsCollectionChangedDelegate();
 
@@ -1004,6 +1003,7 @@ void FSequencer::ResetToNewRootSequence(UMovieSceneSequence& NewSequence)
 {
 	RemoveNodeGroupsCollectionChangedDelegate();
 
+	const UMovieSceneSequence* PreviousRootSequence = RootSequence.Get();
 	RootSequence = &NewSequence;
 	RestorePreAnimatedState();
 
@@ -1043,7 +1043,7 @@ void FSequencer::ResetToNewRootSequence(UMovieSceneSequence& NewSequence)
 	PlayPosition.Reset(ConvertFrameTime(GetPlaybackRange().GetLowerBoundValue(), GetRootTickResolution(), PlayPosition.GetInputRate()));
 	TimeController->Reset(FQualifiedFrameTime(PlayPosition.GetCurrentPosition(), GetRootTickResolution()));
 
-	UpdateSequencerCustomizations();
+	UpdateSequencerCustomizations(PreviousRootSequence);
 
 	AddNodeGroupsCollectionChangedDelegate();
 
@@ -1055,6 +1055,8 @@ void FSequencer::FocusSequenceInstance(UMovieSceneSubSection& InSubSection)
 {
 	RemoveNodeGroupsCollectionChangedDelegate();
 
+	const UMovieSceneSequence* PreviousFocusedSequence = GetFocusedMovieSceneSequence();
+
 	TemplateIDBackwardStack.Push(ActiveTemplateIDs.Top());
 	TemplateIDForwardStack.Reset();
 
@@ -1140,7 +1142,7 @@ void FSequencer::FocusSequenceInstance(UMovieSceneSubSection& InSubSection)
 
 	UpdateSubSequenceData();
 
-	UpdateSequencerCustomizations();
+	UpdateSequencerCustomizations(PreviousFocusedSequence);
 
 	ScrubPositionParent.Reset();
 
@@ -1275,6 +1277,8 @@ void FSequencer::PopToSequenceInstance(FMovieSceneSequenceIDRef SequenceID)
 
 		RemoveNodeGroupsCollectionChangedDelegate();
 
+		const UMovieSceneSequence* PreviousFocusedSequence = GetFocusedMovieSceneSequence();
+
 		// Pop until we find the movie scene to focus
 		while( SequenceID != ActiveTemplateIDs.Last() )
 		{
@@ -1302,7 +1306,7 @@ void FSequencer::PopToSequenceInstance(FMovieSceneSequenceIDRef SequenceID)
 			RootTemplateInstance.FindInstance(MovieSceneSequenceID::Root)->OverrideRootSequence(Linker, ActiveTemplateIDs.Top());
 		}
 
-		UpdateSequencerCustomizations();
+		UpdateSequencerCustomizations(PreviousFocusedSequence);
 
 		AddNodeGroupsCollectionChangedDelegate();
 
@@ -1399,37 +1403,20 @@ void FSequencer::UpdateSubSequenceData()
 	}
 }
 
-void FSequencer::UpdateSequencerCustomizations()
+void FSequencer::UpdateSequencerCustomizations(const UMovieSceneSequence* PreviousFocusedSequence)
 {
-	ISequencerModule& SequencerModule = FModuleManager::LoadModuleChecked<ISequencerModule>("Sequencer");
-	TSharedPtr<FSequencerCustomizationManager> Manager = SequencerModule.GetSequencerCustomizationManager();
-
-	// Get rid of previously active customizations.
-	for (const TUniquePtr<ISequencerCustomization>& Customization : ActiveCustomizations)
-	{
-		Customization->UnregisterSequencerCustomization();
-	}
-	ActiveCustomizations.Reset();
-
-	// Get the customizations for the current sequence.
 	UMovieSceneSequence* FocusedSequence = GetFocusedMovieSceneSequence();
 	check(FocusedSequence != nullptr);
-	Manager->GetSequencerCustomizations(*FocusedSequence, ActiveCustomizations);
-
-	// Get the customization info.
-	FSequencerCustomizationBuilder Builder(*this, *FocusedSequence);
-	for (const TUniquePtr<ISequencerCustomization>& Customization : ActiveCustomizations)
-	{
-		Customization->RegisterSequencerCustomization(Builder);
-	}
+	ViewModel->UpdateSequencerCustomizations(PreviousFocusedSequence);
+	TArrayView<const FSequencerCustomizationInfo> CustomizationInfos = ViewModel->GetActiveCustomizationInfos();
 
 	// Apply customizations to our editor.
-	SequencerWidget->ApplySequencerCustomizations(Builder.GetCustomizations());
+	SequencerWidget->ApplySequencerCustomizations(CustomizationInfos);
 
 	// Apply customizations to ourselves.
 	OnPaste.Reset();
 
-	for (const FSequencerCustomizationInfo& CustomizationInfo : Builder.GetCustomizations())
+	for (const FSequencerCustomizationInfo& CustomizationInfo : CustomizationInfos)
 	{
 		if (CustomizationInfo.OnPaste.IsBound())
 		{
diff --git a/Engine/Source/Editor/Sequencer/Private/Sequencer.h b/Engine/Source/Editor/Sequencer/Private/Sequencer.h
--- a/Engine/Source/Editor/Sequencer/Private/Sequencer.h
+++ b/Engine/Source/Editor/Sequencer/Private/Sequencer.h
@@ -1098,7 +1098,7 @@ private:
 	void UpdateSubSequenceData();
 
 	/** Adjust sequencer customizations based on the currently focused sequence */
-	void UpdateSequencerCustomizations();
+	void UpdateSequencerCustomizations(const UMovieSceneSequence* PreviousFocusedSequence);
 
 	/** Rerun construction scripts on bound actors */
 	void RerunConstructionScripts();
@@ -1181,9 +1181,6 @@ private:
 	/** A copy of the supported features/capabilities we were initialized with. */
 	FSequencerHostCapabilities HostCapabilities;
 	
-	/** Active customizations. */
-	TArray<TUniquePtr<ISequencerCustomization>> ActiveCustomizations;
-
 	/** Active customization callbacks */
 	TArray<FOnSequencerPaste> OnPaste;
 
diff --git a/Engine/Source/Editor/Sequencer/Private/SequencerCustomizationManager.cpp b/Engine/Source/Editor/Sequencer/Private/SequencerCustomizationManager.cpp
--- a/Engine/Source/Editor/Sequencer/Private/SequencerCustomizationManager.cpp
+++ b/Engine/Source/Editor/Sequencer/Private/SequencerCustomizationManager.cpp
@@ -47,9 +47,14 @@ void FSequencerCustomizationManager::UnregisterInstancedSequencerCustomization(c
     ensureMsgf(false, TEXT("No customization found for sequence class: %s"), *SequenceClass->GetName());
 }
 
-void FSequencerCustomizationManager::GetSequencerCustomizations(UMovieSceneSequence& FocusedSequence, TArray<TUniquePtr<ISequencerCustomization>>& OutCustomizations)
+void FSequencerCustomizationManager::GetSequencerCustomizations(UMovieSceneSequence* FocusedSequence, TArray<TUniquePtr<ISequencerCustomization>>& OutCustomizations)
 {
-    const UClass* SequenceClass = FocusedSequence.GetClass();
+	if (UNLIKELY(!FocusedSequence))
+	{
+		return;
+	}
+
+    const UClass* SequenceClass = FocusedSequence->GetClass();
 
     const FCustomizationRegistryEntry* FoundEntry = CustomizationRegistryEntries.FindByPredicate(
             [=](const FCustomizationRegistryEntry& Entry)
@@ -62,3 +67,18 @@ void FSequencerCustomizationManager::GetSequencerCustomizations(UMovieSceneSeque
         OutCustomizations.Add(MoveTemp(Instance));
     }
 }
+
+bool FSequencerCustomizationManager::NeedsCustomizationChange(const UMovieSceneSequence* OldFocusedSequence, const UMovieSceneSequence* NewFocusedSequence) const
+{
+	if (OldFocusedSequence == nullptr || NewFocusedSequence == nullptr)
+	{
+		return true;
+	}
+
+	// For now, we only map customizations to the sequence's class, so we only need to rebuild customizations
+	// if the class has changed between the old and new sequence.
+    const UClass* OldSequenceClass = OldFocusedSequence->GetClass();
+    const UClass* NewSequenceClass = NewFocusedSequence->GetClass();
+	return OldSequenceClass != NewSequenceClass;
+}
+
diff --git a/Engine/Source/Editor/Sequencer/Public/MVVM/ViewModels/SequencerEditorViewModel.h b/Engine/Source/Editor/Sequencer/Public/MVVM/ViewModels/SequencerEditorViewModel.h
--- a/Engine/Source/Editor/Sequencer/Public/MVVM/ViewModels/SequencerEditorViewModel.h
+++ b/Engine/Source/Editor/Sequencer/Public/MVVM/ViewModels/SequencerEditorViewModel.h
@@ -2,7 +2,10 @@
 
 #pragma once
 
+#include "Framework/Commands/UICommandList.h"
 #include "MVVM/ViewModels/EditorViewModel.h"
+#include "SequencerCustomizationManager.h"
+#include "Toolkits/AssetEditorToolkit.h"
 
 struct FMovieSceneSequenceID;
 
@@ -24,7 +27,7 @@ struct ITrackAreaHotspot;
 /**
  * Main view-model for the Sequencer editor.
  */
-class FSequencerEditorViewModel : public FEditorViewModel
+class SEQUENCER_API FSequencerEditorViewModel : public FEditorViewModel
 {
 public:
 
@@ -33,16 +36,24 @@ public:
 	FSequencerEditorViewModel(TSharedRef<ISequencer> InSequencer, const FSequencerHostCapabilities& InHostCapabilities);
 
 	/** Retrieve this editor's selection class */
-	SEQUENCER_API TSharedPtr<FSequencerSelection> GetSelection() const;
+	TSharedPtr<FSequencerSelection> GetSelection() const;
 
 	// @todo_sequencer_mvvm remove this later
-	SEQUENCER_API TSharedPtr<ISequencer> GetSequencer() const;
+	TSharedPtr<ISequencer> GetSequencer() const;
 	// @todo_sequencer_mvvm remove this ASAP
 	TSharedPtr<FSequencer> GetSequencerImpl() const;
 
 	// @todo_sequencer_mvvm move this to the root view-model
 	void SetSequence(UMovieSceneSequence* InRootSequence);
 
+public:
+
+	/** Adjust sequencer customizations based on the currently focused sequence */
+	bool UpdateSequencerCustomizations(const UMovieSceneSequence* PreviousFocusedSequence);
+
+	/** Get the active customization infos */
+	TArrayView<const FSequencerCustomizationInfo> GetActiveCustomizationInfos() const;
+
 	/** Gets the pinned track area view-model. */
 	TSharedPtr<FTrackAreaViewModel> GetPinnedTrackArea() const;
 
@@ -65,15 +76,24 @@ protected:
 
 protected:
 
+	/** Owning sequencer */
 	TWeakPtr<ISequencer> WeakSequencer;
+	
+	/** Pinned track area (the main track area is owned by the base class) */
 	TSharedPtr<FTrackAreaViewModel> PinnedTrackArea;
-	bool bSupportsCurveEditor;
 
+	/** Whether we have a curve editor extension */
+	bool bSupportsCurveEditor;
+	
 	/** The current hotspot, from any of our track areas */
 	TSharedPtr<ITrackAreaHotspot> CurrentHotspot;
 
-	// Cached node paths to be used to compare when the hierarchy changes
+	/** Cached node paths to be used to compare when the hierarchy changes */
 	TMap<TWeakPtr<FViewModel>, FString> NodePaths;
+
+	/** Active customizations. */
+	TArray<TUniquePtr<ISequencerCustomization>> ActiveCustomizations;
+	TArray<FSequencerCustomizationInfo> ActiveCustomizationInfos;
 };
 
 } // namespace Sequencer
diff --git a/Engine/Source/Editor/Sequencer/Public/SequencerCustomizationManager.h b/Engine/Source/Editor/Sequencer/Public/SequencerCustomizationManager.h
--- a/Engine/Source/Editor/Sequencer/Public/SequencerCustomizationManager.h
+++ b/Engine/Source/Editor/Sequencer/Public/SequencerCustomizationManager.h
@@ -133,7 +133,12 @@ public:
 	/**
 	 * Gets the sequencer customizations for the given sequence.
 	 */
-	void GetSequencerCustomizations(UMovieSceneSequence& FocusedSequence, TArray<TUniquePtr<ISequencerCustomization>>& OutCustomizations);
+	void GetSequencerCustomizations(UMovieSceneSequence* FocusedSequence, TArray<TUniquePtr<ISequencerCustomization>>& OutCustomizations);
+
+	/**
+	 * Determines whether the new edited sequence requires changing the customizations.
+	 */
+	bool NeedsCustomizationChange(const UMovieSceneSequence* OldFocusedSequence, const UMovieSceneSequence* NewFocusedSequence) const;
 
 private:
 	FSequencerCustomizationManager(const FSequencerCustomizationManager&) = delete;
