Sequencer: Only sync curve selection when updating the tree if we have no curves, otherwise this will blow away what the animator is working on.
#trivial

diff --git a/Engine/Source/Editor/Sequencer/Private/MVVM/CurveEditorIntegrationExtension.cpp b/Engine/Source/Editor/Sequencer/Private/MVVM/CurveEditorIntegrationExtension.cpp
--- a/Engine/Source/Editor/Sequencer/Private/MVVM/CurveEditorIntegrationExtension.cpp
+++ b/Engine/Source/Editor/Sequencer/Private/MVVM/CurveEditorIntegrationExtension.cpp
@@ -134,8 +134,12 @@ void FCurveEditorIntegrationExtension::UpdateCurveEditor()
 		}
 	}
 
-	// Synchronize selection after the view models have been added to the curve editor. Synchronization depends on curve model tree item IDs
-	CurveEditorExtension->RequestSyncSelection();
+	// only synchronize if we don't have anything in the tree anymore, otherwise this will add curves that currently aren't needed
+	const TMap<FCurveEditorTreeItemID, FCurveEditorTreeItem>& AllItems = CurveEditorTree->GetAllItems();
+	if (AllItems.Num() == 0)
+	{
+		CurveEditorExtension->RequestSyncSelection();
+	}
 }
 
 FCurveEditorTreeItemID FCurveEditorIntegrationExtension::AddToCurveEditor(TViewModelPtr<ICurveEditorTreeItemExtension> InViewModel, TSharedPtr<FCurveEditor> InCurveEditor)
