Sequencer: When changing time via translate hot keys get all changes first, then apply them to avoid double transforms with space/constraint keys.

diff --git a/Engine/Source/Editor/Sequencer/Private/Sequencer.cpp b/Engine/Source/Editor/Sequencer/Private/Sequencer.cpp
--- a/Engine/Source/Editor/Sequencer/Private/Sequencer.cpp
+++ b/Engine/Source/Editor/Sequencer/Private/Sequencer.cpp
@@ -1877,9 +1877,9 @@ void FSequencer::TransformSelectedKeysAndSections(FFrameTime InDeltaTime, float
 	FSelectedKeysByChannel KeysByChannel(ViewModel->GetSelection()->KeySelection);
 	TMap<UMovieSceneSection*, TRange<FFrameNumber>> SectionToNewBounds;
 
-	TArray<FFrameNumber> KeyTimesScratch;
 	if (InScale != 0.f)
 	{
+		TMap<FMovieSceneChannel*, TPair<TArray<FFrameNumber>, TArray<FKeyHandle>>> ChannelsAndKeyTimes;
 		// Dilate the keys
 		for (const FSelectedChannelInfo& ChannelInfo : KeysByChannel.SelectedChannels)
 		{
@@ -1891,20 +1891,20 @@ void FSequencer::TransformSelectedKeysAndSections(FFrameTime InDeltaTime, float
 				{
 					continue;
 				}
-
+				TPair<TArray<FFrameNumber>, TArray<FKeyHandle>>& KeyTimesScratch = ChannelsAndKeyTimes.FindOrAdd(Channel);
 				const int32 NumKeys = ChannelInfo.KeyHandles.Num();
-				KeyTimesScratch.Reset(NumKeys);
-				KeyTimesScratch.SetNum(NumKeys);
-
+				KeyTimesScratch.Key.Reset(NumKeys);
+				KeyTimesScratch.Key.SetNum(NumKeys);
 				// Populate the key times scratch buffer with the times for these handles
-				Channel->GetKeyTimes(ChannelInfo.KeyHandles, KeyTimesScratch);
+				KeyTimesScratch.Value = ChannelInfo.KeyHandles;
+				Channel->GetKeyTimes(ChannelInfo.KeyHandles, KeyTimesScratch.Key);
 
 				// We have to find the lowest key time and the highest key time. They're added based on selection order so we can't rely on their order in the array.
-				FFrameTime LowestFrameTime = KeyTimesScratch[0];
-				FFrameTime HighestFrameTime = KeyTimesScratch[0];
+				FFrameTime LowestFrameTime = KeyTimesScratch.Key[0];
+				FFrameTime HighestFrameTime = KeyTimesScratch.Key[0];
 
 				// Perform the transformation
-				for (FFrameNumber& Time : KeyTimesScratch)
+				for (FFrameNumber& Time : KeyTimesScratch.Key)
 				{
 					FFrameTime KeyTime = Time;
 					Time = (OriginTime + InDeltaTime + (KeyTime - OriginTime) * InScale).FloorToFrame();
@@ -1933,12 +1933,18 @@ void FSequencer::TransformSelectedKeysAndSections(FFrameTime InDeltaTime, float
 				// for sections, but HighestFrameTime is measuring only the key's time.
 				*NewSectionBounds = TRange<FFrameNumber>::Hull(*NewSectionBounds, TRange<FFrameNumber>(LowestFrameTime.GetFrame(), HighestFrameTime.GetFrame() + 1));
 
-				// Apply the new, transformed key times
-				Channel->SetKeyTimes(ChannelInfo.KeyHandles, KeyTimesScratch);
+				
 				bAnythingChanged = true;
 			}
 		}
-
+		// set key times AFTER getting where we want all of them where we want to go, this way we avoid double transforms
+		// on things like smart keys for spaces and constraints.
+		for (TPair<FMovieSceneChannel*, TPair<TArray<FFrameNumber>, TArray<FKeyHandle>>>& ChannelAndKeys : ChannelsAndKeyTimes)
+		{
+			// Apply the new, transformed key times
+			// Channel->SetKeyTimes(KeyHandles,KeyTimes) -->
+			ChannelAndKeys.Key->SetKeyTimes(ChannelAndKeys.Value.Value, ChannelAndKeys.Value.Key);
+		}
 		// Dilate the sections
 		for (UMovieSceneSection* Section : SelectedSections)
 		{
