Sequencer: Add modify when duplicating keys.

This fixes a bug where duplicating keys is not undoable and doesn't update the UI properly. Probably a regression from mvvm refactor a while back.

diff --git a/Engine/Source/Editor/Sequencer/Private/Sequencer.cpp b/Engine/Source/Editor/Sequencer/Private/Sequencer.cpp
--- a/Engine/Source/Editor/Sequencer/Private/Sequencer.cpp
+++ b/Engine/Source/Editor/Sequencer/Private/Sequencer.cpp
@@ -9965,7 +9965,7 @@ void FSequencer::DuplicateSelection()
 
 		// Offset by a visible amount
 		FFrameNumber FrameOffset = FFrameNumber((int32)GetDisplayRateDeltaFrameCount());
-
+		TSet<UMovieSceneSection*> ModifiedSections;
 		FKeySelection NewSelection;
 		for (FKeyHandle KeyHandle : ViewModel->GetSelection()->KeySelection)
 		{
@@ -9973,11 +9973,23 @@ void FSequencer::DuplicateSelection()
 			if (Channel)
 			{
 				TSharedPtr<IKeyArea> KeyArea = Channel->GetKeyArea();
+				UMovieSceneSection* Section = KeyArea->GetOwningSection();
 
-				FKeyHandle NewKeyHandle = KeyArea->DuplicateKey(KeyHandle);
-				KeyArea->SetKeyTime(NewKeyHandle, KeyArea->GetKeyTime(KeyHandle) + FrameOffset);
+				bool bModified = ModifiedSections.Contains(Section);
+				if (!bModified)
+				{
+					bModified = Section->TryModify();
+				}
 
-				NewSelection.Select(Channel, NewKeyHandle);
+				if (bModified)
+				{
+					ModifiedSections.Add(Section);
+
+					FKeyHandle NewKeyHandle = KeyArea->DuplicateKey(KeyHandle);
+					KeyArea->SetKeyTime(NewKeyHandle, KeyArea->GetKeyTime(KeyHandle) + FrameOffset);
+
+					NewSelection.Select(Channel, NewKeyHandle);
+				}
 			}
 		}
 
