Curve Editor: Synchonize selection if new items have been added

diff --git a/Engine/Source/Editor/Sequencer/Private/MVVM/CurveEditorIntegrationExtension.cpp b/Engine/Source/Editor/Sequencer/Private/MVVM/CurveEditorIntegrationExtension.cpp
--- a/Engine/Source/Editor/Sequencer/Private/MVVM/CurveEditorIntegrationExtension.cpp
+++ b/Engine/Source/Editor/Sequencer/Private/MVVM/CurveEditorIntegrationExtension.cpp
@@ -116,6 +116,7 @@ void FCurveEditorIntegrationExtension::UpdateCurveEditor()
 
 	// Iterate all non-filtered out outliners and check for curve editor tree extensions
 	const bool bIncludeRootNode = false;
+	bool bItemsAdded = false;
 	for (TParentFirstChildIterator<IOutlinerExtension> It(OwnerModel, bIncludeRootNode); It; ++It)
 	{
 		if (It->IsFilteredOut())
@@ -129,14 +130,14 @@ void FCurveEditorIntegrationExtension::UpdateCurveEditor()
 			const FCurveEditorTreeItemID ItemID = ViewModelToTreeItemIDMap.FindRef(ChildViewModel.AsModel());
 			if (!ItemID.IsValid())
 			{
+				bItemsAdded = true;
 				AddToCurveEditor(ChildViewModel, CurveEditor);
 			}
 		}
 	}
 
-	// only synchronize if we don't have anything in the tree anymore, otherwise this will add curves that currently aren't needed
-	const TMap<FCurveEditorTreeItemID, FCurveEditorTreeItem>& AllItems = CurveEditorTree->GetAllItems();
-	if (AllItems.Num() == 0)
+	// If new items have been added, synchronize selection after the view models have been added to the curve editor. Synchronization depends on curve model tree item IDs
+	if (bItemsAdded)
 	{
 		CurveEditorExtension->RequestSyncSelection();
 	}
