Sequencer: Use the Section to Key if it exists when keying through the key editors.

This is probably a regression from the mvvm refactor.

diff --git a/Engine/Source/Editor/Sequencer/Private/SequencerAddKeyOperation.cpp b/Engine/Source/Editor/Sequencer/Private/SequencerAddKeyOperation.cpp
--- a/Engine/Source/Editor/Sequencer/Private/SequencerAddKeyOperation.cpp
+++ b/Engine/Source/Editor/Sequencer/Private/SequencerAddKeyOperation.cpp
@@ -112,9 +112,28 @@ bool FAddKeyOperation::ConsiderKeyableAreas(TSharedPtr<ITrackExtension> InTrackM
 	bool bKeyedAnything = false;
 
 	constexpr bool bIncludeThis = true;
-	for (TParentFirstChildIterator<FChannelGroupModel> ChannelGroupModelIt(KeyAnythingBeneath->AsShared(), bIncludeThis); ChannelGroupModelIt; ++ChannelGroupModelIt)
+	
+	// Prefer the section that is marked SectionToKey
+	for (const TViewModelPtr<FSectionModel>& SectionModel : KeyAnythingBeneath->GetDescendantsOfType<FSectionModel>(bIncludeThis))
 	{
-		bKeyedAnything |= ProcessKeyArea(InTrackModel, *ChannelGroupModelIt);
+		UMovieSceneSection* Section = SectionModel->GetSection();
+		UMovieSceneTrack* Track = Cast<UMovieSceneTrack>(Section->GetOuter());
+		if (Track && Track->GetSectionToKey() == Section)
+		{
+			for (TParentFirstChildIterator<FChannelGroupModel> ChannelGroupModelIt(SectionModel->AsShared(), bIncludeThis); ChannelGroupModelIt; ++ChannelGroupModelIt)
+			{
+				bKeyedAnything |= ProcessKeyArea(InTrackModel, *ChannelGroupModelIt);
+			}
+		}
+	}
+
+	// Otherwise if nothing was found, key all
+	if (!bKeyedAnything)
+	{
+		for (TParentFirstChildIterator<FChannelGroupModel> ChannelGroupModelIt(KeyAnythingBeneath->AsShared(), bIncludeThis); ChannelGroupModelIt; ++ChannelGroupModelIt)
+		{
+			bKeyedAnything |= ProcessKeyArea(InTrackModel, *ChannelGroupModelIt);
+		}
 	}
 
 	return bKeyedAnything;
