Fix implicit capture of this using [=] deprecated in C++20
#rnx

diff --git a/Engine/Plugins/Editor/PluginBrowser/Source/PluginBrowser/Private/SNewPluginWizard.cpp b/Engine/Plugins/Editor/PluginBrowser/Source/PluginBrowser/Private/SNewPluginWizard.cpp
--- a/Engine/Plugins/Editor/PluginBrowser/Source/PluginBrowser/Private/SNewPluginWizard.cpp
+++ b/Engine/Plugins/Editor/PluginBrowser/Source/PluginBrowser/Private/SNewPluginWizard.cpp
@@ -190,7 +190,7 @@ void SNewPluginWizard::Construct(const FArguments& Args, TSharedPtr<SDockTab> In
 			SNew(SCheckBox)
 			.OnCheckStateChanged(this, &SNewPluginWizard::OnEnginePluginCheckboxChanged)
 			.IsChecked(this, &SNewPluginWizard::IsEnginePlugin)
-			.Visibility_Lambda([=]()
+			.Visibility_Lambda([this]()
 			{
 				FPluginTemplateDescription* Template = PluginWizardDefinition->GetSelectedTemplate().Get();
 				return ((Template != nullptr) && Template->bCanBePlacedInEngine) ? EVisibility::Visible : EVisibility::Collapsed;
@@ -238,7 +238,7 @@ void SNewPluginWizard::Construct(const FArguments& Args, TSharedPtr<SDockTab> In
 		[
 			SAssignNew(ShowPluginContentDirectoryCheckBox, SCheckBox)
 			.IsChecked(ECheckBoxState::Checked)
-			.Visibility_Lambda([=]()
+			.Visibility_Lambda([this]()
 			{
 				FPluginTemplateDescription* Template = PluginWizardDefinition->GetSelectedTemplate().Get();
 				return ((Template != nullptr) && Template->bCanContainContent) ? EVisibility::Visible : EVisibility::Collapsed;
diff --git a/Engine/Plugins/Experimental/Mutable/Source/CustomizableObjectEditor/Private/MuCOE/Nodes/CustomizableObjectNodeLayoutBlocksDetails.cpp b/Engine/Plugins/Experimental/Mutable/Source/CustomizableObjectEditor/Private/MuCOE/Nodes/CustomizableObjectNodeLayoutBlocksDetails.cpp
--- a/Engine/Plugins/Experimental/Mutable/Source/CustomizableObjectEditor/Private/MuCOE/Nodes/CustomizableObjectNodeLayoutBlocksDetails.cpp
+++ b/Engine/Plugins/Experimental/Mutable/Source/CustomizableObjectEditor/Private/MuCOE/Nodes/CustomizableObjectNodeLayoutBlocksDetails.cpp
@@ -102,7 +102,7 @@ void FCustomizableObjectNodeLayoutBlocksDetails::CustomizeDetails( IDetailLayout
 				.Padding(5.0f,0.0f,0.0f,0.0f)
 				[
 					SAssignNew(LODSelectorWidget, SSpinBox<int32>)
-					.Value_Lambda([=]()->int32
+					.Value_Lambda([this]()->int32
 					{
 						if (Node->Layout)
 						{
diff --git a/Engine/Plugins/Experimental/Mutable/Source/MutableTools/Private/MuT/AST.cpp b/Engine/Plugins/Experimental/Mutable/Source/MutableTools/Private/MuT/AST.cpp
--- a/Engine/Plugins/Experimental/Mutable/Source/MutableTools/Private/MuT/AST.cpp
+++ b/Engine/Plugins/Experimental/Mutable/Source/MutableTools/Private/MuT/AST.cpp
@@ -158,7 +158,7 @@ void ASTOp::Assert()
     });
 
     // Validate the children
-    ForEachChild( [=](ASTChild&c)
+	ForEachChild( [this](ASTChild&c)
     {
         if(c)
         {
diff --git a/Engine/Plugins/FX/Niagara/Source/NiagaraEditor/Private/Customizations/NiagaraComponentDetails.cpp b/Engine/Plugins/FX/Niagara/Source/NiagaraEditor/Private/Customizations/NiagaraComponentDetails.cpp
--- a/Engine/Plugins/FX/Niagara/Source/NiagaraEditor/Private/Customizations/NiagaraComponentDetails.cpp
+++ b/Engine/Plugins/FX/Niagara/Source/NiagaraEditor/Private/Customizations/NiagaraComponentDetails.cpp
@@ -798,7 +798,7 @@ void FNiagaraUserParameterNodeBuilder::GenerateUserParameterRows(IDetailChildren
 								ActiveSection = Section;
 								OnRebuildChildren.ExecuteIfBound();
 							})
-							.IsChecked_Lambda([=]()
+						.IsChecked_Lambda([this, Section]()
 							{
 								return ActiveSection == Section ? ECheckBoxState::Checked : ECheckBoxState::Unchecked;
 							})
@@ -831,7 +831,7 @@ void FNiagaraUserParameterNodeBuilder::GenerateUserParameterRows(IDetailChildren
 						ActiveSection = nullptr;
 						OnRebuildChildren.ExecuteIfBound();
 					})
-					.IsChecked_Lambda([=]()
+					.IsChecked_Lambda([this]()
 					{
 						return ActiveSection == nullptr ? ECheckBoxState::Checked : ECheckBoxState::Unchecked;
 					})
diff --git a/Engine/Plugins/FX/Niagara/Source/NiagaraEditor/Private/NiagaraNodeOp.cpp b/Engine/Plugins/FX/Niagara/Source/NiagaraEditor/Private/NiagaraNodeOp.cpp
--- a/Engine/Plugins/FX/Niagara/Source/NiagaraEditor/Private/NiagaraNodeOp.cpp
+++ b/Engine/Plugins/FX/Niagara/Source/NiagaraEditor/Private/NiagaraNodeOp.cpp
@@ -330,7 +330,7 @@ bool UNiagaraNodeOp::AllowNiagaraTypeForAddPin(const FNiagaraTypeDefinition& InT
 		return true;
 	}
 
-	auto FindPredicate = [=](const FNiagaraTypeDefinition& PinType)
+	auto FindPredicate = [this, InType](const FNiagaraTypeDefinition& PinType)
 	{
 		if (bAllStatic)
 			return PinType.ToStaticDef() == InType;
diff --git a/Engine/Plugins/FX/Niagara/Source/NiagaraEditor/Private/NiagaraNodeSelect.cpp b/Engine/Plugins/FX/Niagara/Source/NiagaraEditor/Private/NiagaraNodeSelect.cpp
--- a/Engine/Plugins/FX/Niagara/Source/NiagaraEditor/Private/NiagaraNodeSelect.cpp
+++ b/Engine/Plugins/FX/Niagara/Source/NiagaraEditor/Private/NiagaraNodeSelect.cpp
@@ -279,11 +279,11 @@ void UNiagaraNodeSelect::GetNodeContextMenuActions(UToolMenu* Menu, UGraphNodeCo
 			FText::Format(LOCTEXT("BrowseToEnumLabel", "Browse to {0}"), FText::FromString(SelectorPinType.GetEnum()->GetName())),
 			LOCTEXT("BrowseToEnumTooltip", "Browses to the enum in the content browser."),
 			FSlateIcon(),
-			FUIAction(FExecuteAction::CreateLambda([=]()
+			FUIAction(FExecuteAction::CreateLambda([this]()
 			{
 				FContentBrowserModule& Module = FModuleManager::Get().LoadModuleChecked<FContentBrowserModule>("ContentBrowser");
 				Module.Get().SyncBrowserToAssets({ FAssetData(SelectorPinType.GetEnum()) });
-			}), FCanExecuteAction::CreateLambda([=]()
+			}), FCanExecuteAction::CreateLambda([this]()
 			{
 				return Cast<UUserDefinedEnum>(SelectorPinType.GetEnum()) != nullptr;
 			})));
@@ -293,10 +293,10 @@ void UNiagaraNodeSelect::GetNodeContextMenuActions(UToolMenu* Menu, UGraphNodeCo
 			FText::Format(LOCTEXT("OpenEnumLabel", "Open {0}"), FText::FromString(SelectorPinType.GetEnum()->GetName())),
 			LOCTEXT("OpenEnumTooltip", "Opens up the enum asset."),
 			FSlateIcon(),
-			FUIAction(FExecuteAction::CreateLambda([=]()
+			FUIAction(FExecuteAction::CreateLambda([this]()
 			{
 				GEditor->GetEditorSubsystem<UAssetEditorSubsystem>()->OpenEditorForAsset(SelectorPinType.GetEnum());
-			}), FCanExecuteAction::CreateLambda([=]()
+			}), FCanExecuteAction::CreateLambda([this]()
 			{
 				return Cast<UUserDefinedEnum>(SelectorPinType.GetEnum()) != nullptr;
 			})));
diff --git a/Engine/Plugins/FX/Niagara/Source/NiagaraEditor/Private/NiagaraNodeStaticSwitch.cpp b/Engine/Plugins/FX/Niagara/Source/NiagaraEditor/Private/NiagaraNodeStaticSwitch.cpp
--- a/Engine/Plugins/FX/Niagara/Source/NiagaraEditor/Private/NiagaraNodeStaticSwitch.cpp
+++ b/Engine/Plugins/FX/Niagara/Source/NiagaraEditor/Private/NiagaraNodeStaticSwitch.cpp
@@ -1185,11 +1185,11 @@ void UNiagaraNodeStaticSwitch::GetNodeContextMenuActions(UToolMenu* Menu, UGraph
 			FText::Format(LOCTEXT("BrowseToEnumLabel", "Browse to {0}"), FText::FromString(SwitchTypeData.Enum->GetName())),
 			LOCTEXT("BrowseToEnumTooltip", "Browses to the enum in the content browser."),
 			FSlateIcon(),
-			FUIAction(FExecuteAction::CreateLambda([=]()
+			FUIAction(FExecuteAction::CreateLambda([this]()
 			{
 				FContentBrowserModule& Module = FModuleManager::Get().LoadModuleChecked<FContentBrowserModule>("ContentBrowser");
 				Module.Get().SyncBrowserToAssets({ FAssetData(SwitchTypeData.Enum) });
-			}) , FCanExecuteAction::CreateLambda([=]()
+			}) , FCanExecuteAction::CreateLambda([this]()
 			{
 				return Cast<UUserDefinedEnum>(SwitchTypeData.Enum) != nullptr;
 			})));
@@ -1216,7 +1216,7 @@ void UNiagaraNodeStaticSwitch::GetNodeContextMenuActions(UToolMenu* Menu, UGraph
 			FUIAction(FExecuteAction::CreateLambda([=]()
 			{
 				GetMutableDefault<UNiagaraSettings>()->AddEnumParameterType(SwitchTypeData.Enum);
-			}), FCanExecuteAction::CreateLambda([=]()
+			}), FCanExecuteAction::CreateLambda([this]()
 			{
 				return !GetDefault<UNiagaraSettings>()->AdditionalParameterEnums.Contains(SwitchTypeData.Enum);
 			})));
diff --git a/Engine/Plugins/FX/Niagara/Source/NiagaraEditor/Private/NiagaraNodeWithDynamicPins.cpp b/Engine/Plugins/FX/Niagara/Source/NiagaraEditor/Private/NiagaraNodeWithDynamicPins.cpp
--- a/Engine/Plugins/FX/Niagara/Source/NiagaraEditor/Private/NiagaraNodeWithDynamicPins.cpp
+++ b/Engine/Plugins/FX/Niagara/Source/NiagaraEditor/Private/NiagaraNodeWithDynamicPins.cpp
@@ -321,25 +321,25 @@ void UNiagaraNodeWithDynamicPins::GetNodeContextMenuActions(UToolMenu* Menu, UGr
 					"MoveDynamicPinUp",
 					MoveUpLabel, 
 					MoveUpTooltip,
-					FNewToolMenuDelegate::CreateLambda([=](UToolMenu* InSubMenuBuilder)
-				{
-					FToolMenuSection& SubSection = InSubMenuBuilder->FindOrAddSection("MovePin");
-					for (int i = PinIdx - 1; i >= 0; i--)
+					FNewToolMenuDelegate::CreateLambda([=,this](UToolMenu* InSubMenuBuilder)
 					{
-						int32 MoveAmount = i - PinIdx;
-						if (!CanMovePin(Context->Pin, MoveAmount))
+						FToolMenuSection& SubSection = InSubMenuBuilder->FindOrAddSection("MovePin");
+						for (int i = PinIdx - 1; i >= 0; i--)
 						{
-							break;
+							int32 MoveAmount = i - PinIdx;
+							if (!CanMovePin(Context->Pin, MoveAmount))
+							{
+								break;
+							}
+							FText PinName = FText::FromName(SameDirectionPins[i]->PinName);
+							SubSection.AddMenuEntry(FName("MoveDynamicPinUp" + FString::FromInt(i)),
+								FText::Format(LOCTEXT("MoveDynamicPinUpLabel", "Move up {0} - above '{1}'"), abs(MoveAmount), PinName),
+								MoveUpTooltip,
+								FSlateIcon(),
+								FUIAction(FExecuteAction::CreateUObject(const_cast<UNiagaraNodeWithDynamicPins*>(this), &UNiagaraNodeWithDynamicPins::MoveDynamicPinFromMenu, const_cast<UEdGraphPin*>(Context->Pin), MoveAmount))
+							);
 						}
-						FText PinName = FText::FromName(SameDirectionPins[i]->PinName);
-						SubSection.AddMenuEntry(FName("MoveDynamicPinUp" + FString::FromInt(i)),
-							FText::Format(LOCTEXT("MoveDynamicPinUpLabel", "Move up {0} - above '{1}'"), abs(MoveAmount), PinName),
-							MoveUpTooltip,
-							FSlateIcon(),
-							FUIAction(FExecuteAction::CreateUObject(const_cast<UNiagaraNodeWithDynamicPins*>(this), &UNiagaraNodeWithDynamicPins::MoveDynamicPinFromMenu, const_cast<UEdGraphPin*>(Context->Pin), MoveAmount))
-						);
-					}
-				}));
+					}));
 			}
 			else
 			{
@@ -362,7 +362,7 @@ void UNiagaraNodeWithDynamicPins::GetNodeContextMenuActions(UToolMenu* Menu, UGr
 					"MoveDynamicPinDown",
 					MoveDownLabel, 
 					MoveDownTooltip,
-					FNewToolMenuDelegate::CreateLambda([=](UToolMenu* InSubMenuBuilder)
+					FNewToolMenuDelegate::CreateLambda([=, this](UToolMenu* InSubMenuBuilder)
 				{
 					FToolMenuSection& SubSection = InSubMenuBuilder->FindOrAddSection("MovePin");
 					for (int i = PinIdx + 1; i < SameDirectionPins.Num(); i++)
diff --git a/Engine/Plugins/FX/Niagara/Source/NiagaraEditor/Private/ViewModels/Stack/NiagaraStackFunctionInputCollection.cpp b/Engine/Plugins/FX/Niagara/Source/NiagaraEditor/Private/ViewModels/Stack/NiagaraStackFunctionInputCollection.cpp
--- a/Engine/Plugins/FX/Niagara/Source/NiagaraEditor/Private/ViewModels/Stack/NiagaraStackFunctionInputCollection.cpp
+++ b/Engine/Plugins/FX/Niagara/Source/NiagaraEditor/Private/ViewModels/Stack/NiagaraStackFunctionInputCollection.cpp
@@ -402,7 +402,7 @@ void UNiagaraStackFunctionInputCollection::RefreshIssues(const TArray<FName>& Du
 				FText ConversionFixDescription = LOCTEXT("ConvertInputToStaticSwitchTransaction", "Copy value to static switch parameter");
 				FStackIssueFix ConvertInputOverrideFix(
 					ConversionFixDescription,
-					UNiagaraStackEntry::FStackIssueFixDelegate::CreateLambda([=]()
+					UNiagaraStackEntry::FStackIssueFixDelegate::CreateLambda([=, this]()
 						{
 							FScopedTransaction ScopedTransaction(ConversionFixDescription);
 							SwitchPin->Modify();
@@ -561,7 +561,7 @@ UNiagaraStackEntry::FStackIssueFix UNiagaraStackFunctionInputCollection::GetNode
 {
 	return FStackIssueFix(
 		FixDescription,
-		UNiagaraStackEntry::FStackIssueFixDelegate::CreateLambda([=]()
+		UNiagaraStackEntry::FStackIssueFixDelegate::CreateLambda([=, this]()
 			{
 				FScopedTransaction ScopedTransaction(FixDescription);
 				TArray<TWeakObjectPtr<UNiagaraDataInterface>> RemovedDataObjects;
diff --git a/Engine/Plugins/LightWeightInstancesEditor/Source/LightWeightInstancesEditor/Private/LightWeightInstancesEditor.cpp b/Engine/Plugins/LightWeightInstancesEditor/Source/LightWeightInstancesEditor/Private/LightWeightInstancesEditor.cpp
--- a/Engine/Plugins/LightWeightInstancesEditor/Source/LightWeightInstancesEditor/Private/LightWeightInstancesEditor.cpp
+++ b/Engine/Plugins/LightWeightInstancesEditor/Source/LightWeightInstancesEditor/Private/LightWeightInstancesEditor.cpp
@@ -52,7 +52,7 @@ void FLightWeightInstancesEditorModule::RemoveLevelViewportMenuExtender()
 		FLevelEditorModule* LevelEditorModule = FModuleManager::Get().GetModulePtr<FLevelEditorModule>("LevelEditor");
 		if (LevelEditorModule)
 		{
-			LevelEditorModule->GetAllLevelViewportContextMenuExtenders().RemoveAll([=](const DelegateType& In) { return In.GetHandle() == LevelViewportExtenderHandle; });
+			LevelEditorModule->GetAllLevelViewportContextMenuExtenders().RemoveAll([this](const DelegateType& In) { return In.GetHandle() == LevelViewportExtenderHandle; });
 		}
 	}
 #endif
diff --git a/Engine/Plugins/Media/MediaIOFramework/Source/MediaIOEditor/Private/Customizations/MediaIODeviceCustomization.cpp b/Engine/Plugins/Media/MediaIOFramework/Source/MediaIOEditor/Private/Customizations/MediaIODeviceCustomization.cpp
--- a/Engine/Plugins/Media/MediaIOFramework/Source/MediaIOEditor/Private/Customizations/MediaIODeviceCustomization.cpp
+++ b/Engine/Plugins/Media/MediaIOFramework/Source/MediaIOEditor/Private/Customizations/MediaIODeviceCustomization.cpp
@@ -59,7 +59,7 @@ TSharedRef<SWidget> FMediaIODeviceCustomization::HandleSourceComboButtonMenuCont
 				DeviceProviderPtr->ToText(Device),
 				FSlateIcon(),
 				FUIAction(
-					FExecuteAction::CreateLambda([=] {
+					FExecuteAction::CreateLambda([this, Device] {
 						AssignValue(Device);
 					}),
 					FCanExecuteAction(),
diff --git a/Engine/Plugins/Media/MediaIOFramework/Source/MediaIOEditor/Private/Customizations/MediaIOVideoTimecodeConfigurationCustomization.cpp b/Engine/Plugins/Media/MediaIOFramework/Source/MediaIOEditor/Private/Customizations/MediaIOVideoTimecodeConfigurationCustomization.cpp
--- a/Engine/Plugins/Media/MediaIOFramework/Source/MediaIOEditor/Private/Customizations/MediaIOVideoTimecodeConfigurationCustomization.cpp
+++ b/Engine/Plugins/Media/MediaIOFramework/Source/MediaIOEditor/Private/Customizations/MediaIOVideoTimecodeConfigurationCustomization.cpp
@@ -67,7 +67,7 @@ namespace MediaIOVideoTimecodeConfigurationCustomization
 TAttribute<FText> FMediaIOVideoTimecodeConfigurationCustomization::GetContentText()
 {
 	FMediaIOVideoTimecodeConfiguration* Value = GetPropertyValueFromPropertyHandle<FMediaIOVideoTimecodeConfiguration>();
-	return MakeAttributeLambda([=]
+	return MakeAttributeLambda([this, Value]
 	{
 		const bool bIsAutoDetected = IsAutoDetected();
 		return Value->ToText(bIsAutoDetected);
diff --git a/Engine/Plugins/Media/MediaPlayerEditor/Source/MediaPlayerEditor/Private/Customizations/BaseMediaSourceCustomization.cpp b/Engine/Plugins/Media/MediaPlayerEditor/Source/MediaPlayerEditor/Private/Customizations/BaseMediaSourceCustomization.cpp
--- a/Engine/Plugins/Media/MediaPlayerEditor/Source/MediaPlayerEditor/Private/Customizations/BaseMediaSourceCustomization.cpp
+++ b/Engine/Plugins/Media/MediaPlayerEditor/Source/MediaPlayerEditor/Private/Customizations/BaseMediaSourceCustomization.cpp
@@ -71,7 +71,7 @@ TSharedRef<SWidget> FBaseMediaSourceCustomization::MakePlatformPlayersMenu(const
 		LOCTEXT("AutoPlayerTooltip", "Select a player automatically based on the media source"),
 		FSlateIcon(),
 		FUIAction(
-			FExecuteAction::CreateLambda([=] { SetPlatformPlayerNamesValue(IniPlatformName, NAME_None); }),
+		FExecuteAction::CreateLambda([this, IniPlatformName] { SetPlatformPlayerNamesValue(IniPlatformName, NAME_None); }),
 			FCanExecuteAction()
 		),
 		NAME_None,
@@ -100,7 +100,7 @@ TSharedRef<SWidget> FBaseMediaSourceCustomization::MakePlatformPlayersMenu(const
 				FText::FromString(FString::Join(Factory->GetSupportedPlatforms(), TEXT(", "))),
 				FSlateIcon(),
 				FUIAction(
-					FExecuteAction::CreateLambda([=] { SetPlatformPlayerNamesValue(IniPlatformName, PlayerName); }),
+				FExecuteAction::CreateLambda([this, IniPlatformName, PlayerName] { SetPlatformPlayerNamesValue(IniPlatformName, PlayerName); }),
 					FCanExecuteAction::CreateLambda([=] { return SupportsPlatform; })
 				),
 				NAME_None,
diff --git a/Engine/Plugins/Media/MediaPlayerEditor/Source/MediaPlayerEditor/Private/Widgets/SMediaPlayerEditorPlaylist.cpp b/Engine/Plugins/Media/MediaPlayerEditor/Source/MediaPlayerEditor/Private/Widgets/SMediaPlayerEditorPlaylist.cpp
--- a/Engine/Plugins/Media/MediaPlayerEditor/Source/MediaPlayerEditor/Private/Widgets/SMediaPlayerEditorPlaylist.cpp
+++ b/Engine/Plugins/Media/MediaPlayerEditor/Source/MediaPlayerEditor/Private/Widgets/SMediaPlayerEditorPlaylist.cpp
@@ -198,7 +198,7 @@ void SMediaPlayerEditorPlaylist::Construct(const FArguments& InArgs, UMediaPlaye
 							{
 								return SNew(SMediaSourceTableRow, OwnerTable)
 									.Entry(Entry)
-									.Opened_Lambda([=]() { return MediaPlayer->GetPlaylistIndex() == Entry->Index; })
+									.Opened_Lambda([this, Entry]() { return MediaPlayer->GetPlaylistIndex() == Entry->Index; })
 									.Style(Style);
 							})
 							.OnMouseButtonDoubleClick_Lambda([this](TSharedPtr<FMediaSourceTableEntry> InItem)
diff --git a/Engine/Plugins/Media/PixelCapture/Source/PixelCapture/Private/PixelCaptureCapturer.cpp b/Engine/Plugins/Media/PixelCapture/Source/PixelCapture/Private/PixelCaptureCapturer.cpp
--- a/Engine/Plugins/Media/PixelCapture/Source/PixelCapture/Private/PixelCaptureCapturer.cpp
+++ b/Engine/Plugins/Media/PixelCapture/Source/PixelCapture/Private/PixelCaptureCapturer.cpp
@@ -55,7 +55,7 @@ void FPixelCaptureCapturer::Capture(const IPixelCaptureInputFrame& InputFrame)
 void FPixelCaptureCapturer::Initialize(int32 InputWidth, int32 InputHeight)
 {
 	Buffer = MakeUnique<UE::PixelCapture::FOutputFrameBuffer>();
-	Buffer->Reset(3, 10, [=]() { return TSharedPtr<IPixelCaptureOutputFrame>(CreateOutputBuffer(InputWidth, InputHeight)); });
+	Buffer->Reset(3, 10, [this, InputWidth, InputHeight]() { return TSharedPtr<IPixelCaptureOutputFrame>(CreateOutputBuffer(InputWidth, InputHeight)); });
 	ExpectedInputWidth = InputWidth;
 	ExpectedInputHeight = InputHeight;
 	bHasOutput = false;
diff --git a/Engine/Plugins/Runtime/GeometryCache/Source/GeometryCacheSequencer/Classes/GeometryCacheTrackEditor.cpp b/Engine/Plugins/Runtime/GeometryCache/Source/GeometryCacheSequencer/Classes/GeometryCacheTrackEditor.cpp
--- a/Engine/Plugins/Runtime/GeometryCache/Source/GeometryCacheSequencer/Classes/GeometryCacheTrackEditor.cpp
+++ b/Engine/Plugins/Runtime/GeometryCache/Source/GeometryCacheSequencer/Classes/GeometryCacheTrackEditor.cpp
@@ -387,7 +387,7 @@ TSharedPtr<SWidget> FGeometryCacheTrackEditor::BuildOutlinerEditWidget(const FGu
 	{
 		TWeakPtr<ISequencer> WeakSequencer = GetSequencer();
 
-		auto SubMenuCallback = [=]() -> TSharedRef<SWidget>
+		auto SubMenuCallback = [this, ObjectBinding, Track]() -> TSharedRef<SWidget>
 		{
 			FMenuBuilder MenuBuilder(true, nullptr);
 
diff --git a/Engine/Plugins/VirtualProduction/Takes/Source/TakeRecorder/Private/Widgets/STakeRecorderCockpit.cpp b/Engine/Plugins/VirtualProduction/Takes/Source/TakeRecorder/Private/Widgets/STakeRecorderCockpit.cpp
--- a/Engine/Plugins/VirtualProduction/Takes/Source/TakeRecorder/Private/Widgets/STakeRecorderCockpit.cpp
+++ b/Engine/Plugins/VirtualProduction/Takes/Source/TakeRecorder/Private/Widgets/STakeRecorderCockpit.cpp
@@ -488,7 +488,7 @@ void STakeRecorderCockpit::Construct(const FArguments& InArgs)
 						.ToolTipText(LOCTEXT("EngineTimeDilation", "Recording speed"))
 						.Style(&FAppStyle::GetWidgetStyle<FSpinBoxStyle>("Sequencer.HyperlinkSpinBox"))
 						.OnValueChanged(this, &STakeRecorderCockpit::SetEngineTimeDilation)
-						.OnValueCommitted_Lambda([=](float InEngineTimeDilation, ETextCommit::Type) { SetEngineTimeDilation(InEngineTimeDilation); })
+						.OnValueCommitted_Lambda([this](float InEngineTimeDilation, ETextCommit::Type) { SetEngineTimeDilation(InEngineTimeDilation); })
 						.MinValue(TOptional<float>())
 						.MaxValue(TOptional<float>())
 						.Value(this, &STakeRecorderCockpit::GetEngineTimeDilation)
diff --git a/Engine/Source/Developer/DirectoryWatcher/Private/DirectoryWatcherTests.cpp b/Engine/Source/Developer/DirectoryWatcher/Private/DirectoryWatcherTests.cpp
--- a/Engine/Source/Developer/DirectoryWatcher/Private/DirectoryWatcherTests.cpp
+++ b/Engine/Source/Developer/DirectoryWatcher/Private/DirectoryWatcherTests.cpp
@@ -99,7 +99,7 @@ bool FDirectoryWatcherSimpleCreateTest::RunTest(const FString& Parameters)
 		// Create a file and check that it got reported as created
 		FFileHelper::SaveStringToFile(TEXT(""), *(WorkingDir / Filename));
 
-		AddCommand(new FDelayedFunctionLatentCommand([=]{
+		AddCommand(new FDelayedFunctionLatentCommand([this, Test]{
 
 			FFileChangeData::EFileChangeAction* Action = Test->ReportedChanges.Find(Filename);
 
@@ -130,12 +130,12 @@ bool FDirectoryWatcherSimpleModifyTest::RunTest(const FString& Parameters)
 		// Start watching the directory
 		TSharedPtr<FDirectoryWatcherTestPayload> Test = MakeShareable(new FDirectoryWatcherTestPayload(WorkingDir));
 
-		AddCommand(new FDelayedFunctionLatentCommand([=]{
+		AddCommand(new FDelayedFunctionLatentCommand([this, WorkingDir, Test]{
 
 			// Manipulate the file
 			FFileHelper::SaveStringToFile(TEXT("Some content"), *(WorkingDir / Filename));
 
-			AddCommand(new FDelayedFunctionLatentCommand([=]{
+			AddCommand(new FDelayedFunctionLatentCommand([this, Test]{
 
 				FFileChangeData::EFileChangeAction* Action = Test->ReportedChanges.Find(Filename);
 
@@ -172,7 +172,7 @@ bool FDirectoryWatcherSimpleDeleteTest::RunTest(const FString& Parameters)
 		// Delete the file
 		IFileManager::Get().Delete(*(WorkingDir / Filename));
 
-		AddCommand(new FDelayedFunctionLatentCommand([=]{
+		AddCommand(new FDelayedFunctionLatentCommand([this, Test]{
 
 			FFileChangeData::EFileChangeAction* Action = Test->ReportedChanges.Find(Filename);
 
@@ -209,7 +209,7 @@ bool FDirectoryWatcherSubFolderTest::RunTest(const FString& Parameters)
 		TSharedPtr<FDirectoryWatcherTestPayload> Test = MakeShareable(new FDirectoryWatcherTestPayload(WorkingDir));
 
 		// Give the stream time to start up before doing the test
-		AddCommand(new FDelayedFunctionLatentCommand([=]{
+		AddCommand(new FDelayedFunctionLatentCommand([this, WorkingDir, Test]{
 
 			// Create a new file
 			FFileHelper::SaveStringToFile(TEXT(""), *(WorkingDir / CreatedFilename));
@@ -218,7 +218,7 @@ bool FDirectoryWatcherSubFolderTest::RunTest(const FString& Parameters)
 			// Delete a file
 			IFileManager::Get().Delete(*(WorkingDir / RemovedFilename));
 
-			AddCommand(new FDelayedFunctionLatentCommand([=]{
+			AddCommand(new FDelayedFunctionLatentCommand([this, Test]{
 
 				FFileChangeData::EFileChangeAction* Action = Test->ReportedChanges.Find(CreatedFilename);
 				if (!Action || *Action != FFileChangeData::FCA_Added)
@@ -266,12 +266,12 @@ bool FDirectoryWatcherNewFolderTest::RunTest(const FString& Parameters)
 		TSharedPtr<FDirectoryWatcherTestPayload> Test = MakeShareable(new FDirectoryWatcherTestPayload(WorkingDir, IDirectoryWatcher::WatchOptions::IncludeDirectoryChanges));
 
 		// Give the stream time to start up before doing the test
-		AddCommand(new FDelayedFunctionLatentCommand([=] {
+		AddCommand(new FDelayedFunctionLatentCommand([this, WorkingDir, Test] {
 
 			IFileManager::Get().MakeDirectory(*(WorkingDir / CreatedDirectory), true);
 			IFileManager::Get().DeleteDirectory(*(WorkingDir / RemovedDirectory), true);
 
-			AddCommand(new FDelayedFunctionLatentCommand([=] {
+			AddCommand(new FDelayedFunctionLatentCommand([this, Test] {
 
 				FFileChangeData::EFileChangeAction* Action = Test->ReportedChanges.Find(CreatedDirectory);
 				if (!Action || *Action != FFileChangeData::FCA_Added)
@@ -311,12 +311,12 @@ bool FDirectoryWatcherIgnoreSubtreeTest::RunTest(const FString& Parameters)
 		TSharedPtr<FDirectoryWatcherTestPayload> Test = MakeShareable(new FDirectoryWatcherTestPayload(WorkingDir, IDirectoryWatcher::WatchOptions::IgnoreChangesInSubtree | IDirectoryWatcher::WatchOptions::IncludeDirectoryChanges));
 
 		// Give the stream time to start up before doing the test
-		AddCommand(new FDelayedFunctionLatentCommand([=]{
+		AddCommand(new FDelayedFunctionLatentCommand([this, WorkingDir, Test]{
 
 			IFileManager::Get().MakeDirectory(*(WorkingDir / ChildDirectory), true);
 			IFileManager::Get().MakeDirectory(*(WorkingDir / ChildDirectory / GrandchildDirectory), true);
 
-			AddCommand(new FDelayedFunctionLatentCommand([=]{
+			AddCommand(new FDelayedFunctionLatentCommand([this, Test]{
 
 				FFileChangeData::EFileChangeAction* Action = Test->ReportedChanges.Find(ChildDirectory);
 				if (!Action || *Action != FFileChangeData::FCA_Added)
diff --git a/Engine/Source/Developer/MaterialBaking/Private/MaterialOptionsCustomization.cpp b/Engine/Source/Developer/MaterialBaking/Private/MaterialOptionsCustomization.cpp
--- a/Engine/Source/Developer/MaterialBaking/Private/MaterialOptionsCustomization.cpp
+++ b/Engine/Source/Developer/MaterialBaking/Private/MaterialOptionsCustomization.cpp
@@ -161,7 +161,7 @@ void FMaterialOptionsCustomization::CustomizeDetails(IDetailLayoutBuilder& Detai
 	}
 
 	// Only allow changes to LOD indices if we have a valid options instance and if there is actually more than one index
-	ContentBox->SetEnabled(TAttribute<bool>::Create([=]() -> bool { return NumLODs > 1 && CurrentOptions != nullptr; }));
+	ContentBox->SetEnabled(TAttribute<bool>::Create([this, CurrentOptions]() -> bool { return NumLODs > 1 && CurrentOptions != nullptr; }));
 	for (int32 LODIndex = 0; LODIndex < NumLODs; ++LODIndex)
 	{
 		ContentBox->AddSlot()
diff --git a/Engine/Source/Developer/MeshUtilities/Private/MeshUtilities.cpp b/Engine/Source/Developer/MeshUtilities/Private/MeshUtilities.cpp
--- a/Engine/Source/Developer/MeshUtilities/Private/MeshUtilities.cpp
+++ b/Engine/Source/Developer/MeshUtilities/Private/MeshUtilities.cpp
@@ -6274,7 +6274,7 @@ void FMeshUtilities::RemoveLevelViewportMenuExtender()
 		if (LevelEditorModule)
 		{
 			typedef FLevelEditorModule::FLevelViewportMenuExtender_SelectedActors DelegateType;
-			LevelEditorModule->GetAllLevelViewportContextMenuExtenders().RemoveAll([=](const DelegateType& In) { return In.GetHandle() == LevelViewportExtenderHandle; });
+			LevelEditorModule->GetAllLevelViewportContextMenuExtenders().RemoveAll([this](const DelegateType& In) { return In.GetHandle() == LevelViewportExtenderHandle; });
 		}
 	}
 }
diff --git a/Engine/Source/Developer/TurnkeyIO/Private/STurnkeyIOWidgets.cpp b/Engine/Source/Developer/TurnkeyIO/Private/STurnkeyIOWidgets.cpp
--- a/Engine/Source/Developer/TurnkeyIO/Private/STurnkeyIOWidgets.cpp
+++ b/Engine/Source/Developer/TurnkeyIO/Private/STurnkeyIOWidgets.cpp
@@ -49,9 +49,9 @@ void STurnkeyReadInputModal::Construct( const FArguments& InArgs )
 			.VAlign(VAlign_Top)
 			[
 				SNew(SEditableTextBox)
-				.Text_Lambda( [=]() { return FText::FromString(Value); } )
-				.OnTextChanged_Lambda( [=]( const FText& NewText ) { Value = NewText.ToString(); } )
-				.OnTextCommitted_Lambda( [=]( const FText& NewText, ETextCommit::Type CommitType )
+				.Text_Lambda( [this]() { return FText::FromString(Value); } )
+				.OnTextChanged_Lambda( [this]( const FText& NewText ) { Value = NewText.ToString(); } )
+				.OnTextCommitted_Lambda( [this]( const FText& NewText, ETextCommit::Type CommitType )
 				{
 					if (CommitType == ETextCommit::OnEnter)
 					{
@@ -74,7 +74,7 @@ void STurnkeyReadInputModal::Construct( const FArguments& InArgs )
 				[
 					SNew(SPrimaryButton)
 					.Text(LOCTEXT("OK", "OK"))
-					.OnClicked_Lambda( [=]()
+					.OnClicked_Lambda( [this]()
 					{ 
 						FinishAction(); 
 						return FReply::Handled();
@@ -152,8 +152,8 @@ void STurnkeyReadInputIntModal::Construct( const FArguments& InArgs )
 				[
 					SNew(SPrimaryButton)
 					.Text(LOCTEXT("OK", "OK"))
-					.IsEnabled_Lambda( [=]() { return Value > 0; } )
-					.OnClicked_Lambda( [=]()
+					.IsEnabled_Lambda( [this]() { return Value > 0; } )
+					.OnClicked_Lambda( [this]()
 					{ 
 						FinishAction();
 						return FReply::Handled();
@@ -176,7 +176,7 @@ void STurnkeyReadInputIntModal::Construct( const FArguments& InArgs )
 			.HAlign(HAlign_Center)
 			.VAlign(VAlign_Center)
 			.Text(LOCTEXT("Cancel", "Cancel"))
-			.OnClicked_Lambda( [=]()
+			.OnClicked_Lambda( [this]()
 			{ 
 				Value = 0;
 				FinishAction();
@@ -216,8 +216,8 @@ void STurnkeyReadInputIntModal::Construct( const FArguments& InArgs )
 				[
 					SNew(SCheckBox)
 					.Style(FAppStyle::Get(), "RadioButton")
-					.IsChecked_Lambda( [=]() { return (Value == ResultValue) ? ECheckBoxState::Checked : ECheckBoxState::Unchecked; } )
-					.OnCheckStateChanged_Lambda([=](const ECheckBoxState NewState)
+					.IsChecked_Lambda( [this, ResultValue]() { return (Value == ResultValue) ? ECheckBoxState::Checked : ECheckBoxState::Unchecked; } )
+					.OnCheckStateChanged_Lambda([this, ResultValue](const ECheckBoxState NewState)
 					{
 						if (NewState == ECheckBoxState::Checked)
 						{
@@ -234,7 +234,7 @@ void STurnkeyReadInputIntModal::Construct( const FArguments& InArgs )
 				[
 					SAssignNew(TextBlock, STextBlock)
 					.Text( FText::FromString(OptionValue) )
-					.OnDoubleClicked_Lambda([=](const FGeometry& MyGeometry, const FPointerEvent& PointerEvent)
+					.OnDoubleClicked_Lambda([this, ResultValue](const FGeometry& MyGeometry, const FPointerEvent& PointerEvent)
 					{
 						// double click the label to select the option immediately
 						Value = ResultValue;
@@ -246,7 +246,7 @@ void STurnkeyReadInputIntModal::Construct( const FArguments& InArgs )
 			];
 
 			// change the option if the user clicks on the label
-			TextBlock->SetOnMouseButtonUp( FPointerEventHandler::CreateLambda( [=](const FGeometry& MyGeometry, const FPointerEvent& PointerEvent)
+			TextBlock->SetOnMouseButtonUp( FPointerEventHandler::CreateLambda( [this, ResultValue](const FGeometry& MyGeometry, const FPointerEvent& PointerEvent)
 			{
 				Value = ResultValue;
 				return FReply::Handled();
diff --git a/Engine/Source/Developer/TurnkeyIO/Private/TurnkeyEditorIOServer.cpp b/Engine/Source/Developer/TurnkeyIO/Private/TurnkeyEditorIOServer.cpp
--- a/Engine/Source/Developer/TurnkeyIO/Private/TurnkeyEditorIOServer.cpp
+++ b/Engine/Source/Developer/TurnkeyIO/Private/TurnkeyEditorIOServer.cpp
@@ -167,7 +167,7 @@ private:
 
 	void DoPauseForUser( TSharedPtr<FJsonObject> JsonObject )
 	{
-		AsyncTask( ENamedThreads::GameThread, [=]()
+		AsyncTask( ENamedThreads::GameThread, [this, JsonObject]()
 		{
 			// read parameters
 			FString Message;
@@ -186,7 +186,7 @@ private:
 
 	void DoReadInput( TSharedPtr<FJsonObject> JsonObject )
 	{
-		AsyncTask( ENamedThreads::GameThread, [=]()
+		AsyncTask( ENamedThreads::GameThread, [this, JsonObject]()
 		{
 			// read parameters
 			FString Prompt;
@@ -208,7 +208,7 @@ private:
 
 	void DoReadInputInt( TSharedPtr<FJsonObject> JsonObject )
 	{
-		AsyncTask( ENamedThreads::GameThread, [=]()
+		AsyncTask( ENamedThreads::GameThread, [this, JsonObject]()
 		{
 			// read parameters
 			FString Prompt;
@@ -237,7 +237,7 @@ private:
 
 	void DoGetUserConfirmation( TSharedPtr<FJsonObject> JsonObject )
 	{
-		AsyncTask( ENamedThreads::GameThread, [=]()
+		AsyncTask( ENamedThreads::GameThread, [this, JsonObject]()
 		{
 			// read parameters
 			FString Message;
@@ -273,7 +273,7 @@ private:
 			];
 
 		// even though HasCloseButton is false, the window can still be closed via Windows' alt-tab screen
-		Window->GetOnWindowClosedEvent().AddLambda( [=]( const TSharedRef<SWindow>& InWindow )
+		Window->GetOnWindowClosedEvent().AddLambda( [this]( const TSharedRef<SWindow>& InWindow )
 		{
 			if (Window.IsValid())
 			{
diff --git a/Engine/Source/Editor/AnimGraph/Private/PoseDriverDetails.cpp b/Engine/Source/Editor/AnimGraph/Private/PoseDriverDetails.cpp
--- a/Engine/Source/Editor/AnimGraph/Private/PoseDriverDetails.cpp
+++ b/Engine/Source/Editor/AnimGraph/Private/PoseDriverDetails.cpp
@@ -329,7 +329,7 @@ TSharedRef< SWidget > SPDD_TargetRow::GenerateWidgetForColumn(const FName& Colum
 						.Padding(FMargin(6, 0, 3, 0))
 						[
 							SNew(SCheckBox)
-							.IsChecked_Lambda([=]() { return IsHidden() ? ECheckBoxState::Checked : ECheckBoxState::Unchecked; })
+							.IsChecked_Lambda([this]() { return IsHidden() ? ECheckBoxState::Checked : ECheckBoxState::Unchecked; })
 							.OnCheckStateChanged(this, &SPDD_TargetRow::OnIsHiddenChanged)
 							.Padding(FMargin(4.0f, 0.0f))
 							.ToolTipText(LOCTEXT("IsHiddenToolTip", "Define if this target should be hidden from debug drawing."))
@@ -382,7 +382,7 @@ TSharedRef< SWidget > SPDD_TargetRow::GenerateWidgetForColumn(const FName& Colum
 						.Padding(FMargin(6, 0, 3, 0))
 						[
 							SNew(SSearchableComboBox)
-							.Visibility_Lambda([=]() { return IsCustomCurveEnabled() ? EVisibility::Collapsed : EVisibility::Visible;  })
+							.Visibility_Lambda([this]() { return IsCustomCurveEnabled() ? EVisibility::Collapsed : EVisibility::Visible;  })
 							.OptionsSource(&FunctionTypeOptions)
 							.OnGenerateWidget(this, &SPDD_TargetRow::MakeDrivenNameWidget)
 							.OnSelectionChanged(this, &SPDD_TargetRow::OnFunctionTypeChanged)
@@ -409,7 +409,7 @@ TSharedRef< SWidget > SPDD_TargetRow::GenerateWidgetForColumn(const FName& Colum
 						.Padding(FMargin(6, 0, 3, 0))
 						[
 							SNew(SCheckBox)
-							.IsChecked_Lambda([=]() { return IsCustomCurveEnabled() ? ECheckBoxState::Checked : ECheckBoxState::Unchecked; })
+							.IsChecked_Lambda([this]() { return IsCustomCurveEnabled() ? ECheckBoxState::Checked : ECheckBoxState::Unchecked; })
 							.OnCheckStateChanged(this, &SPDD_TargetRow::OnApplyCustomCurveChanged)
 							.Padding(FMargin(4.0f, 0.0f))
 							.ToolTipText(LOCTEXT("CustomCurveTooltip", "Define a custom response curve for this target."))
@@ -428,7 +428,7 @@ TSharedRef< SWidget > SPDD_TargetRow::GenerateWidgetForColumn(const FName& Colum
 					.VAlign(VAlign_Fill)
 					[
 						SNew(SBox)
-						.Visibility_Lambda([=]() { return IsCustomCurveEnabled() ? EVisibility::Visible : EVisibility::Collapsed;  })
+						.Visibility_Lambda([this]() { return IsCustomCurveEnabled() ? EVisibility::Visible : EVisibility::Collapsed;  })
 						.IsEnabled_Lambda([this]() { return IsOverrideEnabled(); })
 						.Content()
 						[
@@ -1091,7 +1091,7 @@ void FPoseDriverDetails::CustomizeDetails(IDetailLayoutBuilder& DetailBuilder)
 			.Padding(FMargin(6, 0, 3, 0))
 			[
 				SNew(SCheckBox)
-				.IsChecked_Lambda([=]() { return IsSoloDrivenOnly() ? ECheckBoxState::Checked : ECheckBoxState::Unchecked; })
+				.IsChecked_Lambda([this]() { return IsSoloDrivenOnly() ? ECheckBoxState::Checked : ECheckBoxState::Unchecked; })
 				.OnCheckStateChanged(this, &FPoseDriverDetails::OnSoloDrivenOnlyChanged)
 				.Padding(FMargin(4.0f, 0.0f))
 				.ToolTipText(LOCTEXT("SoloDrivenOnlyHelp", "Only solo the driven poses or curves and leave the source joint(s) in place."))
diff --git a/Engine/Source/Editor/CurveEditor/Private/Tree/CurveEditorTree.cpp b/Engine/Source/Editor/CurveEditor/Private/Tree/CurveEditorTree.cpp
--- a/Engine/Source/Editor/CurveEditor/Private/Tree/CurveEditorTree.cpp
+++ b/Engine/Source/Editor/CurveEditor/Private/Tree/CurveEditorTree.cpp
@@ -437,7 +437,7 @@ void FCurveEditorTree::SortTreeItems(FSortedCurveEditorTreeItems& TreeItemIDsToS
 	// If there is more than one item, sort the items and then repopulate the ChildIDs from the sorted list of items.
 	if (TreeItemIDsToSort.bRequiresSort && TreeItemsToSort.Num() > 1)
 	{
-		TreeItemsToSort.Sort([=](const FCurveEditorTreeItem& ItemA, const FCurveEditorTreeItem& ItemB) { return SortPredicate.Execute(ItemA.GetItem().Get(), ItemB.GetItem().Get()); });
+		TreeItemsToSort.Sort([this](const FCurveEditorTreeItem& ItemA, const FCurveEditorTreeItem& ItemB) { return SortPredicate.Execute(ItemA.GetItem().Get(), ItemB.GetItem().Get()); });
 		for (int32 i = 0; i < TreeItemsToSort.Num(); ++i)
 		{
 			TreeItemIDsToSort.ChildIDs[i] = TreeItemsToSort[i]->GetID();
diff --git a/Engine/Source/Editor/GraphEditor/Private/KismetPins/SGraphPinInteger.cpp b/Engine/Source/Editor/GraphEditor/Private/KismetPins/SGraphPinInteger.cpp
--- a/Engine/Source/Editor/GraphEditor/Private/KismetPins/SGraphPinInteger.cpp
+++ b/Engine/Source/Editor/GraphEditor/Private/KismetPins/SGraphPinInteger.cpp
@@ -163,7 +163,7 @@ TSharedRef<SWidget>	SGraphPinInteger::GetDefaultValueWidget()
 						FSlateIcon(),
 						FUIAction
 						(
-							FExecuteAction::CreateLambda([=]()
+							FExecuteAction::CreateLambda([this, BitmaskFlags, i]()
 							{
 								const FScopedTransaction Transaction(NSLOCTEXT("GraphEditor", "ChangePinValue", "Change Pin Value"));
 								const int32 CurValue = FCString::Atoi(*GraphPinObj->GetDefaultAsString());
@@ -171,7 +171,7 @@ TSharedRef<SWidget>	SGraphPinInteger::GetDefaultValueWidget()
 								GraphPinObj->GetSchema()->TrySetDefaultValue(*GraphPinObj, FString::FromInt(CurValue ^ BitmaskFlags[i].Value));
 							}),
 							FCanExecuteAction(),
-							FIsActionChecked::CreateLambda([=]() -> bool
+							FIsActionChecked::CreateLambda([this, BitmaskFlags, i]() -> bool
 							{
 								const int32 CurValue = FCString::Atoi(*GraphPinObj->GetDefaultAsString());
 								return (CurValue & BitmaskFlags[i].Value) != 0;
diff --git a/Engine/Source/Editor/InternationalizationSettings/Private/InternationalizationSettingsModelDetails.cpp b/Engine/Source/Editor/InternationalizationSettings/Private/InternationalizationSettingsModelDetails.cpp
--- a/Engine/Source/Editor/InternationalizationSettings/Private/InternationalizationSettingsModelDetails.cpp
+++ b/Engine/Source/Editor/InternationalizationSettings/Private/InternationalizationSettingsModelDetails.cpp
@@ -110,7 +110,7 @@ namespace
 		{
 			FCulturePtr EditorLanguage = FInternationalization::Get().GetCurrentLanguage();
 
-			const auto& OnSelectionChangedLambda = [=](FCulturePtr& SelectedCulture, ESelectInfo::Type SelectInfo)
+			const auto& OnSelectionChangedLambda = [this](FCulturePtr& SelectedCulture, ESelectInfo::Type SelectInfo)
 			{
 				if (SettingsModel.IsValid())
 				{
@@ -152,7 +152,7 @@ namespace
 				}
 			};
 
-			const auto& IsCulturePickableLambda = [=](FCulturePtr Culture) -> bool
+			const auto& IsCulturePickableLambda = [this](FCulturePtr Culture) -> bool
 			{
 				TArray<FString> CultureNames = Culture->GetPrioritizedParentCultureNames();
 				for (const FString& CultureName : CultureNames)
@@ -220,7 +220,7 @@ namespace
 		{
 			FCulturePtr EditorLocale = FInternationalization::Get().GetCurrentLocale();
 
-			const auto& OnSelectionChangedLambda = [=](FCulturePtr& SelectedCulture, ESelectInfo::Type SelectInfo)
+			const auto& OnSelectionChangedLambda = [this](FCulturePtr& SelectedCulture, ESelectInfo::Type SelectInfo)
 			{
 				if (SettingsModel.IsValid())
 				{
@@ -331,7 +331,7 @@ namespace
 				}
 			}
 
-			const auto& CulturePickerSelectLambda = [=](FCulturePtr& SelectedCulture, ESelectInfo::Type SelectInfo)
+			const auto& CulturePickerSelectLambda = [this](FCulturePtr& SelectedCulture, ESelectInfo::Type SelectInfo)
 			{
 				if (SettingsModel.IsValid())
 				{
@@ -348,7 +348,7 @@ namespace
 					PreviewGameCultureComboButton->SetIsOpen(false);
 				}
 			};
-			const auto& CulturePickerIsPickableLambda = [=](FCulturePtr Culture) -> bool
+			const auto& CulturePickerIsPickableLambda = [this](FCulturePtr Culture) -> bool
 			{
 				TArray<FString> CultureNames = Culture->GetPrioritizedParentCultureNames();
 				for (const FString& CultureName : CultureNames)
diff --git a/Engine/Source/Editor/KismetCompiler/Private/KismetCompiler.cpp b/Engine/Source/Editor/KismetCompiler/Private/KismetCompiler.cpp
--- a/Engine/Source/Editor/KismetCompiler/Private/KismetCompiler.cpp
+++ b/Engine/Source/Editor/KismetCompiler/Private/KismetCompiler.cpp
@@ -3759,7 +3759,7 @@ void FKismetCompilerContext::MergeUbergraphPagesIn(UEdGraph* Ubergraph)
 // Expands out nodes that need it
 void FKismetCompilerContext::ExpansionStep(UEdGraph* Graph, bool bAllowUbergraphExpansions)
 {
-	auto PruneInner = [=]()
+	auto PruneInner = [this, Graph]()
 	{
 		// Find the connected subgraph starting at the root node and prune out unused nodes
 		const bool bIncludePotentialRootNodes = true;
diff --git a/Engine/Source/Editor/MovieSceneTools/Private/K2Node_GetSequenceBinding.cpp b/Engine/Source/Editor/MovieSceneTools/Private/K2Node_GetSequenceBinding.cpp
--- a/Engine/Source/Editor/MovieSceneTools/Private/K2Node_GetSequenceBinding.cpp
+++ b/Engine/Source/Editor/MovieSceneTools/Private/K2Node_GetSequenceBinding.cpp
@@ -270,7 +270,7 @@ void UK2Node_GetSequenceBinding::GetNodeContextMenuActions(UToolMenu* Menu, UGra
 				"SetSequence",
 				LOCTEXT("SetSequence_Text", "Sequence"),
 				LOCTEXT("SetSequence_ToolTip", "Sets the sequence to get a binding from"),
-				FNewToolMenuDelegate::CreateLambda([=](UToolMenu* SubMenu)
+				FNewToolMenuDelegate::CreateLambda([this, Sequence](UToolMenu* SubMenu)
 				{
 					TArray<const UClass*> AllowedClasses({ UMovieSceneSequence::StaticClass() });
 
diff --git a/Engine/Source/Editor/MovieSceneTools/Private/Sections/ThumbnailSection.cpp b/Engine/Source/Editor/MovieSceneTools/Private/Sections/ThumbnailSection.cpp
--- a/Engine/Source/Editor/MovieSceneTools/Private/Sections/ThumbnailSection.cpp
+++ b/Engine/Source/Editor/MovieSceneTools/Private/Sections/ThumbnailSection.cpp
@@ -115,7 +115,7 @@ void FThumbnailSection::BuildSectionContextMenu(FMenuBuilder& MenuBuilder, const
 		MenuBuilder.AddSubMenu(
 			LOCTEXT("ThumbnailsMenu", "Thumbnails"),
 			FText(),
-			FNewMenuDelegate::CreateLambda([=](FMenuBuilder& InMenuBuilder){
+			FNewMenuDelegate::CreateLambda([this](FMenuBuilder& InMenuBuilder){
 
 				TSharedPtr<ISequencer> Sequencer = SequencerPtr.Pin();
 
@@ -134,7 +134,7 @@ void FThumbnailSection::BuildSectionContextMenu(FMenuBuilder& MenuBuilder, const
 						LOCTEXT("SetSingleTimeTooltip", "Defines the time at which this section should draw its single thumbnail to the current cursor position"),
 						FSlateIcon(),
 						FUIAction(
-							FExecuteAction::CreateLambda([=]{
+						FExecuteAction::CreateLambda([this, Sequencer]{
 								SetSingleTime(Sequencer->GetLocalTime().AsSeconds());
 								GetMutableDefault<UMovieSceneUserThumbnailSettings>()->bDrawSingleThumbnails = true;
 								GetMutableDefault<UMovieSceneUserThumbnailSettings>()->SaveConfig();
diff --git a/Engine/Source/Editor/MovieSceneTools/Private/TrackEditors/CameraShakeSourceShakeTrackEditor.cpp b/Engine/Source/Editor/MovieSceneTools/Private/TrackEditors/CameraShakeSourceShakeTrackEditor.cpp
--- a/Engine/Source/Editor/MovieSceneTools/Private/TrackEditors/CameraShakeSourceShakeTrackEditor.cpp
+++ b/Engine/Source/Editor/MovieSceneTools/Private/TrackEditors/CameraShakeSourceShakeTrackEditor.cpp
@@ -328,7 +328,7 @@ void FCameraShakeSourceShakeTrackEditor::AddCameraShakeSection(TArray<FGuid> Obj
 		}
 	}
 
-	auto OnAddShakeSourceShakeSection = [=](FFrameNumber Time) -> FKeyPropertyResult
+	auto OnAddShakeSourceShakeSection = [this, Objects](FFrameNumber Time) -> FKeyPropertyResult
 	{
 		return this->AddCameraShakeSectionKeyInternal(Time, Objects, true);
 	};
@@ -663,7 +663,7 @@ void FCameraShakeSourceShakeTrackEditor::AddCameraShakeTriggerTrack(const TArray
 		}
 	}
 
-	auto OnAddShakeSourceShakeSection = [=](FFrameNumber Time) -> FKeyPropertyResult
+	auto OnAddShakeSourceShakeSection = [this, Objects](FFrameNumber Time) -> FKeyPropertyResult
 	{
 		return this->AddCameraShakeTriggerTrackInternal(Time, Objects, nullptr);
 	};
diff --git a/Engine/Source/Editor/MovieSceneTools/Private/TrackEditors/SkeletalAnimationTrackEditor.cpp b/Engine/Source/Editor/MovieSceneTools/Private/TrackEditors/SkeletalAnimationTrackEditor.cpp
--- a/Engine/Source/Editor/MovieSceneTools/Private/TrackEditors/SkeletalAnimationTrackEditor.cpp
+++ b/Engine/Source/Editor/MovieSceneTools/Private/TrackEditors/SkeletalAnimationTrackEditor.cpp
@@ -788,7 +788,7 @@ void FSkeletalAnimationSection::BuildSectionContextMenu(FMenuBuilder& MenuBuilde
 		auto MatchToBone = [=](bool bMatchPrevious, int32 Index)
 		{
 			return FUIAction(
-				FExecuteAction::CreateLambda([=]
+				FExecuteAction::CreateLambda([=, this]
 					{
 						FScopedTransaction MatchSection(LOCTEXT("MatchSectionByBone_Transaction", "Match Section By Bone"));
 						Section.Modify();	
@@ -810,7 +810,7 @@ void FSkeletalAnimationSection::BuildSectionContextMenu(FMenuBuilder& MenuBuilde
 					{ 
 						return SequencerPtr.IsValid(); 
 					}),
-				FIsActionChecked::CreateLambda([=]()->bool
+				FIsActionChecked::CreateLambda([=, this]()->bool
 					{
 						if (Index >= 0)
 						{
@@ -874,14 +874,14 @@ void FSkeletalAnimationSection::BuildSectionContextMenu(FMenuBuilder& MenuBuilde
 				LOCTEXT("MatchTranslationTooltip", "Match the Translation to the Specified Bone"),
 				FSlateIcon(),
 				FUIAction(
-					FExecuteAction::CreateLambda([=]()->void {
+					FExecuteAction::CreateLambda([=, this]()->void {
 						FScopedTransaction MatchTransaction(LOCTEXT("MatchTranslation_Transaction", "Match Translation"));
 						Section.Modify();	
 						Section.ToggleMatchTranslation();
 						Section.MatchSectionByBoneTransform(SkelMeshComp, SequencerPtr->GetLocalTime().Time, SequencerPtr->GetLocalTime().Rate, Section.MatchedBoneName);
 						SequencerPtr->NotifyMovieSceneDataChanged(EMovieSceneDataChangeType::TrackValueChanged);}),
-					FCanExecuteAction::CreateLambda([=]()->bool { return true; }),
-					FIsActionChecked::CreateLambda([=]()->bool { return Section.bMatchTranslation; })),
+						FCanExecuteAction::CreateLambda([]()->bool { return true; }),
+						FIsActionChecked::CreateLambda([this]()->bool { return Section.bMatchTranslation; })),
 				NAME_None,
 				EUserInterfaceActionType::ToggleButton
 			);
@@ -891,14 +891,14 @@ void FSkeletalAnimationSection::BuildSectionContextMenu(FMenuBuilder& MenuBuilde
 				LOCTEXT("MatchZHeightTooltip", "Match the Z Height, may want this off for better matching"),
 				FSlateIcon(),
 				FUIAction(
-					FExecuteAction::CreateLambda([=]()->void { 
+					FExecuteAction::CreateLambda([=, this]()->void { 
 						FScopedTransaction MatchTransaction(LOCTEXT("MatchZHeight_Transaction", "Match Z Height"));
 						Section.Modify();
 						Section.ToggleMatchIncludeZHeight(); 
 						Section.MatchSectionByBoneTransform(SkelMeshComp, SequencerPtr->GetLocalTime().Time, SequencerPtr->GetLocalTime().Rate, Section.MatchedBoneName);
 						SequencerPtr->NotifyMovieSceneDataChanged(EMovieSceneDataChangeType::TrackValueChanged); }),
-					FCanExecuteAction::CreateLambda([=]()->bool { return true; }),
-					FIsActionChecked::CreateLambda([=]()->bool { return Section.bMatchIncludeZHeight;  })),
+						FCanExecuteAction::CreateLambda([]()->bool { return true; }),
+						FIsActionChecked::CreateLambda([this]()->bool { return Section.bMatchIncludeZHeight;  })),
 				NAME_None,
 				EUserInterfaceActionType::ToggleButton
 			);
@@ -908,14 +908,14 @@ void FSkeletalAnimationSection::BuildSectionContextMenu(FMenuBuilder& MenuBuilde
 				LOCTEXT("MatchYawRotationTooltip", "Match the Yaw Rotation, may want this off for better matching"),
 				FSlateIcon(),
 				FUIAction(
-					FExecuteAction::CreateLambda([=]()->void {
+					FExecuteAction::CreateLambda([=, this]()->void {
 						FScopedTransaction MatchTransaction(LOCTEXT("MatchYawRotation_Transaction", "Match Yaw Rotation"));
 						Section.Modify();
 						Section.ToggleMatchIncludeYawRotation();
 						Section.MatchSectionByBoneTransform(SkelMeshComp, SequencerPtr->GetLocalTime().Time, SequencerPtr->GetLocalTime().Rate, Section.MatchedBoneName);
 						SequencerPtr->NotifyMovieSceneDataChanged(EMovieSceneDataChangeType::TrackValueChanged); }),
-					FCanExecuteAction::CreateLambda([=]()->bool { return true; }),
-							FIsActionChecked::CreateLambda([=]()->bool { return Section.bMatchRotationYaw; })),
+						FCanExecuteAction::CreateLambda([]()->bool { return true; }),
+						FIsActionChecked::CreateLambda([this]()->bool { return Section.bMatchRotationYaw; })),
 				NAME_None,
 							EUserInterfaceActionType::ToggleButton
 							);
@@ -931,8 +931,8 @@ void FSkeletalAnimationSection::BuildSectionContextMenu(FMenuBuilder& MenuBuilde
 						Section.ToggleMatchIncludePitchRotation();
 						Section.MatchSectionByBoneTransform(SkelMeshComp, SequencerPtr->GetLocalTime().Time, SequencerPtr->GetLocalTime().Rate, Section.MatchedBoneName);
 						SequencerPtr->NotifyMovieSceneDataChanged(EMovieSceneDataChangeType::TrackValueChanged); }),
-					FCanExecuteAction::CreateLambda([=]()->bool { return true; }),
-					FIsActionChecked::CreateLambda([=]()->bool { return Section.bMatchRotationPitch;})),
+						FCanExecuteAction::CreateLambda([]()->bool { return true; }),
+						FIsActionChecked::CreateLambda([this]()->bool { return Section.bMatchRotationPitch;})),
 				NAME_None,
 				EUserInterfaceActionType::ToggleButton
 			);
@@ -942,14 +942,14 @@ void FSkeletalAnimationSection::BuildSectionContextMenu(FMenuBuilder& MenuBuilde
 				LOCTEXT("MatchRollRotationTooltip", "Match the Roll Rotation, may want this off for better matching"),
 				FSlateIcon(),
 				FUIAction(
-					FExecuteAction::CreateLambda([=]()->void {
+					FExecuteAction::CreateLambda([=, this]()->void {
 						FScopedTransaction MatchTransaction(LOCTEXT("MatchRollRotation_Transaction", "Match Roll Rotation"));
 						Section.Modify();
 						Section.ToggleMatchIncludeRollRotation();
 						Section.MatchSectionByBoneTransform(SkelMeshComp, SequencerPtr->GetLocalTime().Time, SequencerPtr->GetLocalTime().Rate, Section.MatchedBoneName);
 						SequencerPtr->NotifyMovieSceneDataChanged(EMovieSceneDataChangeType::TrackValueChanged); }),
-					FCanExecuteAction::CreateLambda([=]()->bool { return true; }),
-							FIsActionChecked::CreateLambda([=]()->bool { return Section.bMatchRotationRoll; })),
+						FCanExecuteAction::CreateLambda([]()->bool { return true; }),
+						FIsActionChecked::CreateLambda([this]()->bool { return Section.bMatchRotationRoll; })),
 				NAME_None,
 							EUserInterfaceActionType::ToggleButton
 				);
@@ -965,13 +965,13 @@ void FSkeletalAnimationSection::BuildSectionContextMenu(FMenuBuilder& MenuBuilde
 				NSLOCTEXT("Sequencer", "ShowSkeletonsTooltip", "Show A Skeleton for this Section."),
 				FSlateIcon(),
 				FUIAction(
-					FExecuteAction::CreateLambda([=]()->void {
+					FExecuteAction::CreateLambda([this, SequencerPtr]()->void {
 						Section.ToggleShowSkeleton();
 						SequencerPtr->NotifyMovieSceneDataChanged(EMovieSceneDataChangeType::TrackValueChanged);
 
 						}),
 					FCanExecuteAction::CreateLambda([=]()->bool { return SequencerPtr != nullptr; }),
-						FIsActionChecked::CreateLambda([=]()->bool { return Section.bShowSkeleton; })),
+						FIsActionChecked::CreateLambda([this]()->bool { return Section.bShowSkeleton; })),
 						NAME_None,
 						EUserInterfaceActionType::ToggleButton
 					);
diff --git a/Engine/Source/Editor/MovieSceneTools/Private/TrackEditors/SubTrackEditor.cpp b/Engine/Source/Editor/MovieSceneTools/Private/TrackEditors/SubTrackEditor.cpp
--- a/Engine/Source/Editor/MovieSceneTools/Private/TrackEditors/SubTrackEditor.cpp
+++ b/Engine/Source/Editor/MovieSceneTools/Private/TrackEditors/SubTrackEditor.cpp
@@ -84,7 +84,7 @@ public:
 			MenuBuilder.AddSubMenu(
 				LOCTEXT("TakesMenu", "Takes"),
 				LOCTEXT("TakesMenuTooltip", "Subsequence takes"),
-				FNewMenuDelegate::CreateLambda([=](FMenuBuilder& InMenuBuilder) { SubTrackEditor.Pin()->AddTakesMenu(Section, InMenuBuilder); }));
+				FNewMenuDelegate::CreateLambda([this, Section](FMenuBuilder& InMenuBuilder) { SubTrackEditor.Pin()->AddTakesMenu(Section, InMenuBuilder); }));
 
 			MenuBuilder.AddMenuEntry(
 				LOCTEXT("NewTake", "New Take"),
diff --git a/Engine/Source/Editor/MovieSceneTools/Private/TrackEditors/TransformTrackEditor.cpp b/Engine/Source/Editor/MovieSceneTools/Private/TrackEditors/TransformTrackEditor.cpp
--- a/Engine/Source/Editor/MovieSceneTools/Private/TrackEditors/TransformTrackEditor.cpp
+++ b/Engine/Source/Editor/MovieSceneTools/Private/TrackEditors/TransformTrackEditor.cpp
@@ -881,11 +881,11 @@ void F3DTransformTrackEditor::AddTransformKeys( UObject* ObjectToKey, const TOpt
 	{
 		NewTrack->SetPropertyNameAndPath(TransformPropertyName, TransformPropertyName.ToString());
 	};
-	auto GenerateKeys = [=](UMovieSceneSection* Section, FGeneratedTrackKeys& GeneratedKeys)
+	auto GenerateKeys = [=, this](UMovieSceneSection* Section, FGeneratedTrackKeys& GeneratedKeys)
 	{
 		this->GetTransformKeys(LastTransform, CurrentTransform, ChannelsToKey, ObjectToKey, Section, GeneratedKeys);
 	};
-	auto OnKeyProperty = [=](FFrameNumber Time) -> FKeyPropertyResult
+	auto OnKeyProperty = [=, this](FFrameNumber Time) -> FKeyPropertyResult
 	{
 		FKeyPropertyResult KeyPropertyResult = this->AddKeysToObjects(MakeArrayView(&ObjectToKey, 1), Time,  KeyMode, UMovieScene3DTransformTrack::StaticClass(), TransformPropertyName, InitializeNewTrack, GenerateKeys);
 		if (KeyPropertyResult.SectionsKeyed.Num() > 0)
diff --git a/Engine/Source/Editor/Persona/Private/SAnimViewportToolBar.cpp b/Engine/Source/Editor/Persona/Private/SAnimViewportToolBar.cpp
--- a/Engine/Source/Editor/Persona/Private/SAnimViewportToolBar.cpp
+++ b/Engine/Source/Editor/Persona/Private/SAnimViewportToolBar.cpp
@@ -618,7 +618,7 @@ TSharedRef<SWidget> SAnimViewportToolBar::GenerateViewMenu() const
 					InSubMenuBuilder.AddSubMenu(
 						LOCTEXT("CameraFollowBone_DisplayName", "Orbit Bone"),
 						LOCTEXT("CameraFollowBone_ToolTip", "Select a bone for the camera to follow and orbit around"),
-						FNewMenuDelegate::CreateLambda([=](FMenuBuilder& SubMenuBuilder)
+						FNewMenuDelegate::CreateLambda([this](FMenuBuilder& SubMenuBuilder)
 						{
 							SubMenuBuilder.BeginSection("CameraFollowModeBoneSubmenu", LOCTEXT("CameraFollowModeBoneSubmenu_Label", "Follow Bone Options"));
 							SubMenuBuilder.AddWidget(MakeFollowBoneWidget(), FText());
diff --git a/Engine/Source/Editor/Persona/Private/SAnimationBlendSpaceGridWidget.cpp b/Engine/Source/Editor/Persona/Private/SAnimationBlendSpaceGridWidget.cpp
--- a/Engine/Source/Editor/Persona/Private/SAnimationBlendSpaceGridWidget.cpp
+++ b/Engine/Source/Editor/Persona/Private/SAnimationBlendSpaceGridWidget.cpp
@@ -497,7 +497,7 @@ TSharedPtr<SWidget> SBlendSpaceGridWidget::CreateGridEntryBox(const int32 BoxInd
 		[
 			SNew(STextBlock)
 			.Visibility(bShowLabel ? EVisibility::Visible : EVisibility::Collapsed)
-			.Text_Lambda([=]() { return (BoxIndex == 0) ? ParameterXName : ParameterYName; })
+			.Text_Lambda([this, BoxIndex]() { return (BoxIndex == 0) ? ParameterXName : ParameterYName; })
 		];
 }
 
diff --git a/Engine/Source/Editor/Sequencer/Private/LevelEditorSequencerIntegration.cpp b/Engine/Source/Editor/Sequencer/Private/LevelEditorSequencerIntegration.cpp
--- a/Engine/Source/Editor/Sequencer/Private/LevelEditorSequencerIntegration.cpp
+++ b/Engine/Source/Editor/Sequencer/Private/LevelEditorSequencerIntegration.cpp
@@ -711,7 +711,7 @@ void FLevelEditorSequencerIntegration::OnEndPlayMap()
 	}
 	if (bAddRestoreCallback)
 	{
-		AcquiredResources.Add([=] { this->RestoreRealtimeViewports(); });
+		AcquiredResources.Add([this] { this->RestoreRealtimeViewports(); });
 	}
 
 	IterateAllSequencers(
@@ -1163,7 +1163,7 @@ void FLevelEditorSequencerIntegration::ActivateRealtimeViewports()
 		}
 	}
 
-	AcquiredResources.Add([=]{ this->RestoreRealtimeViewports(); });
+	AcquiredResources.Add([this]{ this->RestoreRealtimeViewports(); });
 }
 
 void FLevelEditorSequencerIntegration::RestoreRealtimeViewports()
diff --git a/Engine/Source/Editor/Sequencer/Private/MVVM/ViewModels/TrackModel.cpp b/Engine/Source/Editor/Sequencer/Private/MVVM/ViewModels/TrackModel.cpp
--- a/Engine/Source/Editor/Sequencer/Private/MVVM/ViewModels/TrackModel.cpp
+++ b/Engine/Source/Editor/Sequencer/Private/MVVM/ViewModels/TrackModel.cpp
@@ -714,7 +714,7 @@ void FTrackModel::BuildContextMenu(FMenuBuilder& MenuBuilder)
 		MenuBuilder.AddSubMenu(
 			TrackSections.Num() > 1 ? LOCTEXT("BatchEditSections", "Batch Edit Sections") : LOCTEXT("EditSection", "Edit Section"),
 			FText(),
-			FNewMenuDelegate::CreateLambda([=](FMenuBuilder& SubMenuBuilder){
+			FNewMenuDelegate::CreateLambda([this, TrackSections](FMenuBuilder& SubMenuBuilder){
 				FSequencer* Sequencer = static_cast<FSequencer*>(GetEditor()->GetSequencer().Get());
 				SequencerHelpers::AddPropertiesMenu(*Sequencer, SubMenuBuilder, TrackSections);
 			})
diff --git a/Engine/Source/Editor/Sequencer/Private/SSequencer.cpp b/Engine/Source/Editor/Sequencer/Private/SSequencer.cpp
--- a/Engine/Source/Editor/Sequencer/Private/SSequencer.cpp
+++ b/Engine/Source/Editor/Sequencer/Private/SSequencer.cpp
@@ -159,7 +159,7 @@ void SSequencer::Construct(const FArguments& InArgs, TSharedRef<FSequencer> InSe
 
 	// Get the desired display format from the user's settings each time.
 	TAttribute<EFrameNumberDisplayFormats> GetDisplayFormatAttr = MakeAttributeLambda(
-		[=]
+		[this]
 		{
 			if (USequencerSettings* Settings = GetSequencerSettings())
 			{
@@ -171,7 +171,7 @@ void SSequencer::Construct(const FArguments& InArgs, TSharedRef<FSequencer> InSe
 
 	// Get the number of zero pad frames from the user's settings as well.
 	TAttribute<uint8> GetZeroPadFramesAttr = MakeAttributeLambda(
-		[=]()->uint8
+		[this]()->uint8
 		{
 			if (USequencerSettings* Settings = GetSequencerSettings())
 			{
@@ -1149,7 +1149,7 @@ TSharedRef<SWidget> SSequencer::MakeAddButton()
 	.OnGetMenuContent(this, &SSequencer::MakeAddMenu)
 	.Icon(FAppStyle::Get().GetBrush("Icons.Plus"))
 	.Text(LOCTEXT("Track", "Track"))
-	.IsEnabled_Lambda([=]() { return !SequencerPtr.Pin()->IsReadOnly(); });
+	.IsEnabled_Lambda([this]() { return !SequencerPtr.Pin()->IsReadOnly(); });
 }
 
 TSharedRef<SWidget> SSequencer::MakeFilterButton()
@@ -2107,7 +2107,7 @@ TSharedRef<SWidget> SSequencer::MakeViewMenu()
 	MenuBuilder.AddMenuEntry( FSequencerCommands::Get().ToggleShowPreAndPostRoll );
 
 	// Menu entry for zero padding
-	auto OnZeroPadChanged = [=](uint8 NewValue) {
+	auto OnZeroPadChanged = [this](uint8 NewValue) {
 		GetSequencerSettings()->SetZeroPadFrames(NewValue);
 	};
 
@@ -2126,7 +2126,7 @@ TSharedRef<SWidget> SSequencer::MakeViewMenu()
 			.OnValueChanged_Lambda(OnZeroPadChanged)
 			.MinValue(0)
 			.MaxValue(8)
-			.Value_Lambda([=]() -> uint8 {
+			.Value_Lambda([this]() -> uint8 {
 			return GetSequencerSettings()->GetZeroPadFrames();
 		})
 		],
@@ -2309,7 +2309,7 @@ void SSequencer::FillTimeDisplayFormatMenu(FMenuBuilder& MenuBuilder)
 					FUIAction(
 						FExecuteAction::CreateUObject(GetSequencerSettings(), &USequencerSettings::SetTimeDisplayFormat, Value),
 						FCanExecuteAction(),
-						FIsActionChecked::CreateLambda([=] { return GetSequencerSettings()->GetTimeDisplayFormat() == Value; })
+						FIsActionChecked::CreateLambda([this, Value] { return GetSequencerSettings()->GetTimeDisplayFormat() == Value; })
 					),
 					NAME_None,
 					EUserInterfaceActionType::RadioButton
@@ -2327,7 +2327,7 @@ TSharedRef<SWidget> SSequencer::MakePlaybackMenu()
 	MenuBuilder.BeginSection("PlaybackThisSequence", LOCTEXT("PlaybackThisSequenceHeader", "Playback - This Sequence"));
 	{
 		// Menu entry for the start position
-		auto OnStartChanged = [=](double NewValue){
+		auto OnStartChanged = [this](double NewValue){
 
 			FFrameNumber ValueAsFrame = FFrameTime::FromDecimal(NewValue).GetFrame();
 			FFrameNumber PlayStart = ValueAsFrame;
@@ -2369,17 +2369,17 @@ TSharedRef<SWidget> SSequencer::MakePlaybackMenu()
 				[
 					SNew(SSpinBox<double>)
 						.TypeInterface(NumericTypeInterface)
-						.IsEnabled_Lambda([=]() {
+						.IsEnabled_Lambda([this]() {
 							return !SequencerPtr.Pin()->IsPlaybackRangeLocked();
 						})
 						.Style(&FAppStyle::GetWidgetStyle<FSpinBoxStyle>("Sequencer.HyperlinkSpinBox"))
 						.OnValueCommitted_Lambda([=](double Value, ETextCommit::Type){ OnStartChanged(Value); })
 						.OnValueChanged_Lambda([=](double Value) { OnStartChanged(Value); })
 						.OnBeginSliderMovement(OnPlaybackRangeBeginDrag)
-						.OnEndSliderMovement_Lambda([=](double Value){ OnStartChanged(Value); OnPlaybackRangeEndDrag.ExecuteIfBound(); })
+						.OnEndSliderMovement_Lambda([=, this](double Value){ OnStartChanged(Value); OnPlaybackRangeEndDrag.ExecuteIfBound(); })
 						.MinValue(TOptional<double>())
 						.MaxValue(TOptional<double>())
-						.Value_Lambda([=]() -> double {
+						.Value_Lambda([this]() -> double {
 							return SequencerPtr.Pin()->GetPlaybackRange().GetLowerBoundValue().Value;
 						})
 						.Delta(this, &SSequencer::GetSpinboxDelta)
@@ -2388,7 +2388,7 @@ TSharedRef<SWidget> SSequencer::MakePlaybackMenu()
 			LOCTEXT("PlaybackStartLabel", "Start"));
 
 		// Menu entry for the end position
-		auto OnEndChanged = [=](double NewValue) {
+		auto OnEndChanged = [this](double NewValue) {
 
 			FFrameNumber ValueAsFrame = FFrameTime::FromDecimal(NewValue).GetFrame();
 			FFrameNumber PlayStart = UE::MovieScene::DiscreteInclusiveLower(SequencerPtr.Pin()->GetPlaybackRange());
@@ -2430,17 +2430,17 @@ TSharedRef<SWidget> SSequencer::MakePlaybackMenu()
 				[
 					SNew(SSpinBox<double>)
 						.TypeInterface(NumericTypeInterface)
-						.IsEnabled_Lambda([=]() {
+						.IsEnabled_Lambda([this]() {
 					 		return !SequencerPtr.Pin()->IsPlaybackRangeLocked();
 						})
 						.Style(&FAppStyle::GetWidgetStyle<FSpinBoxStyle>("Sequencer.HyperlinkSpinBox"))
 						.OnValueCommitted_Lambda([=](double Value, ETextCommit::Type){ OnEndChanged(Value); })
 						.OnValueChanged_Lambda([=](double Value) { OnEndChanged(Value); })
 						.OnBeginSliderMovement(OnPlaybackRangeBeginDrag)
-						.OnEndSliderMovement_Lambda([=](double Value){ OnEndChanged(Value); OnPlaybackRangeEndDrag.ExecuteIfBound(); })
+						.OnEndSliderMovement_Lambda([=, this](double Value){ OnEndChanged(Value); OnPlaybackRangeEndDrag.ExecuteIfBound(); })
 						.MinValue(TOptional<double>())
 						.MaxValue(TOptional<double>())
-						.Value_Lambda([=]() -> double {
+						.Value_Lambda([this]() -> double {
 					 		return SequencerPtr.Pin()->GetPlaybackRange().GetUpperBoundValue().Value;
 						})
 						.Delta(this, &SSequencer::GetSpinboxDelta)
@@ -2484,7 +2484,7 @@ TSharedRef<SWidget> SSequencer::MakePlaybackMenu()
 	MenuBuilder.EndSection();
 
 	// Menu entry for the jump frame increment
-	auto OnJumpFrameIncrementChanged = [=](double NewValue) {
+	auto OnJumpFrameIncrementChanged = [this](double NewValue) {
 		FFrameRate TickResolution = SequencerPtr.Pin()->GetFocusedTickResolution();
 		FFrameRate DisplayRate = SequencerPtr.Pin()->GetFocusedDisplayRate();
 		FFrameNumber JumpFrameIncrement = FFrameRate::TransformTime(FFrameTime::FromDecimal(NewValue), TickResolution, DisplayRate).CeilToFrame();						
@@ -2507,7 +2507,7 @@ TSharedRef<SWidget> SSequencer::MakePlaybackMenu()
 					.OnValueChanged_Lambda([=](double Value) { OnJumpFrameIncrementChanged(Value); })
 					.MinValue(TOptional<double>())
 					.MaxValue(TOptional<double>())
-					.Value_Lambda([=]() -> double {
+					.Value_Lambda([this]() -> double {
 						FFrameNumber JumpFrameIncrement = SequencerPtr.Pin()->GetSequencerSettings()->GetJumpFrameIncrement();
 						FFrameRate TickResolution = SequencerPtr.Pin()->GetFocusedTickResolution();
 						FFrameRate DisplayRate = SequencerPtr.Pin()->GetFocusedDisplayRate();
diff --git a/Engine/Source/Editor/Sequencer/Private/SSequencerPlayRateCombo.cpp b/Engine/Source/Editor/Sequencer/Private/SSequencerPlayRateCombo.cpp
--- a/Engine/Source/Editor/Sequencer/Private/SSequencerPlayRateCombo.cpp
+++ b/Engine/Source/Editor/Sequencer/Private/SSequencerPlayRateCombo.cpp
@@ -279,7 +279,7 @@ TSharedRef<SWidget> SSequencerPlayRateCombo::OnCreateMenu()
 				SNew(SFrameRateEntryBox)
 				.Value(this, &SSequencerPlayRateCombo::GetDisplayRate)
 				.OnValueChanged(this, &SSequencerPlayRateCombo::SetDisplayRate)
-				.IsEnabled_Lambda([=] { return !GetIsSequenceReadOnly(); })
+				.IsEnabled_Lambda([this] { return !GetIsSequenceReadOnly(); })
 			],
 			LOCTEXT("CustomFramerateDisplayLabel", "Custom")
 		);
diff --git a/Engine/Source/Editor/Sequencer/Private/SequencerContextMenus.cpp b/Engine/Source/Editor/Sequencer/Private/SequencerContextMenus.cpp
--- a/Engine/Source/Editor/Sequencer/Private/SequencerContextMenus.cpp
+++ b/Engine/Source/Editor/Sequencer/Private/SequencerContextMenus.cpp
@@ -154,7 +154,7 @@ void FKeyContextMenu::PopulateMenu(FMenuBuilder& MenuBuilder, TSharedPtr<FExtend
 			FUIAction (
 				FExecuteAction(),
 				// @todo sequencer: only one struct per structure view supported right now :/
-				FCanExecuteAction::CreateLambda([=]{ return KeyStruct.IsValid(); })
+				FCanExecuteAction::CreateLambda([this]{ return KeyStruct.IsValid(); })
 			),
 			NAME_None,
 			EUserInterfaceActionType::Button
@@ -261,7 +261,7 @@ void FSectionContextMenu::PopulateMenu(FMenuBuilder& MenuBuilder, TSharedPtr<FEx
 	MenuBuilder.AddSubMenu(
 		LOCTEXT("SectionProperties", "Properties"),
 		LOCTEXT("SectionPropertiesTooltip", "Modify the section properties"),
-		FNewMenuDelegate::CreateLambda([=](FMenuBuilder& SubMenuBuilder)
+		FNewMenuDelegate::CreateLambda([this](FMenuBuilder& SubMenuBuilder)
 		{
 			TArray<TWeakObjectPtr<UObject>> Sections;
 			for (TViewModelPtr<FSectionModel> SectionModel : Sequencer->GetViewModel()->GetSelection()->TrackArea.Filter<FSectionModel>())
@@ -452,9 +452,9 @@ void FSectionContextMenu::AddEditMenu(FMenuBuilder& MenuBuilder)
 		LOCTEXT("DeleteKeysWhenTrimmingTooltip", "Delete keys outside of the trimmed range"),
 		FSlateIcon(),
 		FUIAction(
-			FExecuteAction::CreateLambda([=] { Sequencer->GetSequencerSettings()->SetDeleteKeysWhenTrimming(!Sequencer->GetSequencerSettings()->GetDeleteKeysWhenTrimming()); }),
+			FExecuteAction::CreateLambda([this] { Sequencer->GetSequencerSettings()->SetDeleteKeysWhenTrimming(!Sequencer->GetSequencerSettings()->GetDeleteKeysWhenTrimming()); }),
 			FCanExecuteAction(),
-			FIsActionChecked::CreateLambda([=] { return Sequencer->GetSequencerSettings()->GetDeleteKeysWhenTrimming(); })),
+			FIsActionChecked::CreateLambda([this] { return Sequencer->GetSequencerSettings()->GetDeleteKeysWhenTrimming(); })),
 		NAME_None,
 		EUserInterfaceActionType::ToggleButton
 	);
@@ -541,7 +541,7 @@ void FSectionContextMenu::AddEditMenu(FMenuBuilder& MenuBuilder)
 			FCanExecuteAction::CreateLambda([=]{ return Shared->CanReduceKeys(); }))
 	);
 
-	auto OnReduceKeysToleranceChanged = [=](float NewValue) {
+	auto OnReduceKeysToleranceChanged = [this](float NewValue) {
 		Sequencer->GetSequencerSettings()->SetReduceKeysTolerance(NewValue);
 	};
 
@@ -560,7 +560,7 @@ void FSectionContextMenu::AddEditMenu(FMenuBuilder& MenuBuilder)
 				.OnValueChanged_Lambda(OnReduceKeysToleranceChanged)
 				.MinValue(0)
 				.MaxValue(TOptional<float>())
-				.Value_Lambda([=]() -> float {
+				.Value_Lambda([this]() -> float {
 				return Sequencer->GetSequencerSettings()->GetReduceKeysTolerance();
 				})
 			],
@@ -1766,7 +1766,7 @@ void FEasingContextMenu::PopulateMenu(FMenuBuilder& MenuBuilder, TSharedPtr<FExt
 					.MaxValue(TOptional<double>())
 					.MaxSliderValue(TOptional<double>())
 					.MinSliderValue(0.f)
-					.Delta_Lambda([=]() -> double { return Sequencer->GetDisplayRateDeltaFrameCount(); })
+					.Delta_Lambda([this]() -> double { return Sequencer->GetDisplayRateDeltaFrameCount(); })
 					.Value_Lambda([=] {
 						TOptional<int32> Current = Shared->GetCurrentLength();
 						if (Current.IsSet())
diff --git a/Engine/Source/Editor/Sequencer/Private/SequencerTimeSliderController.cpp b/Engine/Source/Editor/Sequencer/Private/SequencerTimeSliderController.cpp
--- a/Engine/Source/Editor/Sequencer/Private/SequencerTimeSliderController.cpp
+++ b/Engine/Source/Editor/Sequencer/Private/SequencerTimeSliderController.cpp
@@ -347,7 +347,7 @@ int32 FSequencerTimeSliderController::DrawMarkedFrames( const FGeometry& Allotte
 	//const FLinearColor SelectedColor = FLinearColor::White; //FAppStyle::GetSlateColor("SelectionColor").GetColor(InWidgetStyle);
 	const FLinearColor WhiteColorHSV = FLinearColor::White.LinearRGBToHSV();
 
-	auto DrawFrameMarkers = ([=](const TArray<FMovieSceneMarkedFrame> & InMarkedFrames, FSlateWindowElementList& DrawElements, bool bIsGlobal)
+	auto DrawFrameMarkers = ([=, this](const TArray<FMovieSceneMarkedFrame> & InMarkedFrames, FSlateWindowElementList& DrawElements, bool bIsGlobal)
 	{
 		for (int32 MarkIndex = 0; MarkIndex < InMarkedFrames.Num(); ++MarkIndex)
 		{
@@ -1558,8 +1558,8 @@ TSharedRef<SWidget> FSequencerTimeSliderController::OpenSetPlaybackRangeMenu(con
 			FText(),
 			FSlateIcon(),
 			FUIAction(
-				FExecuteAction::CreateLambda([=]{ SetPlaybackRangeStart(FrameNumber); }),
-				FCanExecuteAction::CreateLambda([=]{ return !TimeSliderArgs.IsPlaybackRangeLocked.Get() && FrameNumber < UE::MovieScene::DiscreteExclusiveUpper(PlaybackRange); })
+				FExecuteAction::CreateLambda([this, FrameNumber]{ SetPlaybackRangeStart(FrameNumber); }),
+				FCanExecuteAction::CreateLambda([this, FrameNumber, PlaybackRange]{ return !TimeSliderArgs.IsPlaybackRangeLocked.Get() && FrameNumber < UE::MovieScene::DiscreteExclusiveUpper(PlaybackRange); })
 			)
 		);
 
@@ -1568,8 +1568,8 @@ TSharedRef<SWidget> FSequencerTimeSliderController::OpenSetPlaybackRangeMenu(con
 			FText(),
 			FSlateIcon(),
 			FUIAction(
-				FExecuteAction::CreateLambda([=]{ SetPlaybackRangeEnd(FrameNumber); }),
-				FCanExecuteAction::CreateLambda([=]{ return !TimeSliderArgs.IsPlaybackRangeLocked.Get() && FrameNumber >= UE::MovieScene::DiscreteInclusiveLower(PlaybackRange); })
+				FExecuteAction::CreateLambda([this, FrameNumber]{ SetPlaybackRangeEnd(FrameNumber); }),
+				FCanExecuteAction::CreateLambda([this, FrameNumber, PlaybackRange]{ return !TimeSliderArgs.IsPlaybackRangeLocked.Get() && FrameNumber >= UE::MovieScene::DiscreteInclusiveLower(PlaybackRange); })
 			)
 		);
 
@@ -1578,9 +1578,9 @@ TSharedRef<SWidget> FSequencerTimeSliderController::OpenSetPlaybackRangeMenu(con
 			LOCTEXT("ToggleLockedTooltip", "Lock/Unlock the playback range"),
 			FSlateIcon(),
 			FUIAction(
-				FExecuteAction::CreateLambda([=] { TimeSliderArgs.OnTogglePlaybackRangeLocked.ExecuteIfBound(); }),
-				FCanExecuteAction::CreateLambda([=]{ return !bReadOnly; }),
-				FIsActionChecked::CreateLambda([=] { return TimeSliderArgs.IsPlaybackRangeLocked.Get(); })
+				FExecuteAction::CreateLambda([this] { TimeSliderArgs.OnTogglePlaybackRangeLocked.ExecuteIfBound(); }),
+				FCanExecuteAction::CreateLambda([bReadOnly]{ return !bReadOnly; }),
+				FIsActionChecked::CreateLambda([this] { return TimeSliderArgs.IsPlaybackRangeLocked.Get(); })
 			),
 			NAME_None,
 			EUserInterfaceActionType::ToggleButton
@@ -1596,7 +1596,7 @@ TSharedRef<SWidget> FSequencerTimeSliderController::OpenSetPlaybackRangeMenu(con
 			FText(),
 			FSlateIcon(),
 			FUIAction(
-				FExecuteAction::CreateLambda([=]{ SetSelectionRangeStart(FrameNumber); }),
+				FExecuteAction::CreateLambda([this, FrameNumber]{ SetSelectionRangeStart(FrameNumber); }),
 				FCanExecuteAction::CreateLambda([=]{ return SelectionRange.IsEmpty() || FrameNumber < UE::MovieScene::DiscreteExclusiveUpper(SelectionRange); })
 			)
 		);
@@ -1606,7 +1606,7 @@ TSharedRef<SWidget> FSequencerTimeSliderController::OpenSetPlaybackRangeMenu(con
 			FText(),
 			FSlateIcon(),
 			FUIAction(
-				FExecuteAction::CreateLambda([=]{ SetSelectionRangeEnd(FrameNumber); }),
+				FExecuteAction::CreateLambda([this, FrameNumber]{ SetSelectionRangeEnd(FrameNumber); }),
 				FCanExecuteAction::CreateLambda([=]{ return SelectionRange.IsEmpty() || FrameNumber >= UE::MovieScene::DiscreteInclusiveLower(SelectionRange); })
 			)
 		);
@@ -1616,7 +1616,7 @@ TSharedRef<SWidget> FSequencerTimeSliderController::OpenSetPlaybackRangeMenu(con
 			FText(),
 			FSlateIcon(),
 			FUIAction(
-				FExecuteAction::CreateLambda([=]{ TimeSliderArgs.OnSelectionRangeChanged.ExecuteIfBound(TRange<FFrameNumber>::Empty()); }),
+				FExecuteAction::CreateLambda([this]{ TimeSliderArgs.OnSelectionRangeChanged.ExecuteIfBound(TRange<FFrameNumber>::Empty()); }),
 				FCanExecuteAction::CreateLambda([=]{ return !SelectionRange.IsEmpty(); })
 			)
 		);
@@ -1649,9 +1649,9 @@ TSharedRef<SWidget> FSequencerTimeSliderController::OpenSetPlaybackRangeMenu(con
 					FText::Format(LOCTEXT("DisplayTimeSpace", "Display time in the space of {0}"), ParentText),
 					FSlateIcon(),
 					FUIAction(
-						FExecuteAction::CreateLambda([=] { TimeSliderArgs.OnScrubPositionParentChanged.ExecuteIfBound(ParentID); }),
+						FExecuteAction::CreateLambda([this, ParentID] { TimeSliderArgs.OnScrubPositionParentChanged.ExecuteIfBound(ParentID); }),
 						FCanExecuteAction(),
-						FIsActionChecked::CreateLambda([=] { return TimeSliderArgs.ScrubPositionParent.Get() == MovieSceneSequenceID::Invalid ? ParentID == TimeSliderArgs.ScrubPositionParentChain.Get().Last() : TimeSliderArgs.ScrubPositionParent.Get() == ParentID; })
+						FIsActionChecked::CreateLambda([this, ParentID] { return TimeSliderArgs.ScrubPositionParent.Get() == MovieSceneSequenceID::Invalid ? ParentID == TimeSliderArgs.ScrubPositionParentChain.Get().Last() : TimeSliderArgs.ScrubPositionParent.Get() == ParentID; })
 					),
 					NAME_None,
 					EUserInterfaceActionType::RadioButton
@@ -1709,7 +1709,7 @@ TSharedRef<SWidget> FSequencerTimeSliderController::OpenSetPlaybackRangeMenu(con
 					TSharedPtr<FStructOnScope> StructOnScope = MakeShared<FStructOnScope>(FMovieSceneMarkedFrame::StaticStruct(), (uint8 *)&InMovieScene->GetMarkedFrames()[InMarkedFrameIndex]);
 
 					DetailsView = PropertyEditorModule.CreateStructureDetailView(DetailsViewArgs, StructureDetailsViewArgs, nullptr);
-					DetailsView->GetDetailsView()->RegisterInstancedCustomPropertyTypeLayout("FrameNumber", FOnGetPropertyTypeCustomizationInstance::CreateLambda([=]() {
+					DetailsView->GetDetailsView()->RegisterInstancedCustomPropertyTypeLayout("FrameNumber", FOnGetPropertyTypeCustomizationInstance::CreateLambda([this]() {
 						return MakeShared<FFrameNumberDetailsCustomization>(WeakSequencer.Pin()->GetNumericTypeInterface()); }));
 					DetailsView->SetStructureData(StructOnScope);
 
@@ -1742,8 +1742,8 @@ TSharedRef<SWidget> FSequencerTimeSliderController::OpenSetPlaybackRangeMenu(con
 			FText(),
 			FSlateIcon(),
 			FUIAction(
-				FExecuteAction::CreateLambda( [=]{ AddMarkAtFrame(FrameNumber); }),
-				FCanExecuteAction::CreateLambda([=]{ return !TimeSliderArgs.AreMarkedFramesLocked.Get() && MarkedIndex == INDEX_NONE; }))
+				FExecuteAction::CreateLambda( [this, FrameNumber]{ AddMarkAtFrame(FrameNumber); }),
+				FCanExecuteAction::CreateLambda([this, MarkedIndex]{ return !TimeSliderArgs.AreMarkedFramesLocked.Get() && MarkedIndex == INDEX_NONE; }))
 		);
 
 		MenuBuilder.AddMenuEntry( 
@@ -1751,8 +1751,8 @@ TSharedRef<SWidget> FSequencerTimeSliderController::OpenSetPlaybackRangeMenu(con
 			FText(),
 			FSlateIcon(),
 			FUIAction(
-				FExecuteAction::CreateLambda([=]{ DeleteMarkAtIndex(MarkedIndex); }),
-				FCanExecuteAction::CreateLambda([=]{ return !TimeSliderArgs.AreMarkedFramesLocked.Get() && MarkedIndex != INDEX_NONE; }))
+				FExecuteAction::CreateLambda([this, MarkedIndex]{ DeleteMarkAtIndex(MarkedIndex); }),
+				FCanExecuteAction::CreateLambda([this, MarkedIndex]{ return !TimeSliderArgs.AreMarkedFramesLocked.Get() && MarkedIndex != INDEX_NONE; }))
 		);
 
 		MenuBuilder.AddMenuEntry( 
@@ -1760,8 +1760,8 @@ TSharedRef<SWidget> FSequencerTimeSliderController::OpenSetPlaybackRangeMenu(con
 			FText(),
 			FSlateIcon(),
 			FUIAction(
-				FExecuteAction::CreateLambda([=]{ DeleteAllMarks(); }),
-				FCanExecuteAction::CreateLambda([=]{ return !TimeSliderArgs.AreMarkedFramesLocked.Get() && bHasMarks; }))
+				FExecuteAction::CreateLambda([this]{ DeleteAllMarks(); }),
+				FCanExecuteAction::CreateLambda([this, bHasMarks]{ return !TimeSliderArgs.AreMarkedFramesLocked.Get() && bHasMarks; }))
 			);
 
 		MenuBuilder.AddMenuEntry(
@@ -1769,9 +1769,9 @@ TSharedRef<SWidget> FSequencerTimeSliderController::OpenSetPlaybackRangeMenu(con
 			LOCTEXT("ToggleLockedMarksTooltip", "Lock/Unlock all marked frames"),
 			FSlateIcon(),
 			FUIAction(
-				FExecuteAction::CreateLambda([=]{ TimeSliderArgs.OnToggleMarkedFramesLocked.ExecuteIfBound(); }),
+				FExecuteAction::CreateLambda([this]{ TimeSliderArgs.OnToggleMarkedFramesLocked.ExecuteIfBound(); }),
 				FCanExecuteAction::CreateLambda([=]{ return !bReadOnly; }),
-				FIsActionChecked::CreateLambda([=]{ return TimeSliderArgs.AreMarkedFramesLocked.Get(); })),
+				FIsActionChecked::CreateLambda([this]{ return TimeSliderArgs.AreMarkedFramesLocked.Get(); })),
 			NAME_None,
 			EUserInterfaceActionType::ToggleButton
 			);
diff --git a/Engine/Source/Editor/UnrealEd/Private/AssetPrivatizeModel.cpp b/Engine/Source/Editor/UnrealEd/Private/AssetPrivatizeModel.cpp
--- a/Engine/Source/Editor/UnrealEd/Private/AssetPrivatizeModel.cpp
+++ b/Engine/Source/Editor/UnrealEd/Private/AssetPrivatizeModel.cpp
@@ -51,14 +51,14 @@ void FPendingPrivateAsset::CheckForIllegalReferences()
 
 	AssetRegistryModule.Get().GetReferencers(OwningPackage->GetFName(), IllegalDiskReferences);
 
-	IllegalDiskReferences.RemoveAll([=](const FName& Reference)
+	IllegalDiskReferences.RemoveAll([this](const FName& Reference)
 		{
 			return !IsReferenceIllegal(Reference);
 		});
 
 	ObjectTools::GatherObjectReferencersForDeletion(Object, bIsReferencedInMemoryByNonUndo, bIsReferencedInMemoryByUndo, &IllegalMemoryReferences);
 
-	IllegalMemoryReferences.ExternalReferences.RemoveAll([=](const FReferencerInformation& ReferenceInfo)
+	IllegalMemoryReferences.ExternalReferences.RemoveAll([this](const FReferencerInformation& ReferenceInfo)
 		{
 			if (ReferenceInfo.Referencer->HasAnyFlags(RF_Transient))
 			{
@@ -67,7 +67,7 @@ void FPendingPrivateAsset::CheckForIllegalReferences()
 			return !IsReferenceIllegal(ReferenceInfo.Referencer->GetPackage()->GetFName());
 		});
 
-	IllegalMemoryReferences.InternalReferences.RemoveAll([=](const FReferencerInformation& ReferenceInfo)
+	IllegalMemoryReferences.InternalReferences.RemoveAll([this](const FReferencerInformation& ReferenceInfo)
 		{
 			if (ReferenceInfo.Referencer->HasAnyFlags(RF_Transient))
 			{
diff --git a/Engine/Source/Editor/UnrealEd/Private/AutoReimport/ContentDirectoryMonitor.cpp b/Engine/Source/Editor/UnrealEd/Private/AutoReimport/ContentDirectoryMonitor.cpp
--- a/Engine/Source/Editor/UnrealEd/Private/AutoReimport/ContentDirectoryMonitor.cpp
+++ b/Engine/Source/Editor/UnrealEd/Private/AutoReimport/ContentDirectoryMonitor.cpp
@@ -126,7 +126,7 @@ void FContentDirectoryMonitor::Tick()
 	// Immediately resolve any changes that we should not consider
 	const FDateTime Threshold = FDateTime::UtcNow() - FTimespan::FromSeconds(GetDefault<UEditorLoadingSavingSettings>()->AutoReimportThreshold);
 
-	TArray<DirectoryWatcher::FUpdateCacheTransaction> InsignificantTransactions = Cache.FilterOutstandingChanges([=](const DirectoryWatcher::FUpdateCacheTransaction& Transaction, const FDateTime& TimeOfChange){
+	TArray<DirectoryWatcher::FUpdateCacheTransaction> InsignificantTransactions = Cache.FilterOutstandingChanges([this, Threshold](const DirectoryWatcher::FUpdateCacheTransaction& Transaction, const FDateTime& TimeOfChange){
 		return TimeOfChange <= Threshold && !ShouldConsiderChange(Transaction);
 	});
 
@@ -168,7 +168,7 @@ int32 FContentDirectoryMonitor::GetNumUnprocessedChanges() const
 	int32 Total = 0;
 
 	// Get all the changes that have happend beyond our import threshold
-	Cache.IterateOutstandingChanges([=, &Total](const DirectoryWatcher::FUpdateCacheTransaction& Transaction, const FDateTime& TimeOfChange){
+	Cache.IterateOutstandingChanges([this, &Total, Threshold](const DirectoryWatcher::FUpdateCacheTransaction& Transaction, const FDateTime& TimeOfChange){
 		if (TimeOfChange <= Threshold && ShouldConsiderChange(Transaction))
 		{
 			++Total;
@@ -191,7 +191,7 @@ int32 FContentDirectoryMonitor::StartProcessing()
 	const FDateTime Threshold = FDateTime::UtcNow() - FTimespan::FromSeconds(GetDefault<UEditorLoadingSavingSettings>()->AutoReimportThreshold);
 
 	// Get all the changes that have happend beyond our import threshold
-	auto OutstandingChanges = Cache.FilterOutstandingChanges([=](const DirectoryWatcher::FUpdateCacheTransaction& Transaction, const FDateTime& TimeOfChange){
+	auto OutstandingChanges = Cache.FilterOutstandingChanges([this, Threshold](const DirectoryWatcher::FUpdateCacheTransaction& Transaction, const FDateTime& TimeOfChange){
 		return TimeOfChange <= Threshold && ShouldConsiderChange(Transaction);
 	});
 
diff --git a/Engine/Source/Runtime/Landscape/Private/LandscapeEditLayers.cpp b/Engine/Source/Runtime/Landscape/Private/LandscapeEditLayers.cpp
--- a/Engine/Source/Runtime/Landscape/Private/LandscapeEditLayers.cpp
+++ b/Engine/Source/Runtime/Landscape/Private/LandscapeEditLayers.cpp
@@ -9196,7 +9196,7 @@ void ALandscape::ClearPaintLayer(const FGuid& InLayerGuid, ULandscapeLayerInfoOb
 	}
 
 	Modify();
-	FScopedSetLandscapeEditingLayer Scope(this, InLayerGuid, [=] { RequestLayersContentUpdate(ELandscapeLayerUpdateMode::Update_Weightmap_All); });
+	FScopedSetLandscapeEditingLayer Scope(this, InLayerGuid, [this] { RequestLayersContentUpdate(ELandscapeLayerUpdateMode::Update_Weightmap_All); });
 
 	FLandscapeEditDataInterface LandscapeEdit(LandscapeInfo);
 	LandscapeInfo->ForEachLandscapeProxy([&](ALandscapeProxy* Proxy)
@@ -9231,7 +9231,7 @@ void ALandscape::ClearLayer(const FGuid& InLayerGuid, TSet<TObjectPtr<ULandscape
 	}
 
 	Modify(bMarkPackageDirty);
-	FScopedSetLandscapeEditingLayer Scope(this, Layer->Guid, [=] { RequestLayersContentUpdate(ELandscapeLayerUpdateMode::Update_All); });
+	FScopedSetLandscapeEditingLayer Scope(this, Layer->Guid, [this] { RequestLayersContentUpdate(ELandscapeLayerUpdateMode::Update_All); });
 
 	TArray<uint16> NewHeightData;
 	NewHeightData.AddZeroed(FMath::Square(ComponentSizeQuads + 1));
@@ -9439,7 +9439,7 @@ void ALandscape::UpdateLandscapeSplines(const FGuid& InTargetLayer, bool bInUpda
 	const FLandscapeLayer* TargetLayer = GetLayer(TargetLayerGuid);
 	if (LandscapeInfo && TargetLayer)
 	{
-		FScopedSetLandscapeEditingLayer Scope(this, TargetLayerGuid, [=] { this->RequestLayersContentUpdate(ELandscapeLayerUpdateMode::Update_All); });
+		FScopedSetLandscapeEditingLayer Scope(this, TargetLayerGuid, [this] { this->RequestLayersContentUpdate(ELandscapeLayerUpdateMode::Update_All); });
 		// Temporarily disable material instance updates since it will be done once at the end (requested by RequestLayersContentUpdateForceAll)
 		GDisableUpdateLandscapeMaterialInstances = true;
 		TSet<TObjectPtr<ULandscapeComponent>>* ModifiedComponent = nullptr;
diff --git a/Engine/Source/Runtime/OpenGLDrv/Private/OpenGLQuery.cpp b/Engine/Source/Runtime/OpenGLDrv/Private/OpenGLQuery.cpp
--- a/Engine/Source/Runtime/OpenGLDrv/Private/OpenGLQuery.cpp
+++ b/Engine/Source/Runtime/OpenGLDrv/Private/OpenGLQuery.cpp
@@ -567,9 +567,9 @@ bool FOpenGLDynamicRHI::RHIGetRenderQueryResult(FRHIRenderQuery* QueryRHI, uint6
 				if (IsRunningRHIInDedicatedThread())
 				{
 					// send a command that will wait, so if the RHIT runs out of work, it just blocks and waits for the GPU
-					ALLOC_COMMAND_CL(RHICmdList, FRHICommandGLCommand)([=]() {GetRenderQueryResult_OnThisThread(ResourceCast(QueryRHI), true); });
+					ALLOC_COMMAND_CL(RHICmdList, FRHICommandGLCommand)([this, QueryRHI]() {GetRenderQueryResult_OnThisThread(ResourceCast(QueryRHI), true); });
 					FGraphEventRef Done = RHICmdList.RHIThreadFence(false);
-					ALLOC_COMMAND_CL(RHICmdList, FRHICommandGLCommand)([=]() {GBatcher.Flush(*this, QueryRHI); });
+					ALLOC_COMMAND_CL(RHICmdList, FRHICommandGLCommand)([this, QueryRHI]() {GBatcher.Flush(*this, QueryRHI); });
 					RHICmdList.ImmediateFlush(EImmediateFlushType::DispatchToRHIThread);
 					while (!Done->IsComplete())
 					{
@@ -587,9 +587,9 @@ bool FOpenGLDynamicRHI::RHIGetRenderQueryResult(FRHIRenderQuery* QueryRHI, uint6
 				}
 				else
 				{
-					ALLOC_COMMAND_CL(RHICmdList, FRHICommandGLCommand)([=]() {GetRenderQueryResult_OnThisThread(ResourceCast(QueryRHI), true); });
+					ALLOC_COMMAND_CL(RHICmdList, FRHICommandGLCommand)([this, QueryRHI]() {GetRenderQueryResult_OnThisThread(ResourceCast(QueryRHI), true); });
 					FGraphEventRef Done = RHICmdList.RHIThreadFence(false);
-					ALLOC_COMMAND_CL(RHICmdList, FRHICommandGLCommand)([=]() {GBatcher.Flush(*this, QueryRHI); });
+					ALLOC_COMMAND_CL(RHICmdList, FRHICommandGLCommand)([this, QueryRHI]() {GBatcher.Flush(*this, QueryRHI); });
 					RHICmdList.ImmediateFlush(EImmediateFlushType::DispatchToRHIThread);
 					FRHICommandListExecutor::WaitOnRHIThreadFence(Done);
 				}
diff --git a/Engine/Source/Runtime/UMG/Private/Components/ScrollBox.cpp b/Engine/Source/Runtime/UMG/Private/Components/ScrollBox.cpp
--- a/Engine/Source/Runtime/UMG/Private/Components/ScrollBox.cpp
+++ b/Engine/Source/Runtime/UMG/Private/Components/ScrollBox.cpp
@@ -565,7 +565,7 @@ void UScrollBox::OnDescendantDeselectedByDesigner( UWidget* DescendantWidget )
 	}
 
 	// because we get a deselect before we get a select, we need to delay this call until we're sure we didn't scroll to another widget.
-	TickHandle = FTSTicker::GetCoreTicker().AddTicker( FTickerDelegate::CreateLambda( [=]( float ) -> bool
+	TickHandle = FTSTicker::GetCoreTicker().AddTicker( FTickerDelegate::CreateLambda( [this]( float ) -> bool
 	                                                                                {
                                                                                         QUICK_SCOPE_CYCLE_COUNTER(STAT_UScrollBox_ScrollToStart_LambdaTick);
 		                                                                                this->ScrollToStart();
