Sequencer: Outliner column registration and column interface declaration

diff --git a/Engine/Source/Editor/Sequencer/Private/MovieSceneOutlinerColumn.cpp b/Engine/Source/Editor/Sequencer/Private/MovieSceneOutlinerColumn.cpp
new file mode 100644
--- /dev/null
+++ b/Engine/Source/Editor/Sequencer/Private/MovieSceneOutlinerColumn.cpp
@@ -0,0 +1,13 @@
+// Copyright Epic Games, Inc. All Rights Reserved.
+
+#include "MovieSceneOutlinerColumn.h"
+
+FMovieSceneOutlinerColumn::FMovieSceneOutlinerColumn()
+{
+
+}
+
+FMovieSceneOutlinerColumn::~FMovieSceneOutlinerColumn()
+{
+
+}
diff --git a/Engine/Source/Editor/Sequencer/Private/Sequencer.cpp b/Engine/Source/Editor/Sequencer/Private/Sequencer.cpp
--- a/Engine/Source/Editor/Sequencer/Private/Sequencer.cpp
+++ b/Engine/Source/Editor/Sequencer/Private/Sequencer.cpp
@@ -127,6 +127,7 @@
 #include "SequencerKeyActor.h"
 #include "ISequencerChannelInterface.h"
 #include "IMovieRendererInterface.h"
+#include "ISequencerOutlinerColumn.h"
 #include "SequencerKeyCollection.h"
 #include "CurveEditor.h"
 #include "CurveEditorScreenSpace.h"
@@ -281,7 +282,7 @@ namespace UE
 } // namespace UE
 
 
-void FSequencer::InitSequencer(const FSequencerInitParams& InitParams, const TSharedRef<ISequencerObjectChangeListener>& InObjectChangeListener, const TArray<FOnCreateTrackEditor>& TrackEditorDelegates, const TArray<FOnCreateEditorObjectBinding>& EditorObjectBindingDelegates)
+void FSequencer::InitSequencer(const FSequencerInitParams& InitParams, const TSharedRef<ISequencerObjectChangeListener>& InObjectChangeListener, const TArray<FOnCreateTrackEditor>& TrackEditorDelegates, const TArray<FOnCreateEditorObjectBinding>& EditorObjectBindingDelegates, const TArray<FOnCreateOutlinerColumn>& OutlinerColumnDelegates)
 {
 	using namespace UE::Sequencer;
 
@@ -445,6 +446,19 @@ void FSequencer::InitSequencer(const FSequencerInitParams& InitParams, const TSh
 		}
 	}
 
+	// OutlinerColumns are registered to be provided to SOutlinerView by SSequencer
+	for (int32 DelegateIndex = 0; DelegateIndex < OutlinerColumnDelegates.Num(); ++DelegateIndex)
+	{
+		check(OutlinerColumnDelegates[DelegateIndex].IsBound());
+		TSharedRef<ISequencerOutlinerColumn> OutlinerColumn = OutlinerColumnDelegates[DelegateIndex].Execute(SharedThis(this));
+
+		if (OutlinerColumn->SupportsSequence(InitParams.RootSequence))
+		{
+			check(!OutlinerColumns.Contains(OutlinerColumn->GetColumnName()));
+			OutlinerColumns.Add(OutlinerColumn->GetColumnName(), OutlinerColumn);
+		}
+	}
+
 	TWeakPtr<UE::MovieScene::IDeferredSignedObjectChangeHandler> WeakDeferredSignedObjectChangeHandler = UMovieSceneSignedObject::GetDeferredHandler();
 	if (WeakDeferredSignedObjectChangeHandler.IsValid())
 	{
diff --git a/Engine/Source/Editor/Sequencer/Private/Sequencer.h b/Engine/Source/Editor/Sequencer/Private/Sequencer.h
--- a/Engine/Source/Editor/Sequencer/Private/Sequencer.h
+++ b/Engine/Source/Editor/Sequencer/Private/Sequencer.h
@@ -129,7 +129,7 @@ public:
 	 * @param TrackEditorDelegates Delegates to call to create auto-key handlers for this sequencer.
 	 * @param EditorObjectBindingDelegates Delegates to call to create object bindings for this sequencer.
 	 */
-	void InitSequencer(const FSequencerInitParams& InitParams, const TSharedRef<ISequencerObjectChangeListener>& InObjectChangeListener, const TArray<FOnCreateTrackEditor>& TrackEditorDelegates, const TArray<FOnCreateEditorObjectBinding>& EditorObjectBindingDelegates);
+	void InitSequencer(const FSequencerInitParams& InitParams, const TSharedRef<ISequencerObjectChangeListener>& InObjectChangeListener, const TArray<FOnCreateTrackEditor>& TrackEditorDelegates, const TArray<FOnCreateEditorObjectBinding>& EditorObjectBindingDelegatess, const TArray<FOnCreateOutlinerColumn>& OutlinerColumnDelegates);
 
 	/** @return The current view range */
 	virtual FAnimatedRange GetViewRange() const override;
@@ -287,6 +287,14 @@ public:
 		return TrackEditors;
 	}
 
+	/**
+	* @return Outliner Columns registered to the sequencer by column name
+	*/
+	const TMap<FName, TSharedPtr<ISequencerOutlinerColumn>>& GetOutlinerColumns() const
+	{
+		return OutlinerColumns;
+	}
+
 public:
 
 	/** @return The set of vertical frames */
@@ -1148,6 +1156,9 @@ private:
 	TArray<TSharedPtr<ISequencerTrackEditor>> TrackEditors;
 	TMap<FObjectKey, TSharedPtr<ISequencerTrackEditor>> TrackEditorsByType;
 
+	/** List of Outliner column creators that are supported by the Sequencer. */
+	TMap<FName, TSharedPtr<ISequencerOutlinerColumn>> OutlinerColumns;
+
 	/** List of object bindings we can use */
 	TArray<TSharedPtr<ISequencerEditorObjectBinding>> ObjectBindings;
 
diff --git a/Engine/Source/Editor/Sequencer/Private/SequencerModule.cpp b/Engine/Source/Editor/Sequencer/Private/SequencerModule.cpp
--- a/Engine/Source/Editor/Sequencer/Private/SequencerModule.cpp
+++ b/Engine/Source/Editor/Sequencer/Private/SequencerModule.cpp
@@ -163,7 +163,7 @@ public:
 
 		OnPreSequencerInit.Broadcast(Sequencer, ObjectChangeListener, InitParams);
 
-		Sequencer->InitSequencer(InitParams, ObjectChangeListener, TrackEditorDelegates, EditorObjectBindingDelegates);
+		Sequencer->InitSequencer(InitParams, ObjectChangeListener, TrackEditorDelegates, EditorObjectBindingDelegates, OutlinerColumnDelegates);
 
 		OnSequencerCreated.Broadcast(Sequencer);
 
@@ -217,6 +217,15 @@ public:
 		TrackModelDelegates.RemoveAll([=](const FOnCreateTrackModel& Delegate) { return Delegate.GetHandle() == InHandle; });
 	}
 
+	virtual FDelegateHandle RegisterOutlinerColumn(FOnCreateOutlinerColumn InCreator) override {
+		OutlinerColumnDelegates.Add(InCreator);
+		return OutlinerColumnDelegates.Last().GetHandle();
+	}
+
+	virtual void UnregisterOutlinerColumn(FDelegateHandle InHandle) override {
+		OutlinerColumnDelegates.RemoveAll([=](const FOnCreateOutlinerColumn& Delegate) { return Delegate.GetHandle() == InHandle; });
+	}
+
 	virtual FDelegateHandle RegisterOnSequencerCreated(FOnSequencerCreated::FDelegate InOnSequencerCreated) override
 	{
 		return OnSequencerCreated.Add(InOnSequencerCreated);
@@ -483,6 +492,9 @@ private:
 	/** List of track model creators */
 	TArray<FOnCreateTrackModel> TrackModelDelegates;
 
+	/** List of outliner column creators */
+	TArray<FOnCreateOutlinerColumn> OutlinerColumnDelegates;
+
 	/** Global details row extension delegate; */
 	FDelegateHandle OnGetGlobalRowExtensionHandle;
 
diff --git a/Engine/Source/Editor/Sequencer/Public/ISequencerModule.h b/Engine/Source/Editor/Sequencer/Public/ISequencerModule.h
--- a/Engine/Source/Editor/Sequencer/Public/ISequencerModule.h
+++ b/Engine/Source/Editor/Sequencer/Public/ISequencerModule.h
@@ -19,6 +19,7 @@ class FExtensibilityManager;
 class FMenuBuilder;
 class FSequencerCustomizationManager;
 class ISequencerTrackEditor;
+class ISequencerOutlinerColumn;
 class ISequencerEditorObjectBinding;
 class IToolkitHost;
 class UMovieSceneSequence;
@@ -55,6 +56,9 @@ DECLARE_DELEGATE_RetVal_OneParam(TSharedRef<ISequencerEditorObjectBinding>, FOnC
 /** A delegate which will create a track model. */
 DECLARE_DELEGATE_RetVal_OneParam(TSharedPtr<UE::Sequencer::FTrackModel>, FOnCreateTrackModel, UMovieSceneTrack*);
 
+/** A delegate which will create an outliner column */
+DECLARE_DELEGATE_RetVal_OneParam(TSharedRef<ISequencerOutlinerColumn>, FOnCreateOutlinerColumn, TSharedRef<ISequencer>);
+
 /** A delegate that is executed when adding menu content. */
 DECLARE_DELEGATE_OneParam(FOnGetContextMenuContent, FMenuBuilder& /*MenuBuilder*/);
 
@@ -235,6 +239,22 @@ public:
 	 */
 	virtual void UnregisterTrackModel(FDelegateHandle InHandle) = 0;
 
+	/**
+	 * Registers a delegate that will create a gutter column view model for a given Outliner column
+	 *
+	 * @param InCreator Delegate to register
+	 * @return A handle to the newly added delegate
+	 */
+	virtual FDelegateHandle RegisterOutlinerColumn(FOnCreateOutlinerColumn InCreator) = 0;
+
+	/**
+	 * Unregisters a previously registered delegate for creating Outliner columns
+	 *
+	 * @param InHandle Handle to the delegate to unregister
+	 */
+	virtual void UnregisterOutlinerColumn(FDelegateHandle Inhandle) = 0;
+
+
 	/** 
 	 * Registers a delegate that will be called when a sequencer is created
 	 *
diff --git a/Engine/Source/Editor/Sequencer/Public/MovieSceneOutlinerColumn.h b/Engine/Source/Editor/Sequencer/Public/MovieSceneOutlinerColumn.h
new file mode 100644
--- /dev/null
+++ b/Engine/Source/Editor/Sequencer/Public/MovieSceneOutlinerColumn.h
@@ -0,0 +1,27 @@
+// Copyright Epic Games, Inc. All Rights Reserved.
+
+#pragma once
+
+#include "ISequencer.h"
+#include "ISequencerOutlinerColumn.h"
+
+/**
+ * Base class for handling creation of outliner column widgets, column settings, and column properties.
+ */
+class SEQUENCER_API FMovieSceneOutlinerColumn
+	: public ISequencerOutlinerColumn
+{
+public:
+
+	/** Constructor */
+	FMovieSceneOutlinerColumn();
+
+	/** Destructor */
+	virtual ~FMovieSceneOutlinerColumn();
+
+public:
+
+	virtual bool SupportsSequence(UMovieSceneSequence* InSequence) const override { return true; }
+
+	virtual TSharedRef<SWidget> CreateColumnWidget(UE::Sequencer::TViewModelPtr<UE::Sequencer::IOutlinerExtension> InOutlinerExtension) const = 0;
+};
diff --git a/Engine/Source/Editor/SequencerCore/Public/ISequencerOutlinerColumn.h b/Engine/Source/Editor/SequencerCore/Public/ISequencerOutlinerColumn.h
new file mode 100644
--- /dev/null
+++ b/Engine/Source/Editor/SequencerCore/Public/ISequencerOutlinerColumn.h
@@ -0,0 +1,43 @@
+// Copyright Epic Games, Inc. All Rights Reserved.
+
+#pragma once
+
+#include "MVVM/ViewModelPtr.h"
+#include "Templates/SubclassOf.h"
+
+class UMovieSceneSequence;
+
+namespace UE
+{
+	namespace Sequencer
+	{
+		class IOutlinerExtension;
+	}
+}
+
+/**
+ * Interface for sequencer track editors.
+ */
+class ISequencerOutlinerColumn
+{
+	
+public:
+
+	/* Gets the unique name of the column to use in the SOutlinerView registry and context menus */
+	virtual FName GetColumnName() const = 0;
+
+	/* The default visibility state of this column when loaded for the first time */
+	virtual bool IsColumnVisibleByDefault() const { return true; }
+
+	/* Gets whether or not this column is supported by a given Sequencer */
+	virtual bool SupportsSequence(UMovieSceneSequence* InSequence) const = 0;
+
+	/* Gets the widget created for each item within the SOutlinerView, column widgets must be fixed width */
+	virtual TSharedRef<SWidget> CreateColumnWidget(UE::Sequencer::TViewModelPtr<UE::Sequencer::IOutlinerExtension>InOutlinerExtension) const = 0;
+
+public:
+
+	/** Virtual destructor. */
+	virtual ~ISequencerOutlinerColumn() { }
+	
+};
