Sequencer: On pasting a possessable, only parent it to the pasted target if the possessable's class is not an actor (ie. a component)

diff --git a/Engine/Source/Editor/Sequencer/Private/SequencerUtilities.cpp b/Engine/Source/Editor/Sequencer/Private/SequencerUtilities.cpp
--- a/Engine/Source/Editor/Sequencer/Private/SequencerUtilities.cpp
+++ b/Engine/Source/Editor/Sequencer/Private/SequencerUtilities.cpp
@@ -2159,11 +2159,16 @@ bool FSequencerUtilities::PasteBindings(const FString& TextToImport, TSharedRef<
 					}
 				}
 
-				if (FMovieScenePossessable* Possessable = MovieScene->FindPossessable(NewGuid))
+				// Set the parent to the pasted target only if it's not an actor
+				const UClass* PossessedObjectClass = CopyableBinding->Possessable.GetPossessedObjectClass();
+				if (PossessedObjectClass && !PossessedObjectClass->IsChildOf(AActor::StaticClass()))
 				{
-					if (TargetIndex < PasteBindingsParams.Bindings.Num())
+					if (FMovieScenePossessable* Possessable = MovieScene->FindPossessable(NewGuid))
 					{
-						Possessable->SetParent(PasteBindingsParams.Bindings[TargetIndex].BindingID, MovieScene);
+						if (TargetIndex < PasteBindingsParams.Bindings.Num())
+						{
+							Possessable->SetParent(PasteBindingsParams.Bindings[TargetIndex].BindingID, MovieScene);
+						}
 					}
 				}
 
