Sequencer: Follow up to 25459614 - group should be selected if all of its items are selected

diff --git a/Engine/Source/Editor/Sequencer/Private/SSequencerGroupManager.cpp b/Engine/Source/Editor/Sequencer/Private/SSequencerGroupManager.cpp
--- a/Engine/Source/Editor/Sequencer/Private/SSequencerGroupManager.cpp
+++ b/Engine/Source/Editor/Sequencer/Private/SSequencerGroupManager.cpp
@@ -492,7 +492,17 @@ const FSlateBrush* SSequencerGroupManager::GetIconBrush(TSharedPtr<FSequencerNod
 	return nullptr;
 }
 
-void SSequencerGroupManager::RequestDeleteNodeGroup(FSequencerNodeGroupNode* NodeGroupNode)
+void SSequencerGroupManager::SelectItemsInGroup(FSequencerNodeGroupNode* Node)
+{
+	TreeView->ClearSelection();
+
+	for (TSharedPtr<FSequencerNodeGroupTreeNode> ChildNode : Node->Children)
+	{
+		TreeView->SetItemSelection(ChildNode, true);
+	}
+}
+
+void SSequencerGroupManager::RequestDeleteNodeGroup(FSequencerNodeGroupNode * NodeGroupNode)
 {
 	UMovieScene* MovieScene = GetMovieScene();
 	if (!ensure(MovieScene) || !ensure(NodeGroupNode))
@@ -600,18 +610,6 @@ void SSequencerGroupManager::GetSelectedItemsNodePaths(TSet<FString>& OutSelecte
 			TSharedPtr<FSequencerGroupItemNode> ItemNode = StaticCastSharedPtr<FSequencerGroupItemNode>(Node);
 			OutSelectedNodePaths.Add(ItemNode->Path);
 		}
-		else if (Node->GetType() == FSequencerNodeGroupTreeNode::Type::GroupNode)
-		{
-			for (TSharedPtr<FSequencerNodeGroupTreeNode> ChildNode : Node->Children)
-			{
-				// Note: Currently, children of a set can only be item nodes, but that may change in the future.
-				if (ChildNode->GetType() == FSequencerNodeGroupTreeNode::Type::ItemNode)
-				{
-					TSharedPtr<FSequencerGroupItemNode> ItemNode = StaticCastSharedPtr<FSequencerGroupItemNode>(ChildNode);
-					OutSelectedNodePaths.Add(ItemNode->Path);
-				}
-			}
-		}
 	}
 }
 
@@ -637,11 +635,6 @@ void SSequencerGroupManager::SelectSelectedItemsInSequencer()
 	
 		Sequencer->SelectNodesByPath(SelectedNodePaths);
 	}
-
-	// Then update the group manager tree's notion of what is selected. When this empties selection, 
-	// the group item itself will be deselected. Any items within it will be synchronized with what 
-	// is selected in Sequencer
-	SelectItemsSelectedInSequencer();
 }
 
 void SSequencerGroupManager::SelectItemsSelectedInSequencer()
@@ -688,7 +681,6 @@ void SSequencerGroupManager::SelectItemsSelectedInSequencer()
 	TreeView->ClearSelection();
 
 	// Build a list of the treenodes which match a nodepath we want to select
-	TArray<TSharedPtr<FSequencerGroupItemNode>> TreeNodesToSelect;
 	for (const TSharedPtr<FSequencerNodeGroupTreeNode>& Node : NodeGroupsTree)
 	{
 		if (Node->GetType() == FSequencerNodeGroupTreeNode::Type::ItemNode)
@@ -750,6 +742,12 @@ TSharedPtr<SWidget> SSequencerGroupManager::OnContextMenuOpening()
 				EUserInterfaceActionType::ToggleButton
 			);
 
+			MenuBuilder.AddMenuEntry(
+				LOCTEXT("SelectItemsInGroup", "Select Items in Group"),
+				LOCTEXT("SelectItemsInGroupTooltip", "Select items in group"),
+				FSlateIcon(),
+				FUIAction(FExecuteAction::CreateSP(this, &SSequencerGroupManager::SelectItemsInGroup, NodeGroupNode.Get())));
+
 			MenuBuilder.AddMenuEntry(
 				FText::Format(LOCTEXT("RenameNodeGroupFormat", "Rename {0}"), NodeGroupNode->DisplayText),
 				FText(),
diff --git a/Engine/Source/Editor/Sequencer/Private/SSequencerGroupManager.h b/Engine/Source/Editor/Sequencer/Private/SSequencerGroupManager.h
--- a/Engine/Source/Editor/Sequencer/Private/SSequencerGroupManager.h
+++ b/Engine/Source/Editor/Sequencer/Private/SSequencerGroupManager.h
@@ -41,6 +41,7 @@ private:
 	virtual void Tick(const FGeometry& AllottedGeometry, const double InCurrentTime, const float InDeltaTime) override;
 	void UpdateTree();
 	void HandleTreeSelectionChanged(TSharedPtr<FSequencerNodeGroupTreeNode> InSelectedNode, ESelectInfo::Type SelectionType);
+	void SelectItemsInGroup(FSequencerNodeGroupNode* NodeGroupNode);
 	void RequestDeleteNodeGroup(FSequencerNodeGroupNode* NodeGroupNode);
 	void RemoveSelectedItemsFromNodeGroup();
 	void CreateNodeGroup();
