Sequencer: Destroy customizations on exit

diff --git a/Engine/Source/Editor/Sequencer/Private/MVVM/ViewModels/SequencerEditorViewModel.cpp b/Engine/Source/Editor/Sequencer/Private/MVVM/ViewModels/SequencerEditorViewModel.cpp
--- a/Engine/Source/Editor/Sequencer/Private/MVVM/ViewModels/SequencerEditorViewModel.cpp
+++ b/Engine/Source/Editor/Sequencer/Private/MVVM/ViewModels/SequencerEditorViewModel.cpp
@@ -160,6 +160,16 @@ void FSequencerEditorViewModel::OnTrackAreaHotspotChanged(TSharedPtr<ITrackAreaH
 	CurrentHotspot = NewHotspot;
 }
 
+FSequencerEditorViewModel::~FSequencerEditorViewModel()
+{
+	// Get rid of previously active customizations.
+	for (const TUniquePtr<ISequencerCustomization>& Customization : ActiveCustomizations)
+	{
+		Customization->UnregisterSequencerCustomization();
+	}
+	ActiveCustomizations.Reset();
+}
+
 bool FSequencerEditorViewModel::UpdateSequencerCustomizations(const UMovieSceneSequence* PreviousFocusedSequence)
 {
 	ISequencerModule& SequencerModule = FModuleManager::LoadModuleChecked<ISequencerModule>("Sequencer");
diff --git a/Engine/Source/Editor/Sequencer/Public/MVVM/ViewModels/SequencerEditorViewModel.h b/Engine/Source/Editor/Sequencer/Public/MVVM/ViewModels/SequencerEditorViewModel.h
--- a/Engine/Source/Editor/Sequencer/Public/MVVM/ViewModels/SequencerEditorViewModel.h
+++ b/Engine/Source/Editor/Sequencer/Public/MVVM/ViewModels/SequencerEditorViewModel.h
@@ -34,6 +34,7 @@ public:
 	UE_SEQUENCER_DECLARE_CASTABLE(FSequencerEditorViewModel, FEditorViewModel);
 
 	FSequencerEditorViewModel(TSharedRef<ISequencer> InSequencer, const FSequencerHostCapabilities& InHostCapabilities);
+	~FSequencerEditorViewModel();
 
 	/** Retrieve this editor's selection class */
 	TSharedPtr<FSequencerSelection> GetSelection() const;
