Sequencer: fix recycling of view-models when track layouts are rebuilt

Recycling should be done once before the track (or track row) layout is built. This is because when track (or track rows) have multiple sections, we want to make the pre-rebuild view-models available for recycling, but those already re-added for the first section shouldn't be recycled when we start the second section. So we do all this first-time recycling outside of the track layout builder, and effectively revert the (wrong) fix from 25693978 which removed the "don't recycle if recycling already exists" logic.

Other changes here include:
- Removing an unused/unimplemented constructor for FHierarchicalModelListRefresher
- Renaming RecycleChildren to ConditionalRecycleChildren to make it clearer that it's not meant to always recycle (it should only recycle the first time a model is encountered)
- Add a call to Pop() after a top level channel has been added. It never really mattered because generally when something has a top level channel, it doesn't have anything below it

diff --git a/Engine/Source/Editor/Sequencer/Private/MVVM/ViewModels/TrackModel.cpp b/Engine/Source/Editor/Sequencer/Private/MVVM/ViewModels/TrackModel.cpp
--- a/Engine/Source/Editor/Sequencer/Private/MVVM/ViewModels/TrackModel.cpp
+++ b/Engine/Source/Editor/Sequencer/Private/MVVM/ViewModels/TrackModel.cpp
@@ -267,8 +267,8 @@ void FTrackModel::ForceUpdate()
 		struct FRowData
 		{
 			TSharedPtr<FTrackRowModel> Row;
-			TUniquePtr<FScopedViewModelListHead> OldSections;
 			TSharedPtr<FViewModel> SectionsTail;
+			TUniquePtr<FScopedViewModelListHead> RecycledModels;
 		};
 		TArray<FRowData, TInlineAllocator<8>> RowModels;
 		RowModels.SetNum(PopulatedRows.Num());
@@ -292,9 +292,12 @@ void FTrackModel::ForceUpdate()
 
 				RowModels[RowIndex].Row = TrackRowModel;
 
-				// Keep sections on rows alive as well
-				RowModels[RowIndex].OldSections = MakeUnique<FScopedViewModelListHead>(TrackRowModel, EViewModelListType::Recycled);
-				TrackRowModel->GetSectionModels().MoveChildrenTo<IRecyclableExtension>(RowModels[RowIndex].OldSections->GetChildren(), IRecyclableExtension::CallOnRecycle);
+				// Recycle sections, outliner children, and more, while keeping them alive.
+				RowModels[RowIndex].RecycledModels = MakeUnique<FScopedViewModelListHead>(TrackRowModel, EViewModelListType::Recycled);
+				FViewModelChildren RecycledRowModels = RowModels[RowIndex].RecycledModels->GetChildren();
+				TrackRowModel->GetSectionModels().MoveChildrenTo<IRecyclableExtension>(RecycledRowModels, IRecyclableExtension::CallOnRecycle);
+				TrackRowModel->GetChildList(EViewModelListType::Outliner).MoveChildrenTo<IRecyclableExtension>(RecycledRowModels, IRecyclableExtension::CallOnRecycle);
+				TrackRowModel->GetTopLevelChannels().MoveChildrenTo<IRecyclableExtension>(RecycledRowModels, IRecyclableExtension::CallOnRecycle);
 			}
 		}
 
diff --git a/Engine/Source/Editor/Sequencer/Private/MVVM/ViewModels/TrackModelLayoutBuilder.cpp b/Engine/Source/Editor/Sequencer/Private/MVVM/ViewModels/TrackModelLayoutBuilder.cpp
--- a/Engine/Source/Editor/Sequencer/Private/MVVM/ViewModels/TrackModelLayoutBuilder.cpp
+++ b/Engine/Source/Editor/Sequencer/Private/MVVM/ViewModels/TrackModelLayoutBuilder.cpp
@@ -40,12 +40,12 @@ FHierarchicalModelListRefresher::FHierarchicalModelListRefresher(TSharedPtr<FVie
 {
 	// Add a recycled child list to the model. This will get cleaned up when
 	// this FHierarchicalModelListRefresher instance is destroyed
-	RecycleChildren(InRoot, InExistingChildren);
+	ConditionalRecycleChildren(InRoot, InExistingChildren);
 
 	ListData.Add(FListData(InRoot, InExistingChildren));
 }
 
-void FHierarchicalModelListRefresher::RecycleChildren(const TSharedPtr<FViewModel>& InModel, FViewModelChildren InExistingChildren)
+void FHierarchicalModelListRefresher::ConditionalRecycleChildren(const TSharedPtr<FViewModel>& InModel, FViewModelChildren InExistingChildren)
 {
 	TOptional<FViewModelChildren> RecycledChildren = InModel->FindChildList(EViewModelListType::Recycled);
 	if (!RecycledChildren)
@@ -59,9 +59,10 @@ void FHierarchicalModelListRefresher::RecycleChildren(const TSharedPtr<FViewMode
 		RecycledLists.Last().Emplace(InModel, EViewModelListType::Recycled);
 
 		RecycledChildren = InModel->GetChildList(EViewModelListType::Recycled);
+
+		check(RecycledChildren.IsSet());
+		InExistingChildren.MoveChildrenTo<IRecyclableExtension>(RecycledChildren.GetValue(), IRecyclableExtension::CallOnRecycle);
 	}
-	check(RecycledChildren.IsSet());
-	InExistingChildren.MoveChildrenTo<IRecyclableExtension>(RecycledChildren.GetValue(), IRecyclableExtension::CallOnRecycle);
 }
 
 bool FHierarchicalModelListRefresher::IsValid() const
@@ -96,8 +97,10 @@ void FHierarchicalModelListRefresher::RecurseInto(TSharedPtr<FViewModel> Item, F
 	// or row that has multiple sections with different channel layouts.
 	// We ensure that children remain eligible for recycling for the duration
 	// of all sections by adding a temporary recycle list to the model while
-	// the layout is being refreshed
-	RecycleChildren(Item, InExistingChildren);
+	// the layout is being refreshed, but we create this recycle list only on the
+	// first try to not recycle "legitimate" children we added in the previous
+	// loop.
+	ConditionalRecycleChildren(Item, InExistingChildren);
 
 	ListData.Add(FListData(Item, InExistingChildren));
 }
@@ -225,6 +228,8 @@ void FTrackModelLayoutBuilder::SetTopLevelChannel(const FMovieSceneChannelHandle
 	OutlinerList.RecurseInto(CurrentModel, TopLevelGroup);
 
 	AddChannel(Channel, true, OptionalFactory);
+
+	OutlinerList.Pop();
 }
 
 void FTrackModelLayoutBuilder::AddChannel(const FMovieSceneChannelHandle& Channel, TFunction<TSharedPtr<UE::Sequencer::FChannelModel>(FName, const FMovieSceneChannelHandle&)> OptionalFactory)
diff --git a/Engine/Source/Editor/Sequencer/Private/MVVM/ViewModels/TrackModelLayoutBuilder.h b/Engine/Source/Editor/Sequencer/Private/MVVM/ViewModels/TrackModelLayoutBuilder.h
--- a/Engine/Source/Editor/Sequencer/Private/MVVM/ViewModels/TrackModelLayoutBuilder.h
+++ b/Engine/Source/Editor/Sequencer/Private/MVVM/ViewModels/TrackModelLayoutBuilder.h
@@ -29,7 +29,6 @@ class FSectionModel;
 struct FHierarchicalModelListRefresher
 {
 	FHierarchicalModelListRefresher();
-	FHierarchicalModelListRefresher(TSharedPtr<FViewModel> InRoot);
 	FHierarchicalModelListRefresher(TSharedPtr<FViewModel> InRoot, FViewModelChildren InExistingChildren);
 
 	bool IsValid() const;
@@ -79,7 +78,9 @@ struct FHierarchicalModelListRefresher
 		return nullptr;
 	}
 
-	void RecycleChildren(const TSharedPtr<FViewModel>& InModel, FViewModelChildren InExistingChildren);
+private:
+
+	void ConditionalRecycleChildren(const TSharedPtr<FViewModel>& InModel, FViewModelChildren InExistingChildren);
 
 	struct FListData
 	{
diff --git a/Engine/Source/Editor/Sequencer/Private/MVVM/ViewModels/TrackRowModel.cpp b/Engine/Source/Editor/Sequencer/Private/MVVM/ViewModels/TrackRowModel.cpp
--- a/Engine/Source/Editor/Sequencer/Private/MVVM/ViewModels/TrackRowModel.cpp
+++ b/Engine/Source/Editor/Sequencer/Private/MVVM/ViewModels/TrackRowModel.cpp
@@ -74,6 +74,11 @@ void FTrackRowModel::Initialize()
 	}
 }
 
+FViewModelChildren FTrackRowModel::GetTopLevelChannels()
+{
+	return GetChildrenForList(&TopLevelChannelList);
+}
+
 UMovieSceneTrack* FTrackRowModel::GetTrack() const
 {
 	return WeakTrack.Get();
diff --git a/Engine/Source/Editor/Sequencer/Public/MVVM/ViewModels/TrackRowModel.h b/Engine/Source/Editor/Sequencer/Public/MVVM/ViewModels/TrackRowModel.h
--- a/Engine/Source/Editor/Sequencer/Public/MVVM/ViewModels/TrackRowModel.h
+++ b/Engine/Source/Editor/Sequencer/Public/MVVM/ViewModels/TrackRowModel.h
@@ -41,7 +41,7 @@ public:
 
 	void Initialize();
 
-	TSharedPtr<ISequencerTrackEditor> GetTrackEditor() const { return TrackEditor; }
+	FViewModelChildren GetTopLevelChannels();
 
 	/*~ FOutlinerItemModel */
 	void BuildContextMenu(FMenuBuilder& MenuBuilder) override;
@@ -68,6 +68,7 @@ public:
 	UMovieSceneTrack* GetTrack() const override;
 	int32 GetRowIndex() const override;
 	FViewModelChildren GetSectionModels() override;
+	TSharedPtr<ISequencerTrackEditor> GetTrackEditor() const override { return TrackEditor; }
 
 	/*~ ITrackAreaExtension */
 	FTrackAreaParameters GetTrackAreaParameters() const override;
