Sequencer: fix a crash when recompiling a user widget while sequencer is open

diff --git a/Engine/Source/Editor/Sequencer/Private/Sequencer.cpp b/Engine/Source/Editor/Sequencer/Private/Sequencer.cpp
--- a/Engine/Source/Editor/Sequencer/Private/Sequencer.cpp
+++ b/Engine/Source/Editor/Sequencer/Private/Sequencer.cpp
@@ -391,12 +391,39 @@ void FSequencer::InitSequencer(const FSequencerInitParams& InitParams, const TSh
 				}
 			}
 
-			//Reset Bindings for replaced objects.
+			// If our playback client is getting killed (e.g. it's a Blueprint that has been recompiled), we have
+			// to bail out.
+			if (CachedPlaybackContext.IsStale())
+			{
+				GEditor->GetEditorSubsystem<UAssetEditorSubsystem>()->CloseAllEditorsForAsset(GetRootMovieSceneSequence());
+				return;
+			}
+
+			// Update our playback context and client if they were replaced.
+			// This needs to happen before we resting our object bindings (below) because object binding resolution
+			// can sometimes try to access the playback context.
+			if (UObject* OldPlaybackContext = CachedPlaybackContext.Get())
+			{
+				if (UObject* const* NewPlaybackContext = ReplacementMap.Find(OldPlaybackContext))
+				{
+					CachedPlaybackContext = *NewPlaybackContext;
+				}
+			}
+			if (UObject* OldPlaybackClient = CachedPlaybackClient.GetObject())
+			{
+				if (UObject* const* NewPlaybackClient = ReplacementMap.Find(OldPlaybackClient))
+				{
+					CachedPlaybackClient = *NewPlaybackClient;
+				}
+			}
+
+			// Reset Bindings for replaced objects.
 			for (TPair<UObject*, UObject*> ReplacedObject : ReplacementMap)
 			{
 				FGuid Guid = GetHandleToObject(ReplacedObject.Key, false);
 			}
 
+			// Replace pointers inside our pre-animate state storages.
 			PreAnimatedState.OnObjectsReplaced(ReplacementMap);
 
 		});
