Sequencer: fix mac builds

diff --git a/Engine/Source/Editor/Sequencer/Private/Sequencer.cpp b/Engine/Source/Editor/Sequencer/Private/Sequencer.cpp
--- a/Engine/Source/Editor/Sequencer/Private/Sequencer.cpp
+++ b/Engine/Source/Editor/Sequencer/Private/Sequencer.cpp
@@ -391,14 +391,6 @@ void FSequencer::InitSequencer(const FSequencerInitParams& InitParams, const TSh
 				}
 			}
 
-			// If our playback client is getting killed (e.g. it's a Blueprint that has been recompiled), we have
-			// to bail out.
-			if (CachedPlaybackContext.IsStale())
-			{
-				GEditor->GetEditorSubsystem<UAssetEditorSubsystem>()->CloseAllEditorsForAsset(GetRootMovieSceneSequence());
-				return;
-			}
-
 			// Update our playback context and client if they were replaced.
 			// This needs to happen before we resting our object bindings (below) because object binding resolution
 			// can sometimes try to access the playback context.
@@ -413,10 +405,19 @@ void FSequencer::InitSequencer(const FSequencerInitParams& InitParams, const TSh
 			{
 				if (UObject* const* NewPlaybackClient = ReplacementMap.Find(OldPlaybackClient))
 				{
-					CachedPlaybackClient = *NewPlaybackClient;
+					IMovieScenePlaybackClient* ClientInterface = Cast<IMovieScenePlaybackClient>(*NewPlaybackClient);
+					CachedPlaybackClient = TWeakInterfacePtr<IMovieScenePlaybackClient>(ClientInterface);
 				}
 			}
 
+			// If our playback client is getting killed (e.g. it's a Blueprint that has been recompiled), we have
+			// to bail out.
+			if (CachedPlaybackContext.IsStale())
+			{
+				GEditor->GetEditorSubsystem<UAssetEditorSubsystem>()->CloseAllEditorsForAsset(GetRootMovieSceneSequence());
+				return;
+			}
+
 			// Reset Bindings for replaced objects.
 			for (TPair<UObject*, UObject*> ReplacedObject : ReplacementMap)
 			{
