Sequencer: fix crashes when a UserWidget is being recompiled with an animation open

diff --git a/Engine/Source/Editor/Sequencer/Private/Sequencer.cpp b/Engine/Source/Editor/Sequencer/Private/Sequencer.cpp
--- a/Engine/Source/Editor/Sequencer/Private/Sequencer.cpp
+++ b/Engine/Source/Editor/Sequencer/Private/Sequencer.cpp
@@ -417,6 +417,14 @@ void FSequencer::InitSequencer(const FSequencerInitParams& InitParams, const TSh
 				GEditor->GetEditorSubsystem<UAssetEditorSubsystem>()->CloseAllEditorsForAsset(GetRootMovieSceneSequence());
 				return;
 			}
+			if (UObject* PlaybackContext = CachedPlaybackContext.Get())
+			{
+				if (ReplacementMap.Contains(PlaybackContext->GetClass()))
+				{
+					GEditor->GetEditorSubsystem<UAssetEditorSubsystem>()->CloseAllEditorsForAsset(GetRootMovieSceneSequence());
+					return;
+				}
+			}
 
 			// Reset Bindings for replaced objects.
 			for (TPair<UObject*, UObject*> ReplacedObject : ReplacementMap)
diff --git a/Engine/Source/Runtime/UMG/Private/Animation/WidgetAnimation.cpp b/Engine/Source/Runtime/UMG/Private/Animation/WidgetAnimation.cpp
--- a/Engine/Source/Runtime/UMG/Private/Animation/WidgetAnimation.cpp
+++ b/Engine/Source/Runtime/UMG/Private/Animation/WidgetAnimation.cpp
@@ -233,6 +233,11 @@ void UWidgetAnimation::LocateBoundObjects(const FGuid& ObjectId, UObject* InCont
 	}
 
 	UUserWidget* PreviewWidget = CastChecked<UUserWidget>(InContext);
+	if (PreviewWidget->WidgetTree == nullptr)
+	{
+		return;
+	}
+
 	for (const FWidgetAnimationBinding& Binding : AnimationBindings)
 	{
 		if (Binding.AnimationGuid == ObjectId)
