Sequencer: Fix OnTimeDisplayFormatChanged

#rnx

diff --git a/Engine/Source/Editor/Sequencer/Private/Sequencer.cpp b/Engine/Source/Editor/Sequencer/Private/Sequencer.cpp
--- a/Engine/Source/Editor/Sequencer/Private/Sequencer.cpp
+++ b/Engine/Source/Editor/Sequencer/Private/Sequencer.cpp
@@ -330,6 +330,7 @@ void FSequencer::InitSequencer(const FSequencerInitParams& InitParams, const TSh
 
 	Settings->GetOnEvaluateSubSequencesInIsolationChanged().AddSP(this, &FSequencer::RestorePreAnimatedState);
 	Settings->GetOnShowSelectedNodesOnlyChanged().AddSP(this, &FSequencer::OnSelectedNodesOnlyChanged);
+	Settings->GetOnTimeDisplayFormatChanged().AddSP(this, &FSequencer::OnTimeDisplayFormatChanged);
 
 	ObjectBindingTagCache = MakeUnique<FObjectBindingTagCache>();
 
@@ -590,9 +591,6 @@ void FSequencer::InitSequencer(const FSequencerInitParams& InitParams, const TSh
 		ObjectBindings.Add(ObjectBinding);
 	}
 
-	TSharedPtr<FFrameNumberInterface> FrameNumberInterface = StaticCastSharedPtr<FFrameNumberInterface>(GetNumericTypeInterface().ToSharedPtr());
-	Settings->GetOnTimeDisplayFormatChanged().AddRaw(FrameNumberInterface.Get(), &FFrameNumberInterface::DisplayFormatChanged);
-
 	FMovieSceneObjectBindingIDCustomization::BindTo(AsShared());
 
 	ZoomAnimation = FCurveSequence();
@@ -701,11 +699,6 @@ FSequencer::~FSequencer()
 
 void FSequencer::Close()
 {
-	if (Settings && SequencerWidget)
-	{
-		Settings->GetOnTimeDisplayFormatChanged().RemoveAll(&SequencerWidget->GetNumericTypeInterface().Get());
-	}
-
 	for (FLevelEditorViewportClient* LevelVC : GEditor->GetLevelViewportClients())
 	{
 		if (LevelVC != nullptr)
@@ -6720,6 +6713,15 @@ void FSequencer::OnSelectedNodesOnlyChanged()
 	SynchronizeSequencerSelectionWithExternalSelection();
 }
 
+void FSequencer::OnTimeDisplayFormatChanged()
+{
+	TSharedPtr<FFrameNumberInterface> FrameNumberInterface = StaticCastSharedPtr<FFrameNumberInterface>(GetNumericTypeInterface().ToSharedPtr());
+	if (FrameNumberInterface.IsValid())
+	{
+		FrameNumberInterface->DisplayFormatChanged();
+	}
+}
+
 void FSequencer::ZoomToFit()
 {
 	using namespace UE::Sequencer;
diff --git a/Engine/Source/Editor/Sequencer/Private/Sequencer.h b/Engine/Source/Editor/Sequencer/Private/Sequencer.h
--- a/Engine/Source/Editor/Sequencer/Private/Sequencer.h
+++ b/Engine/Source/Editor/Sequencer/Private/Sequencer.h
@@ -602,6 +602,8 @@ public:
 
 	void OnSelectedNodesOnlyChanged();
 
+	void OnTimeDisplayFormatChanged();
+
 	/** Will create a custom menu if FSequencerViewParams::OnBuildCustomContextMenuForGuid is specified. */
 	void BuildCustomContextMenuForGuid(FMenuBuilder& MenuBuilder, FGuid ObjectBinding);
 
