Sequencer: When a track or binding is added, adjust sorting orders if there is going to be clash. The incoming track should be sorted to the bottom amongst the other tracks of the same class and name.

diff --git a/Engine/Source/Editor/Sequencer/Private/Sequencer.cpp b/Engine/Source/Editor/Sequencer/Private/Sequencer.cpp
--- a/Engine/Source/Editor/Sequencer/Private/Sequencer.cpp
+++ b/Engine/Source/Editor/Sequencer/Private/Sequencer.cpp
@@ -5630,6 +5630,8 @@ void FSequencer::PostUndo(bool bSuccess)
 		OwnerMovieScene->SortMarkedFrames();
 	}
 
+	NodeTree->SortAllNodesAndDescendants();
+
 	OnActivateSequenceEvent.Broadcast(ActiveTemplateIDs.Top());
 }
 
@@ -8837,18 +8839,93 @@ void FSequencer::OnAddFolder()
 	NotifyMovieSceneDataChanged( EMovieSceneDataChangeType::MovieSceneStructureItemAdded );
 }
 
-void FSequencer::OnAddTrack(const TWeakObjectPtr<UMovieSceneTrack>& InTrack, const FGuid& ObjectBinding)
+void FixSortingOrders(UMovieSceneTrack* InTrack, const TArray<UMovieSceneTrack*>& Tracks)
+{
+	// Fix the sorting orders if InTrack is going to conflict with any of the given Tracks (ie. same class and same name)
+	TOptional<int32> MaxSortingOrder;
+	for (UMovieSceneTrack* Track : Tracks)
+	{
+		if (Track != InTrack && Track->GetClass() == InTrack->GetClass() && Track->GetDisplayName().EqualTo(InTrack->GetDisplayName()))
+		{
+			if (MaxSortingOrder.IsSet())
+			{
+				MaxSortingOrder = FMath::Max(MaxSortingOrder.GetValue(), Track->GetSortingOrder());
+			}
+			else
+			{
+				MaxSortingOrder = Track->GetSortingOrder();
+			}
+		}
+	}
+
+	if (MaxSortingOrder.IsSet())
+	{
+		InTrack->Modify();
+		InTrack->SetSortingOrder(MaxSortingOrder.GetValue() + 1);
+	}
+}
+
+void FixSortingOrders(FMovieSceneBinding* InBinding, UMovieScene* MovieScene)
+{
+	// Fix the sorting orders if InBinding is going to conflict with any of the existing bindings (ie. same class and same name)
+	TOptional<int32> MaxSortingOrder;
+	for (const FMovieSceneBinding& Binding : MovieScene->GetBindings())
+	{
+		if (&Binding != InBinding && Binding.GetName() == InBinding->GetName())
+		{
+			if (MaxSortingOrder.IsSet())
+			{
+				MaxSortingOrder = FMath::Max(MaxSortingOrder.GetValue(), Binding.GetSortingOrder());
+			}
+			else
+			{
+				MaxSortingOrder = Binding.GetSortingOrder();
+			}
+		}
+	}
+
+	if (MaxSortingOrder.IsSet())
+	{
+		InBinding->SetSortingOrder(MaxSortingOrder.GetValue() + 1);
+	}
+}
+
+void FSequencer::OnAddBinding(const FGuid& ObjectBinding, UMovieScene* MovieScene)
+{
+	FMovieSceneBinding* Binding = MovieScene->FindBinding(ObjectBinding);
+	if (Binding)
+	{
+		FixSortingOrders(Binding, MovieScene);
+	}
+
+	NodeTree->SortAllNodesAndDescendants();
+}
+
+void FSequencer::OnAddTrack(const TWeakObjectPtr<UMovieSceneTrack>&InTrack, const FGuid & ObjectBinding)
 {
 	if (!ensureAlwaysMsgf(InTrack.IsValid(), TEXT("Attempted to add a null UMovieSceneTrack to Sequencer. This should never happen.")))
 	{
 		return;
 	}
 
+	UMovieSceneSequence* Sequence = GetFocusedMovieSceneSequence();
+	UMovieScene* MovieScene = Sequence->GetMovieScene();
+
+	FMovieSceneBinding* Binding = MovieScene->FindBinding(ObjectBinding);
+	if (Binding)
+	{
+		FixSortingOrders(InTrack.Get(), Binding->GetTracks());
+	}
+	else if (MovieScene->ContainsTrack(*InTrack))
+	{	
+		FixSortingOrders(InTrack.Get(), MovieScene->GetTracks());
+	}
+
 	FString NewNodePath;
 
 	// If they specified an object binding it's being added to, we don't add it to a folder since we can't have it existing
 	// as a children of two places at once.
-	if(!GetFocusedMovieSceneSequence()->GetMovieScene()->FindBinding(ObjectBinding))
+	if(!Binding)
 	{
 		TArray<UMovieSceneFolder*> SelectedParentFolders;
 		CalculateSelectedFolderAndPath(SelectedParentFolders, NewNodePath);
@@ -8873,6 +8950,8 @@ void FSequencer::OnAddTrack(const TWeakObjectPtr<UMovieSceneTrack>& InTrack, con
 		SelectSection(InTrack->GetAllSections()[0]);
 	}
 	ThrobSectionSelection();
+
+	NodeTree->SortAllNodesAndDescendants();
 }
 
 
diff --git a/Engine/Source/Editor/Sequencer/Private/Sequencer.h b/Engine/Source/Editor/Sequencer/Private/Sequencer.h
--- a/Engine/Source/Editor/Sequencer/Private/Sequencer.h
+++ b/Engine/Source/Editor/Sequencer/Private/Sequencer.h
@@ -1056,6 +1056,9 @@ private:
 	/** Handles loading in previously recorded data. */
 	void OnLoadRecordedData();
 	
+	/** Adds the binding and selects the binding, throbs it. */
+	void OnAddBinding(const FGuid& ObjectBinding, UMovieScene* MovieScene) override;
+
 	/** Adds the track to the selected folder (if FGuid is invalid) and selects the track, throbs it, and notifies the sequence to rebuild any necessary data. */
 	void OnAddTrack(const TWeakObjectPtr<UMovieSceneTrack>& InTrack, const FGuid& ObjectBinding) override;
 
diff --git a/Engine/Source/Editor/Sequencer/Private/SequencerUtilities.cpp b/Engine/Source/Editor/Sequencer/Private/SequencerUtilities.cpp
--- a/Engine/Source/Editor/Sequencer/Private/SequencerUtilities.cpp
+++ b/Engine/Source/Editor/Sequencer/Private/SequencerUtilities.cpp
@@ -2438,6 +2438,8 @@ bool FSequencerUtilities::PasteBindings(const FString& TextToImport, TSharedRef<
 	for (auto BindingPasted : BindingsPasted)
 	{
 		OutBindings.Add(FMovieSceneBindingProxy(BindingPasted.GetObjectGuid(), Sequence));
+		
+		Sequencer->OnAddBinding(BindingPasted.GetObjectGuid(), MovieScene);
 	}
 
 	return true; 
@@ -2532,6 +2534,8 @@ FGuid FSequencerUtilities::CreateBinding(TSharedRef<ISequencer> Sequencer, UObje
 		Sequencer->OnActorAddedToSequencer().Broadcast(ParentActorAdded, ParentGuid);
 	}
 
+	Sequencer->OnAddBinding(PossessableGuid, OwnerMovieScene);
+
 	return PossessableGuid;
 }
 
diff --git a/Engine/Source/Editor/Sequencer/Public/ISequencer.h b/Engine/Source/Editor/Sequencer/Public/ISequencer.h
--- a/Engine/Source/Editor/Sequencer/Public/ISequencer.h
+++ b/Engine/Source/Editor/Sequencer/Public/ISequencer.h
@@ -319,6 +319,11 @@ public:
 	 */
 	virtual TArray<FGuid> AddActors(const TArray<TWeakObjectPtr<AActor> >& InActors, bool bSelectActors = true) = 0;
 
+	/**
+	 * Should be called after adding a binding to the MovieScene.
+	 */
+	virtual void OnAddBinding(const FGuid& ObjectBinding, UMovieScene* MovieScene) = 0;
+
 	/**
 	 * Should be called after adding a track to the MovieScene. This will set the specified track as your current selection
 	 * cause it to throb, notify the sequence to rebuild any data required. The track will be added to the selected folder
diff --git a/Engine/Source/Editor/SequencerCore/Private/MVVM/Extensions/ISortableExtension.cpp b/Engine/Source/Editor/SequencerCore/Private/MVVM/Extensions/ISortableExtension.cpp
--- a/Engine/Source/Editor/SequencerCore/Private/MVVM/Extensions/ISortableExtension.cpp
+++ b/Engine/Source/Editor/SequencerCore/Private/MVVM/Extensions/ISortableExtension.cpp
@@ -29,23 +29,10 @@ bool FSortingKey::ComparePriorityFirst(const FSortingKey& A, const FSortingKey&
 			return Compare < 0;
 		}
 
-		if (A.CustomOrder >= 0 && B.CustomOrder >= 0)
-		{
-			// Both items have custom orders, use those
-			return A.CustomOrder < B.CustomOrder;
-		}
-		else if (A.CustomOrder >= 0 || B.CustomOrder >= 0)
-		{
-			// Only one item has a custom order, list it first
-			return A.CustomOrder > B.CustomOrder;
-		}
-		else
-		{
-			// Equal priorities and identifiers, and custom orders are unset
-			// Just compare unset custom orders (zero or negative values) to be consistent
-			// with the equality operator
-			return A.CustomOrder < B.CustomOrder;
-		}
+		// Note in this case, it's possible that one of the items does not have a custom order (ie. -1). 
+		// That's fine because -1 will be sorted to the top and generally if new items are added and 
+		// there is ambiguity, we would set a max custom order.
+		return A.CustomOrder < B.CustomOrder;
 	}
 	return A.Priority < B.Priority;
 }
