Sequencer: Track Gutter UI Framework, Outliner Column Creation and Settings

diff --git a/Engine/Source/Editor/Sequencer/Private/SSequencer.cpp b/Engine/Source/Editor/Sequencer/Private/SSequencer.cpp
--- a/Engine/Source/Editor/Sequencer/Private/SSequencer.cpp
+++ b/Engine/Source/Editor/Sequencer/Private/SSequencer.cpp
@@ -2173,7 +2173,10 @@ TSharedRef<SWidget> SSequencer::MakeViewMenu()
 	};
 
 	// Menu Entry for Outliner Column Visibilities
-	MenuBuilder.AddSubMenu(LOCTEXT("ColumnVisibilityHeader", "Columns"), FText::GetEmpty(), FNewMenuDelegate::CreateRaw(this, &SSequencer::FillColumnVisibilityMenu));
+	if (OutlinerColumnVisibilities.Num() > 0)
+	{
+		MenuBuilder.AddSubMenu(LOCTEXT("ColumnVisibilityHeader", "Columns"), FText::GetEmpty(), FNewMenuDelegate::CreateRaw(this, &SSequencer::FillColumnVisibilityMenu));
+	}
 
 	MenuBuilder.AddWidget(
 		SNew(SHorizontalBox)
