Sequencer: Track Gutter UI Framework, Outliner Column Creation and Settings

diff --git a/Engine/Source/Editor/Sequencer/Public/SequencerSettings.h b/Engine/Source/Editor/Sequencer/Public/SequencerSettings.h
--- a/Engine/Source/Editor/Sequencer/Public/SequencerSettings.h
+++ b/Engine/Source/Editor/Sequencer/Public/SequencerSettings.h
@@ -103,6 +103,8 @@ struct FColumnVisibilitySetting
 	{}
 
 	FColumnVisibilitySetting()
+		: ColumnName(NAME_None)
+		, bIsVisible(false)
 	{}
 };
 
