Check GetMovieScene() pointer before checking for read only status.

#rnx

diff --git a/Engine/Source/Editor/Sequencer/Private/Sequencer.cpp b/Engine/Source/Editor/Sequencer/Private/Sequencer.cpp
--- a/Engine/Source/Editor/Sequencer/Private/Sequencer.cpp
+++ b/Engine/Source/Editor/Sequencer/Private/Sequencer.cpp
@@ -5471,7 +5471,7 @@ UObject* FSequencer::GetCurrentAsset() const
 
 bool FSequencer::IsReadOnly() const
 {
-	return bReadOnly || (GetFocusedMovieSceneSequence() && GetFocusedMovieSceneSequence()->GetMovieScene()->IsReadOnly());
+	return bReadOnly || (GetFocusedMovieSceneSequence() && GetFocusedMovieSceneSequence()->GetMovieScene() && GetFocusedMovieSceneSequence()->GetMovieScene()->IsReadOnly());
 }
 
 FGuid FSequencer::MakeNewSpawnable(UObject& Object, UActorFactory* ActorFactory, bool bSetupDefaults)
