Sequencer: Outliner Column Editor Extensions

These editor extensions are used by outliner column implementations to query and modify pin, mute, lock, and solo.

diff --git a/Engine/Source/Editor/Sequencer/Private/MVVM/LockEditorExtension.cpp b/Engine/Source/Editor/Sequencer/Private/MVVM/LockEditorExtension.cpp
new file mode 100644
--- /dev/null
+++ b/Engine/Source/Editor/Sequencer/Private/MVVM/LockEditorExtension.cpp
@@ -0,0 +1,110 @@
+// Copyright Epic Games, Inc. All Rights Reserved.
+
+#include "MVVM/LockEditorExtension.h"
+
+#include "MovieScene.h"
+#include "MVVM/Extensions/IOutlinerExtension.h"
+#include "MVVM/Selection/Selection.h"
+#include "MVVM/ViewModels/SequencerEditorViewModel.h"
+#include "ScopedTransaction.h"
+#include "Sequencer.h"
+#include "SequencerCommonHelpers.h"
+
+namespace UE::Sequencer
+{
+
+FLockEditorExtension::FLockEditorExtension()
+{
+
+}
+
+void FLockEditorExtension::OnCreated(TSharedRef<FViewModel> InWeakOwner)
+{
+	ensureMsgf(!WeakOwnerModel.Pin().IsValid(), TEXT("This extension was already created!"));
+	WeakOwnerModel = InWeakOwner->CastThisShared<FSequencerEditorViewModel>();
+}
+
+bool FLockEditorExtension::IsNodeLocked(TWeakViewModelPtr<IOutlinerExtension> InWeakOutlinerExtension) const
+{
+	FViewModelPtr Item = InWeakOutlinerExtension.Pin();
+
+	// Locked if all sections are locked
+	int32 NumSections = 0;
+	TSet<TWeakObjectPtr<UMovieSceneSection>> Sections;
+	SequencerHelpers::GetAllSections(Item, Sections);
+
+	for (auto Section : Sections)
+	{
+		if (!Section->IsLocked())
+		{
+			return false;
+		}
+		++NumSections;
+	}
+	return NumSections > 0;
+}
+
+void FLockEditorExtension::SetNodeLocked(TWeakViewModelPtr<IOutlinerExtension> InWeakOutlinerExtension, bool bInIsLocked)
+{
+	const FScopedTransaction Transaction(NSLOCTEXT("Sequencer", "SetNodeLocked", "Set Node Locked"));
+
+	TSharedPtr<FSequencerEditorViewModel> EditorViewModel = WeakOwnerModel.Pin();
+	TViewModelPtr<IOutlinerExtension> OutlinerItem = InWeakOutlinerExtension.Pin();
+	FViewModelPtr Item = OutlinerItem;
+
+	if (OutlinerItem->GetSelectionState() == EOutlinerSelectionState::SelectedDirectly)
+	{
+		// modify all selected items
+		for (FViewModelPtr OutlinerNode : EditorViewModel->GetSelection()->Outliner)
+		{
+			TSet<TWeakObjectPtr<UMovieSceneSection> > Sections;
+			SequencerHelpers::GetAllSections(OutlinerNode, Sections);
+
+			for (auto Section : Sections)
+			{
+				Section->Modify();
+				Section->SetIsLocked(bInIsLocked);
+			}
+		}
+	}
+	else
+	{
+		// only one unselected item was toggled, toggle just that node
+		TSet<TWeakObjectPtr<UMovieSceneSection> > Sections;
+		SequencerHelpers::GetAllSections(Item, Sections);
+
+		for (auto Section : Sections)
+		{
+			Section->Modify();
+			Section->SetIsLocked(bInIsLocked);
+		}
+	}
+}
+
+bool FLockEditorExtension::IsNodeLockable(TWeakViewModelPtr<IOutlinerExtension> InWeakOutlinerExtension) const
+{
+	FViewModelPtr Item = InWeakOutlinerExtension.Pin();
+
+	TSet<TWeakObjectPtr<UMovieSceneSection> > Sections;
+	SequencerHelpers::GetAllSections(Item, Sections);
+
+	return Sections.Num() > 0;
+}
+
+bool FLockEditorExtension::HasLockedChildNode(TWeakViewModelPtr<IOutlinerExtension> InWeakOutlinerExtension) const
+{
+	TViewModelPtr<IOutlinerExtension> OutlinerItem = InWeakOutlinerExtension.Pin();
+
+	for (const TViewModelPtr<IOutlinerExtension>& Child : OutlinerItem.AsModel()->GetDescendantsOfType<IOutlinerExtension>())
+	{
+		if (IsNodeLocked(Child))
+		{
+			return true;
+		}
+	}
+
+	return false;
+}
+
+} // namespace UE::Sequencer
+
diff --git a/Engine/Source/Editor/Sequencer/Private/MVVM/LockEditorExtension.h b/Engine/Source/Editor/Sequencer/Private/MVVM/LockEditorExtension.h
new file mode 100644
--- /dev/null
+++ b/Engine/Source/Editor/Sequencer/Private/MVVM/LockEditorExtension.h
@@ -0,0 +1,54 @@
+// Copyright Epic Games, Inc. All Rights Reserved.
+
+#pragma once
+
+#include "MVVM/ViewModels/ViewModel.h"
+#include "MVVM/ViewModelPtr.h"
+#include "MVVM/ViewModelTypeID.h"
+#include "MVVM/Extensions/DynamicExtensionContainer.h"
+
+namespace UE::Sequencer
+{
+
+class FSequencerEditorViewModel;
+class IOutlinerExtension;
+
+/**
+* Extension that handles the locked state of items.
+*/
+class FLockEditorExtension : public IDynamicExtension
+{
+public:
+
+	UE_SEQUENCER_DECLARE_VIEW_MODEL_TYPE_ID(FLockEditorExtension)
+
+	FLockEditorExtension();
+
+	/** Called when the extension is created on a data model. */
+	virtual void OnCreated(TSharedRef<FViewModel> InWeakOwner) override;
+
+	/** Returns whether a given item is locked. */
+	bool IsNodeLocked(TWeakViewModelPtr<IOutlinerExtension> InWeakOutlinerExtension) const;
+
+	/** 
+	* Modifies an item and it's child sections to be locked.
+	* This operation applies to all selected items if the modified item is selected.
+	* @param bInIsLocked - Lock state to set the item to
+	* @param InWeakOutlinerExtension - Item to lock
+	*/
+	void SetNodeLocked(TWeakViewModelPtr<IOutlinerExtension> InWeakOutlinerExtension, bool bInIsLocked);
+
+	/** Returns whether or not this item has any sections to lock. */
+	bool IsNodeLockable(TWeakViewModelPtr<IOutlinerExtension> InWeakOutlinerExtension) const;
+
+	/** Returns whether or not a child item of this item is locked. */
+	bool HasLockedChildNode(TWeakViewModelPtr<IOutlinerExtension> InWeakOutlinerExtension) const;
+
+private:
+
+	/** The sequencer editor we are extending. */
+	TWeakPtr<FSequencerEditorViewModel> WeakOwnerModel;
+};
+
+} // namespace UE::Sequencer
+
diff --git a/Engine/Source/Editor/Sequencer/Private/MVVM/MuteEditorExtension.cpp b/Engine/Source/Editor/Sequencer/Private/MVVM/MuteEditorExtension.cpp
new file mode 100644
--- /dev/null
+++ b/Engine/Source/Editor/Sequencer/Private/MVVM/MuteEditorExtension.cpp
@@ -0,0 +1,128 @@
+// Copyright Epic Games, Inc. All Rights Reserved.
+
+#include "MVVM/MuteEditorExtension.h"
+
+#include "MovieScene.h"
+#include "MVVM/Extensions/IOutlinerExtension.h"
+#include "MVVM/Extensions/IMutableExtension.h"
+#include "MVVM/Selection/Selection.h"
+#include "MVVM/ViewModels/SequencerEditorViewModel.h"
+#include "ScopedTransaction.h"
+#include "Sequencer.h"
+
+namespace UE::Sequencer
+{
+
+FMuteEditorExtension::FMuteEditorExtension()
+{
+
+}
+
+void FMuteEditorExtension::OnCreated(TSharedRef<FViewModel> InWeakOwner)
+{
+	ensureMsgf(!WeakOwnerModel.Pin().IsValid(), TEXT("This extension was already created!"));
+	WeakOwnerModel = InWeakOwner->CastThisShared<FSequencerEditorViewModel>();
+}
+
+void FMuteEditorExtension::SetNodeMuted(TWeakViewModelPtr<IOutlinerExtension> InWeakOutlinerExtension, bool bInIsMuted)
+{              
+	const FScopedTransaction Transaction(NSLOCTEXT("Sequencer", "SetNodeMuted", "Set Node Muted"));
+
+	TViewModelPtr<IOutlinerExtension> OutlinerItem = InWeakOutlinerExtension.Pin();
+	TSharedPtr<FSequencerEditorViewModel> EditorViewModel = WeakOwnerModel.Pin();
+	TSharedPtr<FSequencer> Sequencer = EditorViewModel ? EditorViewModel->GetSequencerImpl() : nullptr;
+		
+	if (Sequencer)
+	{
+		UMovieScene* MovieScene = Sequencer->GetFocusedMovieSceneSequence()->GetMovieScene();
+
+		if (MovieScene->IsReadOnly())
+		{
+			return;
+		}
+
+		TArray<FString>& MuteNodes = MovieScene->GetMuteNodes();
+
+		MovieScene->Modify();
+
+		if (OutlinerItem->GetSelectionState() == EOutlinerSelectionState::SelectedDirectly)
+		{
+			// if selected, modify all selected items
+			for (FViewModelPtr Node : Sequencer->GetViewModel()->GetSelection()->Outliner)
+			{
+				FString NodePath = IOutlinerExtension::GetPathName(Node);
+				if (bInIsMuted)
+				{
+					// Mark Mute, being careful as we might be re-marking an already Mute node
+					MuteNodes.AddUnique(NodePath);
+				}
+				else
+				{
+					// UnMute
+					MuteNodes.Remove(NodePath);
+				}
+			}
+		}
+		else
+		{
+			// modify only the toggled item
+			TSharedPtr<FViewModel> Item = OutlinerItem;
+			FString NodePath = IOutlinerExtension::GetPathName(Item);
+			if (bInIsMuted)
+			{
+				// Mark Mute, being careful as we might be re-marking an already Mute node
+				MuteNodes.AddUnique(NodePath);
+			}
+			else
+			{
+				// Unmute the toggled item
+				MuteNodes.Remove(NodePath);
+			}
+		}
+	}
+}
+
+bool FMuteEditorExtension::IsNodeMutable(TWeakViewModelPtr<IOutlinerExtension> InWeakOutlinerExtension) const
+{
+	TViewModelPtr<IOutlinerExtension> OutlinerItem = InWeakOutlinerExtension.Pin();
+	TSharedPtr<FSequencerEditorViewModel> EditorViewModel = WeakOwnerModel.Pin();
+	TSharedPtr<FSequencer> Sequencer = EditorViewModel ? EditorViewModel->GetSequencerImpl() : nullptr;
+		
+	if (Sequencer)
+	{
+		TSharedPtr<FViewModel> Item = OutlinerItem;
+		FString NodePath = IOutlinerExtension::GetPathName(Item);
+		return Sequencer->GetNodeTree()->GetNodeAtPath(NodePath) != nullptr;
+	}
+
+	return false;
+}
+
+bool FMuteEditorExtension::HasMutedChildNode(TWeakViewModelPtr<IOutlinerExtension> InWeakOutlinerExtension) const
+{
+	using namespace UE::Sequencer;
+
+	TViewModelPtr<IOutlinerExtension> OutlinerItem = InWeakOutlinerExtension.Pin();
+
+	for (const TViewModelPtr<IMutableExtension>& Child : OutlinerItem.AsModel()->GetDescendantsOfType<IMutableExtension>())
+	{
+		if (Child->IsMuted())
+		{
+			return true;
+		}
+	}
+
+	return false;
+}
+
+void FMuteEditorExtension::OnPostMute()
+{
+	TSharedPtr<FSequencerEditorViewModel> EditorViewModel = WeakOwnerModel.Pin();
+	TSharedPtr<FSequencer> Sequencer = EditorViewModel ? EditorViewModel->GetSequencerImpl() : nullptr;
+	if (Sequencer)
+	{
+		Sequencer->RefreshTree();
+	}
+}
+
+} // namespace UE::Sequencer
diff --git a/Engine/Source/Editor/Sequencer/Private/MVVM/MuteEditorExtension.h b/Engine/Source/Editor/Sequencer/Private/MVVM/MuteEditorExtension.h
new file mode 100644
--- /dev/null
+++ b/Engine/Source/Editor/Sequencer/Private/MVVM/MuteEditorExtension.h
@@ -0,0 +1,54 @@
+// Copyright Epic Games, Inc. All Rights Reserved.
+
+#pragma once
+
+#include "MVVM/ViewModels/ViewModel.h"
+#include "MVVM/ViewModelPtr.h"
+#include "MVVM/ViewModelTypeID.h"
+#include "MVVM/Extensions/DynamicExtensionContainer.h"
+
+namespace UE::Sequencer
+{
+
+class FSequencerEditorViewModel;
+class IOutlinerExtension;
+
+/**
+* Extension that handles the mute state of items.
+*/
+class FMuteEditorExtension : public IDynamicExtension
+{
+public:
+
+	UE_SEQUENCER_DECLARE_VIEW_MODEL_TYPE_ID(FMuteEditorExtension)
+
+	FMuteEditorExtension();
+
+	/** Called when the extension is created on a data model. */
+	virtual void OnCreated(TSharedRef<FViewModel> InWeakOwner) override;
+
+	/**
+	* Modifies an item to be muted.
+	* This operation applies to all selected items if the modified item is selected.
+	* @param bInIsMuted - Muted state to set the item to
+	* @param InWeakOutlinerExtension - Item to mute
+	*/
+	void SetNodeMuted(TWeakViewModelPtr<IOutlinerExtension> InWeakOutlinerExtension, bool bInIsMuted);
+
+	/** Returns whether or not the item has a valid node path to mute. */
+	bool IsNodeMutable(TWeakViewModelPtr<IOutlinerExtension> InWeakOutlinerExtension) const;
+
+	/** Returns whether or not the item has any muted child items. */
+	bool HasMutedChildNode(TWeakViewModelPtr<IOutlinerExtension> InWeakOutlinerExtension) const;
+
+	/** Refreshes the Sequencer Tree, required to ensure that muting an item doesn't affect the visibility of soloed items. */
+	void OnPostMute();
+
+private:
+
+	/** The sequencer editor we are extending */
+	TWeakPtr<FSequencerEditorViewModel> WeakOwnerModel;
+};
+
+} // namespace UE::Sequencer
+
diff --git a/Engine/Source/Editor/Sequencer/Private/MVVM/PinEditorExtension.cpp b/Engine/Source/Editor/Sequencer/Private/MVVM/PinEditorExtension.cpp
new file mode 100644
--- /dev/null
+++ b/Engine/Source/Editor/Sequencer/Private/MVVM/PinEditorExtension.cpp
@@ -0,0 +1,69 @@
+// Copyright Epic Games, Inc. All Rights Reserved.
+
+#include "MVVM/PinEditorExtension.h"
+
+#include "MovieScene.h"
+#include "MVVM/Extensions/IOutlinerExtension.h"
+#include "MVVM/Selection/Selection.h"
+#include "MVVM/ViewModels/OutlinerViewModel.h"
+#include "MVVM/ViewModels/SequencerEditorViewModel.h"
+#include "ScopedTransaction.h"
+#include "Sequencer.h"
+
+namespace UE::Sequencer
+{
+
+FPinEditorExtension::FPinEditorExtension()
+{
+
+}
+
+void FPinEditorExtension::OnCreated(TSharedRef<FViewModel> InWeakOwner)
+{
+	ensureMsgf(!WeakOwnerModel.Pin().IsValid(), TEXT("This extension was already created!"));
+	WeakOwnerModel = InWeakOwner->CastThisShared<FSequencerEditorViewModel>();
+}
+
+bool FPinEditorExtension::IsNodePinnable(TWeakViewModelPtr<IOutlinerExtension> InWeakOutlinerExtension) const
+{
+	TSharedPtr<FViewModel> Item = InWeakOutlinerExtension.Pin();
+	return Item->GetHierarchicalDepth() == 1;
+}
+
+void FPinEditorExtension::SetNodePinned(TWeakViewModelPtr<IOutlinerExtension> InWeakOutlinerExtension, bool bInIsPinned)
+{
+	const FScopedTransaction Transaction(NSLOCTEXT("Sequencer", "SetNodePinned", "Set Node Pinned"));
+
+	TSharedPtr<FViewModel> Item = InWeakOutlinerExtension.Pin();
+	FSequenceModel* RootModel = Item->GetRoot()->CastThis<FSequenceModel>();
+	TSharedPtr<IPinnableExtension> PinnableParent = Item->FindAncestorOfType<IPinnableExtension>(true);
+
+	if (RootModel && PinnableParent)
+	{
+		TSharedPtr<FOutlinerViewModel> Outliner = WeakOwnerModel.Pin()->GetOutliner();
+		Outliner->UnpinAllNodes();
+
+		PinnableParent->SetPinned(bInIsPinned);
+
+		TSharedPtr<FSequencer> Sequencer = RootModel->GetSequencerImpl();
+		if (Sequencer)
+		{
+			UMovieScene* MovieScene = Sequencer->GetFocusedMovieSceneSequence()->GetMovieScene();
+			FMovieSceneEditorData& EditorData = MovieScene->GetEditorData();
+
+			if (bInIsPinned)
+			{
+				EditorData.PinnedNodes.AddUnique(IOutlinerExtension::GetPathName(Item));
+			}
+			else
+			{
+				EditorData.PinnedNodes.RemoveSingle(IOutlinerExtension::GetPathName(Item));
+			}
+
+			Sequencer->RefreshTree();
+		}
+	}
+}
+
+} // namespace UE::Sequencer
+
diff --git a/Engine/Source/Editor/Sequencer/Private/MVVM/PinEditorExtension.h b/Engine/Source/Editor/Sequencer/Private/MVVM/PinEditorExtension.h
new file mode 100644
--- /dev/null
+++ b/Engine/Source/Editor/Sequencer/Private/MVVM/PinEditorExtension.h
@@ -0,0 +1,48 @@
+// Copyright Epic Games, Inc. All Rights Reserved.
+
+#pragma once
+
+#include "MVVM/ViewModels/ViewModel.h"
+#include "MVVM/ViewModelPtr.h"
+#include "MVVM/ViewModelTypeID.h"
+#include "MVVM/Extensions/DynamicExtensionContainer.h"
+
+namespace UE::Sequencer
+{
+
+class FSequencerEditorViewModel;
+class IOutlinerExtension;
+
+/**
+* Extension that handles the pinned state of items.
+*/
+class FPinEditorExtension : public IDynamicExtension
+{
+public:
+
+	UE_SEQUENCER_DECLARE_VIEW_MODEL_TYPE_ID(FPinEditorExtension)
+
+	FPinEditorExtension();
+
+	/** Called when the extension is created on a data model. */
+	virtual void OnCreated(TSharedRef<FViewModel> InWeakOwner) override;
+
+	/** Returns whether or not this item is at the top of the hierarchy. */
+	bool IsNodePinnable(TWeakViewModelPtr<IOutlinerExtension> InWeakOutlinerExtension) const;
+
+	/**
+	* Modifies an item and it's child items to be pinned.
+	* This operation applies to all selected items if the modified item is selected.
+	* @param bInIsPinned - Pin state to set the item to
+	* @param InWeakOutlinerExtension - Item to pin
+	*/
+	void SetNodePinned(TWeakViewModelPtr<IOutlinerExtension> InWeakOutlinerExtension, bool bInIsPinned);
+
+private:
+
+	/** The sequencer editor we are extending */
+	TWeakPtr<FSequencerEditorViewModel> WeakOwnerModel;
+};
+
+} // namespace UE::Sequencer
+
diff --git a/Engine/Source/Editor/Sequencer/Private/MVVM/SequencerViewModelTypeIDs.cpp b/Engine/Source/Editor/Sequencer/Private/MVVM/SequencerViewModelTypeIDs.cpp
--- a/Engine/Source/Editor/Sequencer/Private/MVVM/SequencerViewModelTypeIDs.cpp
+++ b/Engine/Source/Editor/Sequencer/Private/MVVM/SequencerViewModelTypeIDs.cpp
@@ -25,10 +25,14 @@
 #include "MVVM/CurveEditorExtension.h"
 #include "MVVM/CurveEditorIntegrationExtension.h"
 #include "MVVM/FolderModelStorageExtension.h"
+#include "MVVM/LockEditorExtension.h"
+#include "MVVM/MuteEditorExtension.h"
 #include "MVVM/ObjectBindingModelStorageExtension.h"
+#include "MVVM/PinEditorExtension.h"
 #include "MVVM/TrackModelStorageExtension.h"
 #include "MVVM/TrackRowModelStorageExtension.h"
 #include "MVVM/SectionModelStorageExtension.h"
+#include "MVVM/SoloEditorExtension.h"
 
 
 namespace UE
@@ -65,8 +69,12 @@ UE_SEQUENCER_DEFINE_CASTABLE(FSequencerTrackAreaViewModel);
 UE_SEQUENCER_DEFINE_VIEW_MODEL_TYPE_ID(FCurveEditorExtension);
 UE_SEQUENCER_DEFINE_VIEW_MODEL_TYPE_ID(FCurveEditorIntegrationExtension);
 UE_SEQUENCER_DEFINE_VIEW_MODEL_TYPE_ID(FFolderModelStorageExtension);
+UE_SEQUENCER_DEFINE_VIEW_MODEL_TYPE_ID(FLockEditorExtension);
+UE_SEQUENCER_DEFINE_VIEW_MODEL_TYPE_ID(FMuteEditorExtension);
 UE_SEQUENCER_DEFINE_VIEW_MODEL_TYPE_ID(FObjectBindingModelStorageExtension);
+UE_SEQUENCER_DEFINE_VIEW_MODEL_TYPE_ID(FPinEditorExtension);
 UE_SEQUENCER_DEFINE_VIEW_MODEL_TYPE_ID(FSectionModelStorageExtension);
+UE_SEQUENCER_DEFINE_VIEW_MODEL_TYPE_ID(FSoloEditorExtension);
 UE_SEQUENCER_DEFINE_VIEW_MODEL_TYPE_ID(FTrackModelStorageExtension);
 UE_SEQUENCER_DEFINE_VIEW_MODEL_TYPE_ID(FTrackRowModelStorageExtension);
 UE_SEQUENCER_DEFINE_VIEW_MODEL_TYPE_ID(ICurveEditorTreeItemExtension);
diff --git a/Engine/Source/Editor/Sequencer/Private/MVVM/SoloEditorExtension.cpp b/Engine/Source/Editor/Sequencer/Private/MVVM/SoloEditorExtension.cpp
new file mode 100644
--- /dev/null
+++ b/Engine/Source/Editor/Sequencer/Private/MVVM/SoloEditorExtension.cpp
@@ -0,0 +1,146 @@
+// Copyright Epic Games, Inc. All Rights Reserved.
+
+#include "MVVM/SoloEditorExtension.h"
+
+#include "MovieScene.h"
+#include "MVVM/Extensions/IOutlinerExtension.h"
+#include "MVVM/Extensions/ISoloableExtension.h"
+#include "MVVM/Selection/Selection.h"
+#include "MVVM/ViewModels/SequencerEditorViewModel.h"
+#include "ScopedTransaction.h"
+#include "Sequencer.h"
+
+namespace UE::Sequencer
+{
+	
+FSoloEditorExtension::FSoloEditorExtension()
+{
+
+}
+
+void FSoloEditorExtension::OnCreated(TSharedRef<FViewModel> InWeakOwner)
+{
+	ensureMsgf(!WeakOwnerModel.Pin().IsValid(), TEXT("This extension was already created!"));
+	WeakOwnerModel = InWeakOwner->CastThisShared<FSequencerEditorViewModel>();
+}
+
+bool FSoloEditorExtension::IsNodeSoloed(TWeakViewModelPtr<IOutlinerExtension> InWeakOutlinerExtension) const
+{
+	TViewModelPtr<IOutlinerExtension> OutlinerItem = InWeakOutlinerExtension.Pin();
+	if (OutlinerItem)
+	{
+		if (TViewModelPtr<ISoloableExtension> Soloable = OutlinerItem.ImplicitCast())
+		{
+			return Soloable->IsSolo();
+		}
+	}
+
+	return false;
+}
+
+void FSoloEditorExtension::SetNodeSoloed(TWeakViewModelPtr<IOutlinerExtension> InWeakOutlinerExtension, bool bInIsSoloed)
+{
+	const FScopedTransaction Transaction(NSLOCTEXT("Sequencer", "SetNodeSoloed", "Set Node Soloed"));
+
+	TSharedPtr<FSequencerEditorViewModel> EditorViewModel = WeakOwnerModel.Pin();
+	TSharedPtr<FSequencer> Sequencer = EditorViewModel ? EditorViewModel->GetSequencerImpl() : nullptr;
+	TViewModelPtr<IOutlinerExtension> OutlinerItem = InWeakOutlinerExtension.Pin();
+
+	if (Sequencer)
+	{
+		UMovieScene* MovieScene = Sequencer->GetFocusedMovieSceneSequence()->GetMovieScene();
+
+		if (MovieScene->IsReadOnly())
+		{
+			return;
+		}
+
+		TArray<FString>& SoloNodes = MovieScene->GetSoloNodes();
+
+		MovieScene->Modify();
+
+		if (OutlinerItem->GetSelectionState() == EOutlinerSelectionState::SelectedDirectly)
+		{
+			// if selected, modify all selected items
+			for (FViewModelPtr Node : Sequencer->GetViewModel()->GetSelection()->Outliner)
+			{
+				FString NodePath = IOutlinerExtension::GetPathName(Node);
+				if (bInIsSoloed)
+				{
+					SoloNodes.AddUnique(NodePath);
+				}
+				else
+				{
+					SoloNodes.Remove(NodePath);
+				}
+			}
+		}
+		else
+		{
+			// modify only the toggled item
+			TSharedPtr<FViewModel> Item = OutlinerItem;
+			FString NodePath = IOutlinerExtension::GetPathName(Item);
+			if (bInIsSoloed)
+			{
+				SoloNodes.AddUnique(NodePath);
+			}
+			else
+			{
+				SoloNodes.Remove(NodePath);
+			}
+		}
+	}
+}
+
+bool FSoloEditorExtension::IsNodeSoloable(TWeakViewModelPtr<IOutlinerExtension> InWeakOutlinerExtension) const
+{
+	TSharedPtr<FSequencerEditorViewModel> EditorViewModel = WeakOwnerModel.Pin();
+	TSharedPtr<FSequencer> Sequencer = EditorViewModel ? EditorViewModel->GetSequencerImpl() : nullptr;
+	
+	if (Sequencer)
+	{
+		TSharedPtr<FViewModel> Item = InWeakOutlinerExtension.Pin();
+		FString NodePath = IOutlinerExtension::GetPathName(Item);
+		return Sequencer->GetNodeTree()->GetNodeAtPath(NodePath) != nullptr;
+	}
+	return false;
+}
+
+void FSoloEditorExtension::OnPostSolo()
+{
+	TSharedPtr<FSequencerEditorViewModel> EditorViewModel = WeakOwnerModel.Pin();
+	TSharedPtr<FSequencer> Sequencer = EditorViewModel ? EditorViewModel->GetSequencerImpl() : nullptr;
+	if (Sequencer)
+	{
+		Sequencer->RefreshTree();
+	}
+}
+
+bool FSoloEditorExtension::HasSoloedChildNode(TWeakViewModelPtr<IOutlinerExtension> InWeakOutlinerExtension) const
+{
+	using namespace UE::Sequencer;
+
+	TSharedPtr<FSequencerEditorViewModel> EditorViewModel = WeakOwnerModel.Pin();
+	TSharedPtr<FSequencer> Sequencer = EditorViewModel ? EditorViewModel->GetSequencerImpl() : nullptr;
+
+	if (Sequencer)
+	{
+		TViewModelPtr<IOutlinerExtension> OutlinerItem = InWeakOutlinerExtension.Pin();
+		const TArray<FString>& MuteNodes = Sequencer->GetFocusedMovieSceneSequence()->GetMovieScene()->GetMuteNodes();
+
+		TSharedPtr<FViewModel> Item = OutlinerItem;
+		const FString NodePath = IOutlinerExtension::GetPathName(Item);
+
+		for (const TViewModelPtr<IOutlinerExtension>& Child : OutlinerItem.AsModel()->GetDescendantsOfType<IOutlinerExtension>())
+		{
+			if (Sequencer->GetNodeTree()->IsNodeSolo(Child))
+			{
+				return true;
+			}
+		}
+	}
+
+	return false;
+}
+
+} // namespace UE::Sequencer
diff --git a/Engine/Source/Editor/Sequencer/Private/MVVM/SoloEditorExtension.h b/Engine/Source/Editor/Sequencer/Private/MVVM/SoloEditorExtension.h
new file mode 100644
--- /dev/null
+++ b/Engine/Source/Editor/Sequencer/Private/MVVM/SoloEditorExtension.h
@@ -0,0 +1,56 @@
+// Copyright Epic Games, Inc. All Rights Reserved.
+
+#pragma once
+
+#include "MVVM/ViewModels/ViewModel.h"
+#include "MVVM/ViewModelPtr.h"
+#include "MVVM/ViewModelTypeID.h"
+#include "MVVM/Extensions/DynamicExtensionContainer.h"
+
+namespace UE::Sequencer
+{
+class FSequencerEditorViewModel;
+class IOutlinerExtension;
+
+/**
+* Extension that handles the solo state of items.
+*/
+class FSoloEditorExtension : public IDynamicExtension
+{
+public:
+
+	UE_SEQUENCER_DECLARE_VIEW_MODEL_TYPE_ID(FSoloEditorExtension)
+
+	FSoloEditorExtension();
+
+	/** Called when the extension is created on a data model. */
+	virtual void OnCreated(TSharedRef<FViewModel> InWeakOwner) override;
+
+	/** Returns whether a given item is soloed. */
+	bool IsNodeSoloed(TWeakViewModelPtr<IOutlinerExtension> InWeakOutlinerExtension) const;
+
+	/**
+	* Modifies an item and it's child sections to be soloed.
+	* This operation applies to all selected items if the modified item is selected.
+	* @param bInIsSoloed - Solo state to set the item to
+	* @param InWeakOutlinerExtension - Item to solo
+	*/
+	void SetNodeSoloed(TWeakViewModelPtr<IOutlinerExtension> InWeakOutlinerExtension, bool bInIsSoloed);
+
+	/** Returns whether or not the item has a valid node path to solo. */
+	bool IsNodeSoloable(TWeakViewModelPtr<IOutlinerExtension> InWeakOutlinerExtension) const;
+
+	/** Returns whether or not the item has any soloed child items. */
+	bool HasSoloedChildNode(TWeakViewModelPtr<IOutlinerExtension> InWeakOutlinerExtension) const;
+
+	/** Refreshes the Sequencer Tree, required to mute non-soloed items. */
+	void OnPostSolo();
+
+private:
+
+	/** The sequencer editor we are extending */
+	TWeakPtr<FSequencerEditorViewModel> WeakOwnerModel;
+};
+
+} // namespace UE::Sequencer
+
diff --git a/Engine/Source/Editor/Sequencer/Private/MVVM/ViewModels/SequencerEditorViewModel.cpp b/Engine/Source/Editor/Sequencer/Private/MVVM/ViewModels/SequencerEditorViewModel.cpp
--- a/Engine/Source/Editor/Sequencer/Private/MVVM/ViewModels/SequencerEditorViewModel.cpp
+++ b/Engine/Source/Editor/Sequencer/Private/MVVM/ViewModels/SequencerEditorViewModel.cpp
@@ -5,8 +5,12 @@
 #include "MVVM/ViewModels/SequencerTrackAreaViewModel.h"
 #include "MVVM/Selection/Selection.h"
 #include "MVVM/CurveEditorExtension.h"
+#include "MVVM/LockEditorExtension.h"
+#include "MVVM/MuteEditorExtension.h"
+#include "MVVM/PinEditorExtension.h"
 #include "MVVM/Extensions/IOutlinerExtension.h"
 #include "MVVM/SharedViewModelData.h"
+#include "MVVM/SoloEditorExtension.h"
 #include "ISequencerModule.h"
 #include "Sequencer.h"
 #include "MovieSceneSequenceID.h"
@@ -46,6 +50,11 @@ void FSequencerEditorViewModel::PreInitializeEditorImpl()
 	{
 		AddDynamicExtension(FCurveEditorExtension::ID);
 	}
+
+	AddDynamicExtension(FLockEditorExtension::ID);
+	AddDynamicExtension(FMuteEditorExtension::ID);
+	AddDynamicExtension(FPinEditorExtension::ID);
+	AddDynamicExtension(FSoloEditorExtension::ID);
 }
 
 TSharedPtr<FViewModel> FSequencerEditorViewModel::CreateRootModelImpl()
diff --git a/Engine/Source/Editor/SequencerCore/Private/MVVM/Views/SOutlinerColumnBorder.cpp b/Engine/Source/Editor/SequencerCore/Private/MVVM/Views/SOutlinerColumnBorder.cpp
--- a/Engine/Source/Editor/SequencerCore/Private/MVVM/Views/SOutlinerColumnBorder.cpp
+++ b/Engine/Source/Editor/SequencerCore/Private/MVVM/Views/SOutlinerColumnBorder.cpp
@@ -69,7 +69,6 @@ FSlateColor SOutlinerColumnInnerBorder::GetBackgroundTint() const
 	bool bIsOutlinerItemSelectedDirectly = EnumHasAnyFlags(SelectionState, EOutlinerSelectionState::SelectedDirectly);
 	bool bIsOutlinerItemHovered = Editor->GetOutliner()->GetHoveredItem() == OutlinerItem;
 
-	// float Opacity = ((bIsOutlinerItemSelectedDirectly || bIsOutlinerItemHovered) && !bIsMouseOverInnerBorder) ? 0.f : .4f;
 	float Opacity = .4f;
 	return FLinearColor(0.f, 0.f, 0.f, Opacity);
 }
diff --git a/Engine/Source/Editor/SequencerCore/Public/ISequencerOutlinerColumn.h b/Engine/Source/Editor/SequencerCore/Public/ISequencerOutlinerColumn.h
--- a/Engine/Source/Editor/SequencerCore/Public/ISequencerOutlinerColumn.h
+++ b/Engine/Source/Editor/SequencerCore/Public/ISequencerOutlinerColumn.h
@@ -51,7 +51,10 @@ public:
 	virtual bool IsColumnVisibleByDefault() const { return true; }
 
 	/* Gets whether or not this column is supported by a given Sequencer */
-	virtual bool SupportsSequence(UMovieSceneSequence* InSequence) const = 0;
+	virtual bool SupportsSequence(UMovieSceneSequence* InSequence) const { return true; }
+
+	/* Gets whether or not a widget should be generated for a given item in the outliner column. */
+	virtual bool IsItemCompatibleWithColumn(const UE::Sequencer::FCreateOutlinerColumnParams& InParams) const = 0;
 
 	/* Gets the widget created for each item within the SOutlinerView, column widgets must be fixed width */
 	virtual TSharedRef<SWidget> CreateColumnWidget(const UE::Sequencer::FCreateOutlinerColumnParams& InParams) const = 0;
