Sequencer: Outliner Column Implementations

Implementations for the pin, mute, lock, and solo outliner columns to be displayed in the SOutlinerView of Sequencer. Each outliner column widget implements click and drag functionality. Outliner column visibility is saved per-sequencer and per-user. Default column visibility and order can be saved and restored in BaseEditorPerProjectUserSettings.ini.

diff --git a/Engine/Source/Editor/EditorStyle/Private/StarshipStyle.cpp b/Engine/Source/Editor/EditorStyle/Private/StarshipStyle.cpp
--- a/Engine/Source/Editor/EditorStyle/Private/StarshipStyle.cpp
+++ b/Engine/Source/Editor/EditorStyle/Private/StarshipStyle.cpp
@@ -2343,6 +2343,16 @@ void FStarshipEditorStyle::FStyle::SetupSequencerStyles()
 		Set("SequenceRecorder.Common.RecordingInactive", new IMAGE_BRUSH("Common/SmallCheckBox", Icon14x14));
 	}
 
+	// Sequencer outliner column UI
+	if (IncludeEditorSpecificStyles())
+	{
+		Set("Sequencer.Column.Mute", new IMAGE_BRUSH_SVG("Sequencer/Column_Widgets/SequencerMute", Icon12x12));
+		Set("Sequencer.Column.Locked", new IMAGE_BRUSH_SVG("Sequencer/Column_Widgets/SequencerLocked", Icon12x12));
+		Set("Sequencer.Column.Solo", new IMAGE_BRUSH_SVG("Sequencer/Column_Widgets/SequencerSolo", Icon12x12));
+		Set("Sequencer.Column.Unpinned", new IMAGE_BRUSH_SVG("Sequencer/Column_Widgets/SequencerUnpinned", Icon12x12));
+		Set("Sequencer.Column.CheckBoxIndeterminate", new CORE_IMAGE_BRUSH_SVG("Starship/CoreWidgets/CheckBox/CheckBoxIndeterminate_12", Icon12x12));
+		Set("Sequencer.Column.OutlinerColumnBox", new FSlateRoundedBoxBrush(FStyleColors::Header, 2.0f));
+	}
 }
 
 void FStarshipEditorStyle::FStyle::SetupViewportStyles()
diff --git a/Engine/Source/Editor/Sequencer/Private/MVVM/LockEditorExtension.cpp b/Engine/Source/Editor/Sequencer/Private/MVVM/LockEditorExtension.cpp
--- a/Engine/Source/Editor/Sequencer/Private/MVVM/LockEditorExtension.cpp
+++ b/Engine/Source/Editor/Sequencer/Private/MVVM/LockEditorExtension.cpp
@@ -28,6 +28,11 @@ bool FLockEditorExtension::IsNodeLocked(TWeakViewModelPtr<IOutlinerExtension> In
 {
 	FViewModelPtr Item = InWeakOutlinerExtension.Pin();
 
+	if (!Item)
+	{
+		return false;
+	}
+
 	// Locked if all sections are locked
 	int32 NumSections = 0;
 	TSet<TWeakObjectPtr<UMovieSceneSection>> Sections;
@@ -50,7 +55,11 @@ void FLockEditorExtension::SetNodeLocked(TWeakViewModelPtr<IOutlinerExtension> I
 
 	TSharedPtr<FSequencerEditorViewModel> EditorViewModel = WeakOwnerModel.Pin();
 	TViewModelPtr<IOutlinerExtension> OutlinerItem = InWeakOutlinerExtension.Pin();
-	FViewModelPtr Item = OutlinerItem;
+
+	if (!EditorViewModel || !OutlinerItem)
+	{
+		return;
+	}
 
 	if (OutlinerItem->GetSelectionState() == EOutlinerSelectionState::SelectedDirectly)
 	{
@@ -70,6 +79,7 @@ void FLockEditorExtension::SetNodeLocked(TWeakViewModelPtr<IOutlinerExtension> I
 	else
 	{
 		// only one unselected item was toggled, toggle just that node
+		FViewModelPtr Item = OutlinerItem;
 		TSet<TWeakObjectPtr<UMovieSceneSection> > Sections;
 		SequencerHelpers::GetAllSections(Item, Sections);
 
@@ -85,6 +95,11 @@ bool FLockEditorExtension::IsNodeLockable(TWeakViewModelPtr<IOutlinerExtension>
 {
 	FViewModelPtr Item = InWeakOutlinerExtension.Pin();
 
+	if (!Item)
+	{
+		return false;
+	}
+
 	TSet<TWeakObjectPtr<UMovieSceneSection> > Sections;
 	SequencerHelpers::GetAllSections(Item, Sections);
 
@@ -95,11 +110,14 @@ bool FLockEditorExtension::HasLockedChildNode(TWeakViewModelPtr<IOutlinerExtensi
 {
 	TViewModelPtr<IOutlinerExtension> OutlinerItem = InWeakOutlinerExtension.Pin();
 
-	for (const TViewModelPtr<IOutlinerExtension>& Child : OutlinerItem.AsModel()->GetDescendantsOfType<IOutlinerExtension>())
+	if (OutlinerItem)
 	{
-		if (IsNodeLocked(Child))
+		for (const TViewModelPtr<IOutlinerExtension>& Child : OutlinerItem.AsModel()->GetDescendantsOfType<IOutlinerExtension>())
 		{
-			return true;
+			if (IsNodeLocked(Child))
+			{
+				return true;
+			}
 		}
 	}
 
diff --git a/Engine/Source/Editor/Sequencer/Private/MVVM/MuteEditorExtension.cpp b/Engine/Source/Editor/Sequencer/Private/MVVM/MuteEditorExtension.cpp
--- a/Engine/Source/Editor/Sequencer/Private/MVVM/MuteEditorExtension.cpp
+++ b/Engine/Source/Editor/Sequencer/Private/MVVM/MuteEditorExtension.cpp
@@ -24,6 +24,25 @@ void FMuteEditorExtension::OnCreated(TSharedRef<FViewModel> InWeakOwner)
 	WeakOwnerModel = InWeakOwner->CastThisShared<FSequencerEditorViewModel>();
 }
 
+bool FMuteEditorExtension::IsNodeMuted(TWeakViewModelPtr<IOutlinerExtension> InWeakOutlinerExtension) const
+{
+	TViewModelPtr<IOutlinerExtension> Item = InWeakOutlinerExtension.Pin();
+	TSharedPtr<FSequencerEditorViewModel> EditorViewModel = WeakOwnerModel.Pin();
+	TSharedPtr<FSequencer> Sequencer = EditorViewModel ? EditorViewModel->GetSequencerImpl() : nullptr;
+	if(Item && Sequencer)
+	{
+		const TArray<FString>& MuteNodes = Sequencer->GetFocusedMovieSceneSequence()->GetMovieScene()->GetMuteNodes();
+		const FString NodePath = IOutlinerExtension::GetPathName(Item);
+
+		if (MuteNodes.Contains(NodePath))
+		{
+			return true;
+		}
+	}
+
+	return false;
+}
+
 void FMuteEditorExtension::SetNodeMuted(TWeakViewModelPtr<IOutlinerExtension> InWeakOutlinerExtension, bool bInIsMuted)
 {              
 	const FScopedTransaction Transaction(NSLOCTEXT("Sequencer", "SetNodeMuted", "Set Node Muted"));
@@ -32,7 +51,7 @@ void FMuteEditorExtension::SetNodeMuted(TWeakViewModelPtr<IOutlinerExtension> In
 	TSharedPtr<FSequencerEditorViewModel> EditorViewModel = WeakOwnerModel.Pin();
 	TSharedPtr<FSequencer> Sequencer = EditorViewModel ? EditorViewModel->GetSequencerImpl() : nullptr;
 		
-	if (Sequencer)
+	if (Sequencer && OutlinerItem)
 	{
 		UMovieScene* MovieScene = Sequencer->GetFocusedMovieSceneSequence()->GetMovieScene();
 
@@ -88,7 +107,7 @@ bool FMuteEditorExtension::IsNodeMutable(TWeakViewModelPtr<IOutlinerExtension> I
 	TSharedPtr<FSequencerEditorViewModel> EditorViewModel = WeakOwnerModel.Pin();
 	TSharedPtr<FSequencer> Sequencer = EditorViewModel ? EditorViewModel->GetSequencerImpl() : nullptr;
 		
-	if (Sequencer)
+	if (Sequencer && OutlinerItem)
 	{
 		TSharedPtr<FViewModel> Item = OutlinerItem;
 		FString NodePath = IOutlinerExtension::GetPathName(Item);
@@ -104,25 +123,44 @@ bool FMuteEditorExtension::HasMutedChildNode(TWeakViewModelPtr<IOutlinerExtensio
 
 	TViewModelPtr<IOutlinerExtension> OutlinerItem = InWeakOutlinerExtension.Pin();
 
-	for (const TViewModelPtr<IMutableExtension>& Child : OutlinerItem.AsModel()->GetDescendantsOfType<IMutableExtension>())
+	if (OutlinerItem)
 	{
-		if (Child->IsMuted())
+		for (const TViewModelPtr<IOutlinerExtension>& Child : OutlinerItem.AsModel()->GetDescendantsOfType<IOutlinerExtension>(false, EViewModelListType::Outliner))
 		{
-			return true;
+			if (IsNodeMuted(Child))
+			{
+				return true;
+			}
 		}
 	}
 
 	return false;
 }
 
-void FMuteEditorExtension::OnPostMute()
+bool FMuteEditorExtension::IsNodeImplicitlyMuted(TWeakViewModelPtr<IOutlinerExtension> InWeakOutlinerExtension) const
 {
+	TViewModelPtr<IOutlinerExtension> Item = InWeakOutlinerExtension.Pin();
 	TSharedPtr<FSequencerEditorViewModel> EditorViewModel = WeakOwnerModel.Pin();
 	TSharedPtr<FSequencer> Sequencer = EditorViewModel ? EditorViewModel->GetSequencerImpl() : nullptr;
-	if (Sequencer)
+	if (Item && Sequencer)
 	{
-		Sequencer->RefreshTree();
+		const TArray<FString>& MuteNodes = Sequencer->GetFocusedMovieSceneSequence()->GetMovieScene()->GetMuteNodes();
+
+		TViewModelPtr<IOutlinerExtension> ParentNode = Item.AsModel()->FindAncestorOfType<IOutlinerExtension>();
+		while (ParentNode)
+		{
+			const FString NodePath = IOutlinerExtension::GetPathName(ParentNode);
+			if (MuteNodes.Contains(NodePath))
+			{
+				return true;
+			}
+
+			// continue traversing upwards until no more parents
+			ParentNode = ParentNode.AsModel()->FindAncestorOfType<IOutlinerExtension>();
+		}
 	}
+
+	return false;
 }
 
 } // namespace UE::Sequencer
diff --git a/Engine/Source/Editor/Sequencer/Private/MVVM/MuteEditorExtension.h b/Engine/Source/Editor/Sequencer/Private/MVVM/MuteEditorExtension.h
--- a/Engine/Source/Editor/Sequencer/Private/MVVM/MuteEditorExtension.h
+++ b/Engine/Source/Editor/Sequencer/Private/MVVM/MuteEditorExtension.h
@@ -27,6 +27,13 @@ public:
 	/** Called when the extension is created on a data model. */
 	virtual void OnCreated(TSharedRef<FViewModel> InWeakOwner) override;
 
+	/** 
+	* Returns whether a given item is directly muted. 
+	* Returns false if not directly muted but implicitly muted by a parent node.
+	*/
+	bool IsNodeMuted(TWeakViewModelPtr<IOutlinerExtension> InWeakOutlinerExtension) const;
+
+
 	/**
 	* Modifies an item to be muted.
 	* This operation applies to all selected items if the modified item is selected.
@@ -41,8 +48,11 @@ public:
 	/** Returns whether or not the item has any muted child items. */
 	bool HasMutedChildNode(TWeakViewModelPtr<IOutlinerExtension> InWeakOutlinerExtension) const;
 
-	/** Refreshes the Sequencer Tree, required to ensure that muting an item doesn't affect the visibility of soloed items. */
-	void OnPostMute();
+	/**
+	 * Returns whether a given item is implicitly muted.
+	 * Returns false if not implicitly muted by a parent node but directly muted itself.
+	 */
+	bool IsNodeImplicitlyMuted(TWeakViewModelPtr<IOutlinerExtension> InWeakOutlinerExtension) const;
 
 private:
 
diff --git a/Engine/Source/Editor/Sequencer/Private/MVVM/PinEditorExtension.cpp b/Engine/Source/Editor/Sequencer/Private/MVVM/PinEditorExtension.cpp
--- a/Engine/Source/Editor/Sequencer/Private/MVVM/PinEditorExtension.cpp
+++ b/Engine/Source/Editor/Sequencer/Private/MVVM/PinEditorExtension.cpp
@@ -4,6 +4,7 @@
 
 #include "MovieScene.h"
 #include "MVVM/Extensions/IOutlinerExtension.h"
+#include "MVVM/Extensions/IPinnableExtension.h"
 #include "MVVM/Selection/Selection.h"
 #include "MVVM/ViewModels/OutlinerViewModel.h"
 #include "MVVM/ViewModels/SequencerEditorViewModel.h"
@@ -24,43 +25,85 @@ void FPinEditorExtension::OnCreated(TSharedRef<FViewModel> InWeakOwner)
 	WeakOwnerModel = InWeakOwner->CastThisShared<FSequencerEditorViewModel>();
 }
 
+bool FPinEditorExtension::IsNodePinned(TWeakViewModelPtr<IOutlinerExtension> InWeakOutlinerExtension) const
+{
+	TViewModelPtr<IOutlinerExtension> OutlinerExtension = InWeakOutlinerExtension.Pin();
+	if (OutlinerExtension.IsValid())
+	{
+		if (TViewModelPtr<IPinnableExtension> PinnableExtension = OutlinerExtension.ImplicitCast())
+		{
+			return PinnableExtension->IsPinned();
+		}
+	}
+	return false;
+}
+
 bool FPinEditorExtension::IsNodePinnable(TWeakViewModelPtr<IOutlinerExtension> InWeakOutlinerExtension) const
 {
 	TSharedPtr<FViewModel> Item = InWeakOutlinerExtension.Pin();
-	return Item->GetHierarchicalDepth() == 1;
+	if (Item)
+	{
+		return Item->GetHierarchicalDepth() == 1;
+	}
+	return false;
 }
 
-void FPinEditorExtension::SetNodePinned(TWeakViewModelPtr<IOutlinerExtension> InWeakOutlinerExtension, bool bInIsPinned)
+void FPinEditorExtension::SetNodePinned(TWeakViewModelPtr<IOutlinerExtension> InWeakOutlinerExtension, const bool bInIsPinned)
 {
 	const FScopedTransaction Transaction(NSLOCTEXT("Sequencer", "SetNodePinned", "Set Node Pinned"));
 
-	TSharedPtr<FViewModel> Item = InWeakOutlinerExtension.Pin();
-	FSequenceModel* RootModel = Item->GetRoot()->CastThis<FSequenceModel>();
-	TSharedPtr<IPinnableExtension> PinnableParent = Item->FindAncestorOfType<IPinnableExtension>(true);
+	TViewModelPtr<IOutlinerExtension> OutlinerItem = InWeakOutlinerExtension.Pin();
+	TSharedPtr<FSequencerEditorViewModel> EditorViewModel = WeakOwnerModel.Pin();
+	TSharedPtr<FSequencer> Sequencer = EditorViewModel ? EditorViewModel->GetSequencerImpl() : nullptr;
 
-	if (RootModel && PinnableParent)
+	if (OutlinerItem && Sequencer)
 	{
-		TSharedPtr<FOutlinerViewModel> Outliner = WeakOwnerModel.Pin()->GetOutliner();
+		TSharedPtr<FOutlinerViewModel> Outliner = EditorViewModel->GetOutliner();
 		Outliner->UnpinAllNodes();
 
-		PinnableParent->SetPinned(bInIsPinned);
+		UMovieScene* MovieScene = Sequencer->GetFocusedMovieSceneSequence()->GetMovieScene();
 
-		TSharedPtr<FSequencer> Sequencer = RootModel->GetSequencerImpl();
-		if (Sequencer)
+		if (MovieScene->IsReadOnly())
 		{
-			UMovieScene* MovieScene = Sequencer->GetFocusedMovieSceneSequence()->GetMovieScene();
-			FMovieSceneEditorData& EditorData = MovieScene->GetEditorData();
+			return;
+		}
 
-			if (bInIsPinned)
-			{
-				EditorData.PinnedNodes.AddUnique(IOutlinerExtension::GetPathName(Item));
-			}
-			else
+		MovieScene->Modify();
+
+		FMovieSceneEditorData& EditorData = MovieScene->GetEditorData();
+
+		if (OutlinerItem->GetSelectionState() == EOutlinerSelectionState::SelectedDirectly)
+		{
+			// pin all selected
+			for (FViewModelPtr Node : Sequencer->GetViewModel()->GetSelection()->Outliner)
 			{
-				EditorData.PinnedNodes.RemoveSingle(IOutlinerExtension::GetPathName(Item));
+				PinItem(EditorData, Node, bInIsPinned);
 			}
+		}
+		else
+		{
+			// pin only passed in outliner extension
+			TSharedPtr<FViewModel> Item = OutlinerItem;
+			PinItem(EditorData, Item, bInIsPinned);
+		}
+	}
+}
 
-			Sequencer->RefreshTree();
+void FPinEditorExtension::PinItem(FMovieSceneEditorData& InEditorData, TSharedPtr<FViewModel> InItem, const bool bInIsPinned)
+{
+	TSharedPtr<IPinnableExtension> PinnableParent = InItem->FindAncestorOfType<IPinnableExtension>(true);
+
+	if (PinnableParent)
+	{
+		PinnableParent->SetPinned(bInIsPinned);
+
+		if (bInIsPinned)
+		{
+			InEditorData.PinnedNodes.AddUnique(IOutlinerExtension::GetPathName(InItem));
+		}
+		else
+		{
+			InEditorData.PinnedNodes.RemoveSingle(IOutlinerExtension::GetPathName(InItem));
 		}
 	}
 }
diff --git a/Engine/Source/Editor/Sequencer/Private/MVVM/PinEditorExtension.h b/Engine/Source/Editor/Sequencer/Private/MVVM/PinEditorExtension.h
--- a/Engine/Source/Editor/Sequencer/Private/MVVM/PinEditorExtension.h
+++ b/Engine/Source/Editor/Sequencer/Private/MVVM/PinEditorExtension.h
@@ -7,6 +7,8 @@
 #include "MVVM/ViewModelTypeID.h"
 #include "MVVM/Extensions/DynamicExtensionContainer.h"
 
+struct FMovieSceneEditorData;
+
 namespace UE::Sequencer
 {
 
@@ -27,6 +29,9 @@ public:
 	/** Called when the extension is created on a data model. */
 	virtual void OnCreated(TSharedRef<FViewModel> InWeakOwner) override;
 
+	/** Returns whether a given item is pinned. */
+	bool IsNodePinned(TWeakViewModelPtr<IOutlinerExtension> InWeakOutlinerExtension) const;
+
 	/** Returns whether or not this item is at the top of the hierarchy. */
 	bool IsNodePinnable(TWeakViewModelPtr<IOutlinerExtension> InWeakOutlinerExtension) const;
 
@@ -36,7 +41,12 @@ public:
 	* @param bInIsPinned - Pin state to set the item to
 	* @param InWeakOutlinerExtension - Item to pin
 	*/
-	void SetNodePinned(TWeakViewModelPtr<IOutlinerExtension> InWeakOutlinerExtension, bool bInIsPinned);
+	void SetNodePinned(TWeakViewModelPtr<IOutlinerExtension> InWeakOutlinerExtension, const bool bInIsPinned);
+
+private:
+
+	/** Helper function for modifying the pin state of individual items */
+	void PinItem(FMovieSceneEditorData& InEditorData, TSharedPtr<FViewModel> InItem, const bool bInIsPinned);
 
 private:
 
diff --git a/Engine/Source/Editor/Sequencer/Private/MVVM/SoloEditorExtension.cpp b/Engine/Source/Editor/Sequencer/Private/MVVM/SoloEditorExtension.cpp
--- a/Engine/Source/Editor/Sequencer/Private/MVVM/SoloEditorExtension.cpp
+++ b/Engine/Source/Editor/Sequencer/Private/MVVM/SoloEditorExtension.cpp
@@ -26,12 +26,17 @@ void FSoloEditorExtension::OnCreated(TSharedRef<FViewModel> InWeakOwner)
 
 bool FSoloEditorExtension::IsNodeSoloed(TWeakViewModelPtr<IOutlinerExtension> InWeakOutlinerExtension) const
 {
-	TViewModelPtr<IOutlinerExtension> OutlinerItem = InWeakOutlinerExtension.Pin();
-	if (OutlinerItem)
+	TViewModelPtr<IOutlinerExtension> Item = InWeakOutlinerExtension.Pin();
+	TSharedPtr<FSequencerEditorViewModel> EditorViewModel = WeakOwnerModel.Pin();
+	TSharedPtr<FSequencer> Sequencer = EditorViewModel ? EditorViewModel->GetSequencerImpl() : nullptr;
+	if (Item && Sequencer)
 	{
-		if (TViewModelPtr<ISoloableExtension> Soloable = OutlinerItem.ImplicitCast())
+		const TArray<FString>& SoloNodes = Sequencer->GetFocusedMovieSceneSequence()->GetMovieScene()->GetSoloNodes();
+		const FString NodePath = IOutlinerExtension::GetPathName(Item);
+
+		if (SoloNodes.Contains(NodePath))
 		{
-			return Soloable->IsSolo();
+			return true;
 		}
 	}
 
@@ -42,11 +47,11 @@ void FSoloEditorExtension::SetNodeSoloed(TWeakViewModelPtr<IOutlinerExtension> I
 {
 	const FScopedTransaction Transaction(NSLOCTEXT("Sequencer", "SetNodeSoloed", "Set Node Soloed"));
 
+	TViewModelPtr<IOutlinerExtension> OutlinerItem = InWeakOutlinerExtension.Pin();
 	TSharedPtr<FSequencerEditorViewModel> EditorViewModel = WeakOwnerModel.Pin();
 	TSharedPtr<FSequencer> Sequencer = EditorViewModel ? EditorViewModel->GetSequencerImpl() : nullptr;
-	TViewModelPtr<IOutlinerExtension> OutlinerItem = InWeakOutlinerExtension.Pin();
 
-	if (Sequencer)
+	if (Sequencer && OutlinerItem)
 	{
 		UMovieScene* MovieScene = Sequencer->GetFocusedMovieSceneSequence()->GetMovieScene();
 
@@ -94,49 +99,58 @@ void FSoloEditorExtension::SetNodeSoloed(TWeakViewModelPtr<IOutlinerExtension> I
 
 bool FSoloEditorExtension::IsNodeSoloable(TWeakViewModelPtr<IOutlinerExtension> InWeakOutlinerExtension) const
 {
+	TSharedPtr<FViewModel> Item = InWeakOutlinerExtension.Pin();
 	TSharedPtr<FSequencerEditorViewModel> EditorViewModel = WeakOwnerModel.Pin();
 	TSharedPtr<FSequencer> Sequencer = EditorViewModel ? EditorViewModel->GetSequencerImpl() : nullptr;
 	
-	if (Sequencer)
+	if (Sequencer && Item)
 	{
-		TSharedPtr<FViewModel> Item = InWeakOutlinerExtension.Pin();
 		FString NodePath = IOutlinerExtension::GetPathName(Item);
 		return Sequencer->GetNodeTree()->GetNodeAtPath(NodePath) != nullptr;
 	}
 	return false;
 }
 
-void FSoloEditorExtension::OnPostSolo()
+bool FSoloEditorExtension::HasSoloedChildNode(TWeakViewModelPtr<IOutlinerExtension> InWeakOutlinerExtension) const
 {
-	TSharedPtr<FSequencerEditorViewModel> EditorViewModel = WeakOwnerModel.Pin();
-	TSharedPtr<FSequencer> Sequencer = EditorViewModel ? EditorViewModel->GetSequencerImpl() : nullptr;
-	if (Sequencer)
+	using namespace UE::Sequencer;
+
+	TViewModelPtr<IOutlinerExtension> OutlinerItem = InWeakOutlinerExtension.Pin();
+
+	if (OutlinerItem)
 	{
-		Sequencer->RefreshTree();
+		for (const TViewModelPtr<IOutlinerExtension>& Child : OutlinerItem.AsModel()->GetDescendantsOfType<IOutlinerExtension>(false, EViewModelListType::Outliner))
+		{
+			if (IsNodeSoloed(Child))
+			{
+				return true;
+			}
+		}
 	}
+
+	return false;
 }
 
-bool FSoloEditorExtension::HasSoloedChildNode(TWeakViewModelPtr<IOutlinerExtension> InWeakOutlinerExtension) const
+bool FSoloEditorExtension::IsNodeImplicitlySoloed(TWeakViewModelPtr<IOutlinerExtension> InWeakOutlinerExtension) const
 {
-	using namespace UE::Sequencer;
-
+	TViewModelPtr<IOutlinerExtension> Item = InWeakOutlinerExtension.Pin();
 	TSharedPtr<FSequencerEditorViewModel> EditorViewModel = WeakOwnerModel.Pin();
 	TSharedPtr<FSequencer> Sequencer = EditorViewModel ? EditorViewModel->GetSequencerImpl() : nullptr;
-
-	if (Sequencer)
+	if (Item && Sequencer)
 	{
-		TViewModelPtr<IOutlinerExtension> OutlinerItem = InWeakOutlinerExtension.Pin();
-		const TArray<FString>& MuteNodes = Sequencer->GetFocusedMovieSceneSequence()->GetMovieScene()->GetMuteNodes();
+		const TArray<FString>& SoloNodes = Sequencer->GetFocusedMovieSceneSequence()->GetMovieScene()->GetSoloNodes();
 
-		TSharedPtr<FViewModel> Item = OutlinerItem;
-		const FString NodePath = IOutlinerExtension::GetPathName(Item);
-
-		for (const TViewModelPtr<IOutlinerExtension>& Child : OutlinerItem.AsModel()->GetDescendantsOfType<IOutlinerExtension>())
+		TViewModelPtr<IOutlinerExtension> ParentNode = Item.AsModel()->FindAncestorOfType<IOutlinerExtension>();
+		while (ParentNode)
 		{
-			if (Sequencer->GetNodeTree()->IsNodeSolo(Child))
+			const FString NodePath = IOutlinerExtension::GetPathName(ParentNode);
+			if (SoloNodes.Contains(NodePath))
 			{
 				return true;
 			}
+
+			// continue traversing upwards until no more parents
+			ParentNode = ParentNode.AsModel()->FindAncestorOfType<IOutlinerExtension>();
 		}
 	}
 
diff --git a/Engine/Source/Editor/Sequencer/Private/MVVM/SoloEditorExtension.h b/Engine/Source/Editor/Sequencer/Private/MVVM/SoloEditorExtension.h
--- a/Engine/Source/Editor/Sequencer/Private/MVVM/SoloEditorExtension.h
+++ b/Engine/Source/Editor/Sequencer/Private/MVVM/SoloEditorExtension.h
@@ -26,7 +26,10 @@ public:
 	/** Called when the extension is created on a data model. */
 	virtual void OnCreated(TSharedRef<FViewModel> InWeakOwner) override;
 
-	/** Returns whether a given item is soloed. */
+	/**
+	* Returns whether a given item is directly soloed.
+	* Returns false if not directly soloed but implicitly soloed by a parent node.
+	*/
 	bool IsNodeSoloed(TWeakViewModelPtr<IOutlinerExtension> InWeakOutlinerExtension) const;
 
 	/**
@@ -43,8 +46,11 @@ public:
 	/** Returns whether or not the item has any soloed child items. */
 	bool HasSoloedChildNode(TWeakViewModelPtr<IOutlinerExtension> InWeakOutlinerExtension) const;
 
-	/** Refreshes the Sequencer Tree, required to mute non-soloed items. */
-	void OnPostSolo();
+	/**
+	 * Returns whether a given item is implicitly soloed.
+	 * Returns false if not implicitly soloed by a parent node but directly soloed itself.
+	 */
+	bool IsNodeImplicitlySoloed(TWeakViewModelPtr<IOutlinerExtension> InWeakOutlinerExtension) const;
 
 private:
 
diff --git a/Engine/Source/Editor/Sequencer/Private/MovieSceneOutlinerColumn.cpp b/Engine/Source/Editor/Sequencer/Private/MovieSceneOutlinerColumn.cpp
deleted file mode 100644
diff --git a/Engine/Source/Editor/Sequencer/Private/OutlinerColumns/LockOutlinerColumn.cpp b/Engine/Source/Editor/Sequencer/Private/OutlinerColumns/LockOutlinerColumn.cpp
new file mode 100644
--- /dev/null
+++ b/Engine/Source/Editor/Sequencer/Private/OutlinerColumns/LockOutlinerColumn.cpp
@@ -0,0 +1,54 @@
+// Copyright Epic Games, Inc. All Rights Reserved.
+
+#include "OutlinerColumns/LockOutlinerColumn.h"
+
+#include "ISequencerOutlinerColumn.h"
+#include "MVVM/Extensions/IOutlinerExtension.h"
+#include "MVVM/LockEditorExtension.h"
+#include "MVVM/ViewModels/EditorViewModel.h"
+#include "MVVM/ViewModels/SequencerEditorViewModel.h"
+#include "Widgets/OutlinerColumns/SLockColumnWidget.h"
+
+#define LOCTEXT_NAMESPACE "FLockOutlinerColumn"
+
+FName FLockOutlinerColumn::GetColumnName() const
+{
+	return FName(TEXT("Lock"));
+}
+
+FText FLockOutlinerColumn::GetColumnLabel() const
+{
+	return LOCTEXT("LockColumnLabel", "Lock");
+}
+
+bool FLockOutlinerColumn::IsItemCompatibleWithColumn(const UE::Sequencer::FCreateOutlinerColumnParams& InParams) const
+{
+	using namespace UE::Sequencer;
+
+	if (InParams.Editor)
+	{
+		FLockEditorExtension* LockEditorExtension = InParams.Editor->CastDynamic<FLockEditorExtension>();
+		if (LockEditorExtension)
+		{
+			return LockEditorExtension->IsNodeLockable(InParams.OutlinerExtension);
+		}
+	}
+
+	return false;
+}
+
+TSharedRef<ISequencerOutlinerColumn> FLockOutlinerColumn::CreateOutlinerColumn()
+{
+	return MakeShareable(new FLockOutlinerColumn());
+}
+
+TSharedRef<SWidget> FLockOutlinerColumn::CreateColumnWidget(const TWeakPtr<ISequencerOutlinerColumn> InWeakOutlinerColumn, const UE::Sequencer::FCreateOutlinerColumnParams& InParams) const
+{
+	using namespace UE::Sequencer;
+
+	return SNew(SLockColumnWidget,
+		InWeakOutlinerColumn,
+		InParams);
+}
+
+#undef LOCTEXT_NAMESPACE
\ No newline at end of file
diff --git a/Engine/Source/Editor/Sequencer/Private/OutlinerColumns/LockOutlinerColumn.h b/Engine/Source/Editor/Sequencer/Private/OutlinerColumns/LockOutlinerColumn.h
new file mode 100644
--- /dev/null
+++ b/Engine/Source/Editor/Sequencer/Private/OutlinerColumns/LockOutlinerColumn.h
@@ -0,0 +1,38 @@
+// Copyright Epic Games, Inc. All Rights Reserved.
+
+#pragma once
+
+#include "ISequencerOutlinerColumn.h"
+
+/**
+ * A column for locking/unlocking tracks.
+ */
+class FLockOutlinerColumn
+	: public ISequencerOutlinerColumn
+{
+public:
+
+	FLockOutlinerColumn()
+	{ }
+
+	/** Returns the name of the column. Used for determining column type when dragging or saving settings. */
+	virtual FName GetColumnName() const override;
+
+	/** Returns the label of the column to display in visibility settings. */
+	virtual FText GetColumnLabel() const override;
+
+	/* Gets whether or not an outliner column widget should be generated for a given item in the outliner column. */
+	virtual bool IsItemCompatibleWithColumn(const UE::Sequencer::FCreateOutlinerColumnParams& InParams) const override;
+
+	/**
+	 * Creates an instance of this class.
+	 *
+	 * @return The new instance of this class.
+	 */
+	static TSharedRef<ISequencerOutlinerColumn> CreateOutlinerColumn();
+
+	/**
+	 * Creates a widget that locks an item in the outliner view.
+	 */
+	virtual TSharedRef<SWidget> CreateColumnWidget(const TWeakPtr<ISequencerOutlinerColumn> InWeakOutlinerColumn, const UE::Sequencer::FCreateOutlinerColumnParams& InParams) const override;
+};
\ No newline at end of file
diff --git a/Engine/Source/Editor/Sequencer/Private/OutlinerColumns/MuteOutlinerColumn.cpp b/Engine/Source/Editor/Sequencer/Private/OutlinerColumns/MuteOutlinerColumn.cpp
new file mode 100644
--- /dev/null
+++ b/Engine/Source/Editor/Sequencer/Private/OutlinerColumns/MuteOutlinerColumn.cpp
@@ -0,0 +1,55 @@
+// Copyright Epic Games, Inc. All Rights Reserved.
+
+#include "OutlinerColumns/MuteOutlinerColumn.h"
+
+#include "ISequencerOutlinerColumn.h"
+#include "MVVM/Extensions/IOutlinerExtension.h"
+#include "MVVM/MuteEditorExtension.h"
+#include "MVVM/ViewModels/EditorViewModel.h"
+#include "MVVM/ViewModels/SequencerEditorViewModel.h"
+#include "Widgets/Layout/SScaleBox.h"
+#include "Widgets/OutlinerColumns/SMuteColumnWidget.h"
+
+#define LOCTEXT_NAMESPACE "FMuteOutlinerColumn"
+
+FName FMuteOutlinerColumn::GetColumnName() const
+{
+	return FName(TEXT("Mute"));
+}
+
+FText FMuteOutlinerColumn::GetColumnLabel() const
+{
+	return LOCTEXT("MuteColumnLabel", "Mute");
+}
+
+bool FMuteOutlinerColumn::IsItemCompatibleWithColumn(const UE::Sequencer::FCreateOutlinerColumnParams& InParams) const
+{
+	using namespace UE::Sequencer;
+
+	if (InParams.Editor)
+	{
+		FMuteEditorExtension* MuteEditorExtension = InParams.Editor->CastDynamic<FMuteEditorExtension>();
+		if (MuteEditorExtension)
+		{
+			return MuteEditorExtension->IsNodeMutable(InParams.OutlinerExtension);
+		}
+	}
+
+	return false;
+}
+
+TSharedRef<ISequencerOutlinerColumn> FMuteOutlinerColumn::CreateOutlinerColumn()
+{
+	return MakeShareable(new FMuteOutlinerColumn());
+}
+
+TSharedRef<SWidget> FMuteOutlinerColumn::CreateColumnWidget(const TWeakPtr<ISequencerOutlinerColumn> InWeakOutlinerColumn, const UE::Sequencer::FCreateOutlinerColumnParams& InParams) const
+{
+	using namespace UE::Sequencer;
+
+	return SNew(SMuteColumnWidget,
+		InWeakOutlinerColumn,
+		InParams);
+}
+
+#undef LOCTEXT_NAMESPACE
\ No newline at end of file
diff --git a/Engine/Source/Editor/Sequencer/Private/OutlinerColumns/MuteOutlinerColumn.h b/Engine/Source/Editor/Sequencer/Private/OutlinerColumns/MuteOutlinerColumn.h
new file mode 100644
--- /dev/null
+++ b/Engine/Source/Editor/Sequencer/Private/OutlinerColumns/MuteOutlinerColumn.h
@@ -0,0 +1,38 @@
+// Copyright Epic Games, Inc. All Rights Reserved.
+
+#pragma once
+
+#include "ISequencerOutlinerColumn.h"
+
+/**
+ * A column for muting/unmuting tracks.
+ */
+class FMuteOutlinerColumn
+	: public ISequencerOutlinerColumn
+{
+public:
+
+	FMuteOutlinerColumn()
+	{ }
+
+	/** Returns the name of the column. Used for determining column type when dragging or saving settings. */
+	virtual FName GetColumnName() const override;
+
+	/** Returns the label of the column to display in visibility settings. */
+	virtual FText GetColumnLabel() const override;
+
+	/* Gets whether or not an outliner column widget should be generated for a given item in the outliner column. */
+	virtual bool IsItemCompatibleWithColumn(const UE::Sequencer::FCreateOutlinerColumnParams& InParams) const override;
+
+	/**
+	 * Creates an instance of this class.
+	 *
+	 * @return The new instance of this class.
+	 */
+	static TSharedRef<ISequencerOutlinerColumn> CreateOutlinerColumn();
+	
+	/**
+	 * Creates a widget that mutes an item in the outliner view.
+	 */
+	virtual TSharedRef<SWidget> CreateColumnWidget(const TWeakPtr<ISequencerOutlinerColumn> InWeakOutlinerColumn, const UE::Sequencer::FCreateOutlinerColumnParams& InParams) const override;
+};
diff --git a/Engine/Source/Editor/Sequencer/Private/OutlinerColumns/PinOutlinerColumn.cpp b/Engine/Source/Editor/Sequencer/Private/OutlinerColumns/PinOutlinerColumn.cpp
new file mode 100644
--- /dev/null
+++ b/Engine/Source/Editor/Sequencer/Private/OutlinerColumns/PinOutlinerColumn.cpp
@@ -0,0 +1,56 @@
+// Copyright Epic Games, Inc. All Rights Reserved.
+
+#include "OutlinerColumns/PinOutlinerColumn.h"
+
+#include "MVVM/Extensions/IOutlinerExtension.h"
+#include "ISequencerOutlinerColumn.h"
+#include "MVVM/ViewModels/EditorViewModel.h"
+#include "MVVM/PinEditorExtension.h"
+#include "MVVM/ViewModels/OutlinerItemModel.h"
+#include "MVVM/ViewModels/SequencerEditorViewModel.h"
+#include "Widgets/OutlinerColumns/SPinColumnWidget.h"
+#include "Widgets/Layout/SScaleBox.h"
+
+#define LOCTEXT_NAMESPACE "FPinOutlinerColumn"
+
+FName FPinOutlinerColumn::GetColumnName() const
+{
+	return FName(TEXT("Pin"));
+}
+
+FText FPinOutlinerColumn::GetColumnLabel() const
+{
+	return LOCTEXT("PinColumnLabel", "Pin");
+}
+
+bool FPinOutlinerColumn::IsItemCompatibleWithColumn(const UE::Sequencer::FCreateOutlinerColumnParams& InParams) const
+{
+	using namespace UE::Sequencer;
+
+	if (InParams.Editor)
+	{
+		FPinEditorExtension* PinEditorExtension = InParams.Editor->CastDynamic<FPinEditorExtension>();
+		if (PinEditorExtension)
+		{
+			return PinEditorExtension->IsNodePinnable(InParams.OutlinerExtension);
+		}
+	}
+
+	return false;
+}
+
+TSharedRef<ISequencerOutlinerColumn> FPinOutlinerColumn::CreateOutlinerColumn()
+{
+	return MakeShareable(new FPinOutlinerColumn());
+}
+
+TSharedRef<SWidget> FPinOutlinerColumn::CreateColumnWidget(const TWeakPtr<ISequencerOutlinerColumn> InWeakOutlinerColumn, const UE::Sequencer::FCreateOutlinerColumnParams& InParams) const
+{
+	using namespace UE::Sequencer;
+
+	return SNew(SPinColumnWidget,
+		InWeakOutlinerColumn,
+		InParams);
+}
+
+#undef LOCTEXT_NAMESPACE
\ No newline at end of file
diff --git a/Engine/Source/Editor/Sequencer/Private/OutlinerColumns/PinOutlinerColumn.h b/Engine/Source/Editor/Sequencer/Private/OutlinerColumns/PinOutlinerColumn.h
new file mode 100644
--- /dev/null
+++ b/Engine/Source/Editor/Sequencer/Private/OutlinerColumns/PinOutlinerColumn.h
@@ -0,0 +1,38 @@
+// Copyright Epic Games, Inc. All Rights Reserved.
+
+#pragma once
+
+#include "ISequencerOutlinerColumn.h"
+
+/**
+ * A column for pinning/unpinning tracks.
+ */
+class FPinOutlinerColumn
+	: public ISequencerOutlinerColumn
+{
+public:
+
+	FPinOutlinerColumn()
+	{ }
+
+	/** Returns the name of the column. Used for determining column type when dragging or saving settings. */
+	virtual FName GetColumnName() const override;
+
+	/** Returns the label of the column to display in visibility settings. */
+	virtual FText GetColumnLabel() const override;
+
+	/* Gets whether or not an outliner column widget should be generated for a given item in the outliner column. */
+	virtual bool IsItemCompatibleWithColumn(const UE::Sequencer::FCreateOutlinerColumnParams& InParams) const override;
+
+	/**
+	 * Creates an instance of this class.
+	 *
+	 * @return The new instance of this class.
+	 */
+	static TSharedRef<ISequencerOutlinerColumn> CreateOutlinerColumn();
+	
+	/**
+	 * Creates a widget that pins an item in the outliner view.
+	 */
+	virtual TSharedRef<SWidget> CreateColumnWidget(const TWeakPtr<ISequencerOutlinerColumn> InWeakOutlinerColumn, const UE::Sequencer::FCreateOutlinerColumnParams& InParams) const override;
+};
diff --git a/Engine/Source/Editor/Sequencer/Private/OutlinerColumns/SoloOutlinerColumn.cpp b/Engine/Source/Editor/Sequencer/Private/OutlinerColumns/SoloOutlinerColumn.cpp
new file mode 100644
--- /dev/null
+++ b/Engine/Source/Editor/Sequencer/Private/OutlinerColumns/SoloOutlinerColumn.cpp
@@ -0,0 +1,55 @@
+// Copyright Epic Games, Inc. All Rights Reserved.
+
+#include "OutlinerColumns/SoloOutlinerColumn.h"
+
+#include "ISequencerOutlinerColumn.h"
+#include "MVVM/Extensions/IOutlinerExtension.h"
+#include "MVVM/SoloEditorExtension.h"
+#include "MVVM/ViewModels/EditorViewModel.h"
+#include "MVVM/ViewModels/SequencerEditorViewModel.h"
+#include "Widgets/Layout/SScaleBox.h"
+#include "Widgets/OutlinerColumns/SSoloColumnWidget.h"
+
+#define LOCTEXT_NAMESPACE "FSoloOutlinerColumn"
+
+FName FSoloOutlinerColumn::GetColumnName() const
+{
+	return FName(TEXT("Solo"));
+}
+
+FText FSoloOutlinerColumn::GetColumnLabel() const
+{
+	return LOCTEXT("SoloColumnLabel", "Solo");
+}
+
+bool FSoloOutlinerColumn::IsItemCompatibleWithColumn(const UE::Sequencer::FCreateOutlinerColumnParams& InParams) const
+{
+	using namespace UE::Sequencer;
+
+	if (InParams.Editor)
+	{
+		FSoloEditorExtension* SoloEditorExtension = InParams.Editor->CastDynamic<FSoloEditorExtension>();
+		if (SoloEditorExtension)
+		{
+			return SoloEditorExtension->IsNodeSoloable(InParams.OutlinerExtension);
+		}
+	}
+
+	return false;
+}
+
+TSharedRef<ISequencerOutlinerColumn> FSoloOutlinerColumn::CreateOutlinerColumn()
+{
+	return MakeShareable(new FSoloOutlinerColumn());
+}
+
+TSharedRef<SWidget> FSoloOutlinerColumn::CreateColumnWidget(const TWeakPtr<ISequencerOutlinerColumn> InWeakOutlinerColumn, const UE::Sequencer::FCreateOutlinerColumnParams& InParams) const
+{
+	using namespace UE::Sequencer;
+
+	return SNew(SSoloColumnWidget,
+		InWeakOutlinerColumn,
+		InParams);
+}
+
+#undef LOCTEXT_NAMESPACE
\ No newline at end of file
diff --git a/Engine/Source/Editor/Sequencer/Private/OutlinerColumns/SoloOutlinerColumn.h b/Engine/Source/Editor/Sequencer/Private/OutlinerColumns/SoloOutlinerColumn.h
new file mode 100644
--- /dev/null
+++ b/Engine/Source/Editor/Sequencer/Private/OutlinerColumns/SoloOutlinerColumn.h
@@ -0,0 +1,39 @@
+// Copyright Epic Games, Inc. All Rights Reserved.
+
+#pragma once
+
+#include "ISequencerOutlinerColumn.h"
+
+/**
+ * A column for soloing/unsoloing tracks.
+ */
+class FSoloOutlinerColumn
+	: public ISequencerOutlinerColumn
+{
+public:
+
+	FSoloOutlinerColumn()
+	{ }
+
+	/** Returns the name of the column. Used for determining column type when dragging or saving settings. */
+	virtual FName GetColumnName() const override;
+
+	/** Returns the label of the column to display in visibility settings. */
+	virtual FText GetColumnLabel() const override;
+
+	/* Gets whether or not an outliner column widget should be generated for a given item in the outliner column. */
+	virtual bool IsItemCompatibleWithColumn(const UE::Sequencer::FCreateOutlinerColumnParams& InParams) const override;
+
+	/**
+	 * Creates an instance of this class.
+	 *
+	 * @return The new instance of this class.
+	 */
+	static TSharedRef<ISequencerOutlinerColumn> CreateOutlinerColumn();
+	
+
+	/**
+	 * Creates a widget that solos an item in the outliner view.
+	 */
+	virtual TSharedRef<SWidget> CreateColumnWidget(const TWeakPtr<ISequencerOutlinerColumn> InWeakOutlinerColumn, const UE::Sequencer::FCreateOutlinerColumnParams& InParams) const override;
+};
diff --git a/Engine/Source/Editor/Sequencer/Private/Sequencer.cpp b/Engine/Source/Editor/Sequencer/Private/Sequencer.cpp
--- a/Engine/Source/Editor/Sequencer/Private/Sequencer.cpp
+++ b/Engine/Source/Editor/Sequencer/Private/Sequencer.cpp
@@ -487,7 +487,7 @@ void FSequencer::InitSequencer(const FSequencerInitParams& InitParams, const TSh
 	for (int32 DelegateIndex = 0; DelegateIndex < OutlinerColumnDelegates.Num(); ++DelegateIndex)
 	{
 		check(OutlinerColumnDelegates[DelegateIndex].IsBound());
-		TSharedRef<ISequencerOutlinerColumn> OutlinerColumn = OutlinerColumnDelegates[DelegateIndex].Execute(SharedThis(this));
+		TSharedRef<ISequencerOutlinerColumn> OutlinerColumn = OutlinerColumnDelegates[DelegateIndex].Execute();
 
 		if (OutlinerColumn->SupportsSequence(InitParams.RootSequence))
 		{
diff --git a/Engine/Source/Editor/Sequencer/Private/SequencerModule.cpp b/Engine/Source/Editor/Sequencer/Private/SequencerModule.cpp
--- a/Engine/Source/Editor/Sequencer/Private/SequencerModule.cpp
+++ b/Engine/Source/Editor/Sequencer/Private/SequencerModule.cpp
@@ -31,6 +31,11 @@
 #include "MVVM/ViewModels/SequenceModel.h"
 #include "MVVM/ViewModels/SequencerEditorViewModel.h"
 
+#include "OutlinerColumns/LockOutlinerColumn.h"
+#include "OutlinerColumns/MuteOutlinerColumn.h"
+#include "OutlinerColumns/PinOutlinerColumn.h"
+#include "OutlinerColumns/SoloOutlinerColumn.h"
+
 #include "ToolMenus.h"
 #include "ContentBrowserMenuContexts.h"
 #include "SequencerUtilities.h"
@@ -341,6 +346,12 @@ public:
 
 			FPropertyEditorModule& EditModule = FModuleManager::Get().GetModuleChecked<FPropertyEditorModule>("PropertyEditor");
 			OnGetGlobalRowExtensionHandle = EditModule.GetGlobalRowExtensionDelegate().AddStatic(&RegisterKeyframeExtensionHandler);
+
+			// Register built-in outliner columns (default order is pin, mute, lock, solo)
+			PinOutlinerColumnHandle = RegisterOutlinerColumn(FOnCreateOutlinerColumn::CreateStatic(&FPinOutlinerColumn::CreateOutlinerColumn));
+			MuteOutlinerColumnHandle = RegisterOutlinerColumn(FOnCreateOutlinerColumn::CreateStatic(&FMuteOutlinerColumn::CreateOutlinerColumn));
+			LockOutlinerColumnHandle = RegisterOutlinerColumn(FOnCreateOutlinerColumn::CreateStatic(&FLockOutlinerColumn::CreateOutlinerColumn));
+			SoloOutlinerColumnHandle = RegisterOutlinerColumn(FOnCreateOutlinerColumn::CreateStatic(&FSoloOutlinerColumn::CreateOutlinerColumn));
 		}
 
 		FSequenceModel::CreateExtensionsEvent.AddLambda(
@@ -383,6 +394,12 @@ public:
 			}
 
 			FEditorModeRegistry::Get().UnregisterMode(FSequencerEdMode::EM_SequencerMode);
+
+			// unregister outliner columns
+			UnregisterOutlinerColumn(PinOutlinerColumnHandle);
+			UnregisterOutlinerColumn(MuteOutlinerColumnHandle);
+			UnregisterOutlinerColumn(LockOutlinerColumnHandle);
+			UnregisterOutlinerColumn(SoloOutlinerColumnHandle);
 		}
 	}
 
@@ -528,6 +545,12 @@ private:
 
 	/** Array of movie renderers */
 	TArray<FMovieRendererEntry> MovieRenderers;
+
+	// Outliner Column Delegate Handles
+	FDelegateHandle PinOutlinerColumnHandle;
+	FDelegateHandle MuteOutlinerColumnHandle;
+	FDelegateHandle LockOutlinerColumnHandle;
+	FDelegateHandle SoloOutlinerColumnHandle;
 };
 
 IMPLEMENT_MODULE(FSequencerModule, Sequencer);
diff --git a/Engine/Source/Editor/Sequencer/Private/SequencerSettings.cpp b/Engine/Source/Editor/Sequencer/Private/SequencerSettings.cpp
--- a/Engine/Source/Editor/Sequencer/Private/SequencerSettings.cpp
+++ b/Engine/Source/Editor/Sequencer/Private/SequencerSettings.cpp
@@ -994,11 +994,11 @@ void USequencerSettings::SetTrackFilterEnabled(const FString & TrackFilter, bool
 	}
 }
 
-void USequencerSettings::SetOutlinerColumnVisibility(const TArray<FColumnVisibilitySetting>& InHiddenColumns)
+void USequencerSettings::SetOutlinerColumnVisibility(const TArray<FColumnVisibilitySetting>& InColumnVisibilitySettings)
 {
-	if (InHiddenColumns != HiddenColumns)
+	if (InColumnVisibilitySettings != ColumnVisibilitySettings)
 	{
-		HiddenColumns = InHiddenColumns;
+		ColumnVisibilitySettings = InColumnVisibilitySettings;
 		SaveConfig();
 	}
 }
diff --git a/Engine/Source/Editor/Sequencer/Private/Widgets/OutlinerColumns/SColumnToggleWidget.cpp b/Engine/Source/Editor/Sequencer/Private/Widgets/OutlinerColumns/SColumnToggleWidget.cpp
new file mode 100644
--- /dev/null
+++ b/Engine/Source/Editor/Sequencer/Private/Widgets/OutlinerColumns/SColumnToggleWidget.cpp
@@ -0,0 +1,260 @@
+// Copyright Epic Games, Inc. All Rights Reserved.
+
+#include "Widgets/OutlinerColumns/SColumnToggleWidget.h"
+
+#include "Delegates/Delegate.h"
+#include "ISequencerOutlinerColumn.h"
+#include "MVVM/ViewModels/SequencerEditorViewModel.h"
+#include "MVVM/ViewModels/OutlinerViewModel.h"
+#include "Sequencer.h"
+
+#define LOCTEXT_NAMESPACE "SColumnToggleWidget"
+
+namespace UE::Sequencer
+{
+
+class FColumnToggleDragDropOp : public FDragDropOperation, public TSharedFromThis<FColumnToggleDragDropOp>
+{
+public:
+
+	DRAG_DROP_OPERATOR_TYPE(FColumnToggleDragDropOp, FDragDropOperation)
+
+	/** Flag which defines whether to hide destination items or not */
+	bool bActive;
+
+	/** Name of the column that is being dragged to not affect other columns */
+	FName ColumnName;
+
+	/** Widget determines event when drag drop operation is cancelled */
+	TWeakPtr<SColumnToggleWidget> ColumnToggleWidget;
+
+	/** Which is kept alive for the duration of the drag */
+	TUniquePtr<FScopedTransaction> UndoTransaction;
+
+	virtual TSharedPtr<SWidget> GetDefaultDecorator() const override
+	{
+		return SNullWidget::NullWidget;
+	}
+
+	/** Called when the drag drop operation ends */
+	virtual void OnDrop(bool bDropWasHandled, const FPointerEvent& MouseEvent)
+	{
+		if (TSharedPtr<SColumnToggleWidget> CurrentColumnToggleWidget = ColumnToggleWidget.Pin())
+		{
+			CurrentColumnToggleWidget->OnToggleOperationComplete();
+		}
+	}
+
+	/** Create a new drag and drop operation out of the specified flag */
+	static TSharedRef<FColumnToggleDragDropOp> New(const bool _bActive, FName _ColumnName, TWeakPtr<SColumnToggleWidget> _ColumnToggleWidget, TUniquePtr<FScopedTransaction>& ScopedTransaction)
+	{
+		TSharedRef<FColumnToggleDragDropOp> Operation = MakeShareable(new FColumnToggleDragDropOp);
+
+		Operation->bActive = _bActive;
+		Operation->ColumnName = _ColumnName;
+		Operation->ColumnToggleWidget = _ColumnToggleWidget;
+		Operation->UndoTransaction = MoveTemp(ScopedTransaction);
+
+		Operation->Construct();
+		return Operation;
+	}
+};
+
+void SColumnToggleWidget::Construct(
+	const FArguments& InArgs,
+	const TWeakPtr<ISequencerOutlinerColumn> InOutlinerColumn,
+	const FCreateOutlinerColumnParams& InParams)
+{
+	WeakOutlinerColumn = InOutlinerColumn;
+
+	WeakOutlinerExtension = InParams.OutlinerExtension;
+	WeakEditor = InParams.Editor->CastThisSharedChecked<FSequencerEditorViewModel>();
+
+	bIsMouseOverWidget = false;
+	bIsChildActive = false;
+
+	ActiveBrush = GetActiveBrush();
+
+	static const FName NAME_ChildActiveBrush = TEXT("Sequencer.Column.CheckBoxIndeterminate");
+	ChildActiveBrush = FAppStyle::Get().GetBrush(NAME_ChildActiveBrush);
+
+	SImage::Construct(
+		SImage::FArguments()
+		.Image(this, &SColumnToggleWidget::GetBrush)
+		.IsEnabled(this, &SColumnToggleWidget::IsEnabled)
+		.ColorAndOpacity(this, &SColumnToggleWidget::GetImageColorAndOpacity)
+	);
+}
+
+
+FSlateColor SColumnToggleWidget::GetImageColorAndOpacity() const
+{
+	TSharedPtr<FEditorViewModel> Editor = WeakEditor.Pin();
+	TViewModelPtr<IOutlinerExtension> OutlinerItem = WeakOutlinerExtension.Pin();
+
+	FLinearColor OutColor = FLinearColor::White;
+
+	if (!Editor || !OutlinerItem)
+	{
+		return OutColor;
+	}
+
+	float Opacity = 0.0f;
+
+	if (IsActive()
+		|| bIsMouseOverWidget)
+	{
+		Opacity = 1.0f;
+	}
+	else if (Editor->GetOutliner()->GetHoveredItem() == OutlinerItem)
+	{
+		Opacity = .2f;
+	}
+	else if (IsChildActive())
+	{
+		Opacity = 1.0f;
+	}
+	else if (IsImplicitlyActive())
+	{
+		Opacity = .4f;
+	}
+	else
+	{
+		Opacity = 0.0f;
+	}
+
+	OutColor.A = Opacity;
+	return OutColor;
+}
+
+const FSlateBrush* SColumnToggleWidget::GetBrush() const
+{
+	TSharedPtr<FEditorViewModel> Editor = WeakEditor.Pin();
+	TViewModelPtr<IOutlinerExtension> OutlinerItem = WeakOutlinerExtension.Pin();
+
+	if (!Editor || !OutlinerItem)
+	{
+		return ActiveBrush;
+	}
+
+	const bool bIsMouseOverItem = bIsMouseOverWidget || Editor->GetOutliner()->GetHoveredItem() == OutlinerItem;
+
+	if (bIsMouseOverItem
+		|| IsActive()
+		|| !IsChildActive())
+	{
+		return ActiveBrush;
+	}
+	else
+	{
+		return ChildActiveBrush;
+	}
+}
+
+FReply SColumnToggleWidget::OnDragDetected(const FGeometry& MyGeometry, const FPointerEvent& MouseEvent)
+{
+	TSharedPtr<ISequencerOutlinerColumn> OutlinerColumn = WeakOutlinerColumn.Pin();
+	if (MouseEvent.IsMouseButtonDown(EKeys::LeftMouseButton)
+		&& OutlinerColumn)
+	{
+		return FReply::Handled().BeginDragDrop(FColumnToggleDragDropOp::New(IsActive(), OutlinerColumn->GetColumnName(), SharedThis(this), UndoTransaction));
+	}
+	else
+	{
+		return FReply::Unhandled();
+	}
+}
+
+void SColumnToggleWidget::OnDragEnter(const FGeometry& MyGeometry, const FDragDropEvent& DragDropEvent)
+{
+	auto ColumnToggleDragOp = DragDropEvent.GetOperationAs<FColumnToggleDragDropOp>();
+	TSharedPtr<ISequencerOutlinerColumn> OutlinerColumn = WeakOutlinerColumn.Pin();
+	if (ColumnToggleDragOp.IsValid()
+		&& OutlinerColumn
+		&& ColumnToggleDragOp->ColumnName == OutlinerColumn->GetColumnName())
+	{
+		SetIsActive(ColumnToggleDragOp->bActive);
+	}
+}
+
+FReply SColumnToggleWidget::HandleClick()
+{
+	const bool bIsActive = IsActive();
+
+	TSharedPtr<ISequencerOutlinerColumn> OutlinerColumn = WeakOutlinerColumn.Pin();
+
+	// Open an undo transaction
+	if (OutlinerColumn)
+	{
+		UndoTransaction.Reset(new FScopedTransaction(FText::Format(LOCTEXT("Toggle Outliner Column", "Toggle {0}"), OutlinerColumn->GetColumnLabel())));
+	}
+
+	if (bIsActive)
+	{
+		SetIsActive(false);
+	}
+	else
+	{
+		SetIsActive(true);
+	}
+
+	return FReply::Handled().DetectDrag(SharedThis(this), EKeys::LeftMouseButton);
+}
+
+FReply SColumnToggleWidget::OnMouseButtonDoubleClick(const FGeometry& InMyGeometry, const FPointerEvent& InMouseEvent)
+{
+	return HandleClick();
+}
+
+FReply SColumnToggleWidget::OnMouseButtonDown(const FGeometry& MyGeometry, const FPointerEvent& MouseEvent)
+{
+	if (MouseEvent.GetEffectingButton() != EKeys::LeftMouseButton)
+	{
+		return FReply::Unhandled();
+	}
+
+	return HandleClick();
+}
+
+
+void SColumnToggleWidget::OnMouseCaptureLost(const FCaptureLostEvent& CaptureLostEvent)
+{
+	OnToggleOperationComplete();
+	UndoTransaction.Reset();
+}
+
+FReply SColumnToggleWidget::OnMouseButtonUp(const FGeometry& MyGeometry, const FPointerEvent& MouseEvent)
+{
+	if (MouseEvent.GetEffectingButton() == EKeys::LeftMouseButton)
+	{
+		OnToggleOperationComplete();
+		UndoTransaction.Reset();
+		return FReply::Handled();
+	}
+
+	return FReply::Unhandled();
+}
+
+void SColumnToggleWidget::OnMouseEnter(const FGeometry& MyGeometry, const FPointerEvent& MouseEvent)
+{
+	bIsMouseOverWidget = true;
+}
+
+void SColumnToggleWidget::OnMouseLeave(const FPointerEvent& MouseEvent)
+{
+	bIsMouseOverWidget = false;
+}
+
+void SColumnToggleWidget::RefreshSequencerTree()
+{
+	TSharedPtr<FSequencerEditorViewModel> EditorViewModel = WeakEditor.Pin();
+	TSharedPtr<FSequencer> Sequencer = EditorViewModel ? EditorViewModel->GetSequencerImpl() : nullptr;
+	if (Sequencer)
+	{
+		Sequencer->RefreshTree();
+	}
+}
+
+} // namespace UE::Sequencer
+
+#undef LOCTEXT_NAMESPACE
diff --git a/Engine/Source/Editor/Sequencer/Private/Widgets/OutlinerColumns/SColumnToggleWidget.h b/Engine/Source/Editor/Sequencer/Private/Widgets/OutlinerColumns/SColumnToggleWidget.h
new file mode 100644
--- /dev/null
+++ b/Engine/Source/Editor/Sequencer/Private/Widgets/OutlinerColumns/SColumnToggleWidget.h
@@ -0,0 +1,131 @@
+// Copyright Epic Games, Inc. All Rights Reserved.
+
+#pragma once
+
+#include "ISequencerOutlinerColumn.h"
+#include "MVVM/Extensions/IOutlinerExtension.h"
+#include "MVVM/ViewModelPtr.h"
+#include "ScopedTransaction.h"
+#include "Templates/UniquePtr.h"
+
+namespace UE::Sequencer
+{
+
+class FSequencerEditorViewModel;
+
+/**
+ * Outliner Column Widget for toggling active state with click and drag functionality.
+ * Adjusts image color and transparency based on active state and item hover.
+ * Displays an alternate icon when a child is active and the item is not.
+ * Column Toggle widgets have 4 States: Active, Inactive, Active Child, and Implicitly Active.
+ * Column Toggle widgets will change its brush based on it's state.
+ */
+class SColumnToggleWidget
+	: public SImage
+{
+public:
+
+	SLATE_BEGIN_ARGS(SColumnToggleWidget) {}
+	SLATE_END_ARGS()
+
+	void Construct(
+		const FArguments& InArgs,
+		const TWeakPtr<ISequencerOutlinerColumn> InOutlinerColumn,
+		const FCreateOutlinerColumnParams& InParams);
+
+public:
+
+	/** Called when a drag or click ends and can be used to refresh the SequencerTree if necessary. */
+	virtual void OnToggleOperationComplete() {};
+
+protected:
+
+	/** Returns whether the widget is enabled or not. */
+	virtual bool IsEnabled() const { return true; }
+
+protected:
+
+	/** Returns whether or not the item is directly active or not. */
+	virtual bool IsActive() const = 0;
+
+	/** Sets the active state of this item to the input value. */
+	virtual void SetIsActive(const bool bInIsActive) = 0;
+
+	/** Returns true if a child of this item is active. */
+	virtual bool IsChildActive() const = 0;
+
+	/** Returns true if this item is implicitly active, but not directly active */
+	virtual bool IsImplicitlyActive() const { return false; }
+
+	/** Returns the brush to be used to represent a widget is active. */
+	virtual const FSlateBrush* GetActiveBrush() const = 0;
+
+protected:
+	
+	/** Get the color and opacity of the column toggle widget. */
+	virtual FSlateColor GetImageColorAndOpacity() const;
+
+	/** Get the image this widget displays. */
+	virtual const FSlateBrush* GetBrush() const;
+
+protected:
+
+	/** Start a new drag/drop operation for this widget. */
+	virtual FReply OnDragDetected(const FGeometry& MyGeometry, const FPointerEvent& MouseEvent) override;
+
+	/** If a matching column drag drop operation has entered this widget, set its item to the specified active state. */
+	virtual void OnDragEnter(const FGeometry& MyGeometry, const FDragDropEvent& DragDropEvent) override;
+
+	/** Handles left click-type inputs and potentially begins drag and drop operation. */
+	FReply HandleClick();
+
+	/** Called when the mouse button is double clicked on this widget. */
+	virtual FReply OnMouseButtonDoubleClick(const FGeometry& InMyGeometry, const FPointerEvent& InMouseEvent) override;
+
+	/** Called when the mouse button is pressed down on this widget. */
+	virtual FReply OnMouseButtonDown(const FGeometry& MyGeometry, const FPointerEvent& MouseEvent) override;
+
+	/** Called when this widget had captured the mouse, but that capture has been revoked for some reason. */
+	virtual void OnMouseCaptureLost(const FCaptureLostEvent& CaptureLostEvent) override;
+
+	/** Process a mouse up message. */
+	virtual FReply OnMouseButtonUp(const FGeometry& MyGeometry, const FPointerEvent& MouseEvent) override;
+
+	/** Called when mouse enters the column widget. */
+	void OnMouseEnter(const FGeometry& MyGeometry, const FPointerEvent& MouseEvent) override;
+
+	/** Called when mouse leaves the column widget. */
+	void OnMouseLeave(const FPointerEvent& MouseEvent) override;
+
+protected:
+
+	/** Refreshes the Sequencer Tree. Can be used after toggle operation is complete if operation affects other items. */
+	void RefreshSequencerTree();
+
+protected:
+	/** Editor view-model for this view. */
+	TWeakPtr<FSequencerEditorViewModel> WeakEditor;
+
+	/** Outliner item this widget accesses. */
+	TWeakViewModelPtr<IOutlinerExtension> WeakOutlinerExtension;
+
+	/** Reference to column this widget is a member of. */
+	TWeakPtr<ISequencerOutlinerColumn> WeakOutlinerColumn;
+
+private:
+
+	/** Whether or not the mouse is directly over the widget. */
+	bool bIsMouseOverWidget;
+
+	/** Whether or not a child of this item is active. */
+	bool bIsChildActive;
+
+	/** Brushes for the various states of active widgets. */
+	const FSlateBrush* ActiveBrush;
+	const FSlateBrush* ChildActiveBrush;
+
+	/** Scoped undo transaction. */
+	TUniquePtr<FScopedTransaction> UndoTransaction;
+};
+
+} // namespace UE::Sequencer
\ No newline at end of file
diff --git a/Engine/Source/Editor/Sequencer/Private/Widgets/OutlinerColumns/SLockColumnWidget.cpp b/Engine/Source/Editor/Sequencer/Private/Widgets/OutlinerColumns/SLockColumnWidget.cpp
new file mode 100644
--- /dev/null
+++ b/Engine/Source/Editor/Sequencer/Private/Widgets/OutlinerColumns/SLockColumnWidget.cpp
@@ -0,0 +1,69 @@
+// Copyright Epic Games, Inc. All Rights Reserved.
+
+#include "Widgets/OutlinerColumns/SLockColumnWidget.h"
+
+#include "MVVM/LockEditorExtension.h"
+#include "MVVM/ViewModels/SequencerEditorViewModel.h"
+
+namespace UE::Sequencer
+{
+
+void SLockColumnWidget::Construct(const FArguments& InArgs, const TWeakPtr<ISequencerOutlinerColumn> InWeakOutlinerColumn, const UE::Sequencer::FCreateOutlinerColumnParams& InParams)
+{
+	SColumnToggleWidget::Construct(
+		SColumnToggleWidget::FArguments(),
+		InWeakOutlinerColumn,
+		InParams
+	);
+}
+
+bool SLockColumnWidget::IsActive() const
+{
+	TSharedPtr<FSequencerEditorViewModel> SequencerEditor = WeakEditor.Pin();
+	if (SequencerEditor)
+	{
+		FLockEditorExtension* LockEditorExtension = SequencerEditor->CastDynamic<FLockEditorExtension>();
+		if (LockEditorExtension)
+		{
+			return LockEditorExtension->IsNodeLocked(WeakOutlinerExtension);
+		}
+	}
+
+	return false;
+}
+
+void SLockColumnWidget::SetIsActive(const bool bInIsActive)
+{
+	TSharedPtr<FSequencerEditorViewModel> SequencerEditor = WeakEditor.Pin();
+	if (SequencerEditor)
+	{
+		FLockEditorExtension* LockEditorExtension = SequencerEditor->CastDynamic<FLockEditorExtension>();
+		if (LockEditorExtension)
+		{
+			LockEditorExtension->SetNodeLocked(WeakOutlinerExtension, bInIsActive);
+		}
+	}
+}
+
+bool SLockColumnWidget::IsChildActive() const
+{
+	TSharedPtr<FSequencerEditorViewModel> SequencerEditor = WeakEditor.Pin();
+	if (SequencerEditor)
+	{
+		FLockEditorExtension* LockEditorExtension = SequencerEditor->CastDynamic<FLockEditorExtension>();
+		if (LockEditorExtension)
+		{
+			return LockEditorExtension->HasLockedChildNode(WeakOutlinerExtension);
+		}
+	}
+
+	return false;
+}
+
+const FSlateBrush* SLockColumnWidget::GetActiveBrush() const
+{
+	static const FName NAME_LockBrush = TEXT("Sequencer.Column.Locked");
+	return FAppStyle::Get().GetBrush(NAME_LockBrush);
+}
+
+} // namespace UE::Sequencer
diff --git a/Engine/Source/Editor/Sequencer/Private/Widgets/OutlinerColumns/SLockColumnWidget.h b/Engine/Source/Editor/Sequencer/Private/Widgets/OutlinerColumns/SLockColumnWidget.h
new file mode 100644
--- /dev/null
+++ b/Engine/Source/Editor/Sequencer/Private/Widgets/OutlinerColumns/SLockColumnWidget.h
@@ -0,0 +1,36 @@
+// Copyright Epic Games, Inc. All Rights Reserved.
+
+#pragma once
+
+#include "Widgets/OutlinerColumns/SColumnToggleWidget.h"
+
+namespace UE::Sequencer
+{
+
+/**
+* A widget that shows and controls the locked state of outliner items.
+*/
+class SLockColumnWidget
+	: public SColumnToggleWidget
+{
+public:
+	SLATE_BEGIN_ARGS(SLockColumnWidget) {}
+	SLATE_END_ARGS()
+
+	void Construct(const FArguments& InArgs, const TWeakPtr<ISequencerOutlinerColumn> InWeakOutlinerColumn, const FCreateOutlinerColumnParams& InParams);
+
+public:
+	/** Returns whether or not this outliner item is locked. */
+	virtual bool IsActive() const override;
+
+	/** Sets this item and it's sections as locked or unlocked. If selected, applies to all selected items. */
+	virtual void SetIsActive(const bool bInIsActive) override;
+
+	/** Returns whether or not a child of this item is locked. */
+	virtual bool IsChildActive() const override;
+
+	/** Returns the brush to display when this item is locked. */
+	virtual const FSlateBrush* GetActiveBrush() const override;
+};
+
+} // namespace UE::Sequencer
\ No newline at end of file
diff --git a/Engine/Source/Editor/Sequencer/Private/Widgets/OutlinerColumns/SMuteColumnWidget.cpp b/Engine/Source/Editor/Sequencer/Private/Widgets/OutlinerColumns/SMuteColumnWidget.cpp
new file mode 100644
--- /dev/null
+++ b/Engine/Source/Editor/Sequencer/Private/Widgets/OutlinerColumns/SMuteColumnWidget.cpp
@@ -0,0 +1,90 @@
+// Copyright Epic Games, Inc. All Rights Reserved.
+
+#include "Widgets/OutlinerColumns/SMuteColumnWidget.h"
+
+#include "MVVM/MuteEditorExtension.h"
+#include "MVVM/ViewModels/SequencerEditorViewModel.h"
+
+namespace UE::Sequencer
+{
+
+void SMuteColumnWidget::OnToggleOperationComplete()
+{
+	// refresh the sequencer tree after operation is complete
+	RefreshSequencerTree();
+}
+
+void SMuteColumnWidget::Construct(const FArguments& InArgs, const TWeakPtr<ISequencerOutlinerColumn> InWeakOutlinerColumn, const UE::Sequencer::FCreateOutlinerColumnParams& InParams)
+{
+	SColumnToggleWidget::Construct(
+		SColumnToggleWidget::FArguments(),
+		InWeakOutlinerColumn,
+		InParams
+	);
+}
+
+bool SMuteColumnWidget::IsActive() const
+{
+	TSharedPtr<FSequencerEditorViewModel> SequencerEditor = WeakEditor.Pin();
+	if (SequencerEditor)
+	{
+		FMuteEditorExtension* MuteEditorExtension = SequencerEditor->CastDynamic<FMuteEditorExtension>();
+		if (MuteEditorExtension)
+		{
+			return MuteEditorExtension->IsNodeMuted(WeakOutlinerExtension);
+		}
+	}
+
+	return false;
+}
+
+void SMuteColumnWidget::SetIsActive(const bool bInIsActive)
+{
+	TSharedPtr<FSequencerEditorViewModel> SequencerEditor = WeakEditor.Pin();
+	if (SequencerEditor)
+	{
+		FMuteEditorExtension* MuteEditorExtension = SequencerEditor->CastDynamic<FMuteEditorExtension>();
+		if (MuteEditorExtension)
+		{
+			MuteEditorExtension->SetNodeMuted(WeakOutlinerExtension, bInIsActive);
+		}
+	}
+}
+
+bool SMuteColumnWidget::IsChildActive() const
+{
+	TSharedPtr<FSequencerEditorViewModel> SequencerEditor = WeakEditor.Pin();
+	if (SequencerEditor)
+	{
+		FMuteEditorExtension* MuteEditorExtension = SequencerEditor->CastDynamic<FMuteEditorExtension>();
+		if (MuteEditorExtension)
+		{
+			return MuteEditorExtension->HasMutedChildNode(WeakOutlinerExtension);
+		}
+	}
+
+	return false;
+}
+
+bool SMuteColumnWidget::IsImplicitlyActive() const
+{
+	TSharedPtr<FSequencerEditorViewModel> SequencerEditor = WeakEditor.Pin();
+	if (SequencerEditor)
+	{
+		FMuteEditorExtension* MuteEditorExtension = SequencerEditor->CastDynamic<FMuteEditorExtension>();
+		if (MuteEditorExtension)
+		{
+			return MuteEditorExtension->IsNodeImplicitlyMuted(WeakOutlinerExtension);
+		}
+	}
+
+	return false;
+}
+
+const FSlateBrush* SMuteColumnWidget::GetActiveBrush() const
+{
+	static const FName NAME_MutedBrush = TEXT("Sequencer.Column.Mute");
+	return FAppStyle::Get().GetBrush(NAME_MutedBrush);
+}
+
+} // namespace UE::Sequencer
diff --git a/Engine/Source/Editor/Sequencer/Private/Widgets/OutlinerColumns/SMuteColumnWidget.h b/Engine/Source/Editor/Sequencer/Private/Widgets/OutlinerColumns/SMuteColumnWidget.h
new file mode 100644
--- /dev/null
+++ b/Engine/Source/Editor/Sequencer/Private/Widgets/OutlinerColumns/SMuteColumnWidget.h
@@ -0,0 +1,44 @@
+// Copyright Epic Games, Inc. All Rights Reserved.
+
+#pragma once
+
+#include "Widgets/OutlinerColumns/SColumnToggleWidget.h"
+
+namespace UE::Sequencer
+{
+
+/**
+ * A widget that shows and controls the muted state of outliner items.
+ */
+class SMuteColumnWidget
+	: public SColumnToggleWidget
+{
+public:
+	SLATE_BEGIN_ARGS(SMuteColumnWidget) {}
+	SLATE_END_ARGS()
+
+	void Construct(const FArguments& InArgs, const TWeakPtr<ISequencerOutlinerColumn> InWeakOutlinerColumn, const FCreateOutlinerColumnParams& InParams);
+
+public:
+
+	/** Refreshes the Sequencer Tree to handle simultaneous mute and solo modifiers. */
+	virtual void OnToggleOperationComplete();
+
+protected:
+	/** Returns whether or not this item is muted. */
+	virtual bool IsActive() const override;
+
+	/** Sets this item as muted or unmuted. If selected, applies to all selected items. */
+	virtual void SetIsActive(const bool bInIsActive) override;
+
+	/** Returns whether or not a child of this item is muted. */
+	virtual bool IsChildActive() const override;
+
+	/** Returns true if this item is implicitly active, but not directly active */
+	virtual bool IsImplicitlyActive() const override;
+
+	/** Returns the brush to display when this item is muted. */
+	virtual const FSlateBrush* GetActiveBrush() const override;
+};
+
+} // namespace UE::Sequencer
\ No newline at end of file
diff --git a/Engine/Source/Editor/Sequencer/Private/Widgets/OutlinerColumns/SPinColumnWidget.cpp b/Engine/Source/Editor/Sequencer/Private/Widgets/OutlinerColumns/SPinColumnWidget.cpp
new file mode 100644
--- /dev/null
+++ b/Engine/Source/Editor/Sequencer/Private/Widgets/OutlinerColumns/SPinColumnWidget.cpp
@@ -0,0 +1,67 @@
+// Copyright Epic Games, Inc. All Rights Reserved.
+
+#include "Widgets/OutlinerColumns/SPinColumnWidget.h"
+
+#include "MVVM/Extensions/IPinnableExtension.h"
+#include "MVVM/PinEditorExtension.h"
+#include "MVVM/ViewModels/SequencerEditorViewModel.h"
+
+namespace UE::Sequencer
+{
+
+void SPinColumnWidget::OnToggleOperationComplete()
+{
+	// refresh the sequencer tree after operation is complete
+	RefreshSequencerTree();
+}
+
+void SPinColumnWidget::Construct(const FArguments& InArgs, const TWeakPtr<ISequencerOutlinerColumn> InWeakOutlinerColumn, const FCreateOutlinerColumnParams& InParams)
+{
+	SColumnToggleWidget::Construct(
+		SColumnToggleWidget::FArguments(),
+		InWeakOutlinerColumn,
+		InParams
+	);
+}
+
+bool SPinColumnWidget::IsActive() const
+{
+	TSharedPtr<FSequencerEditorViewModel> SequencerEditor = WeakEditor.Pin();
+	if (SequencerEditor)
+	{
+		FPinEditorExtension* PinEditorExtension = SequencerEditor->CastDynamic<FPinEditorExtension>();
+		if (PinEditorExtension)
+		{
+			return PinEditorExtension->IsNodePinned(WeakOutlinerExtension);
+		}
+	}
+	
+	return false;
+}
+
+void SPinColumnWidget::SetIsActive(const bool bInIsActive)
+{
+	TSharedPtr<FSequencerEditorViewModel> SequencerEditor = WeakEditor.Pin();
+	if (SequencerEditor)
+	{
+		FPinEditorExtension* PinEditorExtension = SequencerEditor->CastDynamic<FPinEditorExtension>();
+		if (PinEditorExtension)
+		{
+			PinEditorExtension->SetNodePinned(WeakOutlinerExtension, bInIsActive);
+		}
+	}
+}
+
+bool SPinColumnWidget::IsChildActive() const
+{
+	// children of a pin widget cannot have a different pinned state from their parent
+	return false;
+}
+
+const FSlateBrush* SPinColumnWidget::GetActiveBrush() const
+{
+	static const FName NAME_PinnedBrush = TEXT("Sequencer.Column.Unpinned");
+	return FAppStyle::Get().GetBrush(NAME_PinnedBrush);
+}
+
+} // namespace UE::Sequencer
diff --git a/Engine/Source/Editor/Sequencer/Private/Widgets/OutlinerColumns/SPinColumnWidget.h b/Engine/Source/Editor/Sequencer/Private/Widgets/OutlinerColumns/SPinColumnWidget.h
new file mode 100644
--- /dev/null
+++ b/Engine/Source/Editor/Sequencer/Private/Widgets/OutlinerColumns/SPinColumnWidget.h
@@ -0,0 +1,45 @@
+// Copyright Epic Games, Inc. All Rights Reserved.
+
+#pragma once
+
+#include "Widgets/OutlinerColumns/SColumnToggleWidget.h"
+
+namespace UE::Sequencer
+{
+
+class IPinnableExtension;
+
+/**
+ * A widget that shows and controls the pinned state of outliner items.
+ */
+class SPinColumnWidget
+	: public SColumnToggleWidget
+{
+public:
+	SLATE_BEGIN_ARGS(SPinColumnWidget) {}
+	SLATE_END_ARGS()
+
+	void Construct(const FArguments& InArgs, const TWeakPtr<ISequencerOutlinerColumn> InWeakOutlinerColumn, const FCreateOutlinerColumnParams& InParams);
+
+public:
+
+	/** Refreshes the Sequencer Tree to update pinned state of items. */
+	virtual void OnToggleOperationComplete();
+
+protected:
+
+	/** Returns whether or not this outliner item is pinned. */
+	virtual bool IsActive() const override;
+
+	/** Sets this item as pinned or unpinned. If selected, applies to all selected items. */
+	virtual void SetIsActive(const bool bInIsActive) override;
+
+	/** Returns whether or not a child of this item is pinned. */
+	virtual bool IsChildActive() const override;
+
+	/** Returns the brush to display when this item is pinned. */
+	virtual const FSlateBrush* GetActiveBrush() const override;
+
+};
+
+} // namespace UE::Sequencer
\ No newline at end of file
diff --git a/Engine/Source/Editor/Sequencer/Private/Widgets/OutlinerColumns/SSoloColumnWidget.cpp b/Engine/Source/Editor/Sequencer/Private/Widgets/OutlinerColumns/SSoloColumnWidget.cpp
new file mode 100644
--- /dev/null
+++ b/Engine/Source/Editor/Sequencer/Private/Widgets/OutlinerColumns/SSoloColumnWidget.cpp
@@ -0,0 +1,91 @@
+// Copyright Epic Games, Inc. All Rights Reserved.
+
+#include "Widgets/OutlinerColumns/SSoloColumnWidget.h"
+
+#include "MVVM/SoloEditorExtension.h"
+#include "MVVM/ViewModels/EditorViewModel.h"
+#include "MVVM/ViewModels/SequencerEditorViewModel.h"
+
+namespace UE::Sequencer
+{
+
+void SSoloColumnWidget::OnToggleOperationComplete()
+{
+	// refresh the sequencer tree after operation is complete
+	RefreshSequencerTree();
+}
+
+void SSoloColumnWidget::Construct(const FArguments& InArgs, const TWeakPtr<ISequencerOutlinerColumn> InWeakOutlinerColumn, const UE::Sequencer::FCreateOutlinerColumnParams& InParams)
+{
+	SColumnToggleWidget::Construct(
+		SColumnToggleWidget::FArguments(),
+		InWeakOutlinerColumn,
+		InParams
+	);
+}
+
+bool SSoloColumnWidget::IsActive() const
+{
+	TSharedPtr<FSequencerEditorViewModel> SequencerEditor = WeakEditor.Pin();
+	if (SequencerEditor)
+	{
+		FSoloEditorExtension* SoloEditorExtension = SequencerEditor->CastDynamic<FSoloEditorExtension>();
+		if (SoloEditorExtension)
+		{
+			return SoloEditorExtension->IsNodeSoloed(WeakOutlinerExtension);
+		}
+	}
+
+	return false;
+}
+
+void SSoloColumnWidget::SetIsActive(const bool bInIsActive)
+{
+	TSharedPtr<FSequencerEditorViewModel> SequencerEditor = WeakEditor.Pin();
+	if (SequencerEditor)
+	{
+		FSoloEditorExtension* SoloEditorExtension = SequencerEditor->CastDynamic<FSoloEditorExtension>();
+		if (SoloEditorExtension)
+		{
+			SoloEditorExtension->SetNodeSoloed(WeakOutlinerExtension, bInIsActive);
+		}
+	}
+}
+
+bool SSoloColumnWidget::IsChildActive() const
+{
+	TSharedPtr<FSequencerEditorViewModel> SequencerEditor = WeakEditor.Pin();
+	if (SequencerEditor)
+	{
+		FSoloEditorExtension* SoloEditorExtension = SequencerEditor->CastDynamic<FSoloEditorExtension>();
+		if (SoloEditorExtension)
+		{
+			return SoloEditorExtension->HasSoloedChildNode(WeakOutlinerExtension);
+		}
+	}
+
+	return false;
+}
+
+bool SSoloColumnWidget::IsImplicitlyActive() const
+{
+	TSharedPtr<FSequencerEditorViewModel> SequencerEditor = WeakEditor.Pin();
+	if (SequencerEditor)
+	{
+		FSoloEditorExtension* SoloEditorExtension = SequencerEditor->CastDynamic<FSoloEditorExtension>();
+		if (SoloEditorExtension)
+		{
+			return SoloEditorExtension->IsNodeImplicitlySoloed(WeakOutlinerExtension);
+		}
+	}
+
+	return false;
+}
+
+const FSlateBrush* SSoloColumnWidget::GetActiveBrush() const
+{
+	static const FName NAME_SoloedBrush = TEXT("Sequencer.Column.Solo");
+	return FAppStyle::Get().GetBrush(NAME_SoloedBrush);
+}
+
+} // namespace UE::Sequencer
diff --git a/Engine/Source/Editor/Sequencer/Private/Widgets/OutlinerColumns/SSoloColumnWidget.h b/Engine/Source/Editor/Sequencer/Private/Widgets/OutlinerColumns/SSoloColumnWidget.h
new file mode 100644
--- /dev/null
+++ b/Engine/Source/Editor/Sequencer/Private/Widgets/OutlinerColumns/SSoloColumnWidget.h
@@ -0,0 +1,45 @@
+// Copyright Epic Games, Inc. All Rights Reserved.
+
+#pragma once
+
+#include "Widgets/OutlinerColumns/SColumnToggleWidget.h"
+
+namespace UE::Sequencer
+{
+
+/**
+* A widget that shows and controls the soloed state of outliner items.
+*/
+class SSoloColumnWidget
+	: public SColumnToggleWidget
+{
+public:
+	SLATE_BEGIN_ARGS(SSoloColumnWidget) {}
+	SLATE_END_ARGS()
+
+	void Construct(const FArguments& InArgs, const TWeakPtr<ISequencerOutlinerColumn> InWeakOutlinerColumn, const FCreateOutlinerColumnParams& InParams);
+
+public:
+
+	/** Refreshes the Sequencer Tree to show other outliner items as greyed out if they are not also soloed. */
+	virtual void OnToggleOperationComplete();
+
+protected:
+
+	/** Returns whether or not this outliner item is soloed. */
+	virtual bool IsActive() const override;
+
+	/** Sets this item as soloed or not soloed. If selected, applies to all selected items. */
+	virtual void SetIsActive(const bool bInIsActive) override;
+
+	/** Returns whether or not a child of this item is soloed. */
+	virtual bool IsChildActive() const override;
+
+	/** Returns true if this item is implicitly active, but not directly active */
+	virtual bool IsImplicitlyActive() const override;
+
+	/** Returns the brush to display when this item is pinned. */
+	virtual const FSlateBrush* GetActiveBrush() const override;
+};
+
+} // namespace UE::Sequencer
\ No newline at end of file
diff --git a/Engine/Source/Editor/Sequencer/Public/ISequencerModule.h b/Engine/Source/Editor/Sequencer/Public/ISequencerModule.h
--- a/Engine/Source/Editor/Sequencer/Public/ISequencerModule.h
+++ b/Engine/Source/Editor/Sequencer/Public/ISequencerModule.h
@@ -57,7 +57,7 @@ DECLARE_DELEGATE_RetVal_OneParam(TSharedRef<ISequencerEditorObjectBinding>, FOnC
 DECLARE_DELEGATE_RetVal_OneParam(TSharedPtr<UE::Sequencer::FTrackModel>, FOnCreateTrackModel, UMovieSceneTrack*);
 
 /** A delegate which will create an outliner column */
-DECLARE_DELEGATE_RetVal_OneParam(TSharedRef<ISequencerOutlinerColumn>, FOnCreateOutlinerColumn, TSharedRef<ISequencer>);
+DECLARE_DELEGATE_RetVal(TSharedRef<ISequencerOutlinerColumn>, FOnCreateOutlinerColumn);
 
 /** A delegate that is executed when adding menu content. */
 DECLARE_DELEGATE_OneParam(FOnGetContextMenuContent, FMenuBuilder& /*MenuBuilder*/);
diff --git a/Engine/Source/Editor/Sequencer/Public/MovieSceneOutlinerColumn.h b/Engine/Source/Editor/Sequencer/Public/MovieSceneOutlinerColumn.h
deleted file mode 100644
diff --git a/Engine/Source/Editor/Sequencer/Public/SequencerSettings.h b/Engine/Source/Editor/Sequencer/Public/SequencerSettings.h
--- a/Engine/Source/Editor/Sequencer/Public/SequencerSettings.h
+++ b/Engine/Source/Editor/Sequencer/Public/SequencerSettings.h
@@ -448,9 +448,9 @@ public:
 	void SetTrackFilterEnabled(const FString& TrackFilter, bool bEnabled);
 
 	/** Get outliner column visibility in display order */
-	TArray<FColumnVisibilitySetting> GetOutlinerColumnSettings() const { return HiddenColumns; }
+	TArray<FColumnVisibilitySetting> GetOutlinerColumnSettings() const { return ColumnVisibilitySettings; }
 	/** Sets the visibility of outliner columns in display order */
-	void SetOutlinerColumnVisibility(const TArray<FColumnVisibilitySetting>& InHiddenColumns);
+	void SetOutlinerColumnVisibility(const TArray<FColumnVisibilitySetting>& InColumnVisibilitySettings);
 
 protected:
 
@@ -706,7 +706,7 @@ protected:
 
 	/** List of all columns and their visibility, in the order to be displayed in the outliner view */
 	UPROPERTY(config, EditAnywhere, Category = General)
-	TArray<FColumnVisibilitySetting> HiddenColumns;
+	TArray<FColumnVisibilitySetting> ColumnVisibilitySettings;
 
 	FOnEvaluateSubSequencesInIsolationChanged OnEvaluateSubSequencesInIsolationChangedEvent;
 	FOnShowSelectedNodesOnlyChanged OnShowSelectedNodesOnlyChangedEvent;
diff --git a/Engine/Source/Editor/SequencerCore/Private/MVVM/Views/SOutlinerColumnBorder.cpp b/Engine/Source/Editor/SequencerCore/Private/MVVM/Views/SOutlinerColumnBorder.cpp
--- a/Engine/Source/Editor/SequencerCore/Private/MVVM/Views/SOutlinerColumnBorder.cpp
+++ b/Engine/Source/Editor/SequencerCore/Private/MVVM/Views/SOutlinerColumnBorder.cpp
@@ -19,7 +19,7 @@ void SOutlinerColumnInnerBorder::Construct(const FArguments& InArgs, const FCrea
 	WeakOutlinerExtension = InParams.OutlinerExtension;
 	WeakEditor = InParams.Editor;
 
-	BackgroundBrush = FAppStyle::GetBrush("Sequencer.AnimationOutliner.DefaultBorder");
+	BackgroundBrush = FAppStyle::GetBrush("Sequencer.Column.OutlinerColumnBox");
 
 	// Size of outliner column widgets is currently 12x12
 	TSharedRef<SWidget>	FinalWidget = SNew(SBox)
@@ -31,8 +31,7 @@ void SOutlinerColumnInnerBorder::Construct(const FArguments& InArgs, const FCrea
 			SNew(SBorder)
 			.VAlign(VAlign_Center)
 			.BorderImage(this, &SOutlinerColumnInnerBorder::GetBorderImage)
-			.BorderBackgroundColor(this, &SOutlinerColumnInnerBorder::GetBackgroundTint)
-			.Padding(FMargin(0.f))
+			.Padding(FMargin(1.f))
 			[
 				InArgs._Content.Widget
 			]
@@ -54,45 +53,45 @@ void SOutlinerColumnInnerBorder::OnMouseLeave(const FPointerEvent& MouseEvent)
 	bIsMouseOverInnerBorder = false;
 }
 
-const FSlateBrush* SOutlinerColumnInnerBorder::GetBorderImage() const
+FReply SOutlinerColumnBorder::OnMouseButtonDown(const FGeometry& MyGeometry, const FPointerEvent& MouseEvent)
 {
-	return BackgroundBrush;
+	return FReply::Handled();
 }
 
-FSlateColor SOutlinerColumnInnerBorder::GetBackgroundTint() const
+const FSlateBrush* SOutlinerColumnInnerBorder::GetBorderImage() const
 {
-	TSharedPtr<FEditorViewModel> Editor = WeakEditor.Pin();
-	TViewModelPtr<IOutlinerExtension> OutlinerItem = WeakOutlinerExtension.Pin();
-
-	EOutlinerSelectionState SelectionState = OutlinerItem->GetSelectionState();
-
-	bool bIsOutlinerItemSelectedDirectly = EnumHasAnyFlags(SelectionState, EOutlinerSelectionState::SelectedDirectly);
-	bool bIsOutlinerItemHovered = Editor->GetOutliner()->GetHoveredItem() == OutlinerItem;
-
-	float Opacity = .4f;
-	return FLinearColor(0.f, 0.f, 0.f, Opacity);
+	return BackgroundBrush;
 }
 
-
-void SOutlinerColumnBorder::Construct(const FArguments& InArgs, const FCreateOutlinerColumnParams& InParams)
+void SOutlinerColumnBorder::Construct(const FArguments& InArgs, const FCreateOutlinerColumnParams& InParams, const bool bHasInnerBorder)
 {
 	WeakOutlinerExtension = InParams.OutlinerExtension;
 	WeakEditor = InParams.Editor;
 
 	BackgroundBrush = FAppStyle::GetBrush("Sequencer.AnimationOutliner.DefaultBorder");
 
+	TSharedPtr<SWidget> InnerWidget;
+	if (bHasInnerBorder)
+	{
+		InnerWidget = SNew(SOutlinerColumnInnerBorder, InParams)
+			[
+				InArgs._Content.Widget
+			];
+	}
+	else
+	{
+		InnerWidget = InArgs._Content.Widget;
+	}
+
 	TSharedRef<SWidget>	FinalWidget = SNew(SBorder)
 		.VAlign(VAlign_Center)
 		.BorderImage(this, &SOutlinerColumnBorder::GetBorderImage)
 		.BorderBackgroundColor(this, &SOutlinerColumnBorder::GetBackgroundTint)
 		.VAlign(VAlign_Center)
 		.HAlign(HAlign_Center)
-		.Padding(FMargin(1.0f))
+		.Padding(FMargin(0.0f))
 		[
-			SNew(SOutlinerColumnInnerBorder, InParams)
-			[
-				InArgs._Content.Widget
-			]
+			InnerWidget.ToSharedRef()
 		];
 
 	ChildSlot
diff --git a/Engine/Source/Editor/SequencerCore/Private/MVVM/Views/SOutlinerView.cpp b/Engine/Source/Editor/SequencerCore/Private/MVVM/Views/SOutlinerView.cpp
--- a/Engine/Source/Editor/SequencerCore/Private/MVVM/Views/SOutlinerView.cpp
+++ b/Engine/Source/Editor/SequencerCore/Private/MVVM/Views/SOutlinerView.cpp
@@ -273,7 +273,7 @@ void SOutlinerView::Construct(const FArguments& InArgs, TWeakPtr<FOutlinerViewMo
 
 	HeaderRow = SNew(SHeaderRow).Visibility(EVisibility::Collapsed);
 
-	GenerateOutlinerColumns();
+	UpdateOutlinerColumns();
 
 	WeakOutliner.Pin()->OnRefreshed.AddSP(this, &SOutlinerView::Refresh);
 
@@ -470,57 +470,100 @@ float SOutlinerView::VirtualToPhysical(float InVirtual) const
 	return InVirtual;
 }
 
-void SOutlinerView::GenerateOutlinerColumns()
+void SOutlinerView::AddOutlinerColumnCallback(TSharedPtr<ISequencerOutlinerColumn> InColumn, TSharedPtr<FEditorViewModel> InEditorViewModel)
 {
-	// Clear columns to ensure consistent order when building UI from Map
-	HeaderRow->ClearColumns();
-	Columns.Empty();
-	Columns.Compact();
+	auto GenerateToolColumn = [=](const TWeakViewModelPtr<IOutlinerExtension>& InWeakModel, const TSharedRef<SOutlinerViewRow>& InRow) -> TSharedRef<SWidget>
+	{
+		if (TViewModelPtr<IOutlinerExtension> OutlinerExtension = InWeakModel.ImplicitPin())
+		{
+			if (!OutlinerExtension.AsModel()->IsA<FOutlinerSpacer>())
+			{
+				const FCreateOutlinerColumnParams ColumnParams = FCreateOutlinerColumnParams(OutlinerExtension, InEditorViewModel);
+				const bool bIsCompatible = InColumn->IsItemCompatibleWithColumn(ColumnParams);
 
-	TSharedPtr<FEditorViewModel> EditorViewModel = WeakOutliner.Pin()->GetEditor();
+				TSharedRef<SWidget> InnerContent = bIsCompatible ? InColumn->CreateColumnWidget(InColumn, ColumnParams) : SNullWidget::NullWidget;
 
-	// Create Visible Outliner Columns (Pin/Mute/Solo...)
-	for (TSharedPtr<ISequencerOutlinerColumn> Column : OutlinerColumns)
+				return SNew(SOutlinerColumnBorder, ColumnParams, bIsCompatible)
+					[
+						InnerContent
+					];
+			}
+			else
+			{
+				return SNew(SBox).HeightOverride(10.f);
+			}
+		}
+
+		ensureMsgf(false, TEXT("Attempting to create an outliner column for a view model that is either dead, or not an outliner item."));
+		return SNew(SBox).HeightOverride(10.f);
+	};
+
+	Columns.Add(InColumn->GetColumnName(), FOutlinerViewColumn(GenerateToolColumn, 16.f, true));
+}
+
+void SOutlinerView::UpdateTrackGutterColumns(TSharedPtr<FEditorViewModel> InEditorViewModel)
+{
+	// Padding columns callback, used to provide apropriately colored padding to the left and right of the track gutter
+	auto GeneratePadding = [=](const TWeakViewModelPtr<IOutlinerExtension>& InWeakModel, const TSharedRef<SOutlinerViewRow>& InRow) -> TSharedRef<SWidget>
 	{
-		auto GenerateToolColumn = [=](const TWeakViewModelPtr<IOutlinerExtension>& InWeakModel, const TSharedRef<SOutlinerViewRow>& InRow) -> TSharedRef<SWidget>
+		if (TViewModelPtr<IOutlinerExtension> OutlinerExtension = InWeakModel.ImplicitPin())
 		{
-			if (TViewModelPtr<IOutlinerExtension> OutlinerExtension = InWeakModel.ImplicitPin())
+			if (!OutlinerExtension.AsModel()->IsA<FOutlinerSpacer>())
 			{
-				if (!OutlinerExtension.AsModel()->IsA<FOutlinerSpacer>())
-				{
-					return SNew(SOutlinerColumnBorder, FCreateOutlinerColumnParams(OutlinerExtension, EditorViewModel))
-						[	
-							Column->CreateColumnWidget(FCreateOutlinerColumnParams(OutlinerExtension, EditorViewModel))
-						];
-				}
-				else
-				{
-					return SNew(SBox).HeightOverride(10.f);
-				}
+				const FCreateOutlinerColumnParams ColumnParams = FCreateOutlinerColumnParams(OutlinerExtension, InEditorViewModel);
+				return SNew(SOutlinerColumnBorder, ColumnParams, false);
+			}
+			else
+			{
+				return SNew(SBox).HeightOverride(10.f);
 			}
+		}
+		ensureMsgf(false, TEXT("Attempting to create an outliner column for a view model that is either dead, or not an outliner item."));
+		return SNew(SBox).HeightOverride(10.f);
+	};
+
+	// Spacer column callback, generates a line to separate the track gutter from the outliner
+	auto GenerateSpacer = [=](const TWeakViewModelPtr<IOutlinerExtension>& InWeakModel, const TSharedRef<SOutlinerViewRow>& InRow) -> TSharedRef<SWidget>
+	{
+		return SNew(SHorizontalBox)
+			+ SHorizontalBox::Slot()
+			.AutoWidth()
+			[
+				SNew(SSeparator)
+				.Orientation(Orient_Vertical)
+				.Thickness(1.0f)
+				.SeparatorImage(FAppStyle::Get().GetBrush("Separator"))
+			];
+	};
 
-			ensureMsgf(false, TEXT("Attempting to create an outliner column for a view model that is either dead, or not an outliner item."));
-			return SNew(SBox).HeightOverride(10.f);
-		};
-		Columns.Add(Column->GetColumnName(), FOutlinerViewColumn(GenerateToolColumn, 16.f, true));
+	// Left padding
+	Columns.Add(FName(TEXT("OutlinerColumnLeftPadding")), FOutlinerViewColumn(GeneratePadding, 2.f, true));
+
+	// Create Visible Outliner Columns in the Track Gutter (Pin/Mute/Solo...)
+	for (TSharedPtr<ISequencerOutlinerColumn> Column : OutlinerColumns)
+	{
+		AddOutlinerColumnCallback(Column, InEditorViewModel);
 	}
 
-	// If there are any outliner columns, add a spacer between them and the outliner view
+	// Right padding
+	Columns.Add(FName(TEXT("OutlinerColumnRightPadding")), FOutlinerViewColumn(GeneratePadding, 2.f, true));
+
+	// Spacer between Track Gutter and Outliner View
+	Columns.Add(FName(TEXT("OutlinerColumnSpacer")), FOutlinerViewColumn(GenerateSpacer, 1.f, true));
+}
+
+void SOutlinerView::UpdateOutlinerColumns()
+{
+	// Clear columns to ensure consistent order when building UI from Map
+	HeaderRow->ClearColumns();
+	Columns.Empty();
+	Columns.Compact();
+
+	TSharedPtr<FEditorViewModel> EditorViewModel = WeakOutliner.Pin()->GetEditor();
+
 	if (OutlinerColumns.Num() > 0)
 	{
-		auto GenerateSpacer = [=](const TWeakViewModelPtr<IOutlinerExtension>& InWeakModel, const TSharedRef<SOutlinerViewRow>& InRow) -> TSharedRef<SWidget>
-		{
-			return SNew(SHorizontalBox)
-				+ SHorizontalBox::Slot()
-				.AutoWidth()
-				[
-					SNew(SSeparator)
-					.Orientation(Orient_Vertical)
-					.Thickness(1.0f)
-					.SeparatorImage(FAppStyle::Get().GetBrush("Separator"))
-				];
-		};
-		Columns.Add(FName(TEXT("OutlinerColumnSpacer")), FOutlinerViewColumn(GenerateSpacer, 1.f, true));
+		UpdateTrackGutterColumns(EditorViewModel);
 	}
 
 	// Create Column for Outliner View (Track name / Keying)
@@ -587,7 +630,7 @@ void SOutlinerView::SetOutlinerColumns(const TArray<TSharedPtr<ISequencerOutline
 {
 	// Reset the way rows are constructed with an updated list of Outliner Columns
 	OutlinerColumns = InOutlinerColumns;
-	GenerateOutlinerColumns();
+	UpdateOutlinerColumns();
 }
 
 void SOutlinerView::OnRightMouseButtonDown(const FPointerEvent& MouseEvent)
diff --git a/Engine/Source/Editor/SequencerCore/Public/ISequencerOutlinerColumn.h b/Engine/Source/Editor/SequencerCore/Public/ISequencerOutlinerColumn.h
--- a/Engine/Source/Editor/SequencerCore/Public/ISequencerOutlinerColumn.h
+++ b/Engine/Source/Editor/SequencerCore/Public/ISequencerOutlinerColumn.h
@@ -41,23 +41,23 @@ class SEQUENCERCORE_API ISequencerOutlinerColumn
 
 public:
 
-	/* Gets the unique name of the column to use in the SOutlinerView registry and context menus */
+	/** Returns the name of the column. Used for determining column type when dragging or saving settings. */
 	virtual FName GetColumnName() const = 0;
 
-	/* Gets the text to display in the UI for visibility */
+	/** Returns the label of the column to display in visibility settings. */
 	virtual FText GetColumnLabel() const = 0;
 
-	/* The default visibility state of this column when loaded for the first time */
+	/* The default visibility state of this column when loaded for the first time. */
 	virtual bool IsColumnVisibleByDefault() const { return true; }
 
-	/* Gets whether or not this column is supported by a given Sequencer */
+	/* Gets whether or not this column is supported by a given Sequencer. */
 	virtual bool SupportsSequence(UMovieSceneSequence* InSequence) const { return true; }
 
 	/* Gets whether or not a widget should be generated for a given item in the outliner column. */
 	virtual bool IsItemCompatibleWithColumn(const UE::Sequencer::FCreateOutlinerColumnParams& InParams) const = 0;
 
-	/* Gets the widget created for each item within the SOutlinerView, column widgets must be fixed width */
-	virtual TSharedRef<SWidget> CreateColumnWidget(const UE::Sequencer::FCreateOutlinerColumnParams& InParams) const = 0;
+	/* Gets the widget created for each item within the SOutlinerView, column widgets must be fixed width. */
+	virtual TSharedRef<SWidget> CreateColumnWidget(const TWeakPtr<ISequencerOutlinerColumn> InWeakOutlinerColumn, const UE::Sequencer::FCreateOutlinerColumnParams& InParams) const = 0;
 
 public:
 
diff --git a/Engine/Source/Editor/SequencerCore/Public/MVVM/Views/SOutlinerColumnBorder.h b/Engine/Source/Editor/SequencerCore/Public/MVVM/Views/SOutlinerColumnBorder.h
--- a/Engine/Source/Editor/SequencerCore/Public/MVVM/Views/SOutlinerColumnBorder.h
+++ b/Engine/Source/Editor/SequencerCore/Public/MVVM/Views/SOutlinerColumnBorder.h
@@ -60,7 +60,7 @@ private:
 };
 
 /**
-* A widget for containing a column widget with the correct style / border
+* A border widget with appropriate colors to match outliner items in the Sequencer outliner view.
 */
 class SOutlinerColumnBorder
 	: public SCompoundWidget
@@ -72,9 +72,17 @@ public:
 		SLATE_DEFAULT_SLOT(FArguments, Content)
 	SLATE_END_ARGS()
 
+	/**
+	* Construct this widget
+	*
+	* @param	InArgs			The declaration data for this widget
+	* @param	InParams		Ouliner Column params for reference to which item is hovered or selected
+	* @param	bHasInnerBorder	Should this widget be contained by a 12x12 dark inner border
+	*/
 	void Construct(
 		const FArguments& InArgs, 
-		const FCreateOutlinerColumnParams& InParams
+		const FCreateOutlinerColumnParams& InParams,
+		const bool bHasInnerBorder
 	);
 
 private:
@@ -83,6 +91,9 @@ private:
 	void OnMouseEnter(const FGeometry& MyGeometry, const FPointerEvent& MouseEvent) override;
 	void OnMouseLeave(const FPointerEvent& MouseEvent) override;
 
+	/** Handle mouse down events in border so that outliner item events are not generated by the track gutter. */
+	virtual FReply OnMouseButtonDown(const FGeometry& MyGeometry, const FPointerEvent& MouseEvent) override;
+
 	/**
 	* @return The border image to show in the tree node.
 	*/
diff --git a/Engine/Source/Editor/SequencerCore/Public/MVVM/Views/SOutlinerView.h b/Engine/Source/Editor/SequencerCore/Public/MVVM/Views/SOutlinerView.h
--- a/Engine/Source/Editor/SequencerCore/Public/MVVM/Views/SOutlinerView.h
+++ b/Engine/Source/Editor/SequencerCore/Public/MVVM/Views/SOutlinerView.h
@@ -235,10 +235,14 @@ protected:
 	/** Linear, sorted array of nodes that we currently have generated widgets for */
 	TArray<FCachedGeometry> PhysicalNodes;
 
-protected:
+	/** Helper function for creating callbacks to generate outliner columns */
+	void AddOutlinerColumnCallback(TSharedPtr<ISequencerOutlinerColumn> InColumn, TSharedPtr<FEditorViewModel> InEditorViewModel);
+
+	/** Helper function for initializing track gutter columns and padding */
+	void UpdateTrackGutterColumns(TSharedPtr<FEditorViewModel> InEditorViewModel);
 
 	/** Populate the map of column definitions, and add relevant columns to the header row. Must be called when outliner columns change */
-	void GenerateOutlinerColumns();
+	void UpdateOutlinerColumns();
 
 	FReply OnDragRow(const FGeometry& InGeometry, const FPointerEvent& InPointerEvent, TSharedRef<SOutlinerViewRow> InRow);
 
