Sequencer: Use the focused movie scene sequence for the dynamic binding customization

Otherwise, this crashes when you try to make a dynamic possession of a possessable in a subsequence.

#rnx

diff --git a/Engine/Source/Editor/Sequencer/Private/MVVM/ViewModels/ObjectBindingModel.cpp b/Engine/Source/Editor/Sequencer/Private/MVVM/ViewModels/ObjectBindingModel.cpp
--- a/Engine/Source/Editor/Sequencer/Private/MVVM/ViewModels/ObjectBindingModel.cpp
+++ b/Engine/Source/Editor/Sequencer/Private/MVVM/ViewModels/ObjectBindingModel.cpp
@@ -931,7 +931,7 @@ void FObjectBindingModel::AddDynamicBindingMenu(FMenuBuilder& MenuBuilder, FMovi
 
 	// Register details customizations for this instance
 	TSharedPtr<FSequencer> Sequencer = OwnerModel->GetSequencerImpl();
-	UMovieSceneSequence* Sequence = Sequencer->GetRootMovieSceneSequence();
+	UMovieSceneSequence* Sequence = Sequencer->GetFocusedMovieSceneSequence();
 	StructureDetailsView->GetDetailsView()->RegisterInstancedCustomPropertyTypeLayout(
 		FMovieSceneDynamicBinding::StaticStruct()->GetFName(),
 		FOnGetPropertyTypeCustomizationInstance::CreateStatic(&FMovieSceneDynamicBindingCustomization::MakeInstance, Sequence->GetMovieScene(), ObjectBindingID));
