Sequencer: Outliner Column Implementations

Implementations for the pin, mute, lock, and solo outliner columns to be displayed in the SOutlinerView of Sequencer. Each outliner column widget implements click and drag functionality. Outliner column visibility is saved per-sequencer and per-user. Default column visibility and order can be saved and restored in BaseEditorPerProjectUserSettings.ini.

diff --git a/Engine/Source/Editor/Sequencer/Private/Widgets/OutlinerColumns/SColumnToggleWidget.cpp b/Engine/Source/Editor/Sequencer/Private/Widgets/OutlinerColumns/SColumnToggleWidget.cpp
--- a/Engine/Source/Editor/Sequencer/Private/Widgets/OutlinerColumns/SColumnToggleWidget.cpp
+++ b/Engine/Source/Editor/Sequencer/Private/Widgets/OutlinerColumns/SColumnToggleWidget.cpp
@@ -101,25 +101,34 @@ FSlateColor SColumnToggleWidget::GetImageColorAndOpacity() const
 
 	float Opacity = 0.0f;
 
-	if (IsActive()
-		|| bIsMouseOverWidget)
+	if (IsActive())
 	{
+		// Directly active, full opacity
 		Opacity = 1.0f;
 	}
-	else if (Editor->GetOutliner()->GetHoveredItem() == OutlinerItem)
+	else if (IsChildActive() && !bIsMouseOverWidget)
 	{
-		Opacity = .2f;
+		// Child is active and mouse is not over widget. Full opacity '-'.
+		Opacity = 1.0f;
 	}
-	else if (IsChildActive())
+	else if (bIsMouseOverWidget)
 	{
-		Opacity = 1.0f;
+		// Mouse is over widget and it is not directly active.
+		Opacity = .65f;
 	}
 	else if (IsImplicitlyActive())
 	{
-		Opacity = .4f;
+		// Implicitly active through another object and mouse is not over.
+		Opacity = .35f;
+	}
+	else if (Editor->GetOutliner()->GetHoveredItem() == OutlinerItem)
+	{
+		// Mouse is hovered over outliner item and not the widget itself, preview icons for widget.
+		Opacity = .2f;
 	}
 	else
 	{
+		// Not active in any way and mouse is not over widget or item.
 		Opacity = 0.0f;
 	}
 
@@ -137,17 +146,15 @@ const FSlateBrush* SColumnToggleWidget::GetBrush() const
 		return ActiveBrush;
 	}
 
-	const bool bIsMouseOverItem = bIsMouseOverWidget || Editor->GetOutliner()->GetHoveredItem() == OutlinerItem;
-
-	if (bIsMouseOverItem
-		|| IsActive()
-		|| !IsChildActive())
+	if (IsChildActive()
+		&& !IsActive()
+		&& !bIsMouseOverWidget)
 	{
-		return ActiveBrush;
+		return ChildActiveBrush;
 	}
 	else
 	{
-		return ChildActiveBrush;
+		return ActiveBrush;
 	}
 }
 
diff --git a/Engine/Source/Editor/Sequencer/Public/SequencerSettings.h b/Engine/Source/Editor/Sequencer/Public/SequencerSettings.h
--- a/Engine/Source/Editor/Sequencer/Public/SequencerSettings.h
+++ b/Engine/Source/Editor/Sequencer/Public/SequencerSettings.h
@@ -86,7 +86,7 @@ struct FColumnVisibilitySetting
 {
 	GENERATED_BODY()
 
-	UPROPERTY(EditAnywhere, Category=General)
+	UPROPERTY(EditAnywhere, Category=General, meta = (EditCondition="false"))
 	FName ColumnName;
 
 	UPROPERTY(EditAnywhere, Category=General)
