Sequencer: Outliner Column Implementations

Implementations for the pin, mute, lock, and solo outliner columns to be displayed in the SOutlinerView of Sequencer. Each outliner column widget implements click and drag functionality. Outliner column visibility is saved per-sequencer and per-user. Default column visibility and order can be saved and restored in BaseEditorPerProjectUserSettings.ini.

diff --git a/Engine/Source/Editor/Sequencer/Public/SequencerSettings.h b/Engine/Source/Editor/Sequencer/Public/SequencerSettings.h
--- a/Engine/Source/Editor/Sequencer/Public/SequencerSettings.h
+++ b/Engine/Source/Editor/Sequencer/Public/SequencerSettings.h
@@ -86,7 +86,7 @@ struct FColumnVisibilitySetting
 {
 	GENERATED_BODY()
 
-	UPROPERTY(EditAnywhere, Category=General, meta = (EditCondition="false"))
+	UPROPERTY(VisibleAnywhere, Category=General)
 	FName ColumnName;
 
 	UPROPERTY(EditAnywhere, Category=General)
