Fix for soloing child-tracks in sequencer resulting in no tracks being solo'd.

diff --git a/Engine/Source/Editor/Sequencer/Private/SequencerNodeTree.cpp b/Engine/Source/Editor/Sequencer/Private/SequencerNodeTree.cpp
--- a/Engine/Source/Editor/Sequencer/Private/SequencerNodeTree.cpp
+++ b/Engine/Source/Editor/Sequencer/Private/SequencerNodeTree.cpp
@@ -327,11 +327,30 @@ bool FSequencerNodeTree::IsNodeSolo(const UE::Sequencer::TViewModelPtr<UE::Seque
 		return true;
 	}
 
-	// Children should follow their parent's behavior unless told otherwise.
+	// If parent is solo'd, then this node should also be solo'd
 	TViewModelPtr<IOutlinerExtension> ParentNode = InNode.AsModel()->FindAncestorOfType<IOutlinerExtension>();
-	if (ParentNode)
+	while (ParentNode)
+	{
+		const FString ParentNodePath = IOutlinerExtension::GetPathName(ParentNode);
+		if (SoloNodes.Contains(ParentNodePath))
+		{
+			return true;
+		}
+
+		ParentNode = ParentNode.AsModel()->FindAncestorOfType<IOutlinerExtension>();
+	}
+
+	// If any child is solo'd then this node needs to also be solo'd so that the child can be evaluated later during playback
+	for (TViewModelPtr<IOutlinerExtension> ChildNode : InNode.AsModel()->GetChildrenOfType<IOutlinerExtension>())
 	{
-		return IsNodeSolo(ParentNode);
+		if (ChildNode)
+		{
+			const FString ChildNodePath = IOutlinerExtension::GetPathName(ChildNode);
+			if (SoloNodes.Contains(ChildNodePath))
+			{
+				return true;
+			}
+		}
 	}
 
 	return false;
