Sequencer: Fixed latent Actor selection events overwriting track area selection when instigated from a track area selection

diff --git a/Engine/Source/Editor/Sequencer/Private/Sequencer.cpp b/Engine/Source/Editor/Sequencer/Private/Sequencer.cpp
--- a/Engine/Source/Editor/Sequencer/Private/Sequencer.cpp
+++ b/Engine/Source/Editor/Sequencer/Private/Sequencer.cpp
@@ -6411,6 +6411,13 @@ void FSequencer::SynchronizeExternalSelectionWithSequencerSelection()
 
 	GEditor->NoteSelectionChange();
 
+	// Ensure that the (newer) typed element selection broadcasts its changes immediately so we don't get an
+	//   end of frame update which might overwrite a newly created track area selection
+	if (UTypedElementSelectionSet* TypedElements = GEditor->GetSelectedActors()->GetElementSelectionSet())
+	{
+		TypedElements->NotifyPendingChanges();
+	}
+
 	if (SelectedSequencerComponents.Num())
 	{
 		GEditor->GetSelectedComponents()->Modify();
@@ -6424,6 +6431,13 @@ void FSequencer::SynchronizeExternalSelectionWithSequencerSelection()
 		GEditor->GetSelectedComponents()->EndBatchSelectOperation();
 
 		GEditor->NoteSelectionChange();
+
+		// Ensure that the (newer) typed element selection broadcasts its changes immediately so we don't get an
+		//   end of frame update which might overwrite a newly created track area selection
+		if (UTypedElementSelectionSet* TypedElements = GEditor->GetSelectedComponents()->GetElementSelectionSet())
+		{
+			TypedElements->NotifyPendingChanges();
+		}
 	}
 }
 
diff --git a/Engine/Source/Editor/Sequencer/Sequencer.Build.cs b/Engine/Source/Editor/Sequencer/Sequencer.Build.cs
--- a/Engine/Source/Editor/Sequencer/Sequencer.Build.cs
+++ b/Engine/Source/Editor/Sequencer/Sequencer.Build.cs
@@ -46,6 +46,7 @@ public class Sequencer : ModuleRules
 				"SerializedRecorderInterface",
 				"ToolMenus",
 				"ToolWidgets",
+				"TypedElementFramework",
 				}
 			);
 
