Sequencer: Fix selected nodes only so that it relies on the bound objects rather than the binding name which isn't always the same as the actor label

diff --git a/Engine/Source/Editor/Sequencer/Private/LevelEditorSequencerIntegration.cpp b/Engine/Source/Editor/Sequencer/Private/LevelEditorSequencerIntegration.cpp
--- a/Engine/Source/Editor/Sequencer/Private/LevelEditorSequencerIntegration.cpp
+++ b/Engine/Source/Editor/Sequencer/Private/LevelEditorSequencerIntegration.cpp
@@ -545,7 +545,7 @@ bool FLevelEditorSequencerIntegration::IsBindingVisible(const FMovieSceneBinding
 		return true;
 	}
 
-	// Disregard if not a level sequence (ie. a control rig sequence)
+	// Disregard if not a level sequence (ie. a control rig sequence)	
 	for (FSequencerAndOptions& SequencerAndOptions : BoundSequencers)
 	{
 		TSharedPtr<FSequencer> Pinned = SequencerAndOptions.Sequencer.Pin();
@@ -557,20 +557,24 @@ bool FLevelEditorSequencerIntegration::IsBindingVisible(const FMovieSceneBinding
 				{
 					return true;
 				}
+				else
+				{
+					TArrayView<TWeakObjectPtr<>> Objects = Pinned->FindObjectsInCurrentSequence(InBinding.GetObjectGuid());
+					for (TWeakObjectPtr<> Object : Objects)
+					{
+						if (AActor* Actor = Cast<AActor>(Object.Get()))
+						{
+							if (GEditor->GetSelectedActors()->IsSelected(Actor))
+							{
+								return true;
+							}
+						}
+					}
+				}
 			}
 		}
 	}
 
-	for( FSelectionIterator SelectionIt( *GEditor->GetSelectedActors() ); SelectionIt; ++SelectionIt )
-	{
-		AActor* SelectedActor = CastChecked<AActor>( *SelectionIt );
-		
-		if (SelectedActor->GetActorLabel() == InBinding.GetName())
-		{
-			return true;
-		}
-	}
-
 	return false;
 }
 
