Sequencer: Added some comments to motion trail options

diff --git a/Engine/Source/Editor/Sequencer/Public/Tools/MotionTrailOptions.h b/Engine/Source/Editor/Sequencer/Public/Tools/MotionTrailOptions.h
--- a/Engine/Source/Editor/Sequencer/Public/Tools/MotionTrailOptions.h
+++ b/Engine/Source/Editor/Sequencer/Public/Tools/MotionTrailOptions.h
@@ -35,51 +35,67 @@ public:
 		, SecondsPerMark(0.1)
 	{}
 
+	/** Whether or not to show motion trails */
 	UPROPERTY(EditAnywhere, Category = Trail);
 	bool bShowTrails;
 
+	/** The color of the motion trail */
 	UPROPERTY(EditAnywhere,  Category = Trail)
 	FLinearColor TrailColor;
 
+	/** Whether or not to show the full motion trail */
 	UPROPERTY(EditAnywhere, Category = Trail)
 	bool bShowFullTrail;
 
+	/* The thickness of the motion trail */
 	UPROPERTY(EditAnywhere, Category = Trail, Meta = (ClampMin = "0.0"))
 	float TrailThickness;
 
+	/** The number of frames to draw before the start of the trail. Requires not showing the full trail */
 	UPROPERTY(EditAnywhere, Category = Trail, Meta = (EditCondition = "!bShowFullTrail", ClampMin = "0"))
 	int32 FramesBefore;
 
+	/** The number of frames to draw after the end of the trail. Requires not showing the full trail */
 	UPROPERTY(EditAnywhere, Category = Trail, Meta = (EditCondition = "!bShowFullTrail", ClampMin = "0"))
 	int32 FramesAfter;
 
+	/** The number of evaluations per frame */
 	UPROPERTY(EditAnywhere, Category = Trail, Meta = (ClampMin = "1.0"))
 	int32 EvalsPerFrame;
 	
+	/** Whether or not to show keys on the motion trail */
 	UPROPERTY(EditAnywhere, Category = Keys)
 	bool bShowKeys;
 
+	/** Whether or not to show frame numbers for the keys on the motion trail */
 	UPROPERTY(EditAnywhere, Category = Keys)
 	bool bShowFrameNumber;
 
+	/** The color of the keys */
 	UPROPERTY(EditAnywhere, Category = Keys)
 	FLinearColor KeyColor;
 
+	/** The size of the keys */
 	UPROPERTY(EditAnywhere, Category = Keys, Meta = (ClampMin = "0.0"))
 	double KeySize;
 
+	/** Whether or not to show marks along the motion trail */
 	UPROPERTY(EditAnywhere, Category = Marks)
 	bool bShowMarks;
 
+	/** The color of the marks */
 	UPROPERTY(EditAnywhere, Category = Marks)
 	FLinearColor MarkColor;
 
+	/** The size of the marks */
 	UPROPERTY(EditAnywhere, Category = Marks, Meta = (ClampMin = "0.0"))
 	double MarkSize;
 
+	/** Whether or not to lock the marks to the frames */
 	UPROPERTY(EditAnywhere, Category = Marks)
 	bool bLockMarksToFrames;
 
+	/** The seconds per mark */
 	UPROPERTY(EditAnywhere, Category = Marks, Meta = (EditCondition = "!bLockMarksToFrames", ClampMin = "0.01"))
 	double SecondsPerMark;
 
