Sequencer: Game mode playback would not turn back on gizmo's, do to editor regression, which now needs you to manully turn on widgets with game mode turned off.
#trivial

diff --git a/Engine/Source/Editor/Sequencer/Private/Sequencer.cpp b/Engine/Source/Editor/Sequencer/Private/Sequencer.cpp
--- a/Engine/Source/Editor/Sequencer/Private/Sequencer.cpp
+++ b/Engine/Source/Editor/Sequencer/Private/Sequencer.cpp
@@ -4652,6 +4652,11 @@ FSequencer::FCachedViewState::RestoreViewState()
 			if (LevelVC && LevelVC->AllowsCinematicControl())
 			{
 				LevelVC->SetGameView(GameViewStates[Index].Value);
+				//if turn off game view now need to make sure widget/gizmo is on
+				if (GameViewStates[Index].Value == false)
+				{
+					LevelVC->ShowWidget(true);
+				}
 			}
 		}
 	}
@@ -5899,6 +5904,11 @@ void FSequencer::SaveCurrentMovieScene()
 		//redraw viewport to have the yellow highlight again
 		GCurrentLevelEditingViewportClient = OldViewportClient;
 		GCurrentLevelEditingViewportClient->SetGameView(bIsInGameView);
+		//if turn off game view now need to make sure widget/gizmo is on
+		if (bIsInGameView == false)
+		{
+			GCurrentLevelEditingViewportClient->ShowWidget(true);
+		}
 		Viewport->Draw();
 	}
 
