Sequencer: Clear selected marked frames after marked frames are deleted.

This fixes a crash moving keys after deleting marked frames because the selected marked frames is stale.

Note, when keys are deleted, we similarly clear the key selection.

diff --git a/Engine/Source/Editor/Sequencer/Private/Sequencer.cpp b/Engine/Source/Editor/Sequencer/Private/Sequencer.cpp
--- a/Engine/Source/Editor/Sequencer/Private/Sequencer.cpp
+++ b/Engine/Source/Editor/Sequencer/Private/Sequencer.cpp
@@ -9779,6 +9779,8 @@ void FSequencer::ToggleMarkAtPlayPosition()
 		FocusedMovieScene->Modify();
 		FocusedMovieScene->AddMarkedFrame(FMovieSceneMarkedFrame(TickFrameNumber));
 	}
+
+	ViewModel->GetSelection()->MarkedFrames.Empty();
 }
 
 void FSequencer::SetMarkedFrame(int32 InMarkIndex, FFrameNumber InFrameNumber)
@@ -9857,6 +9859,8 @@ void FSequencer::DeleteMarkedFrame(int32 InMarkIndex)
 
 		FocusedMovieScene->Modify();
 		FocusedMovieScene->DeleteMarkedFrame(InMarkIndex);
+
+		ViewModel->GetSelection()->MarkedFrames.Empty();
 	}
 }
 
@@ -9884,6 +9888,8 @@ void FSequencer::DeleteAllMarkedFrames()
 
 	FocusedMovieScene->Modify();
 	FocusedMovieScene->DeleteMarkedFrames();
+
+	ViewModel->GetSelection()->MarkedFrames.Empty();
 }
 
 void FSequencer::StepToNextMark()
