Sequencer; Outer the entity system linker to the transient package

diff --git a/Engine/Source/Editor/Sequencer/Private/Sequencer.cpp b/Engine/Source/Editor/Sequencer/Private/Sequencer.cpp
--- a/Engine/Source/Editor/Sequencer/Private/Sequencer.cpp
+++ b/Engine/Source/Editor/Sequencer/Private/Sequencer.cpp
@@ -4197,12 +4197,7 @@ void FSequencer::GetViewportSettings(TMap<FViewportClient*, EMovieSceneViewportP
 
 UMovieSceneEntitySystemLinker* FSequencer::ConstructEntitySystemLinker()
 {
-	UObject* Context = GetPlaybackContext();
-	if (!Context)
-	{
-		Context = GetTransientPackage();
-	}
-
+	UObject* Context = GetTransientPackage();
 	FName EntitySystemLinkerName = MakeUniqueObjectName(Context, UMovieSceneEntitySystemLinker::StaticClass(), TEXT("SequencerEntitySystemLinker"));
 	UMovieSceneEntitySystemLinker* EntitySystemLinker = NewObject<UMovieSceneEntitySystemLinker>(Context, EntitySystemLinkerName);
 	EntitySystemLinker->SetLinkerRole(UE::MovieScene::EEntitySystemLinkerRole::Standalone);
