Sequencer: Components of possessable bindings disappear when containing folder is moved.

Folder within folder moving removes viewmodel child data for possessables until rebuilt. This fix recycles valid child object bindings in OnConstruct for Possessable Models.

diff --git a/Engine/Source/Editor/Sequencer/Private/MVVM/ViewModels/PossessableModel.cpp b/Engine/Source/Editor/Sequencer/Private/MVVM/ViewModels/PossessableModel.cpp
--- a/Engine/Source/Editor/Sequencer/Private/MVVM/ViewModels/PossessableModel.cpp
+++ b/Engine/Source/Editor/Sequencer/Private/MVVM/ViewModels/PossessableModel.cpp
@@ -6,6 +6,7 @@
 #include "HAL/Platform.h"
 #include "ISequencer.h"
 #include "Internationalization/Internationalization.h"
+#include "MVVM/Extensions/IRecyclableExtension.h"
 #include "MVVM/ViewModels/SequenceModel.h"
 #include "MVVM/ViewModels/SequencerEditorViewModel.h"
 #include "Misc/AssertionMacros.h"
@@ -63,7 +64,21 @@ void FPossessableModel::OnConstruct()
 	const FMovieScenePossessable* Possessable = MovieScene->FindPossessable(ObjectBindingID);
 	check(Possessable);
 
+	FScopedViewModelListHead RecycledHead(AsShared(), EViewModelListType::Recycled);
+	GetChildrenForList(&OutlinerChildList).MoveChildrenTo<IRecyclableExtension>(RecycledHead.GetChildren(), IRecyclableExtension::CallOnRecycle);
+
 	FObjectBindingModel::OnConstruct();
+
+	// Preserve objectbindings since those are always added by the object model storage
+	// on reinitialize or in response to an object event.
+	FViewModelChildren OutlinerChildren = GetChildrenForList(&OutlinerChildList);
+	for (TViewModelPtr<FObjectBindingModel> Child : RecycledHead.GetChildren().IterateSubList<FObjectBindingModel>().ToArray())
+	{
+		if (Child)
+		{
+			OutlinerChildren.AddChild(Child);
+		}
+	}
 }
 
 FText FPossessableModel::GetIconToolTipText() const
