Sequencer: Check for valid mouse down position

diff --git a/Engine/Source/Editor/Sequencer/Private/SequencerTimeSliderController.cpp b/Engine/Source/Editor/Sequencer/Private/SequencerTimeSliderController.cpp
--- a/Engine/Source/Editor/Sequencer/Private/SequencerTimeSliderController.cpp
+++ b/Engine/Source/Editor/Sequencer/Private/SequencerTimeSliderController.cpp
@@ -593,14 +593,14 @@ int32 FSequencerTimeSliderController::OnPaintTimeSlider( bool bMirrorLabels, con
 			);
 		}
 		
-		if (MouseDragType == DRAG_SETTING_RANGE)
+		if (MouseDragType == DRAG_SETTING_RANGE && MouseDownPosition[0].IsSet() && MouseDownPosition[1].IsSet())
 		{
 			FFrameRate Resolution = GetTickResolution();
 			FFrameTime MouseDownTime[2];
 
 			FScrubRangeToScreen MouseDownRange(GetViewRange(), MouseDownGeometry.Size);
-			MouseDownTime[0] = ComputeFrameTimeFromMouse(MouseDownGeometry, MouseDownPosition[0], MouseDownRange);
-			MouseDownTime[1] = ComputeFrameTimeFromMouse(MouseDownGeometry, MouseDownPosition[1], MouseDownRange);
+			MouseDownTime[0] = ComputeFrameTimeFromMouse(MouseDownGeometry, MouseDownPosition[0].GetValue(), MouseDownRange);
+			MouseDownTime[1] = ComputeFrameTimeFromMouse(MouseDownGeometry, MouseDownPosition[1].GetValue(), MouseDownRange);
 
 			float      MouseStartPosX = RangeToScreen.InputToLocalX(MouseDownTime[0] / Resolution);
 			float      MouseEndPosX   = RangeToScreen.InputToLocalX(MouseDownTime[1] / Resolution);
@@ -880,6 +880,8 @@ FReply FSequencerTimeSliderController::OnMouseButtonUp( SWidget& WidgetOwner, co
 		
 		bPanning = false;
 		bMouseDownInRegion = false;
+		MouseDownPosition[0].Reset();
+		MouseDownPosition[1].Reset();
 		
 		return FReply::Handled().ReleaseMouseCapture();
 	}
@@ -906,10 +908,10 @@ FReply FSequencerTimeSliderController::OnMouseButtonUp( SWidget& WidgetOwner, co
 			TimeSliderArgs.OnMarkEndDrag.ExecuteIfBound();
 			UpdateMarkSelection(DragMarkIndex, DragMarkCurrentFrameNumber);
 		}
-		else if (MouseDragType == DRAG_SETTING_RANGE)
+		else if (MouseDragType == DRAG_SETTING_RANGE && MouseDownPosition[0].IsSet())
 		{
 			// Zooming
-			FFrameTime MouseDownStart = ComputeFrameTimeFromMouse(MyGeometry, MouseDownPosition[0], RangeToScreen);
+			FFrameTime MouseDownStart = ComputeFrameTimeFromMouse(MyGeometry, MouseDownPosition[0].GetValue(), RangeToScreen);
 
 			const bool bCanZoomIn  = MouseTime > MouseDownStart;
 			const bool bCanZoomOut = ViewRangeStack.Num() > 0;
@@ -987,11 +989,15 @@ FReply FSequencerTimeSliderController::OnMouseButtonUp( SWidget& WidgetOwner, co
 		DistanceDragged = 0.f;
 		DragMarkIndex = INDEX_NONE;
 		bMouseDownInRegion = false;
+		MouseDownPosition[0].Reset();
+		MouseDownPosition[1].Reset();
 
 		return FReply::Handled().ReleaseMouseCapture();
 	}
 
 	bMouseDownInRegion = false;
+	MouseDownPosition[0].Reset();
+	MouseDownPosition[1].Reset();
 	return FReply::Unhandled();
 }
 
@@ -1081,7 +1087,7 @@ FReply FSequencerTimeSliderController::OnMouseMoveImpl( SWidget& WidgetOwner, co
 			SetViewRange(NewViewOutputMin, NewViewOutputMax, EViewRangeInterpolation::Immediate);
 		}
 	}
-	else if (bHandleLeftMouseButton || bHandleMiddleMouseButton)
+	else if ((bHandleLeftMouseButton || bHandleMiddleMouseButton) && MouseDownPosition[0].IsSet())
 	{
 		TRange<double> LocalViewRange = GetViewRange();
 		FScrubRangeToScreen RangeToScreen(LocalViewRange, MyGeometry.Size);
@@ -1091,7 +1097,7 @@ FReply FSequencerTimeSliderController::OnMouseMoveImpl( SWidget& WidgetOwner, co
 		{
 			if ( DistanceDragged > 0.f /*FSlateApplication::Get().GetDragTriggerDistance()*/ )
 			{
-				FFrameTime MouseDownFree = ComputeFrameTimeFromMouse(MyGeometry, MouseDownPosition[0], RangeToScreen, false);
+				FFrameTime MouseDownFree = ComputeFrameTimeFromMouse(MyGeometry, MouseDownPosition[0].GetValue(), RangeToScreen, false);
 
 				const bool       bReadOnly          = Sequencer->IsReadOnly();
 				const FFrameRate TickResolution     = GetTickResolution();
@@ -1149,7 +1155,7 @@ FReply FSequencerTimeSliderController::OnMouseMoveImpl( SWidget& WidgetOwner, co
 		{
 			FFrameTime MouseTime = ComputeFrameTimeFromMouse(MyGeometry, MouseEvent.GetScreenSpacePosition(), RangeToScreen);
 			FFrameTime ScrubTime = ComputeScrubTimeFromMouse(MyGeometry, MouseEvent, RangeToScreen);
-			FFrameTime MouseDownTime = ComputeFrameTimeFromMouse(MyGeometry, MouseDownPosition[0], RangeToScreen);
+			FFrameTime MouseDownTime = ComputeFrameTimeFromMouse(MyGeometry, MouseDownPosition[0].GetValue(), RangeToScreen);
 			FFrameNumber DiffFrame = MouseTime.FrameNumber - MouseDownTime.FrameNumber;
 
 			// Set the start range time?
diff --git a/Engine/Source/Editor/Sequencer/Public/SequencerTimeSliderController.h b/Engine/Source/Editor/Sequencer/Public/SequencerTimeSliderController.h
--- a/Engine/Source/Editor/Sequencer/Public/SequencerTimeSliderController.h
+++ b/Engine/Source/Editor/Sequencer/Public/SequencerTimeSliderController.h
@@ -361,7 +361,7 @@ private:
 	bool bPanning;
 
 	/** Mouse down position range */
-	FVector2D MouseDownPosition[2];
+	TOptional<FVector2D> MouseDownPosition[2];
 
 	/** Geometry on mouse down */
 	FGeometry MouseDownGeometry;
