Sequencer: Fix column visibility menu not updating when checked

#rnx

diff --git a/Engine/Source/Editor/Sequencer/Private/SSequencer.cpp b/Engine/Source/Editor/Sequencer/Private/SSequencer.cpp
--- a/Engine/Source/Editor/Sequencer/Private/SSequencer.cpp
+++ b/Engine/Source/Editor/Sequencer/Private/SSequencer.cpp
@@ -2344,6 +2344,8 @@ void SSequencer::FillPlaybackSpeedMenu(FMenuBuilder& InMenuBarBuilder)
 void SSequencer::FillColumnVisibilityMenu(FMenuBuilder& InMenuBuilder)
 {
 	const bool bShouldCloseWindowAfterMenuSelection = true;
+
+	InMenuBuilder.BeginSection("ColumnVisibility", LOCTEXT("ColumnVisibilityHeader", "Column Visibility"));
 	for (FSequencerOutlinerColumnVisibility& ColumnVisibility : OutlinerColumnVisibilities)
 	{
 		InMenuBuilder.AddMenuEntry(
@@ -2356,7 +2358,17 @@ void SSequencer::FillColumnVisibilityMenu(FMenuBuilder& InMenuBuilder)
 					UpdateOutlinerViewColumns();
 					}),
 				FCanExecuteAction(),
-				FIsActionChecked::CreateLambda([this, ColumnVisibility] { return ColumnVisibility.bIsColumnVisible; })),
+				FIsActionChecked::CreateLambda([this, ColumnVisibility] { 
+						for (const FColumnVisibilitySetting& ColumnSettings : GetSequencerSettings()->GetOutlinerColumnSettings())
+						{
+							if (ColumnSettings.ColumnName == ColumnVisibility.Column.Get()->GetColumnName())
+							{
+								return ColumnSettings.bIsVisible;
+							}
+						}
+						return false;
+					})
+				),
 				NAME_None,
 				EUserInterfaceActionType::ToggleButton
 			);
