Sequencer: Remove section and fix localization error

#rnx

diff --git a/Engine/Source/Editor/Sequencer/Private/SSequencer.cpp b/Engine/Source/Editor/Sequencer/Private/SSequencer.cpp
--- a/Engine/Source/Editor/Sequencer/Private/SSequencer.cpp
+++ b/Engine/Source/Editor/Sequencer/Private/SSequencer.cpp
@@ -2345,7 +2345,6 @@ void SSequencer::FillColumnVisibilityMenu(FMenuBuilder& InMenuBuilder)
 {
 	const bool bShouldCloseWindowAfterMenuSelection = true;
 
-	InMenuBuilder.BeginSection("ColumnVisibility", LOCTEXT("ColumnVisibilityHeader", "Column Visibility"));
 	for (FSequencerOutlinerColumnVisibility& ColumnVisibility : OutlinerColumnVisibilities)
 	{
 		InMenuBuilder.AddMenuEntry(
@@ -2373,7 +2372,6 @@ void SSequencer::FillColumnVisibilityMenu(FMenuBuilder& InMenuBuilder)
 				EUserInterfaceActionType::ToggleButton
 			);
 	}
-	InMenuBuilder.EndSection();
 }
 
 void SSequencer::FillTimeDisplayFormatMenu(FMenuBuilder& MenuBuilder)
