Sequencer: Outer the entity system linker to the transient package

Add SetWorld/GetWorld on the linker so that systems would get the player context world instead of the linker world, which is the transient package.

diff --git a/Engine/Source/Editor/Sequencer/Private/Sequencer.cpp b/Engine/Source/Editor/Sequencer/Private/Sequencer.cpp
--- a/Engine/Source/Editor/Sequencer/Private/Sequencer.cpp
+++ b/Engine/Source/Editor/Sequencer/Private/Sequencer.cpp
@@ -4197,14 +4197,14 @@ void FSequencer::GetViewportSettings(TMap<FViewportClient*, EMovieSceneViewportP
 
 UMovieSceneEntitySystemLinker* FSequencer::ConstructEntitySystemLinker()
 {
-	UObject* Context = GetPlaybackContext();
-	if (!Context)
-	{
-		Context = GetTransientPackage();
-	}
+	UObject* PlaybackContext = GetPlaybackContext();
+	UWorld* World = PlaybackContext ? PlaybackContext->GetWorld() : nullptr;
+
+	UObject* Outer = GetTransientPackage();
+	FName EntitySystemLinkerName = MakeUniqueObjectName(Outer, UMovieSceneEntitySystemLinker::StaticClass(), TEXT("SequencerEntitySystemLinker"));
+	UMovieSceneEntitySystemLinker* EntitySystemLinker = NewObject<UMovieSceneEntitySystemLinker>(Outer, EntitySystemLinkerName);
 
-	FName EntitySystemLinkerName = MakeUniqueObjectName(Context, UMovieSceneEntitySystemLinker::StaticClass(), TEXT("SequencerEntitySystemLinker"));
-	UMovieSceneEntitySystemLinker* EntitySystemLinker = NewObject<UMovieSceneEntitySystemLinker>(Context, EntitySystemLinkerName);
+	EntitySystemLinker->SetWorld(World);
 	EntitySystemLinker->SetLinkerRole(UE::MovieScene::EEntitySystemLinkerRole::Standalone);
 	return EntitySystemLinker;
 }
diff --git a/Engine/Source/Runtime/MovieScene/Public/EntitySystem/MovieSceneEntitySystemLinker.h b/Engine/Source/Runtime/MovieScene/Public/EntitySystem/MovieSceneEntitySystemLinker.h
--- a/Engine/Source/Runtime/MovieScene/Public/EntitySystem/MovieSceneEntitySystemLinker.h
+++ b/Engine/Source/Runtime/MovieScene/Public/EntitySystem/MovieSceneEntitySystemLinker.h
@@ -212,6 +212,16 @@ public:
 	 */
 	MOVIESCENE_API void Reset();
 
+	/**
+	 * Gets the world context for this linker
+	 */
+	virtual UWorld* GetWorld() const override { return WeakWorld.IsValid() ? WeakWorld.Get() : Super::GetWorld(); }
+
+	/**
+	 * Sets the world context for this linker
+	 */
+	void SetWorld(UWorld* InWorld) { WeakWorld = InWorld; }
+
 public:
 
 	/**
@@ -385,6 +395,7 @@ private:
 	uint64 LastInstantiationVersion;
 
 	TWeakPtr<bool> GlobalStateCaptureToken;
+	TWeakObjectPtr<UWorld> WeakWorld;
 
 protected:
 
