Sequencer: Added missing EditAnywhere tags for ExportFBXParams which was preventing it from being visible in blueprints

diff --git a/Engine/Plugins/MovieScene/SequencerScripting/Source/SequencerScriptingEditor/Public/SequencerTools.h b/Engine/Plugins/MovieScene/SequencerScripting/Source/SequencerScriptingEditor/Public/SequencerTools.h
--- a/Engine/Plugins/MovieScene/SequencerScripting/Source/SequencerScriptingEditor/Public/SequencerTools.h
+++ b/Engine/Plugins/MovieScene/SequencerScripting/Source/SequencerScriptingEditor/Public/SequencerTools.h
@@ -85,25 +85,25 @@ struct FSequencerExportFBXParams
 		, OverrideOptions(InOverrideOptions)
 		, FBXFileName(InFBXFileName) {}
 	
-	UPROPERTY(BlueprintReadWrite, Category = "Movie Scene")
+	UPROPERTY(EditAnywhere, BlueprintReadWrite, Category = "Movie Scene")
 	TObjectPtr<UWorld> World;
 	
-	UPROPERTY(BlueprintReadWrite, Category = "Movie Scene")
+	UPROPERTY(EditAnywhere, BlueprintReadWrite, Category = "Movie Scene")
 	TObjectPtr<ULevelSequence> Sequence;
 	
-	UPROPERTY(BlueprintReadWrite, Category = "Movie Scene")
+	UPROPERTY(EditAnywhere, BlueprintReadWrite, Category = "Movie Scene")
 	TObjectPtr<ULevelSequence> RootSequence;
 	
-	UPROPERTY(BlueprintReadWrite, Category = "Movie Scene")
+	UPROPERTY(EditAnywhere, BlueprintReadWrite, Category = "Movie Scene")
 	TArray<FMovieSceneBindingProxy> Bindings;
 	
-	UPROPERTY(BlueprintReadWrite, Category = "Movie Scene")
+	UPROPERTY(EditAnywhere, BlueprintReadWrite, Category = "Movie Scene")
 	TArray<TObjectPtr<UMovieSceneTrack>> Tracks;
 	
-	UPROPERTY(BlueprintReadWrite, Category = "Movie Scene")
+	UPROPERTY(EditAnywhere, BlueprintReadWrite, Category = "Movie Scene")
 	TObjectPtr<UFbxExportOption> OverrideOptions;
 
-	UPROPERTY(BlueprintReadWrite, Category = "Movie Scene")
+	UPROPERTY(EditAnywhere, BlueprintReadWrite, Category = "Movie Scene")
 	FString FBXFileName;
 };
 
diff --git a/Engine/Source/Editor/Sequencer/Public/SequencerUtilities.h b/Engine/Source/Editor/Sequencer/Public/SequencerUtilities.h
--- a/Engine/Source/Editor/Sequencer/Public/SequencerUtilities.h
+++ b/Engine/Source/Editor/Sequencer/Public/SequencerUtilities.h
@@ -36,10 +36,10 @@ struct FMovieScenePasteFoldersParams
 		: Sequence(InSequence)
 		, ParentFolder(InParentFolder) {}
 
-	UPROPERTY(BlueprintReadWrite, Category = "Movie Scene")
+	UPROPERTY(EditAnywhere, BlueprintReadWrite, Category = "Movie Scene")
 	TObjectPtr<UMovieSceneSequence> Sequence;
 	
-	UPROPERTY(BlueprintReadWrite, Category = "Movie Scene")
+	UPROPERTY(EditAnywhere, BlueprintReadWrite, Category = "Movie Scene")
 	TObjectPtr<UMovieSceneFolder> ParentFolder;
 };
 
@@ -55,13 +55,13 @@ struct FMovieScenePasteSectionsParams
 		, TrackRowIndices(InTrackRowIndices)
 		, Time(InTime) {}
 
-	UPROPERTY(BlueprintReadWrite, Category = "Movie Scene")
+	UPROPERTY(EditAnywhere, BlueprintReadWrite, Category = "Movie Scene")
 	TArray<TObjectPtr<UMovieSceneTrack>> Tracks;
 
-	UPROPERTY(BlueprintReadWrite, Category = "Movie Scene")
+	UPROPERTY(EditAnywhere, BlueprintReadWrite, Category = "Movie Scene")
 	TArray<int32> TrackRowIndices;
 
-	UPROPERTY(BlueprintReadWrite, Category = "Movie Scene")
+	UPROPERTY(EditAnywhere, BlueprintReadWrite, Category = "Movie Scene")
 	FFrameTime Time;
 };
 
@@ -78,16 +78,16 @@ struct FMovieScenePasteTracksParams
 		, ParentFolder(InParentFolder)
 		, Folders(InFolders) {}
 
-	UPROPERTY(BlueprintReadWrite, Category = "Movie Scene")
+	UPROPERTY(EditAnywhere, BlueprintReadWrite, Category = "Movie Scene")
 	TObjectPtr<UMovieSceneSequence> Sequence;
 
-	UPROPERTY(BlueprintReadWrite, Category = "Movie Scene")
+	UPROPERTY(EditAnywhere, BlueprintReadWrite, Category = "Movie Scene")
 	TArray<FMovieSceneBindingProxy> Bindings;
 	
-	UPROPERTY(BlueprintReadWrite, Category = "Movie Scene")
+	UPROPERTY(EditAnywhere, BlueprintReadWrite, Category = "Movie Scene")
 	TObjectPtr<UMovieSceneFolder> ParentFolder;
 	
-	UPROPERTY(BlueprintReadWrite, Category = "Movie Scene")
+	UPROPERTY(EditAnywhere, BlueprintReadWrite, Category = "Movie Scene")
 	TArray<TObjectPtr<UMovieSceneFolder>> Folders;
 };
 
@@ -103,16 +103,16 @@ struct FMovieScenePasteBindingsParams
 		, Folders(InFolders)
 		, bDuplicateExistingActors(bInDuplicateExistingActors) {}
 
-	UPROPERTY(BlueprintReadWrite, Category = "Movie Scene")
+	UPROPERTY(EditAnywhere, BlueprintReadWrite, Category = "Movie Scene")
 	TArray<FMovieSceneBindingProxy> Bindings;
 	
-	UPROPERTY(BlueprintReadWrite, Category = "Movie Scene")
+	UPROPERTY(EditAnywhere, BlueprintReadWrite, Category = "Movie Scene")
 	TObjectPtr<UMovieSceneFolder> ParentFolder;
 
-	UPROPERTY(BlueprintReadWrite, Category = "Movie Scene")
+	UPROPERTY(EditAnywhere, BlueprintReadWrite, Category = "Movie Scene")
 	TArray<TObjectPtr<UMovieSceneFolder>> Folders;
 
-	UPROPERTY(BlueprintReadWrite, Category = "Movie Scene")
+	UPROPERTY(EditAnywhere, BlueprintReadWrite, Category = "Movie Scene")
 	bool bDuplicateExistingActors;
 };
 
