Sequencer: Fix recursion in GetHandleToObject

diff --git a/Engine/Source/Editor/Sequencer/Private/Sequencer.cpp b/Engine/Source/Editor/Sequencer/Private/Sequencer.cpp
--- a/Engine/Source/Editor/Sequencer/Private/Sequencer.cpp
+++ b/Engine/Source/Editor/Sequencer/Private/Sequencer.cpp
@@ -3867,7 +3867,7 @@ FGuid FSequencer::GetHandleToObject( UObject* Object, bool bCreateHandleIfMissin
 			}
 		}
 
-		if (OwningActor)
+		if (OwningActor && OwningActor != Object)
 		{
 			GetHandleToObject(OwningActor);
 		}
