Sequencer: Optimized sequencer column logic slightly

  - Removed excessive allocations caused by continuously populating TSets
  - Changed column widget to cache IsActive et. al. on Tick so it only gets called once per-frame

diff --git a/Engine/Source/Editor/Sequencer/Private/MVVM/LockEditorExtension.cpp b/Engine/Source/Editor/Sequencer/Private/MVVM/LockEditorExtension.cpp
--- a/Engine/Source/Editor/Sequencer/Private/MVVM/LockEditorExtension.cpp
+++ b/Engine/Source/Editor/Sequencer/Private/MVVM/LockEditorExtension.cpp
@@ -35,10 +35,7 @@ bool FLockEditorExtension::IsNodeLocked(TWeakViewModelPtr<IOutlinerExtension> In
 
 	// Locked if all sections are locked
 	int32 NumSections = 0;
-	TSet<TWeakObjectPtr<UMovieSceneSection>> Sections;
-	SequencerHelpers::GetAllSections(Item, Sections);
-
-	for (auto Section : Sections)
+	for (const TViewModelPtr<FSectionModel>& Section : Item->GetDescendantsOfType<FSectionModel>())
 	{
 		if (!Section->IsLocked())
 		{
@@ -46,6 +43,7 @@ bool FLockEditorExtension::IsNodeLocked(TWeakViewModelPtr<IOutlinerExtension> In
 		}
 		++NumSections;
 	}
+
 	return NumSections > 0;
 }
 
@@ -66,12 +64,8 @@ void FLockEditorExtension::SetNodeLocked(TWeakViewModelPtr<IOutlinerExtension> I
 		// modify all selected items
 		for (FViewModelPtr OutlinerNode : EditorViewModel->GetSelection()->Outliner)
 		{
-			TSet<TWeakObjectPtr<UMovieSceneSection> > Sections;
-			SequencerHelpers::GetAllSections(OutlinerNode, Sections);
-
-			for (auto Section : Sections)
+			for (TSharedPtr<FSectionModel> Section : OutlinerNode->GetDescendantsOfType<FSectionModel>())
 			{
-				Section->Modify();
 				Section->SetIsLocked(bInIsLocked);
 			}
 		}
@@ -80,12 +74,8 @@ void FLockEditorExtension::SetNodeLocked(TWeakViewModelPtr<IOutlinerExtension> I
 	{
 		// only one unselected item was toggled, toggle just that node
 		FViewModelPtr Item = OutlinerItem;
-		TSet<TWeakObjectPtr<UMovieSceneSection> > Sections;
-		SequencerHelpers::GetAllSections(Item, Sections);
-
-		for (auto Section : Sections)
+		for (TSharedPtr<FSectionModel> Section : Item->GetDescendantsOfType<FSectionModel>())
 		{
-			Section->Modify();
 			Section->SetIsLocked(bInIsLocked);
 		}
 	}
@@ -94,16 +84,7 @@ void FLockEditorExtension::SetNodeLocked(TWeakViewModelPtr<IOutlinerExtension> I
 bool FLockEditorExtension::IsNodeLockable(TWeakViewModelPtr<IOutlinerExtension> InWeakOutlinerExtension) const
 {
 	FViewModelPtr Item = InWeakOutlinerExtension.Pin();
-
-	if (!Item)
-	{
-		return false;
-	}
-
-	TSet<TWeakObjectPtr<UMovieSceneSection> > Sections;
-	SequencerHelpers::GetAllSections(Item, Sections);
-
-	return Sections.Num() > 0;
+	return Item && (bool)Item->GetDescendantsOfType<FSectionModel>();
 }
 
 bool FLockEditorExtension::HasLockedChildNode(TWeakViewModelPtr<IOutlinerExtension> InWeakOutlinerExtension) const
diff --git a/Engine/Source/Editor/Sequencer/Private/MVVM/ViewModels/SectionModel.cpp b/Engine/Source/Editor/Sequencer/Private/MVVM/ViewModels/SectionModel.cpp
--- a/Engine/Source/Editor/Sequencer/Private/MVVM/ViewModels/SectionModel.cpp
+++ b/Engine/Source/Editor/Sequencer/Private/MVVM/ViewModels/SectionModel.cpp
@@ -54,6 +54,20 @@ FSectionModel::~FSectionModel()
 	WeakSection.Reset();
 }
 
+bool FSectionModel::IsLocked() const
+{
+	return bIsLocked;
+}
+
+void FSectionModel::SetIsLocked(bool bInIsLocked)
+{
+	if (UMovieSceneSection* Section = GetSection())
+	{
+		Section->Modify();
+		Section->SetIsLocked(bInIsLocked);
+	}
+}
+
 TRange<FFrameNumber> FSectionModel::GetRange() const
 {
 	UMovieSceneSection* Section = WeakSection.Get();
@@ -134,6 +148,7 @@ void FSectionModel::OnModifiedDirectly(UMovieSceneSignedObject*)
 
 void FSectionModel::UpdateCachedData()
 {
+	bIsLocked = true;
 	SectionRange = TRange<FFrameNumber>::Empty();
 	LayerBarRange = TRange<FFrameNumber>::Empty();
 
@@ -143,6 +158,7 @@ void FSectionModel::UpdateCachedData()
 		return;
 	}
 
+	bIsLocked = Section->IsLocked();
 	SectionRange = Section->GetRange();
 
 	// Compute the layer bar range from this section's effective key range
diff --git a/Engine/Source/Editor/Sequencer/Private/Widgets/OutlinerColumns/SColumnToggleWidget.cpp b/Engine/Source/Editor/Sequencer/Private/Widgets/OutlinerColumns/SColumnToggleWidget.cpp
--- a/Engine/Source/Editor/Sequencer/Private/Widgets/OutlinerColumns/SColumnToggleWidget.cpp
+++ b/Engine/Source/Editor/Sequencer/Private/Widgets/OutlinerColumns/SColumnToggleWidget.cpp
@@ -71,9 +71,12 @@ void SColumnToggleWidget::Construct(
 	WeakEditor = InParams.Editor->CastThisSharedChecked<FSequencerEditorViewModel>();
 
 	bIsMouseOverWidget = false;
+	bIsActive = false;
 	bIsChildActive = false;
+	bIsImplicitlyActive = false;
 
 	ActiveBrush = GetActiveBrush();
+	SetCanTick(true);
 
 	static const FName NAME_ChildActiveBrush = TEXT("Sequencer.Column.CheckBoxIndeterminate");
 	ChildActiveBrush = FAppStyle::Get().GetBrush(NAME_ChildActiveBrush);
@@ -86,6 +89,12 @@ void SColumnToggleWidget::Construct(
 	);
 }
 
+void SColumnToggleWidget::Tick(const FGeometry& AllottedGeometry, const double InCurrentTime, const float InDeltaTime)
+{
+	bIsActive = IsActive();
+	bIsChildActive = IsChildActive();
+	bIsImplicitlyActive = IsImplicitlyActive();
+}
 
 FSlateColor SColumnToggleWidget::GetImageColorAndOpacity() const
 {
@@ -101,12 +110,12 @@ FSlateColor SColumnToggleWidget::GetImageColorAndOpacity() const
 
 	float Opacity = 0.0f;
 
-	if (IsActive())
+	if (bIsActive)
 	{
 		// Directly active, full opacity
 		Opacity = 1.0f;
 	}
-	else if (IsChildActive() && !bIsMouseOverWidget)
+	else if (bIsChildActive && !bIsMouseOverWidget)
 	{
 		// Child is active and mouse is not over widget. Full opacity '-'.
 		Opacity = 1.0f;
@@ -116,7 +125,7 @@ FSlateColor SColumnToggleWidget::GetImageColorAndOpacity() const
 		// Mouse is over widget and it is not directly active.
 		Opacity = .65f;
 	}
-	else if (IsImplicitlyActive())
+	else if (bIsImplicitlyActive)
 	{
 		// Implicitly active through another object and mouse is not over.
 		Opacity = .35f;
@@ -146,8 +155,8 @@ const FSlateBrush* SColumnToggleWidget::GetBrush() const
 		return ActiveBrush;
 	}
 
-	if (IsChildActive()
-		&& !IsActive()
+	if (bIsChildActive
+		&& !bIsActive
 		&& !bIsMouseOverWidget)
 	{
 		return ChildActiveBrush;
@@ -186,8 +195,6 @@ void SColumnToggleWidget::OnDragEnter(const FGeometry& MyGeometry, const FDragDr
 
 FReply SColumnToggleWidget::HandleClick()
 {
-	const bool bIsActive = IsActive();
-
 	TSharedPtr<ISequencerOutlinerColumn> OutlinerColumn = WeakOutlinerColumn.Pin();
 
 	// Open an undo transaction
diff --git a/Engine/Source/Editor/Sequencer/Private/Widgets/OutlinerColumns/SColumnToggleWidget.h b/Engine/Source/Editor/Sequencer/Private/Widgets/OutlinerColumns/SColumnToggleWidget.h
--- a/Engine/Source/Editor/Sequencer/Private/Widgets/OutlinerColumns/SColumnToggleWidget.h
+++ b/Engine/Source/Editor/Sequencer/Private/Widgets/OutlinerColumns/SColumnToggleWidget.h
@@ -36,7 +36,7 @@ public:
 public:
 
 	/** Called when a drag or click ends and can be used to refresh the SequencerTree if necessary. */
-	virtual void OnToggleOperationComplete() {};
+	virtual void OnToggleOperationComplete() {}
 
 protected:
 
@@ -91,6 +91,9 @@ protected:
 	/** Process a mouse up message. */
 	virtual FReply OnMouseButtonUp(const FGeometry& MyGeometry, const FPointerEvent& MouseEvent) override;
 
+	/** Tick this widget to update its cached state */
+	virtual void Tick(const FGeometry& AllottedGeometry, const double InCurrentTime, const float InDeltaTime) override;
+
 	/** Called when mouse enters the column widget. */
 	void OnMouseEnter(const FGeometry& MyGeometry, const FPointerEvent& MouseEvent) override;
 
@@ -114,18 +117,24 @@ protected:
 
 private:
 
-	/** Whether or not the mouse is directly over the widget. */
-	bool bIsMouseOverWidget;
-
-	/** Whether or not a child of this item is active. */
-	bool bIsChildActive;
-
 	/** Brushes for the various states of active widgets. */
 	const FSlateBrush* ActiveBrush;
 	const FSlateBrush* ChildActiveBrush;
 
 	/** Scoped undo transaction. */
 	TUniquePtr<FScopedTransaction> UndoTransaction;
+
+	/** Whether or not the mouse is directly over the widget. */
+	uint8 bIsMouseOverWidget : 1;
+
+	/** Whether this item is active. */
+	uint8 bIsActive : 1;
+
+	/** Whether or not a child of this item is active. */
+	uint8 bIsChildActive : 1;
+
+	/** Whether this item is active implicitly (as a result of some item's state, not this item explicitly). */
+	uint8 bIsImplicitlyActive : 1;
 };
 
 } // namespace UE::Sequencer
\ No newline at end of file
diff --git a/Engine/Source/Editor/Sequencer/Public/MVVM/ViewModels/SectionModel.h b/Engine/Source/Editor/Sequencer/Public/MVVM/ViewModels/SectionModel.h
--- a/Engine/Source/Editor/Sequencer/Public/MVVM/ViewModels/SectionModel.h
+++ b/Engine/Source/Editor/Sequencer/Public/MVVM/ViewModels/SectionModel.h
@@ -86,6 +86,9 @@ public:
 	int32 GetPreRollFrames() const;
 	int32 GetPostRollFrames() const;
 
+	bool IsLocked() const;
+	void SetIsLocked(bool bIsLocked);
+
 public:
 
 	/*~ ILayerBarExtension Interface */
@@ -130,6 +133,8 @@ private:
 	TWeakObjectPtr<UMovieSceneSection> WeakSection;
 	TRange<FFrameNumber> SectionRange;
 	TRange<FFrameNumber> LayerBarRange;
+
+	uint8 bIsLocked : 1;
 };
 
 struct FOverlappingSections
