Sequencer:
- Export LayerBarModel

diff --git a/Engine/Source/Editor/Sequencer/Public/MVVM/ViewModels/LayerBarModel.h b/Engine/Source/Editor/Sequencer/Public/MVVM/ViewModels/LayerBarModel.h
--- a/Engine/Source/Editor/Sequencer/Public/MVVM/ViewModels/LayerBarModel.h
+++ b/Engine/Source/Editor/Sequencer/Public/MVVM/ViewModels/LayerBarModel.h
@@ -26,7 +26,7 @@ namespace Sequencer
 
 class FSequenceModel;
 
-class FLayerBarModel
+class SEQUENCER_API FLayerBarModel
 	: public FViewModel
 	, public FLinkedOutlinerExtension
 	, public ITrackLaneExtension
