Sequencer: Optimized mute/solo/lock columns by implementing caching and removing redundant UI delegation

  + Added a new row type for outliner trees that sets the background color once for the whole row. Previously this logic was implemented by each widget within the row (now each column widget, and the outliner view widget), which meant we re-ran the exact same logic (that was copied around) for background color computation multiple times per row. We now only do this once for the whole row.

  + Added a ModelID to all MVVM models. This is just a non-persistent serial uint32 unique to the instance and can be used for more efficient and stable keys than using the instance ptr.

  + SharedData is now a view model so it can support arbitrary extension and casting. This enables us to:
     + Add a new hierarchical cache extension that is easily accessible from any widget using GetSharedData()->CastThis<FOutlinerCacheExtension>()
     + Hierarchical Cache will perform a single pass on the entire MVVM hierarchy and allow clients to compute cached flags.
     + We have a separate cache for each of mute, solo and lock states that includes whether the item is set directly, and also whether it has settable children, and is implicitly set by a parent or child.

  + Added a dynamic extension iterator that allows anyone to iterate and filter the currently active dynamic extensions.
  + Moved Mute and Solo logic into the actual mute / solo extensions, rather than it being implemented in the NodeTree
  + Fixed header coloration for collapsed outliner rows that had selected keys within them. The logic here was incorrect and looked like an error in porting the code. It is now correct (and more efficient).

  - Only folders, objects, tracks and track rows now implement mute/solo. Channels and categories were never supported, and now no longer show as such on the UI.
  - Deleted the pre-existing XYZEditorExtension classes since they have all been ported to these new caches (the code is largely the same)

diff --git a/Engine/Source/Editor/Sequencer/Private/MVVM/LockEditorExtension.cpp b/Engine/Source/Editor/Sequencer/Private/MVVM/LockEditorExtension.cpp
deleted file mode 100644
diff --git a/Engine/Source/Editor/Sequencer/Private/MVVM/LockEditorExtension.h b/Engine/Source/Editor/Sequencer/Private/MVVM/LockEditorExtension.h
deleted file mode 100644
diff --git a/Engine/Source/Editor/Sequencer/Private/MVVM/MuteEditorExtension.cpp b/Engine/Source/Editor/Sequencer/Private/MVVM/MuteEditorExtension.cpp
deleted file mode 100644
diff --git a/Engine/Source/Editor/Sequencer/Private/MVVM/MuteEditorExtension.h b/Engine/Source/Editor/Sequencer/Private/MVVM/MuteEditorExtension.h
deleted file mode 100644
diff --git a/Engine/Source/Editor/Sequencer/Private/MVVM/Selection.cpp b/Engine/Source/Editor/Sequencer/Private/MVVM/Selection.cpp
--- a/Engine/Source/Editor/Sequencer/Private/MVVM/Selection.cpp
+++ b/Engine/Source/Editor/Sequencer/Private/MVVM/Selection.cpp
@@ -80,6 +80,13 @@ void FSequencerSelection::PreBroadcastChangeEvent()
 		if (OldNode)
 		{
 			OldNode->ToggleSelectionState(EOutlinerSelectionState::HasSelectedKeys | EOutlinerSelectionState::HasSelectedTrackAreaItems, false);
+
+			OldNode = OldNode.AsModel()->FindAncestorOfType<IOutlinerExtension>();
+			while (OldNode)
+			{
+				OldNode->ToggleSelectionState(EOutlinerSelectionState::DescendentHasSelectedTrackAreaItems | EOutlinerSelectionState::DescendentHasSelectedKeys, false);
+				OldNode = OldNode.AsModel()->FindAncestorOfType<IOutlinerExtension>();
+			}
 		}
 	}
 
@@ -95,6 +102,13 @@ void FSequencerSelection::PreBroadcastChangeEvent()
 			ParentOutlinerNode->ToggleSelectionState(EOutlinerSelectionState::HasSelectedTrackAreaItems, true);
 			NodesWithKeysOrSections.Add(ParentOutlinerNode);
 		}
+
+		ParentOutlinerNode = ParentOutlinerNode.AsModel()->FindAncestorOfType<IOutlinerExtension>();
+		while (ParentOutlinerNode)
+		{
+			ParentOutlinerNode->ToggleSelectionState(EOutlinerSelectionState::DescendentHasSelectedTrackAreaItems, true);
+			ParentOutlinerNode = ParentOutlinerNode.AsModel()->FindAncestorOfType<IOutlinerExtension>();
+		}
 	}
 
 	// Gather selection states from selected keys
@@ -107,6 +121,13 @@ void FSequencerSelection::PreBroadcastChangeEvent()
 			ParentOutlinerNode->ToggleSelectionState(EOutlinerSelectionState::HasSelectedKeys, true);
 			NodesWithKeysOrSections.Add(ParentOutlinerNode);
 		}
+
+		ParentOutlinerNode = ParentOutlinerNode.AsModel()->FindAncestorOfType<IOutlinerExtension>();
+		while (ParentOutlinerNode)
+		{
+			ParentOutlinerNode->ToggleSelectionState(EOutlinerSelectionState::DescendentHasSelectedKeys, true);
+			ParentOutlinerNode = ParentOutlinerNode.AsModel()->FindAncestorOfType<IOutlinerExtension>();
+		}
 	}
 
 	FSelectionEventSuppressor EventSuppressor = SuppressEvents();
diff --git a/Engine/Source/Editor/Sequencer/Private/MVVM/SequencerViewModelTypeIDs.cpp b/Engine/Source/Editor/Sequencer/Private/MVVM/SequencerViewModelTypeIDs.cpp
--- a/Engine/Source/Editor/Sequencer/Private/MVVM/SequencerViewModelTypeIDs.cpp
+++ b/Engine/Source/Editor/Sequencer/Private/MVVM/SequencerViewModelTypeIDs.cpp
@@ -25,14 +25,11 @@
 #include "MVVM/CurveEditorExtension.h"
 #include "MVVM/CurveEditorIntegrationExtension.h"
 #include "MVVM/FolderModelStorageExtension.h"
-#include "MVVM/LockEditorExtension.h"
-#include "MVVM/MuteEditorExtension.h"
 #include "MVVM/ObjectBindingModelStorageExtension.h"
 #include "MVVM/PinEditorExtension.h"
 #include "MVVM/TrackModelStorageExtension.h"
 #include "MVVM/TrackRowModelStorageExtension.h"
 #include "MVVM/SectionModelStorageExtension.h"
-#include "MVVM/SoloEditorExtension.h"
 
 
 namespace UE
@@ -49,6 +46,7 @@ UE_SEQUENCER_DEFINE_CASTABLE(FChannelModel);
 UE_SEQUENCER_DEFINE_CASTABLE(FFolderModel);
 UE_SEQUENCER_DEFINE_CASTABLE(FLayerBarModel);
 UE_SEQUENCER_DEFINE_CASTABLE(FOutlinerItemModel);
+UE_SEQUENCER_DEFINE_CASTABLE(FMuteSoloOutlinerItemModel);
 UE_SEQUENCER_DEFINE_CASTABLE(FObjectBindingModel);
 UE_SEQUENCER_DEFINE_CASTABLE(FPossessableModel);
 UE_SEQUENCER_DEFINE_CASTABLE(FSectionModel);
@@ -69,17 +67,15 @@ UE_SEQUENCER_DEFINE_CASTABLE(FSequencerTrackAreaViewModel);
 UE_SEQUENCER_DEFINE_VIEW_MODEL_TYPE_ID(FCurveEditorExtension);
 UE_SEQUENCER_DEFINE_VIEW_MODEL_TYPE_ID(FCurveEditorIntegrationExtension);
 UE_SEQUENCER_DEFINE_VIEW_MODEL_TYPE_ID(FFolderModelStorageExtension);
-UE_SEQUENCER_DEFINE_VIEW_MODEL_TYPE_ID(FLockEditorExtension);
-UE_SEQUENCER_DEFINE_VIEW_MODEL_TYPE_ID(FMuteEditorExtension);
 UE_SEQUENCER_DEFINE_VIEW_MODEL_TYPE_ID(FObjectBindingModelStorageExtension);
 UE_SEQUENCER_DEFINE_VIEW_MODEL_TYPE_ID(FPinEditorExtension);
 UE_SEQUENCER_DEFINE_VIEW_MODEL_TYPE_ID(FSectionModelStorageExtension);
-UE_SEQUENCER_DEFINE_VIEW_MODEL_TYPE_ID(FSoloEditorExtension);
 UE_SEQUENCER_DEFINE_VIEW_MODEL_TYPE_ID(FTrackModelStorageExtension);
 UE_SEQUENCER_DEFINE_VIEW_MODEL_TYPE_ID(FTrackRowModelStorageExtension);
 UE_SEQUENCER_DEFINE_VIEW_MODEL_TYPE_ID(ICurveEditorTreeItemExtension);
 UE_SEQUENCER_DEFINE_VIEW_MODEL_TYPE_ID(IObjectBindingExtension);
 UE_SEQUENCER_DEFINE_VIEW_MODEL_TYPE_ID(ISnappableExtension);
+UE_SEQUENCER_DEFINE_VIEW_MODEL_TYPE_ID(FOutlinerCacheExtension);
 
 } // namespace Sequencer
 } // namespace UE
diff --git a/Engine/Source/Editor/Sequencer/Private/MVVM/SoloEditorExtension.cpp b/Engine/Source/Editor/Sequencer/Private/MVVM/SoloEditorExtension.cpp
deleted file mode 100644
diff --git a/Engine/Source/Editor/Sequencer/Private/MVVM/SoloEditorExtension.h b/Engine/Source/Editor/Sequencer/Private/MVVM/SoloEditorExtension.h
deleted file mode 100644
diff --git a/Engine/Source/Editor/Sequencer/Private/MVVM/ViewModels/FolderModel.cpp b/Engine/Source/Editor/Sequencer/Private/MVVM/ViewModels/FolderModel.cpp
--- a/Engine/Source/Editor/Sequencer/Private/MVVM/ViewModels/FolderModel.cpp
+++ b/Engine/Source/Editor/Sequencer/Private/MVVM/ViewModels/FolderModel.cpp
@@ -43,7 +43,7 @@ TMap<TWeakObjectPtr<UMovieSceneFolder>, FColor> FFolderModel::InitialFolderColor
 bool FFolderModel::bFolderPickerWasCancelled = false;
 
 FFolderModel::FFolderModel(UMovieSceneFolder* Folder)
-	: FOutlinerItemModel()
+	: FMuteSoloOutlinerItemModel()
 	, TrackAreaList(EViewModelListType::TrackArea)
 	, WeakFolder(Folder)
 	, OwnerModel(nullptr)
diff --git a/Engine/Source/Editor/Sequencer/Private/MVVM/ViewModels/OutlinerItemModel.cpp b/Engine/Source/Editor/Sequencer/Private/MVVM/ViewModels/OutlinerItemModel.cpp
--- a/Engine/Source/Editor/Sequencer/Private/MVVM/ViewModels/OutlinerItemModel.cpp
+++ b/Engine/Source/Editor/Sequencer/Private/MVVM/ViewModels/OutlinerItemModel.cpp
@@ -9,6 +9,7 @@
 #include "MVVM/Extensions/IObjectBindingExtension.h"
 #include "MVVM/ViewModels/OutlinerViewModel.h"
 #include "MVVM/ViewModels/SequencerEditorViewModel.h"
+#include "MVVM/SharedViewModelData.h"
 #include "MVVM/Selection/Selection.h"
 
 #include "CurveEditor.h"
@@ -257,42 +258,22 @@
 
 bool FOutlinerItemModelMixin::IsDimmed() const
 {
-	TSharedPtr<FSequencerEditorViewModel> EditorViewModel = GetEditor();
-	TSharedPtr<FSequencer> Sequencer = EditorViewModel ? EditorViewModel->GetSequencerImpl() : nullptr;
-	if (Sequencer)
-	{
-		TViewModelPtr<IOutlinerExtension> ThisShared(const_cast<FViewModel*>(AsViewModel())->AsShared(), const_cast<FOutlinerItemModelMixin*>(this));
-		FSequencerNodeTree& ParentTree = Sequencer->GetNodeTree().Get();
-		if (ParentTree.IsNodeMute(ThisShared) || (ParentTree.HasSoloNodes() && !ParentTree.IsNodeSolo(ThisShared)))
-		{
-			return true;
-		}
-	}
-	return false;
+	const FViewModel* ViewModel = AsViewModel();
+
+	FMuteStateCacheExtension* MuteState = ViewModel->GetSharedData()->CastThis<FMuteStateCacheExtension>();
+	FSoloStateCacheExtension* SoloState = ViewModel->GetSharedData()->CastThis<FSoloStateCacheExtension>();
+
+	check(MuteState && SoloState);
+
+	ECachedMuteState MuteFlags = MuteState->GetCachedFlags(ViewModel->GetModelID());
+	ECachedSoloState SoloFlags = SoloState->GetCachedFlags(ViewModel->GetModelID());
+
+	const bool bAnySoloNodes = EnumHasAnyFlags(SoloState->GetRootFlags(), ECachedSoloState::Soloed | ECachedSoloState::PartiallySoloedChildren);
+	const bool bIsMuted      = EnumHasAnyFlags(MuteFlags, ECachedMuteState::Muted  | ECachedMuteState::ImplicitlyMutedByParent);
+	const bool bIsSoloed     = EnumHasAnyFlags(SoloFlags, ECachedSoloState::Soloed | ECachedSoloState::ImplicitlySoloedByParent);
+
+	const bool bDisableEval = bIsMuted || (bAnySoloNodes && !bIsSoloed);
+	return bDisableEval;
 }
 
-bool FOutlinerItemModelMixin::IsSolo() const
-{
-	TSharedPtr<FSequencerEditorViewModel> EditorViewModel = GetEditor();
-	TSharedPtr<FSequencer> Sequencer = EditorViewModel ? EditorViewModel->GetSequencerImpl() : nullptr;
-	if (Sequencer)
-	{
-		TViewModelPtr<IOutlinerExtension> ThisShared(const_cast<FViewModel*>(AsViewModel())->AsShared(), const_cast<FOutlinerItemModelMixin*>(this));
-		return Sequencer->GetNodeTree()->IsNodeSolo(ThisShared);
-	}
-	return false;
-}
-
-bool FOutlinerItemModelMixin::IsMuted() const
-{
-	TSharedPtr<FSequencerEditorViewModel> EditorViewModel = GetEditor();
-	TSharedPtr<FSequencer> Sequencer = EditorViewModel ? EditorViewModel->GetSequencerImpl() : nullptr;
-	if (Sequencer)
-	{
-		TViewModelPtr<IOutlinerExtension> ThisShared(const_cast<FViewModel*>(AsViewModel())->AsShared(), const_cast<FOutlinerItemModelMixin*>(this));
-		return Sequencer->GetNodeTree()->IsNodeMute(ThisShared);
-	}
-	return false;
-}
-
 bool FOutlinerItemModelMixin::IsRootModelPinned() const
@@ -319,45 +300,75 @@
 	}
 }
 
-bool FOutlinerItemModelMixin::IsSelectedModelsSolo() const
+ECheckBoxState FOutlinerItemModelMixin::SelectedModelsSoloState() const
 {
-	TSharedPtr<FSequencerEditorViewModel> EditorViewModel = GetEditor();
-	TSharedPtr<FSequencer> Sequencer = EditorViewModel ? EditorViewModel->GetSequencerImpl() : nullptr;
-	if (Sequencer)
+	FSoloStateCacheExtension* SoloStateCache = AsViewModel()->GetSharedData()->CastThis<FSoloStateCacheExtension>();
+	check(SoloStateCache);
+
+	int32 NumSoloables = 0;
+	int32 NumSoloed = 0;
+	for (FViewModelPtr Soloable : GetEditor()->GetSelection()->Outliner.Filter<ISoloableExtension>())
 	{
-		return Sequencer->GetNodeTree()->IsSelectedNodesSolo();
+		++NumSoloables;
+		if (EnumHasAnyFlags(SoloStateCache->GetCachedFlags(Soloable), ECachedSoloState::Soloed))
+		{
+			++NumSoloed;
+		}
 	}
-	return false;
+
+	if (NumSoloed == 0)
+	{
+		return ECheckBoxState::Unchecked;
+	}
+	return NumSoloables == NumSoloed ? ECheckBoxState::Checked : ECheckBoxState::Undetermined;
 }
 
 void FOutlinerItemModelMixin::ToggleSelectedModelsSolo()
 {
-	TSharedPtr<FSequencerEditorViewModel> EditorViewModel = GetEditor();
-	TSharedPtr<FSequencer> Sequencer = EditorViewModel ? EditorViewModel->GetSequencerImpl() : nullptr;
-	if (Sequencer)
+	ECheckBoxState CurrentState = SelectedModelsSoloState();
+	const bool bNewSoloState = CurrentState != ECheckBoxState::Checked;
+
+	const FScopedTransaction Transaction(NSLOCTEXT("Sequencer", "ToggleSolo", "Toggle Solo"));
+
+	for (TViewModelPtr<ISoloableExtension> Soloable : GetEditor()->GetSelection()->Outliner.Filter<ISoloableExtension>())
 	{
-		Sequencer->GetNodeTree()->ToggleSelectedNodesSolo();
+		Soloable->SetIsSoloed(bNewSoloState);
 	}
 }
 
-bool FOutlinerItemModelMixin::IsSelectedModelsMuted() const
+ECheckBoxState FOutlinerItemModelMixin::SelectedModelsMuteState() const
 {
-	TSharedPtr<FSequencerEditorViewModel> EditorViewModel = GetEditor();
-	TSharedPtr<FSequencer> Sequencer = EditorViewModel ? EditorViewModel->GetSequencerImpl() : nullptr;
-	if (Sequencer)
+	FMuteStateCacheExtension* MuteStateCache = AsViewModel()->GetSharedData()->CastThis<FMuteStateCacheExtension>();
+	check(MuteStateCache);
+
+	int32 NumMutables = 0;
+	int32 NumMuted = 0;
+	for (FViewModelPtr Mutable : GetEditor()->GetSelection()->Outliner.Filter<IMutableExtension>())
 	{
-		return Sequencer->GetNodeTree()->IsSelectedNodesMute();
+		++NumMutables;
+		if (EnumHasAnyFlags(MuteStateCache->GetCachedFlags(Mutable), ECachedMuteState::Muted))
+		{
+			++NumMuted;
+		}
 	}
-	return false;
+
+	if (NumMuted == 0)
+	{
+		return ECheckBoxState::Unchecked;
+	}
+	return NumMutables == NumMuted ? ECheckBoxState::Checked : ECheckBoxState::Undetermined;
 }
 
 void FOutlinerItemModelMixin::ToggleSelectedModelsMuted()
 {
-	TSharedPtr<FSequencerEditorViewModel> EditorViewModel = GetEditor();
-	TSharedPtr<FSequencer> Sequencer = EditorViewModel ? EditorViewModel->GetSequencerImpl() : nullptr;
-	if (Sequencer)
+	ECheckBoxState CurrentState = SelectedModelsMuteState();
+	const bool bNewMuteState = CurrentState != ECheckBoxState::Checked;
+
+	const FScopedTransaction Transaction(NSLOCTEXT("Sequencer", "ToggleMute", "Toggle Mute"));
+
+	for (TViewModelPtr<IMutableExtension> Muteable : GetEditor()->GetSelection()->Outliner.Filter<IMutableExtension>())
 	{
-		return Sequencer->GetNodeTree()->ToggleSelectedNodesMute();
+		Muteable->SetIsMuted(bNewMuteState);
 	}
 }
 
@@ -434,7 +445,7 @@
 			FUIAction(
 				FExecuteAction::CreateSP(SharedThis, &FOutlinerItemModelMixin::ToggleSelectedModelsSolo),
 				CanExecute,
-				FIsActionChecked::CreateSP(SharedThis, &FOutlinerItemModelMixin::IsSelectedModelsSolo)
+				FGetActionCheckState::CreateSP(SharedThis, &FOutlinerItemModelMixin::SelectedModelsSoloState)
 			),
 			NAME_None,
 			EUserInterfaceActionType::ToggleButton
@@ -447,7 +458,7 @@
 			FUIAction(
 				FExecuteAction::CreateSP(SharedThis, &FOutlinerItemModelMixin::ToggleSelectedModelsMuted),
 				CanExecute,
-				FIsActionChecked::CreateSP(SharedThis, &FOutlinerItemModelMixin::IsSelectedModelsMuted)
+				FGetActionCheckState::CreateSP(SharedThis, &FOutlinerItemModelMixin::SelectedModelsMuteState)
 			),
 			NAME_None,
 			EUserInterfaceActionType::ToggleButton
@@ -808,6 +819,128 @@
 	return false;
 }
 
+
+
+bool FMuteSoloOutlinerItemModel::IsSolo() const
+{
+	TSharedPtr<FSequencerEditorViewModel> EditorViewModel = GetEditor();
+	TSharedPtr<FSequencer> Sequencer = EditorViewModel ? EditorViewModel->GetSequencerImpl() : nullptr;
+	if (Sequencer)
+	{
+		const TArray<FString>& SoloNodes = Sequencer->GetFocusedMovieSceneSequence()->GetMovieScene()->GetSoloNodes();
+
+		// Should be always called on the main thread, but thread_local for safety
+		static thread_local TStringBuilder<256> StaticBuffer;
+		StaticBuffer.Reset();
+
+		IOutlinerExtension::GetPathName(*AsViewModel(), StaticBuffer);
+
+		const TCHAR* NodePath = StaticBuffer.ToString();
+
+		// It's pretty ridiculous to do a linear string search within this array, but that's what we have
+		return SoloNodes.Contains(NodePath);
+	}
+	return false;
+}
+
+void FMuteSoloOutlinerItemModel::SetIsSoloed(bool bIsSoloed)
+{
+	FViewModel* ViewModel = AsViewModel();
+	TSharedPtr<FSequenceModel> SequenceModel = ViewModel->FindAncestorOfType<FSequenceModel>();
+
+	if (SequenceModel)
+	{
+		UMovieScene* MovieScene = SequenceModel->GetMovieScene();
+		if (MovieScene->IsReadOnly())
+		{
+			return;
+		}
+
+		TArray<FString>& SoloNodes = MovieScene->GetSoloNodes();
+
+		MovieScene->Modify();
+
+		// Should be always called on the main thread, but thread_local for safety
+		static thread_local TStringBuilder<256> StaticBuffer;
+		StaticBuffer.Reset();
+
+		IOutlinerExtension::GetPathName(*ViewModel, StaticBuffer);
+
+		const TCHAR* NodePath = StaticBuffer.ToString();
+
+		if (bIsSoloed)
+		{
+			// Mark Mute, being careful as we might be re-marking an already Mute node
+			SoloNodes.AddUnique(NodePath);
+		}
+		else
+		{
+			// UnMute
+			SoloNodes.Remove(NodePath);
+		}
+	}
+}
+
+bool FMuteSoloOutlinerItemModel::IsMuted() const
+{
+	TSharedPtr<FSequencerEditorViewModel> EditorViewModel = GetEditor();
+	TSharedPtr<FSequencer> Sequencer = EditorViewModel ? EditorViewModel->GetSequencerImpl() : nullptr;
+	if (Sequencer)
+	{
+		const TArray<FString>& MuteNodes = Sequencer->GetFocusedMovieSceneSequence()->GetMovieScene()->GetMuteNodes();
+
+		// Should be always called on the main thread, but thread_local for safety
+		static thread_local TStringBuilder<256> StaticBuffer;
+		StaticBuffer.Reset();
+
+		IOutlinerExtension::GetPathName(*AsViewModel(), StaticBuffer);
+
+		const TCHAR* NodePath = StaticBuffer.ToString();
+
+		// It's pretty ridiculous to do a linear string search within this array, but that's what we have
+		return MuteNodes.Contains(NodePath);
+	}
+	return false;
+}
+
+void FMuteSoloOutlinerItemModel::SetIsMuted(bool bIsMuted)
+{
+	FViewModel* ViewModel = AsViewModel();
+	TSharedPtr<FSequenceModel> SequenceModel = ViewModel->FindAncestorOfType<FSequenceModel>();
+
+	if (SequenceModel)
+	{
+		UMovieScene* MovieScene = SequenceModel->GetMovieScene();
+		if (MovieScene->IsReadOnly())
+		{
+			return;
+		}
+
+		TArray<FString>& MuteNodes = MovieScene->GetMuteNodes();
+
+		// Should be always called on the main thread, but thread_local for safety
+		static thread_local TStringBuilder<256> StaticBuffer;
+		StaticBuffer.Reset();
+
+		IOutlinerExtension::GetPathName(*ViewModel, StaticBuffer);
+
+		const TCHAR* NodePath = StaticBuffer.ToString();
+
+		MovieScene->Modify();
+
+		if (bIsMuted)
+		{
+			// Mark Mute, being careful as we might be re-marking an already Mute node
+			MuteNodes.AddUnique(NodePath);
+		}
+		else
+		{
+			// UnMute
+			MuteNodes.Remove(NodePath);
+		}
+	}
+}
+
 } // namespace Sequencer
 } // namespace UE
 
diff --git a/Engine/Source/Editor/Sequencer/Private/MVVM/ViewModels/SectionModel.cpp b/Engine/Source/Editor/Sequencer/Private/MVVM/ViewModels/SectionModel.cpp
--- a/Engine/Source/Editor/Sequencer/Private/MVVM/ViewModels/SectionModel.cpp
+++ b/Engine/Source/Editor/Sequencer/Private/MVVM/ViewModels/SectionModel.cpp
@@ -54,20 +54,6 @@ FSectionModel::~FSectionModel()
 	WeakSection.Reset();
 }
 
-bool FSectionModel::IsLocked() const
-{
-	return bIsLocked;
-}
-
-void FSectionModel::SetIsLocked(bool bInIsLocked)
-{
-	if (UMovieSceneSection* Section = GetSection())
-	{
-		Section->Modify();
-		Section->SetIsLocked(bInIsLocked);
-	}
-}
-
 TRange<FFrameNumber> FSectionModel::GetRange() const
 {
 	UMovieSceneSection* Section = WeakSection.Get();
@@ -148,7 +134,6 @@ void FSectionModel::OnModifiedDirectly(UMovieSceneSignedObject*)
 
 void FSectionModel::UpdateCachedData()
 {
-	bIsLocked = true;
 	SectionRange = TRange<FFrameNumber>::Empty();
 	LayerBarRange = TRange<FFrameNumber>::Empty();
 
@@ -158,7 +143,6 @@ void FSectionModel::UpdateCachedData()
 		return;
 	}
 
-	bIsLocked = Section->IsLocked();
 	SectionRange = Section->GetRange();
 
 	// Compute the layer bar range from this section's effective key range
diff --git a/Engine/Source/Editor/Sequencer/Private/MVVM/ViewModels/SequenceModel.cpp b/Engine/Source/Editor/Sequencer/Private/MVVM/ViewModels/SequenceModel.cpp
--- a/Engine/Source/Editor/Sequencer/Private/MVVM/ViewModels/SequenceModel.cpp
+++ b/Engine/Source/Editor/Sequencer/Private/MVVM/ViewModels/SequenceModel.cpp
@@ -2,6 +2,9 @@
 
 #include "MVVM/ViewModels/SequenceModel.h"
 #include "MVVM/Extensions/IObjectBindingExtension.h"
+#include "MVVM/Extensions/ILockableExtension.h"
+#include "MVVM/Extensions/IMutableExtension.h"
+#include "MVVM/Extensions/ISoloableExtension.h"
 #include "MVVM/ViewModels/FolderModel.h"
 #include "MVVM/ViewModels/OutlinerItemModel.h"
 #include "MVVM/ViewModels/TrackModel.h"
@@ -39,7 +42,22 @@ FSequenceModel::FSequenceModel(TWeakPtr<FSequencerEditorViewModel> InEditorViewM
 
 void FSequenceModel::InitializeExtensions()
 {
-	SetSharedData(MakeShared<FSharedViewModelData>());
+	TSharedPtr<FSharedViewModelData> NewSharedData = MakeShared<FSharedViewModelData>();
+	SetSharedData(NewSharedData);
+
+	// Re-generate hierarchical caches when the sequence changes
+	NewSharedData->AddDynamicExtension(FOutlinerCacheExtension::ID);
+	NewSharedData->AddDynamicExtension(FMuteStateCacheExtension::ID);
+	NewSharedData->AddDynamicExtension(FSoloStateCacheExtension::ID);
+	NewSharedData->AddDynamicExtension(FLockStateCacheExtension::ID);
+
+	// Add our hierarchical cache processor
+	TSharedPtr<FOutlinerCacheExtension> OutlinerCache = NewSharedData->CastThisSharedChecked<FOutlinerCacheExtension>();
+	OutlinerCache->Initialize(AsShared());
+
+	// Make sure the hierarchical cache is updated when our hierarchy changes
+	FSimpleMulticastDelegate& HierarchyChanged = NewSharedData->SubscribeToHierarchyChanged(AsShared());
+	HierarchyChanged.AddSP(OutlinerCache.ToSharedRef(), &FOutlinerCacheExtension::OnHierarchyUpdated);
 
 	// This needs to be run outside of the constructor because a shared pointer to 'this' can't
 	// be created until after the object is fully built.
@@ -336,6 +354,21 @@ void FSequenceModel::PerformDrop(const FViewModelPtr& TargetModel, const FDragDr
 	}
 }
 
+void FSequenceModel::OnModifiedIndirectly(UMovieSceneSignedObject*)
+{
+	if (FOutlinerCacheExtension* OutlinerCache = GetSharedData()->CastThis<FOutlinerCacheExtension>())
+	{
+		OutlinerCache->UpdateCachedFlags();
+	}
+}
+
+void FSequenceModel::OnModifiedDirectly(UMovieSceneSignedObject*)
+{
+	if (FOutlinerCacheExtension* OutlinerCache = GetSharedData()->CastThis<FOutlinerCacheExtension>())
+	{
+		OutlinerCache->UpdateCachedFlags();
+	}
+}
 
 } // namespace Sequencer
 } // namespace UE
diff --git a/Engine/Source/Editor/Sequencer/Private/MVVM/ViewModels/SequencerEditorViewModel.cpp b/Engine/Source/Editor/Sequencer/Private/MVVM/ViewModels/SequencerEditorViewModel.cpp
--- a/Engine/Source/Editor/Sequencer/Private/MVVM/ViewModels/SequencerEditorViewModel.cpp
+++ b/Engine/Source/Editor/Sequencer/Private/MVVM/ViewModels/SequencerEditorViewModel.cpp
@@ -5,12 +5,9 @@
 #include "MVVM/ViewModels/SequencerTrackAreaViewModel.h"
 #include "MVVM/Selection/Selection.h"
 #include "MVVM/CurveEditorExtension.h"
-#include "MVVM/LockEditorExtension.h"
-#include "MVVM/MuteEditorExtension.h"
 #include "MVVM/PinEditorExtension.h"
 #include "MVVM/Extensions/IOutlinerExtension.h"
 #include "MVVM/SharedViewModelData.h"
-#include "MVVM/SoloEditorExtension.h"
 #include "ISequencerModule.h"
 #include "Sequencer.h"
 #include "MovieSceneSequenceID.h"
@@ -51,10 +48,7 @@ void FSequencerEditorViewModel::PreInitializeEditorImpl()
 		AddDynamicExtension(FCurveEditorExtension::ID);
 	}
 
-	AddDynamicExtension(FLockEditorExtension::ID);
-	AddDynamicExtension(FMuteEditorExtension::ID);
 	AddDynamicExtension(FPinEditorExtension::ID);
-	AddDynamicExtension(FSoloEditorExtension::ID);
 }
 
 TSharedPtr<FViewModel> FSequencerEditorViewModel::CreateRootModelImpl()
@@ -79,6 +73,11 @@ TSharedPtr<FTrackAreaViewModel> FSequencerEditorViewModel::CreateTrackAreaImpl()
 	return NewTrackArea;
 }
 
+TViewModelPtr<FSequenceModel> FSequencerEditorViewModel::GetRootSequenceModel() const
+{
+	return GetRootModel().ImplicitCast();
+}
+
 void FSequencerEditorViewModel::InitializeEditorImpl()
 {
 	PinnedTrackArea = CreateTrackAreaImpl();
diff --git a/Engine/Source/Editor/Sequencer/Private/MVVM/ViewModels/TrackModel.cpp b/Engine/Source/Editor/Sequencer/Private/MVVM/ViewModels/TrackModel.cpp
--- a/Engine/Source/Editor/Sequencer/Private/MVVM/ViewModels/TrackModel.cpp
+++ b/Engine/Source/Editor/Sequencer/Private/MVVM/ViewModels/TrackModel.cpp
@@ -2,6 +2,7 @@
 
 #include "MVVM/ViewModels/TrackModel.h"
 
+#include "MVVM/SharedViewModelData.h"
 #include "MVVM/Extensions/IObjectBindingExtension.h"
 #include "MVVM/Extensions/IRecyclableExtension.h"
 #include "MVVM/Extensions/ITrackExtension.h"
@@ -612,6 +613,42 @@ void FTrackModel::Resize(float NewSize)
 	}
 }
 
+ELockableLockState FTrackModel::GetLockState() const
+{
+	int32 NumSections = 0;
+	int32 NumLockedSections = 0;
+
+	for (const TViewModelPtr<FSectionModel>& Section : SectionList.Iterate<FSectionModel>())
+	{
+		++NumSections;
+
+		UMovieSceneSection* SectionObject = Section->GetSection();
+		if (SectionObject && SectionObject->IsLocked())
+		{
+			++NumLockedSections;
+		}
+	}
+
+	if (NumSections == 0 || NumLockedSections == 0)
+	{
+		return ELockableLockState::None;
+	}
+	return NumLockedSections == NumSections ? ELockableLockState::Locked : ELockableLockState::PartiallyLocked;
+}
+
+void FTrackModel::SetIsLocked(bool bInIsLocked)
+{
+	for (const TViewModelPtr<FSectionModel>& Section : SectionList.Iterate<FSectionModel>())
+	{
+		UMovieSceneSection* SectionObject = Section->GetSection();
+		if (SectionObject)
+		{
+			SectionObject->Modify();
+			SectionObject->SetIsLocked(bInIsLocked);
+		}
+	}
+}
+
 bool FTrackModel::CanDrag() const
 {
 	// Can only drag root tracks at the moment
diff --git a/Engine/Source/Editor/Sequencer/Private/MVVM/ViewModels/TrackRowModel.cpp b/Engine/Source/Editor/Sequencer/Private/MVVM/ViewModels/TrackRowModel.cpp
--- a/Engine/Source/Editor/Sequencer/Private/MVVM/ViewModels/TrackRowModel.cpp
+++ b/Engine/Source/Editor/Sequencer/Private/MVVM/ViewModels/TrackRowModel.cpp
@@ -144,6 +144,42 @@ bool FTrackRowModel::IsDimmed() const
 	return FOutlinerItemModel::IsDimmed();
 }
 
+ELockableLockState FTrackRowModel::GetLockState() const
+{
+	int32 NumSections = 0;
+	int32 NumLockedSections = 0;
+
+	for (const TViewModelPtr<FSectionModel>& Section : SectionList.Iterate<FSectionModel>())
+	{
+		++NumSections;
+
+		UMovieSceneSection* SectionObject = Section->GetSection();
+		if (SectionObject && SectionObject->IsLocked())
+		{
+			++NumLockedSections;
+		}
+	}
+
+	if (NumSections == 0 || NumLockedSections == 0)
+	{
+		return ELockableLockState::None;
+	}
+	return NumLockedSections == NumSections ? ELockableLockState::Locked : ELockableLockState::PartiallyLocked;
+}
+
+void FTrackRowModel::SetIsLocked(bool bInIsLocked)
+{
+	for (const TViewModelPtr<FSectionModel>& Section : SectionList.Iterate<FSectionModel>())
+	{
+		UMovieSceneSection* SectionObject = Section->GetSection();
+		if (SectionObject)
+		{
+			SectionObject->Modify();
+			SectionObject->SetIsLocked(bInIsLocked);
+		}
+	}
+}
+
 FSlateFontInfo FTrackRowModel::GetLabelFont() const
 {
 	bool bAllAnimated = false;
diff --git a/Engine/Source/Editor/Sequencer/Private/MVVM/Views/SSequencerOutlinerView.cpp b/Engine/Source/Editor/Sequencer/Private/MVVM/Views/SSequencerOutlinerView.cpp
new file mode 100644
--- /dev/null
+++ b/Engine/Source/Editor/Sequencer/Private/MVVM/Views/SSequencerOutlinerView.cpp
@@ -0,0 +1,97 @@
+// Copyright Epic Games, Inc. All Rights Reserved.
+
+#include "MVVM/Views/SSequencerOutlinerView.h"
+#include "MVVM/Extensions/IOutlinerExtension.h"
+#include "MVVM/ViewModels/OutlinerViewModel.h"
+#include "MVVM/Views/SOutlinerItemViewBase.h"
+
+#include "Styling/StyleColors.h"
+
+namespace UE::Sequencer
+{
+
+class SSequencerOutlinerViewRow : public SOutlinerViewRow
+{
+public:
+
+	void Construct(const FArguments& InArgs, const TSharedRef<STableViewBase>& OwnerTableView, TWeakViewModelPtr<IOutlinerExtension> InWeakModel, TWeakPtr<FOutlinerViewModel> InWeakOutliner)
+	{
+		WeakOutliner = InWeakOutliner;
+
+		SOutlinerViewRow::Construct(InArgs, OwnerTableView, InWeakModel);
+
+		if (InWeakModel.Pin()->HasBackground())
+		{
+			SetBorderImage(TAttribute<const FSlateBrush*>(FAppStyle::GetBrush("WhiteBrush")));
+			SetBorderBackgroundColor(MakeAttributeSP(this, &SSequencerOutlinerViewRow::GetBorderTint));
+		}
+	}
+
+	FSlateColor GetBorderTint() const
+	{
+		TSharedPtr<FOutlinerViewModel>    Outliner     = WeakOutliner.Pin();
+		TViewModelPtr<IOutlinerExtension> OutlinerItem = WeakModel.Pin();
+
+		if (!Outliner || !OutlinerItem)
+		{
+			return FLinearColor(0.f,0.f,0.f,0.f);
+		}
+
+		EOutlinerSelectionState SelectionState = OutlinerItem->GetSelectionState();
+
+		if (EnumHasAnyFlags(SelectionState, EOutlinerSelectionState::SelectedDirectly))
+		{
+			return FStyleColors::Select;
+		}
+		if (EnumHasAnyFlags(SelectionState, EOutlinerSelectionState::HasSelectedKeys | EOutlinerSelectionState::HasSelectedTrackAreaItems))
+		{
+			return FStyleColors::Header;
+		}
+
+		// If this is collapsed but has any children with selected keys or sections, we report that state on the parent
+		if (!OutlinerItem->IsExpanded())
+		{
+			if (EnumHasAnyFlags(SelectionState, EOutlinerSelectionState::DescendentHasSelectedTrackAreaItems | EOutlinerSelectionState::DescendentHasSelectedKeys))
+			{
+				return FStyleColors::Header;
+			}
+		}
+
+		// @todo: Currently we are always using 'default' for the style here. A lot of this is carried over from earlier Sequencer UIs
+		//        so we either need to re-evaluate whether EOutlinerItemViewBaseStyle is even necessary, or add proper support to the outliner extension to specify it
+		EOutlinerItemViewBaseStyle ItemStyle = EOutlinerItemViewBaseStyle::Default; /* OutlinerItem->GetItemViewStyle() */;
+
+		if (Outliner->GetHoveredItem() == OutlinerItem)
+		{
+			return ItemStyle == EOutlinerItemViewBaseStyle::ContainerHeader
+				? FLinearColor(FColor(52, 52, 52, 255))
+				: FLinearColor(FColor(72, 72, 72, 255));
+		}
+
+		return ItemStyle == EOutlinerItemViewBaseStyle::ContainerHeader
+			? FLinearColor(FColor(48, 48, 48, 255))
+			: FLinearColor(FColor(62, 62, 62, 255));
+	}
+
+private:
+
+	TWeakPtr<FOutlinerViewModel> WeakOutliner;
+};
+
+TSharedRef<ITableRow> SSequencerOutlinerView::OnGenerateRow(TWeakViewModelPtr<IOutlinerExtension> InWeakModel, const TSharedRef<STableViewBase>& OwnerTable)
+{
+	TSharedRef<SOutlinerViewRow> Row =
+		SNew(SSequencerOutlinerViewRow, OwnerTable, InWeakModel, WeakOutliner)
+		.OnDetectDrag(this, &SSequencerOutlinerView::OnDragRow)
+		.OnGenerateWidgetForColumn(this, &SSequencerOutlinerView::GenerateWidgetForColumn);
+
+	if (TViewModelPtr<IOutlinerExtension> ViewModel = InWeakModel.Pin())
+	{
+		CreateTrackLanesForRow(Row, ViewModel);
+	}
+
+	return Row;
+}
+
+
+} // namespace UE::Sequencer
\ No newline at end of file
diff --git a/Engine/Source/Editor/Sequencer/Private/MVVM/Views/SSequencerOutlinerView.h b/Engine/Source/Editor/Sequencer/Private/MVVM/Views/SSequencerOutlinerView.h
new file mode 100644
--- /dev/null
+++ b/Engine/Source/Editor/Sequencer/Private/MVVM/Views/SSequencerOutlinerView.h
@@ -0,0 +1,18 @@
+// Copyright Epic Games, Inc. All Rights Reserved.
+
+#pragma once
+
+#include "MVVM/Views/SOutlinerView.h"
+
+namespace UE::Sequencer
+{
+
+class SSequencerOutlinerView : public SOutlinerView
+{
+public:
+
+	TSharedRef<ITableRow> OnGenerateRow(TWeakViewModelPtr<IOutlinerExtension> InWeakModel, const TSharedRef<STableViewBase>& OwnerTable) override;
+};
+
+
+} // namespace UE::Sequencer
\ No newline at end of file
diff --git a/Engine/Source/Editor/Sequencer/Private/OutlinerColumns/LockOutlinerColumn.cpp b/Engine/Source/Editor/Sequencer/Private/OutlinerColumns/LockOutlinerColumn.cpp
--- a/Engine/Source/Editor/Sequencer/Private/OutlinerColumns/LockOutlinerColumn.cpp
+++ b/Engine/Source/Editor/Sequencer/Private/OutlinerColumns/LockOutlinerColumn.cpp
@@ -4,9 +4,10 @@
 
 #include "ISequencerOutlinerColumn.h"
 #include "MVVM/Extensions/IOutlinerExtension.h"
-#include "MVVM/LockEditorExtension.h"
 #include "MVVM/ViewModels/EditorViewModel.h"
 #include "MVVM/ViewModels/SequencerEditorViewModel.h"
+#include "MVVM/Extensions/ILockableExtension.h"
+#include "MVVM/SharedViewModelData.h"
 #include "Widgets/OutlinerColumns/SLockColumnWidget.h"
 
 #define LOCTEXT_NAMESPACE "FLockOutlinerColumn"
@@ -25,13 +26,9 @@ bool FLockOutlinerColumn::IsItemCompatibleWithColumn(const UE::Sequencer::FCreat
 {
 	using namespace UE::Sequencer;
 
-	if (InParams.Editor)
+	if (FLockStateCacheExtension* LockStateCache = InParams.OutlinerExtension.AsModel()->GetSharedData()->CastThis<FLockStateCacheExtension>())
 	{
-		FLockEditorExtension* LockEditorExtension = InParams.Editor->CastDynamic<FLockEditorExtension>();
-		if (LockEditorExtension)
-		{
-			return LockEditorExtension->IsNodeLockable(InParams.OutlinerExtension);
-		}
+		return EnumHasAnyFlags(LockStateCache->GetCachedFlags(InParams.OutlinerExtension), ECachedLockState::Lockable | ECachedLockState::LockableChildren);
 	}
 
 	return false;
diff --git a/Engine/Source/Editor/Sequencer/Private/OutlinerColumns/MuteOutlinerColumn.cpp b/Engine/Source/Editor/Sequencer/Private/OutlinerColumns/MuteOutlinerColumn.cpp
--- a/Engine/Source/Editor/Sequencer/Private/OutlinerColumns/MuteOutlinerColumn.cpp
+++ b/Engine/Source/Editor/Sequencer/Private/OutlinerColumns/MuteOutlinerColumn.cpp
@@ -3,8 +3,10 @@
 #include "OutlinerColumns/MuteOutlinerColumn.h"
 
 #include "ISequencerOutlinerColumn.h"
+
+#include "MVVM/SharedViewModelData.h"
 #include "MVVM/Extensions/IOutlinerExtension.h"
-#include "MVVM/MuteEditorExtension.h"
+#include "MVVM/Extensions/IMutableExtension.h"
 #include "MVVM/ViewModels/EditorViewModel.h"
 #include "MVVM/ViewModels/SequencerEditorViewModel.h"
 #include "Widgets/Layout/SScaleBox.h"
@@ -26,13 +28,9 @@ bool FMuteOutlinerColumn::IsItemCompatibleWithColumn(const UE::Sequencer::FCreat
 {
 	using namespace UE::Sequencer;
 
-	if (InParams.Editor)
+	if (FMuteStateCacheExtension* MuteStateCache = InParams.OutlinerExtension.AsModel()->GetSharedData()->CastThis<FMuteStateCacheExtension>())
 	{
-		FMuteEditorExtension* MuteEditorExtension = InParams.Editor->CastDynamic<FMuteEditorExtension>();
-		if (MuteEditorExtension)
-		{
-			return MuteEditorExtension->IsNodeMutable(InParams.OutlinerExtension);
-		}
+		return EnumHasAnyFlags(MuteStateCache->GetCachedFlags(InParams.OutlinerExtension), ECachedMuteState::Mutable | ECachedMuteState::MutableChildren);
 	}
 
 	return false;
diff --git a/Engine/Source/Editor/Sequencer/Private/OutlinerColumns/SoloOutlinerColumn.cpp b/Engine/Source/Editor/Sequencer/Private/OutlinerColumns/SoloOutlinerColumn.cpp
--- a/Engine/Source/Editor/Sequencer/Private/OutlinerColumns/SoloOutlinerColumn.cpp
+++ b/Engine/Source/Editor/Sequencer/Private/OutlinerColumns/SoloOutlinerColumn.cpp
@@ -3,8 +3,9 @@
 #include "OutlinerColumns/SoloOutlinerColumn.h"
 
 #include "ISequencerOutlinerColumn.h"
+#include "MVVM/SharedViewModelData.h"
 #include "MVVM/Extensions/IOutlinerExtension.h"
-#include "MVVM/SoloEditorExtension.h"
+#include "MVVM/Extensions/ISoloableExtension.h"
 #include "MVVM/ViewModels/EditorViewModel.h"
 #include "MVVM/ViewModels/SequencerEditorViewModel.h"
 #include "Widgets/Layout/SScaleBox.h"
@@ -26,13 +27,9 @@ bool FSoloOutlinerColumn::IsItemCompatibleWithColumn(const UE::Sequencer::FCreat
 {
 	using namespace UE::Sequencer;
 
-	if (InParams.Editor)
+	if (FSoloStateCacheExtension* SoloStateCache = InParams.OutlinerExtension.AsModel()->GetSharedData()->CastThis<FSoloStateCacheExtension>())
 	{
-		FSoloEditorExtension* SoloEditorExtension = InParams.Editor->CastDynamic<FSoloEditorExtension>();
-		if (SoloEditorExtension)
-		{
-			return SoloEditorExtension->IsNodeSoloable(InParams.OutlinerExtension);
-		}
+		return EnumHasAnyFlags(SoloStateCache->GetCachedFlags(InParams.OutlinerExtension), ECachedSoloState::Soloable | ECachedSoloState::SoloableChildren);
 	}
 
 	return false;
diff --git a/Engine/Source/Editor/Sequencer/Private/SSequencer.cpp b/Engine/Source/Editor/Sequencer/Private/SSequencer.cpp
--- a/Engine/Source/Editor/Sequencer/Private/SSequencer.cpp
+++ b/Engine/Source/Editor/Sequencer/Private/SSequencer.cpp
@@ -8,11 +8,13 @@
 #include "MVVM/ViewModels/ChannelModel.h"
 #include "MVVM/ViewModels/ViewModelIterators.h"
 #include "MVVM/Extensions/IRenameableExtension.h"
+#include "MVVM/Extensions/HierarchicalCacheExtension.h"
 #include "MVVM/ViewModels/SequencerEditorViewModel.h"
 #include "MVVM/ViewModels/SequencerTrackAreaViewModel.h"
 #include "MVVM/ViewModels/SequencerOutlinerViewModel.h"
 #include "MVVM/ViewModels/TrackRowModel.h"
 #include "MVVM/Selection/Selection.h"
+#include "MVVM/SharedViewModelData.h"
 #include "Sections/MovieSceneSubSection.h"
 #include "MovieScene.h"
 #include "Framework/MultiBox/MultiBoxDefs.h"
@@ -68,7 +70,7 @@
 #include "DragAndDrop/FolderDragDropOp.h"
 #include "DragAndDrop/CompositeDragDropOp.h"
 #include "Widgets/Input/SSearchBox.h"
-#include "MVVM/Views/SOutlinerView.h"
+#include "MVVM/Views/SSequencerOutlinerView.h"
 #include "MovieSceneTrackEditor.h"
 #include "SSequencerSplitterOverlay.h"
 #include "SequencerHotspots.h"
@@ -288,7 +290,7 @@ void SSequencer::Construct(const FArguments& InArgs, TSharedRef<FSequencer> InSe
 
 	TSharedPtr<FSequencerEditorViewModel> SequencerViewModel = InSequencer->GetViewModel()->CastThisShared<FSequencerEditorViewModel>();
 	SAssignNew(PinnedTrackArea, SSequencerTrackAreaView, SequencerViewModel->GetPinnedTrackArea(), TimeSliderControllerRef);
-	SAssignNew(PinnedTreeView, SOutlinerView, SequencerViewModel->GetOutliner(), PinnedTrackArea.ToSharedRef())
+	SAssignNew(PinnedTreeView, SSequencerOutlinerView, SequencerViewModel->GetOutliner(), PinnedTrackArea.ToSharedRef())
 		.Selection(InSequencer->GetViewModel()->GetSelection())
 		.ExternalScrollbar(PinnedAreaScrollBar)
 		.Clipping(EWidgetClipping::ClipToBounds);
@@ -299,7 +301,7 @@ void SSequencer::Construct(const FArguments& InArgs, TSharedRef<FSequencer> InSe
 	PinnedTreeView->SetShowPinned(true);
 
 	SAssignNew(TrackArea, SSequencerTrackAreaView, SequencerViewModel->GetTrackArea(), TimeSliderControllerRef);
-	SAssignNew(TreeView, SOutlinerView, SequencerViewModel->GetOutliner(), TrackArea.ToSharedRef())
+	SAssignNew(TreeView, SSequencerOutlinerView, SequencerViewModel->GetOutliner(), TrackArea.ToSharedRef())
 		.Selection(InSequencer->GetViewModel()->GetSelection())
 		.ExternalScrollbar(ScrollBar)
 		.Clipping(EWidgetClipping::ClipToBounds);
@@ -2970,19 +2972,34 @@ void SSequencer::UpdateLayoutTree()
 		{
 			bool bAnyChanged = false;
 
-			TSharedPtr<FSequencerNodeTree> NodeTree = Sequencer->GetNodeTree();
-			const bool bHasSoloNodes = NodeTree->HasSoloNodes();
+			TSharedPtr<FSharedViewModelData> SharedData = Sequencer->GetViewModel()->GetRootModel()->GetSharedData();
+			FOutlinerCacheExtension*     OutlinerCache  = SharedData->CastThis<FOutlinerCacheExtension>();
+			FMuteStateCacheExtension*    MuteState      = SharedData->CastThis<FMuteStateCacheExtension>();
+			FSoloStateCacheExtension*    SoloState      = SharedData->CastThis<FSoloStateCacheExtension>();
 
-			for (TViewModelPtr<IOutlinerExtension> OutlinerItem
-				: Sequencer->GetNodeTree()->GetRootNode()->GetDescendantsOfType<IOutlinerExtension>())
+			check(OutlinerCache && MuteState && SoloState);
+
+			// Hack - we shouldn't really not just forcibly update these here, but currently this function is getting forcibly called before
+			//        UpdateCachedFlags has a chance to naturally update itself in response to the signature change
+			OutlinerCache->UpdateCachedFlags();
+
+			const bool bAnySoloNodes = EnumHasAnyFlags(SoloState->GetRootFlags(), ECachedSoloState::Soloed | ECachedSoloState::PartiallySoloedChildren);
+
+			for (TViewModelPtr<ITrackExtension> TrackNode
+				: Sequencer->GetNodeTree()->GetRootNode()->GetDescendantsOfType<ITrackExtension>())
 			{
-				TViewModelPtr<ITrackExtension> TrackNode = OutlinerItem.ImplicitCast();
-				if (TrackNode && TrackNode->GetTrack())
+				if (TrackNode->GetTrack())
 				{
+					ECachedMuteState MuteFlags = MuteState->GetCachedFlags(TrackNode.AsModel());
+					ECachedSoloState SoloFlags = SoloState->GetCachedFlags(TrackNode.AsModel());
+
 					UMovieSceneTrack* Track = TrackNode->GetTrack();
-					bool bDisableEval = NodeTree->IsNodeMute(OutlinerItem) || (bHasSoloNodes && !NodeTree->IsNodeSolo(OutlinerItem));
+					const bool bIsMuted  = EnumHasAnyFlags(MuteFlags, ECachedMuteState::Muted  | ECachedMuteState::ImplicitlyMutedByParent);
+					const bool bIsSoloed = EnumHasAnyFlags(SoloFlags, ECachedSoloState::Soloed | ECachedSoloState::ImplicitlySoloedByParent);
+
+					const bool bDisableEval = bIsMuted || (bAnySoloNodes && !bIsSoloed);
 
-					TViewModelPtr<FTrackRowModel> TrackRowModel = OutlinerItem.ImplicitCast();
+					TViewModelPtr<FTrackRowModel> TrackRowModel = TrackNode.ImplicitCast();
 					if (TrackRowModel)
 					{
 						if (bDisableEval != Track->IsRowEvalDisabled(TrackNode->GetRowIndex()))
diff --git a/Engine/Source/Editor/Sequencer/Private/SequencerNodeTree.cpp b/Engine/Source/Editor/Sequencer/Private/SequencerNodeTree.cpp
--- a/Engine/Source/Editor/Sequencer/Private/SequencerNodeTree.cpp
+++ b/Engine/Source/Editor/Sequencer/Private/SequencerNodeTree.cpp
@@ -153,9 +153,6 @@ void FSequencerNodeTree::Update()
 	// Cache pinned state of nodes, needs to happen after OnTreeRefreshed
 	FPinnableExtensionShim::UpdateCachedPinnedState(RootNode);
 
-	// Update mute/solo markers and cache final state
-	UpdateMuteSolo(MovieScene);
-
 	// Update curve editor tree based on new filtered hierarchy
 	auto CurveEditorIntegration = SequenceModel->CastDynamic<FCurveEditorIntegrationExtension>();
 	if (CurveEditorIntegration)
@@ -315,226 +312,6 @@ bool FSequencerNodeTree::IsTrackLevelFilterActive(const FString& LevelName) cons
 	return !TrackFilterLevelFilter->IsLevelHidden(LevelName);
 }
 
-bool FSequencerNodeTree::IsNodeSolo(const UE::Sequencer::TViewModelPtr<UE::Sequencer::IOutlinerExtension>& InNode) const
-{
-	using namespace UE::Sequencer;
-
-	const TArray<FString>& SoloNodes = Sequencer.GetFocusedMovieSceneSequence()->GetMovieScene()->GetSoloNodes();
-	const FString NodePath = IOutlinerExtension::GetPathName(InNode);
-
-	if (SoloNodes.Contains(NodePath))
-	{
-		return true;
-	}
-
-	// If parent is solo'd, then this node should also be solo'd
-	TViewModelPtr<IOutlinerExtension> ParentNode = InNode.AsModel()->FindAncestorOfType<IOutlinerExtension>();
-	while (ParentNode)
-	{
-		const FString ParentNodePath = IOutlinerExtension::GetPathName(ParentNode);
-		if (SoloNodes.Contains(ParentNodePath))
-		{
-			return true;
-		}
-
-		ParentNode = ParentNode.AsModel()->FindAncestorOfType<IOutlinerExtension>();
-	}
-
-	// If any child is solo'd then this node needs to also be solo'd so that the child can be evaluated later during playback
-	for (TViewModelPtr<IOutlinerExtension> ChildNode : InNode.AsModel()->GetChildrenOfType<IOutlinerExtension>())
-	{
-		if (ChildNode)
-		{
-			const FString ChildNodePath = IOutlinerExtension::GetPathName(ChildNode);
-			if (SoloNodes.Contains(ChildNodePath))
-			{
-				return true;
-			}
-		}
-	}
-
-	return false;
-}
-
-bool FSequencerNodeTree::HasSoloNodes() const
-{
-	return bHasSoloNodes;
-}
-
-bool FSequencerNodeTree::IsSelectedNodesSolo() const
-{
-	using namespace UE::Sequencer;
-
-	if (Sequencer.GetViewModel()->GetSelection()->Outliner.Num() == 0)
-	{
-		return false;
-	}
-
-	const TArray<FString>& SoloNodes = Sequencer.GetFocusedMovieSceneSequence()->GetMovieScene()->GetSoloNodes();
-
-	bool bIsSolo = true;
-	for (FViewModelPtr Node : Sequencer.GetViewModel()->GetSelection()->Outliner)
-	{
-		if (!SoloNodes.Contains(IOutlinerExtension::GetPathName(Node)))
-		{
-			bIsSolo = false;
-			break;
-		}
-	}
-
-	return bIsSolo;
-}
-
-void FSequencerNodeTree::ToggleSelectedNodesSolo()
-{
-	using namespace UE::Sequencer;
-
-	UMovieScene* MovieScene = Sequencer.GetFocusedMovieSceneSequence()->GetMovieScene();
-	
-	if (MovieScene->IsReadOnly())
-	{
-		return;
-	}
-	
-	TArray<FString>& SoloNodes = MovieScene->GetSoloNodes();
-
-	if (Sequencer.GetViewModel()->GetSelection()->Outliner.Num() == 0)
-	{
-		return;
-	}
-
-	// First, determine if any of the selected nodes are not marked as solo
-	// If we have a mix, we should default to setting them all as solo
-	bool bIsSolo = true;
-	for (FViewModelPtr Node : Sequencer.GetViewModel()->GetSelection()->Outliner)
-	{
-		if (!SoloNodes.Contains(IOutlinerExtension::GetPathName(Node)))
-		{
-			bIsSolo = false;
-			break;
-		}
-	}
-
-	const FScopedTransaction Transaction(NSLOCTEXT("Sequencer", "ToggleSolo", "Toggle Solo"));
-
-	MovieScene->Modify();
-
-	for (FViewModelPtr Node : Sequencer.GetViewModel()->GetSelection()->Outliner)
-	{
-		FString NodePath = IOutlinerExtension::GetPathName(Node);
-		if (bIsSolo)
-		{
-			// If we're currently solo, unsolo
-			SoloNodes.Remove(NodePath);
-		}
-		else
-		{
-			// Mark solo, being careful as we might be re-marking an already solo node
-			SoloNodes.AddUnique(NodePath);
-		}
-	}
-	
-	GetSequencer().RefreshTree();
-}
-
-bool FSequencerNodeTree::IsNodeMute(const UE::Sequencer::TViewModelPtr<UE::Sequencer::IOutlinerExtension>& InNode) const
-{
-	using namespace UE::Sequencer;
-
-	const TArray<FString>& MuteNodes = Sequencer.GetFocusedMovieSceneSequence()->GetMovieScene()->GetMuteNodes();
-	const FString NodePath = IOutlinerExtension::GetPathName(InNode);
-
-	if (MuteNodes.Contains(NodePath))
-	{
-		return true;
-	}
-	
-	// Children should follow their parent's behavior unless told otherwise.
-	TViewModelPtr<IOutlinerExtension> ParentNode = InNode.AsModel()->FindAncestorOfType<IOutlinerExtension>();
-	if (ParentNode)
-	{
-		return IsNodeMute(ParentNode);
-	}
-	
-	return false;
-}
-
-bool FSequencerNodeTree::IsSelectedNodesMute() const
-{
-	using namespace UE::Sequencer;
-
-	if (Sequencer.GetViewModel()->GetSelection()->Outliner.Num() == 0)
-	{
-		return false;
-	}
-
-	const TArray<FString>& MuteNodes = Sequencer.GetFocusedMovieSceneSequence()->GetMovieScene()->GetMuteNodes();
-
-	bool bIsMute = true;
-	for (FViewModelPtr Node : Sequencer.GetViewModel()->GetSelection()->Outliner)
-	{
-		if (!MuteNodes.Contains(IOutlinerExtension::GetPathName(Node)))
-		{
-			bIsMute = false;
-			break;
-		}
-	}
-
-	return bIsMute;
-}
-
-void FSequencerNodeTree::ToggleSelectedNodesMute()
-{
-	using namespace UE::Sequencer;
-
-	UMovieScene* MovieScene = Sequencer.GetFocusedMovieSceneSequence()->GetMovieScene();
-	
-	if (MovieScene->IsReadOnly())
-	{
-		return;
-	}
-
-	TArray<FString>& MuteNodes = MovieScene->GetMuteNodes();
-
-	if (Sequencer.GetViewModel()->GetSelection()->Outliner.Num() == 0)
-	{
-		return;
-	}
-
-	// First, determine if any of the selected nodes are not marked as Mute
-	// If we have a mix, we should default to setting them all as Mute
-	bool bIsMute = true;
-	for (FViewModelPtr Node : Sequencer.GetViewModel()->GetSelection()->Outliner)
-	{
-		if (!MuteNodes.Contains(IOutlinerExtension::GetPathName(Node)))
-		{
-			bIsMute = false;
-			break;
-		}
-	}
-
-	const FScopedTransaction Transaction(NSLOCTEXT("Sequencer", "ToggleMute", "Toggle Mute"));
-
-	MovieScene->Modify();
-
-	for (FViewModelPtr Node : Sequencer.GetViewModel()->GetSelection()->Outliner)
-	{
-		FString NodePath = IOutlinerExtension::GetPathName(Node);
-		if (bIsMute)
-		{
-			// If we're currently Mute, unMute
-			MuteNodes.Remove(NodePath);
-		}
-		else
-		{
-			// Mark Mute, being careful as we might be re-marking an already Mute node
-			MuteNodes.AddUnique(NodePath);
-		}
-	}
-
-	GetSequencer().RefreshTree();
-}
-
 void FSequencerNodeTree::SaveExpansionState(const UE::Sequencer::FViewModel& Node, bool bExpanded)
 {
 	using namespace UE::Sequencer;
@@ -1046,25 +823,6 @@ void FSequencerNodeTree::CleanupMuteSolo(UMovieScene* MovieScene)
 	}
 }
 
-void FSequencerNodeTree::UpdateMuteSolo(UMovieScene* MovieScene)
-{
-	// Muting overrides soloing, so a solo node only counts if it's not muted.
-	bHasSoloNodes = false;
-	TArray<FString>& SoloNodes = Sequencer.GetFocusedMovieSceneSequence()->GetMovieScene()->GetSoloNodes();
-	for (const FString& NodePath : SoloNodes)
-	{
-		UE::Sequencer::TViewModelPtr<UE::Sequencer::IOutlinerExtension> Node = GetNodeAtPath(NodePath);
-		if (Node)
-		{
-			if (!IsNodeMute(Node))
-			{
-				bHasSoloNodes = true;
-				break;
-			}
-		}
-	}
-}
-
 int32 FSequencerNodeTree::GetTotalDisplayNodeCount() const 
 { 
 	// Subtract 1 for the spacer node which is always added
diff --git a/Engine/Source/Editor/Sequencer/Private/SequencerNodeTree.h b/Engine/Source/Editor/Sequencer/Private/SequencerNodeTree.h
--- a/Engine/Source/Editor/Sequencer/Private/SequencerNodeTree.h
+++ b/Engine/Source/Editor/Sequencer/Private/SequencerNodeTree.h
@@ -183,21 +183,6 @@ public:
 	void RemoveLevelFilter(const FString& LevelName);
 	bool IsTrackLevelFilterActive(const FString& LevelName) const;
 
-	void ToggleSelectedNodesSolo();
-	bool IsNodeSolo(const UE::Sequencer::TViewModelPtr<UE::Sequencer::IOutlinerExtension>& InNode) const;
-	
-	/** Returns whether any of the nodes in this tree are marked solo */
-	bool HasSoloNodes() const;
-
-	void ToggleSelectedNodesMute();
-	bool IsNodeMute(const UE::Sequencer::TViewModelPtr<UE::Sequencer::IOutlinerExtension>& InNode) const;
-
-	/** Returns whether any of the currently selected nodes are marked solo */
-	bool IsSelectedNodesSolo() const;
-	
-	/** Returns whether any of the currently selected nodes are muted */
-	bool IsSelectedNodesMute() const;
-
 private:
 
 	/** Returns whether this NodeTree should only display selected nodes */
@@ -218,11 +203,6 @@ private:
 	 */
 	void CleanupMuteSolo(UMovieScene* MovieScene);
 
-	/**
-	 * Updates mute/solo markers
-	 */
-	void UpdateMuteSolo(UMovieScene* MovieScene);
-
 public:
 
 	int32 GetTotalDisplayNodeCount() const;
@@ -253,8 +233,5 @@ private:
 
 	bool bFilterUpdateRequested;
 	bool bFilteringOnNodeGroups;
-
-	/** Cached value of whether we have any nodes that should be treated as soloing */
-	bool bHasSoloNodes;
 };
 
diff --git a/Engine/Source/Editor/Sequencer/Private/Widgets/OutlinerColumns/SColumnToggleWidget.cpp b/Engine/Source/Editor/Sequencer/Private/Widgets/OutlinerColumns/SColumnToggleWidget.cpp
--- a/Engine/Source/Editor/Sequencer/Private/Widgets/OutlinerColumns/SColumnToggleWidget.cpp
+++ b/Engine/Source/Editor/Sequencer/Private/Widgets/OutlinerColumns/SColumnToggleWidget.cpp
@@ -6,6 +6,7 @@
 #include "ISequencerOutlinerColumn.h"
 #include "MVVM/ViewModels/SequencerEditorViewModel.h"
 #include "MVVM/ViewModels/OutlinerViewModel.h"
+#include "MVVM/SharedViewModelData.h"
 #include "Sequencer.h"
 
 #define LOCTEXT_NAMESPACE "SColumnToggleWidget"
@@ -70,6 +71,8 @@ void SColumnToggleWidget::Construct(
 	WeakOutlinerExtension = InParams.OutlinerExtension;
 	WeakEditor = InParams.Editor->CastThisSharedChecked<FSequencerEditorViewModel>();
 
+	ModelID = InParams.OutlinerExtension.AsModel()->GetModelID();
+
 	bIsMouseOverWidget = false;
 	bIsActive = false;
 	bIsChildActive = false;
@@ -147,14 +150,6 @@ FSlateColor SColumnToggleWidget::GetImageColorAndOpacity() const
 
 const FSlateBrush* SColumnToggleWidget::GetBrush() const
 {
-	TSharedPtr<FEditorViewModel> Editor = WeakEditor.Pin();
-	TViewModelPtr<IOutlinerExtension> OutlinerItem = WeakOutlinerExtension.Pin();
-
-	if (!Editor || !OutlinerItem)
-	{
-		return ActiveBrush;
-	}
-
 	if (bIsChildActive
 		&& !bIsActive
 		&& !bIsMouseOverWidget)
diff --git a/Engine/Source/Editor/Sequencer/Private/Widgets/OutlinerColumns/SColumnToggleWidget.h b/Engine/Source/Editor/Sequencer/Private/Widgets/OutlinerColumns/SColumnToggleWidget.h
--- a/Engine/Source/Editor/Sequencer/Private/Widgets/OutlinerColumns/SColumnToggleWidget.h
+++ b/Engine/Source/Editor/Sequencer/Private/Widgets/OutlinerColumns/SColumnToggleWidget.h
@@ -124,6 +124,13 @@ private:
 	/** Scoped undo transaction. */
 	TUniquePtr<FScopedTransaction> UndoTransaction;
 
+protected:
+
+	/** Cached model ID of the outliner widget this widget relates to */
+	uint32 ModelID;
+
+private:
+
 	/** Whether or not the mouse is directly over the widget. */
 	uint8 bIsMouseOverWidget : 1;
 
diff --git a/Engine/Source/Editor/Sequencer/Private/Widgets/OutlinerColumns/SLockColumnWidget.cpp b/Engine/Source/Editor/Sequencer/Private/Widgets/OutlinerColumns/SLockColumnWidget.cpp
--- a/Engine/Source/Editor/Sequencer/Private/Widgets/OutlinerColumns/SLockColumnWidget.cpp
+++ b/Engine/Source/Editor/Sequencer/Private/Widgets/OutlinerColumns/SLockColumnWidget.cpp
@@ -1,9 +1,11 @@
 // Copyright Epic Games, Inc. All Rights Reserved.
 
 #include "Widgets/OutlinerColumns/SLockColumnWidget.h"
-
-#include "MVVM/LockEditorExtension.h"
+#include "MVVM/SharedViewModelData.h"
 #include "MVVM/ViewModels/SequencerEditorViewModel.h"
+#include "MVVM/Extensions/ILockableExtension.h"
+#include "MVVM/ViewModels/SequenceModel.h"
+#include "MVVM/Selection/Selection.h"
 
 namespace UE::Sequencer
 {
@@ -15,18 +17,15 @@ void SLockColumnWidget::Construct(const FArguments& InArgs, const TWeakPtr<ISequ
 		InWeakOutlinerColumn,
 		InParams
 	);
+
+	WeakLockStateCacheExtension = CastViewModel<FLockStateCacheExtension>(InParams.OutlinerExtension.AsModel()->GetSharedData());
 }
 
 bool SLockColumnWidget::IsActive() const
 {
-	TSharedPtr<FSequencerEditorViewModel> SequencerEditor = WeakEditor.Pin();
-	if (SequencerEditor)
+	if (TViewModelPtr<FLockStateCacheExtension> StateCache = WeakLockStateCacheExtension.Pin())
 	{
-		FLockEditorExtension* LockEditorExtension = SequencerEditor->CastDynamic<FLockEditorExtension>();
-		if (LockEditorExtension)
-		{
-			return LockEditorExtension->IsNodeLocked(WeakOutlinerExtension);
-		}
+		return EnumHasAnyFlags(StateCache->GetCachedFlags(ModelID), ECachedLockState::Locked);
 	}
 
 	return false;
@@ -34,27 +33,47 @@ bool SLockColumnWidget::IsActive() const
 
 void SLockColumnWidget::SetIsActive(const bool bInIsActive)
 {
-	TSharedPtr<FSequencerEditorViewModel> SequencerEditor = WeakEditor.Pin();
-	if (SequencerEditor)
+	TViewModelPtr<IOutlinerExtension> OutlinerItem = WeakOutlinerExtension.Pin();
+	if (!OutlinerItem)
 	{
-		FLockEditorExtension* LockEditorExtension = SequencerEditor->CastDynamic<FLockEditorExtension>();
-		if (LockEditorExtension)
+		return;
+	}
+
+	TSharedPtr<FSequenceModel> SequenceModel = OutlinerItem.AsModel()->FindAncestorOfType<FSequenceModel>();
+	if (!SequenceModel)
+	{
+		return;
+	}
+
+	const FScopedTransaction Transaction(NSLOCTEXT("Sequencer", "SetNodeLocked", "Set Node Locked"));
+
+	if (OutlinerItem->GetSelectionState() == EOutlinerSelectionState::SelectedDirectly)
+	{
+		// modify all selected items
+		for (FViewModelPtr OutlinerNode : SequenceModel->GetEditor()->GetSelection()->Outliner)
+		{
+			for (TSharedPtr<ILockableExtension> Lockable : OutlinerNode->GetDescendantsOfType<ILockableExtension>(true))
+			{
+				Lockable->SetIsLocked(bInIsActive);
+			}
+		}
+	}
+	else
+	{
+		// only one unselected item was toggled, toggle just that node
+		FViewModelPtr Item = OutlinerItem;
+		for (TSharedPtr<ILockableExtension> Lockable : Item->GetDescendantsOfType<ILockableExtension>(true))
 		{
-			LockEditorExtension->SetNodeLocked(WeakOutlinerExtension, bInIsActive);
+			Lockable->SetIsLocked(bInIsActive);
 		}
 	}
 }
 
 bool SLockColumnWidget::IsChildActive() const
 {
-	TSharedPtr<FSequencerEditorViewModel> SequencerEditor = WeakEditor.Pin();
-	if (SequencerEditor)
+	if (TViewModelPtr<FLockStateCacheExtension> StateCache = WeakLockStateCacheExtension.Pin())
 	{
-		FLockEditorExtension* LockEditorExtension = SequencerEditor->CastDynamic<FLockEditorExtension>();
-		if (LockEditorExtension)
-		{
-			return LockEditorExtension->HasLockedChildNode(WeakOutlinerExtension);
-		}
+		return EnumHasAnyFlags(StateCache->GetCachedFlags(ModelID), ECachedLockState::PartiallyLockedChildren);
 	}
 
 	return false;
diff --git a/Engine/Source/Editor/Sequencer/Private/Widgets/OutlinerColumns/SLockColumnWidget.h b/Engine/Source/Editor/Sequencer/Private/Widgets/OutlinerColumns/SLockColumnWidget.h
--- a/Engine/Source/Editor/Sequencer/Private/Widgets/OutlinerColumns/SLockColumnWidget.h
+++ b/Engine/Source/Editor/Sequencer/Private/Widgets/OutlinerColumns/SLockColumnWidget.h
@@ -7,6 +7,8 @@
 namespace UE::Sequencer
 {
 
+class FLockStateCacheExtension;
+
 /**
 * A widget that shows and controls the locked state of outliner items.
 */
@@ -31,6 +33,11 @@ public:
 
 	/** Returns the brush to display when this item is locked. */
 	virtual const FSlateBrush* GetActiveBrush() const override;
+
+private:
+
+	/** Weak cache extension ptr (can be null). */
+	TWeakViewModelPtr<FLockStateCacheExtension> WeakLockStateCacheExtension;
 };
 
 } // namespace UE::Sequencer
\ No newline at end of file
diff --git a/Engine/Source/Editor/Sequencer/Private/Widgets/OutlinerColumns/SMuteColumnWidget.cpp b/Engine/Source/Editor/Sequencer/Private/Widgets/OutlinerColumns/SMuteColumnWidget.cpp
--- a/Engine/Source/Editor/Sequencer/Private/Widgets/OutlinerColumns/SMuteColumnWidget.cpp
+++ b/Engine/Source/Editor/Sequencer/Private/Widgets/OutlinerColumns/SMuteColumnWidget.cpp
@@ -2,7 +2,11 @@
 
 #include "Widgets/OutlinerColumns/SMuteColumnWidget.h"
 
-#include "MVVM/MuteEditorExtension.h"
+
+#include "MVVM/SharedViewModelData.h"
+#include "MVVM/Selection/Selection.h"
+#include "MVVM/Extensions/IMutableExtension.h"
+#include "MVVM/ViewModels/SequenceModel.h"
 #include "MVVM/ViewModels/SequencerEditorViewModel.h"
 
 namespace UE::Sequencer
@@ -21,18 +25,15 @@ void SMuteColumnWidget::Construct(const FArguments& InArgs, const TWeakPtr<ISequ
 		InWeakOutlinerColumn,
 		InParams
 	);
+
+	WeakMuteStateCacheExtension = CastViewModel<FMuteStateCacheExtension>(InParams.OutlinerExtension.AsModel()->GetSharedData());
 }
 
 bool SMuteColumnWidget::IsActive() const
 {
-	TSharedPtr<FSequencerEditorViewModel> SequencerEditor = WeakEditor.Pin();
-	if (SequencerEditor)
+	if (TViewModelPtr<FMuteStateCacheExtension> MuteStateCache = WeakMuteStateCacheExtension.Pin())
 	{
-		FMuteEditorExtension* MuteEditorExtension = SequencerEditor->CastDynamic<FMuteEditorExtension>();
-		if (MuteEditorExtension)
-		{
-			return MuteEditorExtension->IsNodeMuted(WeakOutlinerExtension);
-		}
+		return EnumHasAnyFlags(MuteStateCache->GetCachedFlags(ModelID), ECachedMuteState::Muted);
 	}
 
 	return false;
@@ -40,27 +41,64 @@ bool SMuteColumnWidget::IsActive() const
 
 void SMuteColumnWidget::SetIsActive(const bool bInIsActive)
 {
-	TSharedPtr<FSequencerEditorViewModel> SequencerEditor = WeakEditor.Pin();
-	if (SequencerEditor)
+	TViewModelPtr<IOutlinerExtension> OutlinerItem = WeakOutlinerExtension.Pin();
+	if (!OutlinerItem)
+	{
+		return;
+	}
+
+	TSharedPtr<FSequenceModel> SequenceModel = OutlinerItem.AsModel()->FindAncestorOfType<FSequenceModel>();
+	if (!SequenceModel)
 	{
-		FMuteEditorExtension* MuteEditorExtension = SequencerEditor->CastDynamic<FMuteEditorExtension>();
-		if (MuteEditorExtension)
+		return;
+	}
+
+	const FScopedTransaction Transaction(NSLOCTEXT("Sequencer", "SetNodeMuted", "Set Node Muted"));
+
+	if (OutlinerItem->GetSelectionState() == EOutlinerSelectionState::SelectedDirectly)
+	{
+		// If selected, modify all selected items
+		for (FViewModelPtr Selected : SequenceModel->GetEditor()->GetSelection()->Outliner)
 		{
-			MuteEditorExtension->SetNodeMuted(WeakOutlinerExtension, bInIsActive);
+			SetIsActive(Selected, bInIsActive);
 		}
 	}
+	else
+	{
+		SetIsActive(OutlinerItem.AsModel(), bInIsActive);
+	}
 }
 
-bool SMuteColumnWidget::IsChildActive() const
+void SMuteColumnWidget::SetIsActive(const FViewModelPtr& ViewModel, const bool bInIsActive)
 {
-	TSharedPtr<FSequencerEditorViewModel> SequencerEditor = WeakEditor.Pin();
-	if (SequencerEditor)
+	TViewModelPtr<IMutableExtension> Mutable = ViewModel.ImplicitCast();
+	if (bInIsActive)
 	{
-		FMuteEditorExtension* MuteEditorExtension = SequencerEditor->CastDynamic<FMuteEditorExtension>();
-		if (MuteEditorExtension)
+		// If this is mutable, mute only this
+		if (Mutable)
 		{
-			return MuteEditorExtension->HasMutedChildNode(WeakOutlinerExtension);
+			Mutable->SetIsMuted(true);
 		}
+		// Otherwise mute mutable children of this (if any)
+		else for (TViewModelPtr<IMutableExtension> Child : ViewModel->GetDescendantsOfType<IMutableExtension>())
+		{
+			Child->SetIsMuted(true);
+		}
+	}
+	else
+	{
+		for (TViewModelPtr<IMutableExtension> Child : ViewModel->GetDescendantsOfType<IMutableExtension>(true))
+		{
+			Child->SetIsMuted(false);
+		}
+	}
+}
+
+bool SMuteColumnWidget::IsChildActive() const
+{
+	if (TViewModelPtr<FMuteStateCacheExtension> MuteStateCache = WeakMuteStateCacheExtension.Pin())
+	{
+		return EnumHasAnyFlags(MuteStateCache->GetCachedFlags(ModelID), ECachedMuteState::PartiallyMutedChildren);
 	}
 
 	return false;
@@ -68,14 +106,9 @@ bool SMuteColumnWidget::IsChildActive() const
 
 bool SMuteColumnWidget::IsImplicitlyActive() const
 {
-	TSharedPtr<FSequencerEditorViewModel> SequencerEditor = WeakEditor.Pin();
-	if (SequencerEditor)
+	if (TViewModelPtr<FMuteStateCacheExtension> MuteStateCache = WeakMuteStateCacheExtension.Pin())
 	{
-		FMuteEditorExtension* MuteEditorExtension = SequencerEditor->CastDynamic<FMuteEditorExtension>();
-		if (MuteEditorExtension)
-		{
-			return MuteEditorExtension->IsNodeImplicitlyMuted(WeakOutlinerExtension);
-		}
+		return EnumHasAnyFlags(MuteStateCache->GetCachedFlags(ModelID), ECachedMuteState::ImplicitlyMutedByParent);
 	}
 
 	return false;
diff --git a/Engine/Source/Editor/Sequencer/Private/Widgets/OutlinerColumns/SMuteColumnWidget.h b/Engine/Source/Editor/Sequencer/Private/Widgets/OutlinerColumns/SMuteColumnWidget.h
--- a/Engine/Source/Editor/Sequencer/Private/Widgets/OutlinerColumns/SMuteColumnWidget.h
+++ b/Engine/Source/Editor/Sequencer/Private/Widgets/OutlinerColumns/SMuteColumnWidget.h
@@ -7,6 +7,8 @@
 namespace UE::Sequencer
 {
 
+class FMuteStateCacheExtension;
+
 /**
  * A widget that shows and controls the muted state of outliner items.
  */
@@ -39,6 +41,13 @@ protected:
 
 	/** Returns the brush to display when this item is muted. */
 	virtual const FSlateBrush* GetActiveBrush() const override;
+
+private:
+
+	void SetIsActive(const FViewModelPtr& ViewModel, const bool bInIsActive);
+
+	/** Weak cache extension ptr (can be null). */
+	TWeakViewModelPtr<FMuteStateCacheExtension> WeakMuteStateCacheExtension;
 };
 
 } // namespace UE::Sequencer
\ No newline at end of file
diff --git a/Engine/Source/Editor/Sequencer/Private/Widgets/OutlinerColumns/SSoloColumnWidget.cpp b/Engine/Source/Editor/Sequencer/Private/Widgets/OutlinerColumns/SSoloColumnWidget.cpp
--- a/Engine/Source/Editor/Sequencer/Private/Widgets/OutlinerColumns/SSoloColumnWidget.cpp
+++ b/Engine/Source/Editor/Sequencer/Private/Widgets/OutlinerColumns/SSoloColumnWidget.cpp
@@ -2,7 +2,10 @@
 
 #include "Widgets/OutlinerColumns/SSoloColumnWidget.h"
 
-#include "MVVM/SoloEditorExtension.h"
+#include "MVVM/SharedViewModelData.h"
+#include "MVVM/Selection/Selection.h"
+#include "MVVM/Extensions/ISoloableExtension.h"
+#include "MVVM/ViewModels/SequenceModel.h"
 #include "MVVM/ViewModels/EditorViewModel.h"
 #include "MVVM/ViewModels/SequencerEditorViewModel.h"
 
@@ -22,18 +25,15 @@ void SSoloColumnWidget::Construct(const FArguments& InArgs, const TWeakPtr<ISequ
 		InWeakOutlinerColumn,
 		InParams
 	);
+
+	WeakSoloStateCacheExtension = CastViewModel<FSoloStateCacheExtension>(InParams.OutlinerExtension.AsModel()->GetSharedData());
 }
 
 bool SSoloColumnWidget::IsActive() const
 {
-	TSharedPtr<FSequencerEditorViewModel> SequencerEditor = WeakEditor.Pin();
-	if (SequencerEditor)
+	if (TViewModelPtr<FSoloStateCacheExtension> SoloStateCache = WeakSoloStateCacheExtension.Pin())
 	{
-		FSoloEditorExtension* SoloEditorExtension = SequencerEditor->CastDynamic<FSoloEditorExtension>();
-		if (SoloEditorExtension)
-		{
-			return SoloEditorExtension->IsNodeSoloed(WeakOutlinerExtension);
-		}
+		return EnumHasAnyFlags(SoloStateCache->GetCachedFlags(ModelID), ECachedSoloState::Soloed);
 	}
 
 	return false;
@@ -41,27 +41,39 @@ bool SSoloColumnWidget::IsActive() const
 
 void SSoloColumnWidget::SetIsActive(const bool bInIsActive)
 {
-	TSharedPtr<FSequencerEditorViewModel> SequencerEditor = WeakEditor.Pin();
-	if (SequencerEditor)
+	TViewModelPtr<IOutlinerExtension> OutlinerItem = WeakOutlinerExtension.Pin();
+	if (!OutlinerItem)
+	{
+		return;
+	}
+
+	TSharedPtr<FSequenceModel> SequenceModel = OutlinerItem.AsModel()->FindAncestorOfType<FSequenceModel>();
+	if (!SequenceModel)
+	{
+		return;
+	}
+
+	const FScopedTransaction Transaction(NSLOCTEXT("Sequencer", "SetNodeMuted", "Set Node Soloed"));
+
+	if (OutlinerItem->GetSelectionState() == EOutlinerSelectionState::SelectedDirectly)
 	{
-		FSoloEditorExtension* SoloEditorExtension = SequencerEditor->CastDynamic<FSoloEditorExtension>();
-		if (SoloEditorExtension)
+		// if selected, modify all selected items
+		for (TViewModelPtr<ISoloableExtension> Soloable : SequenceModel->GetEditor()->GetSelection()->Outliner.Filter<ISoloableExtension>())
 		{
-			SoloEditorExtension->SetNodeSoloed(WeakOutlinerExtension, bInIsActive);
+			Soloable->SetIsSoloed(bInIsActive);
 		}
 	}
+	else if (TViewModelPtr<ISoloableExtension> Soloable = OutlinerItem.ImplicitCast())
+	{
+		Soloable->SetIsSoloed(bInIsActive);
+	}
 }
 
 bool SSoloColumnWidget::IsChildActive() const
 {
-	TSharedPtr<FSequencerEditorViewModel> SequencerEditor = WeakEditor.Pin();
-	if (SequencerEditor)
+	if (TViewModelPtr<FSoloStateCacheExtension> SoloStateCache = WeakSoloStateCacheExtension.Pin())
 	{
-		FSoloEditorExtension* SoloEditorExtension = SequencerEditor->CastDynamic<FSoloEditorExtension>();
-		if (SoloEditorExtension)
-		{
-			return SoloEditorExtension->HasSoloedChildNode(WeakOutlinerExtension);
-		}
+		return EnumHasAnyFlags(SoloStateCache->GetCachedFlags(ModelID), ECachedSoloState::PartiallySoloedChildren);
 	}
 
 	return false;
@@ -69,14 +81,9 @@ bool SSoloColumnWidget::IsChildActive() const
 
 bool SSoloColumnWidget::IsImplicitlyActive() const
 {
-	TSharedPtr<FSequencerEditorViewModel> SequencerEditor = WeakEditor.Pin();
-	if (SequencerEditor)
+	if (TViewModelPtr<FSoloStateCacheExtension> SoloStateCache = WeakSoloStateCacheExtension.Pin())
 	{
-		FSoloEditorExtension* SoloEditorExtension = SequencerEditor->CastDynamic<FSoloEditorExtension>();
-		if (SoloEditorExtension)
-		{
-			return SoloEditorExtension->IsNodeImplicitlySoloed(WeakOutlinerExtension);
-		}
+		return EnumHasAnyFlags(SoloStateCache->GetCachedFlags(ModelID), ECachedSoloState::ImplicitlySoloedByParent);
 	}
 
 	return false;
diff --git a/Engine/Source/Editor/Sequencer/Private/Widgets/OutlinerColumns/SSoloColumnWidget.h b/Engine/Source/Editor/Sequencer/Private/Widgets/OutlinerColumns/SSoloColumnWidget.h
--- a/Engine/Source/Editor/Sequencer/Private/Widgets/OutlinerColumns/SSoloColumnWidget.h
+++ b/Engine/Source/Editor/Sequencer/Private/Widgets/OutlinerColumns/SSoloColumnWidget.h
@@ -7,6 +7,8 @@
 namespace UE::Sequencer
 {
 
+class FSoloStateCacheExtension;
+
 /**
 * A widget that shows and controls the soloed state of outliner items.
 */
@@ -40,6 +42,11 @@ protected:
 
 	/** Returns the brush to display when this item is pinned. */
 	virtual const FSlateBrush* GetActiveBrush() const override;
+
+private:
+
+	/** Weak cache extension ptr (can be null). */
+	TWeakViewModelPtr<FSoloStateCacheExtension> WeakSoloStateCacheExtension;
 };
 
 } // namespace UE::Sequencer
\ No newline at end of file
diff --git a/Engine/Source/Editor/Sequencer/Public/MVVM/ViewModels/FolderModel.h b/Engine/Source/Editor/Sequencer/Public/MVVM/ViewModels/FolderModel.h
--- a/Engine/Source/Editor/Sequencer/Public/MVVM/ViewModels/FolderModel.h
+++ b/Engine/Source/Editor/Sequencer/Public/MVVM/ViewModels/FolderModel.h
@@ -28,7 +28,7 @@ class FLayerBarModel;
 class FSequenceModel;
 
 class FFolderModel
-	: public FOutlinerItemModel
+	: public FMuteSoloOutlinerItemModel
 	, public IRenameableExtension
 	, public ITrackAreaExtension
 	, public IGroupableExtension
@@ -41,7 +41,7 @@ class FFolderModel
 public:
 
 	UE_SEQUENCER_DECLARE_CASTABLE(FFolderModel
-		, FOutlinerItemModel
+		, FMuteSoloOutlinerItemModel
 		, IRenameableExtension
 		, ITrackAreaExtension
 		, IGroupableExtension
diff --git a/Engine/Source/Editor/Sequencer/Public/MVVM/ViewModels/ObjectBindingModel.h b/Engine/Source/Editor/Sequencer/Public/MVVM/ViewModels/ObjectBindingModel.h
--- a/Engine/Source/Editor/Sequencer/Public/MVVM/ViewModels/ObjectBindingModel.h
+++ b/Engine/Source/Editor/Sequencer/Public/MVVM/ViewModels/ObjectBindingModel.h
@@ -40,7 +40,7 @@ enum class EObjectBindingType
 };
 
 class SEQUENCER_API FObjectBindingModel
-	: public FOutlinerItemModel
+	: public FMuteSoloOutlinerItemModel
 	, public IObjectBindingExtension
 	, public IDraggableOutlinerExtension
 	, public ITrackAreaExtension
@@ -52,7 +52,7 @@ class SEQUENCER_API FObjectBindingModel
 public:
 
 	UE_SEQUENCER_DECLARE_CASTABLE(FObjectBindingModel
-		, FOutlinerItemModel
+		, FMuteSoloOutlinerItemModel
 		, IObjectBindingExtension
 		, IDraggableOutlinerExtension
 		, ITrackAreaExtension
diff --git a/Engine/Source/Editor/Sequencer/Public/MVVM/ViewModels/OutlinerItemModel.h b/Engine/Source/Editor/Sequencer/Public/MVVM/ViewModels/OutlinerItemModel.h
--- a/Engine/Source/Editor/Sequencer/Public/MVVM/ViewModels/OutlinerItemModel.h
+++ b/Engine/Source/Editor/Sequencer/Public/MVVM/ViewModels/OutlinerItemModel.h
@@ -10,8 +10,9 @@
 #include "MVVM/Extensions/IHoveredExtension.h"
 #include "MVVM/Extensions/IOutlinerExtension.h"
 #include "MVVM/Extensions/IPinnableExtension.h"
-#include "MVVM/Extensions/ISoloableExtension.h"
 #include "MVVM/Extensions/IMutableExtension.h"
+#include "MVVM/Extensions/ISoloableExtension.h"
+#include "MVVM/Extensions/HierarchicalCacheExtension.h"
 #include "CurveEditorTypes.h"
 #include "Tree/ICurveEditorTreeItem.h"
 
@@ -31,14 +32,12 @@ class SEQUENCER_API FOutlinerItemModelMixin
 	, public FPinnableExtensionShim
 	, public FHoveredExtensionShim
 	, public IDimmableExtension
-	, public ISoloableExtension
-	, public IMutableExtension
 	, public FCurveEditorTreeItemExtensionShim
 	, public ICurveEditorTreeItem
 {
 public:
 
-	using Implements = TImplements<IOutlinerExtension, IGeometryExtension, IPinnableExtension, IHoveredExtension, IDimmableExtension, IMutableExtension, ISoloableExtension, ICurveEditorTreeItemExtension>;
+	using Implements = TImplements<IOutlinerExtension, IGeometryExtension, IPinnableExtension, IHoveredExtension, IDimmableExtension, ICurveEditorTreeItemExtension>;
 
 	FOutlinerItemModelMixin();
 
@@ -67,12 +66,6 @@ public:
 	/*~ IDimmableExtension */
 	bool IsDimmed() const override;
 
-	/*~ ISoloableExtension */
-	bool IsSolo() const override;
-
-	/*~ IMutableExtension */
-	bool IsMuted() const override;
-
 protected:
 
 	/** Get context menu contents. */
@@ -103,10 +96,10 @@ private:
 	bool IsRootModelPinned() const;
 	void ToggleRootModelPinned();
 
-	bool IsSelectedModelsSolo() const;
+	ECheckBoxState SelectedModelsSoloState() const;
 	void ToggleSelectedModelsSolo();
 
-	bool IsSelectedModelsMuted() const;
+	ECheckBoxState SelectedModelsMuteState() const;
 	void ToggleSelectedModelsMuted();
 
 private:
@@ -146,6 +139,39 @@ public:
 	UE_SEQUENCER_DECLARE_CASTABLE(FOutlinerItemModel, FOutlinerItemModelMixin);
 };
 
+class SEQUENCER_API FMuteSoloOutlinerItemModel
+	: public FOutlinerItemModel
+	, public IMutableExtension
+	, public ISoloableExtension
+{
+public:
+
+	UE_SEQUENCER_DECLARE_CASTABLE(FMuteSoloOutlinerItemModel, FOutlinerItemModel, IMutableExtension, ISoloableExtension);
+
+	/*~ ISoloableExtension */
+	bool IsSolo() const override;
+	void SetIsSoloed(bool bIsSoloed) override;
+
+	/*~ IMutableExtension */
+	bool IsMuted() const override;
+	void SetIsMuted(bool bIsMuted) override;
+};
+
+
+class SEQUENCER_API FOutlinerCacheExtension
+	: public FHierarchicalCacheExtension
+{
+public:
+
+	UE_SEQUENCER_DECLARE_VIEW_MODEL_TYPE_ID(FOutlinerCacheExtension);
+
+	FOutlinerCacheExtension()
+	{
+		ModelListFilter = EViewModelListType::Outliner;
+	}
+};
+
+
 } // namespace Sequencer
 } // namespace UE
 
diff --git a/Engine/Source/Editor/Sequencer/Public/MVVM/ViewModels/SectionModel.h b/Engine/Source/Editor/Sequencer/Public/MVVM/ViewModels/SectionModel.h
--- a/Engine/Source/Editor/Sequencer/Public/MVVM/ViewModels/SectionModel.h
+++ b/Engine/Source/Editor/Sequencer/Public/MVVM/ViewModels/SectionModel.h
@@ -86,9 +86,6 @@ public:
 	int32 GetPreRollFrames() const;
 	int32 GetPostRollFrames() const;
 
-	bool IsLocked() const;
-	void SetIsLocked(bool bIsLocked);
-
 public:
 
 	/*~ ILayerBarExtension Interface */
@@ -133,8 +130,6 @@ private:
 	TWeakObjectPtr<UMovieSceneSection> WeakSection;
 	TRange<FFrameNumber> SectionRange;
 	TRange<FFrameNumber> LayerBarRange;
-
-	uint8 bIsLocked : 1;
 };
 
 struct FOverlappingSections
diff --git a/Engine/Source/Editor/Sequencer/Public/MVVM/ViewModels/SequenceModel.h b/Engine/Source/Editor/Sequencer/Public/MVVM/ViewModels/SequenceModel.h
--- a/Engine/Source/Editor/Sequencer/Public/MVVM/ViewModels/SequenceModel.h
+++ b/Engine/Source/Editor/Sequencer/Public/MVVM/ViewModels/SequenceModel.h
@@ -65,6 +65,8 @@ public:
 
 	/*~ ISignedObjectEventHandler */
 	void OnPostUndo() override;
+	void OnModifiedIndirectly(UMovieSceneSignedObject*) override;
+	void OnModifiedDirectly(UMovieSceneSignedObject*) override;
 
 	/*~ ISortableExtension */
 	void SortChildren() override;
diff --git a/Engine/Source/Editor/Sequencer/Public/MVVM/ViewModels/SequencerEditorViewModel.h b/Engine/Source/Editor/Sequencer/Public/MVVM/ViewModels/SequencerEditorViewModel.h
--- a/Engine/Source/Editor/Sequencer/Public/MVVM/ViewModels/SequencerEditorViewModel.h
+++ b/Engine/Source/Editor/Sequencer/Public/MVVM/ViewModels/SequencerEditorViewModel.h
@@ -47,6 +47,8 @@ public:
 	// @todo_sequencer_mvvm move this to the root view-model
 	void SetSequence(UMovieSceneSequence* InRootSequence);
 
+	TViewModelPtr<FSequenceModel> GetRootSequenceModel() const;
+
 public:
 
 	/** Adjust sequencer customizations based on the currently focused sequence */
diff --git a/Engine/Source/Editor/Sequencer/Public/MVVM/ViewModels/TrackModel.h b/Engine/Source/Editor/Sequencer/Public/MVVM/ViewModels/TrackModel.h
--- a/Engine/Source/Editor/Sequencer/Public/MVVM/ViewModels/TrackModel.h
+++ b/Engine/Source/Editor/Sequencer/Public/MVVM/ViewModels/TrackModel.h
@@ -10,6 +10,7 @@
 #include "MVVM/Extensions/IResizableExtension.h"
 #include "MVVM/Extensions/ISortableExtension.h"
 #include "MVVM/Extensions/ITrackAreaExtension.h"
+#include "MVVM/Extensions/ILockableExtension.h"
 #include "MVVM/Extensions/ITrackExtension.h"
 #include "MVVM/Extensions/IDeletableExtension.h"
 #include "MVVM/Extensions/IDraggableOutlinerExtension.h"
@@ -35,11 +36,12 @@ namespace Sequencer
 class FSectionModel;
 
 class SEQUENCER_API FTrackModel
-	: public FOutlinerItemModel
+	: public FMuteSoloOutlinerItemModel
 	, public IRenameableExtension
 	, public IResizableExtension
 	, public ITrackExtension
 	, public ITrackAreaExtension
+	, public ILockableExtension
 	, public IGroupableExtension
 	, public ISortableExtension
 	, public IDraggableOutlinerExtension
@@ -50,11 +52,12 @@ class SEQUENCER_API FTrackModel
 public:
 
 	UE_SEQUENCER_DECLARE_CASTABLE(FTrackModel
-		, FOutlinerItemModel
+		, FMuteSoloOutlinerItemModel
 		, IRenameableExtension
 		, IResizableExtension
 		, ITrackExtension
 		, ITrackAreaExtension
+		, ILockableExtension
 		, IGroupableExtension
 		, ISortableExtension
 		, IDraggableOutlinerExtension
@@ -100,6 +103,10 @@ public:
 	bool IsResizable() const override;
 	void Resize(float NewSize) override;
 
+	/*~ ILockableExtension Interface */
+	ELockableLockState GetLockState() const override;
+	void SetIsLocked(bool bIsLocked) override;
+
 	/*~ ITrackAreaExtension */
 	FTrackAreaParameters GetTrackAreaParameters() const override;
 	FViewModelVariantIterator GetTrackAreaModelList() const override;
diff --git a/Engine/Source/Editor/Sequencer/Public/MVVM/ViewModels/TrackRowModel.h b/Engine/Source/Editor/Sequencer/Public/MVVM/ViewModels/TrackRowModel.h
--- a/Engine/Source/Editor/Sequencer/Public/MVVM/ViewModels/TrackRowModel.h
+++ b/Engine/Source/Editor/Sequencer/Public/MVVM/ViewModels/TrackRowModel.h
@@ -6,6 +6,7 @@
 #include "MVVM/ViewModels/OutlinerItemModel.h"
 #include "MVVM/Extensions/ITrackExtension.h"
 #include "MVVM/Extensions/ITrackAreaExtension.h"
+#include "MVVM/Extensions/ILockableExtension.h"
 #include "MVVM/Extensions/IRenameableExtension.h"
 #include "MVVM/Extensions/IResizableExtension.h"
 #include "MVVM/Extensions/IDeletableExtension.h"
@@ -25,8 +26,9 @@ namespace Sequencer
 class FSectionModel;
 
 class FTrackRowModel
-	: public FOutlinerItemModel
+	: public FMuteSoloOutlinerItemModel
 	, public ITrackAreaExtension
+	, public ILockableExtension
 	, public ITrackExtension
 	, public IRenameableExtension
 	, public IResizableExtension
@@ -34,7 +36,7 @@ class FTrackRowModel
 {
 public:
 
-	UE_SEQUENCER_DECLARE_CASTABLE(FTrackRowModel, FOutlinerItemModel, ITrackAreaExtension, ITrackExtension, IDeletableExtension, IRenameableExtension);
+	UE_SEQUENCER_DECLARE_CASTABLE(FTrackRowModel, FMuteSoloOutlinerItemModel, ITrackAreaExtension, ILockableExtension, ITrackExtension, IDeletableExtension, IMutableExtension, ISoloableExtension, IRenameableExtension);
 
 	explicit FTrackRowModel(UMovieSceneTrack* InTrack, int32 InRowIndex);
 	~FTrackRowModel();
@@ -82,6 +84,10 @@ public:
 	bool CanDelete(FText* OutErrorMessage) const override;
 	void Delete() override;
 
+	/*~ ILockableExtension Interface */
+	ELockableLockState GetLockState() const override;
+	void SetIsLocked(bool bIsLocked) override;
+
 private:
 
 	FViewModelListHead SectionList;
diff --git a/Engine/Source/Editor/SequencerCore/Private/MVVM/Extensions/DynamicExtensionContainer.cpp b/Engine/Source/Editor/SequencerCore/Private/MVVM/Extensions/DynamicExtensionContainer.cpp
--- a/Engine/Source/Editor/SequencerCore/Private/MVVM/Extensions/DynamicExtensionContainer.cpp
+++ b/Engine/Source/Editor/SequencerCore/Private/MVVM/Extensions/DynamicExtensionContainer.cpp
@@ -1,6 +1,7 @@
 // Copyright Epic Games, Inc. All Rights Reserved.
 
 #include "MVVM/Extensions/DynamicExtensionContainer.h"
+#include "MVVM/CastableTypeTable.h"
 
 namespace UE::Sequencer
 {
@@ -17,5 +18,41 @@ const void* FDynamicExtensionContainer::CastDynamic(FViewModelTypeID Type) const
 	return nullptr;
 }
 
+FDynamicExtensionContainerIterator::FDynamicExtensionContainerIterator(IteratorType&& InIterator, FViewModelTypeID InType)
+	: CurrentExtension(nullptr)
+	, Iterator(MoveTemp(InIterator))
+	, Type(InType)
+{
+	for ( ; Iterator; ++Iterator)
+	{
+		CurrentExtension = const_cast<void*>(Iterator->TypeTable->Cast(&Iterator->Extension.Get(), Type.GetTypeID()));
+		if (CurrentExtension)
+		{
+			break;
+		}
+	}
+}
+
+FDynamicExtensionContainerIterator& FDynamicExtensionContainerIterator::operator++()
+{
+	++Iterator;
+
+	for (CurrentExtension = nullptr; Iterator; ++Iterator)
+	{
+		CurrentExtension = const_cast<void*>(Iterator->TypeTable->Cast(&Iterator->Extension.Get(), Type.GetTypeID()));
+		if (CurrentExtension)
+		{
+			break;
+		}
+	}
+
+	return *this;
+}
+
+bool operator!=(const FDynamicExtensionContainerIterator& A, const FDynamicExtensionContainerIterator& B)
+{
+	return A.Iterator != B.Iterator || A.Type != B.Type;
+}
+
 } // namespace UE::Sequencer
 
diff --git a/Engine/Source/Editor/SequencerCore/Private/MVVM/Extensions/HierarchicalCacheExtension.cpp b/Engine/Source/Editor/SequencerCore/Private/MVVM/Extensions/HierarchicalCacheExtension.cpp
new file mode 100644
--- /dev/null
+++ b/Engine/Source/Editor/SequencerCore/Private/MVVM/Extensions/HierarchicalCacheExtension.cpp
@@ -0,0 +1,93 @@
+// Copyright Epic Games, Inc. All Rights Reserved.
+
+#include "MVVM/Extensions/HierarchicalCacheExtension.h"
+#include "MVVM/SharedViewModelData.h"
+#include "MVVM/ViewModels/ViewModelIterators.h"
+
+namespace UE::Sequencer
+{
+
+
+FHierarchicalCacheExtension::~FHierarchicalCacheExtension()
+{
+}
+
+void FHierarchicalCacheExtension::OnCreated(TSharedRef<FViewModel> InWeakOwner)
+{
+	WeakOwnerModel = InWeakOwner;
+}
+
+void FHierarchicalCacheExtension::Initialize(const FViewModelPtr& InRootModel)
+{
+	WeakRootModel = InRootModel;
+}
+
+void FHierarchicalCacheExtension::OnHierarchyUpdated()
+{
+	UpdateCachedFlags();
+}
+
+void FHierarchicalCacheExtension::UpdateCachedFlags()
+{
+	FViewModelPtr Owner = WeakOwnerModel.Pin();
+	if (!Owner)
+	{
+		return;
+	}
+
+	// Trigger our event, allowing caches to be added if necessary
+	PreUpdateCachesEvent.Broadcast(Owner);
+
+	TArray<IHierarchicalCache*> HierarchicalCaches;
+
+	for (IHierarchicalCache& HierarchicalCache : Owner->FilterDynamic<IHierarchicalCache>())
+	{
+		HierarchicalCaches.Add(&HierarchicalCache);
+	}
+
+	if (FViewModelPtr RootModel = WeakRootModel.Pin())
+	{
+		UpdateAllCachedFlags(RootModel, HierarchicalCaches);
+	}
+}
+
+void FHierarchicalCacheExtension::UpdateAllCachedFlags(const FViewModelPtr& ViewModel, TArrayView<IHierarchicalCache* const> HierarchicalCaches)
+{
+	for (IHierarchicalCache* HierarchicalCache : HierarchicalCaches)
+	{
+		HierarchicalCache->BeginUpdate();
+	}
+
+	if (ViewModel)
+	{
+		UpdateCachedFlagsForModel(ViewModel, HierarchicalCaches);
+	}
+
+	for (IHierarchicalCache* HierarchicalCache : HierarchicalCaches)
+	{
+		HierarchicalCache->EndUpdate();
+	}
+}
+
+void FHierarchicalCacheExtension::UpdateCachedFlagsForModel(const FViewModelPtr& ViewModel, TArrayView<IHierarchicalCache* const> HierarchicalCaches)
+{
+	// Visit parent-first caches, and call pre-visit children
+	for (IHierarchicalCache* HierarchicalCache : HierarchicalCaches)
+	{
+		HierarchicalCache->PreVisitChildren(ViewModel);
+	}
+
+	for (const FViewModelPtr& Child : ViewModel->GetChildren(ModelListFilter))
+	{
+		UpdateCachedFlagsForModel(Child, HierarchicalCaches);
+	}
+
+	// Call post-visit children and visit child-first caches
+	for (IHierarchicalCache* HierarchicalCache : HierarchicalCaches)
+	{
+		HierarchicalCache->PostVisitChildren(ViewModel);
+	}
+}
+
+} // namespace UE::Sequencer
+
diff --git a/Engine/Source/Editor/SequencerCore/Private/MVVM/Extensions/ILockableExtension.cpp b/Engine/Source/Editor/SequencerCore/Private/MVVM/Extensions/ILockableExtension.cpp
new file mode 100644
--- /dev/null
+++ b/Engine/Source/Editor/SequencerCore/Private/MVVM/Extensions/ILockableExtension.cpp
@@ -0,0 +1,82 @@
+// Copyright Epic Games, Inc. All Rights Reserved.
+
+#include "MVVM/Extensions/ILockableExtension.h"
+
+namespace UE::Sequencer
+{
+
+ECachedLockState FLockStateCacheExtension::ComputeFlagsForModel(const FViewModelPtr& ViewModel)
+{
+	ECachedLockState& ParentFlags = IndividualItemFlags.Last();
+	ECachedLockState ThisModelFlags = ECachedLockState::None;
+
+	if (EnumHasAnyFlags(ParentFlags, ECachedLockState::Locked | ECachedLockState::ImplicitlyLockedByParent))
+	{
+		ThisModelFlags |= ECachedLockState::ImplicitlyLockedByParent;
+	}
+
+	if (TViewModelPtr<ILockableExtension> Lockable = ViewModel.ImplicitCast())
+	{
+		ThisModelFlags |= ECachedLockState::Lockable;
+
+		switch (Lockable->GetLockState())
+		{
+		case ELockableLockState::Locked:
+			ThisModelFlags |= ECachedLockState::Locked;
+			break;
+		case ELockableLockState::PartiallyLocked:
+			// Lockable, but not locked - any parent must not be fully locked
+			if (EnumHasAnyFlags(ParentFlags, ECachedLockState::Locked))
+			{
+				ParentFlags |= ECachedLockState::PartiallyLockedChildren;
+				ParentFlags &= ~ECachedLockState::Locked;
+			}
+			break;
+		}
+	}
+
+	return ThisModelFlags;
+}
+
+void FLockStateCacheExtension::PostComputeChildrenFlags(const FViewModelPtr& ViewModel, ECachedLockState& OutThisModelFlags, ECachedLockState& OutPropagateToParentFlags)
+{
+	// --------------------------------------------------------------------
+	// Handle lock state propagation
+	const bool bIsLockable              = EnumHasAnyFlags(OutThisModelFlags, ECachedLockState::Lockable);
+	const bool bIsLocked                = EnumHasAnyFlags(OutThisModelFlags, ECachedLockState::Locked);
+	const bool bHasLockableChildren     = EnumHasAnyFlags(OutThisModelFlags, ECachedLockState::LockableChildren);
+	const bool bSiblingsPartiallyLocked = EnumHasAnyFlags(OutPropagateToParentFlags, ECachedLockState::PartiallyLockedChildren);
+	const bool bSiblingsFullyLocked     = EnumHasAnyFlags(OutPropagateToParentFlags, ECachedLockState::Locked);
+	const bool bHasAnyLockableSiblings  = EnumHasAnyFlags(OutPropagateToParentFlags, ECachedLockState::LockableChildren);
+
+	if (bIsLocked)
+	{
+		if (!bSiblingsPartiallyLocked)
+		{
+			if (bHasAnyLockableSiblings && !bSiblingsFullyLocked)
+			{
+				// This is the first locked lockable within the parent, but we know there are already other lockables so this has to be partially locked
+				OutPropagateToParentFlags |= ECachedLockState::PartiallyLockedChildren;
+			}
+			else
+			{
+				// Parent is (currently) fully locked because it contains no other lockable children, and is not already partially locked
+				OutPropagateToParentFlags |= ECachedLockState::Locked;
+			}
+		}
+	}
+	else if (bIsLockable && bSiblingsFullyLocked)
+	{
+		// Parent is no longer fully locked because it contains an unLocked lockable
+		OutPropagateToParentFlags |= ECachedLockState::PartiallyLockedChildren;
+		OutPropagateToParentFlags &= ~ECachedLockState::Locked;
+	}
+
+	if (bIsLockable)
+	{
+		OutPropagateToParentFlags |= ECachedLockState::LockableChildren;
+	}
+}
+
+} // namespace UE::Sequencer
+
diff --git a/Engine/Source/Editor/SequencerCore/Private/MVVM/Extensions/IMutableExtension.cpp b/Engine/Source/Editor/SequencerCore/Private/MVVM/Extensions/IMutableExtension.cpp
new file mode 100644
--- /dev/null
+++ b/Engine/Source/Editor/SequencerCore/Private/MVVM/Extensions/IMutableExtension.cpp
@@ -0,0 +1,71 @@
+// Copyright Epic Games, Inc. All Rights Reserved.
+
+#include "MVVM/Extensions/IMutableExtension.h"
+
+namespace UE::Sequencer
+{
+
+ECachedMuteState FMuteStateCacheExtension::ComputeFlagsForModel(const FViewModelPtr& ViewModel)
+{
+	ECachedMuteState ParentFlags = IndividualItemFlags.Last();
+	ECachedMuteState ThisModelFlags = ECachedMuteState::None;
+
+	if (EnumHasAnyFlags(ParentFlags, ECachedMuteState::Muted | ECachedMuteState::ImplicitlyMutedByParent))
+	{
+		ThisModelFlags |= ECachedMuteState::ImplicitlyMutedByParent;
+	}
+
+	if (TViewModelPtr<IMutableExtension> Mutable = ViewModel.ImplicitCast())
+	{
+		ThisModelFlags |= ECachedMuteState::Mutable;
+		if (Mutable->IsMuted())
+		{
+			ThisModelFlags |= ECachedMuteState::Muted;
+		}
+	}
+
+	return ThisModelFlags;
+}
+
+void FMuteStateCacheExtension::PostComputeChildrenFlags(const FViewModelPtr& ViewModel, ECachedMuteState& OutThisModelFlags, ECachedMuteState& OutPropagateToParentFlags)
+{
+	// --------------------------------------------------------------------
+	// Handle mute state propagation
+	const bool bIsMutable              = EnumHasAnyFlags(OutThisModelFlags, ECachedMuteState::Mutable);
+	const bool bIsMuted                = EnumHasAnyFlags(OutThisModelFlags, ECachedMuteState::Muted);
+	const bool bHasMutableChildren     = EnumHasAnyFlags(OutThisModelFlags, ECachedMuteState::MutableChildren);
+	const bool bSiblingsPartiallyMuted = EnumHasAnyFlags(OutPropagateToParentFlags, ECachedMuteState::PartiallyMutedChildren);
+	const bool bSiblingsFullyMuted     = EnumHasAnyFlags(OutPropagateToParentFlags, ECachedMuteState::Muted);
+	const bool bHasAnyMutableSiblings  = EnumHasAnyFlags(OutPropagateToParentFlags, ECachedMuteState::MutableChildren);
+
+	if (bIsMuted)
+	{
+		if (!bSiblingsPartiallyMuted)
+		{
+			if (bHasAnyMutableSiblings && !bSiblingsFullyMuted)
+			{
+				// This is the first muted mutable within the parent, but we know there are already other mutables so this has to be partially muted
+				OutPropagateToParentFlags |= ECachedMuteState::PartiallyMutedChildren;
+			}
+			else
+			{
+				// Parent is (currently) fully muted because it contains no other mutable children, and is not already partially muted
+				OutPropagateToParentFlags |= ECachedMuteState::Muted;
+			}
+		}
+	}
+	else if (bIsMutable && bSiblingsFullyMuted)
+	{
+		// Parent is no longer fully Muted because it contains an unmuted Mutable
+		OutPropagateToParentFlags |= ECachedMuteState::PartiallyMutedChildren;
+		OutPropagateToParentFlags &= ~ECachedMuteState::Muted;
+	}
+
+	if (bIsMutable)
+	{
+		OutPropagateToParentFlags |= ECachedMuteState::MutableChildren;
+	}
+}
+
+} // namespace UE::Sequencer
+
diff --git a/Engine/Source/Editor/SequencerCore/Private/MVVM/Extensions/ISoloableExtension.cpp b/Engine/Source/Editor/SequencerCore/Private/MVVM/Extensions/ISoloableExtension.cpp
new file mode 100644
--- /dev/null
+++ b/Engine/Source/Editor/SequencerCore/Private/MVVM/Extensions/ISoloableExtension.cpp
@@ -0,0 +1,71 @@
+// Copyright Epic Games, Inc. All Rights Reserved.
+
+#include "MVVM/Extensions/ISoloableExtension.h"
+
+namespace UE::Sequencer
+{
+
+ECachedSoloState FSoloStateCacheExtension::ComputeFlagsForModel(const FViewModelPtr& ViewModel)
+{
+	ECachedSoloState ParentFlags = IndividualItemFlags.Last();
+	ECachedSoloState ThisModelFlags = ECachedSoloState::None;
+
+	if (EnumHasAnyFlags(ParentFlags, ECachedSoloState::Soloed | ECachedSoloState::ImplicitlySoloedByParent))
+	{
+		ThisModelFlags |= ECachedSoloState::ImplicitlySoloedByParent;
+	}
+
+	if (TViewModelPtr<ISoloableExtension> Soloable = ViewModel.ImplicitCast())
+	{
+		ThisModelFlags |= ECachedSoloState::Soloable;
+		if (Soloable->IsSolo())
+		{
+			ThisModelFlags |= ECachedSoloState::Soloed;
+		}
+	}
+
+	return ThisModelFlags;
+}
+
+void FSoloStateCacheExtension::PostComputeChildrenFlags(const FViewModelPtr& ViewModel, ECachedSoloState& OutThisModelFlags, ECachedSoloState& OutPropagateToParentFlags)
+{
+	// --------------------------------------------------------------------
+	// Handle solo state propagation
+	const bool bIsSoloable              = EnumHasAnyFlags(OutThisModelFlags, ECachedSoloState::Soloable);
+	const bool bIsSoloed                = EnumHasAnyFlags(OutThisModelFlags, ECachedSoloState::Soloed);
+	const bool bHasSoloableChildren     = EnumHasAnyFlags(OutThisModelFlags, ECachedSoloState::SoloableChildren);
+	const bool bSiblingsPartiallySoloed = EnumHasAnyFlags(OutPropagateToParentFlags, ECachedSoloState::PartiallySoloedChildren);
+	const bool bSiblingsFullySoloed     = EnumHasAnyFlags(OutPropagateToParentFlags, ECachedSoloState::Soloed);
+	const bool bHasAnySoloableSiblings  = EnumHasAnyFlags(OutPropagateToParentFlags, ECachedSoloState::SoloableChildren);
+
+	if (bIsSoloed)
+	{
+		if (!bSiblingsPartiallySoloed)
+		{
+			if (bHasAnySoloableSiblings && !bSiblingsFullySoloed)
+			{
+				// This is the first soloed soloable within the parent, but we know there are already other soloables so this has to be partially soloed
+				OutPropagateToParentFlags |= ECachedSoloState::PartiallySoloedChildren;
+			}
+			else
+			{
+				// Parent is (currently) fully soloed because it contains no other soloable children, and is not already partially soloed
+				OutPropagateToParentFlags |= ECachedSoloState::Soloed;
+			}
+		}
+	}
+	else if (bIsSoloable && bSiblingsFullySoloed)
+	{
+		// Parent is no longer fully Soloed because it contains an unsoloed soloable
+		OutPropagateToParentFlags |= ECachedSoloState::PartiallySoloedChildren;
+		OutPropagateToParentFlags &= ~ECachedSoloState::Soloed;
+	}
+
+	if (bIsSoloable)
+	{
+		OutPropagateToParentFlags |= ECachedSoloState::SoloableChildren;
+	}
+}
+
+} // namespace UE::Sequencer
+
diff --git a/Engine/Source/Editor/SequencerCore/Private/MVVM/SharedViewModelData.cpp b/Engine/Source/Editor/SequencerCore/Private/MVVM/SharedViewModelData.cpp
--- a/Engine/Source/Editor/SequencerCore/Private/MVVM/SharedViewModelData.cpp
+++ b/Engine/Source/Editor/SequencerCore/Private/MVVM/SharedViewModelData.cpp
@@ -60,7 +60,7 @@ void FSharedViewModelData::PreHierarchicalChange(const TSharedPtr<FViewModel>& I
 	// If we do not have a current hierarchical operation, make a new one that will get flushed next tick
 	if (!CurrentHierarchicalOperation)
 	{
-		LatentOperation = MakeUnique<FViewModelHierarchyOperation>(AsShared());
+		LatentOperation = MakeUnique<FViewModelHierarchyOperation>(SharedThis(this));
 	}
 
 	check(CurrentHierarchicalOperation);
diff --git a/Engine/Source/Editor/SequencerCore/Private/MVVM/ViewModelTypeIDs.cpp b/Engine/Source/Editor/SequencerCore/Private/MVVM/ViewModelTypeIDs.cpp
--- a/Engine/Source/Editor/SequencerCore/Private/MVVM/ViewModelTypeIDs.cpp
+++ b/Engine/Source/Editor/SequencerCore/Private/MVVM/ViewModelTypeIDs.cpp
@@ -20,6 +20,7 @@
 #include "MVVM/Extensions/IResizableExtension.h"
 #include "MVVM/Extensions/ISelectableExtension.h"
 #include "MVVM/Extensions/ISoloableExtension.h"
+#include "MVVM/Extensions/ILockableExtension.h"
 #include "MVVM/Extensions/ISortableExtension.h"
 #include "MVVM/Extensions/IStretchableExtension.h"
 #include "MVVM/Extensions/ITrackAreaExtension.h"
@@ -51,7 +52,9 @@ UE_SEQUENCER_DEFINE_VIEW_MODEL_TYPE_ID(IHotspotExtension);
 UE_SEQUENCER_DEFINE_VIEW_MODEL_TYPE_ID(IHoveredExtension);
 UE_SEQUENCER_DEFINE_VIEW_MODEL_TYPE_ID(ILayerBarExtension);
 UE_SEQUENCER_DEFINE_VIEW_MODEL_TYPE_ID(IMutableExtension);
+UE_SEQUENCER_DEFINE_VIEW_MODEL_TYPE_ID(ILockableExtension);
 UE_SEQUENCER_DEFINE_VIEW_MODEL_TYPE_ID(IOutlinerDropTargetOutlinerExtension);
+UE_SEQUENCER_DEFINE_VIEW_MODEL_TYPE_ID(IHierarchicalCache);
 UE_SEQUENCER_DEFINE_VIEW_MODEL_TYPE_ID(IOutlinerExtension);
 UE_SEQUENCER_DEFINE_VIEW_MODEL_TYPE_ID(IPinnableExtension);
 UE_SEQUENCER_DEFINE_VIEW_MODEL_TYPE_ID(IRecyclableExtension);
@@ -70,6 +73,9 @@ UE_SEQUENCER_DEFINE_VIEW_MODEL_TYPE_ID(IKeyExtension);
 UE_SEQUENCER_DEFINE_VIEW_MODEL_TYPE_ID(IDynamicExtension);
 UE_SEQUENCER_DEFINE_VIEW_MODEL_TYPE_ID(FDynamicExtensionContainer);
 UE_SEQUENCER_DEFINE_VIEW_MODEL_TYPE_ID(FLinkedOutlinerExtension);
+UE_SEQUENCER_DEFINE_VIEW_MODEL_TYPE_ID(FMuteStateCacheExtension);
+UE_SEQUENCER_DEFINE_VIEW_MODEL_TYPE_ID(FSoloStateCacheExtension);
+UE_SEQUENCER_DEFINE_VIEW_MODEL_TYPE_ID(FLockStateCacheExtension);
 
 // View models
 UE_SEQUENCER_DEFINE_CASTABLE(FEditorViewModel);
diff --git a/Engine/Source/Editor/SequencerCore/Private/MVVM/ViewModels/ViewModel.cpp b/Engine/Source/Editor/SequencerCore/Private/MVVM/ViewModels/ViewModel.cpp
--- a/Engine/Source/Editor/SequencerCore/Private/MVVM/ViewModels/ViewModel.cpp
+++ b/Engine/Source/Editor/SequencerCore/Private/MVVM/ViewModels/ViewModel.cpp
@@ -14,11 +14,15 @@ namespace UE
 namespace Sequencer
 {
 
+uint32 GNextViewModelID = 0;
+
 FViewModel::FViewModel()
 	: FirstChildListHead(nullptr)
 	, ActiveIterationCount(0)
 	, bNeedsConstruction(1)
 {
+	// Don't really care about overflow here - by the time it does overflow, any overlapping IDs should have been released
+	ModelID = GNextViewModelID++;
 }
 
 FViewModel::~FViewModel()
diff --git a/Engine/Source/Editor/SequencerCore/Private/MVVM/ViewModels/ViewModelHierarchy.cpp b/Engine/Source/Editor/SequencerCore/Private/MVVM/ViewModels/ViewModelHierarchy.cpp
--- a/Engine/Source/Editor/SequencerCore/Private/MVVM/ViewModels/ViewModelHierarchy.cpp
+++ b/Engine/Source/Editor/SequencerCore/Private/MVVM/ViewModels/ViewModelHierarchy.cpp
@@ -366,11 +366,18 @@ void FViewModelChildren::Empty()
 }
 
 FViewModelHierarchyOperation::FViewModelHierarchyOperation(const TSharedPtr<FViewModel>& InAnyModel)
-	: FViewModelHierarchyOperation(InAnyModel->SharedData)
-{}
+	: SharedData(InAnyModel->SharedData)
+{
+	Construct();
+}
 
-FViewModelHierarchyOperation::FViewModelHierarchyOperation(const TSharedPtr<FSharedViewModelData>& InSharedData)
+FViewModelHierarchyOperation::FViewModelHierarchyOperation(const TSharedRef<FSharedViewModelData>& InSharedData)
 	: SharedData(InSharedData)
+{
+	Construct();
+}
+
+void FViewModelHierarchyOperation::Construct()
 {
 	if (SharedData)
 	{
diff --git a/Engine/Source/Editor/SequencerCore/Private/MVVM/Views/SOutlinerColumnBorder.cpp b/Engine/Source/Editor/SequencerCore/Private/MVVM/Views/SOutlinerColumnBorder.cpp
--- a/Engine/Source/Editor/SequencerCore/Private/MVVM/Views/SOutlinerColumnBorder.cpp
+++ b/Engine/Source/Editor/SequencerCore/Private/MVVM/Views/SOutlinerColumnBorder.cpp
@@ -20,8 +20,6 @@ void SOutlinerColumnInnerBorder::Construct(const FArguments& InArgs, const FCrea
 	WeakOutlinerExtension = InParams.OutlinerExtension;
 	WeakEditor = InParams.Editor;
 
-	BackgroundBrush = FAppStyle::GetBrush("Sequencer.Column.OutlinerColumnBox");
-
 	TSharedPtr<SScaleBox> UniformScaleBox;
 
 	// Size of outliner column widgets stretch to desired height and are usually 12x12 based off padding and fixed width of columns
@@ -35,7 +33,7 @@ void SOutlinerColumnInnerBorder::Construct(const FArguments& InArgs, const FCrea
 			[
 				SNew(SBorder)
 				.VAlign(VAlign_Center)
-				.BorderImage(this, &SOutlinerColumnInnerBorder::GetBorderImage)
+				.BorderImage(FAppStyle::GetBrush("Sequencer.Column.OutlinerColumnBox"))
 				.Padding(FMargin(1.f))
 				[
 					InArgs._Content.Widget
@@ -66,18 +64,11 @@ FReply SOutlinerColumnBorder::OnMouseButtonDown(const FGeometry& MyGeometry, con
 	return FReply::Handled();
 }
 
-const FSlateBrush* SOutlinerColumnInnerBorder::GetBorderImage() const
-{
-	return BackgroundBrush;
-}
-
 void SOutlinerColumnBorder::Construct(const FArguments& InArgs, const FCreateOutlinerColumnParams& InParams, const bool bHasInnerBorder)
 {
 	WeakOutlinerExtension = InParams.OutlinerExtension;
 	WeakEditor = InParams.Editor;
 
-	BackgroundBrush = FAppStyle::GetBrush("Sequencer.AnimationOutliner.DefaultBorder");
-
 	TSharedPtr<SWidget> InnerWidget;
 	if (bHasInnerBorder)
 	{
@@ -91,20 +82,12 @@ void SOutlinerColumnBorder::Construct(const FArguments& InArgs, const FCreateOut
 		InnerWidget = InArgs._Content.Widget;
 	}
 
-	TSharedRef<SWidget>	FinalWidget = SNew(SBorder)
-		.VAlign(VAlign_Center)
-		.BorderImage(this, &SOutlinerColumnBorder::GetBorderImage)
-		.BorderBackgroundColor(this, &SOutlinerColumnBorder::GetBackgroundTint)
-		.VAlign(VAlign_Center)
-		.HAlign(HAlign_Center)
-		.Padding(FMargin(2.0f))
-		[
-			InnerWidget.ToSharedRef()
-		];
-
 	ChildSlot
+	.VAlign(VAlign_Center)
+	.HAlign(HAlign_Center)
+	.Padding(FMargin(2.0f))
 	[
-		FinalWidget
+		InnerWidget.ToSharedRef()
 	];
 }
 
@@ -129,57 +112,5 @@ void SOutlinerColumnBorder::OnMouseLeave(const FPointerEvent& MouseEvent)
 	SWidget::OnMouseLeave(MouseEvent);
 }
 
-const FSlateBrush* SOutlinerColumnBorder::GetBorderImage() const
-{
-	return BackgroundBrush;
-}
-
-FSlateColor SOutlinerColumnBorder::GetBackgroundTint() const
-{
-	TSharedPtr<FEditorViewModel> Editor = WeakEditor.Pin();
-	TViewModelPtr<IOutlinerExtension> OutlinerItem = WeakOutlinerExtension.Pin();
-
-	if (!Editor || !OutlinerItem)
-	{
-		return FLinearColor(0.f, 0.f, 0.f, 0.f);
-	}
-
-	if (!OutlinerItem->HasBackground())
-	{
-		return FLinearColor(0.f, 0.f, 0.f, 0.f);
-	}
-
-	EOutlinerSelectionState SelectionState = OutlinerItem->GetSelectionState();
-
-	if (EnumHasAnyFlags(SelectionState, EOutlinerSelectionState::SelectedDirectly))
-	{
-		return FStyleColors::Select;
-	}
-	if (EnumHasAnyFlags(SelectionState, EOutlinerSelectionState::HasSelectedKeys | EOutlinerSelectionState::HasSelectedTrackAreaItems))
-	{
-		return FStyleColors::Header;
-	}
-
-	// If this is collapsed but has any children with selected keys or sections, we report that state on the parent
-	if (!OutlinerItem->IsExpanded())
-	{
-		for (TViewModelPtr<IOutlinerExtension> Child : OutlinerItem.AsModel()->GetDescendantsOfType<IOutlinerExtension>())
-		{
-			if (EnumHasAnyFlags(SelectionState, EOutlinerSelectionState::HasSelectedKeys | EOutlinerSelectionState::HasSelectedTrackAreaItems))
-			{
-				return FStyleColors::Header;
-			}
-		}
-	}
-
-	// Use same hovered and default colors as SOutlinerItemViewBase
-	if (Editor->GetOutliner()->GetHoveredItem() == OutlinerItem)
-	{
-		return FLinearColor(FColor(72, 72, 72, 255));
-	}
-
-	return FLinearColor(FColor(62, 62, 62, 255));
-}
-
 } // namespace UE::Sequencer
 
diff --git a/Engine/Source/Editor/SequencerCore/Private/MVVM/Views/SOutlinerItemViewBase.cpp b/Engine/Source/Editor/SequencerCore/Private/MVVM/Views/SOutlinerItemViewBase.cpp
--- a/Engine/Source/Editor/SequencerCore/Private/MVVM/Views/SOutlinerItemViewBase.cpp
+++ b/Engine/Source/Editor/SequencerCore/Private/MVVM/Views/SOutlinerItemViewBase.cpp
@@ -146,12 +146,6 @@ void SOutlinerItemViewBase::Construct(
 	static float IndentAmount = 10.f;
 
 	TSharedRef<SWidget>	FinalWidget = 
-		SNew( SBorder )
-		.VAlign( VAlign_Center )
-		.BorderImage( this, &SOutlinerItemViewBase::GetNodeBorderImage )
-		.BorderBackgroundColor( this, &SOutlinerItemViewBase::GetNodeBackgroundTint )
-		.Padding(FMargin(0.f))
-		[
 			SNew( SHorizontalBox )
 
 			+ SHorizontalBox::Slot()
@@ -235,8 +229,7 @@ void SOutlinerItemViewBase::Construct(
 			.AutoWidth()
 			[
 				InArgs._RightGutterContent.Widget
-			]
-		];
+			];
 
 	ChildSlot
 	[
@@ -397,51 +390,6 @@ const FSlateBrush* SOutlinerItemViewBase::GetNodeBorderImage() const
 	return bHasChildren ? ExpandedBackgroundBrush : CollapsedBackgroundBrush;
 }
 
-FSlateColor SOutlinerItemViewBase::GetNodeBackgroundTint() const
-{
-	TSharedPtr<FEditorViewModel> Editor = WeakEditor.Pin();
-	TViewModelPtr<IOutlinerExtension> OutlinerItem = WeakOutlinerExtension.Pin();
-
-	if (!Editor || !OutlinerItem)
-	{
-		return FLinearColor(0.f,0.f,0.f,0.f);
-	}
-
-	EOutlinerSelectionState SelectionState = OutlinerItem->GetSelectionState();
-
-	if (EnumHasAnyFlags(SelectionState, EOutlinerSelectionState::SelectedDirectly))
-	{
-		return FStyleColors::Select;
-	}
-	if (EnumHasAnyFlags(SelectionState, EOutlinerSelectionState::HasSelectedKeys | EOutlinerSelectionState::HasSelectedTrackAreaItems))
-	{
-		return FStyleColors::Header;
-	}
-
-	// If this is collapsed but has any children with selected keys or sections, we report that state on the parent
-	if (!OutlinerItem->IsExpanded())
-	{
-		for (TViewModelPtr<IOutlinerExtension> Child : OutlinerItem.AsModel()->GetDescendantsOfType<IOutlinerExtension>())
-		{
-			if (EnumHasAnyFlags(SelectionState, EOutlinerSelectionState::HasSelectedKeys | EOutlinerSelectionState::HasSelectedTrackAreaItems))
-			{
-				return FStyleColors::Header;
-			}
-		}
-	}
-
-	if (Editor->GetOutliner()->GetHoveredItem() == OutlinerItem)
-	{
-		return ItemStyle == EOutlinerItemViewBaseStyle::ContainerHeader
-			? FLinearColor(FColor(52, 52, 52, 255))
-			: FLinearColor(FColor(72, 72, 72, 255));
-	}
-
-	return ItemStyle == EOutlinerItemViewBaseStyle::ContainerHeader
-		? FLinearColor(FColor(48, 48, 48, 255))
-		: FLinearColor(FColor(62, 62, 62, 255));
-}
-
 FSlateColor SOutlinerItemViewBase::GetNodeInnerBackgroundTint() const
 {
 	TSharedPtr<FEditorViewModel> Editor = WeakEditor.Pin();
diff --git a/Engine/Source/Editor/SequencerCore/Public/MVVM/Extensions/DynamicExtensionContainer.h b/Engine/Source/Editor/SequencerCore/Public/MVVM/Extensions/DynamicExtensionContainer.h
--- a/Engine/Source/Editor/SequencerCore/Public/MVVM/Extensions/DynamicExtensionContainer.h
+++ b/Engine/Source/Editor/SequencerCore/Public/MVVM/Extensions/DynamicExtensionContainer.h
@@ -19,6 +19,14 @@ struct FCastableTypeTable;
 class FDynamicExtensionContainer;
 class FViewModel;
 
+struct FDynamicExtensionContainerIterator;
+
+template<typename T>
+struct TDynamicExtensionContainerIterator;
+
+template<typename T>
+struct TDynamicExtensionContainerIteratorProxy;
+
 /**
  * Base class for dynamic extensions that can be added to sequence data models and participate in dynamic casting.
  */
@@ -40,6 +48,7 @@ public:
 
 DECLARE_DELEGATE_OneParam(FDynamicExtensionCallback, TSharedRef<FViewModel>);
 
+
 /**
  * Base class for supporting dynamic extensions that participate in dynamic casting.
  */
@@ -84,6 +93,12 @@ public:
 		return const_cast<T*>(const_cast<const FDynamicExtensionContainer*>(this)->CastDynamicChecked<T>());
 	}
 
+	template<typename T>
+	TDynamicExtensionContainerIteratorProxy<T> FilterDynamic()
+	{
+		return TDynamicExtensionContainerIteratorProxy<T>(this);
+	}
+
 	SEQUENCERCORE_API const void* CastDynamic(FViewModelTypeID Type) const;
 
 	const void* CastDynamicChecked(FViewModelTypeID Type) const
@@ -140,6 +155,9 @@ protected:
 
 private:
 
+	friend struct FDynamicExtensionContainerIterator;
+	template<typename T> friend struct TDynamicExtensionContainerIteratorProxy;
+
 	struct FDynamicExtensionInfo
 	{
 		FCastableTypeTable* TypeTable;
@@ -149,6 +167,64 @@ private:
 	TArray<FDynamicExtensionInfo> DynamicExtensions;
 };
 
+struct FDynamicExtensionContainerIterator
+{
+	explicit operator bool() const
+	{
+		return (bool)Iterator;
+	}
+
+	SEQUENCERCORE_API FDynamicExtensionContainerIterator& operator++();
+
+	SEQUENCERCORE_API friend bool operator!=(const FDynamicExtensionContainerIterator& A, const FDynamicExtensionContainerIterator& B);
+
+protected:
+	using IteratorType = TArray<FDynamicExtensionContainer::FDynamicExtensionInfo>::TConstIterator;
+
+	friend FDynamicExtensionContainer;
+
+	SEQUENCERCORE_API explicit FDynamicExtensionContainerIterator(IteratorType&& InIterator, FViewModelTypeID InType);
+
+	void* CurrentExtension = nullptr;
+	IteratorType Iterator;
+	FViewModelTypeID Type;
+};
+
+template<typename T>
+struct TDynamicExtensionContainerIteratorProxy
+{
+	FORCEINLINE TDynamicExtensionContainerIterator<T> begin() { return TDynamicExtensionContainerIterator<T>(IteratorType(Container->DynamicExtensions)); }
+	FORCEINLINE TDynamicExtensionContainerIterator<T> end()   { return TDynamicExtensionContainerIterator<T>(IteratorType(Container->DynamicExtensions, Container->DynamicExtensions.Num())); }
+
+private:
+	friend FDynamicExtensionContainer;
+
+	using IteratorType = TArray<FDynamicExtensionContainer::FDynamicExtensionInfo>::TConstIterator;
+
+	TDynamicExtensionContainerIteratorProxy(FDynamicExtensionContainer* InContainer)
+		: Container(InContainer)
+	{}
+
+	FDynamicExtensionContainer* Container;
+};
+
+template<typename T>
+struct TDynamicExtensionContainerIterator : FDynamicExtensionContainerIterator
+{
+	T& operator*() const
+	{
+		check(CurrentExtension);
+		return *static_cast<T*>(CurrentExtension);
+	}
+
+protected:
+	friend TDynamicExtensionContainerIteratorProxy<T>;
+
+	explicit TDynamicExtensionContainerIterator(FDynamicExtensionContainerIterator::IteratorType InIterator)
+		: FDynamicExtensionContainerIterator(MoveTemp(InIterator), T::ID)
+	{}
+};
+
 } // namespace Sequencer
 } // namespace UE
 
diff --git a/Engine/Source/Editor/SequencerCore/Public/MVVM/Extensions/HierarchicalCacheExtension.h b/Engine/Source/Editor/SequencerCore/Public/MVVM/Extensions/HierarchicalCacheExtension.h
new file mode 100644
--- /dev/null
+++ b/Engine/Source/Editor/SequencerCore/Public/MVVM/Extensions/HierarchicalCacheExtension.h
@@ -0,0 +1,143 @@
+// Copyright Epic Games, Inc. All Rights Reserved.
+
+#pragma once
+
+#include "SequencerCoreFwd.h"
+#include "MVVM/ViewModelTypeID.h"
+#include "Delegates/DelegateCombinations.h"
+#include "MVVM/Extensions/DynamicExtensionContainer.h"
+#include "MVVM/ViewModelPtr.h"
+
+struct FGuid;
+
+namespace UE::Sequencer
+{
+
+class IHierarchicalCache;
+
+DECLARE_MULTICAST_DELEGATE_OneParam(FPreUpdateCachesEvent, FViewModelPtr);
+
+class SEQUENCERCORE_API IHierarchicalCache
+{
+public:
+
+	UE_SEQUENCER_DECLARE_VIEW_MODEL_TYPE_ID(IHierarchicalCache)
+
+	virtual ~IHierarchicalCache(){}
+
+	virtual void BeginUpdate() = 0;
+	virtual void PreVisitChildren(const FViewModelPtr& ViewModel) = 0;
+	virtual void PostVisitChildren(const FViewModelPtr& ViewModel) = 0;
+	virtual void EndUpdate() = 0;
+};
+
+
+class SEQUENCERCORE_API FHierarchicalCacheExtension
+	: public IDynamicExtension
+{
+public:
+
+	~FHierarchicalCacheExtension();
+
+	FPreUpdateCachesEvent PreUpdateCachesEvent;
+
+	void Initialize(const FViewModelPtr& InRootModel);
+
+	void OnCreated(TSharedRef<FViewModel> InWeakOwner) override;
+
+	void OnHierarchyUpdated();
+
+	void UpdateCachedFlags();
+	void UpdateAllCachedFlags(const FViewModelPtr& ViewModel, TArrayView<IHierarchicalCache* const> HierarchicalCaches);
+
+private:
+
+	void UpdateCachedFlagsForModel(const FViewModelPtr& ViewModel, TArrayView<IHierarchicalCache* const> HierarchicalCaches);
+
+protected:
+
+	FWeakViewModelPtr WeakOwnerModel;
+	FWeakViewModelPtr WeakRootModel;
+	EViewModelListType ModelListFilter;
+};
+
+
+
+template<typename FlagsType>
+class TFlagStateCacheExtension
+	: public IDynamicExtension
+	, public IHierarchicalCache
+{
+public:
+
+	using Implements = TImplements<IDynamicExtension, IHierarchicalCache>;
+
+	virtual ~TFlagStateCacheExtension() {}
+
+	FlagsType GetRootFlags() const
+	{
+		return RootFlags;
+	}
+	FlagsType GetCachedFlags(const FViewModelPtr& ViewModel) const
+	{
+		return GetCachedFlags(ViewModel->GetModelID());
+	}
+	FlagsType GetCachedFlags(uint32 InModelID) const
+	{
+		return CachedFlagsFromNodeID.FindRef(InModelID);
+	}
+
+private:
+
+	virtual FlagsType ComputeFlagsForModel(const FViewModelPtr& ViewModel) = 0;
+	virtual void PostComputeChildrenFlags(const FViewModelPtr& ViewModel, FlagsType& OutThisModelFlags, FlagsType& OutPropagateToParentFlags) {}
+
+protected:
+
+	void BeginUpdate() override
+	{
+		CachedFlagsFromNodeID.Empty();
+		IndividualItemFlags.Emplace(FlagsType::None);
+		AccumulatedChildFlags.Emplace(FlagsType::None);
+	}
+	void PreVisitChildren(const FViewModelPtr& ViewModel) override
+	{
+		FlagsType ThisModelFlags = ComputeFlagsForModel(ViewModel);
+		IndividualItemFlags.Emplace(ThisModelFlags);
+		AccumulatedChildFlags.Emplace(FlagsType::None);
+	}
+	void PostVisitChildren(const FViewModelPtr& ViewModel) override
+	{
+		FlagsType ThisModelFlags              = IndividualItemFlags.Pop();
+		FlagsType FlagsPropagatedFromChildren = AccumulatedChildFlags.Pop();
+
+		ThisModelFlags |= FlagsPropagatedFromChildren;
+		PostComputeChildrenFlags(ViewModel, ThisModelFlags, AccumulatedChildFlags.Last());
+
+		AccumulatedChildFlags.Last() |= (ThisModelFlags & FlagsType::InheritedFromChildren);
+
+		if (ThisModelFlags != FlagsType::None)
+		{
+			CachedFlagsFromNodeID.Add(ViewModel->GetModelID(), ThisModelFlags);
+		}
+	}
+	void EndUpdate() override
+	{
+		RootFlags = AccumulatedChildFlags.Last();
+		AccumulatedChildFlags.Empty();
+		IndividualItemFlags.Empty();
+
+		CachedFlagsFromNodeID.Shrink();
+	}
+
+protected:
+
+	TArray<FlagsType> AccumulatedChildFlags;
+	TArray<FlagsType> IndividualItemFlags;
+	TMap<uint32, FlagsType> CachedFlagsFromNodeID;
+	FlagsType RootFlags;
+};
+
+
+} // namespace UE::Sequencer
+
diff --git a/Engine/Source/Editor/SequencerCore/Public/MVVM/Extensions/ILockableExtension.h b/Engine/Source/Editor/SequencerCore/Public/MVVM/Extensions/ILockableExtension.h
new file mode 100644
--- /dev/null
+++ b/Engine/Source/Editor/SequencerCore/Public/MVVM/Extensions/ILockableExtension.h
@@ -0,0 +1,68 @@
+// Copyright Epic Games, Inc. All Rights Reserved.
+
+#pragma once
+
+#include "CoreMinimal.h"
+#include "MVVM/ViewModelTypeID.h"
+#include "MVVM/Extensions/HierarchicalCacheExtension.h"
+
+namespace UE
+{
+namespace Sequencer
+{
+
+enum class ELockableLockState
+{
+	None,
+	Locked,
+	PartiallyLocked,
+};
+
+/**
+ * An extension for models that can be locked
+ */
+class SEQUENCERCORE_API ILockableExtension
+{
+public:
+
+	UE_SEQUENCER_DECLARE_VIEW_MODEL_TYPE_ID(ILockableExtension)
+
+	virtual ~ILockableExtension(){}
+
+	/** Returns whether this item is locked */
+	virtual ELockableLockState GetLockState() const = 0;
+
+	/** Set whether this item is locked */
+	virtual void SetIsLocked(bool bIsLocked) = 0;
+};
+
+enum class ECachedLockState
+{
+	None                     = 0,
+
+	Lockable                 = 1 << 0,
+	LockableChildren         = 1 << 1,
+	Locked                   = 1 << 2,
+	PartiallyLockedChildren  = 1 << 3,
+	ImplicitlyLockedByParent = 1 << 4,
+
+	InheritedFromChildren = LockableChildren | PartiallyLockedChildren,
+};
+ENUM_CLASS_FLAGS(ECachedLockState)
+
+class SEQUENCERCORE_API FLockStateCacheExtension
+	: public TFlagStateCacheExtension<ECachedLockState>
+{
+public:
+
+	UE_SEQUENCER_DECLARE_VIEW_MODEL_TYPE_ID(FLockStateCacheExtension);
+
+private:
+
+	ECachedLockState ComputeFlagsForModel(const FViewModelPtr& ViewModel) override;
+	void PostComputeChildrenFlags(const FViewModelPtr& ViewModel, ECachedLockState& OutThisModelFlags, ECachedLockState& OutPropagateToParentFlags) override;
+};
+
+} // namespace Sequencer
+} // namespace UE
+
diff --git a/Engine/Source/Editor/SequencerCore/Public/MVVM/Extensions/IMutableExtension.h b/Engine/Source/Editor/SequencerCore/Public/MVVM/Extensions/IMutableExtension.h
--- a/Engine/Source/Editor/SequencerCore/Public/MVVM/Extensions/IMutableExtension.h
+++ b/Engine/Source/Editor/SequencerCore/Public/MVVM/Extensions/IMutableExtension.h
@@ -4,6 +4,7 @@
 
 #include "CoreMinimal.h"
 #include "MVVM/ViewModelTypeID.h"
+#include "MVVM/Extensions/HierarchicalCacheExtension.h"
 
 namespace UE
 {
@@ -21,10 +22,41 @@ public:
 
 	virtual ~IMutableExtension(){}
 
-	/** Returns whether this item is muted, or is in a muted branch */
+	/** Returns whether this item is muted */
 	virtual bool IsMuted() const = 0;
+
+	/** Set this item's mute state */
+	virtual void SetIsMuted(bool bIsMuted) = 0;
+};
+
+enum class ECachedMuteState
+{
+	None                     = 0,
+
+	Mutable                  = 1 << 0,
+	MutableChildren          = 1 << 1,
+	Muted                    = 1 << 2,
+	PartiallyMutedChildren   = 1 << 3,
+	ImplicitlyMutedByParent  = 1 << 4,
+
+	InheritedFromChildren = MutableChildren | PartiallyMutedChildren,
+};
+ENUM_CLASS_FLAGS(ECachedMuteState)
+
+class SEQUENCERCORE_API FMuteStateCacheExtension
+	: public TFlagStateCacheExtension<ECachedMuteState>
+{
+public:
+
+	UE_SEQUENCER_DECLARE_VIEW_MODEL_TYPE_ID(FMuteStateCacheExtension);
+
+private:
+
+	ECachedMuteState ComputeFlagsForModel(const FViewModelPtr& ViewModel) override;
+	void PostComputeChildrenFlags(const FViewModelPtr& ViewModel, ECachedMuteState& OutThisModelFlags, ECachedMuteState& OutPropagateToParentFlags) override;
 };
 
+
 } // namespace Sequencer
 } // namespace UE
 
diff --git a/Engine/Source/Editor/SequencerCore/Public/MVVM/Extensions/IOutlinerExtension.h b/Engine/Source/Editor/SequencerCore/Public/MVVM/Extensions/IOutlinerExtension.h
--- a/Engine/Source/Editor/SequencerCore/Public/MVVM/Extensions/IOutlinerExtension.h
+++ b/Engine/Source/Editor/SequencerCore/Public/MVVM/Extensions/IOutlinerExtension.h
@@ -41,6 +41,9 @@ enum class EOutlinerSelectionState
 	SelectedDirectly          = 1 << 0,
 	HasSelectedKeys           = 1 << 1,
 	HasSelectedTrackAreaItems = 1 << 2,
+
+	DescendentHasSelectedKeys = 1 << 3,
+	DescendentHasSelectedTrackAreaItems = 1 << 4,
 };
 ENUM_CLASS_FLAGS(EOutlinerSelectionState);
 
diff --git a/Engine/Source/Editor/SequencerCore/Public/MVVM/Extensions/ISoloableExtension.h b/Engine/Source/Editor/SequencerCore/Public/MVVM/Extensions/ISoloableExtension.h
--- a/Engine/Source/Editor/SequencerCore/Public/MVVM/Extensions/ISoloableExtension.h
+++ b/Engine/Source/Editor/SequencerCore/Public/MVVM/Extensions/ISoloableExtension.h
@@ -4,6 +4,7 @@
 
 #include "Templates/SharedPointer.h"
 #include "MVVM/ViewModelTypeID.h"
+#include "MVVM/Extensions/HierarchicalCacheExtension.h"
 
 namespace UE
 {
@@ -21,8 +22,42 @@ public:
 
 	virtual ~ISoloableExtension(){}
 
-	/** Returns whether this item is solo, or in a solo branch */
+	/** Returns whether this item is solo */
 	virtual bool IsSolo() const = 0;
+
+	/** Enable or disable solo for this object */
+	virtual void SetIsSoloed(bool bIsSoloed) = 0;
+};
+
+enum class ECachedSoloState
+{
+	None                     = 0,
+
+	Soloable                 = 1 << 0,
+	SoloableChildren         = 1 << 1,
+	Soloed                   = 1 << 2,
+	PartiallySoloedChildren  = 1 << 3,
+	ImplicitlySoloedByParent = 1 << 4,
+
+	InheritedFromChildren = SoloableChildren | PartiallySoloedChildren,
+};
+ENUM_CLASS_FLAGS(ECachedSoloState)
+
+class SEQUENCERCORE_API FSoloStateCacheExtension
+	: public TFlagStateCacheExtension<ECachedSoloState>
+{
+public:
+
+	UE_SEQUENCER_DECLARE_VIEW_MODEL_TYPE_ID(FSoloStateCacheExtension)
+
+	using Implements = TImplements<IDynamicExtension, IHierarchicalCache>;
+
+	virtual ~FSoloStateCacheExtension() {}
+
+private:
+
+	ECachedSoloState ComputeFlagsForModel(const FViewModelPtr& ViewModel) override;
+	void PostComputeChildrenFlags(const FViewModelPtr& ViewModel, ECachedSoloState& OutThisModelFlags, ECachedSoloState& OutPropagateToParentFlags) override;
 };
 
 } // namespace Sequencer
diff --git a/Engine/Source/Editor/SequencerCore/Public/MVVM/SharedViewModelData.h b/Engine/Source/Editor/SequencerCore/Public/MVVM/SharedViewModelData.h
--- a/Engine/Source/Editor/SequencerCore/Public/MVVM/SharedViewModelData.h
+++ b/Engine/Source/Editor/SequencerCore/Public/MVVM/SharedViewModelData.h
@@ -6,6 +6,7 @@
 #include "CoreMinimal.h"
 #include "Delegates/Delegate.h"
 #include "MVVM/ViewModels/ViewModelHierarchy.h"
+#include "MVVM/ViewModels/ViewModel.h"
 #include "SequencerCoreFwd.h"
 #include "Templates/SharedPointer.h"
 #include "Templates/UniquePtr.h"
@@ -19,7 +20,8 @@ namespace Sequencer
 
 class FViewModel;
 
-class SEQUENCERCORE_API FSharedViewModelData : public TSharedFromThis<FSharedViewModelData>
+class SEQUENCERCORE_API FSharedViewModelData
+	: public FViewModel
 {
 public:
 	void PreHierarchicalChange(const TSharedPtr<FViewModel>& InChangedModel);
diff --git a/Engine/Source/Editor/SequencerCore/Public/MVVM/ViewModels/ViewModel.h b/Engine/Source/Editor/SequencerCore/Public/MVVM/ViewModels/ViewModel.h
--- a/Engine/Source/Editor/SequencerCore/Public/MVVM/ViewModels/ViewModel.h
+++ b/Engine/Source/Editor/SequencerCore/Public/MVVM/ViewModels/ViewModel.h
@@ -131,6 +131,13 @@ public:
 	/** Returns the root model of the hierarchy, i.e. the parent model that has no parent */
 	FViewModelPtr GetRoot() const;
 
+
+	/** Get this model's unique, non-persistent and non-deterministic ID */
+	uint32 GetModelID() const
+	{
+		return ModelID;
+	}
+
 public:
 
 	bool IsConstructed() const;
@@ -279,6 +286,9 @@ private:
 	 */
 	int32 ActiveIterationCount;
 
+	/** Non-deterministic, non-persistent, serially increasing model ID unique to this instance */
+	uint32 ModelID;
+
 	/**
 	 * Flag to track whether this view model needs to be constructed
 	 */
diff --git a/Engine/Source/Editor/SequencerCore/Public/MVVM/ViewModels/ViewModelHierarchy.h b/Engine/Source/Editor/SequencerCore/Public/MVVM/ViewModels/ViewModelHierarchy.h
--- a/Engine/Source/Editor/SequencerCore/Public/MVVM/ViewModels/ViewModelHierarchy.h
+++ b/Engine/Source/Editor/SequencerCore/Public/MVVM/ViewModels/ViewModelHierarchy.h
@@ -290,8 +290,8 @@ private:
 class SEQUENCERCORE_API FViewModelHierarchyOperation
 {
 public:
-	FViewModelHierarchyOperation(const TSharedPtr<FViewModel>& InAnyModel);
-	FViewModelHierarchyOperation(const TSharedPtr<FSharedViewModelData>& InSharedData);
+	explicit FViewModelHierarchyOperation(const TSharedPtr<FViewModel>& InAnyModel);
+	explicit FViewModelHierarchyOperation(const TSharedRef<FSharedViewModelData>& InSharedData);
 	~FViewModelHierarchyOperation();
 
 	/**
@@ -301,6 +301,8 @@ public:
 
 private:
 
+	void Construct();
+
 	struct FCachedHierarchicalPosition
 	{
 		TWeakPtr<FViewModel> PreviousParent;
diff --git a/Engine/Source/Editor/SequencerCore/Public/MVVM/Views/SOutlinerColumnBorder.h b/Engine/Source/Editor/SequencerCore/Public/MVVM/Views/SOutlinerColumnBorder.h
--- a/Engine/Source/Editor/SequencerCore/Public/MVVM/Views/SOutlinerColumnBorder.h
+++ b/Engine/Source/Editor/SequencerCore/Public/MVVM/Views/SOutlinerColumnBorder.h
@@ -36,25 +36,12 @@ private:
 	void OnMouseEnter(const FGeometry& MyGeometry, const FPointerEvent& MouseEvent) override;
 	void OnMouseLeave(const FPointerEvent& MouseEvent) override;
 
-	/**
-	* @return The border image to show in the tree node.
-	*/
-	const FSlateBrush* GetBorderImage() const;
-
-	/**
-	* @return The inner background border color and opacity, darker when mouse is directly over inner widget.
-	*/
-	FSlateColor GetBackgroundTint() const;
-
 	/** Used to determine color if outliner item is directly selected */
 	TWeakViewModelPtr<IOutlinerExtension> WeakOutlinerExtension;
 
 	/** Used to determine color if outliner item is being hovered */
 	TWeakPtr<FEditorViewModel> WeakEditor;
 
-	/** Brush used for inner border */
-	const FSlateBrush* BackgroundBrush;
-
 	/** Is the mouse directly over this widget */
 	bool bIsMouseOverInnerBorder = false;
 };
@@ -94,24 +81,11 @@ private:
 	/** Handle mouse down events in border so that outliner item events are not generated by the track gutter. */
 	virtual FReply OnMouseButtonDown(const FGeometry& MyGeometry, const FPointerEvent& MouseEvent) override;
 
-	/**
-	* @return The border image to show in the tree node.
-	*/
-	const FSlateBrush* GetBorderImage() const;
-
-	/**
-	* @return The tint to apply to the border image
-	*/
-	FSlateColor GetBackgroundTint() const;
-
 	/** Used to determine color if outliner item is directly selected */
 	TWeakViewModelPtr<IOutlinerExtension> WeakOutlinerExtension;
 
 	/** Used to determine color if outliner item is being hovered */
 	TWeakPtr<FEditorViewModel> WeakEditor;
-
-	/** Brush used for border */
-	const FSlateBrush* BackgroundBrush;
 };
 
 } // namespace UE::Sequencer
diff --git a/Engine/Source/Editor/SequencerCore/Public/MVVM/Views/SOutlinerItemViewBase.h b/Engine/Source/Editor/SequencerCore/Public/MVVM/Views/SOutlinerItemViewBase.h
--- a/Engine/Source/Editor/SequencerCore/Public/MVVM/Views/SOutlinerItemViewBase.h
+++ b/Engine/Source/Editor/SequencerCore/Public/MVVM/Views/SOutlinerItemViewBase.h
@@ -138,11 +138,6 @@ private:
 	 * @return The border image to show in the tree node.
 	 */
 	const FSlateBrush* GetNodeBorderImage() const;
-	
-	/**
-	 * @return The tint to apply to the border image
-	 */
-	FSlateColor GetNodeBackgroundTint() const;
 
 	/**
 	* @return The tint to apply to the border image for the inner portion of the node.
diff --git a/Engine/Source/Editor/SequencerCore/Public/MVVM/Views/SOutlinerView.h b/Engine/Source/Editor/SequencerCore/Public/MVVM/Views/SOutlinerView.h
--- a/Engine/Source/Editor/SequencerCore/Public/MVVM/Views/SOutlinerView.h
+++ b/Engine/Source/Editor/SequencerCore/Public/MVVM/Views/SOutlinerView.h
@@ -145,7 +145,7 @@ protected:
 	void ExpandCollapseNode(TViewModelPtr<IOutlinerExtension> InDataModel, bool bExpansionState, ETreeRecursion Recursion);
 
 	/** Generate a row for a particular node */
-	TSharedRef<ITableRow> OnGenerateRow(TWeakViewModelPtr<IOutlinerExtension> InDisplayNode, const TSharedRef<STableViewBase>& OwnerTable);
+	virtual TSharedRef<ITableRow> OnGenerateRow(TWeakViewModelPtr<IOutlinerExtension> InDisplayNode, const TSharedRef<STableViewBase>& OwnerTable);
 
 	void CreateTrackLanesForRow(TSharedRef<SOutlinerViewRow> InRow, TViewModelPtr<IOutlinerExtension> InDataModel);
 	TSharedPtr<STrackLane> FindOrCreateParentLane(TViewModelPtr<IOutlinerExtension> InDataModel);
@@ -175,7 +175,7 @@ public:
 
 	virtual FReply OnKeyDown(const FGeometry& MyGeometry, const FKeyEvent& InKeyEvent) override;
 
-private:
+protected:
 
 	// Private, unimplemented overloaded name for SetItemSelection to prevent external calls - use ForceSetItemSelection instead
 	void SetItemSelection();
@@ -248,7 +248,7 @@ protected:
 
 	FString OnItemToString_Debug(TWeakViewModelPtr<IOutlinerExtension> InWeakModel);
 
-private:
+protected:
 
 	/** The tree view's header row (hidden) */
 	TSharedPtr<SHeaderRow> HeaderRow;
@@ -359,9 +359,10 @@ public:
 	FReply OnAcceptDrop( const FDragDropEvent& DragDropEvent, EItemDropZone ItemDropZone, TWeakViewModelPtr<IOutlinerExtension> InDataModel);
 
 protected:
+
 	FVector2D ComputeDesiredSize(float LayoutScaleMultiplier) const override;
 
-private:
+protected:
 
 	/**
 	 * Cached reference to a track lane that we relate to.
