Fixed loc collision

diff --git a/Engine/Source/Editor/Sequencer/Private/Widgets/OutlinerColumns/SSoloColumnWidget.cpp b/Engine/Source/Editor/Sequencer/Private/Widgets/OutlinerColumns/SSoloColumnWidget.cpp
--- a/Engine/Source/Editor/Sequencer/Private/Widgets/OutlinerColumns/SSoloColumnWidget.cpp
+++ b/Engine/Source/Editor/Sequencer/Private/Widgets/OutlinerColumns/SSoloColumnWidget.cpp
@@ -53,7 +53,7 @@ void SSoloColumnWidget::SetIsActive(const bool bInIsActive)
 		return;
 	}
 
-	const FScopedTransaction Transaction(NSLOCTEXT("Sequencer", "SetNodeMuted", "Set Node Soloed"));
+	const FScopedTransaction Transaction(NSLOCTEXT("Sequencer", "SetNodeSoloed", "Set Node Soloed"));
 
 	if (OutlinerItem->GetSelectionState() == EOutlinerSelectionState::SelectedDirectly)
 	{
