Sequencer: Separate out clear and sort all nodes and descendants so that custom sort orders are retained. Custom sort orders will be removed when  the user invokes Sort All Nodes

diff --git a/Engine/Source/Editor/Sequencer/Private/Sequencer.cpp b/Engine/Source/Editor/Sequencer/Private/Sequencer.cpp
--- a/Engine/Source/Editor/Sequencer/Private/Sequencer.cpp
+++ b/Engine/Source/Editor/Sequencer/Private/Sequencer.cpp
@@ -9477,6 +9477,7 @@ FText FSequencer::GetNavigateBackwardTooltip() const
 void FSequencer::SortAllNodesAndDescendants()
 {
 	FScopedTransaction SortAllNodesTransaction(NSLOCTEXT("Sequencer", "SortAllNodes_Transaction", "Sort Tracks"));
+	NodeTree->ClearCustomSortOrders();
 	NodeTree->SortAllNodesAndDescendants();
 }
 
diff --git a/Engine/Source/Editor/Sequencer/Private/SequencerNodeTree.cpp b/Engine/Source/Editor/Sequencer/Private/SequencerNodeTree.cpp
--- a/Engine/Source/Editor/Sequencer/Private/SequencerNodeTree.cpp
+++ b/Engine/Source/Editor/Sequencer/Private/SequencerNodeTree.cpp
@@ -243,6 +243,20 @@ TArray<UE::Sequencer::TViewModelPtr<UE::Sequencer::IOutlinerExtension>> FSequenc
 	return RootNodes;
 }
 
+void FSequencerNodeTree::ClearCustomSortOrders()
+{
+	using namespace UE::Sequencer;
+
+	const bool bIncludeRootNode = true;
+	for (TSharedPtr<FViewModel> Child : RootNode->GetDescendants(bIncludeRootNode))
+	{
+		if (ISortableExtension* SortableExtension = Child->CastThis<ISortableExtension>())
+		{
+			SortableExtension->SetCustomOrder(-1);
+		}
+	}
+}
+
 void FSequencerNodeTree::SortAllNodesAndDescendants()
 {
 	using namespace UE::Sequencer;
@@ -253,7 +267,6 @@ void FSequencerNodeTree::SortAllNodesAndDescendants()
 	{
 		if (ISortableExtension* SortableExtension = Child->CastThis<ISortableExtension>())
 		{
-			SortableExtension->SetCustomOrder(-1);
 			SortableChildren.Add(SortableExtension);
 		}
 	}
@@ -261,10 +274,6 @@ void FSequencerNodeTree::SortAllNodesAndDescendants()
 	{
 		SortableChild->SortChildren();
 	}
-
-	// Refresh the tree so that our changes are visible.
-	// @todo: Is this necessary any more?
-	//GetSequencer().RefreshTree();
 }
 
 void FSequencerNodeTree::AddFilter(TSharedPtr<FSequencerTrackFilter> TrackFilter)
diff --git a/Engine/Source/Editor/Sequencer/Private/SequencerNodeTree.h b/Engine/Source/Editor/Sequencer/Private/SequencerNodeTree.h
--- a/Engine/Source/Editor/Sequencer/Private/SequencerNodeTree.h
+++ b/Engine/Source/Editor/Sequencer/Private/SequencerNodeTree.h
@@ -156,7 +156,10 @@ public:
 	 */
 	FOnUpdated& OnUpdated() { return OnUpdatedDelegate; }
 
-	/** Sorts all nodes and their descendants by category then alphabetically.*/
+	/** Clears all custom sort orders so that sorting will be by category then alphabetically. */
+	void ClearCustomSortOrders();
+
+	/** Sorts all nodes and their descendants*/
 	void SortAllNodesAndDescendants();
 
 	/**
