Sequencer: Refresh tree when toggling mute/solo from the menu

Follow up fix to 27669817

#rnx

diff --git a/Engine/Source/Editor/Sequencer/Private/MVVM/ViewModels/OutlinerItemModel.cpp b/Engine/Source/Editor/Sequencer/Private/MVVM/ViewModels/OutlinerItemModel.cpp
--- a/Engine/Source/Editor/Sequencer/Private/MVVM/ViewModels/OutlinerItemModel.cpp
+++ b/Engine/Source/Editor/Sequencer/Private/MVVM/ViewModels/OutlinerItemModel.cpp
@@ -330,10 +330,17 @@ void FOutlinerItemModelMixin::ToggleSelectedModelsSolo()
 
 	const FScopedTransaction Transaction(NSLOCTEXT("Sequencer", "ToggleSolo", "Toggle Solo"));
 
-	for (TViewModelPtr<ISoloableExtension> Soloable : GetEditor()->GetSelection()->Outliner.Filter<ISoloableExtension>())
+	TSharedPtr<FSequencerEditorViewModel> EditorViewModel = GetEditor();
+	for (TViewModelPtr<ISoloableExtension> Soloable : EditorViewModel->GetSelection()->Outliner.Filter<ISoloableExtension>())
 	{
 		Soloable->SetIsSoloed(bNewSoloState);
 	}
+
+	TSharedPtr<ISequencer> Sequencer = EditorViewModel->GetSequencer();
+	if (Sequencer)
+	{
+		Sequencer->RefreshTree();
+	}
 }
 
 ECheckBoxState FOutlinerItemModelMixin::SelectedModelsMuteState() const
@@ -366,10 +373,17 @@ void FOutlinerItemModelMixin::ToggleSelectedModelsMuted()
 
 	const FScopedTransaction Transaction(NSLOCTEXT("Sequencer", "ToggleMute", "Toggle Mute"));
 
-	for (TViewModelPtr<IMutableExtension> Muteable : GetEditor()->GetSelection()->Outliner.Filter<IMutableExtension>())
+	TSharedPtr<FSequencerEditorViewModel> EditorViewModel = GetEditor();
+	for (TViewModelPtr<IMutableExtension> Muteable : EditorViewModel->GetSelection()->Outliner.Filter<IMutableExtension>())
 	{
 		Muteable->SetIsMuted(bNewMuteState);
 	}
+
+	TSharedPtr<ISequencer> Sequencer = EditorViewModel->GetSequencer();
+	if (Sequencer)
+	{
+		Sequencer->RefreshTree();
+	}
 }
 
 TSharedPtr<SWidget> FOutlinerItemModelMixin::CreateContextMenuWidget(const FCreateOutlinerContextMenuWidgetParams& InParams)
