Sequencer: Check whether the section is keyable (don't allow keying muted/disabled tracks or track rows)

diff --git a/Engine/Source/Editor/Sequencer/Private/SequencerKeyParams.cpp b/Engine/Source/Editor/Sequencer/Private/SequencerKeyParams.cpp
--- a/Engine/Source/Editor/Sequencer/Private/SequencerKeyParams.cpp
+++ b/Engine/Source/Editor/Sequencer/Private/SequencerKeyParams.cpp
@@ -88,6 +88,11 @@ void FKeyOperation::Populate(UMovieSceneTrack* InTrack, TSharedPtr<ISequencerSec
 		return;
 	}
 
+	if (!MovieSceneHelpers::IsSectionKeyable(SectionObject))
+	{
+		return;
+	}
+
 	FSectionCandidates&   Candidates        = CandidatesByTrack.FindOrAdd(InTrack);
 	FKeySectionOperation* ExistingOperation = Algo::FindBy(Candidates.Operations, InSection, &FKeySectionOperation::Section);
 	if (ExistingOperation)
