Sequencer: Add commands for Clear Global Marked Frames and Show Marked Frames Globally

diff --git a/Engine/Source/Editor/Sequencer/Private/SSequencer.cpp b/Engine/Source/Editor/Sequencer/Private/SSequencer.cpp
--- a/Engine/Source/Editor/Sequencer/Private/SSequencer.cpp
+++ b/Engine/Source/Editor/Sequencer/Private/SSequencer.cpp
@@ -2214,44 +2214,8 @@ TSharedRef<SWidget> SSequencer::MakeViewMenu()
 
 	MenuBuilder.BeginSection("MarkedFrames", LOCTEXT("MarkedFramesHeader", "Marked Frames"));
 	{
-		UMovieSceneSequence* FocusedMovieSequence = Sequencer->GetFocusedMovieSceneSequence();
-		UMovieScene* FocusedMovieScene = nullptr;
-		if (FocusedMovieSequence != nullptr)
-		{
-			FocusedMovieScene = FocusedMovieSequence->GetMovieScene();
-		}
-
-		MenuBuilder.AddMenuEntry(
-			LOCTEXT("ToggleShowMarkedFrameGlobal", "Show Marked Frames Globally"),
-			LOCTEXT("ToggleShowMarkedFrameGlobalTooltip", "Makes marked frames in this sub-sequence visible in parent/sibling sequences for the current editor session"),
-			FSlateIcon(),
-			FUIAction(
-				FExecuteAction::CreateLambda([FocusedMovieScene, this]() {
-					if (FocusedMovieScene != nullptr)
-					{
-						const FScopedTransaction Transaction(LOCTEXT("ToggleShowMarkedFramesGlobally", "Toggle Show Marked Frames Globally"));
-						FocusedMovieScene->Modify();
-						FocusedMovieScene->ToggleGloballyShowMarkedFrames();
-						SequencerPtr.Pin()->InvalidateGlobalMarkedFramesCache();
-					}
-				}),
-				FCanExecuteAction::CreateLambda([FocusedMovieScene]() { return FocusedMovieScene != nullptr; }),
-				FIsActionChecked::CreateLambda([FocusedMovieScene]() {
-					if (FocusedMovieScene != nullptr)
-					{
-						return FocusedMovieScene->GetGloballyShowMarkedFrames();
-					}
-					return false;
-				})),
-			NAME_None,
-			EUserInterfaceActionType::ToggleButton
-		);
-		MenuBuilder.AddMenuEntry(
-			LOCTEXT("ClearMarkedFramesGlobal", "Clear Global Marked Frames"),
-			LOCTEXT("ClearMarkedFramesGlobalTooltip", "Set all marked frames in all sub-sequences to not be globally displayed."),
-			FSlateIcon(),
-			FUIAction(FExecuteAction::CreateLambda([this]() { SequencerPtr.Pin()->ClearGlobalMarkedFrames(); }))
-		);
+		MenuBuilder.AddMenuEntry(FSequencerCommands::Get().ToggleShowMarkedFramesGlobally);
+		MenuBuilder.AddMenuEntry(FSequencerCommands::Get().ClearGlobalMarkedFrames);
 	}
 	MenuBuilder.EndSection();
 
diff --git a/Engine/Source/Editor/Sequencer/Private/Sequencer.cpp b/Engine/Source/Editor/Sequencer/Private/Sequencer.cpp
--- a/Engine/Source/Editor/Sequencer/Private/Sequencer.cpp
+++ b/Engine/Source/Editor/Sequencer/Private/Sequencer.cpp
@@ -9737,6 +9737,32 @@ void FSequencer::UpdateGlobalMarkedFramesCache()
 	bGlobalMarkedFramesCached = true;
 }
 
+void FSequencer::ToggleShowMarkedFramesGlobally()
+{
+	UMovieSceneSequence* FocusedMovieSequence = GetFocusedMovieSceneSequence();
+	if (!FocusedMovieSequence)
+	{
+		return;
+	}
+
+	UMovieScene* FocusedMovieScene = FocusedMovieSequence->GetMovieScene();
+	if (!FocusedMovieScene)
+	{
+		return;
+	}
+
+	if (FocusedMovieScene->IsReadOnly())
+	{
+		FSequencerUtilities::ShowReadOnlyError();
+		return;
+	}
+
+	const FScopedTransaction Transaction(LOCTEXT("ToggleShowMarkedFramesGlobally", "Toggle Show Marked Frames Globally"));
+	FocusedMovieScene->Modify();
+	FocusedMovieScene->ToggleGloballyShowMarkedFrames();
+	InvalidateGlobalMarkedFramesCache();
+}
+
 void FSequencer::ClearGlobalMarkedFrames()
 {
 	FSequencerMarkedFrameHelper::ClearGlobalMarkedFrames(*this);
@@ -10505,6 +10531,16 @@ void FSequencer::BindCommands()
 		Commands.SortAllNodesAndDescendants,
 		FExecuteAction::CreateSP(this, &FSequencer::SortAllNodesAndDescendants));
 
+	SequencerCommandBindings->MapAction(
+		Commands.ToggleShowMarkedFramesGlobally,
+		FExecuteAction::CreateSP(this, &FSequencer::ToggleShowMarkedFramesGlobally),
+		FCanExecuteAction::CreateLambda([this] { return GetFocusedMovieSceneSequence() != nullptr; }),
+		FIsActionChecked::CreateLambda([this] { return GetFocusedMovieSceneSequence()->GetMovieScene()->GetGloballyShowMarkedFrames(); }) );
+
+	SequencerCommandBindings->MapAction(
+		Commands.ClearGlobalMarkedFrames,
+		FExecuteAction::CreateSP(this, &FSequencer::ClearGlobalMarkedFrames));
+
 	SequencerCommandBindings->MapAction(
 		Commands.ToggleAutoExpandNodesOnSelection,
 		FExecuteAction::CreateLambda([this] { Settings->SetAutoExpandNodesOnSelection(!Settings->GetAutoExpandNodesOnSelection()); }),
diff --git a/Engine/Source/Editor/Sequencer/Private/Sequencer.h b/Engine/Source/Editor/Sequencer/Private/Sequencer.h
--- a/Engine/Source/Editor/Sequencer/Private/Sequencer.h
+++ b/Engine/Source/Editor/Sequencer/Private/Sequencer.h
@@ -307,9 +307,10 @@ public:
 	void InvalidateGlobalMarkedFramesCache() { bGlobalMarkedFramesCached = false; }
 	void UpdateGlobalMarkedFramesCache();
 
-	/** 
-	  * Disables all global marked frames from all sub-sequences
-	  */
+	/** Toggle whether to show marked frames globally */
+	void ToggleShowMarkedFramesGlobally();
+
+	/** Disables all global marked frames from all sub-sequences */
 	void ClearGlobalMarkedFrames();
 
 protected:
diff --git a/Engine/Source/Editor/Sequencer/Private/SequencerCommands.cpp b/Engine/Source/Editor/Sequencer/Private/SequencerCommands.cpp
--- a/Engine/Source/Editor/Sequencer/Private/SequencerCommands.cpp
+++ b/Engine/Source/Editor/Sequencer/Private/SequencerCommands.cpp
@@ -117,6 +117,8 @@ void FSequencerCommands::RegisterCommands()
 	UI_COMMAND( StepToNextMark, "Step to Next Marked Frame", "Step to the next marked frame", EUserInterfaceActionType::Button, FInputChord(EModifierKey::Control | EModifierKey::Shift, EKeys::Period) );
 	UI_COMMAND( StepToPreviousMark, "Step to Previous Marked Frame", "Step to the previous marked frame", EUserInterfaceActionType::Button, FInputChord(EModifierKey::Control | EModifierKey::Shift, EKeys::Comma) );
 	UI_COMMAND( ToggleMarksLocked, "Marked Frames Locked", "Prevent editing the marked frames.", EUserInterfaceActionType::ToggleButton, FInputChord() );
+	UI_COMMAND( ToggleShowMarkedFramesGlobally, "Show Marked Frames Globally", "Makes marked frames in this sub-sequence visible in parent/sibling sequences for the current editor session.", EUserInterfaceActionType::ToggleButton, FInputChord() );
+	UI_COMMAND( ClearGlobalMarkedFrames, "Clear Global Marked Frames", "Set all marked frames in all sub-sequences to not be globally displayed.", EUserInterfaceActionType::Button, FInputChord() );
 
 	UI_COMMAND( ToggleAutoScroll, "Auto Scroll", "Toggle auto-scroll: When enabled, automatically scrolls the sequencer view to keep the current time visible", EUserInterfaceActionType::ToggleButton, FInputChord(EModifierKey::Shift, EKeys::S) );
 
diff --git a/Engine/Source/Editor/Sequencer/Public/SequencerCommands.h b/Engine/Source/Editor/Sequencer/Public/SequencerCommands.h
--- a/Engine/Source/Editor/Sequencer/Public/SequencerCommands.h
+++ b/Engine/Source/Editor/Sequencer/Public/SequencerCommands.h
@@ -295,6 +295,12 @@ public:
 	/** Toggle locking marks */
 	TSharedPtr< FUICommandInfo > ToggleMarksLocked;
 
+	/** Toggle show marked frames globally */
+	TSharedPtr< FUICommandInfo > ToggleShowMarkedFramesGlobally;
+
+	/** Clear global marked frames */
+	TSharedPtr< FUICommandInfo > ClearGlobalMarkedFrames;
+
 	/** Rotates through the supported formats for displaying times/frames/timecode. */
 	TSharedPtr< FUICommandInfo > ChangeTimeDisplayFormat;
 
