Sequencer: Fix property owner/container

PropertyContainer is what we look to for the setter function
PropertyOwner is the owner struct, which is where we look for permissions

The original fix was to find the property owner struct to determine DoesPropertyPassFilter. This broke finding setter functions for some properties.

diff --git a/Engine/Source/Editor/Sequencer/Private/KeyPropertyParams.cpp b/Engine/Source/Editor/Sequencer/Private/KeyPropertyParams.cpp
--- a/Engine/Source/Editor/Sequencer/Private/KeyPropertyParams.cpp
+++ b/Engine/Source/Editor/Sequencer/Private/KeyPropertyParams.cpp
@@ -57,7 +57,7 @@ FCanKeyPropertyParams::FCanKeyPropertyParams(const UClass* InObjectClass, const
 {
 }
 
-const UStruct* FCanKeyPropertyParams::FindPropertyContainer(const FProperty* ForProperty) const
+const UStruct* FCanKeyPropertyParams::FindPropertyOwner(const FProperty* ForProperty) const
 {
 	check(ForProperty);
 
@@ -81,6 +81,30 @@ const UStruct* FCanKeyPropertyParams::FindPropertyContainer(const FProperty* For
 	return ObjectClass;
 }
 
+const UStruct* FCanKeyPropertyParams::FindPropertyContainer(const FProperty* ForProperty) const
+{
+	check(ForProperty);
+
+	bool bFoundProperty = false;
+	for (int32 Index = PropertyPath.GetNumProperties() - 1; Index >= 0; --Index)
+	{
+		FProperty* Property = PropertyPath.GetPropertyInfo(Index).Property.Get();
+		if (!bFoundProperty)
+		{
+			bFoundProperty = Property == ForProperty;
+			if (bFoundProperty)
+			{
+				return ObjectClass;
+			}
+		}
+		else if (const FStructProperty* StructProperty = CastField<const FStructProperty>(Property))
+		{
+			return StructProperty->Struct;
+		}
+	}
+	return ObjectClass;
+}
+
 FKeyPropertyParams::FKeyPropertyParams(TArray<UObject*> InObjectsToKey, const FPropertyPath& InPropertyPath, ESequencerKeyMode InKeyMode)
 	: ObjectsToKey(InObjectsToKey)
 	, PropertyPath(InPropertyPath)
diff --git a/Engine/Source/Editor/Sequencer/Private/SequencerObjectChangeListener.cpp b/Engine/Source/Editor/Sequencer/Private/SequencerObjectChangeListener.cpp
--- a/Engine/Source/Editor/Sequencer/Private/SequencerObjectChangeListener.cpp
+++ b/Engine/Source/Editor/Sequencer/Private/SequencerObjectChangeListener.cpp
@@ -288,14 +288,16 @@ bool FSequencerObjectChangeListener::CanKeyProperty_Internal(FCanKeyPropertyPara
 				continue;
 			}
 
-			const UStruct* PropertyContainer = CanKeyPropertyParams.FindPropertyContainer(Property);
-			if (PropertyContainer)
+			const UStruct* PropertyOwner = CanKeyPropertyParams.FindPropertyOwner(Property);
+			if (PropertyOwner)
 			{
-				if (!FPropertyEditorPermissionList::Get().DoesPropertyPassFilter(PropertyContainer, Property->GetFName()))
+				if (!FPropertyEditorPermissionList::Get().DoesPropertyPassFilter(PropertyOwner, Property->GetFName()))
 				{
 					continue;
 				}
 
+				const UStruct* PropertyContainer = CanKeyPropertyParams.FindPropertyContainer(Property);
+				if (PropertyContainer)
 				{
 					FAnimatedPropertyKey PropertyKey = FAnimatedPropertyKey::FromProperty(Property);
 
diff --git a/Engine/Source/Editor/Sequencer/Public/KeyPropertyParams.h b/Engine/Source/Editor/Sequencer/Public/KeyPropertyParams.h
--- a/Engine/Source/Editor/Sequencer/Public/KeyPropertyParams.h
+++ b/Engine/Source/Editor/Sequencer/Public/KeyPropertyParams.h
@@ -37,6 +37,10 @@ struct SEQUENCER_API FCanKeyPropertyParams
 	*/
 	FCanKeyPropertyParams(const UClass* InObjectClass, const IPropertyHandle& InPropertyHandle);
 
+	/** The owner struct */
+	const UStruct* FindPropertyOwner(const FProperty* ForProperty) const;
+
+	/** The container from which to find setter functions */
 	const UStruct* FindPropertyContainer(const FProperty* ForProperty) const;
 
 	/** The class of the object which has the property to be keyed. */
