Sequencer: Fix static analysis

#rnx

diff --git a/Engine/Source/Editor/Sequencer/Private/SequencerObjectChangeListener.cpp b/Engine/Source/Editor/Sequencer/Private/SequencerObjectChangeListener.cpp
--- a/Engine/Source/Editor/Sequencer/Private/SequencerObjectChangeListener.cpp
+++ b/Engine/Source/Editor/Sequencer/Private/SequencerObjectChangeListener.cpp
@@ -289,64 +289,57 @@ bool FSequencerObjectChangeListener::CanKeyProperty_Internal(FCanKeyPropertyPara
 			}
 
 			const UStruct* PropertyOwner = CanKeyPropertyParams.FindPropertyOwner(Property);
-			if (PropertyOwner)
+			if (!PropertyOwner)
 			{
-				if (!FPropertyEditorPermissionList::Get().DoesPropertyPassFilter(PropertyOwner, Property->GetFName()))
-				{
-					continue;
-				}
-
-				const UStruct* PropertyContainer = CanKeyPropertyParams.FindPropertyContainer(Property);
-				if (PropertyContainer)
-				{
-					FAnimatedPropertyKey PropertyKey = FAnimatedPropertyKey::FromProperty(Property);
+				continue;
+			}
 
-					// If there is a custom accessor for this specific property path, it is animatable (as long as there is a supported track editor registered for the property type)
-					if (UE::MovieScene::GlobalCustomAccessorExists(CanKeyPropertyParams.ObjectClass, InOutPropertyPath.ToString(TEXT("."))))
-					{
-						if (const FOnAnimatablePropertyChanged* DelegatePtr = PropertyChangedEventMap.Find(PropertyKey))
-						{
-							InOutProperty = Property;
-							InOutDelegate = *DelegatePtr;
-							return true;
-						}
-					}
-					// Otherwise we check for magic named functions using the default logic
-					const FOnAnimatablePropertyChanged* DelegatePtr = FindPropertySetter(*PropertyContainer, PropertyKey, *Property);
-					if (DelegatePtr != nullptr)
-					{
-						InOutProperty = Property;
-						InOutDelegate = *DelegatePtr;
-						return true;
-					}
-				}
+			if (!FPropertyEditorPermissionList::Get().DoesPropertyPassFilter(PropertyOwner, Property->GetFName()))
+			{
+				continue;
+			}
 
+			const UStruct* PropertyContainer = CanKeyPropertyParams.FindPropertyContainer(Property);
+			if (!PropertyContainer)
+			{
+				continue;
+			}
 
-				FObjectPropertyBase* ObjectProperty = CastField<FObjectPropertyBase>(Property);
+			FAnimatedPropertyKey PropertyKey = FAnimatedPropertyKey::FromProperty(Property);
 
-				// Check each level of the property hierarchy
-				FFieldClass* PropertyType = Property->GetClass();
-				while (PropertyType && PropertyType != FProperty::StaticClass())
+			// If there is a custom accessor for this specific property path, it is animatable (as long as there is a supported track editor registered for the property type)
+			if (UE::MovieScene::GlobalCustomAccessorExists(CanKeyPropertyParams.ObjectClass, InOutPropertyPath.ToString(TEXT("."))))
+			{
+				if (const FOnAnimatablePropertyChanged* DelegatePtr = PropertyChangedEventMap.Find(PropertyKey))
 				{
-					FAnimatedPropertyKey Key = FAnimatedPropertyKey::FromPropertyTypeName(PropertyType->GetFName());
+					InOutProperty = Property;
+					InOutDelegate = *DelegatePtr;
+					return true;
+				}
+			}
 
-					// For object properties, check each parent type of the object (ie, so a track that animates UBaseClass ptrs can be used with a UDerivedClass property)
-					UClass* ClassType = (ObjectProperty && ObjectProperty->PropertyClass) ? ObjectProperty->PropertyClass->GetSuperClass() : nullptr;
-					while (ClassType)
-					{
-						Key.ObjectTypeName = ClassType->GetFName();
+			// Otherwise we check for magic named functions using the default logic
+			if (const FOnAnimatablePropertyChanged* DelegatePtr = FindPropertySetter(*PropertyContainer, PropertyKey, *Property))
+			{
+				InOutProperty = Property;
+				InOutDelegate = *DelegatePtr;
+				return true;
+			}
 
-						if (const FOnAnimatablePropertyChanged* DelegatePtr = FindPropertySetter(*PropertyContainer, Key, *Property))
-						{
-							InOutProperty = Property;
-							InOutDelegate = *DelegatePtr;
-							return true;
-						}
+			FObjectPropertyBase* ObjectProperty = CastField<FObjectPropertyBase>(Property);
 
-						ClassType = ClassType->GetSuperClass();
-					}
+			// Check each level of the property hierarchy
+			FFieldClass* PropertyType = Property->GetClass();
+			while (PropertyType && PropertyType != FProperty::StaticClass())
+			{
+				FAnimatedPropertyKey Key = FAnimatedPropertyKey::FromPropertyTypeName(PropertyType->GetFName());
+
+				// For object properties, check each parent type of the object (ie, so a track that animates UBaseClass ptrs can be used with a UDerivedClass property)
+				UClass* ClassType = (ObjectProperty && ObjectProperty->PropertyClass) ? ObjectProperty->PropertyClass->GetSuperClass() : nullptr;
+				while (ClassType)
+				{
+					Key.ObjectTypeName = ClassType->GetFName();
 
-					Key.ObjectTypeName = NAME_None;
 					if (const FOnAnimatablePropertyChanged* DelegatePtr = FindPropertySetter(*PropertyContainer, Key, *Property))
 					{
 						InOutProperty = Property;
@@ -354,9 +347,19 @@ bool FSequencerObjectChangeListener::CanKeyProperty_Internal(FCanKeyPropertyPara
 						return true;
 					}
 
-					// Look at the property's super class
-					PropertyType = PropertyType->GetSuperClass();
+					ClassType = ClassType->GetSuperClass();
 				}
+
+				Key.ObjectTypeName = NAME_None;
+				if (const FOnAnimatablePropertyChanged* DelegatePtr = FindPropertySetter(*PropertyContainer, Key, *Property))
+				{
+					InOutProperty = Property;
+					InOutDelegate = *DelegatePtr;
+					return true;
+				}
+
+				// Look at the property's super class
+				PropertyType = PropertyType->GetSuperClass();
 			}
 		}
 	}
