Sequencer: don't add a playback capability if it already exists

This can happen if Sequencer is re-focused on a sequence that it was already focused on.

diff --git a/Engine/Source/Editor/Sequencer/Private/Sequencer.cpp b/Engine/Source/Editor/Sequencer/Private/Sequencer.cpp
--- a/Engine/Source/Editor/Sequencer/Private/Sequencer.cpp
+++ b/Engine/Source/Editor/Sequencer/Private/Sequencer.cpp
@@ -477,7 +477,10 @@ void FSequencer::InitSequencer(const FSequencerInitParams& InitParams, const TSh
 	UMovieSceneEntitySystemLinker* Linker = GetEvaluationTemplate().GetEntitySystemLinker();
 	FInstanceRegistry* InstanceRegistry = Linker->GetInstanceRegistry();
 	TSharedRef<FSharedPlaybackState> SharedPlaybackState = InstanceRegistry->GetInstance(RootInstanceHandle).GetSharedPlaybackState();
-	SharedPlaybackState->AddCapabilityRaw<FCameraCutPlaybackCapability>((FCameraCutPlaybackCapability*)this);
+	if (!SharedPlaybackState->HasCapability<FCameraCutPlaybackCapability>())
+	{
+		SharedPlaybackState->AddCapabilityRaw<FCameraCutPlaybackCapability>((FCameraCutPlaybackCapability*)this);
+	}
 
 	// Create tools and bind them to this sequencer
 	for( int32 DelegateIndex = 0; DelegateIndex < TrackEditorDelegates.Num(); ++DelegateIndex )
@@ -1100,7 +1103,10 @@ void FSequencer::ResetToNewRootSequence(UMovieSceneSequence& NewSequence)
 	UMovieSceneEntitySystemLinker* Linker = GetEvaluationTemplate().GetEntitySystemLinker();
 	FInstanceRegistry* InstanceRegistry = Linker->GetInstanceRegistry();
 	TSharedRef<FSharedPlaybackState> SharedPlaybackState = InstanceRegistry->GetInstance(RootInstanceHandle).GetSharedPlaybackState();
-	SharedPlaybackState->AddCapabilityRaw<FCameraCutPlaybackCapability>((FCameraCutPlaybackCapability*)this);
+	if (!SharedPlaybackState->HasCapability<FCameraCutPlaybackCapability>())
+	{
+		SharedPlaybackState->AddCapabilityRaw<FCameraCutPlaybackCapability>((FCameraCutPlaybackCapability*)this);
+	}
 
 	OnActivateSequenceEvent.Broadcast(ActiveTemplateIDs.Top());
 }
diff --git a/Engine/Source/Runtime/MovieScene/Public/EntitySystem/MovieSceneSharedPlaybackState.h b/Engine/Source/Runtime/MovieScene/Public/EntitySystem/MovieSceneSharedPlaybackState.h
--- a/Engine/Source/Runtime/MovieScene/Public/EntitySystem/MovieSceneSharedPlaybackState.h
+++ b/Engine/Source/Runtime/MovieScene/Public/EntitySystem/MovieSceneSharedPlaybackState.h
@@ -15,6 +15,12 @@ struct FSharedPlaybackState
 
 public:
 
+	template<typename T>
+	bool HasCapability() const
+	{
+		return Capabilities.HasCapability(T::ID);
+	}
+
 	template<typename T>
 	T* FindCapability() const
 	{
