Sequencer: Only add components and properties for selected actors if the given owner is actually selected

This fixes an issue where if you select an actor and then choose +Track component for another actor, they will both get a track. This should only happen if the actor is actually selected.

diff --git a/Engine/Plugins/MovieScene/LevelSequenceEditor/Source/LevelSequenceEditor/Private/LevelSequenceEditorToolkit.cpp b/Engine/Plugins/MovieScene/LevelSequenceEditor/Source/LevelSequenceEditor/Private/LevelSequenceEditorToolkit.cpp
--- a/Engine/Plugins/MovieScene/LevelSequenceEditor/Source/LevelSequenceEditor/Private/LevelSequenceEditorToolkit.cpp
+++ b/Engine/Plugins/MovieScene/LevelSequenceEditor/Source/LevelSequenceEditor/Private/LevelSequenceEditorToolkit.cpp
@@ -659,8 +659,9 @@ void FLevelSequenceEditorToolkit::HandleAddComponentActionExecute(UActorComponen
 	TArray<UActorComponent*> ActorComponents;
 	ActorComponents.Add(Component);
 
+	// Add selected actor components if the given component is from a selected actor
 	USelection* SelectedActors = GEditor->GetSelectedActors();
-	if (SelectedActors && SelectedActors->Num() > 0)
+	if (SelectedActors && SelectedActors->Num() > 0 && SelectedActors->IsSelected(Component->GetOwner()))
 	{
 		for (FSelectionIterator Iter(*SelectedActors); Iter; ++Iter)
 		{
diff --git a/Engine/Source/Editor/Sequencer/Private/MVVM/ViewModels/ObjectBindingModel.cpp b/Engine/Source/Editor/Sequencer/Private/MVVM/ViewModels/ObjectBindingModel.cpp
--- a/Engine/Source/Editor/Sequencer/Private/MVVM/ViewModels/ObjectBindingModel.cpp
+++ b/Engine/Source/Editor/Sequencer/Private/MVVM/ViewModels/ObjectBindingModel.cpp
@@ -788,14 +788,19 @@ void FObjectBindingModel::HandlePropertyMenuItemExecute(FPropertyPath PropertyPa
 		}
 	}
 
-	for (TViewModelPtr<FObjectBindingModel> ObjectBindingNode : Sequencer->GetViewModel()->GetSelection()->Outliner.Filter<FObjectBindingModel>())
+	// Only include other selected object bindings if this binding is selected. Otherwise, this will lead to 
+	// confusion with multiple tracks being added to possibly unrelated objects
+	if (Sequencer->GetViewModel()->GetSelection()->Outliner.IsSelected(SharedThis(this)))
 	{
-		FGuid Guid = ObjectBindingNode->GetObjectGuid();
-		for (auto RuntimeObject : Sequencer->FindBoundObjects(Guid, OwnerModel->GetSequenceID()))
+		for (TViewModelPtr<FObjectBindingModel> ObjectBindingNode : Sequencer->GetViewModel()->GetSelection()->Outliner.Filter<FObjectBindingModel>())
 		{
-			if (Sequencer->CanKeyProperty(FCanKeyPropertyParams(RuntimeObject->GetClass(), PropertyPath)))
+			FGuid Guid = ObjectBindingNode->GetObjectGuid();
+			for (auto RuntimeObject : Sequencer->FindBoundObjects(Guid, OwnerModel->GetSequenceID()))
 			{
-				KeyableBoundObjects.AddUnique(RuntimeObject.Get());
+				if (Sequencer->CanKeyProperty(FCanKeyPropertyParams(RuntimeObject->GetClass(), PropertyPath)))
+				{
+					KeyableBoundObjects.AddUnique(RuntimeObject.Get());
+				}
 			}
 		}
 	}
