Sequencer: Change to OwnerModel->GetEditor() instead of Sequencer

#rnx

diff --git a/Engine/Source/Editor/Sequencer/Private/MVVM/ViewModels/ObjectBindingModel.cpp b/Engine/Source/Editor/Sequencer/Private/MVVM/ViewModels/ObjectBindingModel.cpp
--- a/Engine/Source/Editor/Sequencer/Private/MVVM/ViewModels/ObjectBindingModel.cpp
+++ b/Engine/Source/Editor/Sequencer/Private/MVVM/ViewModels/ObjectBindingModel.cpp
@@ -475,9 +475,9 @@ TSharedRef<SWidget> FObjectBindingModel::GetAddTrackMenuContent()
 
 	// Only include other selected object bindings if this binding is selected. Otherwise, this will lead to 
 	// confusion with multiple tracks being added to possibly unrelated objects
-	if (Sequencer->GetViewModel()->GetSelection()->Outliner.IsSelected(SharedThis(this)))
+	if (OwnerModel->GetEditor()->GetSelection()->Outliner.IsSelected(SharedThis(this)))
 	{
-		for (TViewModelPtr<FObjectBindingModel> ObjectBindingNode : Sequencer->GetViewModel()->GetSelection()->Outliner.Filter<FObjectBindingModel>())
+		for (TViewModelPtr<FObjectBindingModel> ObjectBindingNode : OwnerModel->GetEditor()->GetSelection()->Outliner.Filter<FObjectBindingModel>())
 		{
 			const FGuid Guid = ObjectBindingNode->GetObjectGuid();
 			for (auto RuntimeObject : Sequencer->FindBoundObjects(Guid, OwnerModel->GetSequenceID()))
@@ -790,9 +790,9 @@ void FObjectBindingModel::HandlePropertyMenuItemExecute(FPropertyPath PropertyPa
 
 	// Only include other selected object bindings if this binding is selected. Otherwise, this will lead to 
 	// confusion with multiple tracks being added to possibly unrelated objects
-	if (Sequencer->GetViewModel()->GetSelection()->Outliner.IsSelected(SharedThis(this)))
+	if (OwnerModel->GetEditor()->GetSelection()->Outliner.IsSelected(SharedThis(this)))
 	{
-		for (TViewModelPtr<FObjectBindingModel> ObjectBindingNode : Sequencer->GetViewModel()->GetSelection()->Outliner.Filter<FObjectBindingModel>())
+		for (TViewModelPtr<FObjectBindingModel> ObjectBindingNode : OwnerModel->GetEditor()->GetSelection()->Outliner.Filter<FObjectBindingModel>())
 		{
 			FGuid Guid = ObjectBindingNode->GetObjectGuid();
 			for (auto RuntimeObject : Sequencer->FindBoundObjects(Guid, OwnerModel->GetSequenceID()))
@@ -994,7 +994,7 @@ void FObjectBindingModel::AddTagMenu(FMenuBuilder& MenuBuilder)
 
 		// Gather all the tags on all currently selected object binding IDs
 		FMovieSceneSequenceID SequenceID = OwnerModel->GetSequenceID();
-		for (TViewModelPtr<FObjectBindingModel> ObjectBindingNode : Sequencer->GetViewModel()->GetSelection()->Outliner.Filter<FObjectBindingModel>())
+		for (TViewModelPtr<FObjectBindingModel> ObjectBindingNode : OwnerModel->GetEditor()->GetSelection()->Outliner.Filter<FObjectBindingModel>())
 		{
 			const FGuid& ObjectID = ObjectBindingNode->GetObjectGuid();
 
@@ -1045,7 +1045,7 @@ ECheckBoxState FObjectBindingModel::GetTagCheckState(FName TagName)
 	TSharedPtr<FSequencer> Sequencer = OwnerModel->GetSequencerImpl();
 	FMovieSceneSequenceID SequenceID = OwnerModel->GetSequenceID();
 
-	for (TViewModelPtr<FObjectBindingModel> ObjectBindingNode : Sequencer->GetViewModel()->GetSelection()->Outliner.Filter<FObjectBindingModel>())
+	for (TViewModelPtr<FObjectBindingModel> ObjectBindingNode : OwnerModel->GetEditor()->GetSelection()->Outliner.Filter<FObjectBindingModel>())
 	{
 		const FGuid& ObjectID = ObjectBindingNode->GetObjectGuid();
 
@@ -1072,7 +1072,7 @@ void FObjectBindingModel::ToggleTag(FName TagName)
 	TSharedPtr<FSequencer> Sequencer = OwnerModel->GetSequencerImpl();
 	FMovieSceneSequenceID SequenceID = OwnerModel->GetSequenceID();
 
-	for (TViewModelPtr<FObjectBindingModel> ObjectBindingNode : Sequencer->GetViewModel()->GetSelection()->Outliner.Filter<FObjectBindingModel>())
+	for (TViewModelPtr<FObjectBindingModel> ObjectBindingNode : OwnerModel->GetEditor()->GetSelection()->Outliner.Filter<FObjectBindingModel>())
 	{
 		const FGuid& ObjectID = ObjectBindingNode->GetObjectGuid();
 
@@ -1096,7 +1096,7 @@ void FObjectBindingModel::HandleDeleteTag(FName TagName)
 	MovieScene->Modify();
 
 	FMovieSceneSequenceID SequenceID = OwnerModel->GetSequenceID();
-	for (TViewModelPtr<FObjectBindingModel> ObjectBindingNode : Sequencer->GetViewModel()->GetSelection()->Outliner.Filter<FObjectBindingModel>())
+	for (TViewModelPtr<FObjectBindingModel> ObjectBindingNode : OwnerModel->GetEditor()->GetSelection()->Outliner.Filter<FObjectBindingModel>())
 	{
 		const FGuid& ObjectID = ObjectBindingNode->GetObjectGuid();
 		MovieScene->UntagBinding(TagName, UE::MovieScene::FFixedObjectBindingID(ObjectID, SequenceID));
@@ -1112,7 +1112,7 @@ void FObjectBindingModel::HandleAddTag(FName TagName)
 	MovieScene->Modify();
 
 	FMovieSceneSequenceID SequenceID = OwnerModel->GetSequenceID();
-	for (TViewModelPtr<FObjectBindingModel> ObjectBindingNode : Sequencer->GetViewModel()->GetSelection()->Outliner.Filter<FObjectBindingModel>())
+	for (TViewModelPtr<FObjectBindingModel> ObjectBindingNode : OwnerModel->GetEditor()->GetSelection()->Outliner.Filter<FObjectBindingModel>())
 	{
 		const FGuid& ObjectID = ObjectBindingNode->GetObjectGuid();
 		MovieScene->TagBinding(TagName, UE::MovieScene::FFixedObjectBindingID(ObjectID, SequenceID));
