Sequencer:

Added option to have different Keyed Status / Icon via Details Keyframe Handler.
Current Keyframe Icon is completely replaced, but not yet deleted from Content or from the Style Set entry.
Added Sequencer->GetPropertyKeyedStatus, which uses the FSequencerPropertyKeyedStatus class to find the keyed status for a property.

The keyed status of a property is gotten through matching the property handle to the channel via its channel meta data. Two new member variables to the Channel MetaData were added to achieve this:
1) Sub Property Path: For structs, this is the relative path from the topmost property. E.g. for the pitch in a Transform property this would be "Rotation.Pitch"
2) Sub Property Path Map: some channels are re-used by multiple different structs and some might have different property names. The prime example is Transform and EulerTransform. These two have different names for their properties. FTransform uses "Translation" whereas FEulerTransform uses "Location", so these two must be differentiated via a map from the Struct's FName to the FName of the name it represents. This is also the case for Scale where FTransform uses "Scale3D" and FEulerTransform uses "Scale".

In addition to this, External handling is possible via ISequencerPropertyKeyedStatus::GetExternalHandler. This is used by [REDACTED_PERSON] to link the Scene Component properties (Relative Location/Rotation/Scale) so that these can also get their Keyframe Status properly, as the Transform Track itself can be in the Actor and not in the (Root) Scene Component where the properties are actually at.

diff --git a/Engine/Plugins/Compositing/Composure/Source/Composure/Private/MovieScene/MovieSceneComposurePostMoveSettingsSection.cpp b/Engine/Plugins/Compositing/Composure/Source/Composure/Private/MovieScene/MovieSceneComposurePostMoveSettingsSection.cpp
--- a/Engine/Plugins/Compositing/Composure/Source/Composure/Private/MovieScene/MovieSceneComposurePostMoveSettingsSection.cpp
+++ b/Engine/Plugins/Compositing/Composure/Source/Composure/Private/MovieScene/MovieSceneComposurePostMoveSettingsSection.cpp
@@ -18,30 +18,36 @@ struct FPostMoveSettingsChannelEditorData
 		FText TranslationGroup = NSLOCTEXT("PostMoves", "Translation", "Translation");
 
 		MetaData[0].SetIdentifiers("Pivot.X", FCommonChannelData::ChannelX, PivotGroup);
+		MetaData[0].SubPropertyPath = MetaData[0].Name;
 		MetaData[0].SortOrder = 0;
 		MetaData[0].Color = FCommonChannelData::RedChannelColor;
 		MetaData[0].bCanCollapseToTrack = false;
 
 		MetaData[1].SetIdentifiers("Pivot.Y", FCommonChannelData::ChannelY, PivotGroup);
+		MetaData[1].SubPropertyPath = MetaData[1].Name;
 		MetaData[1].SortOrder = 1;
 		MetaData[1].Color = FCommonChannelData::GreenChannelColor;
 		MetaData[1].bCanCollapseToTrack = false;
 
 		MetaData[2].SetIdentifiers("Translation.X", FCommonChannelData::ChannelX, TranslationGroup);
+		MetaData[2].SubPropertyPath = MetaData[2].Name;
 		MetaData[2].SortOrder = 2;
 		MetaData[2].Color = FCommonChannelData::RedChannelColor;
 		MetaData[2].bCanCollapseToTrack = false;
 
 		MetaData[3].SetIdentifiers("Translation.Y", FCommonChannelData::ChannelY, TranslationGroup);
+		MetaData[3].SubPropertyPath = MetaData[3].Name;
 		MetaData[3].SortOrder = 3;
 		MetaData[3].Color = FCommonChannelData::GreenChannelColor;
 		MetaData[3].bCanCollapseToTrack = false;
 
 		MetaData[4].SetIdentifiers("Rotation", NSLOCTEXT("PostMoves", "Rotation", "Rotation"));
+		MetaData[4].SubPropertyPath = TEXT("RotationAngle");
 		MetaData[4].SortOrder = 4;
 		MetaData[4].bCanCollapseToTrack = false;
 
 		MetaData[5].SetIdentifiers("Scale", NSLOCTEXT("PostMoves", "Scale", "Scale"));
+		MetaData[5].SubPropertyPath = MetaData[5].Name;
 		MetaData[5].SortOrder = 5;
 		MetaData[5].bCanCollapseToTrack = false;
 
diff --git a/Engine/Source/Editor/EditorStyle/Private/StarshipStyle.cpp b/Engine/Source/Editor/EditorStyle/Private/StarshipStyle.cpp
--- a/Engine/Source/Editor/EditorStyle/Private/StarshipStyle.cpp
+++ b/Engine/Source/Editor/EditorStyle/Private/StarshipStyle.cpp
@@ -2275,6 +2275,11 @@ void FStarshipEditorStyle::FStyle::SetupSequencerStyles()
 
 		Set("Sequencer.AddKey.Details", new IMAGE_BRUSH("Sequencer/AddKey_Details", FVector2D(11, 11)));
 
+		Set("Sequencer.KeyedStatus.NotKeyed", new IMAGE_BRUSH_SVG("Sequencer/DetailsKeyUnkeyed", FVector2D(11, 11)));
+		Set("Sequencer.KeyedStatus.Keyed", new IMAGE_BRUSH_SVG("Sequencer/DetailsKeyKeyed", FVector2D(11, 11)));
+		Set("Sequencer.KeyedStatus.Animated", new IMAGE_BRUSH_SVG("Sequencer/DetailsKeyAnimated", FVector2D(11, 11)));
+		Set("Sequencer.KeyedStatus.PartialKey", new IMAGE_BRUSH_SVG("Sequencer/DetailsKeyPartialKey", FVector2D(11, 11)));
+
 		const FSplitterStyle OutlinerSplitterStyle = FSplitterStyle()
 		.SetHandleNormalBrush(FSlateNoResource())
 		.SetHandleHighlightBrush(FSlateNoResource());
diff --git a/Engine/Source/Editor/MovieSceneTools/Private/TrackEditors/TransformTrackEditor.cpp b/Engine/Source/Editor/MovieSceneTools/Private/TrackEditors/TransformTrackEditor.cpp
--- a/Engine/Source/Editor/MovieSceneTools/Private/TrackEditors/TransformTrackEditor.cpp
+++ b/Engine/Source/Editor/MovieSceneTools/Private/TrackEditors/TransformTrackEditor.cpp
@@ -43,7 +43,10 @@
 #include "Constraints/TransformConstraintChannelInterface.h"
 #include "Misc/TransactionObjectEvent.h"
 #include "Engine/Selection.h"
+#include "PropertyHandle.h"
+#include "IDetailKeyframeHandler.h"
 #include "ISequencerObjectChangeListener.h"
+#include "ISequencerPropertyKeyedStatus.h"
 
 #define LOCTEXT_NAMESPACE "MovieScene_TransformTrack"
 
@@ -86,13 +89,18 @@ F3DTransformTrackEditor::F3DTransformTrackEditor( TSharedRef<ISequencer> InSeque
 		const FProperty* Scale3DProperty = FindFProperty<FProperty>(USceneComponent::StaticClass(), USceneComponent::GetRelativeScale3DPropertyName());
 
 		ISequencerObjectChangeListener& ObjectChangeListener = SequencerPtr->GetObjectChangeListener();
-		auto AddTransformProperty = [this, &ObjectChangeListener](const FProperty* Property, EMovieSceneTransformChannel TransformChannel)
+		ISequencerPropertyKeyedStatusHandler& PropertyKeyedStatusHandler = SequencerPtr->GetPropertyKeyedStatusHandler();
+
+		auto AddTransformProperty = [this, &ObjectChangeListener, &PropertyKeyedStatusHandler](const FProperty* Property, EMovieSceneTransformChannel TransformChannel)
 		{
 			if (Property)
 			{
 				TransformProperties.Add({ Property, TransformChannel });
 				ObjectChangeListener.GetOnAnimatablePropertyChanged(Property)
 					.AddRaw(this, &F3DTransformTrackEditor::OnTransformPropertyChanged, TransformChannel);
+
+				PropertyKeyedStatusHandler.GetExternalHandler(Property)
+					.BindRaw(this, &F3DTransformTrackEditor::GetPropertyKeyedStatus, TransformChannel);
 			}
 		};
 
@@ -112,9 +120,11 @@ F3DTransformTrackEditor::~F3DTransformTrackEditor()
 	if (TSharedPtr<ISequencer> SequencerPtr = FMovieSceneTrackEditor::GetSequencer())
 	{
 		ISequencerObjectChangeListener& ObjectChangeListener = SequencerPtr->GetObjectChangeListener();
+		ISequencerPropertyKeyedStatusHandler& PropertyKeyedStatusHandler = SequencerPtr->GetPropertyKeyedStatusHandler();
 		for (const FTransformPropertyInfo& TransformProperty : TransformProperties)
 		{
 			ObjectChangeListener.GetOnAnimatablePropertyChanged(TransformProperty.Property).RemoveAll(this);
+			PropertyKeyedStatusHandler.GetExternalHandler(TransformProperty.Property).Unbind();
 		}
 	}
 }
@@ -1116,6 +1126,206 @@ int32 GetPreviousKey(FMovieSceneDoubleChannel& Channel, FFrameNumber Time)
 	return Index;
 }
 
+EPropertyKeyedStatus F3DTransformTrackEditor::GetKeyedStatusInSection(const UMovieScene3DTransformSection& Section, const TRange<FFrameNumber>& Range, EMovieSceneTransformChannel TransformChannel, TConstArrayView<int32> ChannelIndices) const
+{
+	EPropertyKeyedStatus SectionKeyedStatus = EPropertyKeyedStatus::NotKeyed;
+
+	// Skip Section if the Transform Channel is completely masked out
+	if (!EnumHasAnyFlags(Section.GetMask().GetChannels(), TransformChannel))
+	{
+		return SectionKeyedStatus;
+	}
+
+	int32 EmptyChannelCount = 0;
+
+	FMovieSceneChannelProxy& ChannelProxy = Section.GetChannelProxy();
+	for (int32 ChannelIndex : ChannelIndices)
+	{
+		FMovieSceneDoubleChannel* Channel = ChannelProxy.GetChannel<FMovieSceneDoubleChannel>(ChannelIndex);
+		if (!Channel)
+		{
+			continue;
+		}
+
+		if (Channel->GetNumKeys() == 0)
+		{
+			++EmptyChannelCount;
+			continue;
+		}
+
+		SectionKeyedStatus = FMath::Max(SectionKeyedStatus, EPropertyKeyedStatus::KeyedInOtherFrame);
+
+		TArray<FFrameNumber> KeyTimes;
+		Channel->GetKeys(Range, &KeyTimes, nullptr);
+		if (KeyTimes.IsEmpty())
+		{
+			++EmptyChannelCount;
+		}
+		else
+		{
+			SectionKeyedStatus = FMath::Max(SectionKeyedStatus, EPropertyKeyedStatus::PartiallyKeyed);
+		}
+	}
+
+	if (EmptyChannelCount == 0 && SectionKeyedStatus == EPropertyKeyedStatus::PartiallyKeyed)
+	{
+		SectionKeyedStatus = EPropertyKeyedStatus::KeyedInFrame;
+	}
+	return SectionKeyedStatus;
+}
+
+EPropertyKeyedStatus F3DTransformTrackEditor::GetPropertyKeyedStatus(const IPropertyHandle& PropertyHandle, EMovieSceneTransformChannel TransformChannel) const
+{
+	TRACE_CPUPROFILER_EVENT_SCOPE(F3DTransformTrackEditor::GetPropertyKeyedStatus);
+
+	TSharedPtr<ISequencer> SequencerPtr = FMovieSceneTrackEditor::GetSequencer();
+	if (!SequencerPtr.IsValid())
+	{
+		return EPropertyKeyedStatus::NotKeyed;
+	}
+
+	UMovieSceneSequence* Sequence = SequencerPtr->GetFocusedMovieSceneSequence();
+
+	UMovieScene* MovieScene = Sequence ? Sequence->GetMovieScene() : nullptr;
+	if (!MovieScene)
+	{
+		return EPropertyKeyedStatus::NotKeyed;
+	}
+
+	TArray<UObject*> OuterObjects;
+	PropertyHandle.GetOuterObjects(OuterObjects);
+	if (OuterObjects.IsEmpty())
+	{
+		return EPropertyKeyedStatus::NotKeyed;
+	}
+
+	TSet<const UMovieScene3DTransformSection*> ProcessedSections;
+	ProcessedSections.Reserve(OuterObjects.Num());
+
+	TArray<int32, TFixedAllocator<3>> ChannelIndices;
+	switch (TransformChannel)
+	{
+	case EMovieSceneTransformChannel::Translation:
+		ChannelIndices = { 0, 1, 2 };
+		break;
+
+	case EMovieSceneTransformChannel::Rotation:
+		ChannelIndices = { 3, 4, 5 };
+		break;
+
+	case EMovieSceneTransformChannel::Scale:
+		ChannelIndices = { 6, 7, 8 };
+		break;
+	}
+
+	const TRange<FFrameNumber> FrameRange = TRange<FFrameNumber>(SequencerPtr->GetLocalTime().Time.FrameNumber);
+
+	EPropertyKeyedStatus KeyedStatus = EPropertyKeyedStatus::NotKeyed;
+
+	// List of Tracks that had no sections at the current frame that require an additional check by going through all its sections and whether there's a key in any of them.
+	// Used only to determine whether to return "Not Keyed" or "Keyed In Other Frame".
+	TArray<const UMovieScene3DTransformTrack*> TransformTracksToCheck;
+	TransformTracksToCheck.Reserve(OuterObjects.Num());
+
+	for (UObject* Object : OuterObjects)
+	{
+		USceneComponent* SceneComponent = Cast<USceneComponent>(Object);
+		if (!SceneComponent)
+		{
+			continue;
+		}
+
+		TArray<UMovieScene3DTransformTrack*> TransformTracks;
+		{
+			constexpr bool bCreateHandleIfMissing = false;
+
+			// Include Owner Actor Transform Track if Scene Component is Root
+			AActor* OwningActor = SceneComponent->GetOwner();
+			if (OwningActor && OwningActor->GetRootComponent() == SceneComponent)
+			{
+				FGuid OwningActorHandle = SequencerPtr->GetHandleToObject(OwningActor, bCreateHandleIfMissing);
+				if (OwningActorHandle.IsValid())
+				{
+					TransformTracks.Add(MovieScene->FindTrack<UMovieScene3DTransformTrack>(OwningActorHandle, TransformPropertyName));
+				}
+			}
+
+			FGuid SceneComponentHandle = SequencerPtr->GetHandleToObject(SceneComponent, bCreateHandleIfMissing);
+			if (SceneComponentHandle.IsValid())
+			{
+				TransformTracks.Add(MovieScene->FindTrack<UMovieScene3DTransformTrack>(SceneComponentHandle, TransformPropertyName));
+			}
+		}
+
+		for (UMovieScene3DTransformTrack* TransformTrack : TransformTracks)
+		{
+			if (!TransformTrack || TransformTrack->IsEmpty())
+			{
+				continue;
+			}
+
+			TArray<UMovieSceneSection*, TInlineAllocator<4>> Sections = TransformTrack->FindAllSections(FrameRange.GetLowerBoundValue());
+			for (const UMovieSceneSection* Section : Sections)
+			{
+				const UMovieScene3DTransformSection* TransformSection = Cast<UMovieScene3DTransformSection>(Section);
+				if (!TransformSection)
+				{
+					continue;
+				}
+
+				ProcessedSections.Add(TransformSection);
+
+				EPropertyKeyedStatus NewKeyedStatus = GetKeyedStatusInSection(*TransformSection, FrameRange, TransformChannel, ChannelIndices);
+				KeyedStatus = FMath::Max(KeyedStatus, NewKeyedStatus);
+
+				// Maximum Status Reached no need to iterate further
+				if (KeyedStatus == EPropertyKeyedStatus::KeyedInFrame)
+				{
+					return KeyedStatus;
+				}
+			}
+
+			if (KeyedStatus == EPropertyKeyedStatus::NotKeyed)
+			{
+				TransformTracksToCheck.Add(TransformTrack);
+			}
+		}		
+	}
+
+	// If there's no key in the provided sections look through all sections of the tracks
+	// And return "KeyedInOtherFrame" as soon as there's a keyed section
+	if (KeyedStatus == EPropertyKeyedStatus::NotKeyed)
+	{
+		for (const UMovieScene3DTransformTrack* TransformTrack : TransformTracksToCheck)
+		{
+			if (!TransformTrack)
+			{
+				continue;
+			}
+
+			for (const UMovieSceneSection* Section : TransformTrack->GetAllSections())
+			{
+				const UMovieScene3DTransformSection* TransformSection = Cast<UMovieScene3DTransformSection>(Section);
+
+				// Skip sections that were already processed
+				if (TransformSection && !ProcessedSections.Contains(TransformSection))
+				{
+					EPropertyKeyedStatus NewKeyedStatus = GetKeyedStatusInSection(*TransformSection, FrameRange, TransformChannel, ChannelIndices);
+					KeyedStatus = FMath::Max(KeyedStatus, NewKeyedStatus);
+
+					// Maximum Status Reached no need to iterate further
+					if (KeyedStatus >= EPropertyKeyedStatus::KeyedInOtherFrame)
+					{
+						return KeyedStatus;
+					}
+				}
+			}
+		}
+	}
+
+	return KeyedStatus;
+}
+
 void F3DTransformTrackEditor::OnTransformPropertyChanged(const FPropertyChangedParams& PropertyChangedParams, EMovieSceneTransformChannel TransformChannel)
 {
 	// Key Property sends in one object at a time
diff --git a/Engine/Source/Editor/MovieSceneTools/Private/TrackEditors/TransformTrackEditor.h b/Engine/Source/Editor/MovieSceneTools/Private/TrackEditors/TransformTrackEditor.h
--- a/Engine/Source/Editor/MovieSceneTools/Private/TrackEditors/TransformTrackEditor.h
+++ b/Engine/Source/Editor/MovieSceneTools/Private/TrackEditors/TransformTrackEditor.h
@@ -72,6 +72,9 @@ public:
 	virtual void PostRedo(bool bSuccess) override { PostUndo(bSuccess); }
 
 private:
+	EPropertyKeyedStatus GetKeyedStatusInSection(const UMovieScene3DTransformSection& Section, const TRange<FFrameNumber>& Range, EMovieSceneTransformChannel TransformChannel, TConstArrayView<int32> ChannelIndices) const;
+
+	EPropertyKeyedStatus GetPropertyKeyedStatus(const IPropertyHandle& PropertyHandle, EMovieSceneTransformChannel TransformChannel) const;
 
 	void OnTransformPropertyChanged(const FPropertyChangedParams& PropertyChangedParams, EMovieSceneTransformChannel TransformChannel);
 
diff --git a/Engine/Source/Editor/PropertyEditor/Public/IDetailKeyframeHandler.h b/Engine/Source/Editor/PropertyEditor/Public/IDetailKeyframeHandler.h
--- a/Engine/Source/Editor/PropertyEditor/Public/IDetailKeyframeHandler.h
+++ b/Engine/Source/Editor/PropertyEditor/Public/IDetailKeyframeHandler.h
@@ -6,6 +6,25 @@
 
 class IPropertyHandle;
 
+/**
+ * Keyed Status of a Property, ordered by their precedence
+ * E.g. if a property is keyed in current and another frame, Keyed In Frame should take precedence
+ */
+enum class EPropertyKeyedStatus : uint8
+{
+	/** Property not keyed, or not animated in current time */
+	NotKeyed,
+
+	/** Property is animated in the current time, but there's no key in current frame */
+	KeyedInOtherFrame,
+
+	/** For Top-level properties only -- some but not all the properties of the struct are keyed in Current Frame */
+	PartiallyKeyed,
+
+	/** Property (and all its sub-properties) are keyed in current frame */
+	KeyedInFrame,
+};
+
 class IDetailKeyframeHandler
 {
 public:
@@ -19,4 +38,5 @@ public:
 
 	virtual bool IsPropertyAnimated(const IPropertyHandle& PropertyHandle, UObject* ParentObject) const = 0;
 
+	virtual EPropertyKeyedStatus GetPropertyKeyedStatus(const IPropertyHandle& PropertyHandle) const { return EPropertyKeyedStatus::NotKeyed; }
 };
diff --git a/Engine/Source/Editor/Sequencer/Private/LevelEditorSequencerIntegration.cpp b/Engine/Source/Editor/Sequencer/Private/LevelEditorSequencerIntegration.cpp
--- a/Engine/Source/Editor/Sequencer/Private/LevelEditorSequencerIntegration.cpp
+++ b/Engine/Source/Editor/Sequencer/Private/LevelEditorSequencerIntegration.cpp
@@ -156,6 +156,20 @@ public:
 		}
 	}
 
+	virtual EPropertyKeyedStatus GetPropertyKeyedStatus(const IPropertyHandle& PropertyHandle) const override
+	{
+		EPropertyKeyedStatus KeyedStatus = EPropertyKeyedStatus::NotKeyed;
+		for (const TWeakPtr<ISequencer>& WeakSequencer : Sequencers)
+		{
+			if (TSharedPtr<ISequencer> Sequencer = WeakSequencer.Pin())
+			{
+				EPropertyKeyedStatus NewKeyedStatus = Sequencer->GetPropertyKeyedStatus(PropertyHandle);
+				KeyedStatus = FMath::Max(KeyedStatus, NewKeyedStatus);
+			}
+		}
+		return KeyedStatus;
+	}
+
 private:
 	TArray<TWeakPtr<ISequencer>> Sequencers;
 };
diff --git a/Engine/Source/Editor/Sequencer/Private/Sequencer.cpp b/Engine/Source/Editor/Sequencer/Private/Sequencer.cpp
--- a/Engine/Source/Editor/Sequencer/Private/Sequencer.cpp
+++ b/Engine/Source/Editor/Sequencer/Private/Sequencer.cpp
@@ -73,6 +73,7 @@
 #include "SSequencerSection.h"
 #include "SequencerKeyCollection.h"
 #include "SequencerAddKeyOperation.h"
+#include "SequencerPropertyKeyedStatus.h"
 #include "SequencerSettings.h"
 #include "SequencerLog.h"
 #include "SequencerEdMode.h"
@@ -448,6 +449,7 @@ void FSequencer::InitSequencer(const FSequencerInitParams& InitParams, const TSh
 	PlaybackSpeedBeforePlay = PlaybackSpeed;
 	ShuttleMultiplier = 0;
 	ObjectChangeListener = InObjectChangeListener;
+	PropertyKeyedStatusHandler = MakeShared<FSequencerPropertyKeyedStatusHandler>(SharedThis(this));
 
 	check( ObjectChangeListener.IsValid() );
 	
@@ -3849,6 +3851,11 @@ ISequencerObjectChangeListener& FSequencer::GetObjectChangeListener()
 	return *ObjectChangeListener;
 }
 
+ISequencerPropertyKeyedStatusHandler& FSequencer::GetPropertyKeyedStatusHandler()
+{
+	return *PropertyKeyedStatusHandler;
+}
+
 TSharedPtr<class ITimeSlider> FSequencer::GetTopTimeSliderWidget() const
 {
 	return SequencerWidget->GetTopTimeSliderWidget();
@@ -6502,7 +6509,8 @@ void FSequencer::ZoomToFit()
 
 bool FSequencer::CanKeyProperty(FCanKeyPropertyParams CanKeyPropertyParams) const
 {
-	return ObjectChangeListener->CanKeyProperty(CanKeyPropertyParams);
+	FPropertyPath PropertyPath;
+	return ObjectChangeListener->CanKeyProperty(CanKeyPropertyParams, PropertyPath);
 } 
 
 
@@ -6511,6 +6519,10 @@ void FSequencer::KeyProperty(FKeyPropertyParams KeyPropertyParams)
 	ObjectChangeListener->KeyProperty(KeyPropertyParams);
 }
 
+EPropertyKeyedStatus FSequencer::GetPropertyKeyedStatus(const IPropertyHandle& PropertyHandle) const
+{
+	return PropertyKeyedStatusHandler->GetPropertyKeyedStatus(PropertyHandle);
+}
 
 FSequencerSelectionPreview& FSequencer::GetSelectionPreview()
 {
diff --git a/Engine/Source/Editor/Sequencer/Private/Sequencer.h b/Engine/Source/Editor/Sequencer/Private/Sequencer.h
--- a/Engine/Source/Editor/Sequencer/Private/Sequencer.h
+++ b/Engine/Source/Editor/Sequencer/Private/Sequencer.h
@@ -47,6 +47,7 @@ class APlayerController;
 class FLevelEditorViewportClient;
 class FMenuBuilder;
 class FMovieSceneClipboard;
+class FSequencerPropertyKeyedStatusHandler;
 class FViewportClient;
 class IDetailKeyframeHandler;
 class IAssetViewport;
@@ -736,6 +737,8 @@ public:
 	virtual bool IsInSilentMode() const override { return SilentModeCount != 0; }
 	virtual FGuid GetHandleToObject(UObject* Object, bool bCreateHandleIfMissing = true, const FName& CreatedFolderName = NAME_None) override;
 	virtual ISequencerObjectChangeListener& GetObjectChangeListener() override;
+	virtual ISequencerPropertyKeyedStatusHandler& GetPropertyKeyedStatusHandler() override;
+
 protected:
 	virtual void NotifyMovieSceneDataChangedInternal() override;
 public:
@@ -749,6 +752,7 @@ public:
 	virtual void AddSubSequence(UMovieSceneSequence* Sequence) override;
 	virtual bool CanKeyProperty(FCanKeyPropertyParams CanKeyPropertyParams) const override;
 	virtual void KeyProperty(FKeyPropertyParams KeyPropertyParams) override;
+	EPropertyKeyedStatus GetPropertyKeyedStatus(const IPropertyHandle& PropertyHandle) const override;
 	virtual void GetSelectedTracks(TArray<UMovieSceneTrack*>& OutSelectedTracks) override;
 	virtual void GetSelectedSections(TArray<UMovieSceneSection*>& OutSelectedSections) override;
 	virtual void GetSelectedFolders(TArray<UMovieSceneFolder*>& OutSelectedFolders) override;
@@ -1174,6 +1178,9 @@ private:
 	/** Listener for object changes being made while this sequencer is open*/
 	TSharedPtr<ISequencerObjectChangeListener> ObjectChangeListener;
 
+	/** Responsible for getting the keyed status of a property (whether it's keyed in current frame, other frame, etc)  */
+	TSharedPtr<FSequencerPropertyKeyedStatusHandler> PropertyKeyedStatusHandler;
+
 	/** Main sequencer widget */
 	TSharedPtr<SSequencer> SequencerWidget;
 
diff --git a/Engine/Source/Editor/Sequencer/Private/SequencerModule.cpp b/Engine/Source/Editor/Sequencer/Private/SequencerModule.cpp
--- a/Engine/Source/Editor/Sequencer/Private/SequencerModule.cpp
+++ b/Engine/Source/Editor/Sequencer/Private/SequencerModule.cpp
@@ -53,6 +53,16 @@
 	UE::MovieScene::FEntityManager*& GEntityManagerForDebugging = UE::MovieScene::GEntityManagerForDebuggingVisualizers;
 #endif
 
+namespace UE::Sequencer::Private
+{
+	static const TMap<EPropertyKeyedStatus, FName> KeyedStatusStyleNames =
+		{
+			{ EPropertyKeyedStatus::NotKeyed, "Sequencer.KeyedStatus.NotKeyed" },
+			{ EPropertyKeyedStatus::KeyedInOtherFrame, "Sequencer.KeyedStatus.Animated" },
+			{ EPropertyKeyedStatus::KeyedInFrame, "Sequencer.KeyedStatus.Keyed" },
+			{ EPropertyKeyedStatus::PartiallyKeyed, "Sequencer.KeyedStatus.PartialKey" },
+		};
+}
 
 
 #define LOCTEXT_NAMESPACE "SequencerEditor"
@@ -85,6 +95,23 @@ static TSharedPtr<IDetailKeyframeHandler> GetKeyframeHandler(TWeakPtr<IDetailTre
 	return DetailsView->GetKeyframeHandler();
 }
 
+static FSlateIcon GetKeyframeIcon(TWeakPtr<IDetailTreeNode> OwnerTreeNode, TSharedPtr<IPropertyHandle> PropertyHandle)
+{
+	if (!PropertyHandle.IsValid())
+	{
+		return FSlateIcon();
+	}
+
+	EPropertyKeyedStatus KeyedStatus = EPropertyKeyedStatus::NotKeyed;
+
+	if (TSharedPtr<IDetailKeyframeHandler> KeyframeHandler = GetKeyframeHandler(OwnerTreeNode))
+	{
+		KeyedStatus = KeyframeHandler->GetPropertyKeyedStatus(*PropertyHandle);
+	}
+
+	return FSlateIcon(FAppStyle::GetAppStyleSetName(), UE::Sequencer::Private::KeyedStatusStyleNames[KeyedStatus]);
+}
+
 static bool IsKeyframeButtonVisible(TWeakPtr<IDetailTreeNode> OwnerTreeNode, TSharedPtr<IPropertyHandle> PropertyHandle)
 {
 	TSharedPtr<IDetailKeyframeHandler> KeyframeHandler = GetKeyframeHandler(OwnerTreeNode);
@@ -133,12 +160,11 @@ static void RegisterKeyframeExtensionHandler(const FOnGenerateGlobalRowExtension
 		return;
 	}
 
-	static FSlateIcon CreateKeyIcon(FAppStyle::Get().GetStyleSetName(), "Sequencer.AddKey.Details");
-
 	TWeakPtr<IDetailTreeNode> OwnerTreeNode = Args.OwnerTreeNode;
 
 	FPropertyRowExtensionButton& CreateKey = OutExtensionButtons.AddDefaulted_GetRef();
-	CreateKey.Icon = CreateKeyIcon;
+
+	CreateKey.Icon = TAttribute<FSlateIcon>::Create(TAttribute<FSlateIcon>::FGetter::CreateStatic(&GetKeyframeIcon, OwnerTreeNode, PropertyHandle));
 	CreateKey.Label = NSLOCTEXT("PropertyEditor", "CreateKey", "Create Key");
 	CreateKey.ToolTip = NSLOCTEXT("PropertyEditor", "CreateKeyToolTip", "Add a keyframe for this property.");
 	CreateKey.UIAction = FUIAction(
diff --git a/Engine/Source/Editor/Sequencer/Private/SequencerObjectChangeListener.cpp b/Engine/Source/Editor/Sequencer/Private/SequencerObjectChangeListener.cpp
--- a/Engine/Source/Editor/Sequencer/Private/SequencerObjectChangeListener.cpp
+++ b/Engine/Source/Editor/Sequencer/Private/SequencerObjectChangeListener.cpp
@@ -274,6 +274,13 @@ bool FSequencerObjectChangeListener::CanKeyProperty(FCanKeyPropertyParams CanKey
 	return CanKeyProperty_Internal(CanKeyPropertyParams, Delegate, Property, PropertyPath);
 }
 
+bool FSequencerObjectChangeListener::CanKeyProperty(FCanKeyPropertyParams KeyPropertyParams, FPropertyPath& OutPropertyPath) const
+{
+	FProperty* Property = nullptr;
+	FOnAnimatablePropertyChanged Delegate;
+	return CanKeyProperty_Internal(KeyPropertyParams, Delegate, Property, OutPropertyPath);
+}
+
 bool FSequencerObjectChangeListener::CanKeyProperty_Internal(FCanKeyPropertyParams CanKeyPropertyParams, FOnAnimatablePropertyChanged& InOutDelegate, FProperty*& InOutProperty, FPropertyPath& InOutPropertyPath) const
 {
 	if (CanKeyPropertyParams.PropertyPath.GetNumProperties() == 0)
diff --git a/Engine/Source/Editor/Sequencer/Private/SequencerObjectChangeListener.h b/Engine/Source/Editor/Sequencer/Private/SequencerObjectChangeListener.h
--- a/Engine/Source/Editor/Sequencer/Private/SequencerObjectChangeListener.h
+++ b/Engine/Source/Editor/Sequencer/Private/SequencerObjectChangeListener.h
@@ -30,7 +30,9 @@ public:
 	virtual FOnPropagateObjectChanges& GetOnPropagateObjectChanges() override;
 	virtual FOnObjectPropertyChanged& GetOnAnyPropertyChanged(UObject& Object) override;
 	virtual void ReportObjectDestroyed(UObject& Object) override;
+	UE_DEPRECATED(5.4, "CanKeyProperty taking in only FCanKeyPropertyParams is deprecated. Please use the alternative CanKeyProperty that takes in FCanKeyPropertyParams and an out FPropertyPath.")
 	virtual bool CanKeyProperty(FCanKeyPropertyParams KeyPropertyParams) const override;
+	virtual bool CanKeyProperty(FCanKeyPropertyParams KeyPropertyParams, FPropertyPath& OutPropertyPath) const override;
 	virtual void KeyProperty(FKeyPropertyParams KeyPropertyParams) const override;
 	virtual void TriggerAllPropertiesChanged(UObject* Object) override;
 
diff --git a/Engine/Source/Editor/Sequencer/Private/SequencerPropertyKeyedStatus.cpp b/Engine/Source/Editor/Sequencer/Private/SequencerPropertyKeyedStatus.cpp
new file mode 100644
--- /dev/null
+++ b/Engine/Source/Editor/Sequencer/Private/SequencerPropertyKeyedStatus.cpp
@@ -0,0 +1,445 @@
+// Copyright Epic Games, Inc. All Rights Reserved.
+
+#include "SequencerPropertyKeyedStatus.h"
+#include "IDetailKeyframeHandler.h"
+#include "ISequencer.h"
+#include "ISequencerObjectChangeListener.h"
+#include "MovieScene.h"
+#include "PropertyHandle.h"
+#include "Tracks/MovieScenePropertyTrack.h"
+
+FSequencerPropertyKeyedStatusHandler::FPropertyParameters::FPropertyParameters(ISequencer& InSequencer, const IPropertyHandle& ActualProperty)
+	: Sequencer(InSequencer)
+	, ActualProperty(ActualProperty)
+{
+	CurrentFrameRange = TRange<FFrameNumber>(Sequencer.GetLocalTime().Time.FrameNumber);
+
+	FPropertyPath PropertyPath = BuildPropertyPath();
+	TrackPropertyName = FindTrackPropertyName(PropertyPath);
+	SubPropertyPath = BuildSubPropertyPath();
+	TrackPropertyStructName = FindTrackPropertyStructName(PropertyPath);
+}
+
+FPropertyPath FSequencerPropertyKeyedStatusHandler::FPropertyParameters::BuildPropertyPath() const
+{
+	TSharedRef<const IPropertyHandle> CurrentHandle = ActualProperty.AsShared();
+	TSharedPtr<const IPropertyHandle> ParentHandle = ActualProperty.GetParentHandle();
+
+	TArray<FPropertyInfo> PropertyInfos;
+	PropertyInfos.Emplace(CurrentHandle->GetProperty(), CurrentHandle->GetArrayIndex());
+
+	while (ParentHandle.IsValid() && ParentHandle->GetProperty())
+	{
+		PropertyInfos.Emplace(ParentHandle->GetProperty(), ParentHandle->GetArrayIndex());
+		CurrentHandle = ParentHandle.ToSharedRef();
+		ParentHandle = ParentHandle->GetParentHandle();
+	}
+
+	FPropertyPath PropertyPath;
+	for (const FPropertyInfo& PropertyInfo : ReverseIterate(PropertyInfos))
+	{
+		PropertyPath.AddProperty(PropertyInfo);
+	}
+
+	return PropertyPath;
+}
+
+FString FSequencerPropertyKeyedStatusHandler::FPropertyParameters::FindTrackPropertyName(FPropertyPath& InOutPropertyPath) const
+{
+	const UClass* ObjectClass = ActualProperty.GetOuterBaseClass();
+	Sequencer.GetObjectChangeListener().CanKeyProperty(FCanKeyPropertyParams(ObjectClass, InOutPropertyPath), InOutPropertyPath);
+	return InOutPropertyPath.ToString(TEXT("."));
+}
+
+FName FSequencerPropertyKeyedStatusHandler::FPropertyParameters::FindTrackPropertyStructName(const FPropertyPath& PropertyPath)
+{
+	if (PropertyPath.GetNumProperties() > 0)
+	{
+		FStructProperty* StructProperty = CastField<FStructProperty>(PropertyPath.GetLeafMostProperty().Property.Get());
+		if (StructProperty && StructProperty->Struct)
+		{
+			return StructProperty->Struct->GetFName();
+		}
+	}
+	return NAME_None;
+}
+
+FName FSequencerPropertyKeyedStatusHandler::FPropertyParameters::BuildSubPropertyPath() const
+{
+	// Property Path shows up with the format "Euler->Location->X".
+	// This should be converted to "Euler.Location.X"
+	FString PropertyPath = FString(ActualProperty.GetPropertyPath()).Replace(TEXT("->"), TEXT("."));
+
+	int32 Index = PropertyPath.Find(TrackPropertyName);
+	if (Index != INDEX_NONE)
+	{
+		constexpr bool bAllowShrinking = false;
+
+		// Remove the Track Property name from the property path.
+		// From example above, "Euler.Location.X" should be changed to "Location.X" to obtain the sub-property path relative to the Track property
+		// If the Actual Property is the Track Property, SubPropertyPath should be left empty
+		// Additionally, remove an extra character for the dot (i.e. remove first dot in ".Location.X"). Safe to call as it gets clamped out in RightChop
+		PropertyPath.RightChopInline(Index + TrackPropertyName.Len() + 1, bAllowShrinking);
+	}
+
+	return *PropertyPath;
+}
+
+FSequencerPropertyKeyedStatusHandler::FSequencerPropertyKeyedStatusHandler(TSharedRef<ISequencer> InSequencer)
+	: SequencerWeak(InSequencer)
+{
+	InSequencer->OnMovieSceneDataChanged().AddRaw(this, &FSequencerPropertyKeyedStatusHandler::OnMovieSceneDataChanged);
+	InSequencer->OnGlobalTimeChanged().AddRaw(this, &FSequencerPropertyKeyedStatusHandler::OnGlobalTimeChanged);
+	InSequencer->OnEndScrubbingEvent().AddRaw(this, &FSequencerPropertyKeyedStatusHandler::ResetCachedData);
+	InSequencer->OnChannelChanged().AddRaw(this, &FSequencerPropertyKeyedStatusHandler::OnChannelChanged);
+	InSequencer->OnStopEvent().AddRaw(this, &FSequencerPropertyKeyedStatusHandler::ResetCachedData);
+}
+
+FSequencerPropertyKeyedStatusHandler::~FSequencerPropertyKeyedStatusHandler()
+{
+	if (TSharedPtr<ISequencer> Sequencer = SequencerWeak.Pin())
+	{
+		Sequencer->OnMovieSceneDataChanged().RemoveAll(this);
+		Sequencer->OnGlobalTimeChanged().RemoveAll(this);
+		Sequencer->OnEndScrubbingEvent().RemoveAll(this);
+		Sequencer->OnChannelChanged().RemoveAll(this);
+		Sequencer->OnStopEvent().RemoveAll(this);
+	}
+}
+
+EPropertyKeyedStatus FSequencerPropertyKeyedStatusHandler::GetPropertyKeyedStatus(const IPropertyHandle& PropertyHandle) const
+{
+	if (const EPropertyKeyedStatus* ExistingKeyedStatus = CachedPropertyKeyedStatusMap.Find(&PropertyHandle))
+	{
+		return *ExistingKeyedStatus;
+	}
+
+	EPropertyKeyedStatus KeyedStatus;
+
+	const FOnGetPropertyKeyedStatus* ExternalHandler = ExternalHandlers.Find(PropertyHandle.GetProperty());
+	if (ExternalHandler && ExternalHandler->IsBound())
+	{
+		KeyedStatus = ExternalHandler->Execute(PropertyHandle);
+	}
+	else
+	{
+		KeyedStatus = CalculatePropertyKeyedStatus(PropertyHandle);
+	}
+
+	CachedPropertyKeyedStatusMap.Add(&PropertyHandle, KeyedStatus);
+	return KeyedStatus;
+}
+
+ISequencerPropertyKeyedStatusHandler::FOnGetPropertyKeyedStatus& FSequencerPropertyKeyedStatusHandler::GetExternalHandler(const FProperty* Property)
+{
+	return ExternalHandlers.FindOrAdd(Property);
+}
+
+void FSequencerPropertyKeyedStatusHandler::ResetCachedData()
+{
+	CachedPropertyKeyedStatusMap.Reset();
+}
+
+void FSequencerPropertyKeyedStatusHandler::OnGlobalTimeChanged()
+{
+	// Only reset cached data when not playing
+	TSharedPtr<ISequencer> Sequencer = SequencerWeak.Pin();
+	if (Sequencer.IsValid() && Sequencer->GetPlaybackStatus() != EMovieScenePlayerStatus::Playing)
+	{
+		ResetCachedData();
+	}
+}
+
+void FSequencerPropertyKeyedStatusHandler::OnMovieSceneDataChanged(EMovieSceneDataChangeType DataChangeType)
+{
+	if (DataChangeType == EMovieSceneDataChangeType::MovieSceneStructureItemAdded
+		|| DataChangeType == EMovieSceneDataChangeType::MovieSceneStructureItemRemoved
+		|| DataChangeType == EMovieSceneDataChangeType::MovieSceneStructureItemsChanged
+		|| DataChangeType == EMovieSceneDataChangeType::ActiveMovieSceneChanged
+		|| DataChangeType == EMovieSceneDataChangeType::RefreshAllImmediately)
+	{
+		ResetCachedData();
+	}
+}
+
+void FSequencerPropertyKeyedStatusHandler::OnChannelChanged(const FMovieSceneChannelMetaData*, UMovieSceneSection*)
+{
+	ResetCachedData();
+}
+
+EPropertyKeyedStatus FSequencerPropertyKeyedStatusHandler::CalculatePropertyKeyedStatus(const IPropertyHandle& PropertyHandle) const
+{
+	TRACE_CPUPROFILER_EVENT_SCOPE(FSequencerPropertyKeyedStatus::CalculatePropertyKeyedStatus);
+
+	TSharedPtr<ISequencer> Sequencer = SequencerWeak.Pin();
+	if (!Sequencer.IsValid())
+	{
+		return EPropertyKeyedStatus::NotKeyed;
+	}
+
+	UMovieSceneSequence* Sequence = Sequencer->GetFocusedMovieSceneSequence();
+	if (!Sequence)
+	{
+		return EPropertyKeyedStatus::NotKeyed;
+	}
+
+	UMovieScene* MovieScene = Sequence->GetMovieScene();
+	if (!MovieScene)
+	{
+		return EPropertyKeyedStatus::NotKeyed;
+	}
+
+	TArray<UObject*> OuterObjects;
+	PropertyHandle.GetOuterObjects(OuterObjects);
+	if (OuterObjects.IsEmpty())
+	{
+		return EPropertyKeyedStatus::NotKeyed;
+	}
+
+	const FPropertyParameters Parameters(*Sequencer, PropertyHandle);
+
+	EPropertyKeyedStatus KeyedStatus = EPropertyKeyedStatus::NotKeyed;
+
+	// List of Tracks that had no sections at the current frame that require an additional check by going through all its sections and whether there's a key in any of them.
+	// Used only to determine whether to return "Not Keyed" or "Keyed In Other Frame".
+	TArray<const UMovieScenePropertyTrack*> TracksToCheck;
+	TracksToCheck.Reserve(OuterObjects.Num());
+
+	TSet<const UMovieSceneSection*> ProcessedSections;
+	ProcessedSections.Reserve(OuterObjects.Num());
+
+	for (UObject* Object : OuterObjects)
+	{
+		// Object can be null here, but is handled gracefully by GetHandleToObject
+		constexpr bool bCreateHandleIfMissing = false;
+		FGuid ObjectHandle = Parameters.Sequencer.GetHandleToObject(Object, bCreateHandleIfMissing);
+		if (!ObjectHandle.IsValid())
+		{
+			continue;
+		}
+
+		UMovieScenePropertyTrack* PropertyTrack = MovieScene->FindTrack<UMovieScenePropertyTrack>(ObjectHandle, *Parameters.TrackPropertyName);
+		if (!PropertyTrack || PropertyTrack->IsEmpty())
+		{
+			continue;
+		}
+
+		TArray<UMovieSceneSection*, TInlineAllocator<5>> Sections;
+		Sections = PropertyTrack->FindAllSections(Parameters.CurrentFrameRange.GetLowerBoundValue());
+
+		UMovieSceneSection* const SectionToKey = PropertyTrack->GetSectionToKey();
+
+		// if Section to Key is valid, ensure it is in Sections as the first element, as it takes precedence (and returns immediately) when it exists
+		if (SectionToKey)
+		{
+			int32 Index = Sections.Find(SectionToKey);
+			if (Index == INDEX_NONE)
+			{
+				Index = Sections.Add(SectionToKey);
+			}
+
+			if (Index != INDEX_NONE)
+			{
+				// section order shouldn't matter
+				Sections.Swap(0, Index);	
+			}
+		}
+
+		for (const UMovieSceneSection* Section : Sections)
+		{
+			check(Section);
+			ProcessedSections.Add(Section);
+
+			constexpr EPropertyKeyedStatus MaxStatus = EPropertyKeyedStatus::KeyedInFrame;
+			EPropertyKeyedStatus SectionKeyedStatus = GetKeyedStatusInSection(Parameters, *Section, MaxStatus);
+
+			// If the Section is the Section to Key, prioritize it and ignore the rest
+			if (Section == SectionToKey)
+			{
+				return SectionKeyedStatus;
+			}
+
+			KeyedStatus = FMath::Max(KeyedStatus, SectionKeyedStatus);
+
+			// Return if max status already reached
+			if (KeyedStatus >= MaxStatus)
+			{
+				return KeyedStatus;
+			}
+		}
+
+		// If this Track had no keys in any of its sections,
+		// add it to check for all sections (as this only checked for sections in the current time)
+		if (KeyedStatus == EPropertyKeyedStatus::NotKeyed)
+		{
+			TracksToCheck.Add(PropertyTrack);
+		}
+	}
+
+	// If there's no key in the provided sections look through all sections of the tracks
+	// And return "KeyedInOtherFrame" as soon as there's a keyed section
+	if (KeyedStatus == EPropertyKeyedStatus::NotKeyed)
+	{
+		for (const UMovieScenePropertyTrack* PropertyTrack : TracksToCheck)
+		{
+			if (!PropertyTrack)
+			{
+				continue;
+			}
+
+			for (const UMovieSceneSection* Section : PropertyTrack->GetAllSections())
+			{
+				// Skip sections that were already processed
+				if (Section && !ProcessedSections.Contains(Section))
+				{
+					constexpr EPropertyKeyedStatus MaxStatus = EPropertyKeyedStatus::KeyedInOtherFrame;
+					EPropertyKeyedStatus SectionKeyedStatus = GetKeyedStatusInSection(Parameters, *Section, MaxStatus);
+
+					KeyedStatus = FMath::Max(KeyedStatus, SectionKeyedStatus);
+
+					// Maximum Status Reached no need to iterate further
+					if (KeyedStatus >= MaxStatus)
+					{
+						return KeyedStatus;
+					}
+				}
+			}
+		}
+	}
+
+	return KeyedStatus;
+}
+
+EPropertyKeyedStatus FSequencerPropertyKeyedStatusHandler::GetKeyedStatusInSection(const FPropertyParameters& Parameters, const UMovieSceneSection& Section, EPropertyKeyedStatus MaxPropertyKeyedStatus) const
+{
+	EPropertyKeyedStatus KeyedStatus = EPropertyKeyedStatus::NotKeyed;
+
+	int32 EmptyChannelCount = 0;
+
+	for (const FMovieSceneChannelEntry& ChannelEntry : Section.GetChannelProxy().GetAllEntries())
+	{
+		TConstArrayView<FMovieSceneChannel*> Channels = ChannelEntry.GetChannels();
+
+		TConstArrayView<FMovieSceneChannelMetaData> ChannelMetaData = ChannelEntry.GetMetaData();
+
+		for (int32 ChannelIndex = 0; ChannelIndex < Channels.Num(); ++ChannelIndex)
+		{
+			if (KeyedStatus >= MaxPropertyKeyedStatus)
+			{
+				return KeyedStatus;
+			}
+
+			FMovieSceneChannel* const Channel = Channels[ChannelIndex];
+			if (!Channel)
+			{
+				continue;
+			}
+
+			const FMovieSceneChannelMetaData& MetaData = ChannelMetaData[ChannelIndex];
+
+			FName SubPropertyPath = MetaData.SubPropertyPath;
+			if (SubPropertyPath.IsNone())
+			{
+				// If sub property path is none, try finding the entry from the sub-property path map
+				SubPropertyPath = MetaData.SubPropertyPathMap.FindRef(Parameters.TrackPropertyStructName);
+			}
+
+			// Skip Channel if there was a Struct name for the Track (expected a valid Sub-property path) but didn't find one
+			const bool bSkipChannel = !Parameters.TrackPropertyStructName.IsNone() && SubPropertyPath.IsNone();
+
+			if (bSkipChannel || !HasMatchingProperty(Parameters, SubPropertyPath))
+			{
+				continue;
+			}
+
+			// Check if Channel is disabled, as these could still have Keys, but should be viewed as empty
+			if (Channel->GetNumKeys() == 0 || !MetaData.bEnabled)
+			{
+				++EmptyChannelCount;
+				continue;
+			}
+
+			// There's at least 1 key in this Channel, so status is at least KeyedInOtherFrame
+			KeyedStatus = FMath::Max(KeyedStatus, EPropertyKeyedStatus::KeyedInOtherFrame);
+			if (KeyedStatus >= MaxPropertyKeyedStatus)
+			{
+				return KeyedStatus;
+			}
+
+			TArray<FFrameNumber> KeyTimes;
+			Channel->GetKeys(Parameters.CurrentFrameRange, &KeyTimes, nullptr);
+			if (KeyTimes.IsEmpty())
+			{
+				++EmptyChannelCount;
+			}
+			else
+			{
+				// Key Times found, updated the KeyedStatus
+				KeyedStatus = FMath::Max(KeyedStatus, EPropertyKeyedStatus::PartiallyKeyed);
+			}
+		}
+	}
+
+	if (KeyedStatus == EPropertyKeyedStatus::PartiallyKeyed && EmptyChannelCount == 0)
+	{
+		KeyedStatus = EPropertyKeyedStatus::KeyedInFrame;
+	}
+	return KeyedStatus;
+}
+
+bool FSequencerPropertyKeyedStatusHandler::HasMatchingProperty(const FPropertyParameters& Parameters, FName SubPropertyPath) const
+{
+	FProperty* Property = Parameters.ActualProperty.GetProperty();
+	if (Property->GetFName() == SubPropertyPath || Parameters.SubPropertyPath == SubPropertyPath)
+	{
+		return true;
+	}
+
+	// At this point, the actual property is NOT the channel's property,
+	// but the actual property could be a parent of the channel's property -- so iterate downward and see if there's a match.
+
+	// If it isn't a struct property, it can't iterate further down, return no match
+	const FStructProperty* StructProperty = CastField<FStructProperty>(Property);
+	if (!StructProperty)
+	{
+		return false;
+	}
+
+	// Return early no match if the Property Path does not even begin with the pre-calculated Sub Property Path
+	if (!Parameters.SubPropertyPath.IsNone() && !SubPropertyPath.ToString().StartsWith(Parameters.SubPropertyPath.ToString()))
+	{
+		return false;
+	}
+
+	TArray<FString> PropertySegments;
+	SubPropertyPath.ToString().ParseIntoArray(PropertySegments, TEXT("."));
+	return HasMatchingSubProperty(*StructProperty, PropertySegments);
+}
+
+bool FSequencerPropertyKeyedStatusHandler::HasMatchingSubProperty(const FStructProperty& StructProperty, TConstArrayView<FString> InPropertySegments) const
+{
+	for (const FProperty* Property : TFieldRange<FProperty>(StructProperty.Struct))
+	{
+		int32 SegmentIndex = InPropertySegments.Find(Property->GetName());
+		if (SegmentIndex == INDEX_NONE)
+		{
+			continue;
+		}
+
+		// If it's last index, then the property is found
+		if (SegmentIndex == InPropertySegments.Num() - 1)
+		{
+			return true;
+		}
+
+		// If it's not the last segment that matches, there's more to match, recurse downward
+		if (const FStructProperty* SubStructProperty = CastField<FStructProperty>(Property))
+		{
+			if (HasMatchingSubProperty(*SubStructProperty, InPropertySegments.RightChop(SegmentIndex + 1)))
+			{
+				return true;
+			}
+		}
+	}
+	return false;
+}
diff --git a/Engine/Source/Editor/Sequencer/Private/SequencerPropertyKeyedStatus.h b/Engine/Source/Editor/Sequencer/Private/SequencerPropertyKeyedStatus.h
new file mode 100644
--- /dev/null
+++ b/Engine/Source/Editor/Sequencer/Private/SequencerPropertyKeyedStatus.h
@@ -0,0 +1,94 @@
+// Copyright Epic Games, Inc. All Rights Reserved.
+
+#pragma once
+
+#include "Containers/Array.h"
+#include "Containers/ArrayView.h"
+#include "Containers/UnrealString.h"
+#include "ISequencerPropertyKeyedStatus.h"
+#include "Math/Range.h"
+#include "Templates/SharedPointer.h"
+#include "UObject/NameTypes.h"
+
+class FPropertyPath;
+class FSequencer;
+class FStructProperty;
+class IPropertyHandle;
+class ISequencer;
+class UMovieSceneSection;
+enum class EMovieSceneDataChangeType;
+enum class EPropertyKeyedStatus : uint8;
+struct FMovieSceneChannelMetaData;
+
+class FSequencerPropertyKeyedStatusHandler : public ISequencerPropertyKeyedStatusHandler
+{
+	struct FPropertyParameters
+	{
+		explicit FPropertyParameters(ISequencer& InSequencer, const IPropertyHandle& ActualProperty);
+
+		ISequencer& Sequencer;
+
+		/** the actual Property whose keyed status is being queried */
+		const IPropertyHandle& ActualProperty;
+
+		/** the name of the property that the Property Track is made for. Could be the same as the ActualProperty name, or the name of a parent property (e.g. a Struct containing the Actual Property) */
+		FString TrackPropertyName;
+
+		/** The possible underlying struct for the Track's Property, can be none */
+		FName TrackPropertyStructName;
+
+		/**
+		 * The Sub Property Path going from TrackProperty->ActualProperty.
+		 * If the Actual Property is the Track Property, this is Empty
+		 * Else, if Track's Property is "A", and ActualProperty is in the form "A->B->C", the generated SubPropertyPath becomes "B.C"
+		 */
+		FName SubPropertyPath;
+
+		/** Current Sequencer Time, converted into a Frame Range */
+		TRange<FFrameNumber> CurrentFrameRange;
+
+		FPropertyPath BuildPropertyPath() const;
+
+		FString FindTrackPropertyName(FPropertyPath& InOutPropertyPath) const;
+
+		FName FindTrackPropertyStructName(const FPropertyPath& PropertyPath);
+
+		FName BuildSubPropertyPath() const;
+	};
+
+public:
+	explicit FSequencerPropertyKeyedStatusHandler(TSharedRef<ISequencer> InSequencer);
+
+	virtual ~FSequencerPropertyKeyedStatusHandler() override;
+
+	EPropertyKeyedStatus GetPropertyKeyedStatus(const IPropertyHandle& PropertyHandle) const;
+
+	//~ Begin ISequencerPropertyKeyedStatusHandler
+	virtual FOnGetPropertyKeyedStatus& GetExternalHandler(const FProperty* Property) override;
+	//~ End ISequencerPropertyKeyedStatusHandler
+
+private:
+	void ResetCachedData();
+
+	/** Delegates Resetting Cached Data */
+	void OnGlobalTimeChanged();
+	void OnMovieSceneDataChanged(EMovieSceneDataChangeType);
+	void OnChannelChanged(const FMovieSceneChannelMetaData*, UMovieSceneSection*);
+
+	EPropertyKeyedStatus CalculatePropertyKeyedStatus(const IPropertyHandle& PropertyHandle) const;
+
+	EPropertyKeyedStatus GetKeyedStatusInSection(const FPropertyParameters& Parameters, const UMovieSceneSection& Section, EPropertyKeyedStatus MaxPropertyKeyedStatus) const;
+
+	bool HasMatchingProperty(const FPropertyParameters& Parameters, FName SubPropertyPath) const;
+
+	bool HasMatchingSubProperty(const FStructProperty& StructProperty, TConstArrayView<FString> InPropertySegments) const;
+
+	/** Sequencer owning this instance */
+	TWeakPtr<ISequencer> SequencerWeak;
+
+	/** Map of Properties and their delegates used to get the Property Keyed Status via External Logic */
+	TMap<const FProperty*, FOnGetPropertyKeyedStatus> ExternalHandlers;
+
+	/** Map to the last calculated property keyed status. Resets when Scrubbing, changing Movie Scene Data, etc */
+	mutable TMap<const IPropertyHandle*, EPropertyKeyedStatus> CachedPropertyKeyedStatusMap;
+};
diff --git a/Engine/Source/Editor/Sequencer/Public/ISequencer.h b/Engine/Source/Editor/Sequencer/Public/ISequencer.h
--- a/Engine/Source/Editor/Sequencer/Public/ISequencer.h
+++ b/Engine/Source/Editor/Sequencer/Public/ISequencer.h
@@ -67,6 +67,7 @@ struct FQualifiedFrameTime;
 template <typename NumericType> struct INumericTypeInterface;
 
 enum class EMapChangeType : uint8;
+enum class EPropertyKeyedStatus : uint8;
 class FCurveEditor;
 class FCurveModel;
 class IToolkitHost;
@@ -578,10 +579,17 @@ public:
 	 */
 	virtual class ISequencerObjectChangeListener& GetObjectChangeListener() = 0;
 
+	/**
+	 * @return Returns the property keyed status handler for this sequencer instance
+	 */
+	virtual class ISequencerPropertyKeyedStatusHandler& GetPropertyKeyedStatusHandler() = 0;
+
 	virtual bool CanKeyProperty(FCanKeyPropertyParams CanKeyPropertyParams) const = 0;
 
 	virtual void KeyProperty(FKeyPropertyParams KeyPropertyParams) = 0;
 
+	virtual EPropertyKeyedStatus GetPropertyKeyedStatus(const IPropertyHandle& PropertyHandle) const = 0;
+
 	/** Refresh the sequencer tree view */
 	virtual void RefreshTree() = 0;
 protected:
diff --git a/Engine/Source/Editor/Sequencer/Public/ISequencerObjectChangeListener.h b/Engine/Source/Editor/Sequencer/Public/ISequencerObjectChangeListener.h
--- a/Engine/Source/Editor/Sequencer/Public/ISequencerObjectChangeListener.h
+++ b/Engine/Source/Editor/Sequencer/Public/ISequencerObjectChangeListener.h
@@ -48,7 +48,10 @@ public:
 	 */
 	virtual void TriggerAllPropertiesChanged(UObject* Object) = 0;
 
+	UE_DEPRECATED(5.4, "CanKeyProperty taking in only FCanKeyPropertyParams is deprecated. Please use the alternative CanKeyProperty that takes in FCanKeyPropertyParams and an out FPropertyPath.")
 	virtual bool CanKeyProperty(FCanKeyPropertyParams KeyPropertyParams) const = 0;
 
+	virtual bool CanKeyProperty(FCanKeyPropertyParams KeyPropertyParams, FPropertyPath& OutPropertyPath) const = 0;
+
 	virtual void KeyProperty(FKeyPropertyParams KeyPropertyParams) const = 0;
 };
diff --git a/Engine/Source/Editor/Sequencer/Public/ISequencerPropertyKeyedStatus.h b/Engine/Source/Editor/Sequencer/Public/ISequencerPropertyKeyedStatus.h
new file mode 100644
--- /dev/null
+++ b/Engine/Source/Editor/Sequencer/Public/ISequencerPropertyKeyedStatus.h
@@ -0,0 +1,23 @@
+// Copyright Epic Games, Inc. All Rights Reserved.
+
+#pragma once
+
+#include "Delegates/Delegate.h"
+
+class FProperty;
+class IPropertyHandle;
+enum class EPropertyKeyedStatus : uint8;
+
+class ISequencerPropertyKeyedStatusHandler
+{
+public:
+	virtual ~ISequencerPropertyKeyedStatusHandler() = default;
+
+	DECLARE_DELEGATE_RetVal_OneParam(EPropertyKeyedStatus, FOnGetPropertyKeyedStatus, const IPropertyHandle&);
+
+	/**
+	 * Returns the External Handler delegate for a given property
+	 * Used for cases where a property is keyed into a track that does not match the property path so requires custom logic to figure out its status
+	 */
+	virtual FOnGetPropertyKeyedStatus& GetExternalHandler(const FProperty* Property) = 0;
+};
diff --git a/Engine/Source/Editor/UMGEditor/Private/Animation/UMGDetailKeyframeHandler.cpp b/Engine/Source/Editor/UMGEditor/Private/Animation/UMGDetailKeyframeHandler.cpp
--- a/Engine/Source/Editor/UMGEditor/Private/Animation/UMGDetailKeyframeHandler.cpp
+++ b/Engine/Source/Editor/UMGEditor/Private/Animation/UMGDetailKeyframeHandler.cpp
@@ -43,6 +43,16 @@ bool FUMGDetailKeyframeHandler::IsPropertyAnimated(const IPropertyHandle& Proper
 	return false;
 }
 
+EPropertyKeyedStatus FUMGDetailKeyframeHandler::GetPropertyKeyedStatus(const IPropertyHandle& PropertyHandle) const
+{
+	TSharedPtr<ISequencer> Sequencer = BlueprintEditor.Pin()->GetSequencer();
+	if (Sequencer.IsValid())
+	{
+		return Sequencer->GetPropertyKeyedStatus(PropertyHandle);
+	}
+	return EPropertyKeyedStatus::NotKeyed;
+}
+
 void FUMGDetailKeyframeHandler::OnKeyPropertyClicked(const IPropertyHandle& KeyedPropertyHandle)
 {
 	TArray<UObject*> Objects;
diff --git a/Engine/Source/Editor/UMGEditor/Private/Animation/UMGDetailKeyframeHandler.h b/Engine/Source/Editor/UMGEditor/Private/Animation/UMGDetailKeyframeHandler.h
--- a/Engine/Source/Editor/UMGEditor/Private/Animation/UMGDetailKeyframeHandler.h
+++ b/Engine/Source/Editor/UMGEditor/Private/Animation/UMGDetailKeyframeHandler.h
@@ -21,6 +21,8 @@ public:
 
 	virtual bool IsPropertyAnimated(const class IPropertyHandle& PropertyHandle, UObject *ParentObject) const override;
 
+	virtual EPropertyKeyedStatus GetPropertyKeyedStatus(const IPropertyHandle& PropertyHandle) const override;
+
 private:
 	TWeakPtr<class FWidgetBlueprintEditor> BlueprintEditor;
 };
diff --git a/Engine/Source/Runtime/MovieScene/Public/Channels/MovieSceneChannelEditorData.h b/Engine/Source/Runtime/MovieScene/Public/Channels/MovieSceneChannelEditorData.h
--- a/Engine/Source/Runtime/MovieScene/Public/Channels/MovieSceneChannelEditorData.h
+++ b/Engine/Source/Runtime/MovieScene/Public/Channels/MovieSceneChannelEditorData.h
@@ -65,6 +65,17 @@ struct FMovieSceneChannelMetaData
 	uint32 SortOrder;
 	/** This channel's unique name */
 	FName Name;
+	/**
+	 * Path representation of a sub property relative to the class property (i.e. topmost property) but NOT including the topmost property itself.
+	 * E.g. for FWidgetTransform, a sub-property path for the first channel would be "Translation.X"
+	 */
+	FName SubPropertyPath;
+	/**
+	 * Path representation of a sub-property relative to a class property (i.e. topmost property) but NOT including the topmost property itself.
+	 * This should be used when a Channel can be used by multiple struct sources and happen to have different property names for these.
+	 * The prime example of this is FTransform vs FEulerTransform. FTransform uses "Translation" vs FEulerTransform uses "Location"
+	 */
+	TMap<FName, FName> SubPropertyPathMap;
 	/** Text to display on this channel's key area node */
 	FText DisplayText;
 	/** Delegate to get a dynamic tooltip for the key area node */
diff --git a/Engine/Source/Runtime/MovieSceneTracks/Private/Sections/MovieScene3DTransformSection.cpp b/Engine/Source/Runtime/MovieSceneTracks/Private/Sections/MovieScene3DTransformSection.cpp
--- a/Engine/Source/Runtime/MovieSceneTracks/Private/Sections/MovieScene3DTransformSection.cpp
+++ b/Engine/Source/Runtime/MovieSceneTracks/Private/Sections/MovieScene3DTransformSection.cpp
@@ -31,6 +31,15 @@ struct F3DTransformChannelEditorData
 	{
 		static const TSet<FName> PropertyMetaDataKeys = { "UIMin", "UIMax", "SliderExponent", "LinearDeltaSensitivity", "Delta", "ClampMin", "ClampMax", "ForceUnits", "WheelStep" };
 
+		auto MakeSubPropertyMap = [](FName FTransformName, FName FEulerTransformName)->TMap<FName, FName>
+		{
+			return
+			{
+				{ TBaseStructure<FTransform>::Get()->GetFName(), FTransformName },
+				{ TBaseStructure<FEulerTransform>::Get()->GetFName(), FEulerTransformName }
+			};
+		};
+
 		const FProperty* RelativeLocationProperty = USceneComponent::StaticClass()->FindPropertyByName(USceneComponent::GetRelativeLocationPropertyName());
 		const FProperty* RelativeRotationProperty = USceneComponent::StaticClass()->FindPropertyByName(USceneComponent::GetRelativeRotationPropertyName());
 		const FProperty* RelativeScale3DProperty = USceneComponent::StaticClass()->FindPropertyByName(USceneComponent::GetRelativeScale3DPropertyName());
@@ -40,6 +49,7 @@ struct F3DTransformChannelEditorData
 		FText ScaleGroup    = NSLOCTEXT("MovieSceneTransformSection", "Scale",    "Scale");
 		{
 			MetaData[0].SetIdentifiers("Location.X", FCommonChannelData::ChannelX, LocationGroup);
+			MetaData[0].SubPropertyPathMap = MakeSubPropertyMap(TEXT("Translation.X"), TEXT("Location.X"));
 			MetaData[0].bEnabled = EnumHasAllFlags(Mask, EMovieSceneTransformChannel::TranslationX);
 			MetaData[0].Color = FCommonChannelData::RedChannelColor;
 			MetaData[0].SortOrder = SortOrderStart;
@@ -50,6 +60,7 @@ struct F3DTransformChannelEditorData
 			}
 
 			MetaData[1].SetIdentifiers("Location.Y", FCommonChannelData::ChannelY, LocationGroup);
+			MetaData[1].SubPropertyPathMap = MakeSubPropertyMap(TEXT("Translation.Y"), TEXT("Location.Y"));
 			MetaData[1].bEnabled = EnumHasAllFlags(Mask, EMovieSceneTransformChannel::TranslationY);
 			MetaData[1].Color = FCommonChannelData::GreenChannelColor;
 			MetaData[1].SortOrder = SortOrderStart + 1;
@@ -60,6 +71,7 @@ struct F3DTransformChannelEditorData
 			}
 
 			MetaData[2].SetIdentifiers("Location.Z", FCommonChannelData::ChannelZ, LocationGroup);
+			MetaData[2].SubPropertyPathMap = MakeSubPropertyMap(TEXT("Translation.Z"), TEXT("Location.Z"));
 			MetaData[2].bEnabled = EnumHasAllFlags(Mask, EMovieSceneTransformChannel::TranslationZ);
 			MetaData[2].Color = FCommonChannelData::BlueChannelColor;
 			MetaData[2].SortOrder = SortOrderStart + 2;
@@ -71,6 +83,7 @@ struct F3DTransformChannelEditorData
 		}
 		{
 			MetaData[3].SetIdentifiers("Rotation.X", NSLOCTEXT("MovieSceneTransformSection", "RotationX", "Roll"), RotationGroup);
+			MetaData[3].SubPropertyPathMap = MakeSubPropertyMap(TEXT("Rotation.X"), TEXT("Rotation.Roll"));
 			MetaData[3].bEnabled = EnumHasAllFlags(Mask, EMovieSceneTransformChannel::RotationX);
 			MetaData[3].Color = FCommonChannelData::RedChannelColor;
 			MetaData[3].SortOrder = SortOrderStart + 3;
@@ -81,6 +94,7 @@ struct F3DTransformChannelEditorData
 			}
 
 			MetaData[4].SetIdentifiers("Rotation.Y", NSLOCTEXT("MovieSceneTransformSection", "RotationY", "Pitch"), RotationGroup);
+			MetaData[4].SubPropertyPathMap = MakeSubPropertyMap(TEXT("Rotation.Y"), TEXT("Rotation.Pitch"));
 			MetaData[4].bEnabled = EnumHasAllFlags(Mask, EMovieSceneTransformChannel::RotationY);
 			MetaData[4].Color = FCommonChannelData::GreenChannelColor;
 			MetaData[4].SortOrder = SortOrderStart + 4;
@@ -91,6 +105,7 @@ struct F3DTransformChannelEditorData
 			}
 
 			MetaData[5].SetIdentifiers("Rotation.Z", NSLOCTEXT("MovieSceneTransformSection", "RotationZ", "Yaw"), RotationGroup);
+			MetaData[5].SubPropertyPathMap = MakeSubPropertyMap(TEXT("Rotation.Z"), TEXT("Rotation.Yaw"));
 			MetaData[5].bEnabled = EnumHasAllFlags(Mask, EMovieSceneTransformChannel::RotationZ);
 			MetaData[5].Color = FCommonChannelData::BlueChannelColor;
 			MetaData[5].SortOrder = SortOrderStart + 5;
@@ -102,6 +117,7 @@ struct F3DTransformChannelEditorData
 		}
 		{
 			MetaData[6].SetIdentifiers("Scale.X", FCommonChannelData::ChannelX, ScaleGroup);
+			MetaData[6].SubPropertyPathMap = MakeSubPropertyMap(TEXT("Scale3D.X"), TEXT("Scale.X"));
 			MetaData[6].bEnabled = EnumHasAllFlags(Mask, EMovieSceneTransformChannel::ScaleX);
 			MetaData[6].Color = FCommonChannelData::RedChannelColor;
 			MetaData[6].SortOrder = SortOrderStart + 6;
@@ -112,6 +128,7 @@ struct F3DTransformChannelEditorData
 			}
 
 			MetaData[7].SetIdentifiers("Scale.Y", FCommonChannelData::ChannelY, ScaleGroup);
+			MetaData[7].SubPropertyPathMap = MakeSubPropertyMap(TEXT("Scale3D.Y"), TEXT("Scale.Y"));
 			MetaData[7].bEnabled = EnumHasAllFlags(Mask, EMovieSceneTransformChannel::ScaleY);
 			MetaData[7].Color = FCommonChannelData::GreenChannelColor;
 			MetaData[7].SortOrder = SortOrderStart + 7;
@@ -122,6 +139,7 @@ struct F3DTransformChannelEditorData
 			}
 
 			MetaData[8].SetIdentifiers("Scale.Z", FCommonChannelData::ChannelZ, ScaleGroup);
+			MetaData[8].SubPropertyPathMap = MakeSubPropertyMap(TEXT("Scale3D.Z"), TEXT("Scale.Z"));
 			MetaData[8].bEnabled = EnumHasAllFlags(Mask, EMovieSceneTransformChannel::ScaleZ);
 			MetaData[8].Color = FCommonChannelData::BlueChannelColor;
 			MetaData[8].SortOrder = SortOrderStart + 8;
diff --git a/Engine/Source/Runtime/MovieSceneTracks/Private/Sections/MovieSceneColorSection.cpp b/Engine/Source/Runtime/MovieSceneTracks/Private/Sections/MovieSceneColorSection.cpp
--- a/Engine/Source/Runtime/MovieSceneTracks/Private/Sections/MovieSceneColorSection.cpp
+++ b/Engine/Source/Runtime/MovieSceneTracks/Private/Sections/MovieSceneColorSection.cpp
@@ -17,21 +17,25 @@ struct FColorSectionEditorData
 	FColorSectionEditorData()
 	{
 		MetaData[0].SetIdentifiers("Color.R", FCommonChannelData::ChannelR);
+		MetaData[0].SubPropertyPath = TEXT("R");
 		MetaData[0].SortOrder = 0;
 		MetaData[0].Color = FCommonChannelData::RedChannelColor;
 		MetaData[0].bCanCollapseToTrack = false;
 
 		MetaData[1].SetIdentifiers("Color.G", FCommonChannelData::ChannelG);
+		MetaData[1].SubPropertyPath = TEXT("G");
 		MetaData[1].SortOrder = 1;
 		MetaData[1].Color = FCommonChannelData::GreenChannelColor;
 		MetaData[1].bCanCollapseToTrack = false;
 
 		MetaData[2].SetIdentifiers("Color.B", FCommonChannelData::ChannelB);
+		MetaData[2].SubPropertyPath = TEXT("B");
 		MetaData[2].SortOrder = 2;
 		MetaData[2].Color = FCommonChannelData::BlueChannelColor;
 		MetaData[2].bCanCollapseToTrack = false;
 
 		MetaData[3].SetIdentifiers("Color.A", FCommonChannelData::ChannelA);
+		MetaData[3].SubPropertyPath = TEXT("A");
 		MetaData[3].SortOrder = 3;
 		MetaData[3].bCanCollapseToTrack = false;
 
diff --git a/Engine/Source/Runtime/MovieSceneTracks/Private/Sections/MovieSceneVectorSection.cpp b/Engine/Source/Runtime/MovieSceneTracks/Private/Sections/MovieSceneVectorSection.cpp
--- a/Engine/Source/Runtime/MovieSceneTracks/Private/Sections/MovieSceneVectorSection.cpp
+++ b/Engine/Source/Runtime/MovieSceneTracks/Private/Sections/MovieSceneVectorSection.cpp
@@ -85,21 +85,25 @@ struct TVectorSectionEditorData
 	TVectorSectionEditorData(int32 NumChannels)
 	{
 		MetaData[0].SetIdentifiers("Vector.X", FCommonChannelData::ChannelX);
+		MetaData[0].SubPropertyPath = TEXT("X");
 		MetaData[0].SortOrder = 0;
 		MetaData[0].Color = FCommonChannelData::RedChannelColor;
 		MetaData[0].bCanCollapseToTrack = false;
 
 		MetaData[1].SetIdentifiers("Vector.Y", FCommonChannelData::ChannelY);
+		MetaData[1].SubPropertyPath = TEXT("Y");
 		MetaData[1].SortOrder = 1;
 		MetaData[1].Color = FCommonChannelData::GreenChannelColor;
 		MetaData[1].bCanCollapseToTrack = false;
 
 		MetaData[2].SetIdentifiers("Vector.Z", FCommonChannelData::ChannelZ);
+		MetaData[2].SubPropertyPath = TEXT("Z");
 		MetaData[2].SortOrder = 2;
 		MetaData[2].Color = FCommonChannelData::BlueChannelColor;
 		MetaData[2].bCanCollapseToTrack = false;
 
 		MetaData[3].SetIdentifiers("Vector.W", FCommonChannelData::ChannelW);
+		MetaData[3].SubPropertyPath = TEXT("W");
 		MetaData[3].SortOrder = 3;
 		MetaData[3].bCanCollapseToTrack = false;
 
diff --git a/Engine/Source/Runtime/UMG/Private/Animation/MovieScene2DTransformSection.cpp b/Engine/Source/Runtime/UMG/Private/Animation/MovieScene2DTransformSection.cpp
--- a/Engine/Source/Runtime/UMG/Private/Animation/MovieScene2DTransformSection.cpp
+++ b/Engine/Source/Runtime/UMG/Private/Animation/MovieScene2DTransformSection.cpp
@@ -32,36 +32,43 @@ struct F2DTransformSectionEditorData
 		FText ShearGroup = NSLOCTEXT("MovieScene2DTransformSection", "Shear", "Shear");
 
 		MetaData[0].SetIdentifiers("Translation.X", FCommonChannelData::ChannelX, TranslationGroup);
+		MetaData[0].SubPropertyPath = MetaData[0].Name;
 		MetaData[0].bEnabled = EnumHasAllFlags(Mask, EMovieScene2DTransformChannel::TranslationX);
 		MetaData[0].SortOrder = 0;
 		MetaData[0].bCanCollapseToTrack = false;
 
 		MetaData[1].SetIdentifiers("Translation.Y", FCommonChannelData::ChannelY, TranslationGroup);
+		MetaData[1].SubPropertyPath = MetaData[1].Name;
 		MetaData[1].bEnabled = EnumHasAllFlags(Mask, EMovieScene2DTransformChannel::TranslationY);
 		MetaData[1].SortOrder = 1;
 		MetaData[1].bCanCollapseToTrack = false;
 
 		MetaData[2].SetIdentifiers("Angle", NSLOCTEXT("MovieScene2DTransformSection", "AngleText", "Angle"), RotationGroup);
+		MetaData[2].SubPropertyPath = MetaData[2].Name;
 		MetaData[2].bEnabled = EnumHasAllFlags(Mask, EMovieScene2DTransformChannel::Rotation);
 		MetaData[2].SortOrder = 2;
 		MetaData[2].bCanCollapseToTrack = false;
 
 		MetaData[3].SetIdentifiers("Scale.X", FCommonChannelData::ChannelX, ScaleGroup);
+		MetaData[3].SubPropertyPath = MetaData[3].Name;
 		MetaData[3].bEnabled = EnumHasAllFlags(Mask, EMovieScene2DTransformChannel::ScaleX);
 		MetaData[3].SortOrder = 3;
 		MetaData[3].bCanCollapseToTrack = false;
 
 		MetaData[4].SetIdentifiers("Scale.Y", FCommonChannelData::ChannelY, ScaleGroup);
+		MetaData[4].SubPropertyPath = MetaData[4].Name;
 		MetaData[4].bEnabled = EnumHasAllFlags(Mask, EMovieScene2DTransformChannel::ScaleY);
 		MetaData[4].SortOrder = 4;
 		MetaData[4].bCanCollapseToTrack = false;
 
 		MetaData[5].SetIdentifiers("Shear.X", FCommonChannelData::ChannelX, ShearGroup);
+		MetaData[5].SubPropertyPath = MetaData[5].Name;
 		MetaData[5].bEnabled = EnumHasAllFlags(Mask, EMovieScene2DTransformChannel::ShearX);
 		MetaData[5].SortOrder = 5;
 		MetaData[5].bCanCollapseToTrack = false;
 
 		MetaData[6].SetIdentifiers("Shear.Y", FCommonChannelData::ChannelY, ShearGroup);
+		MetaData[6].SubPropertyPath = MetaData[6].Name;
 		MetaData[6].bEnabled = EnumHasAllFlags(Mask, EMovieScene2DTransformChannel::ShearY);
 		MetaData[6].SortOrder = 6;
 		MetaData[6].bCanCollapseToTrack = false;
diff --git a/Engine/Source/Runtime/UMG/Private/Animation/MovieSceneMarginSection.cpp b/Engine/Source/Runtime/UMG/Private/Animation/MovieSceneMarginSection.cpp
--- a/Engine/Source/Runtime/UMG/Private/Animation/MovieSceneMarginSection.cpp
+++ b/Engine/Source/Runtime/UMG/Private/Animation/MovieSceneMarginSection.cpp
@@ -19,18 +19,22 @@ struct FMarginSectionEditorData
 	FMarginSectionEditorData()
 	{
 		MetaData[0].SetIdentifiers("Left", NSLOCTEXT("MovieSceneMarginSection", "LeftText", "Left"));
+		MetaData[0].SubPropertyPath = MetaData[0].Name;
 		MetaData[0].SortOrder = 0;
 		MetaData[0].bCanCollapseToTrack = false;
 
 		MetaData[1].SetIdentifiers("Top", NSLOCTEXT("MovieSceneMarginSection", "TopText", "Top"));
+		MetaData[1].SubPropertyPath = MetaData[1].Name;
 		MetaData[1].SortOrder = 1;
 		MetaData[1].bCanCollapseToTrack = false;
 
 		MetaData[2].SetIdentifiers("Right", NSLOCTEXT("MovieSceneMarginSection", "RightText", "Right"));
+		MetaData[2].SubPropertyPath = MetaData[2].Name;
 		MetaData[2].SortOrder = 2;
 		MetaData[2].bCanCollapseToTrack = false;
 
 		MetaData[3].SetIdentifiers("Bottom", NSLOCTEXT("MovieSceneMarginSection", "BottomText", "Bottom"));
+		MetaData[3].SubPropertyPath = MetaData[3].Name;
 		MetaData[3].SortOrder = 3;
 		MetaData[3].bCanCollapseToTrack = false;
 
