Sequencer: Fix dragging multiple sections on different rows

Fix regression from 24451411

diff --git a/Engine/Source/Editor/Sequencer/Private/Tools/EditToolDragOperations.cpp b/Engine/Source/Editor/Sequencer/Private/Tools/EditToolDragOperations.cpp
--- a/Engine/Source/Editor/Sequencer/Private/Tools/EditToolDragOperations.cpp
+++ b/Engine/Source/Editor/Sequencer/Private/Tools/EditToolDragOperations.cpp
@@ -1107,10 +1107,6 @@ bool FMoveKeysAndSections::HandleSectionMovement(FFrameTime MouseTime, FVector2D
 		{
 			HighestRowIndex = Section->GetRowIndex();
 		}
-		if (LowestRowIndex.IsSet() && LowestRowIndex.GetValue() != Section->GetRowIndex())
-		{
-			bSectionsAreOnDifferentRows = true;
-		}
 		if (FirstTrack)
 		{
 			if (FirstTrack != Track)
@@ -1124,6 +1120,11 @@ bool FMoveKeysAndSections::HandleSectionMovement(FFrameTime MouseTime, FVector2D
 		}
 	}
 
+	if (LowestRowIndex.IsSet() && HighestRowIndex.IsSet() && LowestRowIndex.GetValue() != HighestRowIndex.GetValue())
+	{
+		bSectionsAreOnDifferentRows = true;
+	}
+
 	TArray<TSharedPtr<FViewModel>> Tracks;
 
 	// @todo_sequencer_mvvm: really this code should all be operating on models rather than looking directly at UMovieSceneSections and tracks
