Sequencer: If there is a subsequence range, clamp to that rather than the playback range.

Users want to be able to clamp the scrubbing to the bounds of the shot, rather than the outer playback range.

diff --git a/Engine/Source/Editor/Sequencer/Private/SequencerCommands.cpp b/Engine/Source/Editor/Sequencer/Private/SequencerCommands.cpp
--- a/Engine/Source/Editor/Sequencer/Private/SequencerCommands.cpp
+++ b/Engine/Source/Editor/Sequencer/Private/SequencerCommands.cpp
@@ -55,7 +55,7 @@ void FSequencerCommands::RegisterCommands()
 	UI_COMMAND( ToggleDynamicWeighting, "Dynamic Weighting", "When enabled, all blendable tracks will cache their initial values to ensure that they are able to correctly blend in/out when dynamic weights are being used.", EUserInterfaceActionType::ToggleButton, FInputChord() );
 
 	UI_COMMAND( ToggleResetPlayheadWhenNavigating, "Reset Playhead When Navigating", "When checked, if the playhead is outside the playback range, it will be reset to the beginning when navigating in and out of subsequences", EUserInterfaceActionType::ToggleButton, FInputChord());
-	UI_COMMAND( ToggleKeepCursorInPlaybackRangeWhileScrubbing, "Keep Playhead in Playback Range While Scrubbing", "When checked, the playhead will be constrained to the current playback range while scrubbing", EUserInterfaceActionType::ToggleButton, FInputChord() );
+	UI_COMMAND( ToggleKeepCursorInPlaybackRangeWhileScrubbing, "Keep Playhead in Playback Range While Scrubbing", "When checked, the playhead will be constrained to the current playback (or subsequence/shot) range while scrubbing", EUserInterfaceActionType::ToggleButton, FInputChord() );
 	UI_COMMAND( ToggleKeepPlaybackRangeInSectionBounds, "Keep Playback Range in Section Bounds", "When checked, the playback range will be synchronized to the section bounds", EUserInterfaceActionType::ToggleButton, FInputChord() );
 
 	UI_COMMAND( ToggleAutoExpandNodesOnSelection, "Auto Expand Nodes on Selection", "Toggle auto expanding the outliner tree on child selection", EUserInterfaceActionType::ToggleButton, FInputChord());
diff --git a/Engine/Source/Editor/Sequencer/Private/SequencerTimeSliderController.cpp b/Engine/Source/Editor/Sequencer/Private/SequencerTimeSliderController.cpp
--- a/Engine/Source/Editor/Sequencer/Private/SequencerTimeSliderController.cpp
+++ b/Engine/Source/Editor/Sequencer/Private/SequencerTimeSliderController.cpp
@@ -113,7 +113,17 @@ FFrameTime FSequencerTimeSliderController::ComputeScrubTimeFromMouse(const FGeom
 	// Clamp first, snap to frame last
 	if (Sequencer->GetSequencerSettings()->ShouldKeepCursorInPlayRangeWhileScrubbing())
 	{
-		ScrubTime = UE::MovieScene::ClampToDiscreteRange(ScrubTime, TimeSliderArgs.PlaybackRange.Get());
+		TOptional<TRange<FFrameNumber>> RangeValue;
+		RangeValue = TimeSliderArgs.SubSequenceRange.Get(RangeValue);
+
+		if (RangeValue.IsSet())
+		{
+			ScrubTime = UE::MovieScene::ClampToDiscreteRange(ScrubTime, RangeValue.GetValue());
+		}
+		else
+		{
+			ScrubTime = UE::MovieScene::ClampToDiscreteRange(ScrubTime, TimeSliderArgs.PlaybackRange.Get());
+		}
 	}
 
 	if ( Sequencer->GetSequencerSettings()->GetIsSnapEnabled() || MouseEvent.IsShiftDown() )
diff --git a/Engine/Source/Editor/Sequencer/Public/SequencerSettings.h b/Engine/Source/Editor/Sequencer/Public/SequencerSettings.h
--- a/Engine/Source/Editor/Sequencer/Public/SequencerSettings.h
+++ b/Engine/Source/Editor/Sequencer/Public/SequencerSettings.h
@@ -282,9 +282,9 @@ public:
 	/** Set whether or not the cursor should be reset when navigating in and out of subsequences */
 	void SetResetPlayheadWhenNavigating(bool bInResetPlayheadWhenNavigating);
 
-	/** @return true if the cursor should be kept within the playback range while scrubbing in sequencer, false otherwise */
+	/** @return true if the cursor should be kept within the playback (or subsequence/shot) range while scrubbing in sequencer, false otherwise */
 	bool ShouldKeepCursorInPlayRangeWhileScrubbing() const;
-	/** Set whether or not the cursor should be kept within the playback range while scrubbing in sequencer */
+	/** Set whether or not the cursor should be kept within the playback (or subsequence/shot) range while scrubbing in sequencer */
 	void SetKeepCursorInPlayRangeWhileScrubbing(bool bInKeepCursorInPlayRangeWhileScrubbing);
 
 	/** @return true if the playback range should be synced to the section bounds, false otherwise */
