Sequencer: Distinguish between RefreshUI and RefreshTree.

MovieSceneStructureItemsChanged calls RefreshUI because the entire UI needs refreshing (ie. bindings replaced)
MovieSceneStructureItemRemoved calls RefreshTree because items have been removed as a result of event callbacks

Follow up fix to 28392580 and #jira UE-111561

#rnx

diff --git a/Engine/Source/Editor/Sequencer/Private/Sequencer.cpp b/Engine/Source/Editor/Sequencer/Private/Sequencer.cpp
--- a/Engine/Source/Editor/Sequencer/Private/Sequencer.cpp
+++ b/Engine/Source/Editor/Sequencer/Private/Sequencer.cpp
@@ -2376,7 +2376,16 @@ void FSequencer::NotifyMovieSceneDataChanged( EMovieSceneDataChangeType DataChan
 		EMovieScenePlayerStatus::Type StoredPlaybackState = GetPlaybackStatus();
 		SetPlaybackStatus( EMovieScenePlayerStatus::Stopped );
 		SelectionPreview.Empty();
-		RefreshUI();
+
+		if (DataChangeType == EMovieSceneDataChangeType::MovieSceneStructureItemsChanged)
+		{
+			RefreshUI();
+		}
+		else
+		{	
+			RefreshTree();
+		}
+
 		SetPlaybackStatus( StoredPlaybackState );
 	}
 	else if (DataChangeType == EMovieSceneDataChangeType::TrackValueChangedRefreshImmediately)
