Sequencer: Back out changes to MovieSceneStructureItemsChanged and MovieSceneStructureItemsRemoved

Added a RefreshUI temporarily for just FixActorReferences (will revisit)

Follow up fix to 28392580 and #jira UE-111561

#rnx

diff --git a/Engine/Plugins/MovieScene/LevelSequenceEditor/Source/LevelSequenceEditor/Private/LevelSequenceEditorSubsystem.cpp b/Engine/Plugins/MovieScene/LevelSequenceEditor/Source/LevelSequenceEditor/Private/LevelSequenceEditorSubsystem.cpp
--- a/Engine/Plugins/MovieScene/LevelSequenceEditor/Source/LevelSequenceEditor/Private/LevelSequenceEditorSubsystem.cpp
+++ b/Engine/Plugins/MovieScene/LevelSequenceEditor/Source/LevelSequenceEditor/Private/LevelSequenceEditorSubsystem.cpp
@@ -1463,7 +1463,7 @@ void ULevelSequenceEditorSubsystem::FixActorReferences()
 		FSequencerUtilities::UpdateBindingIDs(Sequencer.ToSharedRef(), GuidPair.Key, GuidPair.Value);
 	}
 
-	Sequencer->NotifyMovieSceneDataChanged(EMovieSceneDataChangeType::MovieSceneStructureItemsChanged);
+	Sequencer->NotifyMovieSceneDataChanged(EMovieSceneDataChangeType::RefreshUI);
 }
 
 void ULevelSequenceEditorSubsystem::AddActorsToBindingInternal()
diff --git a/Engine/Source/Editor/Sequencer/Private/Sequencer.cpp b/Engine/Source/Editor/Sequencer/Private/Sequencer.cpp
--- a/Engine/Source/Editor/Sequencer/Private/Sequencer.cpp
+++ b/Engine/Source/Editor/Sequencer/Private/Sequencer.cpp
@@ -2376,18 +2376,14 @@ void FSequencer::NotifyMovieSceneDataChanged( EMovieSceneDataChangeType DataChan
 		EMovieScenePlayerStatus::Type StoredPlaybackState = GetPlaybackStatus();
 		SetPlaybackStatus( EMovieScenePlayerStatus::Stopped );
 		SelectionPreview.Empty();
-
-		if (DataChangeType == EMovieSceneDataChangeType::MovieSceneStructureItemsChanged)
-		{
-			RefreshUI();
-		}
-		else
-		{	
-			RefreshTree();
-		}
-
+		RefreshTree();
 		SetPlaybackStatus( StoredPlaybackState );
 	}
+	else if (DataChangeType == EMovieSceneDataChangeType::RefreshUI)
+	{
+		RefreshUI();
+		return;
+	}
 	else if (DataChangeType == EMovieSceneDataChangeType::TrackValueChangedRefreshImmediately)
 	{
 		// Evaluate now
diff --git a/Engine/Source/Editor/Sequencer/Public/ISequencer.h b/Engine/Source/Editor/Sequencer/Public/ISequencer.h
--- a/Engine/Source/Editor/Sequencer/Public/ISequencer.h
+++ b/Engine/Source/Editor/Sequencer/Public/ISequencer.h
@@ -195,8 +195,9 @@ enum class EMovieSceneDataChangeType
 	/** It's not known what data has changed. */
 	Unknown,
 	/** Refresh Tree on Next Tick */
-	RefreshTree
-
+	RefreshTree,
+	/** Refresh UI completely */
+	RefreshUI
 };
 
 /**
