Sequencer: Fix empty property

#rnx

diff --git a/Engine/Source/Editor/Sequencer/Private/SequencerPropertyKeyedStatus.cpp b/Engine/Source/Editor/Sequencer/Private/SequencerPropertyKeyedStatus.cpp
--- a/Engine/Source/Editor/Sequencer/Private/SequencerPropertyKeyedStatus.cpp
+++ b/Engine/Source/Editor/Sequencer/Private/SequencerPropertyKeyedStatus.cpp
@@ -26,7 +26,10 @@ FPropertyPath FSequencerPropertyKeyedStatusHandler::FPropertyParameters::BuildPr
 	TSharedPtr<const IPropertyHandle> ParentHandle = ActualProperty.GetParentHandle();
 
 	TArray<FPropertyInfo> PropertyInfos;
-	PropertyInfos.Emplace(CurrentHandle->GetProperty(), CurrentHandle->GetArrayIndex());
+	if (CurrentHandle->GetProperty())
+	{
+		PropertyInfos.Emplace(CurrentHandle->GetProperty(), CurrentHandle->GetArrayIndex());
+	}
 
 	while (ParentHandle.IsValid() && ParentHandle->GetProperty())
 	{
