Sequencer: Check for valid end frame before getting end frame

#rnx

diff --git a/Engine/Source/Editor/Sequencer/Private/MovieSceneTrackEditor.cpp b/Engine/Source/Editor/Sequencer/Private/MovieSceneTrackEditor.cpp
--- a/Engine/Source/Editor/Sequencer/Private/MovieSceneTrackEditor.cpp
+++ b/Engine/Source/Editor/Sequencer/Private/MovieSceneTrackEditor.cpp
@@ -97,7 +97,7 @@ void FMovieSceneTrackEditor::AnimatablePropertyChanged( FOnKeyProperty OnKeyProp
 		{
 			for (TWeakObjectPtr<UMovieSceneSection> NewSection : KeyPropertyResult.SectionsCreated)
 			{
-				if (NewSection.IsValid())
+				if (NewSection.IsValid() && NewSection.Get()->HasEndFrame())
 				{
 					NextKeyTime = NewSection.Get()->GetExclusiveEndFrame();
 					break;
