Sequencer: Check all chord combinations for scrubbing time. Otherwise, Ctrl-B event is eaten and the user cannot browse to the asset.

Tested that Ctrl+B browses to the asset
B and mouse scrub in viewport scrubs time
Rebinding scrub time to Ctrl+B now scrubs time in viewport.

diff --git a/Engine/Source/Editor/Sequencer/Private/SequencerEdMode.cpp b/Engine/Source/Editor/Sequencer/Private/SequencerEdMode.cpp
--- a/Engine/Source/Editor/Sequencer/Private/SequencerEdMode.cpp
+++ b/Engine/Source/Editor/Sequencer/Private/SequencerEdMode.cpp
@@ -235,7 +235,11 @@ bool FSequencerEdMode::IsPressingMoveTimeSlider(FViewport* InViewport) const
 	{
 		EMultipleKeyBindingIndex ChordIndex = static_cast<EMultipleKeyBindingIndex>(i);
 		const FInputChord& Chord = *Commands.ScrubTimeViewport->GetActiveChord(ChordIndex);
-		bIsMovingTimeSlider |= Chord.IsValidChord() && InViewport->KeyState(Chord.Key);
+		bIsMovingTimeSlider |= Chord.IsValidChord() && InViewport->KeyState(Chord.Key) && 
+			(Chord.NeedsAlt() ? IsAltDown(InViewport) : !IsAltDown(InViewport)) &&
+			(Chord.NeedsShift() ? IsShiftDown(InViewport) : !IsShiftDown(InViewport)) &&
+			(Chord.NeedsControl() ? IsCtrlDown(InViewport) : !IsCtrlDown(InViewport));
+
 	}
 	return bIsMovingTimeSlider;
 }
