Sequencer: null check before synchronizing level editor selection.

diff --git a/Engine/Source/Editor/Sequencer/Private/Sequencer.cpp b/Engine/Source/Editor/Sequencer/Private/Sequencer.cpp
--- a/Engine/Source/Editor/Sequencer/Private/Sequencer.cpp
+++ b/Engine/Source/Editor/Sequencer/Private/Sequencer.cpp
@@ -5941,6 +5941,11 @@ void FSequencer::ClearFilters()
 void FSequencer::SynchronizeExternalSelectionWithSequencerSelection()
 {
 	using namespace UE::Sequencer;
+	
+	if (!ViewModel)
+	{
+		return;
+	}
 
 	if ( bUpdatingSequencerSelection || !IsLevelEditorSequencer() )
 	{
