Sequencer: don't invalidate bindings when looking for an object

#noqatesting

diff --git a/Engine/Source/Editor/Sequencer/Private/Sequencer.cpp b/Engine/Source/Editor/Sequencer/Private/Sequencer.cpp
--- a/Engine/Source/Editor/Sequencer/Private/Sequencer.cpp
+++ b/Engine/Source/Editor/Sequencer/Private/Sequencer.cpp
@@ -3748,7 +3748,9 @@ FGuid FSequencer::GetHandleToObject( UObject* Object, bool bCreateHandleIfMissin
 	}
 
 	// Attempt to resolve the object through the movie scene instance first, 
-	FGuid ObjectGuid = FindObjectId(*Object, ActiveTemplateIDs.Top());
+	FGuid ObjectGuid = bCreateHandleIfMissing ?
+		FindObjectId(*Object, ActiveTemplateIDs.Top()) :
+		FindCachedObjectId(*Object, ActiveTemplateIDs.Top());
 
 	if (ObjectGuid.IsValid())
 	{
