Sequencer: add playback capabilities to re-initiatlized sequence instances when playback context changes.

diff --git a/Engine/Source/Editor/Sequencer/Private/LevelEditorSequencerIntegration.cpp b/Engine/Source/Editor/Sequencer/Private/LevelEditorSequencerIntegration.cpp
--- a/Engine/Source/Editor/Sequencer/Private/LevelEditorSequencerIntegration.cpp
+++ b/Engine/Source/Editor/Sequencer/Private/LevelEditorSequencerIntegration.cpp
@@ -700,7 +700,7 @@ void FLevelEditorSequencerIntegration::OnPreBeginPIE(bool bIsSimulating)
 		{
 			if (Options.bRequiresLevelEvents)
 			{
-				In.GetEvaluationTemplate().PlaybackContextChanged(In);
+				In.OnPlaybackContextChanged();
 				In.RestorePreAnimatedState();
 				In.State.ClearObjectCaches(In);
 				In.RequestEvaluate();
@@ -738,7 +738,7 @@ void FLevelEditorSequencerIntegration::OnEndPlayMap()
 			if (Options.bRequiresLevelEvents)
 			{
 				// Update and clear any stale bindings 
-				In.GetEvaluationTemplate().PlaybackContextChanged(In);
+				In.OnPlaybackContextChanged();
 				In.State.ClearObjectCaches(In);
 				In.ForceEvaluate();
 			}
@@ -1261,7 +1261,7 @@ void FLevelEditorSequencerIntegration::OnMapChanged(UWorld* World, EMapChangeTyp
 		{
 			if (Options.bRequiresLevelEvents)
 			{
-				In.GetEvaluationTemplate().PlaybackContextChanged(In);
+				In.OnPlaybackContextChanged();
 				In.RestorePreAnimatedState();
 				In.State.ClearObjectCaches(In);
 
diff --git a/Engine/Source/Editor/Sequencer/Private/Sequencer.cpp b/Engine/Source/Editor/Sequencer/Private/Sequencer.cpp
--- a/Engine/Source/Editor/Sequencer/Private/Sequencer.cpp
+++ b/Engine/Source/Editor/Sequencer/Private/Sequencer.cpp
@@ -474,15 +474,7 @@ void FSequencer::InitSequencer(const FSequencerInitParams& InitParams, const TSh
 	InitialValueCache = UE::MovieScene::FInitialValueCache::GetGlobalInitialValues();
 	RootTemplateInstance.GetEntitySystemLinker()->AddExtension(InitialValueCache.Get());
 
-	// Add the camera cut playback capability.
-	FInstanceHandle RootInstanceHandle = RootTemplateInstance.GetRootInstanceHandle();
-	UMovieSceneEntitySystemLinker* Linker = GetEvaluationTemplate().GetEntitySystemLinker();
-	FInstanceRegistry* InstanceRegistry = Linker->GetInstanceRegistry();
-	TSharedRef<FSharedPlaybackState> SharedPlaybackState = InstanceRegistry->GetInstance(RootInstanceHandle).GetSharedPlaybackState();
-	if (!SharedPlaybackState->HasCapability<FCameraCutPlaybackCapability>())
-	{
-		SharedPlaybackState->AddCapabilityRaw<FCameraCutPlaybackCapability>((FCameraCutPlaybackCapability*)this);
-	}
+	InitRootSequenceInstance();
 
 	// Create tools and bind them to this sequencer
 	for( int32 DelegateIndex = 0; DelegateIndex < TrackEditorDelegates.Num(); ++DelegateIndex )
@@ -657,6 +649,27 @@ void FSequencer::InitSequencer(const FSequencerInitParams& InitParams, const TSh
 	OnActivateSequenceEvent.Broadcast(ActiveTemplateIDs[0]);
 }
 
+void FSequencer::OnPlaybackContextChanged()
+{
+	RootTemplateInstance.PlaybackContextChanged(*this);
+	InitRootSequenceInstance();
+}
+
+void FSequencer::InitRootSequenceInstance()
+{
+	using namespace UE::MovieScene;
+
+	// Add the camera cut playback capability.
+	FInstanceHandle RootInstanceHandle = RootTemplateInstance.GetRootInstanceHandle();
+	UMovieSceneEntitySystemLinker* Linker = GetEvaluationTemplate().GetEntitySystemLinker();
+	FInstanceRegistry* InstanceRegistry = Linker->GetInstanceRegistry();
+	TSharedRef<FSharedPlaybackState> SharedPlaybackState = InstanceRegistry->GetInstance(RootInstanceHandle).GetSharedPlaybackState();
+	if (!SharedPlaybackState->HasCapability<FCameraCutPlaybackCapability>())
+	{
+		SharedPlaybackState->AddCapabilityRaw<FCameraCutPlaybackCapability>((FCameraCutPlaybackCapability*)this);
+	}
+}
+
 FSequencer::FSequencer()
 	: SequencerCommandBindings( new FUICommandList )
 	, SequencerSharedBindings(new FUICommandList)
@@ -3257,7 +3270,7 @@ void FSequencer::UpdateCachedPlaybackContextAndClient()
 		CachedPlaybackContext = NewPlaybackContext;
 		CachedPlaybackClient = NewPlaybackClient;
 
-		RootTemplateInstance.PlaybackContextChanged(*this);
+		OnPlaybackContextChanged();
 	}
 }
 
diff --git a/Engine/Source/Editor/Sequencer/Private/Sequencer.h b/Engine/Source/Editor/Sequencer/Private/Sequencer.h
--- a/Engine/Source/Editor/Sequencer/Private/Sequencer.h
+++ b/Engine/Source/Editor/Sequencer/Private/Sequencer.h
@@ -134,6 +134,20 @@ public:
 	 */
 	void InitSequencer(const FSequencerInitParams& InitParams, const TSharedRef<ISequencerObjectChangeListener>& InObjectChangeListener, const TArray<FOnCreateTrackEditor>& TrackEditorDelegates, const TArray<FOnCreateEditorObjectBinding>& EditorObjectBindingDelegatess, const TArray<FOnCreateOutlinerColumn>& OutlinerColumnDelegates);
 
+	/**
+	 * Reinitializes sequencer after the playback context has changed
+	 *
+	 * The playback context changes when the user attaches the sequencer to a different world either
+	 * explicitly via the world dropdown picker, or indirectly when e.g. starting/stopping PIE.
+	 */
+	void OnPlaybackContextChanged();
+
+protected:
+
+	void InitRootSequenceInstance();
+
+public:
+
 	/** @return The current view range */
 	virtual FAnimatedRange GetViewRange() const override;
 	virtual void SetViewRange(TRange<double> NewViewRange, EViewRangeInterpolation Interpolation = EViewRangeInterpolation::Animated) override;
