Sequencer: Added a menu item to conform the start/end range to the subsequence range

diff --git a/Engine/Source/Editor/Sequencer/Private/SequencerTimeSliderController.cpp b/Engine/Source/Editor/Sequencer/Private/SequencerTimeSliderController.cpp
--- a/Engine/Source/Editor/Sequencer/Private/SequencerTimeSliderController.cpp
+++ b/Engine/Source/Editor/Sequencer/Private/SequencerTimeSliderController.cpp
@@ -1566,6 +1566,7 @@ TSharedRef<SWidget> FSequencerTimeSliderController::OpenSetPlaybackRangeMenu(con
 	CurrentTimeText = FText::FromString(TimeSliderArgs.NumericTypeInterface->ToString(FrameNumber.Value));
 	
 	TRange<FFrameNumber> PlaybackRange = TimeSliderArgs.PlaybackRange.Get();
+	TOptional<TRange<FFrameNumber>> SubSequenceRange = TimeSliderArgs.SubSequenceRange.Get();
 
 	MenuBuilder.BeginSection("SequencerPlaybackRangeMenu", FText::Format(LOCTEXT("PlaybackRangeTextFormat", "Playback Range ({0}):"), CurrentTimeText));
 	{
@@ -1589,6 +1590,16 @@ TSharedRef<SWidget> FSequencerTimeSliderController::OpenSetPlaybackRangeMenu(con
 			)
 		);
 
+		MenuBuilder.AddMenuEntry(
+			LOCTEXT("ConformToSubsequenceRange", "Conform to Range"),
+			LOCTEXT("ConformToSubsequenceRangeTooltip", "Conform the start and end time to the extents of the subsequence range"),
+			FSlateIcon(),
+			FUIAction(
+				FExecuteAction::CreateLambda([this, SubSequenceRange] { SetPlaybackRangeStart(SubSequenceRange.GetValue().GetLowerBoundValue()); SetPlaybackRangeEnd(SubSequenceRange.GetValue().GetUpperBoundValue()); }),
+				FCanExecuteAction::CreateLambda([this, SubSequenceRange] { return !TimeSliderArgs.IsPlaybackRangeLocked.Get() && SubSequenceRange.IsSet(); })
+			)
+		);
+
 		MenuBuilder.AddMenuEntry(
 			LOCTEXT("ToggleLocked", "Locked"),
 			LOCTEXT("ToggleLockedTooltip", "Lock/Unlock the playback range"),
