Sequencer: Recusively add child folders to copy.

diff --git a/Engine/Source/Editor/Sequencer/Private/SequencerUtilities.cpp b/Engine/Source/Editor/Sequencer/Private/SequencerUtilities.cpp
--- a/Engine/Source/Editor/Sequencer/Private/SequencerUtilities.cpp
+++ b/Engine/Source/Editor/Sequencer/Private/SequencerUtilities.cpp
@@ -1302,12 +1302,27 @@ void ExportObjectsToText(const TArray<UObject*>& ObjectsToExport, FString& Expor
  *
  */
 
+void GatherChildFolders(UMovieSceneFolder* ParentFolder, TArray<UObject*>& Objects)
+{
+	for (UMovieSceneFolder* ChildFolder : ParentFolder->GetChildFolders())
+	{
+		if (ChildFolder)
+		{
+			Objects.Add(ChildFolder);
+
+			GatherChildFolders(ChildFolder, Objects);
+		}
+	}
+}
+
 void FSequencerUtilities::CopyFolders(const TArray<UMovieSceneFolder*>& Folders, FString& ExportedText)
 {
 	TArray<UObject*> Objects;
 	for (UMovieSceneFolder* Folder : Folders)
 	{
 		Objects.Add(Folder);
+
+		GatherChildFolders(Folder, Objects);
 	}
 
 	ExportObjectsToText(Objects, /*out*/ ExportedText);
