SequencerMode: fix gizmo jitter

diff --git a/Engine/Source/Editor/Sequencer/Private/SequencerEdMode.cpp b/Engine/Source/Editor/Sequencer/Private/SequencerEdMode.cpp
--- a/Engine/Source/Editor/Sequencer/Private/SequencerEdMode.cpp
+++ b/Engine/Source/Editor/Sequencer/Private/SequencerEdMode.cpp
@@ -302,11 +302,13 @@ bool FSequencerEdMode::StartTracking(FEditorViewportClient* InViewportClient, FV
 	}
 	else if (IsMovingCamera(InViewport))
 	{
+		bUpdatePivot = true;
 		InViewportClient->SetCurrentWidgetAxis(EAxisList::None);
 		return true;
 	}
 	else if (IsDoingDrag(InViewport))
 	{
+		bUpdatePivot = true;
 		DragToolHandler.MakeDragTool(InViewportClient);
 		return DragToolHandler.StartTracking(InViewportClient, InViewport);
 	}
@@ -320,15 +322,26 @@ bool FSequencerEdMode::EndTracking(FEditorViewportClient* InViewportClient, FVie
 	}
 	else if (IsMovingCamera(InViewport))
 	{
+		bUpdatePivot = false;
 		return true;
 	}
 	else if (DragToolHandler.EndTracking(InViewportClient, InViewport))
 	{
+		bUpdatePivot = false;
 		return true;
 	}
 	return FEdMode::EndTracking(InViewportClient, InViewport);
 }
 
+void FSequencerEdMode::Tick(FEditorViewportClient* ViewportClient,float DeltaTime)
+{
+	if (bUpdatePivot)
+	{
+		GUnrealEd->UpdatePivotLocationForSelection();
+	}
+	return FEdMode::Tick(ViewportClient,DeltaTime);
+}
+
 bool FSequencerEdMode::InputDelta(FEditorViewportClient* InViewportClient, FViewport* InViewport, FVector& InDrag, FRotator& InRot, FVector& InScale)
 {
 	if (IsPressingMoveTimeSlider(InViewport))
@@ -350,8 +363,6 @@ bool FSequencerEdMode::InputDelta(FEditorViewportClient* InViewportClient, FView
 	else if (IsDoingDrag(InViewport))
 	{
 		return DragToolHandler.InputDelta(InViewportClient, InViewport, InDrag, InRot, InScale);
-		GUnrealEd->UpdatePivotLocationForSelection();
-		GUnrealEd->RedrawLevelEditingViewports();
 
 	}
 	return FEdMode::InputDelta(InViewportClient, InViewport, InDrag, InRot, InScale);
diff --git a/Engine/Source/Editor/Sequencer/Private/SequencerEdMode.h b/Engine/Source/Editor/Sequencer/Private/SequencerEdMode.h
--- a/Engine/Source/Editor/Sequencer/Private/SequencerEdMode.h
+++ b/Engine/Source/Editor/Sequencer/Private/SequencerEdMode.h
@@ -99,6 +99,7 @@ public:
 	virtual bool MouseMove(FEditorViewportClient* ViewportClient, FViewport* Viewport, int32 x, int32 y) override;
 	virtual bool ProcessCapturedMouseMoves(FEditorViewportClient* InViewportClient, FViewport* InViewport, const TArrayView<FIntPoint>& CapturedMouseMoves) override;
 	virtual bool InputDelta(FEditorViewportClient* InViewportClient, FViewport* InViewport, FVector& InDrag, FRotator& InRot, FVector& InScale) override;
+	virtual void Tick(FEditorViewportClient* ViewportClient,float DeltaTime);
 
 	bool IsPressingMoveTimeSlider(FViewport* InViewport) const;
 	bool IsDoingDrag(FViewport* InViewport) const;
@@ -146,6 +147,8 @@ private:
 
 	FMarqueeDragTool DragToolHandler;
 
+	/** If the pivot location needs to be updated */
+	bool bUpdatePivot = false;
 };
 
 /**
