Fix warning V547: Expression is always...
#rnx

diff --git a/Engine/Source/Editor/Sequencer/Private/MVVM/Views/SSequencerOutlinerView.cpp b/Engine/Source/Editor/Sequencer/Private/MVVM/Views/SSequencerOutlinerView.cpp
--- a/Engine/Source/Editor/Sequencer/Private/MVVM/Views/SSequencerOutlinerView.cpp
+++ b/Engine/Source/Editor/Sequencer/Private/MVVM/Views/SSequencerOutlinerView.cpp
@@ -63,14 +63,14 @@ public:
 
 		if (Outliner->GetHoveredItem() == OutlinerItem)
 		{
-			return ItemStyle == EOutlinerItemViewBaseStyle::ContainerHeader
+			return /*ItemStyle == EOutlinerItemViewBaseStyle::ContainerHeader
 				? FLinearColor(FColor(52, 52, 52, 255))
-				: FLinearColor(FColor(72, 72, 72, 255));
+				: */FLinearColor(FColor(72, 72, 72, 255));
 		}
 
-		return ItemStyle == EOutlinerItemViewBaseStyle::ContainerHeader
+		return /*ItemStyle == EOutlinerItemViewBaseStyle::ContainerHeader
 			? FLinearColor(FColor(48, 48, 48, 255))
-			: FLinearColor(FColor(62, 62, 62, 255));
+			: */FLinearColor(FColor(62, 62, 62, 255));
 	}
 
 private:
diff --git a/Engine/Source/Runtime/Experimental/Chaos/Private/Chaos/Collision/ConvexTriangleContactPoint.cpp b/Engine/Source/Runtime/Experimental/Chaos/Private/Chaos/Collision/ConvexTriangleContactPoint.cpp
--- a/Engine/Source/Runtime/Experimental/Chaos/Private/Chaos/Collision/ConvexTriangleContactPoint.cpp
+++ b/Engine/Source/Runtime/Experimental/Chaos/Private/Chaos/Collision/ConvexTriangleContactPoint.cpp
@@ -478,7 +478,7 @@ namespace Chaos
 
 			ClipAxis = ConvexPlaneN;
 		}
-		else if (ContactType == EContactPointType::EdgeEdge)
+		else if (ContactType == EContactPointType::EdgeEdge) //-V547
 		{
 			// Edge-edge contact - clip triangle vs convex or vice-versa based on most opposing normals
 
