Sequencer: Fix label color of tracks to match label color of object binding tracks in the case the bindings are not valid. We can no longer assume invalid object binding always means red.

diff --git a/Engine/Source/Editor/Sequencer/Private/MVVM/ViewModels/TrackModel.cpp b/Engine/Source/Editor/Sequencer/Private/MVVM/ViewModels/TrackModel.cpp
--- a/Engine/Source/Editor/Sequencer/Private/MVVM/ViewModels/TrackModel.cpp
+++ b/Engine/Source/Editor/Sequencer/Private/MVVM/ViewModels/TrackModel.cpp
@@ -544,6 +544,15 @@ FSlateColor FTrackModel::GetLabelColor() const
 				if (TViewModelPtr<FObjectBindingModel> ObjectBindingModel = FindAncestorOfType<FObjectBindingModel>())
 				{
 					LabelParams.BindingID = ObjectBindingModel->GetObjectGuid();
+
+					// If the object binding model has an invalid binding, we want to use its label color, as it may be red or gray depending on situation
+					// and we want the children of that to have the same color.
+					// Otherwise, we can use the track's label color below
+					TArrayView<TWeakObjectPtr<> > BoundObjects = LabelParams.Player->FindBoundObjects(LabelParams.BindingID, LabelParams.SequenceID);
+					if (BoundObjects.Num() == 0)
+					{
+						return ObjectBindingModel->GetLabelColor();
+					}
 				}
 			}
 		}
diff --git a/Engine/Source/Editor/Sequencer/Private/MVVM/ViewModels/TrackRowModel.cpp b/Engine/Source/Editor/Sequencer/Private/MVVM/ViewModels/TrackRowModel.cpp
--- a/Engine/Source/Editor/Sequencer/Private/MVVM/ViewModels/TrackRowModel.cpp
+++ b/Engine/Source/Editor/Sequencer/Private/MVVM/ViewModels/TrackRowModel.cpp
@@ -237,7 +237,15 @@ FSlateColor FTrackRowModel::GetLabelColor() const
 	LabelParams.SequenceID = SequenceModel->GetSequenceID();
 	if (TViewModelPtr<FObjectBindingModel> ObjectBindingModel = FindAncestorOfType<FObjectBindingModel>())
 	{
-		LabelParams.BindingID = ObjectBindingModel->GetObjectGuid();
+		LabelParams.BindingID = ObjectBindingModel->GetObjectGuid();					
+		// If the object binding model has an invalid binding, we want to use its label color, as it may be red or gray depending on situation
+		// and we want the children of that to have the same color.
+		// Otherwise, we can use the track's label color below
+		TArrayView<TWeakObjectPtr<> > BoundObjects = LabelParams.Player->FindBoundObjects(LabelParams.BindingID, LabelParams.SequenceID);
+		if (BoundObjects.Num() == 0)
+		{
+			return ObjectBindingModel->GetLabelColor();
+		}
 	}
 	return Track->GetLabelColor(LabelParams);
 }
