Sequencer: Disable min width on the tree view outliner

This actually only applied to user-resizing the sequencer tree view, but it would not be respected if you resized the entire sequencer window since the grid panel used for layout only supports fill coefficients

diff --git a/Engine/Source/Editor/Sequencer/Private/SSequencer.cpp b/Engine/Source/Editor/Sequencer/Private/SSequencer.cpp
--- a/Engine/Source/Editor/Sequencer/Private/SSequencer.cpp
+++ b/Engine/Source/Editor/Sequencer/Private/SSequencer.cpp
@@ -876,7 +876,9 @@ void SSequencer::Construct(const FArguments& InArgs, TSharedRef<FSequencer> InSe
 
 					+ SSplitter::Slot()
 					.Value(FillCoefficient_0)
-					.MinSize(200)
+					// Can't use a minsize here because the grid panel that is actually being used to
+					//   lay out the widgets only supports fill coefficients and this leads to a disparity between the two
+					// .MinSize(200)
 					.OnSlotResized(SSplitter::FOnSlotResized::CreateSP(this, &SSequencer::OnColumnFillCoefficientChanged, 0))
 					[
 						SNew(SSpacer)
