Sequencer: Fix an issue with adding new sections to the binding lifetime track where you could create a section overlap by adding a new section to the track on top of the initial infinite section.
Also fix some issues with dragging section borders where for non-blendable sections you could create overlaps.

diff --git a/Engine/Source/Editor/Sequencer/Private/Tools/EditToolDragOperations.cpp b/Engine/Source/Editor/Sequencer/Private/Tools/EditToolDragOperations.cpp
--- a/Engine/Source/Editor/Sequencer/Private/Tools/EditToolDragOperations.cpp
+++ b/Engine/Source/Editor/Sequencer/Private/Tools/EditToolDragOperations.cpp
@@ -101,6 +101,43 @@ void FEditToolDragOperation::EndTransaction()
 	Sequencer.NotifyMovieSceneDataChanged( EMovieSceneDataChangeType::TrackValueChanged );
 }
 
+TRange<FFrameNumber> FEditToolDragOperation::GetSectionBoundaries(const UMovieSceneSection* Section)
+{
+	using namespace UE::Sequencer;
+
+	// Find the borders of where you can drag to
+	FFrameNumber LowerBound = TNumericLimits<int32>::Lowest(), UpperBound = TNumericLimits<int32>::Max();
+
+	// Find the track node for this section
+	TSharedPtr<FSectionModel> SectionHandle = Sequencer.GetNodeTree()->GetSectionModel(Section);
+	if (SectionHandle)
+	{
+		// Get the closest borders on either side
+		TViewModelPtr<ITrackAreaExtension> TrackModel = SectionHandle->FindAncestorOfType<ITrackAreaExtension>();
+		for (const TViewModelPtr<FSectionModel>& SectionModel : TrackModel->GetTrackAreaModelListAs<FSectionModel>())
+		{
+			const UMovieSceneSection* TestSection = SectionModel->GetSection();
+			TArray<UMovieSceneSection*> Sections;
+			GetSections(Sections);
+			if (!TestSection || Sections.Contains(TestSection))
+			{
+				continue;
+			}
+
+			if (TestSection->HasEndFrame() && Section->HasStartFrame() && TestSection->GetExclusiveEndFrame() <= Section->GetInclusiveStartFrame() && TestSection->GetExclusiveEndFrame() > LowerBound)
+			{
+				LowerBound = TestSection->GetExclusiveEndFrame();
+			}
+			if (TestSection->HasStartFrame() && Section->HasEndFrame() && TestSection->GetInclusiveStartFrame() >= Section->GetExclusiveEndFrame() && TestSection->GetInclusiveStartFrame() < UpperBound)
+			{
+				UpperBound = TestSection->GetInclusiveStartFrame();
+			}
+		}
+	}
+
+	return TRange<FFrameNumber>(LowerBound, UpperBound);
+}
+
 FResizeSection::FResizeSection( FSequencer& InSequencer, bool bInDraggingByEnd, bool bInIsSlipping )
 	: FEditToolDragOperation( InSequencer )
 	, bDraggingByEnd(bInDraggingByEnd)
@@ -279,6 +316,32 @@ void FResizeSection::OnDrag(const FPointerEvent& MouseEvent, FVector2D LocalMous
 			DeltaTime += (SnappedTime->SnappedTime - SnappedTime->OriginalTime).RoundToFrame();
 		}
 	}
+
+
+	auto GetMovementMaximums = [this](const UMovieSceneSection* Section, TOptional<FFrameNumber>& LeftMovementMaximum, TOptional<FFrameNumber>& RightMovementMaximum)
+	{
+		// We'll calculate this section's borders and clamp the possible delta time to be less than that
+
+		if (!Section->GetBlendType().IsValid())
+		{
+			TRange<FFrameNumber> SectionBoundaries = GetSectionBoundaries(Section);
+			LeftMovementMaximum = UE::MovieScene::DiscreteInclusiveLower(SectionBoundaries);
+			RightMovementMaximum = UE::MovieScene::DiscreteExclusiveUpper(SectionBoundaries);
+		}
+
+		if (Settings->GetIsSnapEnabled() && Settings->GetSnapKeysAndSectionsToPlayRange() && !Settings->ShouldKeepPlayRangeInSectionBounds())
+		{
+			if (!LeftMovementMaximum.IsSet() || LeftMovementMaximum.GetValue() < Sequencer.GetPlaybackRange().GetLowerBoundValue())
+			{
+				LeftMovementMaximum = Sequencer.GetPlaybackRange().GetLowerBoundValue();
+			}
+
+			if (!RightMovementMaximum.IsSet() || RightMovementMaximum.GetValue() > Sequencer.GetPlaybackRange().GetUpperBoundValue())
+			{
+				RightMovementMaximum = Sequencer.GetPlaybackRange().GetUpperBoundValue();
+			}
+		}
+	};
 	
 	/********************************************************************/
 	EMovieSceneSectionMovedResult SectionMovedResult(EMovieSceneSectionMovedResult::None);
@@ -289,6 +352,10 @@ void FResizeSection::OnDrag(const FPointerEvent& MouseEvent, FVector2D LocalMous
 			// It is only valid to dilate a fixed bound. Tracks can have mixed bounds types (ie: infinite upper, closed lower)
 			check(bDraggingByEnd ? Data.InitialRange.GetUpperBound().IsClosed() : Data.InitialRange.GetLowerBound().IsClosed());
 
+			TOptional<FFrameNumber> LeftMovementMaximum;
+			TOptional<FFrameNumber> RightMovementMaximum;
+			GetMovementMaximums(Data.MovieSection, LeftMovementMaximum, RightMovementMaximum);
+
 			FFrameNumber StartPosition  = bDraggingByEnd ? UE::MovieScene::DiscreteExclusiveUpper(Data.InitialRange) : UE::MovieScene::DiscreteInclusiveLower(Data.InitialRange);
 
 			FFrameNumber DilationOrigin;
@@ -325,7 +392,8 @@ void FResizeSection::OnDrag(const FPointerEvent& MouseEvent, FVector2D LocalMous
 			DataRange.SetLowerBound(TRangeBound<FFrameNumber>(DilationOrigin < StartPosition ? DilationOrigin : StartPosition));
 			DataRange.SetUpperBound(TRangeBound<FFrameNumber>(DilationOrigin > StartPosition ? DilationOrigin : StartPosition));
 
-			FFrameNumber NewPosition    = bDraggingByEnd ? FMath::Max(StartPosition + DeltaTime, DilationOrigin) : FMath::Min(StartPosition + DeltaTime, DilationOrigin);
+			FFrameNumber NewPosition    = bDraggingByEnd ? FMath::Clamp(StartPosition + DeltaTime, DilationOrigin, RightMovementMaximum.Get(TNumericLimits<int32>::Max()))
+														 : FMath::Clamp(StartPosition + DeltaTime, LeftMovementMaximum.Get(TNumericLimits<int32>::Lowest()), DilationOrigin);
 
 			float DilationFactor = FMath::Abs(NewPosition.Value - DilationOrigin.Value) / float(UE::MovieScene::DiscreteSize(DataRange));
 
@@ -382,6 +450,10 @@ void FResizeSection::OnDrag(const FPointerEvent& MouseEvent, FVector2D LocalMous
 			continue;
 		}
 
+		TOptional<FFrameNumber> LeftMovementMaximum;
+		TOptional<FFrameNumber> RightMovementMaximum;
+		GetMovementMaximums(Section, LeftMovementMaximum, RightMovementMaximum);
+
 		TSharedPtr<ISequencerSection> SectionInterface = SectionHandle->GetSectionInterface();
 
 		FFrameNumber NewTime = SectionInitTimes[Section] + DeltaTime;
@@ -396,9 +468,11 @@ void FResizeSection::OnDrag(const FPointerEvent& MouseEvent, FVector2D LocalMous
 				MinFrame = MinFrame + IntervalSnapThreshold;
 			}
 
+			FFrameNumber MaxFrame = RightMovementMaximum.Get(TNumericLimits<int32>::Max());
+
 			// Dragging the end of a section
-			// Ensure we aren't shrinking past the start time
-			NewTime = FMath::Max( NewTime, MinFrame );
+			// Ensure we aren't shrinking past the start time or into another section if we can't blend
+			NewTime = FMath::Clamp( NewTime, MinFrame, MaxFrame );
 			if (bIsSlipping)
 			{
 				SectionInterface->SlipSection( NewTime );
@@ -418,9 +492,11 @@ void FResizeSection::OnDrag(const FPointerEvent& MouseEvent, FVector2D LocalMous
 				MaxFrame = MaxFrame - IntervalSnapThreshold;
 			}
 
+			FFrameNumber MinFrame = LeftMovementMaximum.Get(TNumericLimits<int32>::Lowest());
+
 			// Dragging the start of a section
-			// Ensure we arent expanding past the end time
-			NewTime = FMath::Min( NewTime, MaxFrame );
+			// Ensure we arent expanding past the end time or into another section if we can't blend
+			NewTime = FMath::Clamp( NewTime, MinFrame, MaxFrame );
 
 			if (bIsSlipping)
 			{
@@ -931,41 +1007,6 @@ void FMoveKeysAndSections::ModifyNonSelectedSections()
 	}
 }
 
-TRange<FFrameNumber> FMoveKeysAndSections::GetSectionBoundaries(const UMovieSceneSection* Section)
-{
-	using namespace UE::Sequencer;
-
-	// Find the borders of where you can drag to
-	FFrameNumber LowerBound = TNumericLimits<int32>::Lowest(), UpperBound = TNumericLimits<int32>::Max();
-
-	// Find the track node for this section
-	TSharedPtr<FSectionModel> SectionHandle = Sequencer.GetNodeTree()->GetSectionModel(Section);
-	if (SectionHandle)
-	{
-		// Get the closest borders on either side
-		TViewModelPtr<ITrackAreaExtension> TrackModel = SectionHandle->FindAncestorOfType<ITrackAreaExtension>();
-		for (const TViewModelPtr<FSectionModel>& SectionModel : TrackModel->GetTrackAreaModelListAs<FSectionModel>())
-		{
-			const UMovieSceneSection* TestSection = SectionModel->GetSection();
-			if (!TestSection || Sections.Contains(TestSection))
-			{
-				continue;
-			}
-
-			if (TestSection->HasEndFrame() && Section->HasStartFrame() && TestSection->GetExclusiveEndFrame() <= Section->GetInclusiveStartFrame() && TestSection->GetExclusiveEndFrame() > LowerBound)
-			{
-				LowerBound = TestSection->GetExclusiveEndFrame();
-			}
-			if (TestSection->HasStartFrame() && Section->HasEndFrame() && TestSection->GetInclusiveStartFrame() >= Section->GetExclusiveEndFrame() && TestSection->GetInclusiveStartFrame() < UpperBound)
-			{
-				UpperBound = TestSection->GetInclusiveStartFrame();
-			}
-		}
-	}
-
-	return TRange<FFrameNumber>(LowerBound, UpperBound);
-}
-
 TOptional<FFrameNumber> FMoveKeysAndSections::GetMovementDeltaX(FFrameTime MouseTime)
 {
 	TOptional<FFrameNumber> DeltaX;
diff --git a/Engine/Source/Editor/Sequencer/Private/Tools/EditToolDragOperations.h b/Engine/Source/Editor/Sequencer/Private/Tools/EditToolDragOperations.h
--- a/Engine/Source/Editor/Sequencer/Private/Tools/EditToolDragOperations.h
+++ b/Engine/Source/Editor/Sequencer/Private/Tools/EditToolDragOperations.h
@@ -66,6 +66,11 @@ protected:
 	/** End an existing scoped transaction if one exists */
 	void EndTransaction();
 
+	virtual void GetSections(TArray<UMovieSceneSection*>& OutSections) {}
+
+	/** Get the bounds within which the specified section can be resized based on its surrounding sections */
+	TRange<FFrameNumber> GetSectionBoundaries(const UMovieSceneSection* Section);
+
 protected:
 
 	/** Scoped transaction for this drag operation */
@@ -144,6 +149,9 @@ private:
 
 	/** Optional snap field to use when dragging */
 	TOptional<FSequencerSnapField> SnapField;
+
+protected:
+	void GetSections(TArray<UMovieSceneSection*>& OutSections) override { OutSections = Sections.Array(); }
 };
 
 /**
@@ -172,9 +180,6 @@ protected:
 	/** Calculate the possible horizontal movement we can, constrained by sections running into things. */
 	TOptional<FFrameNumber> GetMovementDeltaX(FFrameTime MouseTime);
 
-	/** Get the bounds within which the specified section can move based on its surrounding sections */
-	TRange<FFrameNumber> GetSectionBoundaries(const UMovieSceneSection* Section);
-
 	/** Move selected sections, if any. */
 	bool HandleSectionMovement(FFrameTime MouseTime, FVector2D VirtualMousePos, FVector2D LocalMousePos, TOptional<FFrameNumber> MaxDeltaX, FFrameNumber DesiredDeltaX);
 	/** Move selected keys, if any. */
@@ -230,6 +235,9 @@ protected:
 
 	/** If the user is moving them via clicking on the Section then we'll allow vertical re-arranging, otherwise not. */
 	bool bAllowVerticalMovement;
+
+protected:
+	void GetSections(TArray<UMovieSceneSection*>& OutSections) override { OutSections = Sections.Array(); }
 };
 
 /**
diff --git a/Engine/Source/Runtime/MovieScene/Private/Sections/MovieSceneBindingLifetimeSection.cpp b/Engine/Source/Runtime/MovieScene/Private/Sections/MovieSceneBindingLifetimeSection.cpp
--- a/Engine/Source/Runtime/MovieScene/Private/Sections/MovieSceneBindingLifetimeSection.cpp
+++ b/Engine/Source/Runtime/MovieScene/Private/Sections/MovieSceneBindingLifetimeSection.cpp
@@ -8,6 +8,7 @@
 #include "EntitySystem/BuiltInComponentTypes.h"
 #include "EntitySystem/MovieSceneEntitySystemLinker.h"
 #include "Evaluation/MovieSceneEvaluationField.h"
+#include "MovieScene.h"
 
 #include UE_INLINE_GENERATED_CPP_BY_NAME(MovieSceneBindingLifetimeSection)
 
@@ -46,3 +47,38 @@ void UMovieSceneBindingLifetimeSection::ExternalPopulateEvaluationField(const TR
 	OutFieldBuilder->AddPersistentEntity(EffectiveRange, EntityIndex, MetaDataIndex);
 }
 
+void UMovieSceneBindingLifetimeSection::InitialPlacement(const TArray<UMovieSceneSection*>& Sections, FFrameNumber InStartTime, int32 Duration, bool bAllowMultipleRows)
+{
+	// We don't allow overlap or multi row here, so need a few extra rules for initial placement
+
+	check(Duration >= 0);
+
+	// Inclusive lower, exclusive upper bounds
+	SectionRange = TRange<FFrameNumber>(InStartTime, InStartTime + Duration);
+
+	for (;;)
+	{
+		UMovieSceneSection* OverlappedSection = const_cast<UMovieSceneSection*>(OverlapsWithSections(Sections));
+		if (OverlappedSection == nullptr)
+		{
+			break;
+		}
+
+		TRange<FFrameNumber> OtherRange = OverlappedSection->GetRange();
+
+		if (OtherRange.GetUpperBound().IsClosed())
+		{
+			MoveSection(OtherRange.GetUpperBoundValue() - InStartTime);
+		}
+		else
+		{
+			// We're likely overlapping the first infinite section here. Split it.
+			UMovieSceneSection* NewSection = OverlappedSection->SplitSection(FQualifiedFrameTime(InStartTime, GetTypedOuter<UMovieScene>()->GetTickResolution()), false);
+			// The new section will be the one on the right and will still have an open upper bound. Resize it to fit the new section.
+			TRange<FFrameNumber> NewSectionRange = NewSection->GetRange();
+			NewSectionRange.SetLowerBoundValue(InStartTime + Duration);
+			NewSection->SetRange(NewSectionRange);
+		}
+	}
+}
+
diff --git a/Engine/Source/Runtime/MovieScene/Public/Sections/MovieSceneBindingLifetimeSection.h b/Engine/Source/Runtime/MovieScene/Public/Sections/MovieSceneBindingLifetimeSection.h
--- a/Engine/Source/Runtime/MovieScene/Public/Sections/MovieSceneBindingLifetimeSection.h
+++ b/Engine/Source/Runtime/MovieScene/Public/Sections/MovieSceneBindingLifetimeSection.h
@@ -23,6 +23,8 @@ public:
 	virtual bool PopulateEvaluationFieldImpl(const TRange<FFrameNumber>& EffectiveRange, const FMovieSceneEvaluationFieldEntityMetaData& InMetaData, FMovieSceneEntityComponentFieldBuilder* OutFieldBuilder) override;
 	virtual void ImportEntityImpl(UMovieSceneEntitySystemLinker* EntityLinker, const FEntityImportParams& Params, FImportedEntity* OutImportedEntity) override;
 	void ExternalPopulateEvaluationField(const TRange<FFrameNumber>& EffectiveRange, const FMovieSceneEvaluationFieldEntityMetaData& InMetaData, FMovieSceneEntityComponentFieldBuilder* OutFieldBuilder);
+
+	void InitialPlacement(const TArray<UMovieSceneSection*>& Sections, FFrameNumber InStartTime, int32 InDuration, bool bAllowMultipleRows) override;
 };
 
 
