Sequencer- Fix a few issues with the way we handle looping subsequences/shots in Sequencer.

* Fix the drawing of sub sequence ranges for looped sequences to respect the loop count so that the sub sequence range shown is based on where we are in the master sequence. For example if we're on the last loop and it's only a partial loop as the sub sequence range cuts it off, then show the proper border. This brings this into parity with how we show partial non-looping subsequence ranges.
* Fix playback from within a looping subsequence. Previously this would constantly loop the same loop of the subsequence whereas dragging within a subsequence would properly advance the loop count. Now this will play properly, advancing the loop count appropriately with each loop.
* Fix some UI elements (scrubber frame count, loop count) to also show this information correctly.

diff --git a/Engine/Source/Editor/Sequencer/Private/Sequencer.cpp b/Engine/Source/Editor/Sequencer/Private/Sequencer.cpp
--- a/Engine/Source/Editor/Sequencer/Private/Sequencer.cpp
+++ b/Engine/Source/Editor/Sequencer/Private/Sequencer.cpp
@@ -681,6 +681,7 @@ FSequencer::FSequencer()
 	, PlaybackState( EMovieScenePlayerStatus::Stopped )
 	, LocalLoopIndexOnBeginScrubbing(FMovieSceneTimeWarping::InvalidWarpCount)
 	, LocalLoopIndexOffsetDuringScrubbing(0)
+	, MaxLocalLoopIndex(FMovieSceneTimeWarping::InvalidWarpCount)
 	, bPerspectiveViewportPossessionEnabled( true )
 	, bPerspectiveViewportCameraCutEnabled( false )
 	, bIsEditingWithinLevelEditor( false )
@@ -903,7 +904,7 @@ void FSequencer::Tick(float InDeltaTime)
 	{
 		FQualifiedFrameTime CurrentTime = GetLocalTime();
 		FFrameTime Offset = (AutoscrubOffset.GetValue() * AutoScrollFactor) * CurrentTime.Rate;
-		SetLocalTimeLooped(CurrentTime.Time + Offset);
+		SetLocalTimeLooped(CurrentTime.Time + Offset, RootToLocalLoopCounter);
 	}
 
 	// Reset to the root sequence if the focused sequence no longer exists. This can happen if either the subsequence has been deleted or the hierarchy has changed.
@@ -921,8 +922,12 @@ void FSequencer::Tick(float InDeltaTime)
 	{
 		FFrameTime NewGlobalTime = TimeController->RequestCurrentTime(GlobalTime, PlaybackSpeed * Dilation, GetFocusedDisplayRate());
 
-		// Put the time into local space
- 		SetLocalTimeLooped(NewGlobalTime * RootToLocalTransform);
+		// Put the time into clamped local space
+		FFrameTime LocalTime;
+		FMovieSceneWarpCounter WarpCounter;
+		CalculateLocalTimeClamped(NewGlobalTime, RootToLocalTransform, LocalTime, WarpCounter);
+
+ 		SetLocalTimeLooped(LocalTime, WarpCounter);
 
 		if (GetPlaybackStatus() == EMovieScenePlayerStatus::Playing)
 		{
@@ -1400,16 +1405,8 @@ void FSequencer::PopToSequenceInstance(FMovieSceneSequenceIDRef SequenceID)
 
 void FSequencer::UpdateSubSequenceData()
 {
-	const bool bIsScrubbing = GetPlaybackStatus() == EMovieScenePlayerStatus::Scrubbing;
-	const bool bIsSubSequenceLooping = RootToLocalTransform.NestedTransforms.Num() > 0 && RootToLocalTransform.NestedTransforms.Last().IsLooping();
-	const bool bIsScrubbingLoopingSubSequence = bIsScrubbing && bIsSubSequenceLooping;
-
 	SubSequenceRange = TRange<FFrameNumber>::Empty();
 	RootToLocalTransform = FMovieSceneSequenceTransform();
-	if (!bIsScrubbingLoopingSubSequence)
-	{
-		RootToLocalLoopCounter = FMovieSceneWarpCounter();
-	}
 	// else: we're scrubbing, and we don't want to increase/decrease the loop index quite yet,
 	// because that would mess up time transforms. This would be because the mouse would still be
 	// before/after the current loop, and therefore would already add/subtract more than a full
@@ -1418,6 +1415,8 @@ void FSequencer::UpdateSubSequenceData()
 	// Find the parent sub section and set up the sub sequence range, if necessary
 	if (ActiveTemplateIDs.Num() <= 1)
 	{
+		// Ensure in this case we also reset the RootToLocalLoopCounter.
+		RootToLocalLoopCounter = FMovieSceneWarpCounter();
 		return;
 	}
 
@@ -1429,10 +1428,16 @@ void FSequencer::UpdateSubSequenceData()
 		SubSequenceRange = SubSequenceData->PlayRange.Value;
 		RootToLocalTransform = SubSequenceData->RootToSequenceTransform;
 
+		const bool bIsScrubbing = GetPlaybackStatus() == EMovieScenePlayerStatus::Scrubbing;
+		const bool bIsSubSequenceLooping = RootToLocalTransform.NestedTransforms.Num() > 0 && RootToLocalTransform.NestedTransforms.Last().IsLooping();
+		const bool bIsScrubbingLoopingSubSequence = bIsScrubbing && bIsSubSequenceLooping;
+
 		const FQualifiedFrameTime RootTime = GetGlobalTime();
+		int32 CurLoopIndexOffset = 0;
 		if (!bIsScrubbingLoopingSubSequence)
 		{
 			FFrameTime LocalTime;
+			RootToLocalLoopCounter = FMovieSceneWarpCounter();
 			RootToLocalTransform.TransformTime(RootTime.Time, LocalTime, RootToLocalLoopCounter);
 		}
 		else
@@ -1460,29 +1465,98 @@ void FSequencer::UpdateSubSequenceData()
 			// of the loop. Note, again, that we only compute the new loop index for UI display purposes at this
 			// point (see comment at the beginning of this method). We will commit to the new loop indices
 			// once we're done scrubbing.
-			uint32 CurLoopIndex = 0;
 			if (PlayRangeSize > 0)
 			{
 				while (LocalTimeWithLastLoopUnwarped >= PlayRangeUpperBound)
 				{
 					LocalTimeWithLastLoopUnwarped = LocalTimeWithLastLoopUnwarped - PlayRangeSize;
-					++CurLoopIndex;
+					++CurLoopIndexOffset;
 				}
 				while (LocalTimeWithLastLoopUnwarped <= PlayRangeLowerBound)
 				{
 					LocalTimeWithLastLoopUnwarped = LocalTimeWithLastLoopUnwarped + PlayRangeSize;
-					--CurLoopIndex;
+					--CurLoopIndexOffset;
+				}
+			}
+		}
+
+		// If we're looping a sub sequence we may need to clamp to the max number of loops, which will modify SubSequenceRange and RootToLocalLoopCounter
+		if (bIsSubSequenceLooping)
+		{
+			const FMovieSceneSectionParameters SubSectionParameters = SubSequenceData->ToSubSectionParameters();
+			const TRange<FFrameNumber> ChildPlayRange = UMovieSceneSubSection::GetValidatedInnerPlaybackRange(SubSectionParameters, *SubSequenceData->GetSequence()->GetMovieScene());
+			const FFrameNumber ChildLength = UE::MovieScene::DiscreteSize(ChildPlayRange);
+			FMovieSceneSequenceID ParentID = Hierarchy.FindNode(ActiveTemplateIDs.Top())->ParentID;
+			UMovieScene* ParentMovieScene = (ParentID == GetRootTemplateID()) ? GetRootMovieSceneSequence()->GetMovieScene() : Hierarchy.FindSubSequence(ParentID)->GetMovieScene();
+
+			// Now we need to know how long this child play range is in the parent's time space. This is how we can figure
+			// out how many loops we can fit.
+			const FFrameRate ParentFrameRate = ParentMovieScene->GetTickResolution();
+			const FFrameRate ChildFrameRate = SubSequenceData->GetSequence()->GetMovieScene()->GetTickResolution();
+
+			const float ChildTimeScale = SubSequenceData->OuterToInnerTransform.LinearTransform.TimeScale;
+			const float InvChildTimeScale = FMath::IsNearlyZero(ChildTimeScale) ? 1.0f : 1.0f / ChildTimeScale;
+
+			const FFrameNumber ChildLengthInParentSpace = (ConvertFrameTime(ChildLength, ChildFrameRate, ParentFrameRate) * InvChildTimeScale).FrameNumber;
+			const FFrameNumber ChildFirstLoopLength = ChildLength - SubSequenceData->ParentFirstLoopStartFrameOffset;
+			const FFrameNumber ChildFirstLoopLengthInParentSpace = (ConvertFrameTime(ChildFirstLoopLength, ChildFrameRate, ParentFrameRate) * InvChildTimeScale).FrameNumber;
+
+			// We can finally start iterating: we iterate for how many times as we can fit the child sequence's length,
+			// modified by the time scale, into the parent play range.
+			const TRange<FFrameNumber> ParentPlayRange = SubSequenceData->ParentPlayRange.Value;
+			const FFrameNumber ParentExclusiveEnd = UE::MovieScene::DiscreteExclusiveUpper(ParentPlayRange);
+			FFrameNumber CurLoopStart = UE::MovieScene::DiscreteInclusiveLower(ParentPlayRange);
+			FFrameNumber CurLoopEnd = CurLoopStart + FMath::Min(
+				FMath::Max(ChildFirstLoopLengthInParentSpace, FFrameNumber(0)),
+				ParentExclusiveEnd);
+
+			int32 LocalLoopIndex = FMath::Max(0, (int32)RootToLocalLoopCounter.WarpCounts.Last() + CurLoopIndexOffset);
+			uint8 CurLoopIndex = 0;
+			TRange<FFrameNumber> NewSubSequenceRange = TRange<FFrameNumber>(CurLoopStart, FMath::Min(CurLoopEnd, ParentExclusiveEnd));
+			while (CurLoopStart < ParentExclusiveEnd)
+			{
+				if (CurLoopIndex <= LocalLoopIndex)
+				{
+					NewSubSequenceRange = TRange<FFrameNumber>(CurLoopStart, FMath::Min(CurLoopEnd, ParentExclusiveEnd));
+					NewSubSequenceRange = SubSequenceData->OuterToInnerTransform.TransformRangeConstrained(NewSubSequenceRange);
+				}
+
+				CurLoopStart = CurLoopEnd;
+				CurLoopEnd += FMath::Max(ChildLengthInParentSpace, FFrameNumber(0));
+				CurLoopIndex++;
+			}
+
+			MaxLocalLoopIndex = CurLoopIndex - 1;
+
+			if (bIsScrubbingLoopingSubSequence)
+			{
+				// Clamp LocalLoopIndexOffset between 0 and MaxLocalLoopIndex- as it's a delta we need to subtract off LocalLoopIndexOnBeginScrubbing.
+				CurLoopIndexOffset = FMath::Max(-LocalLoopIndexOnBeginScrubbing, FMath::Min(MaxLocalLoopIndex - LocalLoopIndexOnBeginScrubbing, CurLoopIndexOffset));
+
+				if (CurLoopIndexOffset != LocalLoopIndexOffsetDuringScrubbing)
+				{
+					LocalLoopIndexOffsetDuringScrubbing = CurLoopIndexOffset;
+					// If we jumped to the previous or next loop, we need to invalidate the global marked frames because
+					// the focused (currently edited) sequence's time transform just changed.
+					InvalidateGlobalMarkedFramesCache();
 				}
 			}
-			if (CurLoopIndex != LocalLoopIndexOffsetDuringScrubbing)
+			else
 			{
-				LocalLoopIndexOffsetDuringScrubbing = CurLoopIndex;
-				// If we jumped to the previous or next loop, we need to invalidate the global marked frames because
-				// the focused (currently edited) sequence's time transform just changed.
-				InvalidateGlobalMarkedFramesCache();
+				ensure(RootToLocalLoopCounter.WarpCounts.Num() > 0);
+				
+				// Clamp warp count
+				RootToLocalLoopCounter.WarpCounts.Last() = FMath::Clamp(RootToLocalLoopCounter.WarpCounts.Last(), 0, MaxLocalLoopIndex);
 			}
+
+			SubSequenceRange = NewSubSequenceRange;
 		}
 	}
+	else
+	{
+		// Ensure in this case we also reset the RootToLocalLoopCounter.
+		RootToLocalLoopCounter = FMovieSceneWarpCounter();
+	}
 }
 
 void FSequencer::UpdateSequencerCustomizations(const UMovieSceneSequence* PreviousFocusedSequence)
@@ -2997,7 +3071,11 @@ FQualifiedFrameTime FSequencer::GetLocalTime() const
 	const FFrameTime CurrentPosition   = PlayPosition.GetCurrentPosition();
 
 	const FFrameTime RootTime = ConvertFrameTime(CurrentPosition, PlayPosition.GetInputRate(), PlayPosition.GetOutputRate());
-	return FQualifiedFrameTime(RootTime * RootToLocalTransform, FocusedResolution);
+
+	FFrameTime LocalTime;
+	FMovieSceneWarpCounter LoopCounter;
+	CalculateLocalTimeClamped(RootTime, RootToLocalTransform, LocalTime, LoopCounter);
+	return FQualifiedFrameTime(LocalTime, FocusedResolution);
 }
 
 
@@ -3382,11 +3460,11 @@ void FSequencer::UpdateAutoScroll(double NewTime, float ThresholdPercentage)
 	{
 		if (AutoscrollOffset.GetValue() < 0 && LocalTime.AsSeconds() > ViewRange.GetLowerBoundValue() + Threshold)
 		{
-			SetLocalTimeLooped( (ViewRange.GetLowerBoundValue() + Threshold) * LocalTime.Rate );
+			SetLocalTimeLooped( (ViewRange.GetLowerBoundValue() + Threshold) * LocalTime.Rate, RootToLocalLoopCounter);
 		}
 		else if (AutoscrollOffset.GetValue() > 0 && LocalTime.AsSeconds() < ViewRange.GetUpperBoundValue() - Threshold)
 		{
-			SetLocalTimeLooped( (ViewRange.GetUpperBoundValue() - Threshold) * LocalTime.Rate );
+			SetLocalTimeLooped( (ViewRange.GetUpperBoundValue() - Threshold) * LocalTime.Rate, RootToLocalLoopCounter);
 		}
 	}
 
@@ -4470,7 +4548,7 @@ FReply FSequencer::OnPlay(bool bTogglePlay)
 	}
 	else
 	{
-		TRange<FFrameNumber> TimeBounds = GetTimeBounds();
+		TRange<FFrameNumber> TimeBounds = GetRootTimeBounds();
 
 		FFrameNumber MinInclusiveTime = UE::MovieScene::DiscreteInclusiveLower(TimeBounds);
 		FFrameNumber MaxInclusiveTime = UE::MovieScene::DiscreteExclusiveUpper(TimeBounds) - 1;
@@ -4703,21 +4781,20 @@ ESequencerLoopMode FSequencer::GetLoopMode() const
 	return Settings->GetLoopMode();
 }
 
-void FSequencer::SetLocalTimeLooped(FFrameTime NewLocalTime)
+void FSequencer::SetLocalTimeLooped(FFrameTime NewLocalTime, FMovieSceneWarpCounter WarpCounter)
 {
 	TOptional<EMovieScenePlayerStatus::Type> NewPlaybackStatus;
+	// Default to the RootToLocalLoopCounter
+	if (WarpCounter.NumWarpCounts() == 0)
+	{
+		WarpCounter = RootToLocalLoopCounter;
+	}
 
-	const FMovieSceneSequenceTransform LocalToRootTransform = RootToLocalTransform.InverseFromLoop(RootToLocalLoopCounter);
+	const FMovieSceneSequenceTransform LocalToRootTransform = RootToLocalTransform.InverseFromLoop(WarpCounter);
 
 	FFrameTime NewGlobalTime = NewLocalTime * LocalToRootTransform;
 
-	TRange<FFrameNumber> TimeBounds = GetTimeBounds();
-
 	bool         bResetPosition       = false;
-	FFrameRate   LocalTickResolution  = GetFocusedTickResolution();
-	FFrameRate   RootTickResolution   = GetRootTickResolution();
-	FFrameNumber MinInclusiveTime     = UE::MovieScene::DiscreteInclusiveLower(TimeBounds);
-	FFrameNumber MaxInclusiveTime     = UE::MovieScene::DiscreteExclusiveUpper(TimeBounds)-1;
 
 	bool bHasJumped = false;
 	bool bRestarted = false;
@@ -4731,6 +4808,10 @@ void FSequencer::SetLocalTimeLooped(FFrameTime NewLocalTime)
 	}
 	else if (GetLoopMode() == ESequencerLoopMode::SLM_Loop || GetLoopMode() == ESequencerLoopMode::SLM_LoopSelectionRange)
 	{
+		TRange<FFrameNumber> TimeBounds = GetTimeBounds();
+		FFrameNumber MinInclusiveTime = UE::MovieScene::DiscreteInclusiveLower(TimeBounds);
+		FFrameNumber MaxInclusiveTime = UE::MovieScene::DiscreteExclusiveUpper(TimeBounds) - 1;
+
 		const UMovieSceneSequence* FocusedSequence = GetFocusedMovieSceneSequence();
 		if (FocusedSequence)
 		{
@@ -4745,8 +4826,9 @@ void FSequencer::SetLocalTimeLooped(FFrameTime NewLocalTime)
 	}
 	else
 	{
-		TRange<double> WorkingRange = GetClampRange();
-
+		TRange<FFrameNumber> RootTimeBounds = GetRootTimeBounds();
+		FFrameNumber MinInclusiveTime = UE::MovieScene::DiscreteInclusiveLower(RootTimeBounds);
+		FFrameNumber MaxInclusiveTime = UE::MovieScene::DiscreteExclusiveUpper(RootTimeBounds) - 1;
 		bool bReachedEnd = false;
 		if (PlaybackSpeed > 0)
 		{
@@ -4765,6 +4847,7 @@ void FSequencer::SetLocalTimeLooped(FFrameTime NewLocalTime)
 		}
 	}
 
+	FFrameRate   RootTickResolution = GetRootTickResolution();
 	// Ensure the time is in the current view - must occur before the time cursor changes
 	UMovieScene* MovieScene = GetFocusedMovieSceneSequence()->GetMovieScene();
 	if (IsAutoScrollEnabled())
@@ -4817,6 +4900,28 @@ void FSequencer::RestorePlaybackSpeedAfterPlay()
 	PlaybackSpeed = PlaybackSpeedBeforePlay;
 }
 
+void FSequencer::CalculateLocalTimeClamped(FFrameTime RootTime, const FMovieSceneSequenceTransform& RootToParentChainTransform, FFrameTime& OutLocalTime, FMovieSceneWarpCounter& OutLoopCounter) const
+{
+	RootToParentChainTransform.TransformTime(RootTime, OutLocalTime, OutLoopCounter);
+	if (RootToParentChainTransform.NestedTransforms.Num() > 0 && RootToParentChainTransform.NestedTransforms.Last().IsLooping())
+	{
+		ensure(OutLoopCounter.WarpCounts.Num() > 0);
+		// Clamp warp count
+		OutLoopCounter.WarpCounts.Last() = FMath::Clamp(OutLoopCounter.WarpCounts.Last(), 0, MaxLocalLoopIndex);
+
+		// Update OutLocalTime based on the current clamped warp count]
+		FMovieSceneSequenceTransform RootToLocalTransformWithoutLeafLooping = RootToParentChainTransform;
+		FMovieSceneNestedSequenceTransform LeafLooping = RootToLocalTransformWithoutLeafLooping.NestedTransforms.Pop();
+		FFrameTime LocalTimeWithLastLoopUnwarped = RootTime * RootToLocalTransformWithoutLeafLooping;
+		LocalTimeWithLastLoopUnwarped = LocalTimeWithLastLoopUnwarped * LeafLooping.LinearTransform;
+		if (LeafLooping.IsLooping())
+		{
+			LeafLooping.Warping.TransformTimeSpecific(
+				LocalTimeWithLastLoopUnwarped, OutLoopCounter.WarpCounts.Last(), OutLocalTime);
+		}
+	}
+}
+
 EPlaybackMode::Type FSequencer::GetPlaybackMode() const
 {
 	if (PlaybackState == EMovieScenePlayerStatus::Playing)
@@ -4842,6 +4947,9 @@ void FSequencer::UpdateTimeBoundsToFocusedMovieScene()
 		return;
 	}
 
+	// We have to update subsequence data before getting the local time on a new subsequence
+	UpdateSubSequenceData();
+
 	FQualifiedFrameTime CurrentTime = GetLocalTime();
 
 	// Set the view range to:
@@ -4889,6 +4997,37 @@ TRange<FFrameNumber> FSequencer::GetTimeBounds() const
 	return SubSequenceRange;
 }
 
+TRange<FFrameNumber> FSequencer::GetRootTimeBounds() const
+{
+	const UMovieSceneSequence* RootMovieSceneSequence = GetRootMovieSceneSequence();
+	if (!RootMovieSceneSequence)
+	{
+		return TRange<FFrameNumber>(-100000, 100000);
+	}
+
+	TRange<FFrameNumber> RootPlayRangeRootSpace = RootMovieSceneSequence->GetMovieScene()->GetPlaybackRange();
+	if (RootMovieSceneSequence == GetFocusedMovieSceneSequence())
+	{
+		return RootPlayRangeRootSpace;
+	}
+	
+	// Edge case for looping- we need to calculate clamped local time
+	if (RootToLocalTransform.NestedTransforms.Num() > 0 && RootToLocalTransform.NestedTransforms.Last().IsLooping())
+	{
+		FFrameTime RootPlayRangeLocalSpaceLowerBound;
+		FFrameTime RootPlayRangeLocalSpaceUpperBound;
+		FMovieSceneWarpCounter WarpCounter;
+		CalculateLocalTimeClamped(RootPlayRangeRootSpace.GetLowerBoundValue(), RootToLocalTransform, RootPlayRangeLocalSpaceLowerBound, WarpCounter);
+		CalculateLocalTimeClamped(RootPlayRangeRootSpace.GetUpperBoundValue(), RootToLocalTransform, RootPlayRangeLocalSpaceUpperBound, WarpCounter);
+
+		return TRange<FFrameNumber>(RootPlayRangeLocalSpaceLowerBound.FloorToFrame(), RootPlayRangeLocalSpaceUpperBound.CeilToFrame());
+	}
+	else
+	{
+		return RootToLocalTransform.TransformRangeConstrained(RootPlayRangeRootSpace);
+	}
+}
+
 
 void FSequencer::SetViewRange(TRange<double> NewViewRange, EViewRangeInterpolation Interpolation)
 {
@@ -5208,7 +5347,10 @@ FString FSequencer::GetFrameTimeText() const
 
 	const FFrameTime RootTime = ConvertFrameTime(CurrentPosition, PlayPosition.GetInputRate(), PlayPosition.GetOutputRate());
 	
-	const FFrameTime LocalTime = RootTime * RootToParentChainTransform;
+	FFrameTime LocalTime;
+	FMovieSceneWarpCounter LoopCounter;
+
+	CalculateLocalTimeClamped(RootTime, RootToParentChainTransform, LocalTime, LoopCounter);
 
 	return GetNumericTypeInterface()->ToString(LocalTime.GetFrame().Value);
 }
diff --git a/Engine/Source/Editor/Sequencer/Private/Sequencer.h b/Engine/Source/Editor/Sequencer/Private/Sequencer.h
--- a/Engine/Source/Editor/Sequencer/Private/Sequencer.h
+++ b/Engine/Source/Editor/Sequencer/Private/Sequencer.h
@@ -513,8 +513,8 @@ public:
 	void OnTogglePilotCamera();
 	bool IsPilotCamera() const;
 
-	/** Set the new global time, accounting for looping options */
-	void SetLocalTimeLooped(FFrameTime InTime);
+	/** Sets the new global time calculated from local time and the given warp counter, accounting for looping options */
+	void SetLocalTimeLooped(FFrameTime InTime, FMovieSceneWarpCounter WarpCounter=FMovieSceneWarpCounter());
 
 	ESequencerLoopMode GetLoopMode() const;
 
@@ -853,6 +853,12 @@ public:
 	 */
 	TRange<FFrameNumber> GetTimeBounds() const;
 
+
+	/**
+	 * Gets the time boundaries of the root movie scene in local space. If this is a looping subsequence, this will include all loops.
+	 */
+	TRange<FFrameNumber> GetRootTimeBounds() const;
+
 protected:
 
 	// FCameraCutPlaybackCapability interface
@@ -1146,6 +1152,9 @@ private:
 	int32 FindClosestPlaybackSpeed(float InPlaybackSpeed, bool bExactOnly = false) const;
 	void RestorePlaybackSpeedAfterPlay();
 
+	// Given the root sequence time, returns the local time and loop counter clamped to the maximum number of loops
+	void CalculateLocalTimeClamped(FFrameTime RootTime, const FMovieSceneSequenceTransform& RootToParentChainTransform, FFrameTime& OutTime, FMovieSceneWarpCounter& OutLoopCounter) const;
+
 public:
 
 	/** Helper function which returns how many frames (in tick resolution) one display rate frame represents. */
@@ -1282,10 +1291,13 @@ private:
 	FMovieScenePlaybackPosition PlayPosition;
 
 	/** Local loop index at the time we began scrubbing */
-	uint32 LocalLoopIndexOnBeginScrubbing;
+	int32 LocalLoopIndexOnBeginScrubbing;
 
 	/** Local loop index to add for the purposes of displaying it in the UI */
-	uint32 LocalLoopIndexOffsetDuringScrubbing;
+	int32 LocalLoopIndexOffsetDuringScrubbing;
+
+	/** MaxLocalLoopIndex as calculated in UpdateSubSequenceData. Used to ensure LocalTime is also clamped to the correct number of loops. */
+	int32 MaxLocalLoopIndex;
 
 	/** The playback speed */
 	float PlaybackSpeed;
