Sequencer: Fixed potential null access picked up by [REDACTED_PERSON]

diff --git a/Engine/Source/Editor/Sequencer/Private/SkeletalMeshComponentSchema.cpp b/Engine/Source/Editor/Sequencer/Private/SkeletalMeshComponentSchema.cpp
--- a/Engine/Source/Editor/Sequencer/Private/SkeletalMeshComponentSchema.cpp
+++ b/Engine/Source/Editor/Sequencer/Private/SkeletalMeshComponentSchema.cpp
@@ -147,10 +147,15 @@ void FSkeletalMeshComponentSchema::BindAnimationInstance(TWeakPtr<ISequencer> We
 	FScopedTransaction Transaction(LOCTEXT("BindAnimInstance", "Add Anim Instance to Sequencer"));
 
 	UMovieSceneSequence*          Sequence          = Sequencer->GetFocusedMovieSceneSequence();
-	UMovieScene*                  MovieScene        = Sequence ? Sequence->GetMovieScene() : nullptr;;
+	if (!Sequence)
+	{
+		return;
+	}
+
+	UMovieScene*                  MovieScene        = Sequence->GetMovieScene();
 	FMovieSceneBindingReferences* BindingReferences = Sequence->GetBindingReferences();
 
-	if (!Sequence || !MovieScene || !BindingReferences)
+	if (!MovieScene || !BindingReferences)
 	{
 		return;
 	}
