Sequencer: Additional GetPlaybackContext fixes

#rnx

diff --git a/Engine/Plugins/MovieScene/LevelSequenceEditor/Source/LevelSequenceEditor/Private/LevelSequenceEditorSubsystem.cpp b/Engine/Plugins/MovieScene/LevelSequenceEditor/Source/LevelSequenceEditor/Private/LevelSequenceEditorSubsystem.cpp
--- a/Engine/Plugins/MovieScene/LevelSequenceEditor/Source/LevelSequenceEditor/Private/LevelSequenceEditorSubsystem.cpp
+++ b/Engine/Plugins/MovieScene/LevelSequenceEditor/Source/LevelSequenceEditor/Private/LevelSequenceEditorSubsystem.cpp
@@ -1411,8 +1411,7 @@ void ULevelSequenceEditorSubsystem::FixActorReferences()
 		return;
 	}
 
-	UWorld* PlaybackContext = Cast<UWorld>(Sequencer->GetPlaybackContext());
-
+	UWorld* PlaybackContext = Sequencer->GetPlaybackContext()->GetWorld();
 	if (!PlaybackContext)
 	{
 		return;
diff --git a/Engine/Source/Editor/MovieSceneTools/Private/MovieSceneToolHelpers.cpp b/Engine/Source/Editor/MovieSceneTools/Private/MovieSceneToolHelpers.cpp
--- a/Engine/Source/Editor/MovieSceneTools/Private/MovieSceneToolHelpers.cpp
+++ b/Engine/Source/Editor/MovieSceneTools/Private/MovieSceneToolHelpers.cpp
@@ -3378,7 +3378,7 @@ private:
 		// Import static cameras first
 		ImportFBXCamera(FbxImporter, Sequence, *Sequencer, ObjectBindingMap, bMatchByNameOnly, bCreateCameras.IsSet() ? bCreateCameras.GetValue() : ImportFBXSettings->bCreateCameras);
 
-		UWorld* World = Cast<UWorld>(Sequencer->GetPlaybackContext());
+		UWorld* World = Sequencer->GetPlaybackContext()->GetWorld();
 		bool bValid = MovieSceneToolHelpers::ImportFBXIfReady(World, Sequence, Sequencer, Sequencer->GetFocusedTemplateID(), ObjectBindingMap, ImportFBXSettings, InOutParams);
 	
 		Sequencer->NotifyMovieSceneDataChanged(EMovieSceneDataChangeType::MovieSceneStructureItemAdded);
diff --git a/Engine/Source/Editor/MovieSceneTools/Private/TrackEditors/MaterialParameterCollectionTrackEditor.cpp b/Engine/Source/Editor/MovieSceneTools/Private/TrackEditors/MaterialParameterCollectionTrackEditor.cpp
--- a/Engine/Source/Editor/MovieSceneTools/Private/TrackEditors/MaterialParameterCollectionTrackEditor.cpp
+++ b/Engine/Source/Editor/MovieSceneTools/Private/TrackEditors/MaterialParameterCollectionTrackEditor.cpp
@@ -323,7 +323,7 @@ void FMaterialParameterCollectionTrackEditor::AddScalarParameter(UMovieSceneMate
 	}
 
 	float Value = Parameter.DefaultValue;
-	if (UWorld* World = Cast<UWorld>(GetSequencer()->GetPlaybackContext()))
+	if (UWorld* World = GetSequencer()->GetPlaybackContext()->GetWorld())
 	{
 		if (UMaterialParameterCollectionInstance* Instance = World->GetParameterCollectionInstance(Track->MPC))
 		{
@@ -348,7 +348,7 @@ void FMaterialParameterCollectionTrackEditor::AddVectorParameter(UMovieSceneMate
 	}
 
 	FLinearColor Value = Parameter.DefaultValue;
-	if (UWorld* World = Cast<UWorld>(GetSequencer()->GetPlaybackContext()))
+	if (UWorld* World = GetSequencer()->GetPlaybackContext()->GetWorld())
 	{
 		if (UMaterialParameterCollectionInstance* Instance = World->GetParameterCollectionInstance(Track->MPC))
 		{
diff --git a/Engine/Source/Editor/Sequencer/Private/SequencerUtilities.cpp b/Engine/Source/Editor/Sequencer/Private/SequencerUtilities.cpp
--- a/Engine/Source/Editor/Sequencer/Private/SequencerUtilities.cpp
+++ b/Engine/Source/Editor/Sequencer/Private/SequencerUtilities.cpp
@@ -861,7 +861,7 @@ FGuid FSequencerUtilities::CreateCameraWithRig(TSharedRef<ISequencer> Sequencer,
 	}
 
 	// Create a cine camera actor
-	UWorld* PlaybackContext = Cast<UWorld>(Sequencer->GetPlaybackContext());
+	UWorld* PlaybackContext = Sequencer->GetPlaybackContext()->GetWorld();
 	OutActor = PlaybackContext->SpawnActor<ACineCameraActor>();
 	CameraGuid = CreateBinding(Sequencer, *OutActor);
 
@@ -1192,7 +1192,7 @@ FMovieScenePossessable* FSequencerUtilities::ConvertToPossessable(TSharedRef<ISe
 	SpawnInfo.bDeferConstruction = true;
 	SpawnInfo.Template = SpawnableActorTemplate;
 
-	UWorld* PlaybackContext = Cast<UWorld>(Sequencer->GetPlaybackContext());
+	UWorld* PlaybackContext = Sequencer->GetPlaybackContext()->GetWorld();
 	AActor* PossessedActor = PlaybackContext->SpawnActor(Spawnable->GetObjectTemplate()->GetClass(), &DefaultTransform, SpawnInfo);
 
 	if (!PossessedActor)
@@ -2315,7 +2315,7 @@ bool FSequencerUtilities::PasteBindings(const FString& TextToImport, TSharedRef<
 	for (int32 PossessableGuidIndex = 0; PossessableGuidIndex < PossessableGuids.Num(); ++PossessableGuidIndex)
 	{
 		FMovieScenePossessable* Possessable = MovieScene->FindPossessable(PossessableGuids[PossessableGuidIndex]);
-		UWorld* PlaybackContext = Cast<UWorld>(Sequencer->GetPlaybackContext());
+		UWorld* PlaybackContext = Sequencer->GetPlaybackContext()->GetWorld();
 		if (Possessable && PlaybackContext)
 		{
 			for (TActorIterator<AActor> ActorItr(PlaybackContext); ActorItr; ++ActorItr)
