Sequencer: Fixup one more assumption where the playback context was a UWorld

diff --git a/Engine/Source/Editor/Sequencer/Private/Sequencer.cpp b/Engine/Source/Editor/Sequencer/Private/Sequencer.cpp
--- a/Engine/Source/Editor/Sequencer/Private/Sequencer.cpp
+++ b/Engine/Source/Editor/Sequencer/Private/Sequencer.cpp
@@ -8590,7 +8590,7 @@ void FSequencer::OnLoadRecordedData()
 				NotifyMovieSceneDataChanged(EMovieSceneDataChangeType::MovieSceneStructureItemAdded);
 
 			}; //callback
-			UWorld* PlaybackContext = Cast<UWorld>(GetPlaybackContext());
+			UWorld* PlaybackContext = GetPlaybackContext()->GetWorld();
 			for (const FString& FileName : OpenFilenames)
 			{
 				Recorder->LoadRecordedSequencerFile(FocusedMovieSceneSequence, PlaybackContext, FileName, OnReadComplete);
