Sequencer: UX changes for Playback Controls

- Change text "+Track" to "+Add"
- Move current frame widgets to playback controls, remove orange text on the current frame, frame counter is now in the tooltip
- Move clear in the status bar to the far left so it's easier to hit when you hover over the status bar
- Status bar changes - visible on hover of the info button, fades away/hides automatically on unhover after 2 seconds
- Add status bar and current frame to the curve editor playback controls

diff --git a/Engine/Plugins/MovieScene/LevelSequenceEditor/Source/LevelSequenceEditor/Private/CinematicViewport/SCinematicLevelViewport.cpp b/Engine/Plugins/MovieScene/LevelSequenceEditor/Source/LevelSequenceEditor/Private/CinematicViewport/SCinematicLevelViewport.cpp
--- a/Engine/Plugins/MovieScene/LevelSequenceEditor/Source/LevelSequenceEditor/Private/CinematicViewport/SCinematicLevelViewport.cpp
+++ b/Engine/Plugins/MovieScene/LevelSequenceEditor/Source/LevelSequenceEditor/Private/CinematicViewport/SCinematicLevelViewport.cpp
@@ -539,6 +539,24 @@ float SCinematicLevelViewport::GetPlayTimeMinDesiredWidth() const
 		FString LowerBoundStr = Sequencer->GetNumericTypeInterface()->ToString(ViewRange.GetLowerBoundValue());
 		FString UpperBoundStr = Sequencer->GetNumericTypeInterface()->ToString(ViewRange.GetUpperBoundValue());
 
+		// Always measure with the negative and subframe indicator so that the size doesn't change when there is and isn't a subframe
+		if (!LowerBoundStr.Contains(TEXT("*")))
+		{
+			LowerBoundStr += TEXT("*");
+		}
+		if (!LowerBoundStr.Contains(TEXT("-")))
+		{
+			LowerBoundStr += TEXT("-");
+		}
+		if (!UpperBoundStr.Contains(TEXT("*")))
+		{
+			UpperBoundStr += TEXT("*");
+		}
+		if (!UpperBoundStr.Contains(TEXT("-")))
+		{
+			UpperBoundStr += TEXT("-");
+		}
+
 		const FSlateFontInfo PlayTimeFont = FAppStyle::GetFontStyle("Sequencer.FixedFont");
 
 		const TSharedRef< FSlateFontMeasure > FontMeasureService = FSlateApplication::Get().GetRenderer()->GetFontMeasureService();
diff --git a/Engine/Source/Editor/CurveEditor/Private/Tree/SCurveEditorTreeFilterStatusBar.cpp b/Engine/Source/Editor/CurveEditor/Private/Tree/SCurveEditorTreeFilterStatusBar.cpp
--- a/Engine/Source/Editor/CurveEditor/Private/Tree/SCurveEditorTreeFilterStatusBar.cpp
+++ b/Engine/Source/Editor/CurveEditor/Private/Tree/SCurveEditorTreeFilterStatusBar.cpp
@@ -29,23 +29,34 @@ void SCurveEditorTreeFilterStatusBar::Construct(const FArguments& InArgs, TShare
 	WeakCurveEditor = InCurveEditor;
 
 	ChildSlot
-	.Padding(FMargin(5.f, 3.f))
+	.Padding(FMargin(5.f, 0.f))
 	[
 		SNew(SHorizontalBox)
 
 		+ SHorizontalBox::Slot()
+		.AutoWidth()
+		.Padding(FMargin(3.f, 0.f, 0.f, 0.f))
 		[
-			SAssignNew(TextBlock, STextBlock)
+			SNew(SBorder)
+			.BorderImage(FAppStyle::GetBrush("ToolPanel.GroupBorder"))
+			.Padding(FMargin(3.0f, 3.0f))
+			[
+				SNew(SHyperlink)
+				.Visibility(this, &SCurveEditorTreeFilterStatusBar::GetVisibilityFromFilter)
+				.Text(LOCTEXT("ClearFilters", "clear"))
+				.OnNavigate(this, &SCurveEditorTreeFilterStatusBar::ClearFilters)
+			]
 		]
 
 		+ SHorizontalBox::Slot()
-		.AutoWidth()
-		.Padding(FMargin(3.f, 0.f, 0.f, 0.f))
 		[
-			SNew(SHyperlink)
-			.Visibility(this, &SCurveEditorTreeFilterStatusBar::GetVisibilityFromFilter)
-			.Text(LOCTEXT("ClearFilters", "clear"))
-			.OnNavigate(this, &SCurveEditorTreeFilterStatusBar::ClearFilters)
+			SNew(SBorder)
+			.BorderImage(FAppStyle::GetBrush("ToolPanel.GroupBorder"))
+			.Padding(FMargin(3.0f, 3.0f))
+			.Visibility(EVisibility::HitTestInvisible)
+			[
+				SAssignNew(TextBlock, STextBlock)
+			]
 		]
 	];
 
@@ -136,4 +147,63 @@ void SCurveEditorTreeFilterStatusBar::UpdateText()
 	TextBlock->SetText(NewText);
 }
 
+static double CurveEditorOpacityDurationSeconds = 2.0;
+
+void SCurveEditorTreeFilterStatusBar::Tick(const FGeometry& AllottedGeometry, const double InCurrentTime, const float InDeltaTime)
+{
+	if (GetVisibility() == EVisibility::Visible && OpacityThrobEndTime != FLT_MAX)
+	{
+		if (IsHovered())
+		{
+			// Make sure the status bar stays visible and keep it from fading until OnUnhovered
+			ShowStatusBar();
+			FadeOutStatusBar();
+		}
+		else
+		{
+			double CurrentTime = FPlatformTime::Seconds();
+
+			float Opacity = 0.f;
+			if (OpacityThrobEndTime > CurrentTime)
+			{
+				double Difference = OpacityThrobEndTime - CurrentTime;
+				Opacity = Difference / CurveEditorOpacityDurationSeconds;
+			}
+
+			if (Opacity > 0.f)
+			{
+				FLinearColor BorderColor = GetColorAndOpacity();
+				BorderColor.A = Opacity;
+				SetColorAndOpacity(BorderColor);
+			}
+			else
+			{
+				SetVisibility(EVisibility::Hidden);
+			}
+		}
+	}
+
+	SCompoundWidget::Tick(AllottedGeometry, InCurrentTime, InDeltaTime);
+}
+
+void SCurveEditorTreeFilterStatusBar::ShowStatusBar()
+{
+	OpacityThrobEndTime = FLT_MAX;
+	FLinearColor BorderColor = GetColorAndOpacity();
+	BorderColor.A = 1.f;
+	SetColorAndOpacity(BorderColor);
+	SetVisibility(EVisibility::Visible);
+}
+
+void SCurveEditorTreeFilterStatusBar::HideStatusBar()
+{
+	SetVisibility(EVisibility::Hidden);
+}
+
+void SCurveEditorTreeFilterStatusBar::FadeOutStatusBar()
+{
+	// Show and then start fading out the status bar
+	OpacityThrobEndTime = FPlatformTime::Seconds() + CurveEditorOpacityDurationSeconds;
+}
+
 #undef LOCTEXT_NAMESPACE
\ No newline at end of file
diff --git a/Engine/Source/Editor/CurveEditor/Public/Tree/SCurveEditorTreeFilterStatusBar.h b/Engine/Source/Editor/CurveEditor/Public/Tree/SCurveEditorTreeFilterStatusBar.h
--- a/Engine/Source/Editor/CurveEditor/Public/Tree/SCurveEditorTreeFilterStatusBar.h
+++ b/Engine/Source/Editor/CurveEditor/Public/Tree/SCurveEditorTreeFilterStatusBar.h
@@ -20,6 +20,12 @@ public:
 	/** Construct the status bar */
 	void Construct(const FArguments& InArgs, TSharedPtr<FCurveEditor> InCurveEditor);
 
+	virtual void Tick(const FGeometry& AllottedGeometry, const double InCurrentTime, const float InDeltaTime) override;
+
+	void ShowStatusBar();
+	void HideStatusBar();
+	void FadeOutStatusBar();
+
 private:
 
 	/** Request the visibility of the clear hyperlink widget based on whether there's a filter active or not */
@@ -35,4 +41,5 @@ private:
 
 	TWeakPtr<FCurveEditor> WeakCurveEditor;
 	TSharedPtr<STextBlock> TextBlock;
+	double OpacityThrobEndTime = 0;
 };
\ No newline at end of file
diff --git a/Engine/Source/Editor/EditorStyle/Private/StarshipStyle.cpp b/Engine/Source/Editor/EditorStyle/Private/StarshipStyle.cpp
--- a/Engine/Source/Editor/EditorStyle/Private/StarshipStyle.cpp
+++ b/Engine/Source/Editor/EditorStyle/Private/StarshipStyle.cpp
@@ -866,81 +866,24 @@ void FStarshipEditorStyle::FStyle::SetupGeneralStyles()
 	
 	// Scrub control buttons
 	{
-		Set( "Animation.Pause", FButtonStyle( Button )
-			.SetNormal(IMAGE_BRUSH( "/Sequencer/Transport_Bar/Pause_24x_OFF", Icon24x24 ))
-			.SetHovered(IMAGE_BRUSH( "/Sequencer/Transport_Bar/Pause_24x_OFF", Icon24x24 ))
-			.SetPressed(IMAGE_BRUSH( "/Sequencer/Transport_Bar/Pause_24x", Icon24x24 ))
-		);
-
-		Set( "Animation.Forward", FButtonStyle( Button )
-			.SetNormal(IMAGE_BRUSH( "/Sequencer/Transport_Bar/Play_24x_OFF", Icon24x24 ))
-			.SetHovered(IMAGE_BRUSH( "/Sequencer/Transport_Bar/Play_24x_OFF", Icon24x24 ))
-			.SetPressed(IMAGE_BRUSH( "/Sequencer/Transport_Bar/Play_24x", Icon24x24 ))
-		);
-
-		Set( "Animation.Forward_Step", FButtonStyle( Button )
-			.SetNormal(IMAGE_BRUSH( "/Sequencer/Transport_Bar/Step_Forward_24x_OFF", Icon24x24 ))
-			.SetHovered(IMAGE_BRUSH( "/Sequencer/Transport_Bar/Step_Forward_24x_OFF", Icon24x24 ))
-			.SetPressed(IMAGE_BRUSH( "/Sequencer/Transport_Bar/Step_Forward_24x", Icon24x24 ))
-			.SetDisabled(IMAGE_BRUSH( "/Sequencer/Transport_Bar/Step_Forward_24x_OFF", Icon24x24, FLinearColor(0.4f, 0.4f, 0.4f, 0.4f) ))
-		);
-
-		Set( "Animation.Forward_End", FButtonStyle( Button )
-			.SetNormal(IMAGE_BRUSH( "/Sequencer/Transport_Bar/Go_To_End_24x_OFF", Icon24x24 ))
-			.SetHovered(IMAGE_BRUSH( "/Sequencer/Transport_Bar/Go_To_End_24x_OFF", Icon24x24 ))
-			.SetPressed(IMAGE_BRUSH( "/Sequencer/Transport_Bar/Go_To_End_24x", Icon24x24 ))
-			.SetDisabled(IMAGE_BRUSH( "/Sequencer/Transport_Bar/Go_To_End_24x_OFF", Icon24x24, FLinearColor(0.4f, 0.4f, 0.4f, 0.4f) ))
-		);
-
-		Set( "Animation.Backward", FButtonStyle( Button )
-			.SetNormal(IMAGE_BRUSH( "/Sequencer/Transport_Bar/Backwards_24x_OFF", Icon24x24 ))
-			.SetHovered(IMAGE_BRUSH( "/Sequencer/Transport_Bar/Backwards_24x_OFF", Icon24x24 ))
-			.SetPressed(IMAGE_BRUSH( "/Sequencer/Transport_Bar/Backwards_24x", Icon24x24 ))
-		);
-
-		Set( "Animation.Backward_Step", FButtonStyle( Button )
-			.SetNormal(IMAGE_BRUSH( "/Sequencer/Transport_Bar/Step_Backwards_24x_OFF", Icon24x24 ))
-			.SetHovered(IMAGE_BRUSH( "/Sequencer/Transport_Bar/Step_Backwards_24x_OFF", Icon24x24 ))
-			.SetPressed(IMAGE_BRUSH( "/Sequencer/Transport_Bar/Step_Backwards_24x", Icon24x24 ))
-			.SetDisabled(IMAGE_BRUSH( "/Sequencer/Transport_Bar/Step_Backwards_24x_OFF", Icon24x24, FLinearColor(0.4f, 0.4f, 0.4f, 0.4f) ))
-		);
-
-		Set( "Animation.Backward_End", FButtonStyle( Button )
-			.SetNormal(IMAGE_BRUSH( "/Sequencer/Transport_Bar/Go_To_Front_24x_OFF", Icon24x24 ))
-			.SetHovered(IMAGE_BRUSH( "/Sequencer/Transport_Bar/Go_To_Front_24x_OFF", Icon24x24 ))
-			.SetPressed(IMAGE_BRUSH( "/Sequencer/Transport_Bar/Go_To_Front_24x", Icon24x24 ))
-			.SetDisabled(IMAGE_BRUSH( "/Sequencer/Transport_Bar/Step_Backwards_24x_OFF", Icon24x24, FLinearColor(0.4f, 0.4f, 0.4f, 0.4f) ))
-		);
-
-		Set( "Animation.Loop.Enabled", FButtonStyle( Button )
-			.SetNormal(IMAGE_BRUSH( "/Sequencer/Transport_Bar/Loop_24x_OFF", Icon24x24 ))
-			.SetHovered(IMAGE_BRUSH( "/Sequencer/Transport_Bar/Loop_24x_OFF", Icon24x24 ))
-			.SetPressed(IMAGE_BRUSH( "/Sequencer/Transport_Bar/Loop_24x", Icon24x24 ))
-		);
-
-		Set( "Animation.Loop.Disabled", FButtonStyle( Button )
-			.SetNormal(IMAGE_BRUSH( "/Sequencer/Transport_Bar/Loop_Toggle_24x_OFF", Icon24x24 ))
-			.SetHovered(IMAGE_BRUSH( "/Sequencer/Transport_Bar/Loop_Toggle_24x_OFF", Icon24x24 ))
-			.SetPressed(IMAGE_BRUSH( "/Sequencer/Transport_Bar/Loop_Toggle_24x", Icon24x24 ))
-		);
-
-		Set( "Animation.Loop.SelectionRange", FButtonStyle( Button )
-			.SetNormal(IMAGE_BRUSH( "/Sequencer/Transport_Bar/Loop_SelectionRange_24x_OFF", Icon24x24 ))
-			.SetHovered(IMAGE_BRUSH( "/Sequencer/Transport_Bar/Loop_SelectionRange_24x_OFF", Icon24x24 ))
-			.SetPressed(IMAGE_BRUSH( "/Sequencer/Transport_Bar/Loop_SelectionRange_24x", Icon24x24 ))
-		);
-
-		Set( "Animation.Record", FButtonStyle( Button )
-			.SetNormal(IMAGE_BRUSH( "/Sequencer/Transport_Bar/Record_24x_OFF", Icon24x24 ))
-			.SetHovered(IMAGE_BRUSH( "/Sequencer/Transport_Bar/Record_24x_OFF", Icon24x24 ))
-			.SetPressed(IMAGE_BRUSH( "/Sequencer/Transport_Bar/Record_24x", Icon24x24 ))
-		);
-
-		Set( "Animation.Recording", FButtonStyle( Button )
-			.SetNormal(IMAGE_BRUSH( "/Sequencer/Transport_Bar/Recording_24x_OFF", Icon24x24 ))
-			.SetHovered(IMAGE_BRUSH( "/Sequencer/Transport_Bar/Recording_24x_OFF", Icon24x24 ))
-			.SetPressed(IMAGE_BRUSH( "/Sequencer/Transport_Bar/Recording_24x", Icon24x24 ))
-		);
+		Set("Animation.PlayControlsButton", FButtonStyle(Button)
+			.SetNormal(FSlateNoResource())
+			.SetDisabled(FSlateNoResource())
+			.SetNormalPadding(FMargin(2.f, 2.f, 2.f, 2.f))
+			.SetPressedPadding(FMargin(2.f, 2.f, 2.f, 2.f)));
+
+		Set("Animation.Pause", new IMAGE_BRUSH_SVG("Sequencer/PlaybackControls/PlayControlsPause", Icon20x20));
+		Set("Animation.Forward", new IMAGE_BRUSH_SVG("Sequencer/PlaybackControls/PlayControlsPlayForward", Icon20x20));
+		Set("Animation.Forward_Step", new IMAGE_BRUSH_SVG("Sequencer/PlaybackControls/PlayControlsToNext", Icon20x20));
+		Set("Animation.Forward_End", new IMAGE_BRUSH_SVG("Sequencer/PlaybackControls/PlayControlsToEnd", Icon20x20));
+		Set("Animation.Backward", new IMAGE_BRUSH_SVG("Sequencer/PlaybackControls/PlayControlsPlayReverse", Icon20x20));
+		Set("Animation.Stop", new IMAGE_BRUSH_SVG("Sequencer/PlaybackControls/PlayControlsStop", Icon20x20));
+		Set("Animation.Backward_Step", new IMAGE_BRUSH_SVG("Sequencer/PlaybackControls/PlayControlsToPrevious", Icon20x20));
+		Set("Animation.Backward_End", new IMAGE_BRUSH_SVG("Sequencer/PlaybackControls/PlayControlsToFront", Icon20x20));
+		Set("Animation.Loop.Enabled", new IMAGE_BRUSH_SVG("Sequencer/PlaybackControls/PlayControlsLooping", Icon20x20));
+		Set("Animation.Loop.Disabled", new IMAGE_BRUSH_SVG("Sequencer/PlaybackControls/PlayControlsNoLooping", Icon20x20));
+		Set("Animation.Loop.SelectionRange", new IMAGE_BRUSH_SVG("Sequencer/PlaybackControls/PlayControlsLoopingSelectionRange", Icon20x20));
+		Set("Animation.Record", new IMAGE_BRUSH_SVG("Sequencer/PlaybackControls/PlayControlsRecord", Icon20x20));
 	}
 
 	// Message Log
@@ -2122,22 +2065,10 @@ void FStarshipEditorStyle::FStyle::SetupSequencerStyles()
 		Set("Sequencer.Timeline.EaseInOut", new IMAGE_BRUSH("Sequencer/EaseInOut", FVector2D(128, 128)));
 		Set("Sequencer.InterpLine", new BOX_BRUSH("Sequencer/InterpLine", FMargin(5.f / 7.f, 0.f, 0.f, 0.f)));
 
-		Set("Sequencer.Transport.JumpToPreviousKey", FButtonStyle()
-			.SetNormal(IMAGE_BRUSH("/Sequencer/Transport_Bar/Previous_Frame_OFF", Icon24x24))
-			.SetPressed(IMAGE_BRUSH("/Sequencer/Transport_Bar/Previous_Frame", Icon24x24))
-			.SetHovered(IMAGE_BRUSH("/Sequencer/Transport_Bar/Previous_Frame_OFF", Icon24x24)));
-		Set("Sequencer.Transport.JumpToNextKey", FButtonStyle()
-			.SetNormal(IMAGE_BRUSH("/Sequencer/Transport_Bar/Next_Frame_24x_OFF", Icon24x24))
-			.SetPressed(IMAGE_BRUSH("/Sequencer/Transport_Bar/Next_Frame_24x", Icon24x24))
-			.SetHovered(IMAGE_BRUSH("/Sequencer/Transport_Bar/Next_Frame_24x_OFF", Icon24x24)));
-		Set("Sequencer.Transport.SetPlayStart", FButtonStyle()
-			.SetNormal(IMAGE_BRUSH("/Sequencer/Transport_Bar/Bracket_In_16x24_OFF", FVector2D(16, 24)))
-			.SetPressed(IMAGE_BRUSH("/Sequencer/Transport_Bar/Bracket_In_16x24", FVector2D(16, 24)))
-			.SetHovered(IMAGE_BRUSH("/Sequencer/Transport_Bar/Bracket_In_16x24_OFF", FVector2D(16, 24))));
-		Set("Sequencer.Transport.SetPlayEnd", FButtonStyle()
-			.SetNormal(IMAGE_BRUSH("/Sequencer/Transport_Bar/Bracket_Out_16x24_OFF", FVector2D(16, 24)))
-			.SetPressed(IMAGE_BRUSH("/Sequencer/Transport_Bar/Bracket_Out_16x24", FVector2D(16, 24)))
-			.SetHovered(IMAGE_BRUSH("/Sequencer/Transport_Bar/Bracket_Out_16x24_OFF", FVector2D(16, 24))));
+		Set("Sequencer.Transport.JumpToPreviousKey", new IMAGE_BRUSH_SVG("Sequencer/PlaybackControls/PlayControlsJumpToPreviousKey", Icon20x20));
+		Set("Sequencer.Transport.JumpToNextKey", new IMAGE_BRUSH_SVG("Sequencer/PlaybackControls/PlayControlsJumpToNextKey", Icon20x20));
+		Set("Sequencer.Transport.SetPlayStart", new IMAGE_BRUSH_SVG("Sequencer/PlaybackControls/PlayControlsSetPlaybackStart", Icon20x20));
+		Set("Sequencer.Transport.SetPlayEnd", new IMAGE_BRUSH_SVG("Sequencer/PlaybackControls/PlayControlsSetPlaybackEnd", Icon20x20));
 
 		Set("Sequencer.Transport.CloseButton", FButtonStyle()
 			.SetNormal(IMAGE_BRUSH("/Docking/CloseApp_Normal", Icon16x16))
@@ -2301,7 +2232,7 @@ void FStarshipEditorStyle::FStyle::SetupSequencerStyles()
 			.SetHoveredBackgroundBrush(FSlateNoResource())
 			.SetInactiveFillBrush(FSlateNoResource())
 			.SetActiveFillBrush(FSlateNoResource())
-			.SetForegroundColor(SelectionColor_Pressed)
+			.SetForegroundColor(FSlateColor::UseForeground())
 			.SetArrowsImage(FSlateNoResource())
 		);
 
diff --git a/Engine/Source/Editor/EditorWidgets/Private/STransportControl.cpp b/Engine/Source/Editor/EditorWidgets/Private/STransportControl.cpp
--- a/Engine/Source/Editor/EditorWidgets/Private/STransportControl.cpp
+++ b/Engine/Source/Editor/EditorWidgets/Private/STransportControl.cpp
@@ -33,42 +33,54 @@ TSharedPtr<SWidget> STransportControl::MakeTransportControlWidget(ETransportCont
 	{
 	case ETransportControlWidgetType::BackwardEnd:
 		return SNew(SButton)
-			. ButtonStyle(&FAppStyle::Get().GetWidgetStyle<FButtonStyle>("Animation.Backward_End"))
+			. ButtonStyle(FAppStyle::Get(), "Animation.PlayControlsButton")
 			. OnClicked(TransportControlArgs.OnBackwardEnd)
 			. Visibility(TransportControlArgs.OnBackwardEnd.IsBound() ? EVisibility::Visible : EVisibility::Collapsed)
 			. ToolTipText( LOCTEXT("ToFront", "To Front") )
-			. ContentPadding(2.0f)
-			. IsFocusable(bAreButtonsFocusable);
+			. ContentPadding(0.0f)
+			. IsFocusable(bAreButtonsFocusable)
+			[
+				SNew(SImage)
+				.ColorAndOpacity(FSlateColor::UseSubduedForeground())
+				.Image(FAppStyle::Get().GetBrush("Animation.Backward_End"))
+			];
 	case ETransportControlWidgetType::BackwardStep:
 		return SNew(SButton)
-			. ButtonStyle(&FAppStyle::Get().GetWidgetStyle<FButtonStyle>("Animation.Backward_Step"))
+			. ButtonStyle(FAppStyle::Get(), "Animation.PlayControlsButton")
 			. OnClicked(TransportControlArgs.OnBackwardStep)
 			. Visibility(TransportControlArgs.OnBackwardStep.IsBound() ? EVisibility::Visible : EVisibility::Collapsed)
 			. ToolTipText( LOCTEXT("ToPrevious", "To Previous") )
-			. ContentPadding(2.0f)
-			. IsFocusable(bAreButtonsFocusable);
+			. ContentPadding(0.0f)
+			. IsFocusable(bAreButtonsFocusable)
+			[
+				SNew(SImage)
+				.ColorAndOpacity(FSlateColor::UseSubduedForeground())
+				.Image(FAppStyle::Get().GetBrush("Animation.Backward_Step"))
+			];
 	case ETransportControlWidgetType::BackwardPlay:
 		return SAssignNew(BackwardPlayButton, SButton)
 			. OnClicked(TransportControlArgs.OnBackwardPlay)
 			. Visibility(TransportControlArgs.OnBackwardPlay.IsBound() ? EVisibility::Visible : EVisibility::Collapsed)
 			. ToolTipText( LOCTEXT("Reverse", "Reverse") )
-			. ButtonStyle( FAppStyle::Get(), "NoBorder" )
-			. ContentPadding(2.0f)
+			. ButtonStyle( FAppStyle::Get(), "Animation.PlayControlsButton" )
+			. ContentPadding(0.0f)
 			. IsFocusable(bAreButtonsFocusable)
 			[
 				SNew(SImage)
-				. Image( this, &STransportControl::GetBackwardStatusIcon )
+				.ColorAndOpacity(FSlateColor::UseSubduedForeground())
+				.Image( this, &STransportControl::GetBackwardStatusIcon )
 			];
 	case ETransportControlWidgetType::Record:
 		return SAssignNew(RecordButton, SButton)
-			. ButtonStyle(FAppStyle::Get(), "NoBorder")
+			. ButtonStyle(FAppStyle::Get(), "Animation.PlayControlsButton")
 			. OnClicked(TransportControlArgs.OnRecord)
 			. Visibility(TransportControlArgs.OnRecord.IsBound() ? EVisibility::Visible : EVisibility::Collapsed)
 			. ToolTipText( this, &STransportControl::GetRecordStatusTooltip )
-			. ContentPadding(2.0f)
+			. ContentPadding(0.0f)
 			. IsFocusable(bAreButtonsFocusable)
 			[
 				SNew(SImage)
+				.ColorAndOpacity(FSlateColor::UseSubduedForeground())
 				.Image(this, &STransportControl::GetRecordStatusIcon )
 			];
 	case ETransportControlWidgetType::ForwardPlay:
@@ -76,40 +88,52 @@ TSharedPtr<SWidget> STransportControl::MakeTransportControlWidget(ETransportCont
 			. OnClicked(TransportControlArgs.OnForwardPlay)
 			. Visibility(TransportControlArgs.OnForwardPlay.IsBound() ? EVisibility::Visible : EVisibility::Collapsed)
 			. ToolTipText( this, &STransportControl::GetForwardStatusTooltip )
-			. ButtonStyle( FAppStyle::Get(), "NoBorder" )
-			. ContentPadding(2.0f)
+			. ButtonStyle( FAppStyle::Get(), "Animation.PlayControlsButton" )
+			. ContentPadding(0.0f)
 			. IsFocusable(bAreButtonsFocusable)
 			[
 				SNew(SImage)
-				. Image( this, &STransportControl::GetForwardStatusIcon )
+				.ColorAndOpacity(FSlateColor::UseSubduedForeground())
+				.Image( this, &STransportControl::GetForwardStatusIcon )
 			];
 	case ETransportControlWidgetType::ForwardStep:
 		return SNew(SButton)
-			. ButtonStyle(&FAppStyle::Get().GetWidgetStyle<FButtonStyle>("Animation.Forward_Step"))
+			. ButtonStyle(FAppStyle::Get(), "Animation.PlayControlsButton")
 			. OnClicked(TransportControlArgs.OnForwardStep)
 			. Visibility(TransportControlArgs.OnForwardStep.IsBound() ? EVisibility::Visible : EVisibility::Collapsed)
 			. ToolTipText( LOCTEXT("ToNext", "To Next") )
-			. ContentPadding(2.0f)
-			. IsFocusable(bAreButtonsFocusable);
+			. ContentPadding(0.0f)
+			. IsFocusable(bAreButtonsFocusable)
+			[
+				SNew(SImage)
+				.ColorAndOpacity(FSlateColor::UseSubduedForeground())
+				.Image(FAppStyle::Get().GetBrush("Animation.Forward_Step"))
+			];
 	case ETransportControlWidgetType::ForwardEnd:
 		return SNew(SButton)
-			. ButtonStyle(&FAppStyle::Get().GetWidgetStyle<FButtonStyle>("Animation.Forward_End"))
+			. ButtonStyle(FAppStyle::Get(), "Animation.PlayControlsButton")
 			. OnClicked(TransportControlArgs.OnForwardEnd)
 			. Visibility(TransportControlArgs.OnForwardEnd.IsBound() ? EVisibility::Visible : EVisibility::Collapsed)
 			. ToolTipText( LOCTEXT("ToEnd", "To End") )
-			. ContentPadding(2.0f)
-			. IsFocusable(bAreButtonsFocusable);
+			. ContentPadding(0.0f)
+			. IsFocusable(bAreButtonsFocusable)
+			[
+				SNew(SImage)
+				.ColorAndOpacity(FSlateColor::UseSubduedForeground())
+				.Image(FAppStyle::Get().GetBrush("Animation.Forward_End"))
+			];
 	case ETransportControlWidgetType::Loop:
 		return SAssignNew(LoopButton, SButton)
 			. OnClicked(TransportControlArgs.OnToggleLooping)
 			. Visibility(TransportControlArgs.OnGetLooping.IsBound() ? EVisibility::Visible : EVisibility::Collapsed)
 			. ToolTipText( this, &STransportControl::GetLoopStatusTooltip )
-			. ButtonStyle( FAppStyle::Get(), "NoBorder" )
-			. ContentPadding(2.0f)
+			. ButtonStyle( FAppStyle::Get(), "Animation.PlayControlsButton" )
+			. ContentPadding(0.0f)
 			. IsFocusable(bAreButtonsFocusable)
 			[
 				SNew(SImage)
-				. Image( this, &STransportControl::GetLoopStatusIcon )
+				.ColorAndOpacity(FSlateColor::UseSubduedForeground())
+				.Image( this, &STransportControl::GetLoopStatusIcon )
 			];
 	case ETransportControlWidgetType::Custom:
 		if(MakeCustomWidgetDelegate.IsBound())
@@ -208,14 +232,10 @@ const FSlateBrush* STransportControl::GetForwardStatusIcon() const
 
 	if ( PlaybackMode == EPlaybackMode::PlayingForward)
 	{
-		return ForwardPlayButton.IsValid() && ForwardPlayButton->IsPressed() ? 
-			&FAppStyle::Get().GetWidgetStyle<FButtonStyle>("Animation.Pause").Pressed : 
-			&FAppStyle::Get().GetWidgetStyle<FButtonStyle>("Animation.Pause").Normal;
+		return FAppStyle::Get().GetBrush("Animation.Pause");
 	}
 
-	return ForwardPlayButton.IsValid() && ForwardPlayButton->IsPressed() ? 
-		&FAppStyle::Get().GetWidgetStyle<FButtonStyle>("Animation.Forward").Pressed :
-		&FAppStyle::Get().GetWidgetStyle<FButtonStyle>("Animation.Forward").Normal;
+	return FAppStyle::Get().GetBrush("Animation.Forward");
 }
 
 FText STransportControl::GetForwardStatusTooltip() const
@@ -239,14 +259,10 @@ const FSlateBrush* STransportControl::GetRecordStatusIcon() const
 
 	if (bIsRecording)
 	{
-		return RecordButton.IsValid() && RecordButton->IsPressed() ?
-			&FAppStyle::Get().GetWidgetStyle<FButtonStyle>("Animation.Recording").Pressed :
-			&FAppStyle::Get().GetWidgetStyle<FButtonStyle>("Animation.Recording").Normal;
+		return FAppStyle::Get().GetBrush("Animation.Recording");
 	}
 
-	return RecordButton.IsValid() && RecordButton->IsPressed() ?
-		&FAppStyle::Get().GetWidgetStyle<FButtonStyle>("Animation.Record").Pressed :
-		&FAppStyle::Get().GetWidgetStyle<FButtonStyle>("Animation.Record").Normal;
+	return FAppStyle::Get().GetBrush("Animation.Record");
 }
 
 FText STransportControl::GetRecordStatusTooltip() const
@@ -265,14 +281,10 @@ const FSlateBrush* STransportControl::GetBackwardStatusIcon() const
 	if (TransportControlArgs.OnGetPlaybackMode.IsBound() &&
 		TransportControlArgs.OnGetPlaybackMode.Execute() == EPlaybackMode::PlayingReverse)
 	{
-		return BackwardPlayButton.IsValid() && BackwardPlayButton->IsPressed() ? 
-			&FAppStyle::Get().GetWidgetStyle<FButtonStyle>("Animation.Pause").Pressed : 
-			&FAppStyle::Get().GetWidgetStyle<FButtonStyle>("Animation.Pause").Normal;
+		return FAppStyle::Get().GetBrush("Animation.Pause");
 	}
 
-	return BackwardPlayButton.IsValid() && BackwardPlayButton->IsPressed() ? 
-		&FAppStyle::Get().GetWidgetStyle<FButtonStyle>("Animation.Backward").Pressed : 
-		&FAppStyle::Get().GetWidgetStyle<FButtonStyle>("Animation.Backward").Normal;
+	return FAppStyle::Get().GetBrush("Animation.Backward");
 }
 
 const FSlateBrush* STransportControl::GetLoopStatusIcon() const
@@ -280,14 +292,10 @@ const FSlateBrush* STransportControl::GetLoopStatusIcon() const
 	if (TransportControlArgs.OnGetLooping.IsBound() &&
 		TransportControlArgs.OnGetLooping.Execute())
 	{
-		return LoopButton.IsValid() && LoopButton->IsPressed() ? 
-			&FAppStyle::Get().GetWidgetStyle<FButtonStyle>("Animation.Loop.Enabled").Pressed : 
-			&FAppStyle::Get().GetWidgetStyle<FButtonStyle>("Animation.Loop.Enabled").Normal;
+		return FAppStyle::Get().GetBrush("Animation.Loop.Enabled");
 	}
 
-	return LoopButton.IsValid() && LoopButton->IsPressed() ? 
-		&FAppStyle::Get().GetWidgetStyle<FButtonStyle>("Animation.Loop.Disabled").Pressed : 
-		&FAppStyle::Get().GetWidgetStyle<FButtonStyle>("Animation.Loop.Disabled").Normal;
+	return FAppStyle::Get().GetBrush("Animation.Loop.Disabled");
 }
 
 FText STransportControl::GetLoopStatusTooltip() const
diff --git a/Engine/Source/Editor/Sequencer/Private/MVVM/CurveEditorExtension.cpp b/Engine/Source/Editor/Sequencer/Private/MVVM/CurveEditorExtension.cpp
--- a/Engine/Source/Editor/Sequencer/Private/MVVM/CurveEditorExtension.cpp
+++ b/Engine/Source/Editor/Sequencer/Private/MVVM/CurveEditorExtension.cpp
@@ -22,6 +22,7 @@
 #include "Widgets/CurveEditor/SSequencerCurveEditor.h"
 #include "Widgets/CurveEditor/SequencerCurveEditorTimeSliderController.h"
 #include "Widgets/Docking/SDockTab.h"
+#include "Widgets/Input/SButton.h"
 #include "Widgets/Layout/SScrollBorder.h"
 
 #define LOCTEXT_NAMESPACE "SequencerCurveEditorExtension"
@@ -223,23 +224,91 @@ void FCurveEditorExtension::CreateCurveEditor(const FTimeSliderArgs& TimeSliderA
 
 			+ SVerticalBox::Slot()
 			[
-				SNew(SScrollBorder, CurveEditorTreeView.ToSharedRef())
+				SNew(SOverlay)
+
+				+ SOverlay::Slot()
 				[
-					CurveEditorTreeView.ToSharedRef()
+					SNew(SScrollBorder, CurveEditorTreeView.ToSharedRef())
+					[
+						CurveEditorTreeView.ToSharedRef()
+					]
 				]
-			]
 
-			+ SVerticalBox::Slot()
-			.AutoHeight()
-			[
-				SNew(SCurveEditorTreeFilterStatusBar, CurveEditorModel)
+				+ SOverlay::Slot()
+				.VAlign(VAlign_Bottom)
+				[
+					SNew(SHorizontalBox)
+
+					+ SHorizontalBox::Slot()
+					.AutoWidth()
+					[
+						SAssignNew(CurveEditorTreeFilterStatusBar, SCurveEditorTreeFilterStatusBar, CurveEditorModel)
+						.Visibility(EVisibility::Hidden) // Initially hidden, visible on hover of the info button
+					]
+				]
 			]
 
 			+ SVerticalBox::Slot()
 			.AutoHeight()
-			.HAlign(HAlign_Center)
 			[
-				Sequencer->MakeTransportControls(true)
+				SNew(SHorizontalBox)
+
+				+ SHorizontalBox::Slot()
+				.AutoWidth()
+				.VAlign(VAlign_Center)
+				.HAlign(HAlign_Left)
+				[
+					SNew(SButton)
+					.VAlign(EVerticalAlignment::VAlign_Center)
+					.ButtonStyle(FAppStyle::Get(), "SimpleButton")
+					.ToolTipText_Lambda([this] { return LOCTEXT("ShowStatus", "Show Status"); })
+					.ContentPadding(FMargin(1, 0))
+					.OnHovered_Lambda([this] { CurveEditorTreeFilterStatusBar->ShowStatusBar(); })
+					.OnUnhovered_Lambda([this] { CurveEditorTreeFilterStatusBar->FadeOutStatusBar(); })
+					.OnClicked_Lambda([this] { CurveEditorTreeFilterStatusBar->HideStatusBar(); return FReply::Handled(); })
+					[
+						SNew(SImage)
+						.ColorAndOpacity(FSlateColor::UseForeground())
+						.Image(FAppStyle::Get().GetBrush("Icons.Info.Small"))
+					]
+				]
+
+				+ SHorizontalBox::Slot()
+				[
+					SNew(SBorder)
+					.BorderImage(FAppStyle::GetBrush("ToolPanel.GroupBorder"))
+					.HAlign(HAlign_Center)
+					[
+						Sequencer->MakeTransportControls(true)
+					]
+				]
+
+				+ SHorizontalBox::Slot()
+				.AutoWidth()
+				.VAlign(VAlign_Center)
+				.HAlign(HAlign_Right)
+				[
+					SNew(SButton)
+					.VAlign(EVerticalAlignment::VAlign_Center)
+					.ButtonStyle(FAppStyle::Get(), "NoBorder")
+					.ContentPadding(FMargin(1, 0))
+					[
+						SNew(SHorizontalBox)
+
+						+ SHorizontalBox::Slot()
+						.AutoWidth()
+						.VAlign(VAlign_Center)
+						.HAlign(HAlign_Right)
+						.Padding(FMargin(3.f, 0.f, 0.f, 0.f))
+						[
+							SNew(SBorder)
+							.BorderImage(nullptr)
+							[
+								Sequencer->MakePlayTimeDisplay(TimeSliderArgs.NumericTypeInterface.ToSharedRef())
+							]
+						]
+					]
+				]
 			]
 		];
 
diff --git a/Engine/Source/Editor/Sequencer/Private/SSequencer.cpp b/Engine/Source/Editor/Sequencer/Private/SSequencer.cpp
--- a/Engine/Source/Editor/Sequencer/Private/SSequencer.cpp
+++ b/Engine/Source/Editor/Sequencer/Private/SSequencer.cpp
@@ -97,7 +97,6 @@
 #include "Tree/SCurveEditorTree.h"
 #include "Tree/CurveEditorTreeFilter.h"
 #include "Tree/SCurveEditorTreeTextFilter.h"
-#include "Tree/SCurveEditorTreeFilterStatusBar.h"
 #include "SequencerSelectionCurveFilter.h"
 #include "SCurveKeyDetailPanel.h"
 #include "MovieSceneTimeHelpers.h"
@@ -319,14 +318,9 @@ void SSequencer::Construct(const FArguments& InArgs, TSharedRef<FSequencer> InSe
 	SequencerViewModel->GetTrackArea()->CastThisChecked<FSequencerTrackAreaViewModel>()->InitializeDefaultEditTools(*TrackArea);
 	SequencerViewModel->GetPinnedTrackArea()->CastThisChecked<FSequencerTrackAreaViewModel>()->InitializeDefaultEditTools(*PinnedTrackArea);
 
-	TAttribute<FAnimatedRange> ViewRangeAttribute = InArgs._ViewRange;
-
-	if (InSequencer->GetHostCapabilities().bSupportsCurveEditor)
-	{
-		FCurveEditorExtension* CurveEditorExtension = SequencerViewModel->CastDynamicChecked<FCurveEditorExtension>();
-		CurveEditorExtension->CreateCurveEditor(TimeSliderArgs);
-	}
+	PlayTimeDisplay = StaticCastSharedRef<STemporarilyFocusedSpinBox<double>>(SequencerPtr.Pin()->MakePlayTimeDisplay(GetNumericTypeInterface()));
 
+	TAttribute<FAnimatedRange> ViewRangeAttribute = InArgs._ViewRange;
 
 	const int32 Column0 = 0, Column1 = 1;
 	const int32 Row0 = 0, Row1 = 1, Row2 = 2, Row3 = 3, Row4 = 4;
@@ -559,82 +553,6 @@ void SSequencer::Construct(const FArguments& InArgs, TSharedRef<FSequencer> InSe
 							[
 								MakeFilterButton()
 							]
-
-							+ SHorizontalBox::Slot()
-							.AutoWidth()
-							.VAlign(VAlign_Center)
-							.HAlign(HAlign_Right)
-							.Padding(FMargin(CommonPadding + 2.0, 0.f, 0.f, 0.f))
-							[
-								SNew(SBorder)
-								.BorderImage(nullptr)
-								[
-									// Current Play Time 
-									SAssignNew(PlayTimeDisplay, STemporarilyFocusedSpinBox<double>)
-									.Style(&FAppStyle::GetWidgetStyle<FSpinBoxStyle>("Sequencer.PlayTimeSpinBox"))
-									.Value_Lambda([this]() -> double {
-										return SequencerPtr.Pin()->GetLocalTime().Time.GetFrame().Value;
-									})
-									.OnValueChanged(this, &SSequencer::SetPlayTimeClampedByWorkingRange)
-									.OnValueCommitted_Lambda([this](double InFrame, ETextCommit::Type) {
-										SetPlayTime(InFrame);
-
-										// Refocus on the previously focused widget so that user can continue on after setting a time
-										PlayTimeDisplay->Refocus();
-									})
-									.MinValue(TOptional<double>())
-									.MaxValue(TOptional<double>())
-									.TypeInterface(NumericTypeInterface)
-									.Delta(this, &SSequencer::GetSpinboxDelta)
-									.LinearDeltaSensitivity(25)
-									.MinDesiredWidth(this, &SSequencer::GetPlayTimeMinDesiredWidth)
-								]
-							]
-							+ SHorizontalBox::Slot()
-							.AutoWidth()
-							.VAlign(VAlign_Center)
-							.HAlign(HAlign_Right)
-							.Padding(FMargin(CommonPadding + 2.0, 0.f, 0.f, 0.f))
-							[
-								SNew(SBorder)
-								.BorderImage(nullptr)
-								[
-									// Current loop index, if any
-									SAssignNew(LoopIndexDisplay, STextBlock)
-									.Text_Lambda([this]() -> FText {
-										uint32 LoopIndex = SequencerPtr.Pin()->GetLocalLoopIndex();
-										return (LoopIndex != FMovieSceneTimeWarping::InvalidWarpCount) ? FText::AsNumber(LoopIndex + 1) : FText();
-									})
-								]
-							]
-							+ SHorizontalBox::Slot()
-							.AutoWidth()
-							.VAlign(VAlign_Center)
-							.HAlign(HAlign_Right)
-							.Padding(FMargin(CommonPadding + 2.0, 0.f, 0.f, 0.f))
-							[
-								SNew(SBorder)
-								.BorderImage(nullptr)
-								[
-									// Frame count and duration
-									SNew(STextBlock)
-									.ColorAndOpacity(FAppStyle::GetSlateColor("SelectionColor_Pressed"))
-									.Text_Lambda([this]() -> FText {
-										FFrameRate TickResolution = SequencerPtr.Pin()->GetFocusedTickResolution();
-										FFrameRate DisplayRate = SequencerPtr.Pin()->GetFocusedDisplayRate();
-
-										TOptional<TRange<FFrameNumber>> SubSequenceRange = SequencerPtr.Pin()->GetSubSequenceRange();
-
-										FFrameNumber CurrentFrame = SequencerPtr.Pin()->GetLocalTime().Time.GetFrame();
-										TRange<FFrameNumber> CurrentRange = SubSequenceRange.IsSet() ? SubSequenceRange.GetValue() : SequencerPtr.Pin()->GetPlaybackRange();
-
-										FFrameNumber FrameCount = FFrameRate::TransformTime((CurrentFrame - CurrentRange.GetLowerBoundValue() + 1).Value, TickResolution, DisplayRate).CeilToFrame();
-										FFrameNumber FrameDuration = FFrameRate::TransformTime(CurrentRange.Size<FFrameNumber>().Value, TickResolution, DisplayRate).CeilToFrame();
-
-										return FText::FromString(FString::Printf(TEXT("%d of %d"), FrameCount.Value, FrameDuration.Value));
-									})
-								]
-							]
 						]
 					]
 
@@ -745,22 +663,99 @@ void SSequencer::Construct(const FArguments& InArgs, TSharedRef<FSequencer> InSe
 								+ SOverlay::Slot()
 								.VAlign(VAlign_Bottom)
 								[
-									SAssignNew(SequencerTreeFilterStatusBar, SSequencerTreeFilterStatusBar, InSequencer)	
-									.Visibility(this, &SSequencer::GetStatusBarVisibility)
+									SNew(SHorizontalBox)
+
+									+ SHorizontalBox::Slot()
+									.AutoWidth()
+									[
+										SAssignNew(SequencerTreeFilterStatusBar, SSequencerTreeFilterStatusBar, InSequencer)
+										.Visibility(EVisibility::Hidden) // Initially hidden, visible on hover of the info button
+									]
 								]
 							]
 						]
 					]
 
-					// playback buttons
+					// Info Button, Transport Controls and Current Frame
 					+ SGridPanel::Slot( Column0, Row4, SGridPanel::Layer(10) )
 					[
-						SNew(SBorder)
-						.BorderImage(FAppStyle::GetBrush("ToolPanel.GroupBorder"))
-						//.BorderBackgroundColor(FLinearColor(.50f, .50f, .50f, 1.0f))
-						.HAlign(HAlign_Center)
+						SNew(SHorizontalBox)
+						
+						+ SHorizontalBox::Slot()
+						.AutoWidth()
+						.VAlign(VAlign_Center)
+						.HAlign(HAlign_Left)
+						[
+							SNew(SButton)
+							.VAlign(EVerticalAlignment::VAlign_Center)
+							.ButtonStyle(FAppStyle::Get(), "SimpleButton")
+							.ToolTipText_Lambda([this] { return LOCTEXT("ShowStatus", "Show Status"); })
+							.ContentPadding(FMargin(1, 0))
+							.Visibility(this, &SSequencer::GetInfoButtonVisibility)
+							.OnHovered_Lambda([this] { SequencerTreeFilterStatusBar->ShowStatusBar(); })
+							.OnUnhovered_Lambda([this] { SequencerTreeFilterStatusBar->FadeOutStatusBar(); })
+							.OnClicked_Lambda([this] { SequencerTreeFilterStatusBar->HideStatusBar(); return FReply::Handled(); })
+							[
+								SNew(SImage)
+								.ColorAndOpacity(FSlateColor::UseForeground())
+								.Image(FAppStyle::Get().GetBrush("Icons.Info.Small"))
+							]
+						]
+
+						+ SHorizontalBox::Slot()
 						[
-							SequencerPtr.Pin()->MakeTransportControls(true)
+							SNew(SBorder)
+							.BorderImage(FAppStyle::GetBrush("ToolPanel.GroupBorder"))
+							.HAlign(HAlign_Center)
+							[
+								SequencerPtr.Pin()->MakeTransportControls(true)
+							]
+						]
+
+						+ SHorizontalBox::Slot()
+						.AutoWidth()
+						.VAlign(VAlign_Center)
+						.HAlign(HAlign_Right)
+						[
+							SNew(SButton)
+							.VAlign(EVerticalAlignment::VAlign_Center)
+							.ButtonStyle(FAppStyle::Get(), "NoBorder")
+							.ContentPadding(FMargin(1, 0))
+							[
+								SNew(SHorizontalBox)
+
+								+ SHorizontalBox::Slot()
+								.AutoWidth()
+								.VAlign(VAlign_Center)
+								.HAlign(HAlign_Right)
+								.Padding(FMargin(CommonPadding, 0.f, 0.f, 0.f))
+								[
+									SNew(SBorder)
+									.BorderImage(nullptr)
+									[
+										PlayTimeDisplay.ToSharedRef()
+									]
+								]
+
+								+ SHorizontalBox::Slot()
+								.AutoWidth()
+								.VAlign(VAlign_Center)
+								.HAlign(HAlign_Right)
+								.Padding(FMargin(CommonPadding, 0.f, 0.f, 0.f))
+								[
+									SNew(SBorder)
+									.BorderImage(nullptr)
+									[
+										// Current loop index, if any
+										SAssignNew(LoopIndexDisplay, STextBlock)
+										.Font(FCoreStyle::GetDefaultFontStyle("Regular", 8))
+										.Text_Lambda([this]() -> FText {
+											uint32 LoopIndex = SequencerPtr.Pin()->GetLocalLoopIndex();
+											return (LoopIndex != FMovieSceneTimeWarping::InvalidWarpCount) ? FText::AsNumber(LoopIndex + 1) : FText();
+										})
+									]
+								]
+							]
 						]
 					]
 
@@ -895,6 +890,12 @@ void SSequencer::Construct(const FArguments& InArgs, TSharedRef<FSequencer> InSe
 		]
 	];
 
+	if (InSequencer->GetHostCapabilities().bSupportsCurveEditor)
+	{
+		FCurveEditorExtension* CurveEditorExtension = SequencerViewModel->CastDynamicChecked<FCurveEditorExtension>();
+		CurveEditorExtension->CreateCurveEditor(TimeSliderArgs);
+	}
+
 	ApplySequencerCustomization(RootCustomization);
 
 	InSequencer->GetViewModel()->GetSelection()->KeySelection.OnChanged.AddSP(this, &SSequencer::HandleKeySelectionChanged);
@@ -938,7 +939,7 @@ void SSequencer::BindCommands(TSharedRef<FUICommandList> SequencerCommandBinding
 
 	SequencerCommandBindings->MapAction(
 		FSequencerCommands::Get().ToggleShowGotoBox,
-		FExecuteAction::CreateLambda([this] { PlayTimeDisplay->Setup();  FSlateApplication::Get().SetKeyboardFocus(PlayTimeDisplay, EFocusCause::SetDirectly); })
+		FExecuteAction::CreateLambda([this] { PlayTimeDisplay->Setup(); FSlateApplication::Get().SetKeyboardFocus(PlayTimeDisplay, EFocusCause::SetDirectly); })
 	);
 
 	SequencerCommandBindings->MapAction(
@@ -1223,7 +1224,7 @@ TSharedRef<SWidget> SSequencer::MakeAddButton()
 	return SNew(SPositiveActionButton)
 	.OnGetMenuContent(this, &SSequencer::MakeAddMenu)
 	.Icon(FAppStyle::Get().GetBrush("Icons.Plus"))
-	.Text(LOCTEXT("Track", "Track"))
+	.Text(LOCTEXT("Add", "Add"))
 	.IsEnabled_Lambda([this]() { return !SequencerPtr.Pin()->IsReadOnly(); });
 }
 
@@ -2170,7 +2171,7 @@ TSharedRef<SWidget> SSequencer::MakeViewMenu()
 	MenuBuilder.AddMenuEntry( FSequencerCommands::Get().ToggleLayerBars );
 	MenuBuilder.AddMenuEntry( FSequencerCommands::Get().ToggleKeyBars );
 	MenuBuilder.AddMenuEntry( FSequencerCommands::Get().ToggleChannelColors );
-	MenuBuilder.AddMenuEntry( FSequencerCommands::Get().ToggleShowStatusBar );
+	MenuBuilder.AddMenuEntry( FSequencerCommands::Get().ToggleShowInfoButton );
 	MenuBuilder.AddMenuEntry( FSequencerCommands::Get().ToggleShowPreAndPostRoll );
 
 	// Menu entry for zero padding
@@ -3726,9 +3727,9 @@ EVisibility SSequencer::GetTimeRangeVisibility() const
 	return GetSequencerSettings()->GetShowRangeSlider() ? EVisibility::Visible : EVisibility::Hidden;
 }
 
-EVisibility SSequencer::GetStatusBarVisibility() const
+EVisibility SSequencer::GetInfoButtonVisibility() const
 {
-	return GetSequencerSettings()->GetShowStatusBar() ? EVisibility::SelfHitTestInvisible : EVisibility::Hidden;
+	return GetSequencerSettings()->GetShowInfoButton() ? EVisibility::SelfHitTestInvisible : EVisibility::Hidden;
 }
 
 EVisibility SSequencer::GetShowTickLines() const
@@ -4045,24 +4046,6 @@ double SSequencer::GetSpinboxDelta() const
 	return Sequencer->GetDisplayRateDeltaFrameCount();
 }
 
-float SSequencer::GetPlayTimeMinDesiredWidth() const
-{
-	TRange<double> ViewRange = SequencerPtr.Pin()->GetViewRange();
-
-	FString LowerBoundStr = NumericTypeInterface->ToString(ViewRange.GetLowerBoundValue());
-	FString UpperBoundStr = NumericTypeInterface->ToString(ViewRange.GetUpperBoundValue());
-
-	const FSlateFontInfo NormalFont = FCoreStyle::Get().GetFontStyle(TEXT("NormalFont"));
-	
-	const TSharedRef< FSlateFontMeasure > FontMeasureService = FSlateApplication::Get().GetRenderer()->GetFontMeasureService();
-
-	FVector2D LowerTextSize = FontMeasureService->Measure(LowerBoundStr, NormalFont);
-	FVector2D UpperTextSize = FontMeasureService->Measure(UpperBoundStr, NormalFont);
-
-	return FMath::Max(LowerTextSize.X, UpperTextSize.X);
-}
-
-
 bool SSequencer::GetIsSequenceReadOnly() const
 {
 	TSharedPtr<FSequencer> Sequencer = SequencerPtr.Pin();
@@ -4156,6 +4139,12 @@ void SSequencer::SetPlayTime(double Frame)
 		Sequencer->SetViewRange(NewViewRange);
 		
 		Sequencer->SetLocalTime(NewFrame);
+
+		// Refocus on the previously focused widget so that user can continue on after setting a time
+		if (PlayTimeDisplay)
+		{
+			PlayTimeDisplay->Refocus();
+		}
 	}
 }
 
diff --git a/Engine/Source/Editor/Sequencer/Private/SSequencer.h b/Engine/Source/Editor/Sequencer/Private/SSequencer.h
--- a/Engine/Source/Editor/Sequencer/Private/SSequencer.h
+++ b/Engine/Source/Editor/Sequencer/Private/SSequencer.h
@@ -620,8 +620,8 @@ private:
 	/** Gets whether or not the time range should be visible. */
 	EVisibility GetTimeRangeVisibility() const;
 
-	/** Gets whether the status bar should be visible. */
-	EVisibility GetStatusBarVisibility() const;
+	/** Gets whether the info button in the playback controls should be visible. */
+	EVisibility GetInfoButtonVisibility() const;
 
 	/** Gets whether the tick lines should be drawn. */
 	EVisibility GetShowTickLines() const;
@@ -648,9 +648,6 @@ private:
 	/** Controls how fast Spinboxes change values. */
 	double GetSpinboxDelta() const;
 
-	/** Get minimum desired width of the current time spin box */
-	float GetPlayTimeMinDesiredWidth() const;
-
 	bool GetIsSequenceReadOnly() const;
 	void OnSetSequenceReadOnly(ECheckBoxState CheckBoxState);
 
diff --git a/Engine/Source/Editor/Sequencer/Private/SSequencerTreeFilterStatusBar.cpp b/Engine/Source/Editor/Sequencer/Private/SSequencerTreeFilterStatusBar.cpp
--- a/Engine/Source/Editor/Sequencer/Private/SSequencerTreeFilterStatusBar.cpp
+++ b/Engine/Source/Editor/Sequencer/Private/SSequencerTreeFilterStatusBar.cpp
@@ -28,28 +28,28 @@ void SSequencerTreeFilterStatusBar::Construct(const FArguments& InArgs, TSharedP
 
 		+ SHorizontalBox::Slot()
 		.AutoWidth()
+		.Padding(FMargin(3.f, 0.f, 0.f, 0.f))
 		[
 			SNew(SBorder)
 			.BorderImage(FAppStyle::GetBrush("ToolPanel.GroupBorder"))
 			.Padding(FMargin(3.0f, 3.0f))
-			.Visibility(EVisibility::HitTestInvisible)
 			[
-				SAssignNew(TextBlock, STextBlock)
+				SNew(SHyperlink)
+				.Visibility(this, &SSequencerTreeFilterStatusBar::GetVisibilityFromFilter)
+				.Text(LOCTEXT("ClearFilters", "clear"))
+				.OnNavigate(this, &SSequencerTreeFilterStatusBar::ClearFilters)
 			]
 		]
-
+	
 		+ SHorizontalBox::Slot()
 		.AutoWidth()
-		.Padding(FMargin(3.f, 0.f, 0.f, 0.f))
 		[
 			SNew(SBorder)
 			.BorderImage(FAppStyle::GetBrush("ToolPanel.GroupBorder"))
 			.Padding(FMargin(3.0f, 3.0f))
+			.Visibility(EVisibility::HitTestInvisible)
 			[
-				SNew(SHyperlink)
-				.Visibility(this, &SSequencerTreeFilterStatusBar::GetVisibilityFromFilter)
-				.Text(LOCTEXT("ClearFilters", "clear"))
-				.OnNavigate(this, &SSequencerTreeFilterStatusBar::ClearFilters)
+				SAssignNew(TextBlock, STextBlock)
 			]
 		]
 	];
@@ -138,4 +138,63 @@ void SSequencerTreeFilterStatusBar::UpdateText()
 	TextBlock->SetText(NewText);
 }
 
+static double SequencerOpacityDurationSeconds = 2.0;
+
+void SSequencerTreeFilterStatusBar::Tick(const FGeometry& AllottedGeometry, const double InCurrentTime, const float InDeltaTime)
+{
+	if (GetVisibility() == EVisibility::Visible && OpacityThrobEndTime != FLT_MAX)
+	{
+		if (IsHovered())
+		{
+			// Make sure the status bar stays visible and keep it from fading until OnUnhovered
+			ShowStatusBar();
+			FadeOutStatusBar();
+		}
+		else
+		{
+			double CurrentTime = FPlatformTime::Seconds();
+
+			float Opacity = 0.f;
+			if (OpacityThrobEndTime > CurrentTime)
+			{
+				double Difference = OpacityThrobEndTime - CurrentTime;
+				Opacity = Difference / SequencerOpacityDurationSeconds;
+			}
+
+			if (Opacity > 0.f)
+			{
+				FLinearColor BorderColor = GetColorAndOpacity();
+				BorderColor.A = Opacity;
+				SetColorAndOpacity(BorderColor);
+			}
+			else
+			{
+				SetVisibility(EVisibility::Hidden);
+			}
+		}
+	}
+
+	SCompoundWidget::Tick(AllottedGeometry, InCurrentTime, InDeltaTime);
+}
+
+void SSequencerTreeFilterStatusBar::ShowStatusBar()
+{
+	OpacityThrobEndTime = FLT_MAX;
+	FLinearColor BorderColor = GetColorAndOpacity();
+	BorderColor.A = 1.f;
+	SetColorAndOpacity(BorderColor);
+	SetVisibility(EVisibility::Visible);
+}
+
+void SSequencerTreeFilterStatusBar::HideStatusBar()
+{
+	SetVisibility(EVisibility::Hidden);
+}
+
+void SSequencerTreeFilterStatusBar::FadeOutStatusBar()
+{
+	// Show and then start fading out the status bar
+	OpacityThrobEndTime = FPlatformTime::Seconds() + SequencerOpacityDurationSeconds;
+}
+
 #undef LOCTEXT_NAMESPACE
diff --git a/Engine/Source/Editor/Sequencer/Private/SSequencerTreeFilterStatusBar.h b/Engine/Source/Editor/Sequencer/Private/SSequencerTreeFilterStatusBar.h
--- a/Engine/Source/Editor/Sequencer/Private/SSequencerTreeFilterStatusBar.h
+++ b/Engine/Source/Editor/Sequencer/Private/SSequencerTreeFilterStatusBar.h
@@ -21,6 +21,12 @@ public:
 	/** Update the filter text to represent the current filter states in the tree */
 	void UpdateText();
 
+	virtual void Tick(const FGeometry& AllottedGeometry, const double InCurrentTime, const float InDeltaTime) override;
+
+	void ShowStatusBar();
+	void HideStatusBar();
+	void FadeOutStatusBar();
+
 private:
 
 	/** Request the visibility of the clear hyperlink widget based on whether there's a filter active or not */
@@ -33,4 +39,5 @@ private:
 
 	TWeakPtr<FSequencer> WeakSequencer;
 	TSharedPtr<STextBlock> TextBlock;
+	double OpacityThrobEndTime = 0;
 };
\ No newline at end of file
diff --git a/Engine/Source/Editor/Sequencer/Private/Sequencer.cpp b/Engine/Source/Editor/Sequencer/Private/Sequencer.cpp
--- a/Engine/Source/Editor/Sequencer/Private/Sequencer.cpp
+++ b/Engine/Source/Editor/Sequencer/Private/Sequencer.cpp
@@ -4215,6 +4215,71 @@ TSharedRef<SWidget> FSequencer::MakeTransportControls(bool bExtended)
 	return EditorWidgetsModule.CreateTransportControl( TransportControlArgs );
 }
 
+TSharedRef<SWidget> FSequencer::MakePlayTimeDisplay(const TSharedRef<INumericTypeInterface<double>>& InNumericTypeInterface)
+{
+	return SNew(STemporarilyFocusedSpinBox<double>)
+		.Style(&FAppStyle::GetWidgetStyle<FSpinBoxStyle>("Sequencer.PlayTimeSpinBox"))
+		.Value_Lambda([this]() -> double {
+			return GetLocalTime().Time.GetFrame().Value;
+		})
+		.OnValueChanged_Lambda([this](double InFrame) {
+			SequencerWidget->SetPlayTimeClampedByWorkingRange(InFrame);
+		})
+		.OnValueCommitted_Lambda([this](double InFrame, ETextCommit::Type) {
+			SequencerWidget->SetPlayTime(InFrame);
+		})
+		.MinValue(TOptional<double>())
+		.MaxValue(TOptional<double>())
+		.ToolTipText_Lambda([this] {
+			FFrameRate TickResolution = GetFocusedTickResolution();
+			FFrameRate DisplayRate = GetFocusedDisplayRate();
+
+			FFrameNumber CurrentFrame = GetLocalTime().Time.GetFrame();
+			TRange<FFrameNumber> CurrentRange = GetSubSequenceRange().IsSet() ? GetSubSequenceRange().GetValue() : GetPlaybackRange();
+
+			FFrameNumber FrameCount = FFrameRate::TransformTime((CurrentFrame - CurrentRange.GetLowerBoundValue() + 1).Value, TickResolution, DisplayRate).CeilToFrame();
+			FFrameNumber FrameDuration = FFrameRate::TransformTime(CurrentRange.Size<FFrameNumber>().Value, TickResolution, DisplayRate).CeilToFrame();
+
+			return FText::FromString(FString::Printf(TEXT("%d of %d"), FrameCount.Value, FrameDuration.Value));
+		})
+		.TypeInterface(InNumericTypeInterface)
+		.Delta(this, &FSequencer::GetDisplayRateDeltaFrameCount)
+		.LinearDeltaSensitivity(25)
+		.MinDesiredWidth_Lambda([this] {
+			TRange<double> ViewRange = GetViewRange();
+
+			FString LowerBoundStr = GetNumericTypeInterface()->ToString(ViewRange.GetLowerBoundValue());
+			FString UpperBoundStr = GetNumericTypeInterface()->ToString(ViewRange.GetUpperBoundValue());
+
+			// Always measure with the negative and subframe indicator so that the size doesn't change when there is and isn't a subframe
+			if (!LowerBoundStr.Contains(TEXT("*")))
+			{
+				LowerBoundStr += TEXT("*");
+			}
+			if (!LowerBoundStr.Contains(TEXT("-")))
+			{
+				LowerBoundStr += TEXT("-");
+			}
+			if (!UpperBoundStr.Contains(TEXT("*")))
+			{
+				UpperBoundStr += TEXT("*");
+			}
+			if (!UpperBoundStr.Contains(TEXT("-")))
+			{
+				UpperBoundStr += TEXT("-");
+			}
+
+			const FSlateFontInfo NormalFont = FCoreStyle::Get().GetFontStyle(TEXT("NormalFont"));
+
+			const TSharedRef< FSlateFontMeasure > FontMeasureService = FSlateApplication::Get().GetRenderer()->GetFontMeasureService();
+
+			FVector2D LowerTextSize = FontMeasureService->Measure(LowerBoundStr, NormalFont);
+			FVector2D UpperTextSize = FontMeasureService->Measure(UpperBoundStr, NormalFont);
+
+			return FMath::Max(LowerTextSize.X, UpperTextSize.X);
+		});
+}
+
 TSharedRef<SWidget> FSequencer::OnCreateTransportSetPlaybackStart()
 {
 	FText SetPlaybackStartToolTip = FText::Format(LOCTEXT("SetPlayStart_Tooltip", "Set playback start to the current position ({0})"), FSequencerCommands::Get().SetStartPlaybackRange->GetInputText());
@@ -4222,9 +4287,14 @@ TSharedRef<SWidget> FSequencer::OnCreateTransportSetPlaybackStart()
 	return SNew(SButton)
 		.OnClicked(this, &FSequencer::SetPlaybackStart)
 		.ToolTipText(SetPlaybackStartToolTip)
-		.ButtonStyle(FAppStyle::Get(), "Sequencer.Transport.SetPlayStart")
-		.ContentPadding(2.0f)
-		.IsFocusable(false);
+		.ButtonStyle(FAppStyle::Get(), "Animation.PlayControlsButton")
+		.ContentPadding(0.0f)
+		.IsFocusable(false)
+		[
+			SNew(SImage)
+			.ColorAndOpacity(FLinearColor::Green)
+			.Image(FAppStyle::Get().GetBrush("Sequencer.Transport.SetPlayStart"))
+		];
 }
 
 TSharedRef<SWidget> FSequencer::OnCreateTransportJumpToPreviousKey()
@@ -4234,9 +4304,14 @@ TSharedRef<SWidget> FSequencer::OnCreateTransportJumpToPreviousKey()
 	return SNew(SButton)
 		.OnClicked(this, &FSequencer::JumpToPreviousKey)
 		.ToolTipText(JumpToPreviousKeyToolTip)
-		.ButtonStyle(FAppStyle::Get(), "Sequencer.Transport.JumpToPreviousKey")
-		.ContentPadding(2.0f)
-		.IsFocusable(false);
+		.ButtonStyle(FAppStyle::Get(), "Animation.PlayControlsButton")
+		.ContentPadding(0.0f)
+		.IsFocusable(false)
+		[
+			SNew(SImage)
+			.ColorAndOpacity(FSlateColor::UseSubduedForeground())
+			.Image(FAppStyle::Get().GetBrush("Sequencer.Transport.JumpToPreviousKey"))
+		];
 }
 
 TSharedRef<SWidget> FSequencer::OnCreateTransportJumpToNextKey()
@@ -4246,9 +4321,14 @@ TSharedRef<SWidget> FSequencer::OnCreateTransportJumpToNextKey()
 	return SNew(SButton)
 		.OnClicked(this, &FSequencer::JumpToNextKey)
 		.ToolTipText(JumpToNextKeyToolTip)
-		.ButtonStyle(FAppStyle::Get(), "Sequencer.Transport.JumpToNextKey")
-		.ContentPadding(2.0f)
-		.IsFocusable(false);
+		.ButtonStyle(FAppStyle::Get(), "Animation.PlayControlsButton")
+		.ContentPadding(0.0f)
+		.IsFocusable(false)
+		[
+			SNew(SImage)
+			.ColorAndOpacity(FSlateColor::UseSubduedForeground())
+			.Image(FAppStyle::Get().GetBrush("Sequencer.Transport.JumpToNextKey"))
+		];
 }
 
 TSharedRef<SWidget> FSequencer::OnCreateTransportSetPlaybackEnd()
@@ -4258,16 +4338,21 @@ TSharedRef<SWidget> FSequencer::OnCreateTransportSetPlaybackEnd()
 	return SNew(SButton)
 		.OnClicked(this, &FSequencer::SetPlaybackEnd)
 		.ToolTipText(SetPlaybackEndToolTip)
-		.ButtonStyle(FAppStyle::Get(), "Sequencer.Transport.SetPlayEnd")
-		.ContentPadding(2.0f)
-		.IsFocusable(false);
+		.ButtonStyle(FAppStyle::Get(), "Animation.PlayControlsButton")
+		.ContentPadding(0.0f)
+		.IsFocusable(false)
+		[
+			SNew(SImage)
+			.ColorAndOpacity(FLinearColor::Red)
+			.Image(FAppStyle::Get().GetBrush("Sequencer.Transport.SetPlayEnd"))
+		];
 }
 
 TSharedRef<SWidget> FSequencer::OnCreateTransportLoopMode()
 {
 	TSharedRef<SButton> LoopButton = SNew(SButton)
 		.OnClicked(this, &FSequencer::OnCycleLoopMode)
-		.ButtonStyle( FAppStyle::Get(), "NoBorder" )
+		.ButtonStyle( FAppStyle::Get(), "Animation.PlayControlsButton" )
 		.IsFocusable(false)
 		.ToolTipText_Lambda([&]()
 		{ 
@@ -4284,30 +4369,25 @@ TSharedRef<SWidget> FSequencer::OnCreateTransportLoopMode()
 				return LOCTEXT("LoopModeLoopSelectionRange_Tooltip", "Loop selection range");
 			}
 		})
-		.ContentPadding(2.0f);
+		.ContentPadding(0.0f);
 
 	TWeakPtr<SButton> WeakButton = LoopButton;
 
 	LoopButton->SetContent(SNew(SImage)
+		.ColorAndOpacity(FSlateColor::UseSubduedForeground())
 		.Image_Lambda([&, WeakButton]()
 		{
 			if (GetLoopMode() == ESequencerLoopMode::SLM_NoLoop)
 			{
-				return WeakButton.IsValid() && WeakButton.Pin()->IsPressed() ? 
-					&FAppStyle::Get().GetWidgetStyle<FButtonStyle>("Animation.Loop.Disabled").Pressed : 
-					&FAppStyle::Get().GetWidgetStyle<FButtonStyle>("Animation.Loop.Disabled").Normal;
+				return FAppStyle::Get().GetBrush("Animation.Loop.Disabled");
 			}
 			else if (GetLoopMode() == ESequencerLoopMode::SLM_Loop)
 			{
-				return WeakButton.IsValid() && WeakButton.Pin()->IsPressed() ? 
-					&FAppStyle::Get().GetWidgetStyle<FButtonStyle>("Animation.Loop.Enabled").Pressed : 
-					&FAppStyle::Get().GetWidgetStyle<FButtonStyle>("Animation.Loop.Enabled").Normal;
+				return FAppStyle::Get().GetBrush("Animation.Loop.Enabled");
 			}
 			else
 			{
-				return WeakButton.IsValid() && WeakButton.Pin()->IsPressed() ? 
-					&FAppStyle::Get().GetWidgetStyle<FButtonStyle>("Animation.Loop.SelectionRange").Pressed : 
-					&FAppStyle::Get().GetWidgetStyle<FButtonStyle>("Animation.Loop.SelectionRange").Normal;
+				return FAppStyle::Get().GetBrush("Animation.Loop.SelectionRange");
 			}
 		})
 	);
@@ -4319,7 +4399,7 @@ TSharedRef<SWidget> FSequencer::OnCreateTransportRecord()
 {
 	TSharedRef<SButton> RecordButton = SNew(SButton)
 		.OnClicked(this, &FSequencer::OnRecord)
-		.ButtonStyle( FAppStyle::Get(), "NoBorder" )
+		.ButtonStyle( FAppStyle::Get(), "Animation.PlayControlsButton" )
 		.IsFocusable(false)
 		.ToolTipText_Lambda([&]()
 		{
@@ -4340,24 +4420,12 @@ TSharedRef<SWidget> FSequencer::OnCreateTransportRecord()
 		})
 		.Visibility_Lambda([&] { return HostCapabilities.bSupportsRecording && OnGetCanRecord().IsBound() ? EVisibility::Visible : EVisibility::Collapsed; })
 		.IsEnabled_Lambda([&] { FText OutErrorText; return OnGetCanRecord().IsBound() && OnGetCanRecord().Execute(OutErrorText); })
-		.ContentPadding(2.0f);
+		.ContentPadding(0.0f);
 
 	TWeakPtr<SButton> WeakButton = RecordButton;
 
 	RecordButton->SetContent(SNew(SImage)
-		.Image_Lambda([this, WeakButton]()
-		{
-			if (OnGetIsRecording().IsBound() && OnGetIsRecording().Execute())
-			{
-				return WeakButton.IsValid() && WeakButton.Pin()->IsPressed() ? 
-					&FAppStyle::Get().GetWidgetStyle<FButtonStyle>("Animation.Recording").Pressed : 
-					&FAppStyle::Get().GetWidgetStyle<FButtonStyle>("Animation.Recording").Normal;
-			}
-
-			return WeakButton.IsValid() && WeakButton.Pin()->IsPressed() ? 
-				&FAppStyle::Get().GetWidgetStyle<FButtonStyle>("Animation.Record").Pressed : 
-				&FAppStyle::Get().GetWidgetStyle<FButtonStyle>("Animation.Record").Normal;
-		})
+		.Image(FAppStyle::Get().GetBrush("Animation.Record"))
 		.ColorAndOpacity_Lambda([this]()
 		{
 			if (OnGetIsRecording().IsBound() && OnGetIsRecording().Execute())
@@ -4367,11 +4435,13 @@ TSharedRef<SWidget> FSequencer::OnCreateTransportRecord()
 					RecordingAnimation.Play(SequencerWidget.ToSharedRef(), true);
 				}
 
-				return FLinearColor(1.f, 1.f, 1.f, 0.2f + 0.8f * RecordingAnimation.GetLerp());
+				FLinearColor Color = FSlateColor::UseSubduedForeground().GetSpecifiedColor();
+				Color.A = 0.2f + 0.8f * RecordingAnimation.GetLerp();
+				return FSlateColor(Color);
 			}
 
 			RecordingAnimation.Pause();
-			return FLinearColor::White;		
+			return FSlateColor::UseSubduedForeground();
 		})
 	);
 
@@ -10615,12 +10685,12 @@ void FSequencer::BindCommands()
 		FIsActionChecked::CreateLambda( [this]{ return Settings->GetShowChannelColors(); } ) );
 
 	SequencerCommandBindings->MapAction(
-		Commands.ToggleShowStatusBar,
+		Commands.ToggleShowInfoButton,
 		FExecuteAction::CreateLambda([this] {
-			Settings->SetShowStatusBar(!Settings->GetShowStatusBar());
+			Settings->SetShowInfoButton(!Settings->GetShowInfoButton());
 		}),
 		FCanExecuteAction::CreateLambda([] { return true; }),
-		FIsActionChecked::CreateLambda([this] { return Settings->GetShowStatusBar(); }));
+		FIsActionChecked::CreateLambda([this] { return Settings->GetShowInfoButton(); }));
 
 	SequencerCommandBindings->MapAction(
 		Commands.ToggleShowSelectedNodesOnly,
diff --git a/Engine/Source/Editor/Sequencer/Private/Sequencer.h b/Engine/Source/Editor/Sequencer/Private/Sequencer.h
--- a/Engine/Source/Editor/Sequencer/Private/Sequencer.h
+++ b/Engine/Source/Editor/Sequencer/Private/Sequencer.h
@@ -1169,6 +1169,8 @@ public:
 	/** Get the name of the movie renderer to use, defaults to the first available if the setting is empty */
 	FString GetMovieRendererName() const;
 
+	TSharedRef<SWidget> MakePlayTimeDisplay(const TSharedRef<INumericTypeInterface<double>>& InNumericTypeInterface);
+
 private:
 
 	/** Update the time bases for the current movie scene */
diff --git a/Engine/Source/Editor/Sequencer/Private/SequencerCommands.cpp b/Engine/Source/Editor/Sequencer/Private/SequencerCommands.cpp
--- a/Engine/Source/Editor/Sequencer/Private/SequencerCommands.cpp
+++ b/Engine/Source/Editor/Sequencer/Private/SequencerCommands.cpp
@@ -131,7 +131,7 @@ void FSequencerCommands::RegisterCommands()
 	UI_COMMAND( OpenTaggedBindingManager, "Open Binding Tag Manager", "Specifies options for tagging bindings within this sequence for external systems to reference as a persistent name.", EUserInterfaceActionType::Button, FInputChord() );
 	UI_COMMAND( OpenNodeGroupsManager, "Open Sequencer Group Manager", "Manage groups within this sequence.", EUserInterfaceActionType::Button, FInputChord());
 	UI_COMMAND( ToggleShowRangeSlider, "Range Slider", "Enables and disables showing the time range slider", EUserInterfaceActionType::ToggleButton, FInputChord());
-	UI_COMMAND( ToggleShowStatusBar, "Status Bar", "Enables and disables showing the status bar", EUserInterfaceActionType::ToggleButton, FInputChord());
+	UI_COMMAND( ToggleShowInfoButton, "Info Button", "Enables and disables showing the information button in the playback controls", EUserInterfaceActionType::ToggleButton, FInputChord());
 	UI_COMMAND( ToggleIsSnapEnabled, "Enable Snapping", "Enables and disables snapping", EUserInterfaceActionType::ToggleButton, FInputChord() );
 
 	UI_COMMAND( ToggleSnapKeyTimesToInterval, "Snap to the Interval", "Snap keys to the time snapping interval", EUserInterfaceActionType::ToggleButton, FInputChord() );
diff --git a/Engine/Source/Editor/Sequencer/Private/SequencerNodeTree.cpp b/Engine/Source/Editor/Sequencer/Private/SequencerNodeTree.cpp
--- a/Engine/Source/Editor/Sequencer/Private/SequencerNodeTree.cpp
+++ b/Engine/Source/Editor/Sequencer/Private/SequencerNodeTree.cpp
@@ -55,6 +55,7 @@ FSequencerNodeTree::~FSequencerNodeTree()
 
 FSequencerNodeTree::FSequencerNodeTree(FSequencer& InSequencer)
 	: Sequencer(InSequencer)
+	, DisplayNodeCount(0)
 	, bFilterUpdateRequested(false)
 {
 	TrackFilters = MakeShared<FSequencerTrackFilterCollection>();
diff --git a/Engine/Source/Editor/Sequencer/Private/SequencerSettings.cpp b/Engine/Source/Editor/Sequencer/Private/SequencerSettings.cpp
--- a/Engine/Source/Editor/Sequencer/Private/SequencerSettings.cpp
+++ b/Engine/Source/Editor/Sequencer/Private/SequencerSettings.cpp
@@ -49,7 +49,7 @@ USequencerSettings::USequencerSettings( const FObjectInitializer& ObjectInitiali
 	bShowKeyBars = true;
 	bInfiniteKeyAreas = false;
 	bShowChannelColors = false;
-	bShowStatusBar = true;
+	bShowInfoButton = true;
 	ReduceKeysTolerance = KINDA_SMALL_NUMBER;
 	KeyAreaHeightWithCurves = SequencerLayoutConstants::KeyAreaHeight;
 	bDeleteKeysWhenTrimming = true;
@@ -624,16 +624,16 @@ void USequencerSettings::SetShowChannelColors(bool InbShowChannelColors)
 	}
 }
 
-bool USequencerSettings::GetShowStatusBar() const
+bool USequencerSettings::GetShowInfoButton() const
 {
-	return bShowStatusBar;
+	return bShowInfoButton;
 }
 
-void USequencerSettings::SetShowStatusBar(bool InbShowStatusBar)
+void USequencerSettings::SetShowInfoButton(bool InbShowInfoButton)
 {
-	if (bShowStatusBar != InbShowStatusBar)
+	if (bShowInfoButton != InbShowInfoButton)
 	{
-		bShowStatusBar = InbShowStatusBar;
+		bShowInfoButton = InbShowInfoButton;
 		SaveConfig();
 	}
 }
diff --git a/Engine/Source/Editor/Sequencer/Public/MVVM/CurveEditorExtension.h b/Engine/Source/Editor/Sequencer/Public/MVVM/CurveEditorExtension.h
--- a/Engine/Source/Editor/Sequencer/Public/MVVM/CurveEditorExtension.h
+++ b/Engine/Source/Editor/Sequencer/Public/MVVM/CurveEditorExtension.h
@@ -14,6 +14,7 @@ class ISequencer;
 class IPropertyTypeCustomization;
 class SCurveEditorPanel;
 class SCurveEditorTree;
+class SCurveEditorTreeFilterStatusBar;
 struct FTimeSliderArgs;
 
 namespace UE
@@ -81,6 +82,8 @@ private:
 	TSharedPtr<SWidget> CurveEditorWidget;
 	/** The curve editor panel. This is created and updated even if it is not currently visible. */
 	TSharedPtr<SCurveEditorPanel> CurveEditorPanel;
+	/** Filter Status Bar */
+	TSharedPtr<SCurveEditorTreeFilterStatusBar> CurveEditorTreeFilterStatusBar;
 
 	friend class FCurveEditorIntegrationExtension;
 };
diff --git a/Engine/Source/Editor/Sequencer/Public/SequencerCommands.h b/Engine/Source/Editor/Sequencer/Public/SequencerCommands.h
--- a/Engine/Source/Editor/Sequencer/Public/SequencerCommands.h
+++ b/Engine/Source/Editor/Sequencer/Public/SequencerCommands.h
@@ -391,8 +391,8 @@ public:
 	/** Toggles whether to show channel colors in the track area. */
 	TSharedPtr< FUICommandInfo > ToggleChannelColors;
 
-	/** Toggles whether to show the status bar. */
-	TSharedPtr< FUICommandInfo > ToggleShowStatusBar;
+	/** Toggles whether to show the info button in the playback controls. */
+	TSharedPtr< FUICommandInfo > ToggleShowInfoButton;
 
 	/** Turns auto scroll on and off. */
 	TSharedPtr< FUICommandInfo > ToggleAutoScroll;
diff --git a/Engine/Source/Editor/Sequencer/Public/SequencerSettings.h b/Engine/Source/Editor/Sequencer/Public/SequencerSettings.h
--- a/Engine/Source/Editor/Sequencer/Public/SequencerSettings.h
+++ b/Engine/Source/Editor/Sequencer/Public/SequencerSettings.h
@@ -322,10 +322,10 @@ public:
 	/** Set whether to show channel colors for the key bars */
 	void SetShowChannelColors(bool bInShowChannelColors);
 
-	/** @return true if showing status bar */
-	bool GetShowStatusBar() const;
-	/** Set whether to show status bar */
-	void SetShowStatusBar(bool bInShowStatusBar);
+	/** @return true if showing the info button in the playback controls */
+	bool GetShowInfoButton() const;
+	/** Set whether to show the info button in the playback controls */
+	void SetShowInfoButton(bool bInShowInfoButton);
 
 	/** @return true if showing tick lines */
 	bool GetShowTickLines() const;
@@ -628,9 +628,9 @@ protected:
 	UPROPERTY(config, EditAnywhere, Category = Timeline)
 	bool bShowChannelColors;
 
-	/** Enable or disable displaying the status bar for number of items. */
+	/** Enable or disable displaying the info button in the playback controls. */
 	UPROPERTY(config, EditAnywhere, Category = Timeline)
-	bool bShowStatusBar;
+	bool bShowInfoButton;
 
 	/** Enable or disable displaying the tick lines. */
 	UPROPERTY(config, EditAnywhere, Category = Timeline)
diff --git a/Engine/Source/Runtime/SlateCore/Private/Styling/StarshipCoreStyle.cpp b/Engine/Source/Runtime/SlateCore/Private/Styling/StarshipCoreStyle.cpp
--- a/Engine/Source/Runtime/SlateCore/Private/Styling/StarshipCoreStyle.cpp
+++ b/Engine/Source/Runtime/SlateCore/Private/Styling/StarshipCoreStyle.cpp
@@ -315,6 +315,7 @@ TSharedRef<ISlateStyle> FStarshipCoreStyle::Create()
 		Style->Set("Icons.WarningWithColor.Thumbnail", new IMAGE_BRUSH_SVG("Starship/Common/alert-triangle-64", Icon64x64));
 
 		Style->Set("Icons.Info", new IMAGE_BRUSH_SVG("Starship/Common/Info", Icon16x16));
+		Style->Set("Icons.Info.Small", new IMAGE_BRUSH_SVG("Starship/Common/Info", Icon14x14));
 		Style->Set("Icons.Info.Large", new IMAGE_BRUSH_SVG("Starship/Common/Info", Icon32x32));
 		Style->Set("Icons.InfoWithColor", new IMAGE_BRUSH_SVG("Starship/Common/Info", Icon16x16, FStyleColors::Foreground));
 		Style->Set("Icons.InfoWithColor.Large", new IMAGE_BRUSH_SVG("Starship/Common/Info", Icon32x32, FStyleColors::Foreground));
