Sequencer: Tree View UX

- TableRowStyle for Sequencer which is a combination of the alternative row style and the show parent style - Alternating rows have different tints, parents have different tints.
- Make sure top level rows also report as a parent row
- Rows that have keys or sections selected also report as highlighted

diff --git a/Engine/Source/Editor/EditorStyle/Private/StarshipStyle.cpp b/Engine/Source/Editor/EditorStyle/Private/StarshipStyle.cpp
--- a/Engine/Source/Editor/EditorStyle/Private/StarshipStyle.cpp
+++ b/Engine/Source/Editor/EditorStyle/Private/StarshipStyle.cpp
@@ -2016,6 +2016,13 @@ void FStarshipEditorStyle::FStyle::SetupSequencerStyles()
 
 		Set("SequencerToolBar", SequencerToolbar);
 
+		const FTableRowStyle AlternatingTableRowStyle = GetWidgetStyle<FTableRowStyle>("TableView.AlternatingRow");
+
+		Set("Sequencer.TableView.Row", FTableRowStyle(AlternatingTableRowStyle)
+			.SetUseParentRowBrush(true)
+			.SetParentRowBackgroundBrush(FSlateRoundedBoxBrush(FStyleColors::Header, 2.f, FStyleColors::Transparent, 1.f))
+			.SetParentRowBackgroundHoveredBrush(FSlateRoundedBoxBrush(FStyleColors::SelectHover, 2.f)));
+
 		Set("Sequencer.IconKeySmartAuto", new IMAGE_BRUSH("Sequencer/IconKeySmartAuto", Icon12x12));
 		Set("Sequencer.IconKeyAuto", new IMAGE_BRUSH("Sequencer/IconKeyAuto", Icon12x12));
 		Set("Sequencer.IconKeyBreak", new IMAGE_BRUSH("Sequencer/IconKeyBreak", Icon12x12));
diff --git a/Engine/Source/Editor/Sequencer/Private/MVVM/Views/SSequencerOutlinerView.cpp b/Engine/Source/Editor/Sequencer/Private/MVVM/Views/SSequencerOutlinerView.cpp
--- a/Engine/Source/Editor/Sequencer/Private/MVVM/Views/SSequencerOutlinerView.cpp
+++ b/Engine/Source/Editor/Sequencer/Private/MVVM/Views/SSequencerOutlinerView.cpp
@@ -2,9 +2,11 @@
 
 #include "MVVM/Views/SSequencerOutlinerView.h"
 #include "MVVM/Extensions/IOutlinerExtension.h"
+#include "MVVM/ViewModels/OutlinerSpacer.h"
 #include "MVVM/ViewModels/OutlinerViewModel.h"
 #include "MVVM/Views/SOutlinerItemViewBase.h"
 
+#include "SequencerCoreFwd.h"
 #include "Styling/StyleColors.h"
 
 namespace UE::Sequencer
@@ -19,33 +21,39 @@ public:
 		WeakOutliner = InWeakOutliner;
 
 		SOutlinerViewRow::Construct(InArgs, OwnerTableView, InWeakModel);
-
-		if (InWeakModel.Pin()->HasBackground())
-		{
-			SetBorderImage(TAttribute<const FSlateBrush*>(FAppStyle::GetBrush("WhiteBrush")));
-			SetBorderBackgroundColor(MakeAttributeSP(this, &SSequencerOutlinerViewRow::GetBorderTint));
-		}
 	}
 
-	FSlateColor GetBorderTint() const
+	virtual const FSlateBrush* GetBorder() const override
 	{
-		TSharedPtr<FOutlinerViewModel>    Outliner     = WeakOutliner.Pin();
+		const bool bEvenEntryIndex = (IndexInList % 2 == 0);
+
+		TSharedPtr<FOutlinerViewModel>    Outliner = WeakOutliner.Pin();
 		TViewModelPtr<IOutlinerExtension> OutlinerItem = WeakModel.Pin();
 
 		if (!Outliner || !OutlinerItem)
 		{
-			return FLinearColor(0.f,0.f,0.f,0.f);
+			return SOutlinerViewRow::GetBorder();
+		}
+
+		// If selected or highlighted, early out and return the default selected or highlighted border
+		if (IsSelected() || IsHighlighted())
+		{
+			return SOutlinerViewRow::GetBorder();
 		}
 
 		EOutlinerSelectionState SelectionState = OutlinerItem->GetSelectionState();
 
 		if (EnumHasAnyFlags(SelectionState, EOutlinerSelectionState::SelectedDirectly))
 		{
-			return FStyleColors::Select;
+			return &Style->ActiveBrush;
 		}
+
+		// If keys or the track area are selected, highlight this track row
 		if (EnumHasAnyFlags(SelectionState, EOutlinerSelectionState::HasSelectedKeys | EOutlinerSelectionState::HasSelectedTrackAreaItems))
 		{
-			return FStyleColors::Header;
+			return IsHovered()
+				? (bEvenEntryIndex ? &Style->EvenRowBackgroundHoveredBrush : &Style->OddRowBackgroundHoveredBrush)
+				: &Style->ActiveHighlightedBrush;
 		}
 
 		// If this is collapsed but has any children with selected keys or sections, we report that state on the parent
@@ -53,24 +61,21 @@ public:
 		{
 			if (EnumHasAnyFlags(SelectionState, EOutlinerSelectionState::DescendentHasSelectedTrackAreaItems | EOutlinerSelectionState::DescendentHasSelectedKeys))
 			{
-				return FStyleColors::Header;
+				return IsHovered()
+					? (bEvenEntryIndex ? &Style->EvenRowBackgroundHoveredBrush : &Style->OddRowBackgroundHoveredBrush)
+					: &Style->ActiveHighlightedBrush;
 			}
 		}
 
-		// @todo: Currently we are always using 'default' for the style here. A lot of this is carried over from earlier Sequencer UIs
-		//        so we either need to re-evaluate whether EOutlinerItemViewBaseStyle is even necessary, or add proper support to the outliner extension to specify it
-		EOutlinerItemViewBaseStyle ItemStyle = EOutlinerItemViewBaseStyle::Default; /* OutlinerItem->GetItemViewStyle() */;
-
-		if (Outliner->GetHoveredItem() == OutlinerItem)
+		// If this is at the root level, return it as a parent row
+		if (GetIndentLevel() == 0 && !OutlinerItem.AsModel()->IsA<FOutlinerSpacer>())
 		{
-			return /*ItemStyle == EOutlinerItemViewBaseStyle::ContainerHeader
-				? FLinearColor(FColor(52, 52, 52, 255))
-				: */FLinearColor(FColor(72, 72, 72, 255));
+			return IsHovered()
+				? &Style->ParentRowBackgroundHoveredBrush
+				: &Style->ParentRowBackgroundBrush;
 		}
 
-		return /*ItemStyle == EOutlinerItemViewBaseStyle::ContainerHeader
-			? FLinearColor(FColor(48, 48, 48, 255))
-			: */FLinearColor(FColor(62, 62, 62, 255));
+		return SOutlinerViewRow::GetBorder();
 	}
 
 private:
diff --git a/Engine/Source/Editor/SequencerCore/Private/MVVM/Views/SOutlinerItemViewBase.cpp b/Engine/Source/Editor/SequencerCore/Private/MVVM/Views/SOutlinerItemViewBase.cpp
--- a/Engine/Source/Editor/SequencerCore/Private/MVVM/Views/SOutlinerItemViewBase.cpp
+++ b/Engine/Source/Editor/SequencerCore/Private/MVVM/Views/SOutlinerItemViewBase.cpp
@@ -92,7 +92,7 @@ void SOutlinerItemViewBase::Construct(
 		InnerNodePadding = FMargin(0.f);
 	}
 
-	TableRowStyle = &FAppStyle::Get().GetWidgetStyle<FTableRowStyle>("TableView.Row");
+	TableRowStyle = &FAppStyle::Get().GetWidgetStyle<FTableRowStyle>("Sequencer.TableView.Row");
 
 	TSharedPtr<SWidget> LabelContent;
 
diff --git a/Engine/Source/Editor/SequencerCore/Private/MVVM/Views/SOutlinerView.cpp b/Engine/Source/Editor/SequencerCore/Private/MVVM/Views/SOutlinerView.cpp
--- a/Engine/Source/Editor/SequencerCore/Private/MVVM/Views/SOutlinerView.cpp
+++ b/Engine/Source/Editor/SequencerCore/Private/MVVM/Views/SOutlinerView.cpp
@@ -76,6 +76,7 @@ void SOutlinerViewRow::Construct(const FArguments& InArgs, const TSharedRef<STab
 
 	SMultiColumnTableRow::Construct(
 		SMultiColumnTableRow::FArguments()
+			.Style(&FAppStyle::Get().GetWidgetStyle<FTableRowStyle>("Sequencer.TableView.Row"))
 			.OnDragDetected(this, &SOutlinerViewRow::OnDragDetected)
 			.OnCanAcceptDrop(this, &SOutlinerViewRow::OnCanAcceptDrop)
 			.OnAcceptDrop(this, &SOutlinerViewRow::OnAcceptDrop)
diff --git a/Engine/Source/Runtime/Slate/Public/Widgets/Views/STableRow.h b/Engine/Source/Runtime/Slate/Public/Widgets/Views/STableRow.h
--- a/Engine/Source/Runtime/Slate/Public/Widgets/Views/STableRow.h
+++ b/Engine/Source/Runtime/Slate/Public/Widgets/Views/STableRow.h
@@ -1068,6 +1068,23 @@ public:
 		return false;
 	}
 
+	/**
+	 * Callback to determine if the row is highlighted or not
+	 *
+	 * @return		true if highlighted by owning widget.
+	 */
+	bool IsHighlighted() const
+	{
+		TSharedRef< ITypedTableView< ItemType > > OwnerTable = OwnerTablePtr.Pin().ToSharedRef();
+
+		if (const TObjectPtrWrapTypeOf<ItemType>* MyItemPtr = GetItemForThis(OwnerTable))
+		{
+			return OwnerTable->Private_IsItemHighlighted(*MyItemPtr);
+		}
+
+		return false;
+	}
+
 	/** By default, this function does nothing, it should be implemented by derived class */
 	virtual FVector2D GetRowSizeForColumn(const FName& InColumnName) const override
 	{
