Sequencer: Additional Tree View UX - Highlight parents even if the selected item is collapsed

Same as SSceneOutlinerTreeView::Private_UpdateParentHighlights

diff --git a/Engine/Source/Editor/Sequencer/Private/MVVM/Views/SSequencerOutlinerView.cpp b/Engine/Source/Editor/Sequencer/Private/MVVM/Views/SSequencerOutlinerView.cpp
--- a/Engine/Source/Editor/Sequencer/Private/MVVM/Views/SSequencerOutlinerView.cpp
+++ b/Engine/Source/Editor/Sequencer/Private/MVVM/Views/SSequencerOutlinerView.cpp
@@ -49,24 +49,13 @@ public:
 		}
 
 		// If keys or the track area are selected, highlight this track row
-		if (EnumHasAnyFlags(SelectionState, EOutlinerSelectionState::HasSelectedKeys | EOutlinerSelectionState::HasSelectedTrackAreaItems))
+		if (EnumHasAnyFlags(SelectionState, EOutlinerSelectionState::HasSelectedKeys | EOutlinerSelectionState::HasSelectedTrackAreaItems | EOutlinerSelectionState::DescendentHasSelectedTrackAreaItems | EOutlinerSelectionState::DescendentHasSelectedKeys))
 		{
 			return IsHovered()
 				? (bEvenEntryIndex ? &Style->EvenRowBackgroundHoveredBrush : &Style->OddRowBackgroundHoveredBrush)
 				: &Style->ActiveHighlightedBrush;
 		}
 
-		// If this is collapsed but has any children with selected keys or sections, we report that state on the parent
-		if (!OutlinerItem->IsExpanded())
-		{
-			if (EnumHasAnyFlags(SelectionState, EOutlinerSelectionState::DescendentHasSelectedTrackAreaItems | EOutlinerSelectionState::DescendentHasSelectedKeys))
-			{
-				return IsHovered()
-					? (bEvenEntryIndex ? &Style->EvenRowBackgroundHoveredBrush : &Style->OddRowBackgroundHoveredBrush)
-					: &Style->ActiveHighlightedBrush;
-			}
-		}
-
 		// If this is at the root level, return it as a parent row
 		if (GetIndentLevel() == 0 && !OutlinerItem.AsModel()->IsA<FOutlinerSpacer>())
 		{
diff --git a/Engine/Source/Editor/SequencerCore/Private/MVVM/Views/SOutlinerView.cpp b/Engine/Source/Editor/SequencerCore/Private/MVVM/Views/SOutlinerView.cpp
--- a/Engine/Source/Editor/SequencerCore/Private/MVVM/Views/SOutlinerView.cpp
+++ b/Engine/Source/Editor/SequencerCore/Private/MVVM/Views/SOutlinerView.cpp
@@ -887,6 +887,21 @@ FReply SOutlinerView::OnKeyDown(const FGeometry& MyGeometry, const FKeyEvent& In
 	return STreeView<TWeakViewModelPtr<IOutlinerExtension>>::OnKeyDown(MyGeometry, InKeyEvent);
 }
 
+void SOutlinerView::Private_UpdateParentHighlights()
+{
+	this->ClearHighlightedItems();
+
+	// For the Outliner, we want to highlight parent items even if the current selection is not visible (i.e collapsed)
+	for (TWeakViewModelPtr<IOutlinerExtension> WeakSelectedItem : GetSelectedItems())
+	{
+		TViewModelPtr<IOutlinerExtension> SelectedItem = WeakSelectedItem.Pin();
+		for (TViewModelPtr<IOutlinerExtension> Parent : SelectedItem.AsModel()->GetAncestorsOfType<IOutlinerExtension>())
+		{
+			Private_SetItemHighlighted(Parent, true);
+		}
+	}
+}
+
 void SOutlinerView::Private_SetItemSelection( TWeakViewModelPtr<IOutlinerExtension> TheItem, bool bShouldBeSelected, bool bWasUserDirected )
 {
 	if (TSharedPtr<FOutlinerSpacer> Spacer = TheItem.ImplicitPin())
diff --git a/Engine/Source/Editor/SequencerCore/Public/MVVM/Views/SOutlinerView.h b/Engine/Source/Editor/SequencerCore/Public/MVVM/Views/SOutlinerView.h
--- a/Engine/Source/Editor/SequencerCore/Public/MVVM/Views/SOutlinerView.h
+++ b/Engine/Source/Editor/SequencerCore/Public/MVVM/Views/SOutlinerView.h
@@ -160,6 +160,7 @@ protected:
 	void OnExpansionChanged(TWeakViewModelPtr<IOutlinerExtension> InItem, bool bIsExpanded);
 
 	// Tree selection methods which must be overriden to maintain selection consistency with the rest of sequencer.
+	virtual void Private_UpdateParentHighlights() override;
 	virtual void Private_SetItemSelection( TWeakViewModelPtr<IOutlinerExtension> TheItem, bool bShouldBeSelected, bool bWasUserDirected = false ) override;
 	virtual void Private_ClearSelection() override;
 	virtual void Private_SelectRangeFromCurrentTo( TWeakViewModelPtr<IOutlinerExtension> InRangeSelectionEnd ) override;
