Sequencer: Improvements to component possession UX

- No longer show components in the + Track menu if they are already possessed.
- Select components on creation

diff --git a/Engine/Source/Editor/Sequencer/Private/ActorObjectSchema.cpp b/Engine/Source/Editor/Sequencer/Private/ActorObjectSchema.cpp
--- a/Engine/Source/Editor/Sequencer/Private/ActorObjectSchema.cpp
+++ b/Engine/Source/Editor/Sequencer/Private/ActorObjectSchema.cpp
@@ -9,6 +9,13 @@
 #include "ISequencerModule.h"
 #include "Framework/MultiBox/MultiBoxBuilder.h"
 #include "Framework/MultiBox/MultiBoxExtender.h"
+
+#include "MVVM/ObjectBindingModelStorageExtension.h"
+#include "MVVM/Selection/SequencerOutlinerSelection.h"
+#include "MVVM/Selection/Selection.h"
+#include "MVVM/ViewModels/SequencerEditorViewModel.h"
+#include "MVVM/ViewModels/SequenceModel.h"
+#include "MVVM/ViewModels/ObjectBindingModel.h"
 #include "ScopedTransaction.h"
 
 #define LOCTEXT_NAMESPACE "FActorSchema"
@@ -80,6 +87,12 @@ void FActorSchema::HandleTrackMenuExtensionAddTrack(FMenuBuilder& AddTrackMenuBu
 	TSharedRef<FClassViewerFilterFuncs> ClassFilterFuncs = ClassViewerModule.CreateFilterFuncs();
 	FClassViewerInitializationOptions ClassViewerOptions = {};
 
+	TSharedPtr<ISequencer> Sequencer = WeakSequencer.Pin();
+	if (!Sequencer)
+	{
+		return;
+	}
+
 	TSet<FName> AllComponentNames;
 	AddTrackMenuBuilder.BeginSection("Components", LOCTEXT("ComponentsSection", "Components"));
 	{
@@ -92,6 +105,11 @@ void FActorSchema::HandleTrackMenuExtensionAddTrack(FMenuBuilder& AddTrackMenuBu
 					continue;
 				}
 
+				if (Sequencer->GetHandleToObject(Component, false).IsValid())
+				{
+					continue;
+				}
+
 				bool bValidComponent = !Component->IsVisualizationComponent();
 
 				if (GlobalClassFilter.IsValid())
@@ -139,13 +157,27 @@ void FActorSchema::HandleAddComponentActionExecute(FName ComponentName, TWeakPtr
 		return;
 	}
 
+	FObjectBindingModelStorageExtension* ObjectStorage = Sequencer->GetViewModel()->GetRootModel()->CastDynamic<FObjectBindingModelStorageExtension>();
+	check(ObjectStorage);
+
+	TSharedPtr<FSequencerSelection> Selection = Sequencer->GetViewModel()->GetSelection();
+
+	FSelectionEventSuppressor SupressEvents = Selection->SuppressEvents();
+	Selection->Outliner.Empty();
+
 	for (AActor* Actor : Actors)
 	{
 		for (UActorComponent* Component : Actor->GetComponents())
 		{
 			if (Component->GetFName() == ComponentName)
 			{
-				Sequencer->GetHandleToObject(Component);
+				FGuid ObjectId = Sequencer->GetHandleToObject(Component);
+
+				TSharedPtr<FObjectBindingModel> Model = ObjectStorage->FindModelForObjectBinding(ObjectId);
+				if (Model)
+				{
+					Selection->Outliner.Select(Model);
+				}
 			}
 		}
 	}
