Sequencer: Allow dropping actors that haven't been possessed into folders in Sequencer

diff --git a/Engine/Source/Editor/Sequencer/Private/MVVM/ViewModels/FolderModel.cpp b/Engine/Source/Editor/Sequencer/Private/MVVM/ViewModels/FolderModel.cpp
--- a/Engine/Source/Editor/Sequencer/Private/MVVM/ViewModels/FolderModel.cpp
+++ b/Engine/Source/Editor/Sequencer/Private/MVVM/ViewModels/FolderModel.cpp
@@ -29,6 +29,7 @@
 #include "MovieSceneBinding.h"
 #include "MovieSceneFolder.h"
 
+#include "DragAndDrop/ActorDragDropOp.h"
 #include "ScopedTransaction.h"
 #include "Internationalization/Internationalization.h"
 
@@ -501,6 +502,12 @@ void FFolderModel::SetCustomOrder(int32 InCustomOrder)
 
 TOptional<EItemDropZone> FFolderModel::CanAcceptDrop(const FViewModelPtr& TargetModel, const FDragDropEvent& DragDropEvent, EItemDropZone ItemDropZone)
 {
+	TSharedPtr<FActorDragDropOp> ActorDragDropOp = DragDropEvent.GetOperationAs<FActorDragDropOp>();
+	if (ActorDragDropOp)
+	{
+		return ItemDropZone;
+	}
+
 	TSharedPtr<FSequencerOutlinerDragDropOp> DragDropOp = DragDropEvent.GetOperationAs<FSequencerOutlinerDragDropOp>();
 	if (!DragDropOp)
 	{
@@ -564,141 +571,179 @@ TOptional<EItemDropZone> FFolderModel::CanAcceptDrop(const FViewModelPtr& Target
 	return ItemDropZone;
 }
 
-void FFolderModel::PerformDrop(const FViewModelPtr& TargetModel, const FDragDropEvent& DragDropEvent, EItemDropZone InItemDropZone)
+void FFolderModel::PerformDropActors(const FViewModelPtr& TargetModel, TSharedPtr<FActorDragDropOp> ActorDragDropEvent, TSharedPtr<FViewModel> AttachAfter)
 {
-	FViewModelHierarchyOperation ScopedOperation(GetSharedData());
-
 	UMovieSceneFolder* Folder = WeakFolder.Get();
 	if (!Folder)
 	{
 		return;
 	}
+	
+	TSharedPtr<FSequencer> Sequencer = StaticCastSharedPtr<FSequencer>(GetEditor()->GetSequencer());
+	if (!Sequencer)
+	{
+		return;
+	}
 
-	const FScopedTransaction Transaction(FText::Format(LOCTEXT("MoveItems", "Move into {0}"), FText::FromName(Folder->GetFolderName())));
-
-	Folder->SetFlags(RF_Transactional);
-	Folder->Modify();
+	TArray<FGuid> PossessableGuids = FSequencerUtilities::AddActors(Sequencer.ToSharedRef(), ActorDragDropEvent->Actors);
 
-	TSharedPtr<FViewModel> AttachAfter;
-	if (TargetModel)
+	for (FGuid PossessableGuid : PossessableGuids)
 	{
-		if (InItemDropZone == EItemDropZone::AboveItem)
-		{
-			AttachAfter = TargetModel->GetPreviousSibling();
-		}
-		else if (InItemDropZone == EItemDropZone::BelowItem)
-		{
-			AttachAfter = TargetModel;
-		}
+		Folder->AddChildObjectBinding(PossessableGuid);
 	}
+}
 
-	// Open this folder if an item was dropped into the folder
-	if (InItemDropZone == EItemDropZone::OntoItem && !TargetModel)
+void FFolderModel::PerformDropOutliner(const FViewModelPtr& TargetModel, TSharedPtr<FSequencerOutlinerDragDropOp> OutlinerDragDropEvent, TSharedPtr<FViewModel> AttachAfter)
+{
+	UMovieSceneFolder* Folder = WeakFolder.Get();
+	if (!Folder)
 	{
-		SetExpansion(true);
+		return;
 	}
 
-	TSharedPtr<FSequencerOutlinerDragDropOp> DragDropOp = DragDropEvent.GetOperationAs<FSequencerOutlinerDragDropOp>();
-
 	// Drop handing for outliner drag/drop operations.
 	// Warning: this handler may be dragging nodes from a *different* sequence
-	if (DragDropOp)
+	FViewModelChildren OutlinerChildren = GetChildList(EViewModelListType::Outliner);
+
+	UMovieScene* MovieScene = Folder->GetTypedOuter<UMovieScene>();
+	if (MovieScene)
 	{
-		FViewModelChildren OutlinerChildren = GetChildList(EViewModelListType::Outliner);
+		MovieScene->SetFlags(RF_Transactional);
+		MovieScene->Modify();
+	}
 
-		UMovieScene* MovieScene = Folder->GetTypedOuter<UMovieScene>();
-		if (MovieScene)
+	for (const TWeakViewModelPtr<IOutlinerExtension>& WeakModel : OutlinerDragDropEvent->GetDraggedViewModels())
+	{
+		FViewModelPtr DraggedModel = WeakModel.Pin();
+		if (!DraggedModel || DraggedModel == AttachAfter)
 		{
-			MovieScene->SetFlags(RF_Transactional);
-			MovieScene->Modify();
+			continue;
 		}
 
-		for (const TWeakViewModelPtr<IOutlinerExtension>& WeakModel : DragDropOp->GetDraggedViewModels())
+		TViewModelPtr<FFolderModel> ExistingParentFolder = DraggedModel->CastParent<FFolderModel>();
+		UMovieSceneFolder* OldFolder = ExistingParentFolder ? ExistingParentFolder->GetFolder() : nullptr;
+
+		bool bSuccess = false;
+
+		// Handle dropoping a folder into another folder
+		// @todo: if we ever support folders within object bindings this will need to have better validation
+		if (TSharedPtr<FFolderModel> DraggedFolderModel = DraggedModel.ImplicitCast())
 		{
-			FViewModelPtr DraggedModel = WeakModel.Pin();
-			if (!DraggedModel || DraggedModel == AttachAfter)
+			if (UMovieSceneFolder* DraggedFolder = DraggedFolderModel->GetFolder())
 			{
-				continue;
-			}
-
-			TViewModelPtr<FFolderModel> ExistingParentFolder = DraggedModel->CastParent<FFolderModel>();
-			UMovieSceneFolder* OldFolder = ExistingParentFolder ? ExistingParentFolder->GetFolder() : nullptr;
+				if (OldFolder)
+				{
+					OldFolder->SetFlags(RF_Transactional);
+					OldFolder->Modify();
+					OldFolder->RemoveChildFolder(DraggedFolder);
+				}
+				else if (MovieScene)
+				{
+					MovieScene->RemoveRootFolder(DraggedFolder);
+				}
 
-			bool bSuccess = false;
+				// Give this folder a unique name inside its new parent if necessary
+				FName FolderName = Folder->MakeUniqueChildFolderName(DraggedFolder->GetFolderName());
 
-			// Handle dropoping a folder into another folder
-			// @todo: if we ever support folders within object bindings this will need to have better validation
-			if (TSharedPtr<FFolderModel> DraggedFolderModel = DraggedModel.ImplicitCast())
-			{
-				if (UMovieSceneFolder* DraggedFolder = DraggedFolderModel->GetFolder())
+				if (FolderName != DraggedFolder->GetFolderName())
 				{
-					if (OldFolder)
-					{
-						OldFolder->SetFlags(RF_Transactional);
-						OldFolder->Modify();
-						OldFolder->RemoveChildFolder(DraggedFolder);
-					}
-					else if (MovieScene)
-					{
-						MovieScene->RemoveRootFolder(DraggedFolder);
-					}
-
-					// Give this folder a unique name inside its new parent if necessary
-					FName FolderName = Folder->MakeUniqueChildFolderName(DraggedFolder->GetFolderName());
-
-					if (FolderName != DraggedFolder->GetFolderName())
-					{
-						DraggedFolder->SetFlags(RF_Transactional);
-						DraggedFolder->Modify();
-						DraggedFolder->SetFolderName(FolderName);
-					}
-
-					Folder->AddChildFolder(DraggedFolder);
-					bSuccess = true;
+					DraggedFolder->SetFlags(RF_Transactional);
+					DraggedFolder->Modify();
+					DraggedFolder->SetFolderName(FolderName);
 				}
+
+				Folder->AddChildFolder(DraggedFolder);
+				bSuccess = true;
 			}
-			else if (TSharedPtr<IObjectBindingExtension> ObjectBinding = DraggedModel.ImplicitCast())
+		}
+		else if (TSharedPtr<IObjectBindingExtension> ObjectBinding = DraggedModel.ImplicitCast())
+		{
+			TViewModelPtr<IObjectBindingExtension> ParentObject = DraggedModel->CastParent<IObjectBindingExtension>();
+			// Don't allow dropping an object binding if it has an object parent
+			if (ParentObject == nullptr)
 			{
-				TViewModelPtr<IObjectBindingExtension> ParentObject = DraggedModel->CastParent<IObjectBindingExtension>();
-				// Don't allow dropping an object binding if it has an object parent
-				if (ParentObject == nullptr)
+				if (OldFolder)
 				{
-					if (OldFolder)
-					{
-						OldFolder->SetFlags(RF_Transactional);
-						OldFolder->Modify();
-						OldFolder->RemoveChildObjectBinding(ObjectBinding->GetObjectGuid());
-					}
-
-					Folder->AddChildObjectBinding(ObjectBinding->GetObjectGuid());
-					bSuccess = true;
+					OldFolder->SetFlags(RF_Transactional);
+					OldFolder->Modify();
+					OldFolder->RemoveChildObjectBinding(ObjectBinding->GetObjectGuid());
 				}
+
+				Folder->AddChildObjectBinding(ObjectBinding->GetObjectGuid());
+				bSuccess = true;
 			}
-			else if (TSharedPtr<FTrackModel> TrackModel = DraggedModel.ImplicitCast())
+		}
+		else if (TSharedPtr<FTrackModel> TrackModel = DraggedModel.ImplicitCast())
+		{
+			UMovieSceneTrack* Track = TrackModel->GetTrack();
+			TViewModelPtr<IObjectBindingExtension> ParentObject = DraggedModel->CastParent<IObjectBindingExtension>();
+			// Don't allow dropping a track if it has an object parent
+			if (ParentObject == nullptr && Track != nullptr)
 			{
-				UMovieSceneTrack* Track = TrackModel->GetTrack();
-				TViewModelPtr<IObjectBindingExtension> ParentObject = DraggedModel->CastParent<IObjectBindingExtension>();
-				// Don't allow dropping a track if it has an object parent
-				if (ParentObject == nullptr && Track != nullptr)
+				if (OldFolder)
 				{
-					if (OldFolder)
-					{
-						OldFolder->SetFlags(RF_Transactional);
-						OldFolder->Modify();
-						OldFolder->RemoveChildTrack(Track);
-					}
-
-					Folder->AddChildTrack(Track);
-					bSuccess = true;
+					OldFolder->SetFlags(RF_Transactional);
+					OldFolder->Modify();
+					OldFolder->RemoveChildTrack(Track);
 				}
-			}
 
-			if (bSuccess)
-			{
-				// Attach it to the right node, possibly setting its parent too
-				OutlinerChildren.InsertChild(DraggedModel, AttachAfter);
+				Folder->AddChildTrack(Track);
+				bSuccess = true;
 			}
 		}
+
+		if (bSuccess)
+		{
+			// Attach it to the right node, possibly setting its parent too
+			OutlinerChildren.InsertChild(DraggedModel, AttachAfter);
+		}
+	}
+}
+
+void FFolderModel::PerformDrop(const FViewModelPtr& TargetModel, const FDragDropEvent& DragDropEvent, EItemDropZone InItemDropZone)
+{
+	FViewModelHierarchyOperation ScopedOperation(GetSharedData());
+
+	UMovieSceneFolder* Folder = WeakFolder.Get();
+	if (!Folder)
+	{
+		return;
+	}
+
+	const FScopedTransaction Transaction(FText::Format(LOCTEXT("MoveItems", "Move into {0}"), FText::FromName(Folder->GetFolderName())));
+
+	Folder->SetFlags(RF_Transactional);
+	Folder->Modify();
+
+	TSharedPtr<FViewModel> AttachAfter;
+	if (TargetModel)
+	{
+		if (InItemDropZone == EItemDropZone::AboveItem)
+		{
+			AttachAfter = TargetModel->GetPreviousSibling();
+		}
+		else if (InItemDropZone == EItemDropZone::BelowItem)
+		{
+			AttachAfter = TargetModel;
+		}
+	}
+
+	// Open this folder if an item was dropped into the folder
+	if (InItemDropZone == EItemDropZone::OntoItem && !TargetModel)
+	{
+		SetExpansion(true);
+	}
+
+	TSharedPtr<FActorDragDropOp> ActorDragDropOp = DragDropEvent.GetOperationAs<FActorDragDropOp>();
+	if (ActorDragDropOp)
+	{
+		PerformDropActors(TargetModel, ActorDragDropOp, AttachAfter);
+	}
+
+	TSharedPtr<FSequencerOutlinerDragDropOp> OutlinerDragDropOp = DragDropEvent.GetOperationAs<FSequencerOutlinerDragDropOp>();
+	if (OutlinerDragDropOp)
+	{
+		PerformDropOutliner(TargetModel, OutlinerDragDropOp, AttachAfter);
 	}
 
 	// Forcibly update sorting
diff --git a/Engine/Source/Editor/Sequencer/Public/MVVM/ViewModels/FolderModel.h b/Engine/Source/Editor/Sequencer/Public/MVVM/ViewModels/FolderModel.h
--- a/Engine/Source/Editor/Sequencer/Public/MVVM/ViewModels/FolderModel.h
+++ b/Engine/Source/Editor/Sequencer/Public/MVVM/ViewModels/FolderModel.h
@@ -18,6 +18,7 @@
 #include "UObject/WeakObjectPtrTemplates.h"
 
 class UMovieSceneFolder;
+class FActorDragDropOp;
 
 namespace UE
 {
@@ -26,6 +27,7 @@ namespace Sequencer
 
 class FLayerBarModel;
 class FSequenceModel;
+class FSequencerOutlinerDragDropOp;
 
 class FFolderModel
 	: public FMuteSoloOutlinerItemModel
@@ -118,6 +120,9 @@ private:
 
 	void RepopulateChildren();
 
+	void PerformDropActors(const FViewModelPtr& TargetModel, TSharedPtr<FActorDragDropOp> ActorDragDropEvent, TSharedPtr<FViewModel> AttachAfter);
+	void PerformDropOutliner(const FViewModelPtr& TargetModel, TSharedPtr<FSequencerOutlinerDragDropOp> OutlinerDragDropEvent, TSharedPtr<FViewModel> AttachAfter);
+
 private:
 
 	FViewModelListHead TrackAreaList;
