Curve Editor: Additional Tree View UX - Focus the time in the curve editor

diff --git a/Engine/Source/Editor/Sequencer/Private/MVVM/CurveEditorExtension.cpp b/Engine/Source/Editor/Sequencer/Private/MVVM/CurveEditorExtension.cpp
--- a/Engine/Source/Editor/Sequencer/Private/MVVM/CurveEditorExtension.cpp
+++ b/Engine/Source/Editor/Sequencer/Private/MVVM/CurveEditorExtension.cpp
@@ -13,6 +13,7 @@
 #include "SCurveEditorToolProperties.h"
 #include "SCurveKeyDetailPanel.h"
 #include "SSequencerTreeFilterStatusBar.h"
+#include "STemporarilyFocusedSpinBox.h"
 #include "Sequencer.h"
 #include "SequencerCommands.h"
 #include "Toolkits/IToolkitHost.h"
@@ -203,6 +204,8 @@ void FCurveEditorExtension::CreateCurveEditor(const FTimeSliderArgs& TimeSliderA
 	// the Curve Editor time slider controller. We want everything else to just pass through though.
 	TSharedRef<ITimeSliderController> CurveEditorTimeSliderController = MakeShared<FSequencerCurveEditorTimeSliderController>(
 			TimeSliderArgs, Sequencer, CurveEditorModel.ToSharedRef());
+	
+	PlayTimeDisplay = StaticCastSharedRef<STemporarilyFocusedSpinBox<double>>(Sequencer->MakePlayTimeDisplay(TimeSliderArgs.NumericTypeInterface.ToSharedRef()));
 
 	CurveEditorTreeView = SNew(SCurveEditorTree, CurveEditorModel);
 	CurveEditorPanel = SNew(SCurveEditorPanel, CurveEditorModel.ToSharedRef())
@@ -304,7 +307,7 @@ void FCurveEditorExtension::CreateCurveEditor(const FTimeSliderArgs& TimeSliderA
 							SNew(SBorder)
 							.BorderImage(nullptr)
 							[
-								Sequencer->MakePlayTimeDisplay(TimeSliderArgs.NumericTypeInterface.ToSharedRef())
+								PlayTimeDisplay.ToSharedRef()
 							]
 						]
 					]
@@ -327,6 +330,11 @@ void FCurveEditorExtension::CreateCurveEditor(const FTimeSliderArgs& TimeSliderA
 		FExecuteAction::CreateLambda([this] { FSlateApplication::Get().SetKeyboardFocus(CurveEditorSearchBox, EFocusCause::SetDirectly); })
 	);
 
+	CurveEditorModel->GetCommands()->MapAction(
+		FSequencerCommands::Get().ToggleShowGotoBox,
+		FExecuteAction::CreateLambda([this] { PlayTimeDisplay->Setup(); FSlateApplication::Get().SetKeyboardFocus(PlayTimeDisplay, EFocusCause::SetDirectly); })
+	);
+
 	CurveEditorWidget = SNew(SSequencerCurveEditor, CurveEditorPanel.ToSharedRef(), Sequencer);
 
 	CurveEditorPanel->OnFilterClassChanged.BindRaw(this, &FCurveEditorExtension::FilterClassChanged);
diff --git a/Engine/Source/Editor/Sequencer/Private/SSequencer.cpp b/Engine/Source/Editor/Sequencer/Private/SSequencer.cpp
--- a/Engine/Source/Editor/Sequencer/Private/SSequencer.cpp
+++ b/Engine/Source/Editor/Sequencer/Private/SSequencer.cpp
@@ -61,6 +61,7 @@
 #include "SequencerTimeSliderController.h"
 #include "SequencerToolMenuContext.h"
 #include "SSequencerSectionOverlay.h"
+#include "STemporarilyFocusedSpinBox.h"
 #include "MVVM/Views/STrackAreaView.h"
 #include "MVVM/Views/SSequencerTrackAreaView.h"
 #include "DragAndDrop/AssetDragDropOp.h"
diff --git a/Engine/Source/Editor/Sequencer/Private/SSequencer.h b/Engine/Source/Editor/Sequencer/Private/SSequencer.h
--- a/Engine/Source/Editor/Sequencer/Private/SSequencer.h
+++ b/Engine/Source/Editor/Sequencer/Private/SSequencer.h
@@ -22,6 +22,7 @@
 #include "Widgets/Text/STextBlock.h"
 #include "Sequencer.h"
 #include "SequencerWidgetsDelegates.h"
+#include "STemporarilyFocusedSpinBox.h"
 
 class FActorDragDropOp;
 class FFolderDragDropOp;
@@ -130,30 +131,6 @@ struct FSequencerOutlinerColumnVisibility
 };
 
 
-/**
- * A widget that holds a widget that is to be refocused on completion
- */
-template<typename T>
-struct STemporarilyFocusedSpinBox : SSpinBox<T>
-{
-public:
-	void Setup()
-	{
-		PreviousFocusedWidget = FSlateApplication::Get().GetKeyboardFocusedWidget();
-	}
-
-	void Refocus()
-	{
-		if (PreviousFocusedWidget.IsValid())
-		{
-			FSlateApplication::Get().SetKeyboardFocus(PreviousFocusedWidget.Pin());
-		}
-	}
-
-private:
-	TWeakPtr<SWidget> PreviousFocusedWidget;
-};
-
 /**
  * Main sequencer UI widget
  */
diff --git a/Engine/Source/Editor/Sequencer/Public/MVVM/CurveEditorExtension.h b/Engine/Source/Editor/Sequencer/Public/MVVM/CurveEditorExtension.h
--- a/Engine/Source/Editor/Sequencer/Public/MVVM/CurveEditorExtension.h
+++ b/Engine/Source/Editor/Sequencer/Public/MVVM/CurveEditorExtension.h
@@ -9,6 +9,7 @@
 #include "CurveDataAbstraction.h"
 #include "CurveEditorTypes.h"
 #include "Tree/ICurveEditorTreeItem.h"
+#include "STemporarilyFocusedSpinBox.h"
 
 class ISequencer;
 class IPropertyTypeCustomization;
@@ -84,6 +85,8 @@ private:
 	TSharedPtr<SCurveEditorPanel> CurveEditorPanel;
 	/** Filter Status Bar */
 	TSharedPtr<SCurveEditorTreeFilterStatusBar> CurveEditorTreeFilterStatusBar;
+	/** The current playback time display. */
+	TSharedPtr<STemporarilyFocusedSpinBox<double>> PlayTimeDisplay;
 
 	friend class FCurveEditorIntegrationExtension;
 };
diff --git a/Engine/Source/Editor/SequencerWidgets/Public/STemporarilyFocusedSpinBox.h b/Engine/Source/Editor/SequencerWidgets/Public/STemporarilyFocusedSpinBox.h
new file mode 100644
--- /dev/null
+++ b/Engine/Source/Editor/SequencerWidgets/Public/STemporarilyFocusedSpinBox.h
@@ -0,0 +1,31 @@
+// Copyright Epic Games, Inc. All Rights Reserved.
+
+#pragma once
+
+#include "Widgets/SWidget.h"
+#include "Widgets/Input/SSpinBox.h"
+#include "Framework/Application/SlateApplication.h"
+
+/**
+ * A widget that holds a widget that is to be refocused on completion
+ */
+template<typename T>
+struct STemporarilyFocusedSpinBox : SSpinBox<T>
+{
+public:
+	void Setup()
+	{
+		PreviousFocusedWidget = FSlateApplication::Get().GetKeyboardFocusedWidget();
+	}
+
+	void Refocus()
+	{
+		if (PreviousFocusedWidget.IsValid())
+		{
+			FSlateApplication::Get().SetKeyboardFocus(PreviousFocusedWidget.Pin());
+		}
+	}
+
+private:
+	TWeakPtr<SWidget> PreviousFocusedWidget;
+};
\ No newline at end of file
