Sequencer: UX changes for Playback Controls - Prevent throttling

diff --git a/Engine/Source/Editor/Sequencer/Private/Sequencer.cpp b/Engine/Source/Editor/Sequencer/Private/Sequencer.cpp
--- a/Engine/Source/Editor/Sequencer/Private/Sequencer.cpp
+++ b/Engine/Source/Editor/Sequencer/Private/Sequencer.cpp
@@ -4208,6 +4208,7 @@ TSharedRef<SWidget> FSequencer::MakePlayTimeDisplay(const TSharedRef<INumericTyp
 {
 	return SNew(STemporarilyFocusedSpinBox<double>)
 		.Style(&FAppStyle::GetWidgetStyle<FSpinBoxStyle>("Sequencer.PlayTimeSpinBox"))
+		.PreventThrottling(true)
 		.Value_Lambda([this]() -> double {
 			return GetLocalTime().Time.GetFrame().Value;
 		})
