Sequencer: Remove SSequencer::StepToNextKey/StepToPreviousKey, StepToNextCameraKey, StepToPreviousCameraKey and reuse the existing functions in FSequencer::JumpToNextKey, JumpToPreviousKey

#rnx

diff --git a/Engine/Source/Editor/Sequencer/Private/SSequencer.cpp b/Engine/Source/Editor/Sequencer/Private/SSequencer.cpp
--- a/Engine/Source/Editor/Sequencer/Private/SSequencer.cpp
+++ b/Engine/Source/Editor/Sequencer/Private/SSequencer.cpp
@@ -3463,174 +3463,6 @@ void SSequencer::OnSaveMovieSceneClicked()
 	SequencerPtr.Pin()->SaveCurrentMovieScene();
 }
 
-
-void SSequencer::StepToNextKey()
-{
-	StepToKey(true, false);
-}
-
-
-void SSequencer::StepToPreviousKey()
-{
-	StepToKey(false, false);
-}
-
-
-void SSequencer::StepToNextCameraKey()
-{
-	StepToKey(true, true);
-}
-
-
-void SSequencer::StepToPreviousCameraKey()
-{
-	StepToKey(false, true);
-}
-
-
-void SSequencer::StepToKey(bool bStepToNextKey, bool bCameraOnly)
-{
-	using namespace UE::Sequencer;
-
-	TSharedPtr<FSequencer> Sequencer = SequencerPtr.Pin();
-	if ( Sequencer.IsValid() )
-	{
-		TSet<TViewModelPtr<IOutlinerExtension>> Nodes;
-
-		if ( bCameraOnly )
-		{
-			TSet<TViewModelPtr<IOutlinerExtension>> RootNodes(Sequencer->GetNodeTree()->GetRootNodes());
-
-			TSet<TWeakObjectPtr<AActor> > LockedActors;
-			for (FLevelEditorViewportClient* LevelVC : GEditor->GetLevelViewportClients())
-			{
-				if ( LevelVC && LevelVC->IsPerspective() && LevelVC->GetViewMode() != VMI_Unknown )
-				{
-					TWeakObjectPtr<AActor> ActorLock = LevelVC->GetActiveActorLock();
-					if ( ActorLock.IsValid() )
-					{
-						LockedActors.Add( ActorLock );
-					}
-				}
-			}
-
-			for (TViewModelPtr<IOutlinerExtension> RootNode : RootNodes)
-			{
-				TSharedPtr<FObjectBindingModel> ObjectBindingNode = RootNode.ImplicitCast();
-				if (ObjectBindingNode)
-				{
-					for (TWeakObjectPtr<>& Object : Sequencer->FindObjectsInCurrentSequence(ObjectBindingNode->GetObjectGuid()))
-					{
-						AActor* RuntimeActor = Cast<AActor>( Object.Get() );
-						if ( RuntimeActor != nullptr && LockedActors.Contains( RuntimeActor ) )
-						{
-							Nodes.Add(RootNode);
-						}
-					}
-				}
-			}
-		}
-		else
-		{
-			Nodes.Reserve(Sequencer->GetViewModel()->GetSelection()->Outliner.Num());
-			for (TViewModelPtr<IOutlinerExtension> Model : Sequencer->GetViewModel()->GetSelection()->Outliner.Filter<IOutlinerExtension>())
-			{
-				Nodes.Add(Model);
-			}
-			if ( Nodes.Num() == 0 )
-			{
-				TSet<TViewModelPtr<IOutlinerExtension>> RootNodes(Sequencer->GetNodeTree()->GetRootNodes());
-				for (TViewModelPtr<IOutlinerExtension> RootNode : RootNodes)
-				{
-					Nodes.Add(RootNode);
-					for (TViewModelPtr<IOutlinerExtension> Child : RootNode.AsModel()->GetDescendantsOfType<IOutlinerExtension>())
-					{
-						Nodes.Add(Child);
-					}
-				}
-			}
-		}
-
-		if ( Nodes.Num() > 0 )
-		{
-			FFrameTime ClosestKeyDistance = FFrameTime(TNumericLimits<int32>::Max(), 0.99999f);
-			FFrameTime CurrentTime = Sequencer->GetLocalTime().Time;
-			TOptional<FFrameTime> NextTime;
-
-			TOptional<FFrameNumber> StepToTime;
-
-			auto It = Nodes.CreateConstIterator();
-			bool bExpand = !( *It )->IsExpanded();
-
-			for (TViewModelPtr<IOutlinerExtension> Node : Nodes)
-			{
-				TArray<FFrameNumber> AllTimes;
-
-				TSet<TSharedPtr<IKeyArea>> KeyAreas;
-				SequencerHelpers::GetAllKeyAreas( Node.AsModel(), KeyAreas );
-				for ( TSharedPtr<IKeyArea> KeyArea : KeyAreas )
-				{
-					if (KeyArea->GetOwningSection())
-					{
-						KeyArea->GetKeyTimes(AllTimes, KeyArea->GetOwningSection()->GetRange());
-					}
-				}
-
-				TSet<TWeakObjectPtr<UMovieSceneSection> > Sections;
-				SequencerHelpers::GetAllSections( Node.AsModel(), Sections );
-
-				for ( TWeakObjectPtr<UMovieSceneSection> Section : Sections )
-				{
-					if (Section.IsValid())
-					{
-						Section->GetSnapTimes(AllTimes, true);
-					}
-				}
-
-				for (FFrameNumber Time : AllTimes)
-				{
-					if ( bStepToNextKey )
-					{
-						if ( Time > CurrentTime && Time - CurrentTime < ClosestKeyDistance )
-						{
-							StepToTime = Time;
-							ClosestKeyDistance = Time - CurrentTime;
-						}
-					}
-					else
-					{
-						if ( Time < CurrentTime && CurrentTime - Time < ClosestKeyDistance )
-						{
-							StepToTime = Time;
-							ClosestKeyDistance = CurrentTime - Time;
-						}
-					}
-				}
-
-				if (!StepToTime.IsSet() && AllTimes.Num() > 0)
-				{
-					AllTimes.Sort();
-
-					if (bStepToNextKey)
-					{
-						StepToTime = AllTimes[0];
-					}
-					else
-					{
-						StepToTime = AllTimes.Last();
-					}
-				}
-			}
-
-			if ( StepToTime.IsSet() )
-			{
-				Sequencer->SetLocalTime( StepToTime.GetValue() );
-			}
-		}
-	}
-}
-
-
 float SSequencer::GetPinnedAreaMaxHeight() const
 {
 	if (!MainSequencerArea.IsValid())
diff --git a/Engine/Source/Editor/Sequencer/Private/SSequencer.h b/Engine/Source/Editor/Sequencer/Private/SSequencer.h
--- a/Engine/Source/Editor/Sequencer/Private/SSequencer.h
+++ b/Engine/Source/Editor/Sequencer/Private/SSequencer.h
@@ -322,13 +322,6 @@ public:
 	void ResetBreadcrumbs();
 	void PopBreadcrumb();
 
-	/** Step to next and previous keyframes */
-	void StepToNextKey();
-	void StepToPreviousKey();
-	void StepToNextCameraKey();
-	void StepToPreviousCameraKey();
-	void StepToKey(bool bStepToNextKey, bool bCameraOnly);
-
 	/** Called when the save button is clicked */
 	void OnSaveMovieSceneClicked();
 
diff --git a/Engine/Source/Editor/Sequencer/Private/Sequencer.cpp b/Engine/Source/Editor/Sequencer/Private/Sequencer.cpp
--- a/Engine/Source/Editor/Sequencer/Private/Sequencer.cpp
+++ b/Engine/Source/Editor/Sequencer/Private/Sequencer.cpp
@@ -9085,31 +9085,6 @@ void FSequencer::JumpBackward()
 	OnStepBackward(Settings->GetJumpFrameIncrement());
 }
 
-
-void FSequencer::StepToNextKey()
-{
-	SequencerWidget->StepToNextKey();
-}
-
-
-void FSequencer::StepToPreviousKey()
-{
-	SequencerWidget->StepToPreviousKey();
-}
-
-
-void FSequencer::StepToNextCameraKey()
-{
-	SequencerWidget->StepToNextCameraKey();
-}
-
-
-void FSequencer::StepToPreviousCameraKey()
-{
-	SequencerWidget->StepToPreviousCameraKey();
-}
-
-
 void FSequencer::StepToNextShot()
 {
 	if (ActiveTemplateIDs.Num() < 2)
@@ -10361,11 +10336,11 @@ void FSequencer::BindCommands()
 
 	SequencerCommandBindings->MapAction(
 		Commands.StepToNextKey,
-		FExecuteAction::CreateSP( this, &FSequencer::StepToNextKey) );
+		FExecuteAction::CreateLambda([this] { JumpToNextKey(); }));
 
 	SequencerCommandBindings->MapAction(
 		Commands.StepToPreviousKey,
-		FExecuteAction::CreateSP( this, &FSequencer::StepToPreviousKey ) );
+		FExecuteAction::CreateLambda([this] { JumpToPreviousKey(); }));
 
 	SequencerCommandBindings->MapAction(
 		Commands.StepForwardViewport,
@@ -10377,14 +10352,6 @@ void FSequencer::BindCommands()
 		FExecuteAction::CreateSP(this, &FSequencer::StepBackward),
 		EUIActionRepeatMode::RepeatEnabled);
 
-	SequencerCommandBindings->MapAction(
-		Commands.StepToNextCameraKey,
-		FExecuteAction::CreateSP( this, &FSequencer::StepToNextCameraKey ) );
-
-	SequencerCommandBindings->MapAction(
-		Commands.StepToPreviousCameraKey,
-		FExecuteAction::CreateSP( this, &FSequencer::StepToPreviousCameraKey ) );
-
 	SequencerCommandBindings->MapAction(
 		Commands.SortAllNodesAndDescendants,
 		FExecuteAction::CreateSP(this, &FSequencer::SortAllNodesAndDescendants));
diff --git a/Engine/Source/Editor/Sequencer/Private/Sequencer.h b/Engine/Source/Editor/Sequencer/Private/Sequencer.h
--- a/Engine/Source/Editor/Sequencer/Private/Sequencer.h
+++ b/Engine/Source/Editor/Sequencer/Private/Sequencer.h
@@ -1005,10 +1005,6 @@ protected:
 	void StepBackward();
 	void JumpForward();
 	void JumpBackward();
-	void StepToNextKey();
-	void StepToPreviousKey();
-	void StepToNextCameraKey();
-	void StepToPreviousCameraKey();
 	void StepToNextShot();
 	void StepToPreviousShot();
 
diff --git a/Engine/Source/Editor/Sequencer/Private/SequencerCommands.cpp b/Engine/Source/Editor/Sequencer/Private/SequencerCommands.cpp
--- a/Engine/Source/Editor/Sequencer/Private/SequencerCommands.cpp
+++ b/Engine/Source/Editor/Sequencer/Private/SequencerCommands.cpp
@@ -31,8 +31,6 @@ void FSequencerCommands::RegisterCommands()
 	UI_COMMAND( JumpBackward, "Jump Backward", "Jump the timeline backward a user defined number of frames/times", EUserInterfaceActionType::Button, FInputChord(EModifierKey::Shift, EKeys::Left) );
 	UI_COMMAND( StepToNextKey, "Step to Next Key (Viewport)", "Step to the next key in all viewports and sequencer", EUserInterfaceActionType::Button, FInputChord(EKeys::Period) );
 	UI_COMMAND( StepToPreviousKey, "Step to Previous Key (Viewport)", "Step to the previous key in all viewports and sequencer", EUserInterfaceActionType::Button, FInputChord(EKeys::Comma) );
-	UI_COMMAND( StepToNextCameraKey, "Step to Next Camera Key", "Step to the next camera key", EUserInterfaceActionType::Button, FInputChord() );
-	UI_COMMAND( StepToPreviousCameraKey, "Step to Previous Camera Key", "Step to the previous camera key", EUserInterfaceActionType::Button, FInputChord() );
 	UI_COMMAND( StepToNextShot, "Step to Next Shot", "Step to the next shot", EUserInterfaceActionType::Button, FInputChord(EModifierKey::Shift, EKeys::Period) );
 	UI_COMMAND( StepToPreviousShot, "Step to Previous Shot", "Step to the previous shot", EUserInterfaceActionType::Button, FInputChord(EModifierKey::Shift, EKeys::Comma) );
 	UI_COMMAND( SetStartPlaybackRange, "Set Start Playback Range", "Set the start playback range", EUserInterfaceActionType::Button, FInputChord(EKeys::LeftBracket) );
diff --git a/Engine/Source/Editor/Sequencer/Public/SequencerCommands.h b/Engine/Source/Editor/Sequencer/Public/SequencerCommands.h
--- a/Engine/Source/Editor/Sequencer/Public/SequencerCommands.h
+++ b/Engine/Source/Editor/Sequencer/Public/SequencerCommands.h
@@ -85,12 +85,6 @@ public:
 	/** Step to previous key */
 	TSharedPtr< FUICommandInfo > StepToPreviousKey;
 
-	/** Step to next camera key */
-	TSharedPtr< FUICommandInfo > StepToNextCameraKey;
-
-	/** Step to previous camera key */
-	TSharedPtr< FUICommandInfo > StepToPreviousCameraKey;
-
 	/** Step to next shot */
 	TSharedPtr< FUICommandInfo > StepToNextShot;
 
