Sequencer: Added back missing drop shadow on text in sections

diff --git a/Engine/Source/Editor/Sequencer/Private/SSequencerSection.cpp b/Engine/Source/Editor/Sequencer/Private/SSequencerSection.cpp
--- a/Engine/Source/Editor/Sequencer/Private/SSequencerSection.cpp
+++ b/Engine/Source/Editor/Sequencer/Private/SSequencerSection.cpp
@@ -1482,6 +1482,20 @@ int32 SSequencerSection::OnPaint( const FPaintArgs& Args, const FGeometry& Allot
 			FontInfo.Size = FMath::Max(FMath::FloorToInt(FontInfo.Size - 6.f), 11);
 		}
 
+		// Drop shadow
+		FSlateDrawElement::MakeText(
+			OutDrawElements,
+			LayerId,
+			SectionGeometry.MakeChild(
+				FVector2D(SectionGeometry.Size.X, GetFontHeight()),
+				FSlateLayoutTransform(TopLeft + FVector2D(ContentPadding.Left, ContentPadding.Top) + FVector2D(1.f, 1.f))
+			).ToPaintGeometry(),
+			SectionTitle,
+			FontInfo,
+			DrawEffects,
+			FLinearColor(0, 0, 0, .5f * Painter.GhostAlpha)
+		);
+
 		FSlateDrawElement::MakeText(
 			OutDrawElements,
 			LayerId,
