Sequencer- Fix static analysis nullptr check warning.

diff --git a/Engine/Source/Editor/Sequencer/Private/MVVM/ViewModels/ObjectBindingModel.cpp b/Engine/Source/Editor/Sequencer/Private/MVVM/ViewModels/ObjectBindingModel.cpp
--- a/Engine/Source/Editor/Sequencer/Private/MVVM/ViewModels/ObjectBindingModel.cpp
+++ b/Engine/Source/Editor/Sequencer/Private/MVVM/ViewModels/ObjectBindingModel.cpp
@@ -1237,41 +1237,41 @@ bool FObjectBindingModel::CanDelete(FText* OutErrorMessage) const
 
 void FObjectBindingModel::Delete()
 {
-	TSharedPtr<ISequencer> Sequencer = OwnerModel->GetSequencer();
-	UMovieScene* MovieScene = Sequencer->GetRootMovieSceneSequence()->GetMovieScene();
+	if (OwnerModel)
+	{
+		TSharedPtr<ISequencer> Sequencer = OwnerModel->GetSequencer();
+		UMovieScene* MovieScene = Sequencer->GetRootMovieSceneSequence()->GetMovieScene();
 
-	MovieScene->Modify();
+		MovieScene->Modify();
 
-	// Untag this binding
-	UE::MovieScene::FFixedObjectBindingID BindingID(ObjectBindingID, OwnerModel->GetSequenceID());
-	for (auto It = OwnerModel->GetSequencerImpl()->GetObjectBindingTagCache()->IterateTags(BindingID); It; ++It)
-	{
-		MovieScene->UntagBinding(It.Value(), BindingID);
-	}
+		// Untag this binding
+		UE::MovieScene::FFixedObjectBindingID BindingID(ObjectBindingID, OwnerModel->GetSequenceID());
+		for (auto It = OwnerModel->GetSequencerImpl()->GetObjectBindingTagCache()->IterateTags(BindingID); It; ++It)
+		{
+			MovieScene->UntagBinding(It.Value(), BindingID);
+		}
 
-	// Delete any child object bindings - this will remove their tracks implicitly
-	// so no need to delete those manually
-	for (const TViewModelPtr<FObjectBindingModel>& ChildObject : GetChildrenOfType<FObjectBindingModel>(EViewModelListType::Outliner).ToArray())
-	{
-		ChildObject->Delete();
-	}
+		// Delete any child object bindings - this will remove their tracks implicitly
+		// so no need to delete those manually
+		for (const TViewModelPtr<FObjectBindingModel>& ChildObject : GetChildrenOfType<FObjectBindingModel>(EViewModelListType::Outliner).ToArray())
+		{
+			ChildObject->Delete();
+		}
 
-	// Remove from a parent folder if necessary.
-	if (TViewModelPtr<FFolderModel> ParentFolder = CastParent<FFolderModel>())
-	{
-		ParentFolder->GetFolder()->RemoveChildObjectBinding(ObjectBindingID);
-	}
+		// Remove from a parent folder if necessary.
+		if (TViewModelPtr<FFolderModel> ParentFolder = CastParent<FFolderModel>())
+		{
+			ParentFolder->GetFolder()->RemoveChildObjectBinding(ObjectBindingID);
+		}
 
-	if (OwnerModel)
-	{
 		// Delete any loaded object that may be bound to this object binding
 		if (FMovieSceneObjectCache* Cache = Sequencer->State.FindObjectCache(OwnerModel->GetSequenceID()))
 		{
 			Cache->UnloadBinding(ObjectBindingID, Sequencer->GetSharedPlaybackState());
 		}
-	}
 
-	BindingLifetimeOverlayModel.Reset();
+		BindingLifetimeOverlayModel.Reset();
+	}
 }
 
 } // namespace Sequencer
