Sequencer- Fix issue introduced by previous changelist where skeletal mesh/static mesh assets could no longer be dragged into Sequencer- actor factory was broken.

diff --git a/Engine/Source/Editor/Sequencer/Private/SSequencer.cpp b/Engine/Source/Editor/Sequencer/Private/SSequencer.cpp
--- a/Engine/Source/Editor/Sequencer/Private/SSequencer.cpp
+++ b/Engine/Source/Editor/Sequencer/Private/SSequencer.cpp
@@ -3302,17 +3302,13 @@ void SSequencer::OnAssetsDropped( const FAssetDragDropOp& DragDropOp )
 					continue;
 				}
 
-				// Make a test locator from the object, and if it's not an actor type, make a possessable, otherwise make a spawnable by default.
-				// TODO: We might want to generalize this somehow from the locator types.
+				// If the object dragged resolves to a locator of actor type, or is natively a non-supported object type, attempt to use the actor factory to make a spawnable actor from it.
+				// Otherwise, the CreateBinding will have succeeded.
 				FUniversalObjectLocator TestLocator(CurObject, SequencerRef.GetPlaybackContext());
-				if (TestLocator.GetLastFragmentTypeHandle() == FActorLocatorFragment::FragmentType)
+				if (TestLocator.GetLastFragmentTypeHandle() == FActorLocatorFragment::FragmentType || !SequencerRef.CreateBinding(*CurObject, CurObject->GetName()).IsValid())
 				{
 					SequencerRef.MakeNewSpawnable(*CurObject, DragDropOp.GetActorFactory());
 				}
-				else
-				{
-					SequencerRef.CreateBinding(*CurObject, CurObject->GetName());
-				}
 			}
 
 			DropResult = ESequencerDropResult::DropHandled;
