Sequencer: Remove copy relative option because it's a bit of a hidden feature and causes confusion when copying keys with the mouse hovered over keys at different times

diff --git a/Engine/Source/Editor/Sequencer/Private/Sequencer.cpp b/Engine/Source/Editor/Sequencer/Private/Sequencer.cpp
--- a/Engine/Source/Editor/Sequencer/Private/Sequencer.cpp
+++ b/Engine/Source/Editor/Sequencer/Private/Sequencer.cpp
@@ -10142,14 +10142,6 @@ void FSequencer::CopySelectedKeys()
 {
 	using namespace UE::Sequencer;
 
-	TOptional<FFrameNumber> CopyRelativeTo;
-
-	// Copy relative to the current key hotspot, if applicable
-	if (TSharedPtr<FKeyHotspot> KeyHotspot = HotspotCast<FKeyHotspot>(ViewModel->GetHotspot()))
-	{
-		CopyRelativeTo = KeyHotspot->GetTime();
-	}
-
 	FMovieSceneClipboardBuilder Builder;
 
 	// Map selected keys to their key areas
@@ -10169,7 +10161,7 @@ void FSequencer::CopySelectedKeys()
 		Pair.Key->GetKeyArea()->CopyKeys(Builder, Pair.Value);
 	}
 
-	TSharedRef<FMovieSceneClipboard> Clipboard = MakeShareable( new FMovieSceneClipboard(Builder.Commit(CopyRelativeTo)) );
+	TSharedRef<FMovieSceneClipboard> Clipboard = MakeShareable( new FMovieSceneClipboard(Builder.Commit(TOptional<FFrameNumber>())) );
 	
 	Clipboard->GetEnvironment().TickResolution = GetFocusedTickResolution();
 
