Sequencer: PopToSequenceInstance in Tick() before any other updates

diff --git a/Engine/Source/Editor/Sequencer/Private/Sequencer.cpp b/Engine/Source/Editor/Sequencer/Private/Sequencer.cpp
--- a/Engine/Source/Editor/Sequencer/Private/Sequencer.cpp
+++ b/Engine/Source/Editor/Sequencer/Private/Sequencer.cpp
@@ -840,6 +840,12 @@ void FSequencer::Tick(float InDeltaTime)
 		{
 			CompiledDataManager->Compile(RootSequencePtr);
 
+			// Reset to the root sequence if the focused sequence no longer exists. This can happen if either the subsequence has been deleted or the hierarchy has changed.
+			if (!GetFocusedMovieSceneSequence() || !GetFocusedMovieSceneSequence()->GetMovieScene())
+			{
+				PopToSequenceInstance(MovieSceneSequenceID::Root);
+			}
+
 			// Suppress auto evaluation if the sequence signature matches the one to be suppressed
 			if (!SuppressAutoEvalSignature.IsSet())
 			{
@@ -913,12 +919,6 @@ void FSequencer::Tick(float InDeltaTime)
 		SetLocalTimeLooped(CurrentTime.Time + Offset, RootToLocalLoopCounter);
 	}
 
-	// Reset to the root sequence if the focused sequence no longer exists. This can happen if either the subsequence has been deleted or the hierarchy has changed.
-	if (!MovieScene)
-	{
-		PopToSequenceInstance(MovieSceneSequenceID::Root);
-	}
-
 	if (GetSelectionRange().IsEmpty() && GetLoopMode() == SLM_LoopSelectionRange)
 	{
 		Settings->SetLoopMode(SLM_Loop);
