Sequencer: Added focus playback time, which focuses the view range on the current playback time without changing zoom level. The curve editor has this functionality, so adding it here for Sequencer as well

diff --git a/Engine/Source/Editor/Sequencer/Private/Sequencer.cpp b/Engine/Source/Editor/Sequencer/Private/Sequencer.cpp
--- a/Engine/Source/Editor/Sequencer/Private/Sequencer.cpp
+++ b/Engine/Source/Editor/Sequencer/Private/Sequencer.cpp
@@ -2916,6 +2916,20 @@ void FSequencer::TogglePlaybackRangeLocked()
 	}
 }
 
+void FSequencer::FocusPlaybackTime()
+{
+	const double CurrentTime = GetLocalTime().AsSeconds();
+
+	TRange<double> NewViewRange = GetViewRange();
+
+	double MidRange = (NewViewRange.GetUpperBoundValue() - NewViewRange.GetLowerBoundValue()) / 2.0 + NewViewRange.GetLowerBoundValue();
+
+	NewViewRange.SetLowerBoundValue(NewViewRange.GetLowerBoundValue() - (MidRange - CurrentTime));
+	NewViewRange.SetUpperBoundValue(NewViewRange.GetUpperBoundValue() - (MidRange - CurrentTime));
+
+	SetViewRange(NewViewRange, EViewRangeInterpolation::Animated);
+}
+
 void FSequencer::ResetViewRange()
 {
 	TRange<double> PlayRangeSeconds = GetPlaybackRange() / GetFocusedTickResolution();
@@ -11335,6 +11349,10 @@ void FSequencer::BindCommands()
 		Commands.SetStartPlaybackRange,
 		FExecuteAction::CreateLambda([this] { SetPlaybackStart(); }) );
 
+	SequencerCommandBindings->MapAction(
+		Commands.FocusPlaybackTime,
+		FExecuteAction::CreateSP(this, &FSequencer::FocusPlaybackTime));
+
 	SequencerCommandBindings->MapAction(
 		Commands.ResetViewRange,
 		FExecuteAction::CreateSP( this, &FSequencer::ResetViewRange ) );
diff --git a/Engine/Source/Editor/Sequencer/Private/Sequencer.h b/Engine/Source/Editor/Sequencer/Private/Sequencer.h
--- a/Engine/Source/Editor/Sequencer/Private/Sequencer.h
+++ b/Engine/Source/Editor/Sequencer/Private/Sequencer.h
@@ -230,6 +230,7 @@ public:
 
 	bool IsPlaybackRangeLocked() const;
 	void TogglePlaybackRangeLocked();
+	void FocusPlaybackTime();
 	void ResetViewRange();
 	void ZoomViewRange(float InZoomDelta);
 	void ZoomInViewRange();
diff --git a/Engine/Source/Editor/Sequencer/Private/SequencerCommands.cpp b/Engine/Source/Editor/Sequencer/Private/SequencerCommands.cpp
--- a/Engine/Source/Editor/Sequencer/Private/SequencerCommands.cpp
+++ b/Engine/Source/Editor/Sequencer/Private/SequencerCommands.cpp
@@ -35,6 +35,7 @@ void FSequencerCommands::RegisterCommands()
 	UI_COMMAND( StepToPreviousShot, "Step to Previous Shot", "Step to the previous shot", EUserInterfaceActionType::Button, FInputChord(EModifierKey::Shift, EKeys::Comma) );
 	UI_COMMAND( SetStartPlaybackRange, "Set Start Playback Range", "Set the start playback range", EUserInterfaceActionType::Button, FInputChord(EKeys::LeftBracket) );
 	UI_COMMAND( SetEndPlaybackRange, "Set End Playback Range", "Set the end playback range", EUserInterfaceActionType::Button, FInputChord(EKeys::RightBracket) );
+	UI_COMMAND( FocusPlaybackTime, "Focus Playback Time", "Focus the view range on the current playback time without changing zoom level", EUserInterfaceActionType::Button, FInputChord());
 	UI_COMMAND( ResetViewRange, "Reset View Range", "Reset view range to the playback range", EUserInterfaceActionType::Button, FInputChord(EKeys::Home) );
 	UI_COMMAND( ZoomToFit, "Zoom to Fit", "Zoom to Fit", EUserInterfaceActionType::Button, FInputChord(EKeys::F) );
 	UI_COMMAND( ZoomInViewRange, "Zoom into the View Range", "Zoom into the view range", EUserInterfaceActionType::Button, FInputChord(EKeys::Equals) );
diff --git a/Engine/Source/Editor/Sequencer/Public/SequencerCommands.h b/Engine/Source/Editor/Sequencer/Public/SequencerCommands.h
--- a/Engine/Source/Editor/Sequencer/Public/SequencerCommands.h
+++ b/Engine/Source/Editor/Sequencer/Public/SequencerCommands.h
@@ -97,6 +97,9 @@ public:
 	/** Set end playback range */
 	TSharedPtr< FUICommandInfo > SetEndPlaybackRange;
 
+	/** Focus the view range on the current playback time without changing zoom level */
+	TSharedPtr< FUICommandInfo > FocusPlaybackTime;
+
 	/** Reset the view range to the playback range */
 	TSharedPtr< FUICommandInfo > ResetViewRange;
 
