Sequencer: Add missing castable to FViewModel

diff --git a/Engine/Source/Editor/Sequencer/Public/MVVM/ViewModels/OutlinerItemModel.h b/Engine/Source/Editor/Sequencer/Public/MVVM/ViewModels/OutlinerItemModel.h
--- a/Engine/Source/Editor/Sequencer/Public/MVVM/ViewModels/OutlinerItemModel.h
+++ b/Engine/Source/Editor/Sequencer/Public/MVVM/ViewModels/OutlinerItemModel.h
@@ -112,6 +112,9 @@ private:
 	mutable bool bInitializedPinnedState;
 };
 
+//
+// Note: You must add the base class you use as a template parameter to the UE_SEQUENCER_DECLARE_CASTABLE list
+// 
 template<typename BaseType>
 class TOutlinerModelMixin : public BaseType, public FOutlinerItemModelMixin
 {
@@ -137,7 +140,7 @@ public:
 class SEQUENCER_API FOutlinerItemModel : public TOutlinerModelMixin<FViewModel>
 {
 public:
-	UE_SEQUENCER_DECLARE_CASTABLE(FOutlinerItemModel, FOutlinerItemModelMixin);
+	UE_SEQUENCER_DECLARE_CASTABLE(FOutlinerItemModel, FViewModel, FOutlinerItemModelMixin);
 };
 
 class SEQUENCER_API FMuteSoloOutlinerItemModel
