Sequencer: Step to Next/Previous Mark should also look at global marked frames

diff --git a/Engine/Source/Editor/Sequencer/Private/Sequencer.cpp b/Engine/Source/Editor/Sequencer/Private/Sequencer.cpp
--- a/Engine/Source/Editor/Sequencer/Private/Sequencer.cpp
+++ b/Engine/Source/Editor/Sequencer/Private/Sequencer.cpp
@@ -5247,8 +5247,6 @@ FFrameNumber FSequencer::OnGetNearestKey(FFrameTime InTime, ENearestKeyOption Ne
 		TArray<FMovieSceneMarkedFrame> GlobalMarkedFrames = GetGlobalMarkedFrames();
 		for (const FMovieSceneMarkedFrame& GlobalMarkedFrame : GlobalMarkedFrames)
 		{
-			FFrameNumber Diff = FMath::Abs(GlobalMarkedFrame.FrameNumber - CurrentTime);
-
 			if (!NearestTime.IsSet())
 			{
 				NearestTime = GlobalMarkedFrame.FrameNumber;
@@ -9876,11 +9874,35 @@ void FSequencer::StepToNextMark()
 		return;
 	}
 
+	FFrameNumber CurrentTime = GetLocalTime().Time.FloorToFrame();
+	TOptional<FFrameNumber> NearestTime;
+
 	const bool bForwards = true;
-	int32 MarkedIndex = FocusedMovieScene->FindNextMarkedFrame(GetLocalTime().Time.FloorToFrame(), bForwards);
+	int32 MarkedIndex = FocusedMovieScene->FindNextMarkedFrame(CurrentTime, bForwards);
 	if (MarkedIndex != INDEX_NONE)
 	{
-		SetLocalTimeDirectly(FocusedMovieScene->GetMarkedFrames()[MarkedIndex].FrameNumber.Value);
+		NearestTime = FocusedMovieScene->GetMarkedFrames()[MarkedIndex].FrameNumber.Value;
+	}
+
+	TArray<FMovieSceneMarkedFrame> GlobalMarkedFrames = GetGlobalMarkedFrames();
+	for (const FMovieSceneMarkedFrame& GlobalMarkedFrame : GlobalMarkedFrames)
+	{
+		if (GlobalMarkedFrame.FrameNumber > CurrentTime)
+		{
+			if (!NearestTime.IsSet())
+			{
+				NearestTime = GlobalMarkedFrame.FrameNumber;
+			}
+			else if (FMath::Abs(GlobalMarkedFrame.FrameNumber - CurrentTime) < FMath::Abs(NearestTime.GetValue() - CurrentTime))
+			{
+				NearestTime = GlobalMarkedFrame.FrameNumber;
+			}
+		}
+	}
+
+	if (NearestTime.IsSet())
+	{
+		SetLocalTimeDirectly(NearestTime.GetValue());
 	}
 }
 
@@ -9898,11 +9920,35 @@ void FSequencer::StepToPreviousMark()
 		return;
 	}
 	
+	FFrameNumber CurrentTime = GetLocalTime().Time.FloorToFrame();
+	TOptional<FFrameNumber> NearestTime;
+
 	const bool bForwards = false;
-	int32 MarkedIndex = FocusedMovieScene->FindNextMarkedFrame(GetLocalTime().Time.FloorToFrame(), bForwards);
+	int32 MarkedIndex = FocusedMovieScene->FindNextMarkedFrame(CurrentTime, bForwards);
 	if (MarkedIndex != INDEX_NONE)
 	{
-		SetLocalTimeDirectly(FocusedMovieScene->GetMarkedFrames()[MarkedIndex].FrameNumber.Value);
+		NearestTime = FocusedMovieScene->GetMarkedFrames()[MarkedIndex].FrameNumber.Value;
+	}
+
+	TArray<FMovieSceneMarkedFrame> GlobalMarkedFrames = GetGlobalMarkedFrames();
+	for (const FMovieSceneMarkedFrame& GlobalMarkedFrame : GlobalMarkedFrames)
+	{
+		if (GlobalMarkedFrame.FrameNumber < CurrentTime)
+		{
+			if (!NearestTime.IsSet())
+			{
+				NearestTime = GlobalMarkedFrame.FrameNumber;
+			}
+			else if (FMath::Abs(GlobalMarkedFrame.FrameNumber - CurrentTime) < FMath::Abs(NearestTime.GetValue() - CurrentTime))
+			{
+				NearestTime = GlobalMarkedFrame.FrameNumber;
+			}
+		}
+	}
+
+	if (NearestTime.IsSet())
+	{
+		SetLocalTimeDirectly(NearestTime.GetValue());
 	}
 }
 
