Sequencer: When resizing sections, also make sure they snap to other section start/end points

This used to work at some point (5.2) but has since regressed.

diff --git a/Engine/Source/Editor/Sequencer/Private/Tools/EditToolDragOperations.cpp b/Engine/Source/Editor/Sequencer/Private/Tools/EditToolDragOperations.cpp
--- a/Engine/Source/Editor/Sequencer/Private/Tools/EditToolDragOperations.cpp
+++ b/Engine/Source/Editor/Sequencer/Private/Tools/EditToolDragOperations.cpp
@@ -179,6 +179,7 @@ void FResizeSection::OnBeginDrag(const FPointerEvent& MouseEvent, FVector2D Loca
 	FInvalidKeyAndSectionSnappingCandidates SnapCandidates(EmptyKeySet, Sections);
 	SnapField = FSequencerSnapField(Sequencer, SnapCandidates, ESequencerEntity::Section | ESequencerEntity::Key);
 	SnapField.GetValue().SetSnapToInterval(Sequencer.GetSequencerSettings()->GetSnapSectionTimesToInterval());
+	SnapField.GetValue().SetSnapToLikeTypes(Sequencer.GetSequencerSettings()->GetSnapSectionTimesToSections());
 
 	SectionInitTimes.Empty();
 
