Sequencer: Sequencer module scripting layer

Precursor to sequencer specific functions like GetNextKey/GetPreviousKey

#rnx

diff --git a/Engine/Source/Editor/Sequencer/Private/MVVM/ViewModels/SequencerEditorViewModel.cpp b/Engine/Source/Editor/Sequencer/Private/MVVM/ViewModels/SequencerEditorViewModel.cpp
--- a/Engine/Source/Editor/Sequencer/Private/MVVM/ViewModels/SequencerEditorViewModel.cpp
+++ b/Engine/Source/Editor/Sequencer/Private/MVVM/ViewModels/SequencerEditorViewModel.cpp
@@ -8,13 +8,12 @@
 #include "MVVM/PinEditorExtension.h"
 #include "MVVM/Extensions/IOutlinerExtension.h"
 #include "MVVM/SharedViewModelData.h"
+#include "Scripting/SequencerModuleScriptingLayer.h"
 #include "ISequencerModule.h"
 #include "Sequencer.h"
 #include "MovieSceneSequenceID.h"
 
-namespace UE
-{
-namespace Sequencer
+namespace UE::Sequencer
 {
 
 TMap<TWeakPtr<FViewModel>, FString> GetNodePaths(FViewModelPtr RootModel)
@@ -73,6 +72,11 @@ TSharedPtr<FTrackAreaViewModel> FSequencerEditorViewModel::CreateTrackAreaImpl()
 	return NewTrackArea;
 }
 
+USequencerScriptingLayer* FSequencerEditorViewModel::CreateScriptingLayerImpl()
+{
+	return NewObject<USequencerModuleScriptingLayer>();
+}
+
 TViewModelPtr<FSequenceModel> FSequencerEditorViewModel::GetRootSequenceModel() const
 {
 	return GetRootModel().ImplicitCast();
@@ -249,6 +253,5 @@ TSharedPtr<FExtender> FSequencerEditorViewModel::GetSequencerMenuExtender(
 	return FExtender::Combine(Extenders);
 }
 
-} // namespace Sequencer
-} // namespace UE
+} // namespace UE::Sequencer
 
diff --git a/Engine/Source/Editor/Sequencer/Private/Scripting/SequencerModuleOutlinerScriptingObject.cpp b/Engine/Source/Editor/Sequencer/Private/Scripting/SequencerModuleOutlinerScriptingObject.cpp
new file mode 100644
--- /dev/null
+++ b/Engine/Source/Editor/Sequencer/Private/Scripting/SequencerModuleOutlinerScriptingObject.cpp
@@ -0,0 +1,6 @@
+// Copyright Epic Games, Inc. All Rights Reserved.
+
+#include "Scripting/SequencerModuleOutlinerScriptingObject.h"
+
+#include "Containers/Array.h"
+#include "Misc/FrameNumber.h"
diff --git a/Engine/Source/Editor/Sequencer/Private/Scripting/SequencerModuleOutlinerScriptingObject.h b/Engine/Source/Editor/Sequencer/Private/Scripting/SequencerModuleOutlinerScriptingObject.h
new file mode 100644
--- /dev/null
+++ b/Engine/Source/Editor/Sequencer/Private/Scripting/SequencerModuleOutlinerScriptingObject.h
@@ -0,0 +1,20 @@
+// Copyright Epic Games, Inc. All Rights Reserved.
+
+#pragma once
+
+#include "CoreTypes.h"
+#include "Containers/ContainersFwd.h"
+#include "Scripting/OutlinerScriptingObject.h"
+
+#include "SequencerModuleOutlinerScriptingObject.generated.h"
+
+struct FFrameNumber;
+struct FSequencerViewModelScriptingStruct;
+
+UCLASS()
+class USequencerModuleOutlinerScriptingObject : public USequencerOutlinerScriptingObject
+{
+public:
+
+	GENERATED_BODY()
+};
\ No newline at end of file
diff --git a/Engine/Source/Editor/Sequencer/Private/Scripting/SequencerModuleScriptingLayer.cpp b/Engine/Source/Editor/Sequencer/Private/Scripting/SequencerModuleScriptingLayer.cpp
new file mode 100644
--- /dev/null
+++ b/Engine/Source/Editor/Sequencer/Private/Scripting/SequencerModuleScriptingLayer.cpp
@@ -0,0 +1,13 @@
+// Copyright Epic Games, Inc. All Rights Reserved.
+
+#include "Scripting/SequencerModuleScriptingLayer.h"
+#include "Scripting/SequencerModuleOutlinerScriptingObject.h"
+
+#include "MVVM/ViewModels/OutlinerViewModel.h"
+#include "MVVM/ViewModels/EditorViewModel.h"
+
+void USequencerModuleScriptingLayer::Initialize(TSharedPtr<UE::Sequencer::FEditorViewModel> InViewModel)
+{
+	Outliner = NewObject<USequencerModuleOutlinerScriptingObject>(this, "Outliner");
+	Outliner->Initialize(InViewModel->GetOutliner());
+}
\ No newline at end of file
diff --git a/Engine/Source/Editor/Sequencer/Private/Scripting/SequencerModuleScriptingLayer.h b/Engine/Source/Editor/Sequencer/Private/Scripting/SequencerModuleScriptingLayer.h
new file mode 100644
--- /dev/null
+++ b/Engine/Source/Editor/Sequencer/Private/Scripting/SequencerModuleScriptingLayer.h
@@ -0,0 +1,18 @@
+// Copyright Epic Games, Inc. All Rights Reserved.
+
+#pragma once
+
+#include "CoreTypes.h"
+#include "Scripting/SequencerScriptingLayer.h"
+
+#include "SequencerModuleScriptingLayer.generated.h"
+
+UCLASS(MinimalAPI)
+class USequencerModuleScriptingLayer : public USequencerScriptingLayer
+{
+public:
+
+	GENERATED_BODY()
+
+	virtual void Initialize(TSharedPtr<UE::Sequencer::FEditorViewModel> InViewModel);
+};
\ No newline at end of file
diff --git a/Engine/Source/Editor/Sequencer/Public/MVVM/ViewModels/SequencerEditorViewModel.h b/Engine/Source/Editor/Sequencer/Public/MVVM/ViewModels/SequencerEditorViewModel.h
--- a/Engine/Source/Editor/Sequencer/Public/MVVM/ViewModels/SequencerEditorViewModel.h
+++ b/Engine/Source/Editor/Sequencer/Public/MVVM/ViewModels/SequencerEditorViewModel.h
@@ -14,9 +14,7 @@ class FSequencer;
 class UMovieSceneSequence;
 struct FSequencerHostCapabilities;
 
-namespace UE
-{
-namespace Sequencer
+namespace UE::Sequencer
 {
 
 class FSequenceModel;
@@ -78,6 +76,7 @@ protected:
 	virtual TSharedPtr<FOutlinerViewModel> CreateOutlinerImpl() override;
 	virtual TSharedPtr<FTrackAreaViewModel> CreateTrackAreaImpl() override;
 	virtual TSharedPtr<FSequencerCoreSelection> CreateSelectionImpl() override;
+	virtual USequencerScriptingLayer* CreateScriptingLayerImpl() override;
 	virtual bool IsReadOnly() const override;
 
 	void OnTrackAreaHotspotChanged(TSharedPtr<ITrackAreaHotspot> NewHotspot);
@@ -104,6 +103,5 @@ protected:
 	TArray<FSequencerCustomizationInfo> ActiveCustomizationInfos;
 };
 
-} // namespace Sequencer
-} // namespace UE
+} // namespace UE::Sequencer
 
diff --git a/Engine/Source/Editor/SequencerCore/Public/Scripting/OutlinerScriptingObject.h b/Engine/Source/Editor/SequencerCore/Public/Scripting/OutlinerScriptingObject.h
--- a/Engine/Source/Editor/SequencerCore/Public/Scripting/OutlinerScriptingObject.h
+++ b/Engine/Source/Editor/SequencerCore/Public/Scripting/OutlinerScriptingObject.h
@@ -15,7 +15,7 @@ namespace UE::Sequencer
 
 DECLARE_DYNAMIC_MULTICAST_DELEGATE(FSequencerOutlinerSelectionChanged);
 
-UCLASS()
+UCLASS(MinimalAPI)
 class USequencerOutlinerScriptingObject : public UObject
 {
 public:
@@ -25,7 +25,7 @@ public:
 	UPROPERTY(BlueprintAssignable, Category="Sequencer Editor")
 	FSequencerOutlinerSelectionChanged OnSelectionChanged;
 
-	void Initialize(UE::Sequencer::TViewModelPtr<UE::Sequencer::FOutlinerViewModel> InOutliner);
+	SEQUENCERCORE_API void Initialize(UE::Sequencer::TViewModelPtr<UE::Sequencer::FOutlinerViewModel> InOutliner);
 
 	UFUNCTION(BlueprintCallable, Category="Sequencer Editor")
 	FSequencerViewModelScriptingStruct GetRootNode() const;
