Sequencer: Next/Previous Key rework to expose it to BPs

The SequencerOutlinerScriptingObject now has GetNextKey/GetPreviousKey functions which can be called from scripts.

SequencerOutlinerViewModel implements GetNextKey/GetPreviousKey

GetNextKey/GetPreviousKey now operate within the the sequence's time bounds (which can be restricted to the selection range)

diff --git a/Engine/Source/Editor/Sequencer/Private/MVVM/ViewModels/SequencerOutlinerViewModel.cpp b/Engine/Source/Editor/Sequencer/Private/MVVM/ViewModels/SequencerOutlinerViewModel.cpp
--- a/Engine/Source/Editor/Sequencer/Private/MVVM/ViewModels/SequencerOutlinerViewModel.cpp
+++ b/Engine/Source/Editor/Sequencer/Private/MVVM/ViewModels/SequencerOutlinerViewModel.cpp
@@ -5,6 +5,7 @@
 #include "MVVM/ViewModels/SequenceModel.h"
 
 #include "Sequencer.h"
+#include "SequencerKeyCollection.h"
 #include "SequencerOutlinerItemDragDropOp.h"
 
 #include "Framework/MultiBox/MultiBoxBuilder.h"
@@ -81,6 +82,62 @@ TSharedRef<FDragDropOperation> FSequencerOutlinerViewModel::InitiateDrag(TArray<
 	return FSequencerOutlinerDragDropOp::New( MoveTemp(InDraggedModels), DefaultText, nullptr );
 }
 
+FFrameNumber FSequencerOutlinerViewModel::GetNextKey(const TArray<TSharedRef<UE::Sequencer::FViewModel>>& InNodes, FFrameNumber FrameNumber, EMovieSceneTimeUnit TimeUnit, const TRange<FFrameNumber>& Range)
+{
+	return GetNextKeyInternal(InNodes, FrameNumber, TimeUnit, Range, EFindKeyDirection::Forwards);
+}
+
+FFrameNumber FSequencerOutlinerViewModel::GetPreviousKey(const TArray<TSharedRef<UE::Sequencer::FViewModel>>& InNodes, FFrameNumber FrameNumber, EMovieSceneTimeUnit TimeUnit, const TRange<FFrameNumber>& Range)
+{
+	return GetNextKeyInternal(InNodes, FrameNumber, TimeUnit, Range, EFindKeyDirection::Backwards);
+}
+
+FFrameNumber FSequencerOutlinerViewModel::GetNextKeyInternal(const TArray<TSharedRef<UE::Sequencer::FViewModel>>& InNodes, FFrameNumber FrameNumber, EMovieSceneTimeUnit TimeUnit, const TRange<FFrameNumber>& Range, EFindKeyDirection Direction)
+{
+	FSequencerEditorViewModel* EditorViewModel = GetEditor()->CastThisChecked<FSequencerEditorViewModel>();
+	TSharedPtr<ISequencer> Sequencer = EditorViewModel->GetSequencer();
+	if (Sequencer.IsValid())
+	{
+		FFrameRate TickResolution = Sequencer->GetFocusedTickResolution();
+		FFrameRate DisplayRate = Sequencer->GetFocusedDisplayRate();
+
+		FSequencerKeyCollection* KeyCollection = Sequencer->GetKeyCollection();
+
+		const float DuplicateThresholdSeconds = SMALL_NUMBER;
+		const int64 TotalMaxSeconds = static_cast<int64>(TNumericLimits<int32>::Max() / TickResolution.AsDecimal());
+
+		FFrameNumber ThresholdFrames = (DuplicateThresholdSeconds * TickResolution).FloorToFrame();
+		if (ThresholdFrames.Value < -TotalMaxSeconds)
+		{
+			ThresholdFrames.Value = TotalMaxSeconds;
+		}
+		else if (ThresholdFrames.Value > TotalMaxSeconds)
+		{
+			ThresholdFrames.Value = TotalMaxSeconds;
+		}
+
+		if (TimeUnit == EMovieSceneTimeUnit::DisplayRate)
+		{
+			FrameNumber = ConvertFrameTime(FrameNumber, DisplayRate, TickResolution).FloorToFrame();
+		}
+
+		KeyCollection->Update(FSequencerKeyCollectionSignature::FromNodesRecursive(InNodes, ThresholdFrames));
+
+		TOptional<FFrameNumber> NextKey = KeyCollection->GetNextKey(FrameNumber, Direction, Range);
+		if (NextKey.IsSet())
+		{
+			if (TimeUnit == EMovieSceneTimeUnit::DisplayRate)
+			{
+				return ConvertFrameTime(NextKey.GetValue(), TickResolution, DisplayRate).FloorToFrame();
+			}
+
+			return NextKey.GetValue();
+		}
+	}
+
+	return FrameNumber;
+}
+
 } // namespace Sequencer
 } // namespace UE
 
diff --git a/Engine/Source/Editor/Sequencer/Private/Scripting/SequencerModuleOutlinerScriptingObject.cpp b/Engine/Source/Editor/Sequencer/Private/Scripting/SequencerModuleOutlinerScriptingObject.cpp
--- a/Engine/Source/Editor/Sequencer/Private/Scripting/SequencerModuleOutlinerScriptingObject.cpp
+++ b/Engine/Source/Editor/Sequencer/Private/Scripting/SequencerModuleOutlinerScriptingObject.cpp
@@ -1,6 +1,71 @@
 // Copyright Epic Games, Inc. All Rights Reserved.
 
 #include "Scripting/SequencerModuleOutlinerScriptingObject.h"
+#include "Scripting/ViewModelScriptingStruct.h"
+
+#include "MVVM/ViewModels/SequencerOutlinerviewModel.h"
+#include "MVVM/ViewModels/ViewModel.h"
+#include "MVVM/ViewModels/ViewModelIterators.h"
+#include "MVVM/ViewModels/OutlinerViewModel.h"
+#include "MVVM/ViewModels/EditorViewModel.h"
+#include "MVVM/ViewModelPtr.h"
 
 #include "Containers/Array.h"
 #include "Misc/FrameNumber.h"
+
+#define LOCTEXT_NAMESPACE "SequencerModuleOutlinerScriptingObject"
+
+
+FFrameNumber USequencerModuleOutlinerScriptingObject::GetNextKey(const TArray<FSequencerViewModelScriptingStruct>& InNodes, FFrameNumber FrameNumber, EMovieSceneTimeUnit TimeUnit) const
+{
+	using namespace UE::Sequencer;
+
+	TViewModelPtr<FSequencerOutlinerViewModel> Outliner = CastViewModel<FSequencerOutlinerViewModel>(WeakOutliner.Pin());
+	if (!Outliner)
+	{
+		FFrame::KismetExecutionMessage(*LOCTEXT("OutlinerInvalid", "Outliner is no longer valid.").ToString(), ELogVerbosity::Error);
+		return FFrameNumber();
+	}
+
+	TRange<FFrameNumber> Range = TRange<FFrameNumber>::All();
+
+	TArray<TSharedRef<FViewModel>> Nodes;
+	for (const FSequencerViewModelScriptingStruct& InNode : InNodes)
+	{
+		FViewModelPtr Item = InNode.WeakViewModel.ImplicitPin();
+		if (Item)
+		{
+			Nodes.Add(Item.ToSharedRef());
+		}
+	}
+
+	return Outliner->GetNextKey(Nodes, FrameNumber, TimeUnit, Range);
+}
+
+FFrameNumber USequencerModuleOutlinerScriptingObject::GetPreviousKey(const TArray<FSequencerViewModelScriptingStruct>& InNodes, FFrameNumber FrameNumber, EMovieSceneTimeUnit TimeUnit) const
+{
+	using namespace UE::Sequencer;
+
+	TViewModelPtr<FSequencerOutlinerViewModel> Outliner = CastViewModel<FSequencerOutlinerViewModel>(WeakOutliner.Pin());
+	if (!Outliner)
+	{
+		FFrame::KismetExecutionMessage(*LOCTEXT("OutlinerInvalid", "Outliner is no longer valid.").ToString(), ELogVerbosity::Error);
+		return FFrameNumber();
+	}
+
+	TRange<FFrameNumber> Range = TRange<FFrameNumber>::All();
+
+	TArray<TSharedRef<FViewModel>> Nodes;
+	for (const FSequencerViewModelScriptingStruct& InNode : InNodes)
+	{
+		TViewModelPtr<FViewModel> Item = InNode.WeakViewModel.ImplicitPin();
+		if (Item)
+		{
+			Nodes.Add(Item.ToSharedRef());
+		}
+	}
+
+	return Outliner->GetPreviousKey(Nodes, FrameNumber, TimeUnit, Range);
+}
+
+#undef LOCTEXT_NAMESPACE
diff --git a/Engine/Source/Editor/Sequencer/Private/Scripting/SequencerModuleOutlinerScriptingObject.h b/Engine/Source/Editor/Sequencer/Private/Scripting/SequencerModuleOutlinerScriptingObject.h
--- a/Engine/Source/Editor/Sequencer/Private/Scripting/SequencerModuleOutlinerScriptingObject.h
+++ b/Engine/Source/Editor/Sequencer/Private/Scripting/SequencerModuleOutlinerScriptingObject.h
@@ -4,6 +4,7 @@
 
 #include "CoreTypes.h"
 #include "Containers/ContainersFwd.h"
+#include "MovieSceneTimeUnit.h"
 #include "Scripting/OutlinerScriptingObject.h"
 
 #include "SequencerModuleOutlinerScriptingObject.generated.h"
@@ -17,4 +18,10 @@ class USequencerModuleOutlinerScriptingObject : public USequencerOutlinerScripti
 public:
 
 	GENERATED_BODY()
+
+	UFUNCTION(BlueprintCallable, Category = "Sequencer Editor")
+	FFrameNumber GetNextKey(const TArray<FSequencerViewModelScriptingStruct>& InNodes, FFrameNumber FrameNumber, EMovieSceneTimeUnit TimeUnit = EMovieSceneTimeUnit::DisplayRate) const;
+
+	UFUNCTION(BlueprintCallable, Category = "Sequencer Editor")
+	FFrameNumber GetPreviousKey(const TArray<FSequencerViewModelScriptingStruct>& InNodes, FFrameNumber FrameNumber, EMovieSceneTimeUnit TimeUnit = EMovieSceneTimeUnit::DisplayRate) const;
 };
\ No newline at end of file
diff --git a/Engine/Source/Editor/Sequencer/Private/Sequencer.cpp b/Engine/Source/Editor/Sequencer/Private/Sequencer.cpp
--- a/Engine/Source/Editor/Sequencer/Private/Sequencer.cpp
+++ b/Engine/Source/Editor/Sequencer/Private/Sequencer.cpp
@@ -1332,10 +1332,20 @@ void FSequencer::GetKeysFromSelection(TUniquePtr<FSequencerKeyCollection>& KeyCo
 	{
 		ThresholdFrames.Value = TotalMaxSeconds;
 	}
-
+	
 	KeyCollection->Update(FSequencerKeyCollectionSignature::FromNodesRecursive(SelectedItems, ThresholdFrames));
 }
 
+FSequencerKeyCollection* FSequencer::GetKeyCollection()
+{
+	if (!SelectedKeyCollection.IsValid())
+	{
+		SelectedKeyCollection.Reset(new FSequencerKeyCollection);
+	}
+
+	return SelectedKeyCollection.Get();
+}
+
 void FSequencer::GetAllKeys(TUniquePtr<FSequencerKeyCollection>& KeyCollection, float DuplicateThresholdSeconds) const
 {
 	using namespace UE::Sequencer;
@@ -3218,7 +3228,7 @@ void FSequencer::PlayTo(FMovieSceneSequencePlaybackParams PlaybackParams)
 
 	if (PlaybackParams.PositionType == EMovieScenePositionType::Frame)
 	{
-		PlayToTime = (PlaybackParams.Frame / GetFocusedDisplayRate()) * GetFocusedTickResolution();
+		PlayToTime = PlaybackParams.Frame;
 	}
 	else if (PlaybackParams.PositionType == EMovieScenePositionType::Time)
 	{
@@ -4773,8 +4783,9 @@ FReply FSequencer::JumpToPreviousKey()
 
 	if (SelectedKeyCollection.IsValid())
 	{
+		TRange<FFrameNumber> Range = GetTimeBounds();
 		FFrameNumber FrameNumber = GetLocalTime().Time.FloorToFrame();
-		TOptional<FFrameNumber> NewTime = SelectedKeyCollection->GetNextKey(FrameNumber, EFindKeyDirection::Backwards);
+		TOptional<FFrameNumber> NewTime = SelectedKeyCollection->GetNextKey(FrameNumber, EFindKeyDirection::Backwards, Range);
 		if (NewTime.IsSet())
 		{
 			SetPlaybackStatus(EMovieScenePlayerStatus::Stepping);
@@ -4802,8 +4813,9 @@ FReply FSequencer::JumpToNextKey()
 
 	if (SelectedKeyCollection.IsValid())
 	{
+		TRange<FFrameNumber> Range = GetTimeBounds();
 		FFrameNumber FrameNumber = GetLocalTime().Time.FloorToFrame();
-		TOptional<FFrameNumber> NewTime = SelectedKeyCollection->GetNextKey(FrameNumber, EFindKeyDirection::Forwards);
+		TOptional<FFrameNumber> NewTime = SelectedKeyCollection->GetNextKey(FrameNumber, EFindKeyDirection::Forwards, Range);
 		if (NewTime.IsSet())
 		{
 			SetPlaybackStatus(EMovieScenePlayerStatus::Stepping);
diff --git a/Engine/Source/Editor/Sequencer/Private/Sequencer.h b/Engine/Source/Editor/Sequencer/Private/Sequencer.h
--- a/Engine/Source/Editor/Sequencer/Private/Sequencer.h
+++ b/Engine/Source/Editor/Sequencer/Private/Sequencer.h
@@ -987,6 +987,7 @@ protected:
 
 	/** Get all the keys for the current sequencer selection */
 	virtual void GetKeysFromSelection(TUniquePtr<FSequencerKeyCollection>& KeyCollection, float DuplicateThresholdSeconds) override;
+	virtual FSequencerKeyCollection* GetKeyCollection() override;
 
 	void GetAllKeys(TUniquePtr<FSequencerKeyCollection>& KeyCollection, float DuplicateThresoldSeconds) const;
 
diff --git a/Engine/Source/Editor/Sequencer/Private/SequencerKeyCollection.cpp b/Engine/Source/Editor/Sequencer/Private/SequencerKeyCollection.cpp
--- a/Engine/Source/Editor/Sequencer/Private/SequencerKeyCollection.cpp
+++ b/Engine/Source/Editor/Sequencer/Private/SequencerKeyCollection.cpp
@@ -248,35 +248,38 @@ TArrayView<const FFrameNumber> GetKeysInRangeInternal(const TArray<FFrameNumber>
 	return TArrayView<const FFrameNumber>();
 }
 
-TOptional<FFrameNumber> GetNextKeyInternal(const TArray<FFrameNumber>& Times, FFrameNumber FrameNumber, EFindKeyDirection Direction)
+TOptional<FFrameNumber> GetNextKeyInternal(const TArray<FFrameNumber>& Times, FFrameNumber FrameNumber, EFindKeyDirection Direction, const TRange<FFrameNumber>& Range)
 {
+	TArrayView<const FFrameNumber> KeysInRange = GetKeysInRangeInternal(Times, Range);
+
 	int32 Index = INDEX_NONE;
 	if (Direction == EFindKeyDirection::Forwards)
 	{
-		Index = Algo::UpperBound(Times, FrameNumber);
+		Index = Algo::UpperBound(KeysInRange, FrameNumber);
 	}
 	else
 	{
-		Index = Algo::LowerBound(Times, FrameNumber) - 1;
+		Index = Algo::LowerBound(KeysInRange, FrameNumber) - 1;
 	}
 
-	if (Times.IsValidIndex(Index))
+	if (KeysInRange.IsValidIndex(Index))
 	{
-		return Times[Index];
+		return KeysInRange[Index];
 	}
-	else if (Times.Num() > 0)
+	else if (KeysInRange.Num() > 0)
 	{
 		if (Direction == EFindKeyDirection::Forwards)
 		{
-			return Times[0];
+			return KeysInRange[0];
 		}
 		else
 		{
-			return Times.Last();
+			return KeysInRange.Last();
 		}
 	}
 
 	return TOptional<FFrameNumber>();
+
 }
 
 TOptional<FFrameNumber> FindFirstKeyInRangeInternal(const TArray<FFrameNumber>& Times, const TRange<FFrameNumber>& Range, EFindKeyDirection Direction)
@@ -309,12 +312,12 @@ TArrayView<const FFrameNumber> FSequencerKeyCollection::GetSectionKeysInRange(co
 	return GetKeysInRangeInternal(GroupedSectionTimes, Range);
 }
 
-TOptional<FFrameNumber> FSequencerKeyCollection::GetNextKey(FFrameNumber FrameNumber, EFindKeyDirection Direction) const
+TOptional<FFrameNumber> FSequencerKeyCollection::GetNextKey(FFrameNumber FrameNumber, EFindKeyDirection Direction, const TRange<FFrameNumber>& Range) const
 {
-	return GetNextKeyInternal(GroupedTimes, FrameNumber, Direction);
+	return GetNextKeyInternal(GroupedTimes, FrameNumber, Direction, Range);
 }
 
-TOptional<FFrameNumber> FSequencerKeyCollection::GetNextSectionKey(FFrameNumber FrameNumber, EFindKeyDirection Direction) const
+TOptional<FFrameNumber> FSequencerKeyCollection::GetNextSectionKey(FFrameNumber FrameNumber, EFindKeyDirection Direction, const TRange<FFrameNumber>& Range) const
 {
-	return GetNextKeyInternal(GroupedSectionTimes, FrameNumber, Direction);
+	return GetNextKeyInternal(GroupedSectionTimes, FrameNumber, Direction, Range);
 }
diff --git a/Engine/Source/Editor/Sequencer/Public/ISequencer.h b/Engine/Source/Editor/Sequencer/Public/ISequencer.h
--- a/Engine/Source/Editor/Sequencer/Public/ISequencer.h
+++ b/Engine/Source/Editor/Sequencer/Public/ISequencer.h
@@ -615,6 +615,7 @@ public:
 
 	/** Get all the keys for the current sequencer selection */
 	virtual void GetKeysFromSelection(TUniquePtr<FSequencerKeyCollection>& KeyCollection, float DuplicateThresoldTime) = 0;
+	virtual FSequencerKeyCollection* GetKeyCollection() = 0;
 
 	virtual TArray<FMovieSceneMarkedFrame> GetMarkedFrames() const = 0;
 
diff --git a/Engine/Source/Editor/Sequencer/Public/MVVM/ViewModels/SequencerOutlinerViewModel.h b/Engine/Source/Editor/Sequencer/Public/MVVM/ViewModels/SequencerOutlinerViewModel.h
--- a/Engine/Source/Editor/Sequencer/Public/MVVM/ViewModels/SequencerOutlinerViewModel.h
+++ b/Engine/Source/Editor/Sequencer/Public/MVVM/ViewModels/SequencerOutlinerViewModel.h
@@ -5,6 +5,8 @@
 #include "Input/DragAndDrop.h"
 #include "MVVM/ViewModels/OutlinerViewModel.h"
 #include "MVVM/ViewModels/SequencerEditorViewModel.h"
+#include "MovieSceneTimeUnit.h"
+#include "SequencerKeyCollection.h"
 #include "ISequencerModule.h"
 
 class FMenuBuilder;
@@ -31,6 +33,13 @@ public:
 	TSharedPtr<SWidget> CreateContextMenuWidget() override;
 	TSharedRef<FDragDropOperation> InitiateDrag(TArray<TWeakViewModelPtr<IOutlinerExtension>>&& InDraggedModels) override;
 
+	FFrameNumber GetNextKey(const TArray<TSharedRef<UE::Sequencer::FViewModel>>& InNodes, FFrameNumber FrameNumber, EMovieSceneTimeUnit TimeUnit, const TRange<FFrameNumber>& Range);
+	FFrameNumber GetPreviousKey(const TArray<TSharedRef<UE::Sequencer::FViewModel>>& InNodes, FFrameNumber FrameNumber, EMovieSceneTimeUnit TimeUnit, const TRange<FFrameNumber>& Range);
+
+private:
+
+	FFrameNumber GetNextKeyInternal(const TArray<TSharedRef<UE::Sequencer::FViewModel>>& InNodes, FFrameNumber FrameNumber, EMovieSceneTimeUnit TimeUnit, const TRange<FFrameNumber>& Range, EFindKeyDirection Direction);
+
 public:
 
 	/** Called when the add menu is created */
diff --git a/Engine/Source/Editor/Sequencer/Public/SequencerKeyCollection.h b/Engine/Source/Editor/Sequencer/Public/SequencerKeyCollection.h
--- a/Engine/Source/Editor/Sequencer/Public/SequencerKeyCollection.h
+++ b/Engine/Source/Editor/Sequencer/Public/SequencerKeyCollection.h
@@ -112,8 +112,8 @@ public:
 	* @param Direction  Whether to return the next key or previous key from that time
 	* @return (Optional)  Frame number of the key that's next or previous from that time 
 	*/
-	SEQUENCER_API TOptional<FFrameNumber> GetNextKey(FFrameNumber FrameNumber, EFindKeyDirection Direction) const;
-	SEQUENCER_API TOptional<FFrameNumber> GetNextSectionKey(FFrameNumber FrameNumber, EFindKeyDirection Direction) const;
+	SEQUENCER_API TOptional<FFrameNumber> GetNextKey(FFrameNumber FrameNumber, EFindKeyDirection Direction, const TRange<FFrameNumber>& Range) const;
+	SEQUENCER_API TOptional<FFrameNumber> GetNextSectionKey(FFrameNumber FrameNumber, EFindKeyDirection Direction, const TRange<FFrameNumber>& Range) const;
 
 	/**
 	 * Access the signature this collection was generated with
diff --git a/Engine/Source/Editor/SequencerCore/Public/Scripting/OutlinerScriptingObject.h b/Engine/Source/Editor/SequencerCore/Public/Scripting/OutlinerScriptingObject.h
--- a/Engine/Source/Editor/SequencerCore/Public/Scripting/OutlinerScriptingObject.h
+++ b/Engine/Source/Editor/SequencerCore/Public/Scripting/OutlinerScriptingObject.h
@@ -39,9 +39,11 @@ public:
 	UFUNCTION(BlueprintCallable, Category="Sequencer Editor")
 	void SetSelection(const TArray<FSequencerViewModelScriptingStruct>& InSelection);
 
+protected:
+
+	UE::Sequencer::TWeakViewModelPtr<UE::Sequencer::FOutlinerViewModel> WeakOutliner;
+
 private:
 
 	void BroadcastSelectionChanged();
-
-	UE::Sequencer::TWeakViewModelPtr<UE::Sequencer::FOutlinerViewModel> WeakOutliner;
 };
\ No newline at end of file
