Sequencer: Fix capitalization

#rnx

diff --git a/Engine/Source/Editor/Sequencer/Private/Scripting/SequencerModuleOutlinerScriptingObject.cpp b/Engine/Source/Editor/Sequencer/Private/Scripting/SequencerModuleOutlinerScriptingObject.cpp
--- a/Engine/Source/Editor/Sequencer/Private/Scripting/SequencerModuleOutlinerScriptingObject.cpp
+++ b/Engine/Source/Editor/Sequencer/Private/Scripting/SequencerModuleOutlinerScriptingObject.cpp
@@ -3,7 +3,7 @@
 #include "Scripting/SequencerModuleOutlinerScriptingObject.h"
 #include "Scripting/ViewModelScriptingStruct.h"
 
-#include "MVVM/ViewModels/SequencerOutlinerviewModel.h"
+#include "MVVM/ViewModels/SequencerOutlinerViewModel.h"
 #include "MVVM/ViewModels/ViewModel.h"
 #include "MVVM/ViewModels/ViewModelIterators.h"
 #include "MVVM/ViewModels/OutlinerViewModel.h"
