Sequencer: Outliner UX improvements

This change represents the final round of improvements to the Sequencer outliner UI per UX team's mockups. The notable changes are:

 - All widgets on the outliner are now properly aligned and maintained through a rich and extensible column management API. All columns can be toggled on and off as desired.
 - Previous/Add/Next key navigation buttons have updated icons and spacing and are consistent with the pin/lock/mute/solo icons
 - Labels and edit widgets will now forcibly overflow other, less important, widgets on hover if they are occluded because of a narrow overal outliner width
 - All rows now have a separator line
 - Track rows are now uniform heights
 - Lock camera button is now more prominent

diff --git a/Engine/Plugins/Experimental/Animation/ContextualAnimation/Source/ContextualAnimationEditor/Private/ContextualAnimMovieSceneNotifyTrackEditor.cpp b/Engine/Plugins/Experimental/Animation/ContextualAnimation/Source/ContextualAnimationEditor/Private/ContextualAnimMovieSceneNotifyTrackEditor.cpp
--- a/Engine/Plugins/Experimental/Animation/ContextualAnimation/Source/ContextualAnimationEditor/Private/ContextualAnimMovieSceneNotifyTrackEditor.cpp
+++ b/Engine/Plugins/Experimental/Animation/ContextualAnimation/Source/ContextualAnimationEditor/Private/ContextualAnimMovieSceneNotifyTrackEditor.cpp
@@ -6,7 +6,7 @@
 #include "ContextualAnimSceneAsset.h"
 #include "ContextualAnimViewModel.h"
 #include "DetailsViewArgs.h"
-#include "SequencerUtilities.h"
+#include "MVVM/Views/ViewUtilities.h"
 #include "Framework/Commands/UICommandList.h"
 #include "SequencerSectionPainter.h"
 #include "Animation/AnimNotifies/AnimNotify.h"
@@ -184,13 +184,10 @@ TSharedPtr<SWidget> FContextualAnimMovieSceneNotifyTrackEditor::BuildOutlinerEdi
 		return MenuBuilder.MakeWidget();
 	};
 
-	return SNew(SHorizontalBox)
-		+ SHorizontalBox::Slot()
-		.AutoWidth()
-		.VAlign(VAlign_Center)
-		[
-			FSequencerUtilities::MakeAddButton(LOCTEXT("Notify", "Notify"), FOnGetContent::CreateLambda(SubMenuCallback), Params.NodeIsHovered, GetSequencer())
-		];
+	return UE::Sequencer::MakeAddButton(
+		LOCTEXT("Notify", "Notify"),
+		FOnGetContent::CreateLambda(SubMenuCallback),
+		Params.ViewModel);
 }
 
 void FContextualAnimMovieSceneNotifyTrackEditor::BuildNewIKTargetSubMenu(FMenuBuilder& MenuBuilder, UContextualAnimMovieSceneNotifyTrack* Track, int32 RowIndex)
diff --git a/Engine/Plugins/Experimental/Animation/ContextualAnimation/Source/ContextualAnimationEditor/Private/ContextualAnimMovieSceneTrackEditor.cpp b/Engine/Plugins/Experimental/Animation/ContextualAnimation/Source/ContextualAnimationEditor/Private/ContextualAnimMovieSceneTrackEditor.cpp
--- a/Engine/Plugins/Experimental/Animation/ContextualAnimation/Source/ContextualAnimationEditor/Private/ContextualAnimMovieSceneTrackEditor.cpp
+++ b/Engine/Plugins/Experimental/Animation/ContextualAnimation/Source/ContextualAnimationEditor/Private/ContextualAnimMovieSceneTrackEditor.cpp
@@ -99,7 +99,7 @@ FReply FContextualAnimSection::OnSectionDoubleClicked(const FGeometry& SectionGe
 
 float FContextualAnimSection::GetSectionHeight() const
 {
-	return 20.f;
+	return 28.f;
 }
 
 #undef LOCTEXT_NAMESPACE
diff --git a/Engine/Plugins/Experimental/ChaosCaching/Source/ChaosCachingEditor/Private/Sequencer/ChaosCacheTrackEditor.cpp b/Engine/Plugins/Experimental/ChaosCaching/Source/ChaosCachingEditor/Private/Sequencer/ChaosCacheTrackEditor.cpp
--- a/Engine/Plugins/Experimental/ChaosCaching/Source/ChaosCachingEditor/Private/Sequencer/ChaosCacheTrackEditor.cpp
+++ b/Engine/Plugins/Experimental/ChaosCaching/Source/ChaosCachingEditor/Private/Sequencer/ChaosCacheTrackEditor.cpp
@@ -12,7 +12,8 @@
 #include "Rendering/SlateRenderer.h"
 #include "SequencerSectionPainter.h"
 #include "ScopedTransaction.h"
-#include "SequencerUtilities.h"
+#include "MVVM/Views/ViewUtilities.h"
+#include "MVVM/ViewModels/ViewDensity.h"
 #include "Styling/SlateIconFinder.h"
 #include "TimeToPixel.h"
 #include "Widgets/SBoxPanel.h"
@@ -20,7 +21,7 @@
 namespace ChaosCacheEditorConstants
 {
 	// @todo Sequencer Allow this to be customizable
-	constexpr float AnimationTrackHeight = 20.f;
+	constexpr float AnimationTrackHeight = 28.f;
 }
 
 #define LOCTEXT_NAMESPACE "FChaosCacheTrackEditor"
@@ -62,9 +63,9 @@ FText FChaosCacheSection::GetSectionTitle() const
 	return LOCTEXT("NoChaosCacheSection", "No ChaosCache");
 }
 
-float FChaosCacheSection::GetSectionHeight() const
+float FChaosCacheSection::GetSectionHeight(const UE::Sequencer::FViewDensityInfo& ViewDensity) const
 {
-	return ChaosCacheEditorConstants::AnimationTrackHeight;
+	return ViewDensity.UniformHeight.Get(ChaosCacheEditorConstants::AnimationTrackHeight);
 }
 
 int32 FChaosCacheSection::OnPaintSection(FSequencerSectionPainter& Painter) const
@@ -369,13 +370,10 @@ TSharedPtr<SWidget> FChaosCacheTrackEditor::BuildOutlinerEditWidget(const FGuid&
 			return MenuBuilder.MakeWidget();
 		};
 
-		return SNew(SHorizontalBox)
-			+ SHorizontalBox::Slot()
-			.AutoWidth()
-			.VAlign(VAlign_Center)
-			[
-				FSequencerUtilities::MakeAddButton(LOCTEXT("ChaosCacheText", "Chaos Cache"), FOnGetContent::CreateLambda(SubMenuCallback), Params.NodeIsHovered, GetSequencer())
-			];
+		return UE::Sequencer::MakeAddButton(
+			LOCTEXT("ChaosCacheText", "Chaos Cache"),
+			FOnGetContent::CreateLambda(SubMenuCallback),
+			Params.ViewModel);
 	}
 	else
 	{
diff --git a/Engine/Plugins/Experimental/ChaosCaching/Source/ChaosCachingEditor/Public/Sequencer/ChaosCacheTrackEditor.h b/Engine/Plugins/Experimental/ChaosCaching/Source/ChaosCachingEditor/Public/Sequencer/ChaosCacheTrackEditor.h
--- a/Engine/Plugins/Experimental/ChaosCaching/Source/ChaosCachingEditor/Public/Sequencer/ChaosCacheTrackEditor.h
+++ b/Engine/Plugins/Experimental/ChaosCaching/Source/ChaosCachingEditor/Public/Sequencer/ChaosCacheTrackEditor.h
@@ -69,7 +69,7 @@ public:
 	// ~Begin ISequencerSection interface
 	virtual UMovieSceneSection* GetSectionObject() override;
 	virtual FText GetSectionTitle() const override;
-	virtual float GetSectionHeight() const override;
+	virtual float GetSectionHeight(const UE::Sequencer::FViewDensityInfo& ViewDensity) const override;
 	virtual int32 OnPaintSection( FSequencerSectionPainter& Painter ) const override;
 	virtual void BeginResizeSection() override;
 	virtual void ResizeSection(ESequencerSectionResizeMode ResizeMode, FFrameNumber ResizeTime) override;
diff --git a/Engine/Plugins/Experimental/CinematicPrestreaming/Source/CinematicPrestreamingEditor/Private/TrackEditors/CinePrestreamingTrackEditor.cpp b/Engine/Plugins/Experimental/CinematicPrestreaming/Source/CinematicPrestreamingEditor/Private/TrackEditors/CinePrestreamingTrackEditor.cpp
--- a/Engine/Plugins/Experimental/CinematicPrestreaming/Source/CinematicPrestreamingEditor/Private/TrackEditors/CinePrestreamingTrackEditor.cpp
+++ b/Engine/Plugins/Experimental/CinematicPrestreaming/Source/CinematicPrestreamingEditor/Private/TrackEditors/CinePrestreamingTrackEditor.cpp
@@ -10,7 +10,7 @@
 #include "MovieSceneTimeHelpers.h"
 #include "Sections/MovieSceneCinePrestreamingSection.h"
 #include "SequencerSectionPainter.h"
-#include "SequencerUtilities.h"
+#include "MVVM/Views/ViewUtilities.h"
 #include "Styling/StyleColors.h"
 #include "Tracks/MovieSceneCinePrestreamingTrack.h"
 #include "Widgets/SBoxPanel.h"
@@ -182,10 +182,10 @@ void FCinePrestreamingTrackEditor::BuildAddTrackMenu(FMenuBuilder& MenuBuilder)
 
 TSharedPtr<SWidget> FCinePrestreamingTrackEditor::BuildOutlinerEditWidget(const FGuid& ObjectBinding, UMovieSceneTrack* Track, const FBuildEditWidgetParams& Params)
 {
-	return FSequencerUtilities::MakeAddButton(
+	return UE::Sequencer::MakeAddButton(
 		LOCTEXT("AddCinePrestreaming_ButtonLabel", "Cinematic Prestreaming"),
 		FOnGetContent::CreateSP(this, &FCinePrestreamingTrackEditor::BuildAddDataLayerMenu, Track),
-		Params.NodeIsHovered, GetSequencer());
+		Params.ViewModel);
 }
 
 UMovieSceneCinePrestreamingSection* FCinePrestreamingTrackEditor::AddNewSection(UMovieScene* MovieScene, UMovieSceneTrack* PrestreamingTrack)
diff --git a/Engine/Plugins/Experimental/GeometryCollectionPlugin/Source/GeometryCollectionSequencer/Private/GeometryCollectionTrackEditor.cpp b/Engine/Plugins/Experimental/GeometryCollectionPlugin/Source/GeometryCollectionSequencer/Private/GeometryCollectionTrackEditor.cpp
--- a/Engine/Plugins/Experimental/GeometryCollectionPlugin/Source/GeometryCollectionSequencer/Private/GeometryCollectionTrackEditor.cpp
+++ b/Engine/Plugins/Experimental/GeometryCollectionPlugin/Source/GeometryCollectionSequencer/Private/GeometryCollectionTrackEditor.cpp
@@ -16,7 +16,7 @@
 #include "MovieSceneGeometryCollectionTrack.h"
 #include "MovieSceneGeometryCollectionSection.h"
 #include "ContentBrowserModule.h"
-#include "SequencerUtilities.h"
+#include "MVVM/Views/ViewUtilities.h"
 #include "ISectionLayoutBuilder.h"
 #include "Styling/AppStyle.h"
 #include "MovieSceneTimeHelpers.h"
@@ -33,7 +33,7 @@
 namespace GeometryCollectionEditorConstants
 {
 	// @todo Sequencer Allow this to be customizable
-	const uint32 AnimationTrackHeight = 20;
+	const uint32 AnimationTrackHeight = 28;
 }
 
 #define LOCTEXT_NAMESPACE "FGeometryCollectionTrackEditor"
@@ -389,13 +389,7 @@ TSharedPtr<SWidget> FGeometryCollectionTrackEditor::BuildOutlinerEditWidget(cons
 			return MenuBuilder.MakeWidget();
 		};
 
-		return SNew(SHorizontalBox)
-			+ SHorizontalBox::Slot()
-			.AutoWidth()
-			.VAlign(VAlign_Center)
-			[
-				FSequencerUtilities::MakeAddButton(LOCTEXT("GeomCollectionText", "Geometry Collection"), FOnGetContent::CreateLambda(SubMenuCallback), Params.NodeIsHovered, GetSequencer())
-			];
+		return UE::Sequencer::MakeAddButton(LOCTEXT("GeomCollectionText", "Geometry Collection"), FOnGetContent::CreateLambda(SubMenuCallback), Params.ViewModel);
 	}
 	else
 	{
diff --git a/Engine/Plugins/FX/NiagaraSimCaching/Source/NiagaraSimCachingEditor/Private/Sequencer/NiagaraCacheTrackEditor.cpp b/Engine/Plugins/FX/NiagaraSimCaching/Source/NiagaraSimCachingEditor/Private/Sequencer/NiagaraCacheTrackEditor.cpp
--- a/Engine/Plugins/FX/NiagaraSimCaching/Source/NiagaraSimCachingEditor/Private/Sequencer/NiagaraCacheTrackEditor.cpp
+++ b/Engine/Plugins/FX/NiagaraSimCaching/Source/NiagaraSimCachingEditor/Private/Sequencer/NiagaraCacheTrackEditor.cpp
@@ -25,10 +25,11 @@
 #include "UObject/Package.h"
 #include "Widgets/Images/SImage.h"
 #include "Widgets/Input/SButton.h"
+#include "MVVM/ViewModels/ViewDensity.h"
 
 namespace NiagaraCacheEditorConstants
 {
-	constexpr float AnimationTrackHeight = 20.f;
+	constexpr float AnimationTrackHeight = 28.f;
 }
 
 #define LOCTEXT_NAMESPACE "FNiagaraCacheTrackEditor"
@@ -63,9 +64,9 @@ FText FNiagaraCacheSection::GetSectionTitle() const
 	return LOCTEXT("NoNiagaraCacheSection", "No NiagaraCache");
 }
 
-float FNiagaraCacheSection::GetSectionHeight() const
+float FNiagaraCacheSection::GetSectionHeight(const UE::Sequencer::FViewDensityInfo& ViewDensity) const
 {
-	return NiagaraCacheEditorConstants::AnimationTrackHeight;
+	return ViewDensity.UniformHeight.Get(NiagaraCacheEditorConstants::AnimationTrackHeight);
 }
 
 int32 FNiagaraCacheSection::OnPaintSection(FSequencerSectionPainter& Painter) const
diff --git a/Engine/Plugins/FX/NiagaraSimCaching/Source/NiagaraSimCachingEditor/Public/Sequencer/NiagaraCacheTrackEditor.h b/Engine/Plugins/FX/NiagaraSimCaching/Source/NiagaraSimCachingEditor/Public/Sequencer/NiagaraCacheTrackEditor.h
--- a/Engine/Plugins/FX/NiagaraSimCaching/Source/NiagaraSimCachingEditor/Public/Sequencer/NiagaraCacheTrackEditor.h
+++ b/Engine/Plugins/FX/NiagaraSimCaching/Source/NiagaraSimCachingEditor/Public/Sequencer/NiagaraCacheTrackEditor.h
@@ -72,7 +72,7 @@ public:
 	// ~Begin ISequencerSection interface
 	virtual UMovieSceneSection* GetSectionObject() override;
 	virtual FText GetSectionTitle() const override;
-	virtual float GetSectionHeight() const override;
+	virtual float GetSectionHeight(const UE::Sequencer::FViewDensityInfo& ViewDensity) const override;
 	virtual int32 OnPaintSection( FSequencerSectionPainter& Painter ) const override;
 	virtual void BeginResizeSection() override;
 	virtual void ResizeSection(ESequencerSectionResizeMode ResizeMode, FFrameNumber ResizeTime) override;
diff --git a/Engine/Plugins/Media/MediaCompositing/Source/MediaCompositingEditor/Private/Sequencer/MediaPlayerPropertyTrackEditor.cpp b/Engine/Plugins/Media/MediaCompositing/Source/MediaCompositingEditor/Private/Sequencer/MediaPlayerPropertyTrackEditor.cpp
--- a/Engine/Plugins/Media/MediaCompositing/Source/MediaCompositingEditor/Private/Sequencer/MediaPlayerPropertyTrackEditor.cpp
+++ b/Engine/Plugins/Media/MediaCompositing/Source/MediaCompositingEditor/Private/Sequencer/MediaPlayerPropertyTrackEditor.cpp
@@ -11,7 +11,7 @@
 #include "SequencerSectionPainter.h"
 #include "MovieSceneMediaPlayerPropertySection.h"
 #include "MovieSceneMediaPlayerPropertyTrack.h"
-#include "SequencerUtilities.h"
+#include "MVVM/Views/ViewUtilities.h"
 #include "Widgets/SBoxPanel.h"
 #include "Widgets/Layout/SBox.h"
 #include "MediaSource.h"
@@ -174,14 +174,7 @@ TSharedPtr<SWidget> FMediaPlayerPropertyTrackEditor::BuildOutlinerEditWidget(con
 		return Picker;
 	};
 
-	return SNew(SHorizontalBox)
-
-		+ SHorizontalBox::Slot()
-		.AutoWidth()
-		.VAlign(VAlign_Center)
-		[
-			FSequencerUtilities::MakeAddButton(LOCTEXT("AddMediaSection_Text", "Media"), FOnGetContent::CreateLambda(CreatePicker), Params.NodeIsHovered, GetSequencer())
-		];
+	return UE::Sequencer::MakeAddButton(LOCTEXT("AddMediaSection_Text", "Media"), FOnGetContent::CreateLambda(CreatePicker), Params.ViewModel);
 }
 
 TSharedRef<ISequencerSection> FMediaPlayerPropertyTrackEditor::MakeSectionInterface(UMovieSceneSection& SectionObject, UMovieSceneTrack& Track, FGuid ObjectBinding)
diff --git a/Engine/Plugins/Media/MediaCompositing/Source/MediaCompositingEditor/Private/Sequencer/MediaThumbnailSection.cpp b/Engine/Plugins/Media/MediaCompositing/Source/MediaCompositingEditor/Private/Sequencer/MediaThumbnailSection.cpp
--- a/Engine/Plugins/Media/MediaCompositing/Source/MediaCompositingEditor/Private/Sequencer/MediaThumbnailSection.cpp
+++ b/Engine/Plugins/Media/MediaCompositing/Source/MediaCompositingEditor/Private/Sequencer/MediaThumbnailSection.cpp
@@ -66,9 +66,9 @@ FMargin FMediaThumbnailSection::GetContentPadding() const
 }
 
 
-float FMediaThumbnailSection::GetSectionHeight() const
+float FMediaThumbnailSection::GetSectionHeight(const UE::Sequencer::FViewDensityInfo& ViewDensity) const
 {
-	return FThumbnailSection::GetSectionHeight() + 2 * 9.0f; // make space for the film border
+	return FThumbnailSection::GetSectionHeight(ViewDensity) + 2 * 9.0f; // make space for the film border
 }
 
 
diff --git a/Engine/Plugins/Media/MediaCompositing/Source/MediaCompositingEditor/Private/Sequencer/MediaTrackEditor.cpp b/Engine/Plugins/Media/MediaCompositing/Source/MediaCompositingEditor/Private/Sequencer/MediaTrackEditor.cpp
--- a/Engine/Plugins/Media/MediaCompositing/Source/MediaCompositingEditor/Private/Sequencer/MediaTrackEditor.cpp
+++ b/Engine/Plugins/Media/MediaCompositing/Source/MediaCompositingEditor/Private/Sequencer/MediaTrackEditor.cpp
@@ -16,7 +16,7 @@
 #include "MovieSceneMediaSection.h"
 #include "MovieSceneMediaTrack.h"
 #include "MovieSceneToolsUserSettings.h"
-#include "SequencerUtilities.h"
+#include "MVVM/Views/ViewUtilities.h"
 #include "TrackEditorThumbnail/TrackEditorThumbnailPool.h"
 #include "Widgets/SBoxPanel.h"
 #include "Widgets/Layout/SBox.h"
@@ -136,14 +136,7 @@ TSharedPtr<SWidget> FMediaTrackEditor::BuildOutlinerEditWidget(const FGuid& Obje
 		return MenuBuilder.MakeWidget();
 	};
 
-	return SNew(SHorizontalBox)
-
-	+ SHorizontalBox::Slot()
-		.AutoWidth()
-		.VAlign(VAlign_Center)
-		[
-			FSequencerUtilities::MakeAddButton(LOCTEXT("AddMediaSection_Text", "Media"), FOnGetContent::CreateLambda(CreatePicker), Params.NodeIsHovered, GetSequencer())
-		];
+	return UE::Sequencer::MakeAddButton(LOCTEXT("AddMediaSection_Text", "Media"), FOnGetContent::CreateLambda(CreatePicker), Params.ViewModel);
 }
 
 
diff --git a/Engine/Plugins/MovieScene/CustomizableSequencerTracks/Source/CustomizableSequencerTracksEditor/Private/SequencerTrackBPEditor.cpp b/Engine/Plugins/MovieScene/CustomizableSequencerTracks/Source/CustomizableSequencerTracksEditor/Private/SequencerTrackBPEditor.cpp
--- a/Engine/Plugins/MovieScene/CustomizableSequencerTracks/Source/CustomizableSequencerTracksEditor/Private/SequencerTrackBPEditor.cpp
+++ b/Engine/Plugins/MovieScene/CustomizableSequencerTracks/Source/CustomizableSequencerTracksEditor/Private/SequencerTrackBPEditor.cpp
@@ -8,7 +8,7 @@
 #include "ISequencerSection.h"
 #include "SequencerSectionPainter.h"
 #include "Misc/PackageName.h"
-#include "SequencerUtilities.h"
+#include "MVVM/Views/ViewUtilities.h"
 #include "CustomizableSequencerTracksStyle.h"
 
 #include "AssetRegistry/AssetRegistryModule.h"
@@ -238,13 +238,7 @@ TSharedPtr<SWidget> FSequencerTrackBPEditor::BuildOutlinerEditWidget(const FGuid
 			return MenuBuilder.MakeWidget();
 		};
 
-		return SNew(SHorizontalBox)
-			+ SHorizontalBox::Slot()
-			.AutoWidth()
-			.VAlign(VAlign_Center)
-			[
-				FSequencerUtilities::MakeAddButton(LOCTEXT("AddSectionText", "Section"), FOnGetContent::CreateLambda(SubMenuCallback), Params.NodeIsHovered, GetSequencer())
-			];
+		return UE::Sequencer::MakeAddButton(LOCTEXT("AddSectionText", "Section"), FOnGetContent::CreateLambda(SubMenuCallback), Params.ViewModel);
 	}
 
 
diff --git a/Engine/Plugins/MovieScene/TemplateSequence/Source/TemplateSequenceEditor/Private/TrackEditors/TemplateSequenceTrackEditor.cpp b/Engine/Plugins/MovieScene/TemplateSequence/Source/TemplateSequenceEditor/Private/TrackEditors/TemplateSequenceTrackEditor.cpp
--- a/Engine/Plugins/MovieScene/TemplateSequence/Source/TemplateSequenceEditor/Private/TrackEditors/TemplateSequenceTrackEditor.cpp
+++ b/Engine/Plugins/MovieScene/TemplateSequence/Source/TemplateSequenceEditor/Private/TrackEditors/TemplateSequenceTrackEditor.cpp
@@ -11,7 +11,7 @@
 #include "LevelSequence.h"
 #include "Misc/ConfigCacheIni.h"
 #include "Misc/PackageName.h"
-#include "SequencerUtilities.h"
+#include "MVVM/Views/ViewUtilities.h"
 #include "MovieScenePossessable.h"
 #include "TemplateSequence.h"
 #include "MovieSceneSpawnable.h"
@@ -75,16 +75,10 @@ TSharedPtr<SWidget> FTemplateSequenceTrackEditor::BuildOutlinerEditWidget(const
 
 	if (ObjectClass != nullptr)
 	{
-		return SNew(SHorizontalBox)
-			+ SHorizontalBox::Slot()
-			.AutoWidth()
-			.VAlign(VAlign_Center)
-			[
-				FSequencerUtilities::MakeAddButton(
+		return UE::Sequencer::MakeAddButton(
 					LOCTEXT("TemplateSequenceAddButton", "Template Sequence"),
 					FOnGetContent::CreateSP(this, &FTemplateSequenceTrackEditor::BuildTemplateSequenceAssetSubMenu, ObjectBinding, ObjectClass),
-					Params.NodeIsHovered, GetSequencer())
-			];
+					Params.ViewModel);
 	}
 	else
 	{
diff --git a/Engine/Plugins/Runtime/GameplayAbilities/Source/GameplayAbilitiesEditor/Private/Sequencer/GameplayCueTrackEditor.cpp b/Engine/Plugins/Runtime/GameplayAbilities/Source/GameplayAbilitiesEditor/Private/Sequencer/GameplayCueTrackEditor.cpp
--- a/Engine/Plugins/Runtime/GameplayAbilities/Source/GameplayAbilitiesEditor/Private/Sequencer/GameplayCueTrackEditor.cpp
+++ b/Engine/Plugins/Runtime/GameplayAbilities/Source/GameplayAbilitiesEditor/Private/Sequencer/GameplayCueTrackEditor.cpp
@@ -10,7 +10,7 @@
 #include "Tracks/MovieSceneEventTrack.h"
 #include "ISequencerSection.h"
 #include "Widgets/SBoxPanel.h"
-#include "SequencerUtilities.h"
+#include "MVVM/Views/ViewUtilities.h"
 #include "MovieSceneSequenceEditor.h"
 #include "LevelSequence.h"
 
@@ -120,13 +120,7 @@ TSharedPtr<SWidget> FGameplayCueTrackEditor::BuildOutlinerEditWidget(const FGuid
 		return MenuBuilder.MakeWidget();
 	};
 
-	return SNew(SHorizontalBox)
-	+ SHorizontalBox::Slot()
-	.AutoWidth()
-	.VAlign(VAlign_Center)
-	[
-		FSequencerUtilities::MakeAddButton(LOCTEXT("AddSection", "Section"), FOnGetContent::CreateLambda(SubMenuCallback), Params.NodeIsHovered, GetSequencer())
-	];
+	return UE::Sequencer::MakeAddButton(LOCTEXT("AddSection", "Section"), FOnGetContent::CreateLambda(SubMenuCallback), Params.ViewModel);
 }
 
 void FGameplayCueTrackEditor::BuildTrackContextMenu(FMenuBuilder& MenuBuilder, UMovieSceneTrack* Track)
diff --git a/Engine/Plugins/Runtime/GeometryCache/Source/GeometryCacheSequencer/Classes/GeometryCacheTrackEditor.cpp b/Engine/Plugins/Runtime/GeometryCache/Source/GeometryCacheSequencer/Classes/GeometryCacheTrackEditor.cpp
--- a/Engine/Plugins/Runtime/GeometryCache/Source/GeometryCacheSequencer/Classes/GeometryCacheTrackEditor.cpp
+++ b/Engine/Plugins/Runtime/GeometryCache/Source/GeometryCacheSequencer/Classes/GeometryCacheTrackEditor.cpp
@@ -9,7 +9,7 @@
 #include "SequencerSectionPainter.h"
 #include "MovieSceneGeometryCacheTrack.h"
 #include "MovieSceneGeometryCacheSection.h"
-#include "SequencerUtilities.h"
+#include "MVVM/Views/ViewUtilities.h"
 #include "Fonts/FontMeasure.h"
 #include "GeometryCacheComponent.h"
 #include "GeometryCache.h"
@@ -20,7 +20,7 @@
 namespace GeometryCacheEditorConstants
 {
 	// @todo Sequencer Allow this to be customizable
-	const uint32 AnimationTrackHeight = 20;
+	const uint32 AnimationTrackHeight = 28;
 }
 
 #define LOCTEXT_NAMESPACE "FGeometryCacheTrackEditor"
@@ -399,13 +399,7 @@ TSharedPtr<SWidget> FGeometryCacheTrackEditor::BuildOutlinerEditWidget(const FGu
 			return MenuBuilder.MakeWidget();
 		};
 
-		return SNew(SHorizontalBox)
-			+ SHorizontalBox::Slot()
-			.AutoWidth()
-			.VAlign(VAlign_Center)
-			[
-				FSequencerUtilities::MakeAddButton(LOCTEXT("GeomCacheText", "Geometry Cache"), FOnGetContent::CreateLambda(SubMenuCallback), Params.NodeIsHovered, GetSequencer())
-			];
+		return UE::Sequencer::MakeAddButton(LOCTEXT("GeomCacheText", "Geometry Cache"), FOnGetContent::CreateLambda(SubMenuCallback), Params.ViewModel);
 	}
 	else
 	{
diff --git a/Engine/Plugins/Runtime/HairStrands/Source/HairStrandsEditor/Private/GroomCacheTrackEditor.cpp b/Engine/Plugins/Runtime/HairStrands/Source/HairStrandsEditor/Private/GroomCacheTrackEditor.cpp
--- a/Engine/Plugins/Runtime/HairStrands/Source/HairStrandsEditor/Private/GroomCacheTrackEditor.cpp
+++ b/Engine/Plugins/Runtime/HairStrands/Source/HairStrandsEditor/Private/GroomCacheTrackEditor.cpp
@@ -11,14 +11,14 @@
 #include "MovieSceneGroomCacheTrack.h"
 #include "MovieSceneGroomCacheSection.h"
 #include "SequencerSectionPainter.h"
-#include "SequencerUtilities.h"
+#include "MVVM/Views/ViewUtilities.h"
 #include "Styling/SlateIconFinder.h"
 #include "TimeToPixel.h"
 
 namespace GroomCacheEditorConstants
 {
 	// @todo Sequencer Allow this to be customizable
-	const uint32 AnimationTrackHeight = 20;
+	const uint32 AnimationTrackHeight = 28;
 }
 
 #define LOCTEXT_NAMESPACE "FGroomCacheTrackEditor"
@@ -390,13 +390,7 @@ TSharedPtr<SWidget> FGroomCacheTrackEditor::BuildOutlinerEditWidget(const FGuid&
 			return MenuBuilder.MakeWidget();
 		};
 
-		return SNew(SHorizontalBox)
-			+ SHorizontalBox::Slot()
-			.AutoWidth()
-			.VAlign(VAlign_Center)
-			[
-				FSequencerUtilities::MakeAddButton(LOCTEXT("GroomCacheText", "Groom Cache"), FOnGetContent::CreateLambda(SubMenuCallback), Params.NodeIsHovered, GetSequencer())
-			];
+		return UE::Sequencer::MakeAddButton(LOCTEXT("GroomCacheText", "Groom Cache"), FOnGetContent::CreateLambda(SubMenuCallback), Params.ViewModel);
 	}
 	else
 	{
diff --git a/Engine/Plugins/VirtualProduction/DMX/DMXEngine/Source/DMXEditor/Private/Sequencer/DMXLibraryTrackEditor.cpp b/Engine/Plugins/VirtualProduction/DMX/DMXEngine/Source/DMXEditor/Private/Sequencer/DMXLibraryTrackEditor.cpp
--- a/Engine/Plugins/VirtualProduction/DMX/DMXEngine/Source/DMXEditor/Private/Sequencer/DMXLibraryTrackEditor.cpp
+++ b/Engine/Plugins/VirtualProduction/DMX/DMXEngine/Source/DMXEditor/Private/Sequencer/DMXLibraryTrackEditor.cpp
@@ -9,7 +9,7 @@
 #include "Sequencer/DMXLibrarySection.h"
 
 #include "Sections/MovieSceneParameterSection.h"
-#include "SequencerUtilities.h"
+#include "MVVM/Views/ViewUtilities.h"
 #include "MovieSceneSection.h"
 
 #include "Framework/MultiBox/MultiBoxBuilder.h"
@@ -161,7 +161,7 @@ TSharedPtr<SWidget> FDMXLibraryTrackEditor::BuildOutlinerEditWidget(const FGuid&
 	// Create combo button "+ Patch" to pick an asset to add
 	// sub menu content callback
 	FOnGetContent AddPatchMenuContent = FOnGetContent::CreateSP(this, &FDMXLibraryTrackEditor::OnGetAddPatchMenuContent, DMXTrack);
-	return FSequencerUtilities::MakeAddButton(LOCTEXT("AddPatchButton", "Patch"), AddPatchMenuContent, Params.NodeIsHovered, GetSequencer());
+	return UE::Sequencer::MakeAddButton(LOCTEXT("AddPatchButton", "Patch"), AddPatchMenuContent, Params.ViewModel);
 }
 
 TSharedRef<SWidget> FDMXLibraryTrackEditor::OnGetAddPatchMenuContent(UMovieSceneDMXLibraryTrack* DMXTrack)
diff --git a/Engine/Source/Editor/EditorStyle/Private/StarshipStyle.cpp b/Engine/Source/Editor/EditorStyle/Private/StarshipStyle.cpp
--- a/Engine/Source/Editor/EditorStyle/Private/StarshipStyle.cpp
+++ b/Engine/Source/Editor/EditorStyle/Private/StarshipStyle.cpp
@@ -2024,10 +2024,34 @@ void FStarshipEditorStyle::FStyle::SetupSequencerStyles()
 		BrighterHeader.B = FMath::Min(1.f, BrighterHeader.B * 2.0f);
 		BrighterHeader = BrighterHeader.HSVToLinearRGB();
 
-		Set("Sequencer.TableView.Row", FTableRowStyle(AlternatingTableRowStyle)
+		Set("Sequencer.Outliner.Row", FTableRowStyle(AlternatingTableRowStyle)
 			.SetUseParentRowBrush(true)
-			.SetParentRowBackgroundBrush(FSlateRoundedBoxBrush(FStyleColors::Header, 2.f, FStyleColors::Transparent, 1.f))
-			.SetParentRowBackgroundHoveredBrush(FSlateRoundedBoxBrush(BrighterHeader, 2.f)));
+			.SetParentRowBackgroundBrush(FSlateColorBrush(FStyleColors::Header))
+			.SetParentRowBackgroundHoveredBrush(FSlateColorBrush(BrighterHeader)));
+
+		Set("Sequencer.Outliner.Separator", new FSlateColorBrush(FStyleColors::Input));
+		Set("Sequencer.Outliner.Plus", new IMAGE_BRUSH_SVG("Sequencer/Column_Widgets/Plus", Icon14x14));
+		Set("Sequencer.Outliner.AddKey", new IMAGE_BRUSH_SVG("Sequencer/Column_Widgets/AddKey", Icon14x14));
+		Set("Sequencer.Outliner.NextKey", new IMAGE_BRUSH_SVG("Sequencer/Column_Widgets/NextKey", Icon14x14));
+		Set("Sequencer.Outliner.PreviousKey", new IMAGE_BRUSH_SVG("Sequencer/Column_Widgets/PreviousKey", Icon14x14));
+		Set("Sequencer.Outliner.CameraLock", new IMAGE_BRUSH_SVG("Sequencer/Column_Widgets/SequencerCamera", Icon14x14));
+
+		Set("Sequencer.Outliner.ColumnButton", FButtonStyle()
+				.SetNormal(FSlateNoResource())
+				.SetHovered(FSlateNoResource())
+				.SetPressed(FSlateNoResource())
+				.SetNormalPadding(FMargin(0,0,0,1))
+				.SetPressedPadding(FMargin(0,1,0,0)) );
+
+		Set("Sequencer.Outliner.ToggleButton", FCheckBoxStyle( GetWidgetStyle<FCheckBoxStyle>("ToggleButtonCheckbox"))
+			.SetUncheckedImage(FSlateRoundedBoxBrush(FStyleColors::Header, 4.0f, FStyleColors::Input, 1.0f))
+			.SetUncheckedHoveredImage(FSlateRoundedBoxBrush(FStyleColors::Hover, 4.0f, FStyleColors::Input, 1.0f))
+			.SetUncheckedPressedImage(FSlateRoundedBoxBrush(FStyleColors::Hover, 4.0f, FStyleColors::Input, 1.0f))
+			.SetCheckedImage(FSlateRoundedBoxBrush(FStyleColors::Primary, 4.0f, FStyleColors::Input, 1.0f))
+			.SetCheckedHoveredImage(FSlateRoundedBoxBrush(FStyleColors::PrimaryHover, 4.0f, FStyleColors::Input, 1.0f))
+			.SetCheckedPressedImage(FSlateRoundedBoxBrush(FStyleColors::PrimaryHover, 4.0f, FStyleColors::Input, 1.0f))
+			.SetPadding(FMargin(6.f, 1.f))
+		);
 
 		Set("Sequencer.IconKeySmartAuto", new IMAGE_BRUSH("Sequencer/IconKeySmartAuto", Icon12x12));
 		Set("Sequencer.IconKeyAuto", new IMAGE_BRUSH("Sequencer/IconKeyAuto", Icon12x12));
diff --git a/Engine/Source/Editor/MovieSceneTools/Private/Sections/CameraCutSection.cpp b/Engine/Source/Editor/MovieSceneTools/Private/Sections/CameraCutSection.cpp
--- a/Engine/Source/Editor/MovieSceneTools/Private/Sections/CameraCutSection.cpp
+++ b/Engine/Source/Editor/MovieSceneTools/Private/Sections/CameraCutSection.cpp
@@ -172,9 +172,9 @@ FText FCameraCutSection::GetSectionTitle() const
 	return HandleThumbnailTextBlockText();
 }
 
-float FCameraCutSection::GetSectionHeight() const
+float FCameraCutSection::GetSectionHeight(const UE::Sequencer::FViewDensityInfo& ViewDensity) const
 {
-	return FViewportThumbnailSection::GetSectionHeight() + 10.f;
+	return FViewportThumbnailSection::GetSectionHeight(ViewDensity) + 10.f;
 }
 
 FMargin FCameraCutSection::GetContentPadding() const
diff --git a/Engine/Source/Editor/MovieSceneTools/Private/Sections/CinematicShotSection.cpp b/Engine/Source/Editor/MovieSceneTools/Private/Sections/CinematicShotSection.cpp
--- a/Engine/Source/Editor/MovieSceneTools/Private/Sections/CinematicShotSection.cpp
+++ b/Engine/Source/Editor/MovieSceneTools/Private/Sections/CinematicShotSection.cpp
@@ -66,9 +66,9 @@ FText FCinematicShotSection::GetSectionTitle() const
 	return GetRenameVisibility() == EVisibility::Visible ? FText::GetEmpty() : HandleThumbnailTextBlockText();
 }
 
-float FCinematicShotSection::GetSectionHeight() const
+float FCinematicShotSection::GetSectionHeight(const UE::Sequencer::FViewDensityInfo& ViewDensity) const
 {
-	return FViewportThumbnailSection::GetSectionHeight() + 2*9.f;
+	return FViewportThumbnailSection::GetSectionHeight(ViewDensity) + 2*9.f;
 }
 
 FMargin FCinematicShotSection::GetContentPadding() const
diff --git a/Engine/Source/Editor/MovieSceneTools/Private/Sections/ThumbnailSection.cpp b/Engine/Source/Editor/MovieSceneTools/Private/Sections/ThumbnailSection.cpp
--- a/Engine/Source/Editor/MovieSceneTools/Private/Sections/ThumbnailSection.cpp
+++ b/Engine/Source/Editor/MovieSceneTools/Private/Sections/ThumbnailSection.cpp
@@ -179,7 +179,7 @@ float FThumbnailSection::GetSectionGripSize() const
 }
 
 
-float FThumbnailSection::GetSectionHeight() const
+float FThumbnailSection::GetSectionHeight(const UE::Sequencer::FViewDensityInfo& ViewDensity) const
 {
 	auto* Settings = GetDefault<UMovieSceneUserThumbnailSettings>();
 	if (Settings->bDrawThumbnails)
diff --git a/Engine/Source/Editor/MovieSceneTools/Private/TrackEditors/AudioTrackEditor.cpp b/Engine/Source/Editor/MovieSceneTools/Private/TrackEditors/AudioTrackEditor.cpp
--- a/Engine/Source/Editor/MovieSceneTools/Private/TrackEditors/AudioTrackEditor.cpp
+++ b/Engine/Source/Editor/MovieSceneTools/Private/TrackEditors/AudioTrackEditor.cpp
@@ -17,6 +17,7 @@
 #include "Framework/Application/SlateApplication.h"
 #include "Widgets/Layout/SBox.h"
 #include "SequencerSectionPainter.h"
+#include "MVVM/ViewModels/SequencerEditorViewModel.h"
 #include "Styling/AppStyle.h"
 #include "Editor/UnrealEdEngine.h"
 #include "Sound/SoundCue.h"
@@ -30,7 +31,7 @@
 #include "AudioDecompress.h"
 #include "IContentBrowserSingleton.h"
 #include "ContentBrowserModule.h"
-#include "SequencerUtilities.h"
+#include "MVVM/Views/ViewUtilities.h"
 #include "AssetRegistry/AssetRegistryModule.h"
 #include "ISectionLayoutBuilder.h"
 #include "MovieSceneToolHelpers.h"
@@ -705,13 +706,13 @@ FText FAudioSection::GetSectionToolTip() const
 	return FText::GetEmpty();
 }
 
-float FAudioSection::GetSectionHeight() const
+float FAudioSection::GetSectionHeight(const UE::Sequencer::FViewDensityInfo& ViewDensity) const
 {
 	if (UMovieSceneAudioTrack* Track = Section.GetTypedOuter<UMovieSceneAudioTrack>())
 	{
 		return Track->GetRowHeight();
 	}
-	return ISequencerSection::GetSectionHeight();
+	return ISequencerSection::GetSectionHeight(ViewDensity);
 }
 
 int32 FAudioSection::OnPaintSection( FSequencerSectionPainter& Painter ) const
@@ -724,7 +725,7 @@ int32 FAudioSection::OnPaintSection( FSequencerSectionPainter& Painter ) const
 		FSlateDrawElement::MakeViewport(
 			Painter.DrawElements,
 			++LayerId,
-			Painter.SectionGeometry.ToPaintGeometry(FVector2f(StoredXSize, GetSectionHeight() + 8.f), FSlateLayoutTransform(FVector2f(StoredXOffset, 0))),
+			Painter.SectionGeometry.ToPaintGeometry(FVector2f(StoredXSize, Painter.SectionGeometry.GetLocalSize().Y), FSlateLayoutTransform(FVector2f(StoredXOffset, 0))),
 			WaveformThumbnail,
 			(Painter.bParentEnabled ? ESlateDrawEffect::None : ESlateDrawEffect::DisabledEffect) | ESlateDrawEffect::NoGamma,
 			FLinearColor::White
@@ -836,7 +837,7 @@ void FAudioSection::Tick( const FGeometry& AllottedGeometry, const FGeometry& Pa
 			!FMath::IsNearlyEqual(DrawRange.GetUpperBoundValue(), StoredDrawRange.GetUpperBoundValue()) ||
 			XOffset != StoredXOffset || XSize != StoredXSize || Track->GetColorTint() != StoredColor ||
 			StoredSoundWave != SoundWave ||
-			StoredSectionHeight != GetSectionHeight() ||
+			StoredSectionHeight != GetSectionHeight(SequencerPin->GetViewModel()->GetViewDensity()) ||
 			StoredStartOffset != AudioSection->GetStartOffset() ||
 			bStoredLooping != AudioSection->GetLooping())
 		{
@@ -909,13 +910,18 @@ void FAudioSection::SlipSection(FFrameNumber SlipTime)
 void FAudioSection::RegenerateWaveforms(TRange<float> DrawRange, int32 XOffset, int32 XSize, const FColor& ColorTint, float DisplayScale)
 {
 	UMovieSceneAudioSection* AudioSection = Cast<UMovieSceneAudioSection>(&Section);
+	TSharedPtr<ISequencer> SequencerPin = Sequencer.Pin();
+	if (!SequencerPin)
+	{
+		return;
+	}
 
 	StoredDrawRange = DrawRange;
 	StoredXOffset = XOffset;
 	StoredXSize = XSize;
 	StoredColor = ColorTint;
 	StoredStartOffset = AudioSection->GetStartOffset();
-	StoredSectionHeight = GetSectionHeight();
+	StoredSectionHeight = GetSectionHeight(SequencerPin->GetViewModel()->GetViewDensity());
 	bStoredLooping = AudioSection->GetLooping();
 
 	if (DrawRange.IsDegenerate() || DrawRange.IsEmpty() || AudioSection->GetSound() == NULL)
@@ -1144,16 +1150,7 @@ TSharedRef<ISequencerSection> FAudioTrackEditor::MakeSectionInterface( UMovieSce
 
 TSharedPtr<SWidget> FAudioTrackEditor::BuildOutlinerEditWidget(const FGuid& ObjectBinding, UMovieSceneTrack* Track, const FBuildEditWidgetParams& Params)
 {
-	// Create a container edit box
-	return SNew(SHorizontalBox)
-
-	// Add the audio combo box
-	+ SHorizontalBox::Slot()
-	.AutoWidth()
-	.VAlign(VAlign_Center)
-	[
-		FSequencerUtilities::MakeAddButton(LOCTEXT("AudioText", "Audio"), FOnGetContent::CreateSP(this, &FAudioTrackEditor::BuildAudioSubMenu, FOnAssetSelected::CreateRaw(this, &FAudioTrackEditor::OnAudioAssetSelected, Track), FOnAssetEnterPressed::CreateRaw(this, &FAudioTrackEditor::OnAudioAssetEnterPressed, Track)), Params.NodeIsHovered, GetSequencer())
-	];
+	return UE::Sequencer::MakeAddButton(LOCTEXT("AudioText", "Audio"), FOnGetContent::CreateSP(this, &FAudioTrackEditor::BuildAudioSubMenu, FOnAssetSelected::CreateRaw(this, &FAudioTrackEditor::OnAudioAssetSelected, Track), FOnAssetEnterPressed::CreateRaw(this, &FAudioTrackEditor::OnAudioAssetEnterPressed, Track)), Params.ViewModel);
 }
 
 bool FAudioTrackEditor::HandleAssetAdded(UObject* Asset, const FGuid& TargetObjectGuid)
diff --git a/Engine/Source/Editor/MovieSceneTools/Private/TrackEditors/BindingLifetimeTrackEditor.cpp b/Engine/Source/Editor/MovieSceneTools/Private/TrackEditors/BindingLifetimeTrackEditor.cpp
--- a/Engine/Source/Editor/MovieSceneTools/Private/TrackEditors/BindingLifetimeTrackEditor.cpp
+++ b/Engine/Source/Editor/MovieSceneTools/Private/TrackEditors/BindingLifetimeTrackEditor.cpp
@@ -23,7 +23,7 @@
 #include "MovieSceneTrackEditor.h"
 #include "Sections/BindingLifetimeSection.h"
 #include "Widgets/SBoxPanel.h"
-#include "SequencerUtilities.h"
+#include "MVVM/Views/ViewUtilities.h"
 
 class ISequencerTrackEditor;
 
@@ -181,13 +181,7 @@ TSharedPtr<SWidget> FBindingLifetimeTrackEditor::BuildOutlinerEditWidget(const F
 		return FReply::Handled();
 	};
 
-	return SNew(SHorizontalBox)
-		+ SHorizontalBox::Slot()
-		.AutoWidth()
-		.VAlign(VAlign_Center)
-		[
-			FSequencerUtilities::MakeAddButton(LOCTEXT("AddSection", "Section"), FOnClicked::CreateLambda(OnClickedCallback), Params.NodeIsHovered, GetSequencer())
-		];
+	return UE::Sequencer::MakeAddButton(LOCTEXT("AddSection", "Section"), FOnClicked::CreateLambda(OnClickedCallback), Params.ViewModel);
 }
 
 
diff --git a/Engine/Source/Editor/MovieSceneTools/Private/TrackEditors/CVarTrackEditor.cpp b/Engine/Source/Editor/MovieSceneTools/Private/TrackEditors/CVarTrackEditor.cpp
--- a/Engine/Source/Editor/MovieSceneTools/Private/TrackEditors/CVarTrackEditor.cpp
+++ b/Engine/Source/Editor/MovieSceneTools/Private/TrackEditors/CVarTrackEditor.cpp
@@ -7,6 +7,7 @@
 #include "Styling/AppStyle.h"
 #include "Sections/MovieSceneCVarSection.h"
 #include "SequencerUtilities.h"
+#include "MVVM/Views/ViewUtilities.h"
 #include "MovieSceneToolHelpers.h"
 
 #include "SequencerSectionPainter.h"
@@ -50,17 +51,8 @@ TSharedPtr<SWidget> FCVarTrackEditor::BuildOutlinerEditWidget(const FGuid& Objec
 		FSequencerUtilities::CreateNewSection(Track, WeakSequencer, RowIndex, EMovieSceneBlendType::Absolute);
 		return FReply::Handled();
 	};
-	
-	// Create a container edit box
-	return SNew(SHorizontalBox)
-
-	// Add the camera combo box
-	+ SHorizontalBox::Slot()
-	.AutoWidth()
-	.VAlign(VAlign_Center)
-	[
-		FSequencerUtilities::MakeAddButton(LOCTEXT("VarText", "Section"), FOnClicked::CreateLambda(OnClickedCallback), Params.NodeIsHovered, GetSequencer())
-	];
+
+	return UE::Sequencer::MakeAddButton(LOCTEXT("VarText", "Section"), FOnClicked::CreateLambda(OnClickedCallback), Params.ViewModel);
 }
 
 
diff --git a/Engine/Source/Editor/MovieSceneTools/Private/TrackEditors/CameraCutTrackEditor.cpp b/Engine/Source/Editor/MovieSceneTools/Private/TrackEditors/CameraCutTrackEditor.cpp
--- a/Engine/Source/Editor/MovieSceneTools/Private/TrackEditors/CameraCutTrackEditor.cpp
+++ b/Engine/Source/Editor/MovieSceneTools/Private/TrackEditors/CameraCutTrackEditor.cpp
@@ -9,7 +9,8 @@
 #include "Sections/CameraCutSection.h"
 #include "Sections/MovieSceneCameraCutSection.h"
 #include "SequencerSettings.h"
-#include "SequencerUtilities.h"
+#include "MVVM/Views/ViewUtilities.h"
+#include "MVVM/ViewModels/OutlinerColumns/OutlinerColumnTypes.h"
 #include "TrackEditorThumbnail/TrackEditorThumbnailPool.h"
 #include "TrackInstances/MovieSceneCameraCutTrackInstance.h"
 #include "Tracks/MovieSceneCameraCutTrack.h"
@@ -186,41 +187,50 @@ void FCameraCutTrackEditor::HandleToggleAutoArrangeSectionsExecute(UMovieSceneCa
 	}
 }
 
-TSharedPtr<SWidget> FCameraCutTrackEditor::BuildOutlinerEditWidget(const FGuid& ObjectBinding, UMovieSceneTrack* Track, const FBuildEditWidgetParams& Params)
+TSharedPtr<SWidget> FCameraCutTrackEditor::BuildOutlinerColumnWidget(const FBuildColumnWidgetParams& Params, const FName& ColumnName)
 {
-	const FCameraCutTrackCommands& Commands = FCameraCutTrackCommands::Get();
-	// Create a container edit box
-	return SNew(SHorizontalBox)
-
-	// Add the camera combo box
-	+ SHorizontalBox::Slot()
-	.AutoWidth()
-	.VAlign(VAlign_Center)
-	[
-		FSequencerUtilities::MakeAddButton(LOCTEXT("CameraCutText", "Camera"), FOnGetContent::CreateSP(this, &FCameraCutTrackEditor::HandleAddCameraCutComboButtonGetMenuContent), Params.NodeIsHovered, GetSequencer())
-	]
-
-	+ SHorizontalBox::Slot()
-	.VAlign(VAlign_Center)
-	.HAlign(HAlign_Right)
-	.AutoWidth()
-	.Padding(4, 0, 0, 0)
-	[
-		SNew(SCheckBox)
-        .Style( &FAppStyle::Get().GetWidgetStyle<FCheckBoxStyle>("ToggleButtonCheckBoxAlt"))
-		.Type(ESlateCheckBoxType::CheckBox)
-		.Padding(FMargin(0.f))
+	using namespace UE::Sequencer;
+
+	if (ColumnName == FCommonOutlinerNames::Add)
+	{
+		return UE::Sequencer::MakeAddButton(
+			LOCTEXT("CameraCutText", "Camera"),
+			FOnGetContent::CreateSP(this, &FCameraCutTrackEditor::HandleAddCameraCutComboButtonGetMenuContent),
+			Params.ViewModel);
+	}
+
+	// Show the camera lock button in the Edit area if the nav column is disabled
+	const bool bEditColumn = (ColumnName == FCommonOutlinerNames::Edit && !Params.TreeViewRow->IsColumnVisible(FCommonOutlinerNames::Nav));
+	if (ColumnName == FCommonOutlinerNames::Nav || bEditColumn)
+	{
+		TSharedRef<SWidget> Button = SNew(SCheckBox)
+		.Style(FAppStyle::Get(), "Sequencer.Outliner.ToggleButton")
+		.Type(ESlateCheckBoxType::ToggleButton)
 		.IsFocusable(false)
 		.IsChecked(this, &FCameraCutTrackEditor::IsCameraLocked)
 		.OnCheckStateChanged(this, &FCameraCutTrackEditor::OnLockCameraClicked)
 		.ToolTipText(this, &FCameraCutTrackEditor::GetLockCameraToolTip)
-		.CheckedImage(FAppStyle::GetBrush("Sequencer.LockCamera"))
-		.CheckedHoveredImage(FAppStyle::GetBrush("Sequencer.LockCamera"))
-		.CheckedPressedImage(FAppStyle::GetBrush("Sequencer.LockCamera"))
-		.UncheckedImage(FAppStyle::GetBrush("Sequencer.UnlockCamera"))
-		.UncheckedHoveredImage(FAppStyle::GetBrush("Sequencer.UnlockCamera"))
-		.UncheckedPressedImage(FAppStyle::GetBrush("Sequencer.UnlockCamera"))
-	];
+		[
+			SNew(SImage)
+			.Image(FAppStyle::GetBrush("Sequencer.Outliner.CameraLock"))
+		];
+
+		if (bEditColumn)
+		{
+			return SNew(SBox)
+			.HAlign(HAlign_Left)
+			.Padding(4.f, 0.f)
+			[
+				Button
+			];
+		}
+		else
+		{
+			return Button;
+		}
+	}
+
+	return FMovieSceneTrackEditor::BuildOutlinerColumnWidget(Params, ColumnName);;
 }
 
 
diff --git a/Engine/Source/Editor/MovieSceneTools/Private/TrackEditors/CameraCutTrackEditor.h b/Engine/Source/Editor/MovieSceneTools/Private/TrackEditors/CameraCutTrackEditor.h
--- a/Engine/Source/Editor/MovieSceneTools/Private/TrackEditors/CameraCutTrackEditor.h
+++ b/Engine/Source/Editor/MovieSceneTools/Private/TrackEditors/CameraCutTrackEditor.h
@@ -56,7 +56,7 @@ public:
 	virtual void BindCommands(TSharedRef<FUICommandList> SequencerCommandBindings) override;
 	virtual void BuildAddTrackMenu(FMenuBuilder& MenuBuilder) override;
 	virtual void BuildTrackContextMenu(FMenuBuilder& MenuBuilder, UMovieSceneTrack* Track) override;
-	virtual TSharedPtr<SWidget> BuildOutlinerEditWidget(const FGuid& ObjectBinding, UMovieSceneTrack* Track, const FBuildEditWidgetParams& Params) override;
+	virtual TSharedPtr<SWidget> BuildOutlinerColumnWidget(const FBuildColumnWidgetParams& Params, const FName& ColumnName) override;
 	virtual TSharedRef<ISequencerSection> MakeSectionInterface(UMovieSceneSection& SectionObject, UMovieSceneTrack& Track, FGuid ObjectBinding) override;
 	virtual void OnRelease() override;
 	virtual bool SupportsType(TSubclassOf<UMovieSceneTrack> Type) const override;
diff --git a/Engine/Source/Editor/MovieSceneTools/Private/TrackEditors/CameraShakeSourceShakeTrackEditor.cpp b/Engine/Source/Editor/MovieSceneTools/Private/TrackEditors/CameraShakeSourceShakeTrackEditor.cpp
--- a/Engine/Source/Editor/MovieSceneTools/Private/TrackEditors/CameraShakeSourceShakeTrackEditor.cpp
+++ b/Engine/Source/Editor/MovieSceneTools/Private/TrackEditors/CameraShakeSourceShakeTrackEditor.cpp
@@ -57,7 +57,7 @@
 #include "Sections/MovieSceneCameraShakeSourceShakeSection.h"
 #include "Sections/MovieSceneCameraShakeSourceTriggerSection.h"
 #include "SequencerSectionPainter.h"
-#include "SequencerUtilities.h"
+#include "MVVM/Views/ViewUtilities.h"
 #include "SlotBase.h"
 #include "Styling/AppStyle.h"
 #include "Styling/CoreStyle.h"
@@ -341,30 +341,17 @@ TSharedPtr<SWidget> FCameraShakeSourceShakeTrackEditor::BuildOutlinerEditWidget(
 {
 	if (UMovieSceneCameraShakeSourceShakeTrack* ShakeTrack = Cast<UMovieSceneCameraShakeSourceShakeTrack>(Track))
 	{
-		return SNew(SHorizontalBox)
-			+SHorizontalBox::Slot()
-			.AutoWidth()
-			.VAlign(VAlign_Center)
-			[
-				FSequencerUtilities::MakeAddButton(
+		return UE::Sequencer::MakeAddButton(
 						LOCTEXT("AddShakeSourceShakeSection", "Camera Shake"),
 						FOnGetContent::CreateSP(this, &FCameraShakeSourceShakeTrackEditor::BuildCameraShakeSubMenu, ObjectBinding),
-						Params.NodeIsHovered, GetSequencer())
-			];
+						Params.ViewModel);
 	}
 	else
 	{
-		return SNew(SHorizontalBox)
-			+ SHorizontalBox::Slot()
-			.AutoWidth()
-			.VAlign(VAlign_Center)
-			[
-				FSequencerUtilities::MakeAddButton(
+		return UE::Sequencer::MakeAddButton(
 						LOCTEXT("AddSection", "Section"),
 						FOnGetContent::CreateSP(this, &FCameraShakeSourceShakeTrackEditor::BuildCameraShakeTracksMenu, ObjectBinding),
-						Params.NodeIsHovered, 
-						GetSequencer())
-			];
+						Params.ViewModel);
 	}
 }
 
diff --git a/Engine/Source/Editor/MovieSceneTools/Private/TrackEditors/CameraShakeTrackEditor.cpp b/Engine/Source/Editor/MovieSceneTools/Private/TrackEditors/CameraShakeTrackEditor.cpp
--- a/Engine/Source/Editor/MovieSceneTools/Private/TrackEditors/CameraShakeTrackEditor.cpp
+++ b/Engine/Source/Editor/MovieSceneTools/Private/TrackEditors/CameraShakeTrackEditor.cpp
@@ -41,7 +41,7 @@
 #include "MovieSceneTrack.h"
 #include "ScopedTransaction.h"
 #include "Sections/MovieSceneCameraShakeSection.h"
-#include "SequencerUtilities.h"
+#include "MVVM/Views/ViewUtilities.h"
 #include "SlotBase.h"
 #include "Templates/Casts.h"
 #include "Tracks/MovieSceneCameraShakeTrack.h"
@@ -206,16 +206,7 @@ void FCameraShakeTrackEditor::AddCameraShakeSubMenu(FMenuBuilder& MenuBuilder, T
 
 TSharedPtr<SWidget> FCameraShakeTrackEditor::BuildOutlinerEditWidget(const FGuid& ObjectBinding, UMovieSceneTrack* Track, const FBuildEditWidgetParams& Params)
 {
-	// Create a container edit box
-	return SNew(SHorizontalBox)
-
-	// Add the camera shake combo box
-	+ SHorizontalBox::Slot()
-	.AutoWidth()
-	.VAlign(VAlign_Center)
-	[
-		FSequencerUtilities::MakeAddButton(LOCTEXT("AddCameraShake", "Camera Shake"), FOnGetContent::CreateSP(this, &FCameraShakeTrackEditor::BuildCameraShakeSubMenu, ObjectBinding), Params.NodeIsHovered, GetSequencer())
-	];
+	return UE::Sequencer::MakeAddButton(LOCTEXT("AddCameraShake", "Camera Shake"), FOnGetContent::CreateSP(this, &FCameraShakeTrackEditor::BuildCameraShakeSubMenu, ObjectBinding), Params.ViewModel);
 }
 
 
diff --git a/Engine/Source/Editor/MovieSceneTools/Private/TrackEditors/CinematicShotTrackEditor.cpp b/Engine/Source/Editor/MovieSceneTools/Private/TrackEditors/CinematicShotTrackEditor.cpp
--- a/Engine/Source/Editor/MovieSceneTools/Private/TrackEditors/CinematicShotTrackEditor.cpp
+++ b/Engine/Source/Editor/MovieSceneTools/Private/TrackEditors/CinematicShotTrackEditor.cpp
@@ -11,7 +11,9 @@
 #include "Sections/CinematicShotSection.h"
 #include "Sections/MovieSceneCinematicShotSection.h"
 #include "SequencerSettings.h"
-#include "SequencerUtilities.h"
+#include "MVVM/Views/ViewUtilities.h"
+#include "MVVM/ViewModels/OutlinerColumns/OutlinerColumnTypes.h"
+#include "MVVM/Extensions/ITrackExtension.h"
 #include "TrackEditorThumbnail/TrackEditorThumbnailPool.h"
 #include "TrackInstances/MovieSceneCameraCutTrackInstance.h"
 #include "Tracks/MovieSceneCinematicShotTrack.h"
@@ -72,40 +74,50 @@ void FCinematicShotTrackEditor::OnRelease()
 /* ISequencerTrackEditor interface
  *****************************************************************************/
 
-TSharedPtr<SWidget> FCinematicShotTrackEditor::BuildOutlinerEditWidget(const FGuid& ObjectBinding, UMovieSceneTrack* Track, const FBuildEditWidgetParams& Params)
-{
-	// Create a container edit box
-	return SNew(SHorizontalBox)
-
-	// Add the camera combo box
-	+ SHorizontalBox::Slot()
-	.AutoWidth()
-	.VAlign(VAlign_Center)
-	[
-		FSequencerUtilities::MakeAddButton(LOCTEXT("CinematicShotText", "Shot"), FOnGetContent::CreateSP(this, &FCinematicShotTrackEditor::HandleAddSubSequenceComboButtonGetMenuContent, Track), Params.NodeIsHovered, GetSequencer())
-	]
-
-	+ SHorizontalBox::Slot()
-	.VAlign(VAlign_Center)
-	.HAlign(HAlign_Right)
-	.AutoWidth()
-	.Padding(4, 0, 0, 0)
-	[
-		SNew(SCheckBox)
-        .Style( &FAppStyle::Get().GetWidgetStyle<FCheckBoxStyle>("ToggleButtonCheckBoxAlt"))
-		.Type(ESlateCheckBoxType::CheckBox)
-		.Padding(FMargin(0.f))
+TSharedPtr<SWidget> FCinematicShotTrackEditor::BuildOutlinerColumnWidget(const FBuildColumnWidgetParams& Params, const FName& ColumnName)
+{
+	using namespace UE::Sequencer;
+
+	if (ColumnName == FCommonOutlinerNames::Add)
+	{
+		return UE::Sequencer::MakeAddButton(
+			LOCTEXT("CinematicShotText", "Shot"),
+			FOnGetContent::CreateSP(this, &FCinematicShotTrackEditor::HandleAddSubSequenceComboButtonGetMenuContent, Params.TrackModel->GetTrack()),
+			Params.ViewModel);
+	}
+
+	// Show the camera lock button in the Edit area if the nav column is disabled
+	const bool bEditColumn = (ColumnName == FCommonOutlinerNames::Edit && !Params.TreeViewRow->IsColumnVisible(FCommonOutlinerNames::Nav));
+	if (ColumnName == FCommonOutlinerNames::Nav || bEditColumn)
+	{
+		TSharedRef<SWidget> Button = SNew(SCheckBox)
+		.Style(FAppStyle::Get(), "Sequencer.Outliner.ToggleButton")
+		.Type(ESlateCheckBoxType::ToggleButton)
 		.IsFocusable(false)
 		.IsChecked(this, &FCinematicShotTrackEditor::AreShotsLocked)
 		.OnCheckStateChanged(this, &FCinematicShotTrackEditor::OnLockShotsClicked)
 		.ToolTipText(this, &FCinematicShotTrackEditor::GetLockShotsToolTip)
-		.CheckedImage(FAppStyle::GetBrush("Sequencer.LockCamera"))
-		.CheckedHoveredImage(FAppStyle::GetBrush("Sequencer.LockCamera"))
-		.CheckedPressedImage(FAppStyle::GetBrush("Sequencer.LockCamera"))
-		.UncheckedImage(FAppStyle::GetBrush("Sequencer.UnlockCamera"))
-		.UncheckedHoveredImage(FAppStyle::GetBrush("Sequencer.UnlockCamera"))
-		.UncheckedPressedImage(FAppStyle::GetBrush("Sequencer.UnlockCamera"))
-	];
+		[
+			SNew(SImage)
+				.Image(FAppStyle::GetBrush("Sequencer.Outliner.CameraLock"))
+		];
+
+		if (bEditColumn)
+		{
+			return SNew(SBox)
+			.HAlign(HAlign_Left)
+			.Padding(4.f, 0.f)
+			[
+				Button
+			];
+		}
+		else
+		{
+			return Button;
+		}
+	}
+
+	return FMovieSceneTrackEditor::BuildOutlinerColumnWidget(Params, ColumnName);
 }
 
 TSharedRef<ISequencerSection> FCinematicShotTrackEditor::MakeSectionInterface(UMovieSceneSection& SectionObject, UMovieSceneTrack& Track, FGuid ObjectBinding)
diff --git a/Engine/Source/Editor/MovieSceneTools/Private/TrackEditors/CustomPrimitiveDataTrackEditor.cpp b/Engine/Source/Editor/MovieSceneTools/Private/TrackEditors/CustomPrimitiveDataTrackEditor.cpp
--- a/Engine/Source/Editor/MovieSceneTools/Private/TrackEditors/CustomPrimitiveDataTrackEditor.cpp
+++ b/Engine/Source/Editor/MovieSceneTools/Private/TrackEditors/CustomPrimitiveDataTrackEditor.cpp
@@ -10,6 +10,7 @@
 #include "Tracks/MovieSceneCustomPrimitiveDataTrack.h"
 #include "Sections/ParameterSection.h"
 #include "SequencerUtilities.h"
+#include "MVVM/Views/ViewUtilities.h"
 #include "Modules/ModuleManager.h"
 #include "MaterialEditorModule.h"
 #include "Engine/Selection.h"
@@ -41,7 +42,7 @@ TSharedPtr<SWidget> FCustomPrimitiveDataTrackEditor::BuildOutlinerEditWidget(con
 	UMovieSceneCustomPrimitiveDataTrack* CPDTrack = Cast<UMovieSceneCustomPrimitiveDataTrack>(Track);
 	FOnGetContent MenuContent = FOnGetContent::CreateSP(this, &FCustomPrimitiveDataTrackEditor::OnGetAddMenuContent, ObjectBinding, CPDTrack, Params.TrackInsertRowIndex);
 
-	return FSequencerUtilities::MakeAddButton(LOCTEXT("AddParameterButton", "Parameter"), MenuContent, Params.NodeIsHovered, GetSequencer());
+	return UE::Sequencer::MakeAddButton(LOCTEXT("AddParameterButton", "Parameter"), MenuContent, Params.ViewModel);
 }
 
 
diff --git a/Engine/Source/Editor/MovieSceneTools/Private/TrackEditors/DataLayerTrackEditor.cpp b/Engine/Source/Editor/MovieSceneTools/Private/TrackEditors/DataLayerTrackEditor.cpp
--- a/Engine/Source/Editor/MovieSceneTools/Private/TrackEditors/DataLayerTrackEditor.cpp
+++ b/Engine/Source/Editor/MovieSceneTools/Private/TrackEditors/DataLayerTrackEditor.cpp
@@ -9,7 +9,7 @@
 #include "Styling/AppStyle.h"
 #include "Sections/MovieSceneDataLayerSection.h"
 #include "Tracks/MovieSceneDataLayerTrack.h"
-#include "SequencerUtilities.h"
+#include "MVVM/Views/ViewUtilities.h"
 #include "LevelUtils.h"
 #include "MovieSceneTimeHelpers.h"
 #include "MovieSceneToolHelpers.h"
@@ -301,10 +301,10 @@ void FDataLayerTrackEditor::BuildAddTrackMenu(FMenuBuilder& MenuBuilder)
 
 TSharedPtr<SWidget> FDataLayerTrackEditor::BuildOutlinerEditWidget(const FGuid& ObjectBinding, UMovieSceneTrack* Track, const FBuildEditWidgetParams& Params)
 {
-	return FSequencerUtilities::MakeAddButton(
+	return UE::Sequencer::MakeAddButton(
 		LOCTEXT("AddDataLayer_ButtonLabel", "Data Layer"),
 		FOnGetContent::CreateSP(this, &FDataLayerTrackEditor::BuildAddDataLayerMenu, Track),
-		Params.NodeIsHovered, GetSequencer());
+		Params.ViewModel);
 }
 
 UMovieSceneDataLayerSection* FDataLayerTrackEditor::AddNewSection(UMovieScene* MovieScene, UMovieSceneTrack* DataLayerTrack, EDataLayerRuntimeState DesiredState)
diff --git a/Engine/Source/Editor/MovieSceneTools/Private/TrackEditors/EventTrackEditor.cpp b/Engine/Source/Editor/MovieSceneTools/Private/TrackEditors/EventTrackEditor.cpp
--- a/Engine/Source/Editor/MovieSceneTools/Private/TrackEditors/EventTrackEditor.cpp
+++ b/Engine/Source/Editor/MovieSceneTools/Private/TrackEditors/EventTrackEditor.cpp
@@ -20,7 +20,7 @@
 #include "Sections/MovieSceneEventSection.h"
 #include "Sections/MovieSceneEventTriggerSection.h"
 #include "Sections/MovieSceneEventRepeaterSection.h"
-#include "SequencerUtilities.h"
+#include "MVVM/Views/ViewUtilities.h"
 #include "MovieSceneSequenceEditor.h"
 
 #define LOCTEXT_NAMESPACE "FEventTrackEditor"
@@ -159,13 +159,7 @@ TSharedPtr<SWidget> FEventTrackEditor::BuildOutlinerEditWidget(const FGuid& Obje
 		return MenuBuilder.MakeWidget();
 	};
 
-	return SNew(SHorizontalBox)
-	+ SHorizontalBox::Slot()
-	.AutoWidth()
-	.VAlign(VAlign_Center)
-	[
-		FSequencerUtilities::MakeAddButton(LOCTEXT("AddSection", "Section"), FOnGetContent::CreateLambda(SubMenuCallback), Params.NodeIsHovered, GetSequencer())
-	];
+	return UE::Sequencer::MakeAddButton(LOCTEXT("AddSection", "Section"), FOnGetContent::CreateLambda(SubMenuCallback), Params.ViewModel);
 }
 
 void FEventTrackEditor::BuildTrackContextMenu(FMenuBuilder& MenuBuilder, UMovieSceneTrack* Track)
diff --git a/Engine/Source/Editor/MovieSceneTools/Private/TrackEditors/LevelVisibilityTrackEditor.cpp b/Engine/Source/Editor/MovieSceneTools/Private/TrackEditors/LevelVisibilityTrackEditor.cpp
--- a/Engine/Source/Editor/MovieSceneTools/Private/TrackEditors/LevelVisibilityTrackEditor.cpp
+++ b/Engine/Source/Editor/MovieSceneTools/Private/TrackEditors/LevelVisibilityTrackEditor.cpp
@@ -6,7 +6,7 @@
 #include "Sections/MovieSceneLevelVisibilitySection.h"
 #include "Sections/LevelVisibilitySection.h"
 #include "Tracks/MovieSceneLevelVisibilityTrack.h"
-#include "SequencerUtilities.h"
+#include "MVVM/Views/ViewUtilities.h"
 #include "Engine/LevelStreaming.h"
 #include "Engine/World.h"
 #include "Misc/PackageName.h"
@@ -60,11 +60,10 @@ void FLevelVisibilityTrackEditor::BuildAddTrackMenu( FMenuBuilder& MenuBuilder )
 
 TSharedPtr<SWidget> FLevelVisibilityTrackEditor::BuildOutlinerEditWidget( const FGuid& ObjectBinding, UMovieSceneTrack* Track, const FBuildEditWidgetParams& Params )
 {
-	// Create a container edit box
-	return FSequencerUtilities::MakeAddButton( 
+	return UE::Sequencer::MakeAddButton(
 		LOCTEXT( "AddVisibilityTrigger", "Visibility Trigger" ),
 		FOnGetContent::CreateSP( this, &FLevelVisibilityTrackEditor::BuildAddVisibilityTriggerMenu, Track ),
-		Params.NodeIsHovered, GetSequencer() );
+		Params.ViewModel);
 }
 
 void FLevelVisibilityTrackEditor::BuildTrackContextMenu( FMenuBuilder& MenuBuilder, UMovieSceneTrack* Track )
diff --git a/Engine/Source/Editor/MovieSceneTools/Private/TrackEditors/MaterialParameterCollectionTrackEditor.cpp b/Engine/Source/Editor/MovieSceneTools/Private/TrackEditors/MaterialParameterCollectionTrackEditor.cpp
--- a/Engine/Source/Editor/MovieSceneTools/Private/TrackEditors/MaterialParameterCollectionTrackEditor.cpp
+++ b/Engine/Source/Editor/MovieSceneTools/Private/TrackEditors/MaterialParameterCollectionTrackEditor.cpp
@@ -37,6 +37,7 @@
 #include "Sections/MovieSceneParameterSection.h"
 #include "Sections/ParameterSection.h"
 #include "SequencerUtilities.h"
+#include "MVVM/Views/ViewUtilities.h"
 #include "Styling/SlateIconFinder.h"
 #include "Templates/Casts.h"
 #include "Textures/SlateIcon.h"
@@ -259,7 +260,7 @@ TSharedPtr<SWidget> FMaterialParameterCollectionTrackEditor::BuildOutlinerEditWi
 	UMovieSceneMaterialParameterCollectionTrack* MPCTrack = Cast<UMovieSceneMaterialParameterCollectionTrack>(Track);
 	FOnGetContent MenuContent = FOnGetContent::CreateSP(this, &FMaterialParameterCollectionTrackEditor::OnGetAddParameterMenuContent, MPCTrack, Params.RowIndex, Params.TrackInsertRowIndex);
 
-	return FSequencerUtilities::MakeAddButton(LOCTEXT("AddParameterButton", "Parameter"), MenuContent, Params.NodeIsHovered, GetSequencer());
+	return UE::Sequencer::MakeAddButton(LOCTEXT("AddParameterButton", "Parameter"), MenuContent, Params.ViewModel);
 }
 
 TSharedRef<SWidget> FMaterialParameterCollectionTrackEditor::OnGetAddParameterMenuContent(UMovieSceneMaterialParameterCollectionTrack* MPCTrack, int32 RowIndex, int32 TrackInsertRowIndex)
diff --git a/Engine/Source/Editor/MovieSceneTools/Private/TrackEditors/MaterialTrackEditor.cpp b/Engine/Source/Editor/MovieSceneTools/Private/TrackEditors/MaterialTrackEditor.cpp
--- a/Engine/Source/Editor/MovieSceneTools/Private/TrackEditors/MaterialTrackEditor.cpp
+++ b/Engine/Source/Editor/MovieSceneTools/Private/TrackEditors/MaterialTrackEditor.cpp
@@ -11,6 +11,7 @@
 #include "Sections/ComponentMaterialParameterSection.h"
 #include "Sections/ParameterSection.h"
 #include "SequencerUtilities.h"
+#include "MVVM/Views/ViewUtilities.h"
 #include "Modules/ModuleManager.h"
 #include "MaterialEditorModule.h"
 #include "Engine/Selection.h"
@@ -50,7 +51,7 @@ TSharedPtr<SWidget> FMaterialTrackEditor::BuildOutlinerEditWidget( const FGuid&
 	UMovieSceneMaterialTrack* MaterialTrack = Cast<UMovieSceneMaterialTrack>(Track);
 	FOnGetContent MenuContent = FOnGetContent::CreateSP(this, &FMaterialTrackEditor::OnGetAddMenuContent, ObjectBinding, MaterialTrack, Params.TrackInsertRowIndex);
 
-	return FSequencerUtilities::MakeAddButton(LOCTEXT( "AddParameterButton", "Parameter" ), MenuContent, Params.NodeIsHovered, GetSequencer());
+	return UE::Sequencer::MakeAddButton(LOCTEXT( "AddParameterButton", "Parameter" ), MenuContent, Params.ViewModel);
 }
 
 
diff --git a/Engine/Source/Editor/MovieSceneTools/Private/TrackEditors/ParticleParameterTrackEditor.cpp b/Engine/Source/Editor/MovieSceneTools/Private/TrackEditors/ParticleParameterTrackEditor.cpp
--- a/Engine/Source/Editor/MovieSceneTools/Private/TrackEditors/ParticleParameterTrackEditor.cpp
+++ b/Engine/Source/Editor/MovieSceneTools/Private/TrackEditors/ParticleParameterTrackEditor.cpp
@@ -7,7 +7,7 @@
 #include "Sections/MovieSceneParameterSection.h"
 #include "Tracks/MovieSceneParticleParameterTrack.h"
 #include "Sections/ParameterSection.h"
-#include "SequencerUtilities.h"
+#include "MVVM/Views/ViewUtilities.h"
 
 
 #define LOCTEXT_NAMESPACE "ParticleParameterTrackEditor"
@@ -38,11 +38,10 @@ TSharedRef<ISequencerSection> FParticleParameterTrackEditor::MakeSectionInterfac
 TSharedPtr<SWidget> FParticleParameterTrackEditor::BuildOutlinerEditWidget( const FGuid& ObjectBinding, UMovieSceneTrack* Track, const FBuildEditWidgetParams& Params )
 {
 	UMovieSceneParticleParameterTrack* ParticleParameterTrack = Cast<UMovieSceneParticleParameterTrack>(Track);
-	
-	// Create a container edit box
-	return FSequencerUtilities::MakeAddButton(LOCTEXT("ParameterText", "Parameter"),
+
+	return UE::Sequencer::MakeAddButton(LOCTEXT("ParameterText", "Parameter"),
 		FOnGetContent::CreateSP(this, &FParticleParameterTrackEditor::OnGetAddParameterMenuContent, ObjectBinding, ParticleParameterTrack),
-		Params.NodeIsHovered, GetSequencer());
+		Params.ViewModel);
 }
 
 
diff --git a/Engine/Source/Editor/MovieSceneTools/Private/TrackEditors/ParticleTrackEditor.cpp b/Engine/Source/Editor/MovieSceneTools/Private/TrackEditors/ParticleTrackEditor.cpp
--- a/Engine/Source/Editor/MovieSceneTools/Private/TrackEditors/ParticleTrackEditor.cpp
+++ b/Engine/Source/Editor/MovieSceneTools/Private/TrackEditors/ParticleTrackEditor.cpp
@@ -17,11 +17,12 @@
 #include "Particles/ParticleLODLevel.h"
 #include "Particles/ParticleModuleRequired.h"
 #include "TimeToPixel.h"
+#include "MVVM/ViewModels/ViewDensity.h"
 
 namespace AnimatableParticleEditorConstants
 {
 	// @todo Sequencer Allow this to be customizable
-	const uint32 ParticleTrackHeight = 20;
+	const uint32 ParticleTrackHeight = 28;
 }
 
 
@@ -45,9 +46,9 @@ UMovieSceneSection* FParticleSection::GetSectionObject()
 	return &Section;
 }
 
-float FParticleSection::GetSectionHeight() const
+float FParticleSection::GetSectionHeight(const UE::Sequencer::FViewDensityInfo& ViewDensity) const
 {
-	return (float)AnimatableParticleEditorConstants::ParticleTrackHeight;
+	return ViewDensity.UniformHeight.Get(AnimatableParticleEditorConstants::ParticleTrackHeight);
 }
 
 int32 FParticleSection::OnPaintSection( FSequencerSectionPainter& InPainter ) const
diff --git a/Engine/Source/Editor/MovieSceneTools/Private/TrackEditors/PropertyTrackEditors/EulerTransformPropertyTrackEditor.cpp b/Engine/Source/Editor/MovieSceneTools/Private/TrackEditors/PropertyTrackEditors/EulerTransformPropertyTrackEditor.cpp
--- a/Engine/Source/Editor/MovieSceneTools/Private/TrackEditors/PropertyTrackEditors/EulerTransformPropertyTrackEditor.cpp
+++ b/Engine/Source/Editor/MovieSceneTools/Private/TrackEditors/PropertyTrackEditors/EulerTransformPropertyTrackEditor.cpp
@@ -4,6 +4,7 @@
 #include "UnrealEdGlobals.h"
 #include "Editor/UnrealEdEngine.h"
 #include "SequencerUtilities.h"
+#include "MVVM/Views/ViewUtilities.h"
 #include "Sections/TransformPropertySection.h"
 #include "MovieSceneToolHelpers.h"
 #include "IKeyArea.h"
@@ -40,13 +41,7 @@ TSharedPtr<SWidget> FEulerTransformPropertyTrackEditor::BuildOutlinerEditWidget(
 		return MenuBuilder.MakeWidget();
 	};
 
-	return SNew(SHorizontalBox)
-	+ SHorizontalBox::Slot()
-	.AutoWidth()
-	.VAlign(VAlign_Center)
-	[
-		FSequencerUtilities::MakeAddButton(NSLOCTEXT("FEulerTransformPropertyTrackEditor", "AddSection", "Section"), FOnGetContent::CreateLambda(SubMenuCallback), Params.NodeIsHovered, GetSequencer())
-	];
+	return UE::Sequencer::MakeAddButton(NSLOCTEXT("FEulerTransformPropertyTrackEditor", "AddSection", "Section"), FOnGetContent::CreateLambda(SubMenuCallback), Params.ViewModel);
 }
 
 
diff --git a/Engine/Source/Editor/MovieSceneTools/Private/TrackEditors/PropertyTrackEditors/TransformPropertyTrackEditor.cpp b/Engine/Source/Editor/MovieSceneTools/Private/TrackEditors/PropertyTrackEditors/TransformPropertyTrackEditor.cpp
--- a/Engine/Source/Editor/MovieSceneTools/Private/TrackEditors/PropertyTrackEditors/TransformPropertyTrackEditor.cpp
+++ b/Engine/Source/Editor/MovieSceneTools/Private/TrackEditors/PropertyTrackEditors/TransformPropertyTrackEditor.cpp
@@ -6,6 +6,7 @@
 #include "Editor/UnrealEdEngine.h"
 #include "Sections/TransformPropertySection.h"
 #include "SequencerUtilities.h"
+#include "MVVM/Views/ViewUtilities.h"
 #include "MovieSceneToolHelpers.h"
 #include "Evaluation/MovieScenePropertyTemplate.h"
 
@@ -38,13 +39,7 @@ TSharedPtr<SWidget> FTransformPropertyTrackEditor::BuildOutlinerEditWidget(const
 		return MenuBuilder.MakeWidget();
 	};
 
-	return SNew(SHorizontalBox)
-	+ SHorizontalBox::Slot()
-	.AutoWidth()
-	.VAlign(VAlign_Center)
-	[
-		FSequencerUtilities::MakeAddButton(NSLOCTEXT("FTransformPropertyTrackEditor", "AddSection", "Section"), FOnGetContent::CreateLambda(SubMenuCallback), Params.NodeIsHovered, GetSequencer())
-	];
+	return UE::Sequencer::MakeAddButton(NSLOCTEXT("FTransformPropertyTrackEditor", "AddSection", "Section"), FOnGetContent::CreateLambda(SubMenuCallback), Params.ViewModel);
 }
 
 
diff --git a/Engine/Source/Editor/MovieSceneTools/Private/TrackEditors/SkeletalAnimationTrackEditor.cpp b/Engine/Source/Editor/MovieSceneTools/Private/TrackEditors/SkeletalAnimationTrackEditor.cpp
--- a/Engine/Source/Editor/MovieSceneTools/Private/TrackEditors/SkeletalAnimationTrackEditor.cpp
+++ b/Engine/Source/Editor/MovieSceneTools/Private/TrackEditors/SkeletalAnimationTrackEditor.cpp
@@ -23,7 +23,8 @@
 #include "IContentBrowserSingleton.h"
 #include "ContentBrowserModule.h"
 #include "SequencerSettings.h"
-#include "SequencerUtilities.h"
+#include "MVVM/Views/ViewUtilities.h"
+#include "MVVM/ViewModels/ViewDensity.h"
 #include "ISectionLayoutBuilder.h"
 #include "Animation/AnimMontage.h"
 #include "Animation/AnimSequence.h"
@@ -88,7 +89,7 @@ int32 FSkeletalAnimationTrackEditor::NumberActive = 0;
 namespace SkeletalAnimationEditorConstants
 {
 	// @todo Sequencer Allow this to be customizable
-	const uint32 AnimationTrackHeight = 20;
+	const uint32 AnimationTrackHeight = 28;
 }
 
 #define LOCTEXT_NAMESPACE "FSkeletalAnimationTrackEditor"
@@ -715,9 +716,9 @@ TOptional<FFrameTime> FSkeletalAnimationSection::GetSectionTime(FSequencerSectio
 }
 
 
-float FSkeletalAnimationSection::GetSectionHeight() const
+float FSkeletalAnimationSection::GetSectionHeight(const UE::Sequencer::FViewDensityInfo& ViewDensity) const
 {
-	return (float)SkeletalAnimationEditorConstants::AnimationTrackHeight;
+	return SkeletalAnimationEditorConstants::AnimationTrackHeight;
 }
 
 
@@ -2070,18 +2071,8 @@ TSharedPtr<SWidget> FSkeletalAnimationTrackEditor::BuildOutlinerEditWidget(const
 
 	if (Skeleton)
 	{
-		// Create a container edit box
-		return SNew(SHorizontalBox)
-
-		// Add the animation combo box
-		+ SHorizontalBox::Slot()
-		.AutoWidth()
-		.VAlign(VAlign_Center)
-		[
-			FSequencerUtilities::MakeAddButton(LOCTEXT("AnimationText", "Animation"), FOnGetContent::CreateSP(this, &FSkeletalAnimationTrackEditor::BuildAnimationSubMenu, ObjectBinding, Skeleton, Track), Params.NodeIsHovered, GetSequencer())
-		];
+		return UE::Sequencer::MakeAddButton(LOCTEXT("AnimationText", "Animation"), FOnGetContent::CreateSP(this, &FSkeletalAnimationTrackEditor::BuildAnimationSubMenu, ObjectBinding, Skeleton, Track), Params.ViewModel);
 	}
-
 	else
 	{
 		return TSharedPtr<SWidget>();
diff --git a/Engine/Source/Editor/MovieSceneTools/Private/TrackEditors/SubTrackEditor.cpp b/Engine/Source/Editor/MovieSceneTools/Private/TrackEditors/SubTrackEditor.cpp
--- a/Engine/Source/Editor/MovieSceneTools/Private/TrackEditors/SubTrackEditor.cpp
+++ b/Engine/Source/Editor/MovieSceneTools/Private/TrackEditors/SubTrackEditor.cpp
@@ -18,7 +18,7 @@
 #include "Tracks/MovieSceneCinematicShotTrack.h"
 #include "IContentBrowserSingleton.h"
 #include "ContentBrowserModule.h"
-#include "SequencerUtilities.h"
+#include "MVVM/Views/ViewUtilities.h"
 #include "SequencerSectionPainter.h"
 #include "TrackEditors/SubTrackEditorBase.h"
 #include "DragAndDrop/AssetDragDropOp.h"
@@ -215,16 +215,7 @@ void FSubTrackEditor::BuildAddTrackMenu(FMenuBuilder& MenuBuilder)
 
 TSharedPtr<SWidget> FSubTrackEditor::BuildOutlinerEditWidget(const FGuid& ObjectBinding, UMovieSceneTrack* Track, const FBuildEditWidgetParams& Params)
 {
-	// Create a container edit box
-	return SNew(SHorizontalBox)
-
-	// Add the sub sequence combo box
-	+ SHorizontalBox::Slot()
-	.AutoWidth()
-	.VAlign(VAlign_Center)
-	[
-		FSequencerUtilities::MakeAddButton(GetSubTrackName(), FOnGetContent::CreateSP(this, &FSubTrackEditor::HandleAddSubSequenceComboButtonGetMenuContent, Track), Params.NodeIsHovered, GetSequencer())
-	];
+	return UE::Sequencer::MakeAddButton(GetSubTrackName(), FOnGetContent::CreateSP(this, &FSubTrackEditor::HandleAddSubSequenceComboButtonGetMenuContent, Track), Params.ViewModel);
 }
 
 
diff --git a/Engine/Source/Editor/MovieSceneTools/Private/TrackEditors/TransformTrackEditor.cpp b/Engine/Source/Editor/MovieSceneTools/Private/TrackEditors/TransformTrackEditor.cpp
--- a/Engine/Source/Editor/MovieSceneTools/Private/TrackEditors/TransformTrackEditor.cpp
+++ b/Engine/Source/Editor/MovieSceneTools/Private/TrackEditors/TransformTrackEditor.cpp
@@ -23,6 +23,9 @@
 #include "Editor.h"
 #include "Sections/TransformPropertySection.h"
 #include "SequencerUtilities.h"
+#include "MVVM/Views/ViewUtilities.h"
+#include "MVVM/ViewModels/OutlinerColumns/OutlinerColumnTypes.h"
+#include "MVVM/ViewModels/ObjectBindingModel.h"
 #include "MovieSceneToolHelpers.h"
 #include "Animation/AnimData/IAnimationDataModel.h"
 
@@ -495,31 +498,32 @@ void F3DTransformTrackEditor::OnAddTransformKeysForSelectedObjects( EMovieSceneT
 	}
 }
 
-void F3DTransformTrackEditor::BuildObjectBindingEditButtons(TSharedPtr<SHorizontalBox> EditBox, const FGuid& ObjectGuid, const UClass* ObjectClass)
+void F3DTransformTrackEditor::BuildObjectBindingColumnWidgets(TFunctionRef<TSharedRef<SHorizontalBox>()> GetEditBox, const UE::Sequencer::TViewModelPtr<UE::Sequencer::FObjectBindingModel>& ObjectBinding, const UE::Sequencer::FCreateOutlinerViewParams& InParams, const FName& InColumnName)
 {
-	// If this is a camera track, add a button to lock the viewport to the camera
-	EditBox.Get()->AddSlot()
+	using namespace UE::Sequencer;
+
+	if (InColumnName == FCommonOutlinerNames::Nav || (InColumnName == FCommonOutlinerNames::Edit && !InParams.TreeViewRow->IsColumnVisible(FCommonOutlinerNames::Nav)))
+	{
+		FGuid ObjectGuid = ObjectBinding->GetObjectGuid();
+		GetEditBox()->AddSlot()
 		.VAlign(VAlign_Center)
-		.HAlign(HAlign_Right)
+		.HAlign(HAlign_Center)
 		.AutoWidth()
-		.Padding(4, 0, 0, 0)
 		[
-			SNew(SCheckBox)		
-				.Style( &FAppStyle::Get().GetWidgetStyle<FCheckBoxStyle>("ToggleButtonCheckBoxAlt"))
-				.Type(ESlateCheckBoxType::CheckBox)
-				.Padding(FMargin(0.f))
-				.IsFocusable(false)
-				.Visibility(this, &F3DTransformTrackEditor::IsCameraVisible, ObjectGuid)
-				.IsChecked(this, &F3DTransformTrackEditor::IsCameraLocked, ObjectGuid)
-				.OnCheckStateChanged(this, &F3DTransformTrackEditor::OnLockCameraClicked, ObjectGuid)
-				.ToolTipText(this, &F3DTransformTrackEditor::GetLockCameraToolTip, ObjectGuid)
-				.CheckedImage(FAppStyle::GetBrush("Sequencer.LockCamera"))
-				.CheckedHoveredImage(FAppStyle::GetBrush("Sequencer.LockCamera"))
-				.CheckedPressedImage(FAppStyle::GetBrush("Sequencer.LockCamera"))
-				.UncheckedImage(FAppStyle::GetBrush("Sequencer.UnlockCamera"))
-				.UncheckedHoveredImage(FAppStyle::GetBrush("Sequencer.UnlockCamera"))
-				.UncheckedPressedImage(FAppStyle::GetBrush("Sequencer.UnlockCamera"))
+			SNew(SCheckBox)
+			.Style(FAppStyle::Get(), "Sequencer.Outliner.ToggleButton")
+			.Type(ESlateCheckBoxType::ToggleButton)
+			.IsFocusable(false)
+			.Visibility(this, &F3DTransformTrackEditor::IsCameraVisible, ObjectGuid)
+			.IsChecked(this, &F3DTransformTrackEditor::IsCameraLocked, ObjectGuid)
+			.OnCheckStateChanged(this, &F3DTransformTrackEditor::OnLockCameraClicked, ObjectGuid)
+			.ToolTipText(this, &F3DTransformTrackEditor::GetLockCameraToolTip, ObjectGuid)
+			[
+				SNew(SImage)
+				.Image(FAppStyle::GetBrush("Sequencer.Outliner.CameraLock"))
+			]
 		];
+	}
 };
 void F3DTransformTrackEditor::BuildObjectBindingTrackMenu(FMenuBuilder& MenuBuilder, const TArray<FGuid>& ObjectBindings, const UClass* ObjectClass)
 {
diff --git a/Engine/Source/Editor/MovieSceneTools/Private/TrackEditors/TransformTrackEditor.h b/Engine/Source/Editor/MovieSceneTools/Private/TrackEditors/TransformTrackEditor.h
--- a/Engine/Source/Editor/MovieSceneTools/Private/TrackEditors/TransformTrackEditor.h
+++ b/Engine/Source/Editor/MovieSceneTools/Private/TrackEditors/TransformTrackEditor.h
@@ -55,7 +55,7 @@ public:
 
 	// ISequencerTrackEditor interface
 	virtual void ProcessKeyOperation(FFrameNumber InKeyTime, const UE::Sequencer::FKeyOperation& Operation, ISequencer& InSequencer) override;
-	virtual void BuildObjectBindingEditButtons(TSharedPtr<SHorizontalBox> EditBox, const FGuid& ObjectBinding, const UClass* ObjectClass) override;
+	virtual void BuildObjectBindingColumnWidgets(TFunctionRef<TSharedRef<SHorizontalBox>()> GetEditBox, const UE::Sequencer::TViewModelPtr<UE::Sequencer::FObjectBindingModel>& ObjectBinding, const UE::Sequencer::FCreateOutlinerViewParams& InParams, const FName& InColumnName) override;
 	virtual void BuildObjectBindingTrackMenu(FMenuBuilder& MenuBuilder, const TArray<FGuid>& ObjectBindinsg, const UClass* ObjectClass) override;
 	virtual TSharedRef<ISequencerSection> MakeSectionInterface( UMovieSceneSection& SectionObject, UMovieSceneTrack& Track, FGuid ObjectBinding ) override;
 	virtual void OnRelease() override;
diff --git a/Engine/Source/Editor/MovieSceneTools/Public/TrackEditors/CinematicShotTrackEditor.h b/Engine/Source/Editor/MovieSceneTools/Public/TrackEditors/CinematicShotTrackEditor.h
--- a/Engine/Source/Editor/MovieSceneTools/Public/TrackEditors/CinematicShotTrackEditor.h
+++ b/Engine/Source/Editor/MovieSceneTools/Public/TrackEditors/CinematicShotTrackEditor.h
@@ -51,7 +51,7 @@ public:
 	// ISequencerTrackEditor interface
 	virtual void OnInitialize() override;
 	virtual void OnRelease() override;
-	virtual TSharedPtr<SWidget> BuildOutlinerEditWidget(const FGuid& ObjectBinding, UMovieSceneTrack* Track, const FBuildEditWidgetParams& Params) override;
+	virtual TSharedPtr<SWidget> BuildOutlinerColumnWidget(const FBuildColumnWidgetParams& Params, const FName& ColumnName) override;
 	virtual TSharedRef<ISequencerSection> MakeSectionInterface(UMovieSceneSection& SectionObject, UMovieSceneTrack& Track, FGuid ObjectBinding) override;
 	virtual bool SupportsSequence(UMovieSceneSequence* InSequence) const override;
 	virtual void Tick(float DeltaTime) override;
diff --git a/Engine/Source/Editor/MovieSceneTools/Public/TrackEditors/SubTrackEditorBase.h b/Engine/Source/Editor/MovieSceneTools/Public/TrackEditors/SubTrackEditorBase.h
--- a/Engine/Source/Editor/MovieSceneTools/Public/TrackEditors/SubTrackEditorBase.h
+++ b/Engine/Source/Editor/MovieSceneTools/Public/TrackEditors/SubTrackEditorBase.h
@@ -137,7 +137,7 @@ public:
     virtual FText GetSectionTitle() const override;
 	virtual FText GetSectionToolTip() const override;
 	virtual TOptional<FFrameTime> GetSectionTime(FSequencerSectionPainter& InPainter) const override;
-	virtual float GetSectionHeight() const override;
+	virtual float GetSectionHeight(const UE::Sequencer::FViewDensityInfo& ViewDensity) const override;
     virtual bool IsReadOnly() const override;
     virtual int32 OnPaintSection( FSequencerSectionPainter& InPainter ) const override;
     virtual FReply OnSectionDoubleClicked(const FGeometry& SectionGeometry, const FPointerEvent& MouseEvent) override;
@@ -288,7 +288,7 @@ TOptional<FFrameTime> TSubSectionMixin<ParentSectionClass>::GetSectionTime(FSequ
 #undef LOCTEXT_NAMESPACE
 
 template<typename ParentSectionClass>
-float TSubSectionMixin<ParentSectionClass>::GetSectionHeight() const
+float TSubSectionMixin<ParentSectionClass>::GetSectionHeight(const UE::Sequencer::FViewDensityInfo& ViewDensity) const
 {
 	if (UMovieSceneSubTrack* Track = SubSectionObject.GetTypedOuter<UMovieSceneSubTrack>())
 	{
diff --git a/Engine/Source/Editor/Sequencer/Private/ISequencerSection.cpp b/Engine/Source/Editor/Sequencer/Private/ISequencerSection.cpp
--- a/Engine/Source/Editor/Sequencer/Private/ISequencerSection.cpp
+++ b/Engine/Source/Editor/Sequencer/Private/ISequencerSection.cpp
@@ -23,6 +23,7 @@
 #include "ISequencerChannelInterface.h"
 #include "ISequencerModule.h"
 #include "Modules/ModuleManager.h"
+#include "MVVM/ViewModels/ViewDensity.h"
 
 struct FMovieSceneChannel;
 
@@ -200,6 +201,25 @@ void ISequencerSection::GenerateSectionLayout( ISectionLayoutBuilder& LayoutBuil
 	}
 }
 
+float ISequencerSection::GetSectionHeight() const
+{
+	return SequencerSectionConstants::DefaultSectionHeight;
+}
+
+PRAGMA_DISABLE_DEPRECATION_WARNINGS
+float ISequencerSection::GetSectionHeight(const UE::Sequencer::FViewDensityInfo& ViewDensity) const
+{
+	// Call the deprecated method
+	float Height = GetSectionHeight();
+	if (Height != SequencerSectionConstants::DefaultSectionHeight)
+	{
+		// Override the uniform height for some sections
+		return Height;
+	}
+	return ViewDensity.UniformHeight.Get(Height);
+}
+PRAGMA_ENABLE_DEPRECATION_WARNINGS
+
 void ISequencerSection::ResizeSection(ESequencerSectionResizeMode ResizeMode, FFrameNumber ResizeFrameNumber)
 {
 	UMovieSceneSection* SectionObject = GetSectionObject();
diff --git a/Engine/Source/Editor/Sequencer/Private/MVVM/ViewModels/CategoryModel.cpp b/Engine/Source/Editor/Sequencer/Private/MVVM/ViewModels/CategoryModel.cpp
--- a/Engine/Source/Editor/Sequencer/Private/MVVM/ViewModels/CategoryModel.cpp
+++ b/Engine/Source/Editor/Sequencer/Private/MVVM/ViewModels/CategoryModel.cpp
@@ -14,6 +14,8 @@
 #include "MVVM/Views/SChannelView.h"
 #include "MVVM/Views/SOutlinerItemViewBase.h"
 #include "MVVM/Views/SSequencerKeyNavigationButtons.h"
+#include "MVVM/Views/SOutlinerTrackColorPicker.h"
+#include "MVVM/ViewModels/OutlinerColumns/OutlinerColumnTypes.h"
 #include "Math/UnrealMathUtility.h"
 #include "Misc/AssertionMacros.h"
 #include "MovieSceneTrack.h"
@@ -33,9 +35,7 @@ class SWidget;
 
 #define LOCTEXT_NAMESPACE "SequencerCategoryModel"
 
-namespace UE
-{
-namespace Sequencer
+namespace UE::Sequencer
 {
 
 FCategoryModel::FCategoryModel(FName InCategoryName)
@@ -159,7 +159,14 @@ FOutlinerSizing FCategoryGroupModel::GetOutlinerSizing() const
 	{
 		const_cast<FCategoryGroupModel*>(this)->RecomputeSizing();
 	}
-	return ComputedSizing;
+
+	FOutlinerSizing FinalSizing = ComputedSizing;
+	if (!EnumHasAnyFlags(ComputedSizing.Flags, EOutlinerSizingFlags::CustomHeight))
+	{
+		FViewDensityInfo Density = GetEditor()->GetViewDensity();
+		FinalSizing.Height = Density.UniformHeight.Get(FinalSizing.Height);
+	}
+	return FinalSizing;
 }
 
 FText FCategoryGroupModel::GetLabel() const
@@ -199,27 +206,35 @@ FSlateFontInfo FCategoryGroupModel::GetLabelFont() const
 		: FOutlinerItemModel::GetLabelFont();
 }
 
-TSharedRef<SWidget> FCategoryGroupModel::CreateOutlinerView(const FCreateOutlinerViewParams& InParams)
+TSharedPtr<SWidget> FCategoryGroupModel::CreateOutlinerViewForColumn(const FCreateOutlinerViewParams& InParams, const FName& InColumnName)
 {
-	TSharedPtr<FSequencerEditorViewModel> EditorViewModel = GetEditor();
-
-	return SNew(SOutlinerItemViewBase, SharedThis(this), InParams.Editor, InParams.TreeViewRow)
-		.CustomContent()
-		[
-			SNew(SBox)
-			.HAlign(HAlign_Right)
-			.VAlign(VAlign_Center)
-			[
-				SNew(SHorizontalBox)
-
-				+ SHorizontalBox::Slot()
-				.AutoWidth()
-				.VAlign(VAlign_Center)
-				[
-					SNew(SSequencerKeyNavigationButtons, SharedThis(this), EditorViewModel->GetSequencer())
-				]
-			]
-		];
+	TViewModelPtr<FSequencerEditorViewModel> Editor = InParams.Editor->CastThisShared<FSequencerEditorViewModel>();
+	if (!Editor)
+	{
+		return SNullWidget::NullWidget;
+	}
+
+	if (InColumnName == FCommonOutlinerNames::Label)
+	{
+		return SNew(SOutlinerItemViewBase, SharedThis(this), InParams.Editor, InParams.TreeViewRow);
+	}
+
+	if (InColumnName == FCommonOutlinerNames::Edit)
+	{
+
+	}
+
+	if (InColumnName == FCommonOutlinerNames::Nav)
+	{
+		return SNew(SSequencerKeyNavigationButtons, SharedThis(this), Editor->GetSequencer());
+	}
+
+	if (InColumnName == FCommonOutlinerNames::ColorPicker)
+	{
+		return SNew(SOutlinerTrackColorPicker, SharedThis(this), InParams.Editor);
+	}
+
+	return nullptr;
 }
 
 FTrackAreaParameters FCategoryGroupModel::GetTrackAreaParameters() const
@@ -261,8 +276,7 @@ void FCategoryGroupModel::OnRecycle()
 	Categories.Empty();
 }
 
-} // namespace Sequencer
-} // namespace UE
+} // namespace UE::Sequencer
 
 #undef LOCTEXT_NAMESPACE
 
diff --git a/Engine/Source/Editor/Sequencer/Private/MVVM/ViewModels/ChannelModel.cpp b/Engine/Source/Editor/Sequencer/Private/MVVM/ViewModels/ChannelModel.cpp
--- a/Engine/Source/Editor/Sequencer/Private/MVVM/ViewModels/ChannelModel.cpp
+++ b/Engine/Source/Editor/Sequencer/Private/MVVM/ViewModels/ChannelModel.cpp
@@ -21,10 +21,12 @@
 #include "MVVM/ViewModels/ViewModelIterators.h"
 #include "MVVM/ViewModels/OutlinerViewModel.h"
 #include "MVVM/ViewModels/TrackModel.h"
+#include "MVVM/ViewModels/OutlinerColumns/OutlinerColumnTypes.h"
 #include "MVVM/Selection/Selection.h"
 #include "MVVM/Views/SOutlinerView.h"
 #include "MVVM/Views/SOutlinerItemViewBase.h"
 #include "MVVM/Views/SSequencerKeyNavigationButtons.h"
+#include "MVVM/Views/SOutlinerTrackColorPicker.h"
 #include "MVVM/Views/STrackLane.h"
 #include "Math/UnrealMathUtility.h"
 #include "Misc/AssertionMacros.h"
@@ -49,9 +51,7 @@ class SWidget;
 
 #define LOCTEXT_NAMESPACE "SequencerChannelModel"
 
-namespace UE
-{
-namespace Sequencer
+namespace UE::Sequencer
 {
 
 FChannelModel::FChannelModel(FName InChannelName, TWeakPtr<ISequencerSection> InSection, FMovieSceneChannelHandle InChannel)
@@ -99,15 +99,19 @@ UMovieSceneSection* FChannelModel::GetSection() const
 
 FOutlinerSizing FChannelModel::GetDesiredSizing() const
 {
+	FOutlinerSizing Sizing(15.f);
+
 	if (KeyArea->ShouldShowCurve())
 	{
 		TViewModelPtr<FSequenceModel> Sequence = FindAncestorOfType<FSequenceModel>();
 		if (Sequence)
 		{
-			return Sequence->GetSequencer()->GetSequencerSettings()->GetKeyAreaHeightWithCurves();
+			Sizing.Height = Sequence->GetSequencer()->GetSequencerSettings()->GetKeyAreaHeightWithCurves();
+			Sizing.Flags  |= EOutlinerSizingFlags::CustomHeight;
 		}
 	}
-	return FOutlinerSizing(15.f);
+
+	return Sizing;
 }
 
 TSharedPtr<ITrackLaneWidget> FChannelModel::CreateTrackLaneView(const FCreateTrackLaneViewParams& InParams)
@@ -812,35 +816,46 @@ FOutlinerSizing FChannelGroupOutlinerModel::GetOutlinerSizing() const
 	{
 		const_cast<FChannelGroupOutlinerModel*>(this)->RecomputeSizing();
 	}
-	return ComputedSizing;
+
+	FOutlinerSizing FinalSizing = ComputedSizing;
+	if (!EnumHasAnyFlags(ComputedSizing.Flags, EOutlinerSizingFlags::CustomHeight))
+	{
+		FViewDensityInfo Density = GetEditor()->GetViewDensity();
+		FinalSizing.Height = Density.UniformHeight.Get(FinalSizing.Height);
+	}
+	return FinalSizing;
 }
 
-TSharedRef<SWidget> FChannelGroupOutlinerModel::CreateOutlinerView(const FCreateOutlinerViewParams& InParams)
+TSharedPtr<SWidget> FChannelGroupOutlinerModel::CreateOutlinerViewForColumn(const FCreateOutlinerViewParams& InParams, const FName& InColumnName)
 {
-	TSharedPtr<FSequencerEditorViewModel> EditorViewModel = GetEditor();
+	TViewModelPtr<FSequencerEditorViewModel> Editor = InParams.Editor->CastThisShared<FSequencerEditorViewModel>();
+	if (!Editor)
+	{
+		return SNullWidget::NullWidget;
+	}
 
-	return SNew(SOutlinerItemViewBase, SharedThis(this), InParams.Editor, InParams.TreeViewRow)
-	.CustomContent()
-	[
-		// Even if this key area node doesn't have any key areas right now, it may in the future
-		// so we always create the switcher, and just hide it if it is not relevant
-		SNew(SHorizontalBox)
-		.Visibility(this, &FChannelGroupOutlinerModel::GetKeyEditorVisibility)
+	if (InColumnName == FCommonOutlinerNames::Label)
+	{
+		return SNew(SOutlinerItemViewBase, SharedThis(this), InParams.Editor, InParams.TreeViewRow);
+	}
 
-		+ SHorizontalBox::Slot()
-		.HAlign(HAlign_Right)
-		.VAlign(VAlign_Center)
-		[
-			SNew(SKeyAreaEditorSwitcher, SharedThis(this), EditorViewModel->GetSequencer())
-		]
+	if (InColumnName == FCommonOutlinerNames::Edit)
+	{
+		return SNew(SKeyAreaEditorSwitcher, SharedThis(this), Editor)
+			.Visibility(this, &FChannelGroupOutlinerModel::GetKeyEditorVisibility);
+	}
+
+	if (InColumnName == FCommonOutlinerNames::Nav)
+	{
+		return SNew(SSequencerKeyNavigationButtons, SharedThis(this), Editor->GetSequencer());
+	}
 
-		+ SHorizontalBox::Slot()
-		.AutoWidth()
-		.VAlign(VAlign_Center)
-		[
-			SNew(SSequencerKeyNavigationButtons, SharedThis(this), EditorViewModel->GetSequencer())
-		]
-	];
+	if (InColumnName == FCommonOutlinerNames::ColorPicker)
+	{
+		return SNew(SOutlinerTrackColorPicker, SharedThis(this), InParams.Editor);
+	}
+
+	return nullptr;
 }
 
 EVisibility FChannelGroupOutlinerModel::GetKeyEditorVisibility() const
@@ -908,8 +923,7 @@ void FChannelGroupOutlinerModel::BuildContextMenu(FMenuBuilder& MenuBuilder)
 	BuildChannelOverrideMenu(MenuBuilder);
 }
 
-} // namespace Sequencer
-} // namespace UE
+} // namespace UE::Sequencer
 
 #undef LOCTEXT_NAMESPACE
 
diff --git a/Engine/Source/Editor/Sequencer/Private/MVVM/ViewModels/FolderModel.cpp b/Engine/Source/Editor/Sequencer/Private/MVVM/ViewModels/FolderModel.cpp
--- a/Engine/Source/Editor/Sequencer/Private/MVVM/ViewModels/FolderModel.cpp
+++ b/Engine/Source/Editor/Sequencer/Private/MVVM/ViewModels/FolderModel.cpp
@@ -15,6 +15,7 @@
 #include "MVVM/ObjectBindingModelStorageExtension.h"
 #include "MVVM/FolderModelStorageExtension.h"
 #include "MVVM/SharedViewModelData.h"
+#include "MVVM/ViewModels/OutlinerColumns/OutlinerColumnTypes.h"
 
 #include "Sequencer.h"
 #include "SequencerNodeTree.h"
@@ -35,9 +36,7 @@
 
 #define LOCTEXT_NAMESPACE "SequencerFolderModel"
 
-namespace UE
-{
-namespace Sequencer
+namespace UE::Sequencer
 {
 
 TMap<TWeakObjectPtr<UMovieSceneFolder>, FColor> FFolderModel::InitialFolderColors;
@@ -229,7 +228,9 @@ void FFolderModel::OnChildFolderRemoved(UMovieSceneFolder* Folder)
 
 FOutlinerSizing FFolderModel::GetOutlinerSizing() const
 {
-	return FOutlinerSizing(20.f, 4.f);
+	const float CompactHeight = 28.f;
+	FViewDensityInfo Density = GetEditor()->GetViewDensity();
+	return FOutlinerSizing(Density.UniformHeight.Get(CompactHeight));
 }
 
 void FFolderModel::GetIdentifierForGrouping(TStringBuilder<128>& OutString) const
@@ -321,7 +322,8 @@ void FFolderModel::BuildContextMenu(FMenuBuilder& MenuBuilder)
 	MenuBuilder.EndSection();
 }
 
-TSharedRef<SWidget> FFolderModel::CreateOutlinerView(const FCreateOutlinerViewParams& InParams)
+
+TSharedPtr<SWidget> FFolderModel::CreateOutlinerViewForColumn(const FCreateOutlinerViewParams& InParams, const FName& InColumnName)
 {
 	class SOutlinerFolderView
 		: public SOutlinerItemViewBase
@@ -344,7 +346,12 @@ TSharedRef<SWidget> FFolderModel::CreateOutlinerView(const FCreateOutlinerViewPa
 		TAttribute<bool> IsExpandedAttribute;
 	};
 
-	return SNew(SOutlinerFolderView, SharedThis(this), InParams.Editor, InParams.TreeViewRow);
+	if (InColumnName == FCommonOutlinerNames::Label)
+	{
+		return SNew(SOutlinerFolderView, SharedThis(this), InParams.Editor, InParams.TreeViewRow);
+	}
+
+	return nullptr;
 }
 
 void FFolderModel::SetFolderColor()
@@ -800,7 +807,6 @@ void FFolderModel::Delete()
 }
 
 
-} // namespace Sequencer
-} // namespace UE
+} // namespace UE::Sequencer
 
 #undef LOCTEXT_NAMESPACE
diff --git a/Engine/Source/Editor/Sequencer/Private/MVVM/ViewModels/ObjectBindingModel.cpp b/Engine/Source/Editor/Sequencer/Private/MVVM/ViewModels/ObjectBindingModel.cpp
--- a/Engine/Source/Editor/Sequencer/Private/MVVM/ViewModels/ObjectBindingModel.cpp
+++ b/Engine/Source/Editor/Sequencer/Private/MVVM/ViewModels/ObjectBindingModel.cpp
@@ -11,6 +11,7 @@
 #include "MVVM/ObjectBindingModelStorageExtension.h"
 #include "MVVM/ViewModels/SequencerEditorViewModel.h"
 #include "MVVM/ViewModels/OutlinerViewModelDragDropOp.h"
+#include "MVVM/ViewModels/OutlinerColumns/OutlinerColumnTypes.h"
 #include "MVVM/Views/SOutlinerObjectBindingView.h"
 #include "MVVM/Views/STrackLane.h"
 #include "MVVM/Selection/Selection.h"
@@ -45,16 +46,14 @@
 #include "SequencerCommands.h"
 #include "SequencerNodeTree.h"
 #include "SequencerSettings.h"
-#include "SequencerUtilities.h"
+#include "MVVM/Views/ViewUtilities.h"
 #include "Styling/AppStyle.h"
 #include "Styling/SlateIconFinder.h"
 #include "Widgets/SSequencerBindingLifetimeOverlay.h"
 
 #define LOCTEXT_NAMESPACE "ObjectBindingModel"
 
-namespace UE
-{
-namespace Sequencer
+namespace UE::Sequencer
 {
 
 namespace
@@ -266,7 +265,9 @@ FGuid FObjectBindingModel::GetObjectGuid() const
 
 FOutlinerSizing FObjectBindingModel::GetOutlinerSizing() const
 {
-	return FOutlinerSizing(20.f, 4.f);
+	const float CompactHeight = 28.f;
+	FViewDensityInfo Density = GetEditor()->GetViewDensity();
+	return FOutlinerSizing(Density.UniformHeight.Get(CompactHeight));
 }
 
 void FObjectBindingModel::GetIdentifierForGrouping(TStringBuilder<128>& OutString) const
@@ -274,19 +275,75 @@ void FObjectBindingModel::GetIdentifierForGrouping(TStringBuilder<128>& OutStrin
 	FOutlinerItemModel::GetIdentifier().ToString(OutString);
 }
 
-TSharedRef<SWidget> FObjectBindingModel::CreateOutlinerView(const FCreateOutlinerViewParams& InParams)
+TSharedPtr<SWidget> FObjectBindingModel::CreateOutlinerViewForColumn(const FCreateOutlinerViewParams& InParams, const FName& InColumnName)
 {
 	TSharedPtr<FSequencerEditorViewModel> EditorViewModel = GetEditor();
-	TSharedPtr<FSequencer> Sequencer = EditorViewModel->GetSequencerImpl();
+	TSharedPtr<FSequencer>                Sequencer       = EditorViewModel->GetSequencerImpl();
+
+	if (InColumnName == FCommonOutlinerNames::Label)
+	{
+		const FMovieSceneSequenceID SequenceID = OwnerModel->GetSequenceID();
+		const MovieScene::FFixedObjectBindingID FixedObjectBindingID(ObjectBindingID, SequenceID);
+
+		return SNew(SOutlinerItemViewBase, SharedThis(this), EditorViewModel, InParams.TreeViewRow)
+			.AdditionalLabelContent()
+			[
+				SNew(SObjectBindingTags, FixedObjectBindingID, Sequencer->GetObjectBindingTagCache())
+			];
+	}
+
+	if (InColumnName == FCommonOutlinerNames::Add)
+	{
+		return UE::Sequencer::MakeAddButton(
+			LOCTEXT("TrackText", "Track"),
+			FOnGetContent::CreateSP(this, &FObjectBindingModel::GetAddTrackMenuContent),
+			SharedThis(this));
+	}
+
+
+	// Ask track editors to populate the column.
+	// @todo: this is potentially very slow and will not scale as the number of track editors increases.
+	const bool bIsEditColumn = InColumnName == FCommonOutlinerNames::Edit;
+	TSharedPtr<SHorizontalBox> Box;
+
+	auto GetEditBox = [&Box]
+	{
+		if (!Box)
+		{
+			Box = SNew(SHorizontalBox);
 
-	const FMovieSceneSequenceID SequenceID = OwnerModel->GetSequenceID();
-	const MovieScene::FFixedObjectBindingID FixedObjectBindingID(ObjectBindingID, SequenceID);
+			auto CollapsedIfAllSlotsCollapsed = [Box]() -> EVisibility
+			{
+				for (int32 Index = 0; Index < Box->NumSlots(); ++Index)
+				{
+					EVisibility SlotVisibility = Box->GetSlot(Index).GetWidget()->GetVisibility();
+					if (SlotVisibility != EVisibility::Collapsed)
+					{
+						return EVisibility::SelfHitTestInvisible;
+					}
+				}
+				return EVisibility::Collapsed;
+			};
+
+			// Make the edit box collapsed if all of its slots are collapsed (or it has none)
+			Box->SetVisibility(MakeAttributeLambda(CollapsedIfAllSlotsCollapsed));
+		}
+		return Box.ToSharedRef();
+	};
+
+	for (const TSharedPtr<ISequencerTrackEditor>& TrackEditor : Sequencer->GetTrackEditors())
+	{
+		TrackEditor->BuildObjectBindingColumnWidgets(GetEditBox, SharedThis(this), InParams, InColumnName);
+
+		if (bIsEditColumn)
+		{
+			// Backwards compat
+			GetEditBox();
+			TrackEditor->BuildObjectBindingEditButtons(Box, ObjectBindingID, FindObjectClass());
+		}
+	}
 
-	return SNew(SOutlinerObjectBindingView, SharedThis(this), EditorViewModel, InParams.TreeViewRow)
-		.AdditionalLabelContent()
-		[
-			SNew(SObjectBindingTags, FixedObjectBindingID, Sequencer->GetObjectBindingTagCache())
-		];
+	return Box && Box->NumSlots() != 0 ? Box : nullptr;
 }
 
 bool FObjectBindingModel::GetDefaultExpansionState() const
@@ -1274,8 +1331,7 @@ void FObjectBindingModel::Delete()
 	}
 }
 
-} // namespace Sequencer
-} // namespace UE
+} // namespace UE::Sequencer
 
 #undef LOCTEXT_NAMESPACE
 
diff --git a/Engine/Source/Editor/Sequencer/Private/MVVM/ViewModels/SequenceModel.cpp b/Engine/Source/Editor/Sequencer/Private/MVVM/ViewModels/SequenceModel.cpp
--- a/Engine/Source/Editor/Sequencer/Private/MVVM/ViewModels/SequenceModel.cpp
+++ b/Engine/Source/Editor/Sequencer/Private/MVVM/ViewModels/SequenceModel.cpp
@@ -13,6 +13,7 @@
 #include "MVVM/ViewModels/SequencerEditorViewModel.h"
 #include "MVVM/ViewModels/ViewModel.h"
 #include "MVVM/ViewModels/ViewModelIterators.h"
+#include "MVVM/ViewModels/EditorSharedViewModelData.h"
 
 #include "MovieScene.h"
 #include "MovieSceneFolder.h"
@@ -42,7 +43,7 @@ FSequenceModel::FSequenceModel(TWeakPtr<FSequencerEditorViewModel> InEditorViewM
 
 void FSequenceModel::InitializeExtensions()
 {
-	TSharedPtr<FSharedViewModelData> NewSharedData = MakeShared<FSharedViewModelData>();
+	TSharedPtr<FSharedViewModelData> NewSharedData = MakeShared<FEditorSharedViewModelData>(WeakEditor.Pin().ToSharedRef());
 	SetSharedData(NewSharedData);
 
 	// Re-generate hierarchical caches when the sequence changes
diff --git a/Engine/Source/Editor/Sequencer/Private/MVVM/ViewModels/SequencerOutlinerViewModel.cpp b/Engine/Source/Editor/Sequencer/Private/MVVM/ViewModels/SequencerOutlinerViewModel.cpp
--- a/Engine/Source/Editor/Sequencer/Private/MVVM/ViewModels/SequencerOutlinerViewModel.cpp
+++ b/Engine/Source/Editor/Sequencer/Private/MVVM/ViewModels/SequencerOutlinerViewModel.cpp
@@ -20,6 +20,17 @@ FSequencerOutlinerViewModel::FSequencerOutlinerViewModel()
 {
 }
 
+void FSequencerOutlinerViewModel::RequestUpdate()
+{
+	FSequencerEditorViewModel* EditorViewModel = GetEditor()->CastThisChecked<FSequencerEditorViewModel>();
+	TSharedPtr<ISequencer> Sequencer = EditorViewModel->GetSequencer();
+
+	if (Sequencer)
+	{
+		Sequencer->RefreshTree();
+	}
+}
+
 TSharedPtr<SWidget> FSequencerOutlinerViewModel::CreateContextMenuWidget()
 {
 	FSequencerEditorViewModel* EditorViewModel = GetEditor()->CastThisChecked<FSequencerEditorViewModel>();
diff --git a/Engine/Source/Editor/Sequencer/Private/MVVM/ViewModels/TrackModel.cpp b/Engine/Source/Editor/Sequencer/Private/MVVM/ViewModels/TrackModel.cpp
--- a/Engine/Source/Editor/Sequencer/Private/MVVM/ViewModels/TrackModel.cpp
+++ b/Engine/Source/Editor/Sequencer/Private/MVVM/ViewModels/TrackModel.cpp
@@ -41,16 +41,8 @@
 
 #define LOCTEXT_NAMESPACE "TrackModel"
 
-namespace UE
+namespace UE::Sequencer
 {
-namespace Sequencer
-{
-
-namespace LayoutConstants
-{
-	extern const float CommonPadding;
-	const float CommonPadding = 4.f;
-}
 
 FTrackModel::FTrackModel(UMovieSceneTrack* Track)
 	: SectionList(EViewModelListType::TrackArea)
@@ -337,13 +329,15 @@ void FTrackModel::ForceUpdate()
 
 FOutlinerSizing FTrackModel::GetOutlinerSizing() const
 {
-	float Height = SequencerLayoutConstants::SectionAreaDefaultHeight;
+	FViewDensityInfo Density = GetEditor()->GetViewDensity();
+
+	float Height = Density.UniformHeight.Get(SequencerLayoutConstants::SectionAreaDefaultHeight);
 	for (TSharedPtr<FSectionModel> Section : SectionList.Iterate<FSectionModel>())
 	{
-		Height = Section->GetSectionInterface()->GetSectionHeight();
+		Height = Section->GetSectionInterface()->GetSectionHeight(Density);
 		break;
 	}
-	return FOutlinerSizing(Height + 2 * LayoutConstants::CommonPadding);
+	return FOutlinerSizing(Height);
 }
 
 void FTrackModel::GetIdentifierForGrouping(TStringBuilder<128>& OutString) const
@@ -590,12 +584,10 @@ FText FTrackModel::GetLabelToolTipText() const
 	return FText();
 }
 
-TSharedRef<SWidget> FTrackModel::CreateOutlinerView(const FCreateOutlinerViewParams& InParams)
+TSharedPtr<SWidget> FTrackModel::CreateOutlinerViewForColumn(const FCreateOutlinerViewParams& InParams, const FName& InColumnName)
 {
-	return SNew(SOutlinerTrackView, 
-			TWeakViewModelPtr<IOutlinerExtension>(SharedThis(this)),
-			InParams.Editor->CastThisSharedChecked<FSequencerEditorViewModel>(), 
-			InParams.TreeViewRow);
+	FBuildColumnWidgetParams Params(SharedThis(this), InParams);
+	return TrackEditor->BuildOutlinerColumnWidget(Params, InColumnName);
 }
 
 bool FTrackModel::IsResizable() const
@@ -608,14 +600,6 @@ void FTrackModel::Resize(float NewSize)
 {
 	UMovieSceneTrack* Track = GetTrack();
 
-	float PaddingAmount = 2 * LayoutConstants::CommonPadding;
-	if (Track)
-	{
-		PaddingAmount *= (Track->GetMaxRowIndex() + 1);
-	}
-	
-	NewSize -= PaddingAmount;
-
 	if (Track && TrackEditor->IsResizable(Track))
 	{
 		TrackEditor->Resize(NewSize, Track);
@@ -829,8 +813,7 @@ bool FTrackModel::FindBoundObjects(TArray<UObject*>& OutBoundObjects) const
 	return true;
 }
 
-} // namespace Sequencer
-} // namespace UE
+} // namespace UE::Sequencer
 
 #undef LOCTEXT_NAMESPACE
 
diff --git a/Engine/Source/Editor/Sequencer/Private/MVVM/ViewModels/TrackRowModel.cpp b/Engine/Source/Editor/Sequencer/Private/MVVM/ViewModels/TrackRowModel.cpp
--- a/Engine/Source/Editor/Sequencer/Private/MVVM/ViewModels/TrackRowModel.cpp
+++ b/Engine/Source/Editor/Sequencer/Private/MVVM/ViewModels/TrackRowModel.cpp
@@ -34,16 +34,9 @@
 
 #define LOCTEXT_NAMESPACE "TrackRowModel"
 
-namespace UE
-{
-namespace Sequencer
+namespace UE::Sequencer
 {
 
-namespace LayoutConstants
-{
-	extern const float CommonPadding;
-}
-
 FTrackRowModel::FTrackRowModel(UMovieSceneTrack* Track, int32 InRowIndex)
 	: SectionList(EViewModelListType::TrackArea)
 	, TopLevelChannelList(FTrackModel::GetTopLevelChannelGroupType())
@@ -96,13 +89,15 @@ FViewModelChildren FTrackRowModel::GetSectionModels()
 
 FOutlinerSizing FTrackRowModel::GetOutlinerSizing() const
 {
-	float Height = SequencerLayoutConstants::SectionAreaDefaultHeight;
+	FViewDensityInfo Density = GetEditor()->GetViewDensity();
+
+	float Height = Density.UniformHeight.Get(SequencerLayoutConstants::SectionAreaDefaultHeight);
 	for (TSharedPtr<FSectionModel> Section : SectionList.Iterate<FSectionModel>())
 	{
-		Height = Section->GetSectionInterface()->GetSectionHeight();
+		Height = Section->GetSectionInterface()->GetSectionHeight(Density);
 		break;
 	}
-	return FOutlinerSizing(Height + 2 * LayoutConstants::CommonPadding);
+	return FOutlinerSizing(Height);
 }
 
 FTrackAreaParameters FTrackRowModel::GetTrackAreaParameters() const
@@ -250,12 +245,10 @@ FSlateColor FTrackRowModel::GetLabelColor() const
 	return Track->GetLabelColor(LabelParams);
 }
 
-TSharedRef<SWidget> FTrackRowModel::CreateOutlinerView(const FCreateOutlinerViewParams& InParams)
+TSharedPtr<SWidget> FTrackRowModel::CreateOutlinerViewForColumn(const FCreateOutlinerViewParams& InParams, const FName& InColumnName)
 {
-	return SNew(SOutlinerTrackView, 
-			TWeakViewModelPtr<IOutlinerExtension>(SharedThis(this)), 
-			InParams.Editor->CastThisSharedChecked<FSequencerEditorViewModel>(), 
-			InParams.TreeViewRow);
+	FBuildColumnWidgetParams Params(SharedThis(this), InParams);
+	return TrackEditor->BuildOutlinerColumnWidget(Params, InColumnName);
 }
 
 bool FTrackRowModel::CanRename() const
@@ -303,14 +296,6 @@ bool FTrackRowModel::IsResizable() const
 void FTrackRowModel::Resize(float NewSize)
 {
 	UMovieSceneTrack* Track = GetTrack();
-	float PaddingAmount = 2 * LayoutConstants::CommonPadding;
-	if (Track)
-	{
-		PaddingAmount *= (Track->GetMaxRowIndex() + 1);
-	}
-	
-	NewSize -= PaddingAmount;
-
 	TSharedPtr<ISequencerTrackEditor> TrackEditorPtr = GetTrackEditor();
 	if (Track && TrackEditorPtr && TrackEditorPtr->IsResizable(Track))
 	{
@@ -443,8 +428,7 @@ void FTrackRowModel::Delete()
 	Track->FixRowIndices();
 }
 
-} // namespace Sequencer
-} // namespace UE
+} // namespace UE::Sequencer
 
 #undef LOCTEXT_NAMESPACE
 
diff --git a/Engine/Source/Editor/Sequencer/Private/MVVM/Views/SOutlinerObjectBindingView.cpp b/Engine/Source/Editor/Sequencer/Private/MVVM/Views/SOutlinerObjectBindingView.cpp
--- a/Engine/Source/Editor/Sequencer/Private/MVVM/Views/SOutlinerObjectBindingView.cpp
+++ b/Engine/Source/Editor/Sequencer/Private/MVVM/Views/SOutlinerObjectBindingView.cpp
@@ -35,40 +35,7 @@ void SOutlinerObjectBindingView::Construct(
 		TSharedPtr<FSequencerEditorViewModel> InEditor,
 		const TSharedRef<ISequencerTreeViewRow>& InTableRow)
 {
-	TSharedPtr<FSequencer> Sequencer = StaticCastSharedPtr<FSequencer>(InEditor->GetSequencer());
-
-	TAttribute<bool> HoverState = MakeAttributeSP(&InTableRow.Get(), &ISequencerTreeViewRow::IsHovered);
-
-	TSharedRef<SHorizontalBox> CustomBoxPanel =
-		SNew(SHorizontalBox)
-		+ SHorizontalBox::Slot()
-		[
-			SNew(SSpacer)
-		]
-
-		+ SHorizontalBox::Slot()
-		.AutoWidth()
-		.VAlign(VAlign_Center)
-		[
-			FSequencerUtilities::MakeAddButton(LOCTEXT("TrackText", "Track"), FOnGetContent::CreateSP(InModel.ToSharedRef(), &FObjectBindingModel::GetAddTrackMenuContent), HoverState, Sequencer->AsShared())
-		];
-
-	const UClass* ObjectClass = InModel->FindObjectClass();
-	Sequencer->BuildObjectBindingEditButtons(CustomBoxPanel, InModel->GetObjectGuid(), ObjectClass);
-
-	SOutlinerItemViewBase::Construct(
-		FArguments(InArgs)
-		.CustomContent()
-		[
-			CustomBoxPanel
-		]
-		.RightGutterContent()
-		[
-			// Spacer to match TrackView's color strip
-			SNew(SBox)
-			.WidthOverride(6.f)
-		], TWeakViewModelPtr<IOutlinerExtension>(InModel), InEditor, InTableRow
-	);
+	SOutlinerItemViewBase::Construct(InArgs, TWeakViewModelPtr<IOutlinerExtension>(InModel), InEditor, InTableRow);
 }
 
 } // namespace Sequencer
