Sequencer: Outliner UX improvements

diff --git a/Engine/Source/Editor/Sequencer/Private/MVVM/Views/SOutlinerTrackView.cpp b/Engine/Source/Editor/Sequencer/Private/MVVM/Views/SOutlinerTrackView.cpp
--- a/Engine/Source/Editor/Sequencer/Private/MVVM/Views/SOutlinerTrackView.cpp
+++ b/Engine/Source/Editor/Sequencer/Private/MVVM/Views/SOutlinerTrackView.cpp
@@ -1,42 +1,9 @@
 // Copyright Epic Games, Inc. All Rights Reserved.
 
 #include "MVVM/Views/SOutlinerTrackView.h"
-#include "MVVM/ViewModels/ViewModel.h"
-#include "MVVM/ICastable.h"
-#include "MVVM/Extensions/ITrackExtension.h"
-#include "MVVM/Extensions/IObjectBindingExtension.h"
-#include "MVVM/ViewModels/ChannelModel.h"
-#include "MVVM/ViewModels/TrackModel.h"
-#include "MVVM/ViewModels/ViewModelIterators.h"
-#include "MVVM/ViewModels/OutlinerViewModel.h"
 #include "MVVM/ViewModels/SequencerEditorViewModel.h"
-#include "MVVM/Views/STrackAreaView.h"
-#include "ISequencerTrackEditor.h"
-#include "SSequencerKeyNavigationButtons.h"
-#include "SKeyAreaEditorSwitcher.h"
 
-#include "MovieSceneTrack.h"
-#include "Tracks/MovieScenePropertyTrack.h"
-#include "Tracks/MovieScenePrimitiveMaterialTrack.h"
-#include "Tracks/MovieScene3DTransformTrack.h"
-
-#include "Widgets/DeclarativeSyntaxSupport.h"
-#include "Widgets/SBoxPanel.h"
-#include "Widgets/SOverlay.h"
-#include "Widgets/Images/SImage.h"
-#include "Widgets/Input/SButton.h"
-#include "Widgets/Layout/SSpacer.h"
-#include "Widgets/Layout/SBox.h"
-#include "Widgets/Colors/SColorPicker.h"
-
-#include "Engine/Engine.h"
-#include "Styling/AppStyle.h"
-
-#define LOCTEXT_NAMESPACE "SOutlinerTrackView"
-
-namespace UE
-{
-namespace Sequencer
+namespace UE::Sequencer
 {
 
 void SOutlinerTrackView::Construct(
@@ -45,252 +12,8 @@ void SOutlinerTrackView::Construct(
 		TWeakPtr<FSequencerEditorViewModel> InWeakEditor,
 		const TSharedRef<ISequencerTreeViewRow>& InTableRow)
 {
-	using namespace UE::Sequencer;
-
-	TSharedPtr<FViewModel> DataModel = InWeakOutlinerExtension.Pin().AsModel();
-	FSequencerEditorViewModel* EditorViewModel = ICastable::CastWeakPtr<FSequencerEditorViewModel>(InWeakEditor);
-	check(DataModel && EditorViewModel);
-
-	ITrackExtension* TrackExtension = DataModel->CastThis<ITrackExtension>();
-	checkf(TrackExtension != nullptr, TEXT("It is invalid to create a SOutlinerTrackView widget with a model that does not implement ITrackExtension"));
-
-	TSharedPtr<ISequencerTrackEditor> TrackEditor = EditorViewModel->GetSequencer()->GetTrackEditor(TrackExtension->GetTrack());
-
-	TOptional<FViewModelChildren> TopLevelChannels = DataModel->FindChildList(FTrackModel::GetTopLevelChannelGroupType());
-
-	TAttribute<bool> HoverState = MakeAttributeSP(&InTableRow.Get(), &ISequencerTreeViewRow::IsHovered);
-
-	TSharedRef<SHorizontalBox> BoxPanel = SNew(SHorizontalBox);
-
-	TSharedPtr<IObjectBindingExtension> ObjectBinding = DataModel->FindAncestorOfType<IObjectBindingExtension>();
-
-	UMovieSceneTrack* Track = TrackExtension->GetTrack();
-	int32 RowIndex = TrackExtension->GetRowIndex();
-
-	FBuildEditWidgetParams Params;
-	Params.NodeIsHovered = HoverState;
-	Params.RowIndex = RowIndex;
-	if (DataModel->FindAncestorOfType<ITrackExtension>())
-	{
-		Params.TrackInsertRowIndex = RowIndex;
-	}
-	else if (Track && Track->SupportsMultipleRows())
-	{
-		Params.TrackInsertRowIndex = Track->GetMaxRowIndex() + 1;
-	}
-
-	TSharedPtr<SWidget> CustomWidget = TrackEditor ? TrackEditor->BuildOutlinerEditWidget(ObjectBinding ? ObjectBinding->GetObjectGuid() : FGuid(), Track, Params) : nullptr;
-
-	TSharedPtr<FChannelGroupModel> TopLevelChannel = TopLevelChannels ? TopLevelChannels->FindFirstChildOfType<FChannelGroupModel>() : nullptr;
-	if (TopLevelChannel)
-	{
-		TSharedRef<SOverlay> Overlay = SNew(SOverlay);
-
-		Overlay->AddSlot()
-		.HAlign(HAlign_Left)
-		.VAlign(VAlign_Center)
-		[
-			SNew(SKeyAreaEditorSwitcher, TopLevelChannel, EditorViewModel->GetSequencer())
-		];
-
-		if (CustomWidget.IsValid())
-		{
-			Overlay->AddSlot()
-			.HAlign(HAlign_Right)
-			.VAlign(VAlign_Center)
-			[
-				CustomWidget.ToSharedRef()
-			];
-		}
-
-		BoxPanel->AddSlot()
-		.AutoWidth()
-		.VAlign(VAlign_Center)
-		[
-			Overlay
-		];
-
-		BoxPanel->AddSlot()
-		.AutoWidth()
-		.VAlign(VAlign_Center)
-		[
-			SNew(SSequencerKeyNavigationButtons, DataModel, EditorViewModel->GetSequencer())
-		];
-	}
-	else
-	{
-		if (CustomWidget.IsValid())
-		{
-			BoxPanel->AddSlot()
-			.AutoWidth()
-			.VAlign(VAlign_Center)
-			[
-				CustomWidget.ToSharedRef()
-			];
-		}
-
-		bool bHasKeyableAreas = false;
-		for (TSharedPtr<FChannelGroupModel> ChannelGroup : DataModel->GetDescendantsOfType<FChannelGroupModel>())
-		{
-			for (const TWeakViewModelPtr<FChannelModel>& WeakChannel : ChannelGroup->GetChannels())
-			{
-				if (TViewModelPtr<FChannelModel> Channel = WeakChannel.Pin())
-				{
-					if (Channel->GetKeyArea()->CanCreateKeyEditor())
-					{
-						bHasKeyableAreas = true;
-						break;
-					}
-				}
-			}
-			if (bHasKeyableAreas)
-			{
-				break;
-			}
-		}
-
-		if (bHasKeyableAreas)
-		{
-			BoxPanel->AddSlot()
-			.VAlign(VAlign_Center)
-			.AutoWidth()
-			[
-				SNew(SSequencerKeyNavigationButtons, DataModel, EditorViewModel->GetSequencer())
-			];
-		}
-	}
-
-	SOutlinerItemViewBase::Construct(
-		FArguments(InArgs)
-		.CustomContent()
-		[
-			SNew(SBox)
-			.VAlign(VAlign_Center)
-			.HAlign(HAlign_Right)
-			[
-				BoxPanel
-			]
-		]
-		.RightGutterContent()
-		[
-			SNew(SButton)
-			.ContentPadding(0)
-			.VAlign(VAlign_Fill)
-			.IsFocusable(false) // Intentionally false so that it's easier to tab to the next numeric input
-			.IsEnabled(!IsReadOnlyAttribute.Get())
-			.ButtonStyle(FAppStyle::Get(), "Sequencer.AnimationOutliner.ColorStrip")
-			.OnClicked(this, &SOutlinerTrackView::OnSetTrackColor)
-			.Content()
-			[
-				SNew(SBox)
-				.WidthOverride(6.f)
-				[
-					SNew(SImage)
-					.Image(FAppStyle::GetBrush("WhiteBrush"))
-					.ColorAndOpacity(this, &SOutlinerTrackView::GetTrackColorTint)
-				]
-			]
-		], InWeakOutlinerExtension, InWeakEditor, InTableRow);
-}
-
-// We store these for when the Color Picker is canceled so we can restore the old value.
-namespace AnimationOutlinerTreeNode
-{
-	FLinearColor InitialTrackColor;
-	bool bFolderPickerWasCancelled;
+	SOutlinerItemViewBase::Construct(InArgs, InWeakOutlinerExtension, InWeakEditor, InTableRow);
 }
 
-FReply SOutlinerTrackView::OnSetTrackColor()
-{
-	AnimationOutlinerTreeNode::InitialTrackColor = GetTrackColorTint().GetSpecifiedColor();
-	AnimationOutlinerTreeNode::bFolderPickerWasCancelled = false;
-
-	FColorPickerArgs PickerArgs;
-	PickerArgs.bUseAlpha = false;
-	PickerArgs.DisplayGamma = TAttribute<float>::Create(TAttribute<float>::FGetter::CreateUObject(GEngine, &UEngine::GetDisplayGamma));
-	PickerArgs.InitialColor = AnimationOutlinerTreeNode::InitialTrackColor;
-	PickerArgs.ParentWidget = GetParentWidget();
-	PickerArgs.OnColorCommitted = FOnLinearColorValueChanged::CreateSP(this, &SOutlinerTrackView::OnColorPickerPicked);
-	PickerArgs.OnColorPickerWindowClosed = FOnWindowClosed::CreateSP(this, &SOutlinerTrackView::OnColorPickerClosed);
-	PickerArgs.OnColorPickerCancelled = FOnColorPickerCancelled::CreateSP(this, &SOutlinerTrackView::OnColorPickerCancelled);
-
-	OpenColorPicker(PickerArgs);
-	return FReply::Handled();
-}
-
-void SOutlinerTrackView::OnColorPickerPicked(FLinearColor NewFolderColor)
-{
-	UMovieSceneTrack* Track = GetTrackFromNode();
-	if (Track)
-	{
-		// This is called every time the user adjusts the UI so we don't want to create a transaction for it, just directly
-		// modify the track so we can see the change immediately.
-		Track->SetColorTint(NewFolderColor.ToFColor(true));
-	}
-}
-
-void SOutlinerTrackView::OnColorPickerClosed(const TSharedRef<SWindow>& Window)
-{
-	// Under Unreal UX terms, closing the Color Picker (via the UI) is the same as confirming it since we've been live updating
-	// the color. The track already has the latest color change so we undo the change before calling Modify so that Undo sets us
-	// to the original color. This is also called in the event of pressing cancel so we need to detect if it was canceled or not.
-	if (!AnimationOutlinerTreeNode::bFolderPickerWasCancelled)
-	{
-		UMovieSceneTrack* Track = GetTrackFromNode();
-		if(Track)
-		{
-			const FScopedTransaction Transaction(LOCTEXT("SetTrackColor", "Set Track Color"));
-			FSlateColor CurrentColor = GetTrackColorTint();
-			Track->SetColorTint(AnimationOutlinerTreeNode::InitialTrackColor.ToFColor(true));
-			Track->Modify();
-			Track->SetColorTint(CurrentColor.GetSpecifiedColor().ToFColor(true));
-		}
-	}
-}
-
-void SOutlinerTrackView::OnColorPickerCancelled(FLinearColor NewFolderColor)
-{
-	AnimationOutlinerTreeNode::bFolderPickerWasCancelled = true;
-
-	// Restore the original color of the track. No transaction will be created when the OnColorPickerClosed callback is called.
-	UMovieSceneTrack* Track = GetTrackFromNode();
-	if (Track)
-	{
-		Track->SetColorTint(AnimationOutlinerTreeNode::InitialTrackColor.ToFColor(true));
-	}
-}
-
-UMovieSceneTrack* SOutlinerTrackView::GetTrackFromNode() const
-{
-	TViewModelPtr<IOutlinerExtension> Outliner = WeakOutlinerExtension.Pin();
-	if (Outliner)
-	{
-		constexpr bool bIncludeThis = true;
-		for (TSharedPtr<ITrackExtension> TrackExtension : Outliner.AsModel()->GetAncestorsOfType<ITrackExtension>(bIncludeThis))
-		{
-			if (UMovieSceneTrack* Track = TrackExtension->GetTrack())
-			{
-				return Track;
-			}
-		}
-	}
-
-	return nullptr;
-}
-
-FSlateColor SOutlinerTrackView::GetTrackColorTint() const
-{
-	UMovieSceneTrack* Track = GetTrackFromNode();
-	if (Track)
-	{
-		return STrackAreaView::BlendDefaultTrackColor(Track->GetColorTint());
-	}
-
-	return FLinearColor::Transparent;
-}
-
-} // namespace Sequencer
-} // namespace UE
-
-#undef LOCTEXT_NAMESPACE
+} // namespace UE::Sequencer
 
diff --git a/Engine/Source/Editor/Sequencer/Private/MVVM/Views/SSequencerOutlinerView.cpp b/Engine/Source/Editor/Sequencer/Private/MVVM/Views/SSequencerOutlinerView.cpp
--- a/Engine/Source/Editor/Sequencer/Private/MVVM/Views/SSequencerOutlinerView.cpp
+++ b/Engine/Source/Editor/Sequencer/Private/MVVM/Views/SSequencerOutlinerView.cpp
@@ -1,6 +1,7 @@
 // Copyright Epic Games, Inc. All Rights Reserved.
 
 #include "MVVM/Views/SSequencerOutlinerView.h"
+#include "MVVM/Views/SOutlinerViewRow.h"
 #include "MVVM/Extensions/IOutlinerExtension.h"
 #include "MVVM/ViewModels/OutlinerSpacer.h"
 #include "MVVM/ViewModels/OutlinerViewModel.h"
@@ -77,6 +78,7 @@ TSharedRef<ITableRow> SSequencerOutlinerView::OnGenerateRow(TWeakViewModelPtr<IO
 	TSharedRef<SOutlinerViewRow> Row =
 		SNew(SSequencerOutlinerViewRow, OwnerTable, InWeakModel, WeakOutliner)
 		.OnDetectDrag(this, &SSequencerOutlinerView::OnDragRow)
+		.OnGetColumnVisibility(this, &SSequencerOutlinerView::IsColumnVisible)
 		.OnGenerateWidgetForColumn(this, &SSequencerOutlinerView::GenerateWidgetForColumn);
 
 	if (TViewModelPtr<IOutlinerExtension> ViewModel = InWeakModel.Pin())
diff --git a/Engine/Source/Editor/Sequencer/Private/MovieSceneTrackEditor.cpp b/Engine/Source/Editor/Sequencer/Private/MovieSceneTrackEditor.cpp
--- a/Engine/Source/Editor/Sequencer/Private/MovieSceneTrackEditor.cpp
+++ b/Engine/Source/Editor/Sequencer/Private/MovieSceneTrackEditor.cpp
@@ -14,8 +14,13 @@
 #include "Sequencer.h"
 #include "Framework/MultiBox/MultiBoxBuilder.h"
 #include "SequencerUtilities.h"
+#include "MVVM/Views/ViewUtilities.h"
 #include "Framework/Application/SlateApplication.h"
-
+#include "MVVM/ViewModels/OutlinerColumns/OutlinerColumnTypes.h"
+#include "MVVM/Views/SSequencerKeyNavigationButtons.h"
+#include "MVVM/Views/SOutlinerTrackColorPicker.h"
+#include "MVVM/Views/SOutlinerItemViewBase.h"
+#include "SKeyAreaEditorSwitcher.h"
 
 FMovieSceneTrackEditor::FMovieSceneTrackEditor(TSharedRef<ISequencer> InSequencer)
 	: Sequencer(InSequencer)
@@ -185,14 +190,99 @@ void FMovieSceneTrackEditor::BuildAddTrackMenu(FMenuBuilder& MenuBuilder)
 { 
 }
 
-void FMovieSceneTrackEditor::BuildObjectBindingEditButtons(TSharedPtr<SHorizontalBox> EditBox, const FGuid& ObjectBinding, const UClass* ObjectClass) 
-{ 
+void FMovieSceneTrackEditor::BuildObjectBindingColumnWidgets(TFunctionRef<TSharedRef<SHorizontalBox>()> GetEditBox, const UE::Sequencer::TViewModelPtr<UE::Sequencer::FObjectBindingModel>& ObjectBinding, const UE::Sequencer::FCreateOutlinerViewParams& InParams, const FName& InColumnName)
+{
 }
 
 void FMovieSceneTrackEditor::BuildObjectBindingTrackMenu(FMenuBuilder& MenuBuilder, const TArray<FGuid>& ObjectBindings, const UClass* ObjectClass)
 {
 }
 
+TSharedPtr<SWidget> FMovieSceneTrackEditor::BuildOutlinerColumnWidget(const FBuildColumnWidgetParams& Params, const FName& ColumnName)
+{
+	using namespace UE::Sequencer;
+
+	UMovieSceneTrack* Track = Params.TrackModel->GetTrack();
+	TViewModelPtr<FSequencerEditorViewModel> Editor       = Params.Editor->CastThisShared<FSequencerEditorViewModel>();
+	TViewModelPtr<IOutlinerExtension>        OutlinerItem = Params.ViewModel.ImplicitCast();
+	if (!Track || !Editor || !OutlinerItem)
+	{
+		return SNullWidget::NullWidget;
+	}
+
+	if (ColumnName == FCommonOutlinerNames::Label)
+	{
+		return SNew(SOutlinerItemViewBase, OutlinerItem, Params.Editor, Params.TreeViewRow);
+	}
+
+	if (ColumnName == FCommonOutlinerNames::Edit)
+	{
+		TOptional<FViewModelChildren>  TopLevelChannels = Params.ViewModel->FindChildList(FTrackModel::GetTopLevelChannelGroupType());
+		TSharedPtr<FChannelGroupModel> TopLevelChannel  = TopLevelChannels ? TopLevelChannels->FindFirstChildOfType<FChannelGroupModel>() : nullptr;
+		if (TopLevelChannel)
+		{
+			return SNew(SKeyAreaEditorSwitcher, TopLevelChannel, Editor);
+		}
+	}
+
+	if (ColumnName == FCommonOutlinerNames::Add)
+	{
+		TAttribute<bool> IsHovered = MakeAttributeSP(Params.TreeViewRow, &ISequencerTreeViewRow::IsHovered);
+
+		// This function is a misnomer but is called this way for backwards compatibility
+		FBuildEditWidgetParams EditWidgetParams(Params);
+		EditWidgetParams.NodeIsHovered = IsHovered;
+		EditWidgetParams.RowIndex = Params.TrackModel->GetRowIndex();
+		if (Params.ViewModel->FindAncestorOfType<ITrackExtension>())
+		{
+			EditWidgetParams.TrackInsertRowIndex = Params.TrackModel->GetRowIndex();
+		}
+		else if (Track && Track->SupportsMultipleRows())
+		{
+			EditWidgetParams.TrackInsertRowIndex = Track->GetMaxRowIndex() + 1;
+		}
+		TViewModelPtr<IObjectBindingExtension> ObjectBinding = Params.ViewModel->FindAncestorOfType<IObjectBindingExtension>();
+		return BuildOutlinerEditWidget(ObjectBinding ? ObjectBinding->GetObjectGuid() : FGuid(), Track, EditWidgetParams);
+	}
+
+	if (ColumnName == FCommonOutlinerNames::Nav)
+	{
+		bool bHasKeyableAreas = false;
+		for (TSharedPtr<FChannelGroupModel> ChannelGroup : Params.ViewModel->GetDescendantsOfType<FChannelGroupModel>())
+		{
+			for (const TWeakViewModelPtr<FChannelModel>& WeakChannel : ChannelGroup->GetChannels())
+			{
+				if (TViewModelPtr<FChannelModel> Channel = WeakChannel.Pin())
+				{
+					if (Channel->GetKeyArea()->CanCreateKeyEditor())
+					{
+						bHasKeyableAreas = true;
+						break;
+					}
+				}
+			}
+			if (bHasKeyableAreas)
+			{
+				break;
+			}
+		}
+
+		if (bHasKeyableAreas)
+		{
+			return SNew(SSequencerKeyNavigationButtons, Params.ViewModel, Editor->GetSequencer());
+		}
+
+		return nullptr;
+	}
+
+	if (ColumnName == FCommonOutlinerNames::ColorPicker)
+	{
+		return SNew(SOutlinerTrackColorPicker, Params.ViewModel.ImplicitCast(), Params.Editor);
+	}
+
+	return nullptr;
+}
+
 TSharedPtr<SWidget> FMovieSceneTrackEditor::BuildOutlinerEditWidget(const FGuid& ObjectBinding, UMovieSceneTrack* Track, const FBuildEditWidgetParams& Params)
 {
 	if (Track && Track->GetSupportedBlendTypes().Num() > 0)
@@ -209,13 +299,7 @@ TSharedPtr<SWidget> FMovieSceneTrackEditor::BuildOutlinerEditWidget(const FGuid&
 			return MenuBuilder.MakeWidget();
 		};
 
-		return SNew(SHorizontalBox)
-		+ SHorizontalBox::Slot()
-		.AutoWidth()
-		.VAlign(VAlign_Center)
-		[
-			FSequencerUtilities::MakeAddButton(NSLOCTEXT("MovieSceneTrackEditor", "AddSection", "Section"), FOnGetContent::CreateLambda(SubMenuCallback), Params.NodeIsHovered, GetSequencer())
-		];
+		return UE::Sequencer::MakeAddButton(NSLOCTEXT("MovieSceneTrackEditor", "AddSection", "Section"), FOnGetContent::CreateLambda(SubMenuCallback), Params.ViewModel);
 	}
 	else
 	{
diff --git a/Engine/Source/Editor/Sequencer/Private/SKeyAreaEditorSwitcher.cpp b/Engine/Source/Editor/Sequencer/Private/SKeyAreaEditorSwitcher.cpp
--- a/Engine/Source/Editor/Sequencer/Private/SKeyAreaEditorSwitcher.cpp
+++ b/Engine/Source/Editor/Sequencer/Private/SKeyAreaEditorSwitcher.cpp
@@ -8,18 +8,17 @@
 #include "SequencerCommonHelpers.h"
 #include "IKeyArea.h"
 #include "MVVM/ViewModels/ChannelModel.h"
+#include "MVVM/ViewModels/SequencerEditorViewModel.h"
 #include "MVVM/Extensions/IObjectBindingExtension.h"
 #include "MVVM/Views/SKeyNavigationButtons.h"
 
-namespace UE
-{
-namespace Sequencer
+namespace UE::Sequencer
 {
 
-void SKeyAreaEditorSwitcher::Construct(const FArguments& InArgs, TSharedPtr<FChannelGroupModel> InModel, TWeakPtr<ISequencer> InWeakSequencer)
+void SKeyAreaEditorSwitcher::Construct(const FArguments& InArgs, TSharedPtr<FChannelGroupModel> InModel, TWeakPtr<FSequencerEditorViewModel> InWeakEditorModel)
 {
 	WeakModel = InModel;
-	WeakSequencer = InWeakSequencer;
+	WeakEditorModel = InWeakEditorModel;
 	CachedChannelsSerialNumber = 0;
 }
 
@@ -30,9 +29,9 @@ int32 SKeyAreaEditorSwitcher::GetWidgetIndex() const
 
 void SKeyAreaEditorSwitcher::Rebuild()
 {
-	TSharedPtr<FChannelGroupModel> Model     = WeakModel.Pin();
-	TSharedPtr<ISequencer>         Sequencer = WeakSequencer.Pin();
-	if (!Model || !Sequencer)
+	TSharedPtr<FChannelGroupModel>        Model = WeakModel.Pin();
+	TSharedPtr<FSequencerEditorViewModel> Editor = WeakEditorModel.Pin();
+	if (!Model || !Editor)
 	{
 		// Empty our cache so we don't persistently rebuild
 		CachedKeyAreas.Empty();
@@ -45,13 +44,12 @@ void SKeyAreaEditorSwitcher::Rebuild()
 		return;
 	}
 
-	const bool bIsEnabled = !Sequencer->IsReadOnly();
+	SetEnabled(MakeAttributeSP(Editor.ToSharedRef(), &FEditorViewModel::IsEditable));
 
 	// Index 0 is always the spacer node
 	VisibleIndex = 0;
 
 	TSharedRef<SWidgetSwitcher> Switcher = SNew(SWidgetSwitcher)
-		.IsEnabled(bIsEnabled)
 		.WidgetIndex(this, &SKeyAreaEditorSwitcher::GetWidgetIndex)
 
 		+ SWidgetSwitcher::Slot()
@@ -75,14 +73,9 @@ void SKeyAreaEditorSwitcher::Rebuild()
 		else
 		{
 			Switcher->AddSlot()
+			.HAlign(HAlign_Left)
 			[
-				SNew(SBox)
-				.IsEnabled(bIsEnabled)
-				.MinDesiredWidth(100)
-				.HAlign(HAlign_Left)
-				[
-					KeyArea->CreateKeyEditor(Sequencer, ObjectBindingID)
-				]
+				KeyArea->CreateKeyEditor(Editor->GetSequencer(), ObjectBindingID)
 			];
 		}
 	}
@@ -95,9 +88,9 @@ void SKeyAreaEditorSwitcher::Rebuild()
 
 void SKeyAreaEditorSwitcher::Tick( const FGeometry& AllottedGeometry, const double InCurrentTime, const float InDeltaTime )
 {
-	TSharedPtr<FChannelGroupModel> Model = WeakModel.Pin();
-	TSharedPtr<ISequencer>         Sequencer = WeakSequencer.Pin();
-	if (!Model || !Sequencer)
+	TSharedPtr<FChannelGroupModel>        Model  = WeakModel.Pin();
+	TSharedPtr<FSequencerEditorViewModel> Editor = WeakEditorModel.Pin();
+	if (!Model || !Editor)
 	{
 		if (CachedKeyAreas.Num() != 0)
 		{
@@ -126,7 +119,7 @@ void SKeyAreaEditorSwitcher::Tick( const FGeometry& AllottedGeometry, const doub
 			AllSections.Add(KeyArea->GetOwningSection());
 		}
 
-		const int32 ActiveKeyArea = SequencerHelpers::GetSectionFromTime(AllSections, Sequencer->GetLocalTime().Time.FrameNumber);
+		const int32 ActiveKeyArea = SequencerHelpers::GetSectionFromTime(AllSections, Editor->GetSequencer()->GetLocalTime().Time.FrameNumber);
 		if (ActiveKeyArea != INDEX_NONE)
 		{
 			// Index 0 is the spacer node, so add 1 to the key area index to get the widget index
@@ -139,5 +132,4 @@ void SKeyAreaEditorSwitcher::Tick( const FGeometry& AllottedGeometry, const doub
 	}
 }
 
-} // namespace Sequencer
-} // namespace UE
+} // namespace UE::Sequencer
diff --git a/Engine/Source/Editor/Sequencer/Private/SKeyAreaEditorSwitcher.h b/Engine/Source/Editor/Sequencer/Private/SKeyAreaEditorSwitcher.h
--- a/Engine/Source/Editor/Sequencer/Private/SKeyAreaEditorSwitcher.h
+++ b/Engine/Source/Editor/Sequencer/Private/SKeyAreaEditorSwitcher.h
@@ -15,6 +15,7 @@ namespace Sequencer
 {
 
 class FChannelGroupModel;
+class FSequencerEditorViewModel;
 
 class SKeyAreaEditorSwitcher : public SCompoundWidget
 {
@@ -23,7 +24,7 @@ public:
 	SLATE_END_ARGS()
 
 	/** Construct the widget */
-	void Construct(const FArguments& InArgs, TSharedPtr<FChannelGroupModel> InModel, TWeakPtr<ISequencer> InWeakSequencer);
+	void Construct(const FArguments& InArgs, TSharedPtr<FChannelGroupModel> InModel, TWeakPtr<FSequencerEditorViewModel> InWeakEditorViewModel);
 
 	/** Rebuild this widget from its cached key area node */
 	void Rebuild();
@@ -43,8 +44,8 @@ private:
 	int32 VisibleIndex;
 	/** The key area to which we relate */
 	TWeakPtr<FChannelGroupModel> WeakModel;
-	/** The sequencer we belong to */
-	TWeakPtr<ISequencer> WeakSequencer;
+	/** Weak editor view model */
+	TWeakPtr<FSequencerEditorViewModel> WeakEditorModel;
 	/** Key areas cached from the node */
 	TArray<TSharedRef<IKeyArea>> CachedKeyAreas;
 	/** Serial cached from the node */
diff --git a/Engine/Source/Editor/Sequencer/Private/SSequencer.cpp b/Engine/Source/Editor/Sequencer/Private/SSequencer.cpp
--- a/Engine/Source/Editor/Sequencer/Private/SSequencer.cpp
+++ b/Engine/Source/Editor/Sequencer/Private/SSequencer.cpp
@@ -72,6 +72,8 @@
 #include "DragAndDrop/CompositeDragDropOp.h"
 #include "Widgets/Input/SSearchBox.h"
 #include "MVVM/Views/SSequencerOutlinerView.h"
+#include "MVVM/ViewModels/OutlinerColumns/IOutlinerColumn.h"
+#include "MVVM/ViewModels/OutlinerColumns/OutlinerColumnTypes.h"
 #include "MovieSceneTrackEditor.h"
 #include "SSequencerSplitterOverlay.h"
 #include "SequencerHotspots.h"
@@ -126,6 +128,16 @@
 
 #define LOCTEXT_NAMESPACE "Sequencer"
 
+FSequencerOutlinerColumnVisibility::FSequencerOutlinerColumnVisibility(TSharedPtr<UE::Sequencer::IOutlinerColumn> InColumn)
+	: Column(InColumn)
+	, bIsColumnVisible(InColumn->IsColumnVisibleByDefault())
+{}
+
+FSequencerOutlinerColumnVisibility::FSequencerOutlinerColumnVisibility(TSharedPtr<UE::Sequencer::IOutlinerColumn> InColumn, bool bInIsColumnVisible)
+	: Column(InColumn)
+	, bIsColumnVisible(bInIsColumnVisible)
+{}
+
 /* SSequencer interface
  *****************************************************************************/
 UE_DISABLE_OPTIMIZATION_SHIP
@@ -320,6 +332,12 @@ void SSequencer::Construct(const FArguments& InArgs, TSharedRef<FSequencer> InSe
 	SequencerViewModel->GetTrackArea()->CastThisChecked<FSequencerTrackAreaViewModel>()->InitializeDefaultEditTools(*TrackArea);
 	SequencerViewModel->GetPinnedTrackArea()->CastThisChecked<FSequencerTrackAreaViewModel>()->InitializeDefaultEditTools(*PinnedTrackArea);
 
+
+	if (USequencerSettings* Settings = GetSequencerSettings())
+	{
+		SequencerViewModel->SetViewDensity(Settings->GetViewDensity());
+	}
+
 	PlayTimeDisplay = StaticCastSharedRef<STemporarilyFocusedSpinBox<double>>(SequencerPtr.Pin()->MakePlayTimeDisplay(GetNumericTypeInterface()));
 
 	TAttribute<FAnimatedRange> ViewRangeAttribute = InArgs._ViewRange;
@@ -1088,6 +1106,8 @@ void SSequencer::InitializeTrackFilters()
 
 void SSequencer::UpdateOutlinerViewColumns()
 {
+	using namespace UE::Sequencer;
+
 	// Save updated column list in settings
 	TArray<FColumnVisibilitySetting> ColumnVisibilitySettings;
 
@@ -1099,7 +1119,7 @@ void SSequencer::UpdateOutlinerViewColumns()
 	GetSequencerSettings()->SetOutlinerColumnVisibility(ColumnVisibilitySettings);
 
 	// Filter out hidden columns to create a list of visible columns for the outliner views
-	TArray<TSharedPtr<ISequencerOutlinerColumn>> VisibleColumns;
+	TArray<TSharedPtr<IOutlinerColumn>> VisibleColumns;
 	for (FSequencerOutlinerColumnVisibility ColumnVisibility : OutlinerColumnVisibilities)
 	{
 		if (ColumnVisibility.bIsColumnVisible)
@@ -1123,7 +1143,7 @@ void SSequencer::InitializeOutlinerColumns()
 		return;
 	}
 
-	const TMap<FName, TSharedPtr<ISequencerOutlinerColumn>>& RegisteredColumns = Sequencer->GetOutlinerColumns();
+	const TMap<FName, TSharedPtr<IOutlinerColumn>>& RegisteredColumns = Sequencer->GetOutlinerColumns();
 	
 	// Retrieve previously saved column names and visibilities
 	TArray<FColumnVisibilitySetting> ColumnSettings = GetSequencerSettings()->GetOutlinerColumnSettings();
@@ -1132,7 +1152,7 @@ void SSequencer::InitializeOutlinerColumns()
 	// Add registered columns found in settings with their saved visibility state
 	for (const FColumnVisibilitySetting& ColumnVisibility : ColumnSettings)
 	{
-		const TSharedPtr<ISequencerOutlinerColumn>* OutlinerColumn = RegisteredColumns.Find(ColumnVisibility.ColumnName);
+		const TSharedPtr<IOutlinerColumn>* OutlinerColumn = RegisteredColumns.Find(ColumnVisibility.ColumnName);
 		if (OutlinerColumn)
 		{
 			ColumnNamesFoundInSettings.Add(ColumnVisibility.ColumnName);
@@ -1141,7 +1161,7 @@ void SSequencer::InitializeOutlinerColumns()
 	}
 
 	// Add registered columns not found in settings with their default visibility state
-	for (const TTuple<FName, TSharedPtr<ISequencerOutlinerColumn>>& RegisteredColumn : RegisteredColumns)
+	for (const TTuple<FName, TSharedPtr<IOutlinerColumn>>& RegisteredColumn : RegisteredColumns)
 	{
 		if (!ColumnNamesFoundInSettings.Contains(RegisteredColumn.Key))
 		{
@@ -1149,6 +1169,10 @@ void SSequencer::InitializeOutlinerColumns()
 		}
 	}
 
+	Algo::Sort(OutlinerColumnVisibilities, [](const FSequencerOutlinerColumnVisibility& A, const FSequencerOutlinerColumnVisibility& B){
+		return A.Column->GetPosition() < B.Column->GetPosition();
+	});
+
 	UpdateOutlinerViewColumns();
 }
 
@@ -2181,6 +2205,8 @@ TSharedRef<SWidget> SSequencer::MakeViewMenu()
 		GetSequencerSettings()->SetZeroPadFrames(NewValue);
 	};
 
+	MenuBuilder.AddSubMenu(LOCTEXT("ViewDensityMenuLabel", "View Density"), FText::GetEmpty(), FNewMenuDelegate::CreateRaw(this, &SSequencer::FillViewDensityMenu));
+
 	// Menu Entry for Outliner Column Visibilities
 	if (OutlinerColumnVisibilities.Num() > 0)
 	{
@@ -2314,6 +2340,75 @@ void SSequencer::FillPlaybackSpeedMenu(FMenuBuilder& InMenuBarBuilder)
 	InMenuBarBuilder.EndSection();
 }
 
+void SSequencer::FillViewDensityMenu(FMenuBuilder& InMenuBuilder)
+{
+	using namespace UE::Sequencer;
+
+	auto SetViewDensity = [this](EViewDensity InViewDensity){
+		TSharedPtr<FEditorViewModel> Editor = this->SequencerPtr.Pin()->GetViewModel();
+		Editor->SetViewDensity(InViewDensity);
+
+		if (USequencerSettings* Settings = this->GetSequencerSettings())
+		{
+			if (InViewDensity == EViewDensity::Relaxed)
+			{
+				Settings->SetViewDensity("Relaxed");
+			}
+			else if (InViewDensity == EViewDensity::Expanded)
+			{
+				Settings->SetViewDensity("Expanded");
+			}
+			else
+			{
+				Settings->SetViewDensity("Compact");
+			}
+		}
+	};
+	auto IsCurrentViewDensity = [this](EViewDensity InViewDensity){
+		TSharedPtr<FEditorViewModel> Editor = this->SequencerPtr.Pin()->GetViewModel();
+		return Editor->GetViewDensity().Density == InViewDensity;
+	};
+
+	InMenuBuilder.AddMenuEntry(
+		LOCTEXT("CompactViewDensity_Label", "Compact"),
+		LOCTEXT("CompactViewDensity_Tooltip", "Change Sequencer to use a compact view mode"),
+		FSlateIcon(),
+		FUIAction(
+			FExecuteAction::CreateLambda(SetViewDensity, EViewDensity::Compact),
+			FCanExecuteAction(),
+			FIsActionChecked::CreateLambda(IsCurrentViewDensity, EViewDensity::Compact)
+		),
+		NAME_None,
+		EUserInterfaceActionType::RadioButton
+	);
+
+	InMenuBuilder.AddMenuEntry(
+		LOCTEXT("RelaxedViewDensity_Label", "Relaxed"),
+		LOCTEXT("RelaxedViewDensity_Tooltip", "Change Sequencer to use a relaxed view mode with uniform track heights"),
+		FSlateIcon(),
+		FUIAction(
+			FExecuteAction::CreateLambda(SetViewDensity, EViewDensity::Relaxed),
+			FCanExecuteAction(),
+			FIsActionChecked::CreateLambda(IsCurrentViewDensity, EViewDensity::Relaxed)
+		),
+		NAME_None,
+		EUserInterfaceActionType::RadioButton
+	);
+
+	InMenuBuilder.AddMenuEntry(
+		LOCTEXT("ExpandedViewDensity_Label", "Expanded"),
+		LOCTEXT("ExpandedViewDensity_Tooltip", "Change Sequencer to use an expanded view mode with larger uniform track heights"),
+		FSlateIcon(),
+		FUIAction(
+			FExecuteAction::CreateLambda(SetViewDensity, EViewDensity::Expanded),
+			FCanExecuteAction(),
+			FIsActionChecked::CreateLambda(IsCurrentViewDensity, EViewDensity::Expanded)
+		),
+		NAME_None,
+		EUserInterfaceActionType::RadioButton
+	);
+}
+
 void SSequencer::FillColumnVisibilityMenu(FMenuBuilder& InMenuBuilder)
 {
 	const bool bShouldCloseWindowAfterMenuSelection = true;
@@ -3626,8 +3721,6 @@ void SSequencer::OnCurveEditorVisibilityChanged(bool bShouldBeVisible)
 	{
 		CurveEditorExtension->CloseCurveEditor();
 	}
-
-	TreeView->UpdateTrackArea();
 }
 
 
diff --git a/Engine/Source/Editor/Sequencer/Private/SSequencer.h b/Engine/Source/Editor/Sequencer/Private/SSequencer.h
--- a/Engine/Source/Editor/Sequencer/Private/SSequencer.h
+++ b/Engine/Source/Editor/Sequencer/Private/SSequencer.h
@@ -14,7 +14,6 @@
 #include "MovieSceneSequenceID.h"
 #include "ITimeSlider.h"
 #include "ISequencerModule.h"
-#include "ISequencerOutlinerColumn.h"
 #include "ToolMenu.h"
 #include "Framework/Commands/UICommandList.h"
 #include "Widgets/Input/NumericTypeInterface.h"
@@ -51,6 +50,7 @@ namespace Sequencer
 	class SOutlinerView;
 	class STrackAreaView;
 	class FVirtualTrackArea;
+	class IOutlinerColumn;
 	struct FSequencerSelectionCurveFilter;
 
 } // namespace Sequencer
@@ -68,7 +68,7 @@ namespace SequencerLayoutConstants
 	const float ObjectNodeHeight = 20.0f;
 
 	/** Height of each section area if there are no sections (note: section areas may be larger than this if they have children. This is the height of a section area with no children or all children hidden) */
-	const float SectionAreaDefaultHeight = 15.0f;
+	const float SectionAreaDefaultHeight = 27.0f;
 
 	/** Height of each key area */
 	const float KeyAreaHeight = 15.0f;
@@ -116,18 +116,11 @@ struct FSequencerBreadcrumb
  */
 struct FSequencerOutlinerColumnVisibility
 {
-	TSharedPtr<ISequencerOutlinerColumn> Column;
+	TSharedPtr<UE::Sequencer::IOutlinerColumn> Column;
 	bool bIsColumnVisible = false;
 
-	FSequencerOutlinerColumnVisibility(TSharedPtr<ISequencerOutlinerColumn> InColumn)
-		: Column(InColumn)
-		, bIsColumnVisible(InColumn->IsColumnVisibleByDefault())
-	{}
-
-	FSequencerOutlinerColumnVisibility(TSharedPtr<ISequencerOutlinerColumn> InColumn, bool bInIsColumnVisible)
-		: Column(InColumn)
-		, bIsColumnVisible(bInIsColumnVisible)
-	{}
+	FSequencerOutlinerColumnVisibility(TSharedPtr<UE::Sequencer::IOutlinerColumn> InColumn);
+	FSequencerOutlinerColumnVisibility(TSharedPtr<UE::Sequencer::IOutlinerColumn> InColumn, bool bInIsColumnVisible);
 };
 
 
@@ -459,6 +452,9 @@ private:
 	/** Makes the playback speed menu for the toolbar. */
 	void FillPlaybackSpeedMenu(FMenuBuilder& InMenuBuilder);
 
+	/** Makes the view density menu for the toolbar. */
+	void FillViewDensityMenu(FMenuBuilder& InMenuBuilder);
+
 	/** Makes the column visibility menu for the toolbar. */
 	void FillColumnVisibilityMenu(FMenuBuilder& InMenuBuilder);
 
diff --git a/Engine/Source/Editor/Sequencer/Private/SSequencerSection.cpp b/Engine/Source/Editor/Sequencer/Private/SSequencerSection.cpp
--- a/Engine/Source/Editor/Sequencer/Private/SSequencerSection.cpp
+++ b/Engine/Source/Editor/Sequencer/Private/SSequencerSection.cpp
@@ -6,6 +6,7 @@
 #include "MVVM/ViewModels/CategoryModel.h"
 #include "MVVM/ViewModels/ViewModelIterators.h"
 #include "MVVM/ViewModels/SequencerEditorViewModel.h"
+#include "MVVM/ViewModels/EditorSharedViewModelData.h"
 #include "MVVM/Views/ITrackAreaHotspot.h"
 #include "MVVM/Views/STrackAreaView.h"
 #include "MVVM/Views/SCompoundTrackLaneView.h"
@@ -1464,11 +1465,14 @@ int32 SSequencerSection::OnPaint( const FPaintArgs& Args, const FGeometry& Allot
 
 	if (!SectionTitle.IsEmpty())
 	{
+		TViewModelPtr<IGeometryExtension> Geometry = SectionModel->FindAncestorOfType<IOutlinerExtension>().ImplicitCast();
+		const float AllocatedFontHeight = (Geometry ? Geometry->GetVirtualGeometry().GetHeight() : SectionGeometry.Size.Y) - 4.f; // 2px minimum padding
+
+		// Align the section title within the actual track area geometry, excluding any expanded channels
 		FSlateClippingZone ClippingZone(Painter.SectionClippingRect);
 		OutDrawElements.PushClip(ClippingZone);
 
 		FVector2D TopLeft = SectionGeometry.AbsoluteToLocal(Painter.SectionClippingRect.GetTopLeft()) + FVector2D(1.f, -1.f);
-
 		FSlateFontInfo FontInfo = FAppStyle::GetFontStyle("NormalFont");
 
 		TSharedRef<FSlateFontCache> FontCache = FSlateApplication::Get().GetRenderer()->GetFontCache();
@@ -1477,18 +1481,20 @@ int32 SSequencerSection::OnPaint( const FPaintArgs& Args, const FGeometry& Allot
 		{
 			return FontCache->GetMaxCharacterHeight(FontInfo, 1.f) + FontCache->GetBaseline(FontInfo, 1.f);
 		};
-		while (GetFontHeight() > SectionGeometry.Size.Y && FontInfo.Size > 11)
+		while (GetFontHeight() > AllocatedFontHeight && FontInfo.Size > 7.f)
 		{
-			FontInfo.Size = FMath::Max(FMath::FloorToInt(FontInfo.Size - 6.f), 11);
+			FontInfo.Size = FMath::Max(FMath::FloorToInt(FontInfo.Size - 6.f), 7.f);
 		}
 
+		const float TitlePosition = (AllocatedFontHeight - FontInfo.Size) * 0.5f;
+
 		// Drop shadow
 		FSlateDrawElement::MakeText(
 			OutDrawElements,
 			LayerId,
 			SectionGeometry.MakeChild(
 				FVector2D(SectionGeometry.Size.X, GetFontHeight()),
-				FSlateLayoutTransform(TopLeft + FVector2D(ContentPadding.Left, ContentPadding.Top) + FVector2D(1.f, 1.f))
+				FSlateLayoutTransform(TopLeft + FVector2D(ContentPadding.Left, TitlePosition) + FVector2D(1.f, 1.f))
 			).ToPaintGeometry(),
 			SectionTitle,
 			FontInfo,
@@ -1501,7 +1507,7 @@ int32 SSequencerSection::OnPaint( const FPaintArgs& Args, const FGeometry& Allot
 			LayerId,
 			SectionGeometry.MakeChild(
 				FVector2D(SectionGeometry.Size.X, GetFontHeight()),
-				FSlateLayoutTransform(TopLeft + FVector2D(ContentPadding.Left, ContentPadding.Top))
+				FSlateLayoutTransform(TopLeft + FVector2D(ContentPadding.Left, TitlePosition))
 			).ToPaintGeometry(),
 			SectionTitle,
 			FontInfo,
@@ -1561,6 +1567,14 @@ void SSequencerSection::PaintEasingHandles( FSequencerSectionPainter& InPainter,
 		return;
 	}
 
+	TSharedPtr<FSectionModel>              SectionModel = WeakSectionModel.Pin();
+	TSharedPtr<FSharedViewModelData>       SharedModelData = SectionModel ? SectionModel->GetSharedData() : nullptr;
+	TSharedPtr<FEditorSharedViewModelData> EditorSharedModelData = SharedModelData ? SharedModelData->CastThisShared<FEditorSharedViewModelData>() : nullptr;
+	if (!SectionModel || !EditorSharedModelData)
+	{
+		return;
+	}
+
 	TArray<TSharedPtr<FSectionModel>> AllUnderlappingSections;
 	if (IsSectionHighlighted(SectionInterface->GetSectionObject(), Hotspot))
 	{
@@ -1580,9 +1594,9 @@ void SSequencerSection::PaintEasingHandles( FSequencerSectionPainter& InPainter,
 	}
 
 	FTimeToPixel TimeToPixelConverter = InPainter.GetTimeConverter();
-	for (const TSharedPtr<FSectionModel>& SectionModel : AllUnderlappingSections)
+	for (const TSharedPtr<FSectionModel>& UnderlappingSectionModel : AllUnderlappingSections)
 	{
-		UMovieSceneSection* UnderlappingSectionObj = SectionModel->GetSection();
+		UMovieSceneSection* UnderlappingSectionObj = UnderlappingSectionModel->GetSection();
 		if (UnderlappingSectionObj->GetRange() == TRange<FFrameNumber>::All())
 		{
 			continue;
@@ -1597,7 +1611,7 @@ void SSequencerSection::PaintEasingHandles( FSequencerSectionPainter& InPainter,
 		{
 			if (const FSectionEasingHandleHotspot* EasingHotspot = Hotspot->CastThis<FSectionEasingHandleHotspot>())
 			{
-				bDrawThisSectionsHandles = (EasingHotspot->WeakSectionModel.Pin() == SectionModel);
+				bDrawThisSectionsHandles = (EasingHotspot->WeakSectionModel.Pin() == UnderlappingSectionModel);
 				bLeftHandleActive = EasingHotspot->HandleType == ESequencerEasingType::In;
 				bRightHandleActive = EasingHotspot->HandleType == ESequencerEasingType::Out;
 			}
@@ -1647,7 +1661,7 @@ void SSequencerSection::PaintEasingHandles( FSequencerSectionPainter& InPainter,
 		// smaller.
 		// If there isn't any "section to key" border, we just draw the handle white.
 		const FLinearColor InactiveHandleColor = bIsSectionToKey ? FStyleColors::Success.GetColor(FWidgetStyle()) : FStyleColors::AccentYellow.GetColor(FWidgetStyle());
-		const float HalfSectionHeight = SectionInterface->GetSectionHeight() / 2.f;
+		const float HalfSectionHeight = SectionInterface->GetSectionHeight(EditorSharedModelData->GetEditor()->GetViewDensity()) / 2.f;
 		const float HandleSize = FMath::Max(MinHandleSize, FMath::Min(HalfSectionHeight, SectionInterface->GetSectionGripSize())) + (bIsSectionToKey ? 3.f : 0.f);
 
 		const FSlateBrush* HandleBrush = FAppStyle::GetBrush("Sequencer.Section.EasingHandle");
diff --git a/Engine/Source/Editor/Sequencer/Private/Sequencer.cpp b/Engine/Source/Editor/Sequencer/Private/Sequencer.cpp
--- a/Engine/Source/Editor/Sequencer/Private/Sequencer.cpp
+++ b/Engine/Source/Editor/Sequencer/Private/Sequencer.cpp
@@ -128,7 +128,7 @@
 #include "SequencerKeyActor.h"
 #include "ISequencerChannelInterface.h"
 #include "IMovieRendererInterface.h"
-#include "ISequencerOutlinerColumn.h"
+#include "MVVM/ViewModels/OutlinerColumns/IOutlinerColumn.h"
 #include "SequencerKeyCollection.h"
 #include "CurveEditor.h"
 #include "CurveEditorScreenSpace.h"
@@ -502,7 +502,7 @@ void FSequencer::InitSequencer(const FSequencerInitParams& InitParams, const TSh
 	for (int32 DelegateIndex = 0; DelegateIndex < OutlinerColumnDelegates.Num(); ++DelegateIndex)
 	{
 		check(OutlinerColumnDelegates[DelegateIndex].IsBound());
-		TSharedRef<ISequencerOutlinerColumn> OutlinerColumn = OutlinerColumnDelegates[DelegateIndex].Execute();
+		TSharedRef<IOutlinerColumn> OutlinerColumn = OutlinerColumnDelegates[DelegateIndex].Execute();
 
 		if (OutlinerColumn->SupportsSequence(InitParams.RootSequence))
 		{
@@ -11508,15 +11508,6 @@ void FSequencer::BuildObjectBindingTrackMenu(FMenuBuilder& MenuBuilder, const TA
 	}
 }
 
-
-void FSequencer::BuildObjectBindingEditButtons(TSharedPtr<SHorizontalBox> EditBox, const FGuid& ObjectBinding, const UClass* ObjectClass)
-{
-	for (int32 i = 0; i < TrackEditors.Num(); ++i)
-	{
-		TrackEditors[i]->BuildObjectBindingEditButtons(EditBox, ObjectBinding, ObjectClass);
-	}
-}
-
 void FSequencer::BuildAddSelectedToFolderMenu(FMenuBuilder& MenuBuilder)
 {
 	using namespace UE::Sequencer;
diff --git a/Engine/Source/Editor/Sequencer/Private/Sequencer.h b/Engine/Source/Editor/Sequencer/Private/Sequencer.h
--- a/Engine/Source/Editor/Sequencer/Private/Sequencer.h
+++ b/Engine/Source/Editor/Sequencer/Private/Sequencer.h
@@ -308,7 +308,7 @@ public:
 	/**
 	* @return Outliner Columns registered to the sequencer by column name
 	*/
-	const TMap<FName, TSharedPtr<ISequencerOutlinerColumn>>& GetOutlinerColumns() const
+	const TMap<FName, TSharedPtr<UE::Sequencer::IOutlinerColumn>>& GetOutlinerColumns() const
 	{
 		return OutlinerColumns;
 	}
@@ -462,15 +462,6 @@ public:
 	 */
 	void BuildObjectBindingTrackMenu(FMenuBuilder& MenuBuilder, const TArray<FGuid>& ObjectBindings, const UClass* ObjectClass);
 
-	/**
-	 * Builds up the edit buttons for object binding nodes in the outliner
-	 * 
-	 * @param EditBox	    The edit box to add things to
-	 * @param ObjectBinding	The object binding of the selected node
-	 * @param ObjectClass	The class of the selected object
-	 */
-	void BuildObjectBindingEditButtons(TSharedPtr<SHorizontalBox> EditBox, const FGuid& ObjectBinding, const UClass* ObjectClass);
-
 	/**
 	 * Builds up the menu of folders to add selected nodes to
 	 *
@@ -1193,7 +1184,7 @@ private:
 	TMap<FObjectKey, TSharedPtr<ISequencerTrackEditor>> TrackEditorsByType;
 
 	/** List of Outliner column creators that are supported by the Sequencer. */
-	TMap<FName, TSharedPtr<ISequencerOutlinerColumn>> OutlinerColumns;
+	TMap<FName, TSharedPtr<UE::Sequencer::IOutlinerColumn>> OutlinerColumns;
 
 	/** List of object bindings we can use */
 	TArray<TSharedPtr<ISequencerEditorObjectBinding>> ObjectBindings;
diff --git a/Engine/Source/Editor/Sequencer/Private/SequencerModule.cpp b/Engine/Source/Editor/Sequencer/Private/SequencerModule.cpp
--- a/Engine/Source/Editor/Sequencer/Private/SequencerModule.cpp
+++ b/Engine/Source/Editor/Sequencer/Private/SequencerModule.cpp
@@ -31,10 +31,15 @@
 #include "MVVM/ViewModels/SequenceModel.h"
 #include "MVVM/ViewModels/SequencerEditorViewModel.h"
 
-#include "OutlinerColumns/LockOutlinerColumn.h"
-#include "OutlinerColumns/MuteOutlinerColumn.h"
-#include "OutlinerColumns/PinOutlinerColumn.h"
-#include "OutlinerColumns/SoloOutlinerColumn.h"
+#include "MVVM/ViewModels/OutlinerColumns/LockOutlinerColumn.h"
+#include "MVVM/ViewModels/OutlinerColumns/MuteOutlinerColumn.h"
+#include "MVVM/ViewModels/OutlinerColumns/PinOutlinerColumn.h"
+#include "MVVM/ViewModels/OutlinerColumns/SoloOutlinerColumn.h"
+#include "MVVM/ViewModels/OutlinerColumns/LabelOutlinerColumn.h"
+#include "MVVM/ViewModels/OutlinerColumns/EditOutlinerColumn.h"
+#include "MVVM/ViewModels/OutlinerColumns/AddOutlinerColumn.h"
+#include "MVVM/ViewModels/OutlinerColumns/NavOutlinerColumn.h"
+#include "MVVM/ViewModels/OutlinerColumns/ColorPickerOutlinerColumn.h"
 
 #include "ToolMenus.h"
 #include "ContentBrowserMenuContexts.h"
@@ -380,11 +385,20 @@ public:
 			FPropertyEditorModule& EditModule = FModuleManager::Get().GetModuleChecked<FPropertyEditorModule>("PropertyEditor");
 			OnGetGlobalRowExtensionHandle = EditModule.GetGlobalRowExtensionDelegate().AddStatic(&RegisterKeyframeExtensionHandler);
 
-			// Register built-in outliner columns (default order is pin, mute, lock, solo)
-			PinOutlinerColumnHandle = RegisterOutlinerColumn(FOnCreateOutlinerColumn::CreateStatic(&FPinOutlinerColumn::CreateOutlinerColumn));
-			MuteOutlinerColumnHandle = RegisterOutlinerColumn(FOnCreateOutlinerColumn::CreateStatic(&FMuteOutlinerColumn::CreateOutlinerColumn));
-			LockOutlinerColumnHandle = RegisterOutlinerColumn(FOnCreateOutlinerColumn::CreateStatic(&FLockOutlinerColumn::CreateOutlinerColumn));
-			SoloOutlinerColumnHandle = RegisterOutlinerColumn(FOnCreateOutlinerColumn::CreateStatic(&FSoloOutlinerColumn::CreateOutlinerColumn));
+			// Register left gutter columns
+			PinOutlinerColumnHandle  = RegisterOutlinerColumn(FOnCreateOutlinerColumn::CreateStatic([]{ return TSharedRef<IOutlinerColumn>(MakeShared<FPinOutlinerColumn>()); }));
+			MuteOutlinerColumnHandle = RegisterOutlinerColumn(FOnCreateOutlinerColumn::CreateStatic([]{ return TSharedRef<IOutlinerColumn>(MakeShared<FMuteOutlinerColumn>()); }));
+			LockOutlinerColumnHandle = RegisterOutlinerColumn(FOnCreateOutlinerColumn::CreateStatic([]{ return TSharedRef<IOutlinerColumn>(MakeShared<FLockOutlinerColumn>()); }));
+			SoloOutlinerColumnHandle = RegisterOutlinerColumn(FOnCreateOutlinerColumn::CreateStatic([]{ return TSharedRef<IOutlinerColumn>(MakeShared<FSoloOutlinerColumn>()); }));
+
+			// Register center columns
+			LabelOutlinerColumnHandle = RegisterOutlinerColumn(FOnCreateOutlinerColumn::CreateStatic([]{ return TSharedRef<IOutlinerColumn>(MakeShared<FLabelOutlinerColumn>()); }));
+			EditOutlinerColumnHandle  = RegisterOutlinerColumn(FOnCreateOutlinerColumn::CreateStatic([]{ return TSharedRef<IOutlinerColumn>(MakeShared<FEditOutlinerColumn>()); }));
+			AddOutlinerColumnHandle   = RegisterOutlinerColumn(FOnCreateOutlinerColumn::CreateStatic([]{ return TSharedRef<IOutlinerColumn>(MakeShared<FAddOutlinerColumn>()); }));
+
+			// Register right gutter columns
+			NavOutlinerColumnHandle          = RegisterOutlinerColumn(FOnCreateOutlinerColumn::CreateStatic([]{ return TSharedRef<IOutlinerColumn>(MakeShared<FNavOutlinerColumn>()); }));
+			ColorPickerOutlinerColumnHandle  = RegisterOutlinerColumn(FOnCreateOutlinerColumn::CreateStatic([]{ return TSharedRef<IOutlinerColumn>(MakeShared<FColorPickerOutlinerColumn>()); }));
 
 			RegisterObjectSchemas();
 		}
@@ -435,6 +449,11 @@ public:
 			UnregisterOutlinerColumn(MuteOutlinerColumnHandle);
 			UnregisterOutlinerColumn(LockOutlinerColumnHandle);
 			UnregisterOutlinerColumn(SoloOutlinerColumnHandle);
+			UnregisterOutlinerColumn(LabelOutlinerColumnHandle);
+			UnregisterOutlinerColumn(EditOutlinerColumnHandle);
+			UnregisterOutlinerColumn(AddOutlinerColumnHandle);
+			UnregisterOutlinerColumn(NavOutlinerColumnHandle);
+			UnregisterOutlinerColumn(ColorPickerOutlinerColumnHandle);
 		}
 	}
 
@@ -621,6 +640,11 @@ private:
 	FDelegateHandle MuteOutlinerColumnHandle;
 	FDelegateHandle LockOutlinerColumnHandle;
 	FDelegateHandle SoloOutlinerColumnHandle;
+	FDelegateHandle LabelOutlinerColumnHandle;
+	FDelegateHandle EditOutlinerColumnHandle;
+	FDelegateHandle AddOutlinerColumnHandle;
+	FDelegateHandle NavOutlinerColumnHandle;
+	FDelegateHandle ColorPickerOutlinerColumnHandle;
 };
 
 IMPLEMENT_MODULE(FSequencerModule, Sequencer);
diff --git a/Engine/Source/Editor/Sequencer/Private/SequencerSettings.cpp b/Engine/Source/Editor/Sequencer/Private/SequencerSettings.cpp
--- a/Engine/Source/Editor/Sequencer/Private/SequencerSettings.cpp
+++ b/Engine/Source/Editor/Sequencer/Private/SequencerSettings.cpp
@@ -4,6 +4,7 @@
 #include "KeyParams.h"
 #include "ISequencer.h"
 #include "SSequencer.h"
+#include "MVVM/ViewModels/ViewDensity.h"
 
 USequencerSettings::USequencerSettings( const FObjectInitializer& ObjectInitializer )
 	: Super( ObjectInitializer )
@@ -66,6 +67,7 @@ USequencerSettings::USequencerSettings( const FObjectInitializer& ObjectInitiali
 	TreeViewWidth = 0.3f;
 	bShowTickLines = true;
 	bShowSequencerToolbar = true;
+	ViewDensity = "Relaxed";
 
 	SectionColorTints.Add(FColor(88, 102, 142, 255)); // blue
 	SectionColorTints.Add(FColor(99, 137, 132, 255)); // blue-green
@@ -979,6 +981,30 @@ void USequencerSettings::SetTreeViewWidth(float InTreeViewWidth)
 	}
 }
 
+UE::Sequencer::EViewDensity USequencerSettings::GetViewDensity() const
+{
+	static FName NAME_Relaxed("Relaxed");
+	static FName NAME_Expanded("Expanded");
+	if (ViewDensity == NAME_Relaxed)
+	{
+		return UE::Sequencer::EViewDensity::Relaxed;
+	}
+	if (ViewDensity == NAME_Expanded)
+	{
+		return UE::Sequencer::EViewDensity::Expanded;
+	}
+	return UE::Sequencer::EViewDensity::Compact;
+}
+
+void USequencerSettings::SetViewDensity(FName InViewDensity)
+{
+	if (InViewDensity != ViewDensity)
+	{
+		ViewDensity = InViewDensity;
+		SaveConfig();
+	}
+}
+
 bool USequencerSettings::IsTrackFilterEnabled(const FString& TrackFilter) const
 {
 	return TrackFilters.Contains(TrackFilter);
diff --git a/Engine/Source/Editor/Sequencer/Private/SequencerUtilities.cpp b/Engine/Source/Editor/Sequencer/Private/SequencerUtilities.cpp
--- a/Engine/Source/Editor/Sequencer/Private/SequencerUtilities.cpp
+++ b/Engine/Source/Editor/Sequencer/Private/SequencerUtilities.cpp
@@ -84,34 +84,8 @@ TSharedRef<SWidget> FSequencerUtilities::MakeAddButton(FText HoverText, FOnGetCo
 
 TSharedRef<SWidget> FSequencerUtilities::MakeAddButton(FText HoverText, FOnClicked OnClicked, const TAttribute<bool>& HoverState, TWeakPtr<ISequencer> InSequencer)
 {
-	TSharedRef<SButton> Button =
-
-		SNew(SButton)
-		.IsFocusable(true)
-		.ButtonStyle(FAppStyle::Get(), "HoverHintOnly")
-		.ForegroundColor(FSlateColor::UseForeground())
-		.IsEnabled_Lambda([=]() { return InSequencer.IsValid() ? !InSequencer.Pin()->IsReadOnly() : false; })
-		.OnClicked(OnClicked)
-		.ContentPadding(FMargin(5, 2))
-		.HAlign(HAlign_Center)
-		.VAlign(VAlign_Center)
-		.Content()
-		[
-			SNew(SHorizontalBox)
-
-			+ SHorizontalBox::Slot()
-			.AutoWidth()
-			.VAlign(VAlign_Center)
-			.Padding(FMargin(0, 0, 2, 0))
-			[
-				SNew(SImage)
-				.ColorAndOpacity(FSlateColor::UseForeground())
-				.Image(FAppStyle::GetBrush("Plus"))
-				.ToolTipText(HoverText)
-			]
-		];
-
-	return Button;
+	TAttribute<bool> IsEnabled = MakeAttributeLambda([InSequencer]() -> bool { return InSequencer.IsValid() ? !InSequencer.Pin()->IsReadOnly() : false; });
+	return UE::Sequencer::MakeAddButton(HoverText, OnClicked, HoverState, IsEnabled);
 }
 
 void FSequencerUtilities::CreateNewSection(UMovieSceneTrack* InTrack, TWeakPtr<ISequencer> InSequencer, int32 InRowIndex, EMovieSceneBlendType InBlendType)
diff --git a/Engine/Source/Editor/Sequencer/Private/Widgets/OutlinerColumns/SPinColumnWidget.cpp b/Engine/Source/Editor/Sequencer/Private/Widgets/OutlinerColumns/SPinColumnWidget.cpp
--- a/Engine/Source/Editor/Sequencer/Private/Widgets/OutlinerColumns/SPinColumnWidget.cpp
+++ b/Engine/Source/Editor/Sequencer/Private/Widgets/OutlinerColumns/SPinColumnWidget.cpp
@@ -4,7 +4,7 @@
 
 #include "MVVM/Extensions/IPinnableExtension.h"
 #include "MVVM/PinEditorExtension.h"
-#include "MVVM/ViewModels/SequencerEditorViewModel.h"
+#include "MVVM/ViewModels/EditorViewModel.h"
 
 namespace UE::Sequencer
 {
@@ -15,7 +15,7 @@ void SPinColumnWidget::OnToggleOperationComplete()
 	RefreshSequencerTree();
 }
 
-void SPinColumnWidget::Construct(const FArguments& InArgs, const TWeakPtr<ISequencerOutlinerColumn> InWeakOutlinerColumn, const FCreateOutlinerColumnParams& InParams)
+void SPinColumnWidget::Construct(const FArguments& InArgs, const TWeakPtr<IOutlinerColumn> InWeakOutlinerColumn, const FCreateOutlinerColumnParams& InParams)
 {
 	SColumnToggleWidget::Construct(
 		SColumnToggleWidget::FArguments(),
@@ -26,10 +26,10 @@ void SPinColumnWidget::Construct(const FArguments& InArgs, const TWeakPtr<ISeque
 
 bool SPinColumnWidget::IsActive() const
 {
-	TSharedPtr<FSequencerEditorViewModel> SequencerEditor = WeakEditor.Pin();
-	if (SequencerEditor)
+	TSharedPtr<FEditorViewModel> Editor = WeakEditor.Pin();
+	if (Editor)
 	{
-		FPinEditorExtension* PinEditorExtension = SequencerEditor->CastDynamic<FPinEditorExtension>();
+		FPinEditorExtension* PinEditorExtension = Editor->CastDynamic<FPinEditorExtension>();
 		if (PinEditorExtension)
 		{
 			return PinEditorExtension->IsNodePinned(WeakOutlinerExtension);
@@ -41,10 +41,10 @@ bool SPinColumnWidget::IsActive() const
 
 void SPinColumnWidget::SetIsActive(const bool bInIsActive)
 {
-	TSharedPtr<FSequencerEditorViewModel> SequencerEditor = WeakEditor.Pin();
-	if (SequencerEditor)
+	TSharedPtr<FEditorViewModel> Editor = WeakEditor.Pin();
+	if (Editor)
 	{
-		FPinEditorExtension* PinEditorExtension = SequencerEditor->CastDynamic<FPinEditorExtension>();
+		FPinEditorExtension* PinEditorExtension = Editor->CastDynamic<FPinEditorExtension>();
 		if (PinEditorExtension)
 		{
 			PinEditorExtension->SetNodePinned(WeakOutlinerExtension, bInIsActive);
diff --git a/Engine/Source/Editor/Sequencer/Private/Widgets/OutlinerColumns/SPinColumnWidget.h b/Engine/Source/Editor/Sequencer/Private/Widgets/OutlinerColumns/SPinColumnWidget.h
--- a/Engine/Source/Editor/Sequencer/Private/Widgets/OutlinerColumns/SPinColumnWidget.h
+++ b/Engine/Source/Editor/Sequencer/Private/Widgets/OutlinerColumns/SPinColumnWidget.h
@@ -2,7 +2,7 @@
 
 #pragma once
 
-#include "Widgets/OutlinerColumns/SColumnToggleWidget.h"
+#include "MVVM/Views/OutlinerColumns/SColumnToggleWidget.h"
 
 namespace UE::Sequencer
 {
@@ -19,7 +19,7 @@ public:
 	SLATE_BEGIN_ARGS(SPinColumnWidget) {}
 	SLATE_END_ARGS()
 
-	void Construct(const FArguments& InArgs, const TWeakPtr<ISequencerOutlinerColumn> InWeakOutlinerColumn, const FCreateOutlinerColumnParams& InParams);
+	void Construct(const FArguments& InArgs, const TWeakPtr<IOutlinerColumn> InWeakOutlinerColumn, const FCreateOutlinerColumnParams& InParams);
 
 public:
 
diff --git a/Engine/Source/Editor/Sequencer/Public/ISequencerModule.h b/Engine/Source/Editor/Sequencer/Public/ISequencerModule.h
--- a/Engine/Source/Editor/Sequencer/Public/ISequencerModule.h
+++ b/Engine/Source/Editor/Sequencer/Public/ISequencerModule.h
@@ -19,7 +19,6 @@ class FExtensibilityManager;
 class FMenuBuilder;
 class FSequencerCustomizationManager;
 class ISequencerTrackEditor;
-class ISequencerOutlinerColumn;
 class ISequencerEditorObjectBinding;
 class IToolkitHost;
 class UMovieSceneSequence;
@@ -29,6 +28,7 @@ namespace UE::Sequencer
 {
 	class FTrackModel;
 	class IObjectSchema;
+	class IOutlinerColumn;
 } // namespace UE::Sequencer
 
 enum class ECurveEditorTreeFilterType : uint32;
@@ -61,7 +61,7 @@ DECLARE_DELEGATE_RetVal_OneParam(TSharedRef<ISequencerEditorObjectBinding>, FOnC
 DECLARE_DELEGATE_RetVal_OneParam(TSharedPtr<UE::Sequencer::FTrackModel>, FOnCreateTrackModel, UMovieSceneTrack*);
 
 /** A delegate which will create an outliner column */
-DECLARE_DELEGATE_RetVal(TSharedRef<ISequencerOutlinerColumn>, FOnCreateOutlinerColumn);
+DECLARE_DELEGATE_RetVal(TSharedRef<UE::Sequencer::IOutlinerColumn>, FOnCreateOutlinerColumn);
 
 /** A delegate that is executed when adding menu content. */
 DECLARE_DELEGATE_OneParam(FOnGetContextMenuContent, FMenuBuilder& /*MenuBuilder*/);
diff --git a/Engine/Source/Editor/Sequencer/Public/ISequencerSection.h b/Engine/Source/Editor/Sequencer/Public/ISequencerSection.h
--- a/Engine/Source/Editor/Sequencer/Public/ISequencerSection.h
+++ b/Engine/Source/Editor/Sequencer/Public/ISequencerSection.h
@@ -42,6 +42,7 @@ namespace UE::Sequencer
 {
 	class FCategoryModel;
 	class FChannelModel;
+	struct FViewDensityInfo;
 }
 
 /** Enumerates which edge is being resized */
@@ -62,7 +63,7 @@ namespace SequencerSectionConstants
 
 	const float DefaultSectionGripSize = 8.0f;
 
-	const float DefaultSectionHeight = 15.f;
+	const float DefaultSectionHeight = 27.f;
 
 	const FName SelectionColorName("SelectionColor");
 
@@ -191,8 +192,11 @@ public:
 	/**
 	 * @return The height of the section
 	 */
-	virtual float GetSectionHeight() const { return SequencerSectionConstants::DefaultSectionHeight; }
-	
+	UE_DEPRECATED(5.4, "Please call GetSectionHeight(const FViewDensityInfo& ViewDensity) instead.")
+	SEQUENCER_API virtual float GetSectionHeight() const;
+	SEQUENCER_API virtual float GetSectionHeight(const UE::Sequencer::FViewDensityInfo& ViewDensity) const;
+
+
 	/**
 	 * @return The width of the section drag handles
 	 */
diff --git a/Engine/Source/Editor/Sequencer/Public/ISequencerTrackEditor.h b/Engine/Source/Editor/Sequencer/Public/ISequencerTrackEditor.h
--- a/Engine/Source/Editor/Sequencer/Public/ISequencerTrackEditor.h
+++ b/Engine/Source/Editor/Sequencer/Public/ISequencerTrackEditor.h
@@ -11,6 +11,8 @@
 #include "MovieSceneTrack.h"
 #include "SequencerKeyParams.h"
 #include "Sections/MovieScene3DTransformSection.h"
+#include "MVVM/Extensions/IOutlinerExtension.h"
+#include "MVVM/ViewModelPtr.h"
 
 class FExtender;
 class ISequencer;
@@ -22,11 +24,33 @@ class UMovieScene;
 class UMovieSceneTrack;
 class UMovieSceneSequence;
 
+namespace UE::Sequencer
+{
+	class ITrackExtension;
+	class FObjectBindingModel;
+}
+
+
+
+struct FBuildColumnWidgetParams : UE::Sequencer::FCreateOutlinerViewParams
+{
+	FBuildColumnWidgetParams(UE::Sequencer::TViewModelPtr<UE::Sequencer::ITrackExtension> InTrack, const UE::Sequencer::FCreateOutlinerViewParams& InBaseParams)
+		: UE::Sequencer::FCreateOutlinerViewParams(InBaseParams)
+		, ViewModel(InTrack.AsModel())
+		, TrackModel(InTrack)
+	{}
+
+	/**  */
+	UE::Sequencer::FViewModelPtr ViewModel;
+	UE::Sequencer::TViewModelPtr<UE::Sequencer::ITrackExtension> TrackModel;
+};
+
 /** Data structure containing information required to build an edit widget */
-struct FBuildEditWidgetParams
+struct FBuildEditWidgetParams : FBuildColumnWidgetParams
 {
-	FBuildEditWidgetParams()
-		: RowIndex(0)
+	FBuildEditWidgetParams(const FBuildColumnWidgetParams& Other)
+		: FBuildColumnWidgetParams(Other)
+		, RowIndex(0)
 		, TrackInsertRowIndex(0)
 	{}
 
@@ -107,13 +131,18 @@ public:
 	virtual void BuildAddTrackMenu(FMenuBuilder& MenuBuilder) = 0;
 
 	/**
-	 * Builds up the object binding edit buttons for the outliner.
+	 * Deprecated in principle, not practice (due to difficulty deprecating public virtual functions).
+	 * Still called from the newer BuildObjectBindingColumnWidgets which should now be preferred
+	 */
+	virtual void BuildObjectBindingEditButtons(TSharedPtr<SHorizontalBox> EditBox, const FGuid& ObjectBinding, const UClass* ObjectClass){}
+
+	/**
+	 * Build a column widget for the specified object binding on the outliner
 	 *
-	 * @param EditBox The edit box to add buttons to.
+	 * @param GetEditBox A function to retrieve the edit box to populate
 	 * @param ObjectBinding The object binding this is for.
-	 * @param ObjectClass The class of the object this is for.
 	 */
-	virtual void BuildObjectBindingEditButtons(TSharedPtr<SHorizontalBox> EditBox, const FGuid& ObjectBinding, const UClass* ObjectClass) = 0;
+	virtual void BuildObjectBindingColumnWidgets(TFunctionRef<TSharedRef<SHorizontalBox>()> GetEditBox, const UE::Sequencer::TViewModelPtr<UE::Sequencer::FObjectBindingModel>& ObjectBinding, const UE::Sequencer::FCreateOutlinerViewParams& InParams, const FName& InColumnName) = 0;
 
 	/**
 	 * Builds up the object binding track menu for the outliner.
@@ -153,6 +182,16 @@ public:
 	 */
 	virtual TSharedPtr<SWidget> BuildOutlinerEditWidget(const FGuid& ObjectBinding, UMovieSceneTrack* Track, const FBuildEditWidgetParams& Params) = 0;
 
+
+	/**
+	 * Builds an outliner column widget for the outliner nodes which represent tracks which are edited by this editor.
+	 * 
+	 * @param Params Parameter struct containing data relevant to the edit widget
+	 * @param ColumnName The name of the column. See FCommonOutlinerNames
+	 * @returns The the widget to display in the outliner, or an empty shared ptr if not widget is to be displayed.
+	 */
+	virtual TSharedPtr<SWidget> BuildOutlinerColumnWidget(const FBuildColumnWidgetParams& Params, const FName& ColumnName) = 0;
+
 	/**
 	 * Builds the context menu for the track.
 	 * @param MenuBuilder The menu builder to use to build the track menu. 
diff --git a/Engine/Source/Editor/Sequencer/Public/MVVM/Extensions/ITrackExtension.h b/Engine/Source/Editor/Sequencer/Public/MVVM/Extensions/ITrackExtension.h
--- a/Engine/Source/Editor/Sequencer/Public/MVVM/Extensions/ITrackExtension.h
+++ b/Engine/Source/Editor/Sequencer/Public/MVVM/Extensions/ITrackExtension.h
@@ -11,13 +11,12 @@
 class UMovieSceneSection;
 class UMovieSceneTrack;
 class ISequencerTrackEditor;
-struct FViewModelChildren;
 
-namespace UE
-{
-namespace Sequencer
+namespace UE::Sequencer
 {
 
+struct FViewModelChildren;
+
 class SEQUENCER_API ITrackExtension
 {
 public:
@@ -47,6 +46,5 @@ public:
 	virtual TSharedPtr<ISequencerTrackEditor> GetTrackEditor() const = 0;
 };
 
-} // namespace Sequencer
-} // namespace UE
+} // namespace UE::Sequencer
 
diff --git a/Engine/Source/Editor/Sequencer/Public/MVVM/ViewModels/CategoryModel.h b/Engine/Source/Editor/Sequencer/Public/MVVM/ViewModels/CategoryModel.h
--- a/Engine/Source/Editor/Sequencer/Public/MVVM/ViewModels/CategoryModel.h
+++ b/Engine/Source/Editor/Sequencer/Public/MVVM/ViewModels/CategoryModel.h
@@ -121,7 +121,7 @@ public:
 	FText GetLabel() const override;
 	FSlateFontInfo GetLabelFont() const override;
 	FText GetLabelToolTipText() const override;
-	TSharedRef<SWidget> CreateOutlinerView(const FCreateOutlinerViewParams& InParams) override;
+	TSharedPtr<SWidget> CreateOutlinerViewForColumn(const FCreateOutlinerViewParams& InParams, const FName& InColumnName) override;
 
 	/*~ ITrackAreaExtension */
 	FTrackAreaParameters GetTrackAreaParameters() const override;
diff --git a/Engine/Source/Editor/Sequencer/Public/MVVM/ViewModels/ChannelModel.h b/Engine/Source/Editor/Sequencer/Public/MVVM/ViewModels/ChannelModel.h
--- a/Engine/Source/Editor/Sequencer/Public/MVVM/ViewModels/ChannelModel.h
+++ b/Engine/Source/Editor/Sequencer/Public/MVVM/ViewModels/ChannelModel.h
@@ -227,7 +227,7 @@ public:
 	FText GetLabel() const override;
 	FSlateFontInfo GetLabelFont() const override;
 	FText GetLabelToolTipText() const override;
-	TSharedRef<SWidget> CreateOutlinerView(const FCreateOutlinerViewParams& InParams) override;
+	TSharedPtr<SWidget> CreateOutlinerViewForColumn(const FCreateOutlinerViewParams& InParams, const FName& InColumnName) override;
 
 	/*~ ICurveEditorTreeItem */
 	void CreateCurveModels(TArray<TUniquePtr<FCurveModel>>& OutCurveModels) override;
diff --git a/Engine/Source/Editor/Sequencer/Public/MVVM/ViewModels/SequencerOutlinerViewModel.h b/Engine/Source/Editor/Sequencer/Public/MVVM/ViewModels/SequencerOutlinerViewModel.h
--- a/Engine/Source/Editor/Sequencer/Public/MVVM/ViewModels/SequencerOutlinerViewModel.h
+++ b/Engine/Source/Editor/Sequencer/Public/MVVM/ViewModels/SequencerOutlinerViewModel.h
@@ -32,6 +32,7 @@ public:
 	/*~ FOutlinerViewModel */
 	TSharedPtr<SWidget> CreateContextMenuWidget() override;
 	TSharedRef<FDragDropOperation> InitiateDrag(TArray<TWeakViewModelPtr<IOutlinerExtension>>&& InDraggedModels) override;
+	void RequestUpdate() override;
 
 	FFrameNumber GetNextKey(const TArray<TSharedRef<UE::Sequencer::FViewModel>>& InNodes, FFrameNumber FrameNumber, EMovieSceneTimeUnit TimeUnit, const TRange<FFrameNumber>& Range);
 	FFrameNumber GetPreviousKey(const TArray<TSharedRef<UE::Sequencer::FViewModel>>& InNodes, FFrameNumber FrameNumber, EMovieSceneTimeUnit TimeUnit, const TRange<FFrameNumber>& Range);
diff --git a/Engine/Source/Editor/Sequencer/Public/MovieSceneTrackEditor.h b/Engine/Source/Editor/Sequencer/Public/MovieSceneTrackEditor.h
--- a/Engine/Source/Editor/Sequencer/Public/MovieSceneTrackEditor.h
+++ b/Engine/Source/Editor/Sequencer/Public/MovieSceneTrackEditor.h
@@ -215,9 +215,10 @@ public:
 
 	virtual void BindCommands(TSharedRef<FUICommandList> SequencerCommandBindings) override;
 	virtual void BuildAddTrackMenu(FMenuBuilder& MenuBuilder) override;
-	virtual void BuildObjectBindingEditButtons(TSharedPtr<SHorizontalBox> EditBox, const FGuid& ObjectBinding, const UClass* ObjectClass) override;
+	virtual void BuildObjectBindingColumnWidgets(TFunctionRef<TSharedRef<SHorizontalBox>()> GetEditBox, const UE::Sequencer::TViewModelPtr<UE::Sequencer::FObjectBindingModel>& ObjectBinding, const UE::Sequencer::FCreateOutlinerViewParams& InParams, const FName& InColumnName) override;
 	virtual void BuildObjectBindingTrackMenu(FMenuBuilder& MenuBuilder, const TArray<FGuid>& ObjectBindings, const UClass* ObjectClass) override;
 	virtual TSharedPtr<SWidget> BuildOutlinerEditWidget(const FGuid& ObjectBinding, UMovieSceneTrack* Track, const FBuildEditWidgetParams& Params) override;
+	virtual TSharedPtr<SWidget> BuildOutlinerColumnWidget(const FBuildColumnWidgetParams& Params, const FName& ColumnName) override;
 	virtual void BuildTrackContextMenu( FMenuBuilder& MenuBuilder, UMovieSceneTrack* Track ) override;
 	virtual bool HandleAssetAdded(UObject* Asset, const FGuid& TargetObjectGuid) override;
 	virtual bool OnAllowDrop(const FDragDropEvent& DragDropEvent, FSequencerDragDropParams& DragDropParams) override;
diff --git a/Engine/Source/Editor/Sequencer/Public/SequencerSettings.h b/Engine/Source/Editor/Sequencer/Public/SequencerSettings.h
--- a/Engine/Source/Editor/Sequencer/Public/SequencerSettings.h
+++ b/Engine/Source/Editor/Sequencer/Public/SequencerSettings.h
@@ -15,6 +15,11 @@ enum class EAllowEditsMode : uint8;
 enum class EKeyGroupMode : uint8;
 enum class EMovieSceneKeyInterpolation : uint8;
 
+namespace UE::Sequencer
+{
+	enum class EViewDensity;
+}
+
 UENUM()
 enum ESequencerSpawnPosition : int
 {
@@ -454,6 +459,11 @@ public:
 	/** Sets the tree view width percentage */
 	void SetTreeViewWidth(float InTreeViewWidth);
 
+	/** Gets the saved view density */
+	UE::Sequencer::EViewDensity GetViewDensity() const;
+	/** Sets the saved view density */
+	void SetViewDensity(FName InViewDensity);
+
 	/** Gets whether the given track filter is enabled */
 	bool IsTrackFilterEnabled(const FString& TrackFilter) const;
 	/** Sets whether the track filter should be enabled/disabled */
@@ -720,6 +730,9 @@ protected:
 	UPROPERTY(config, EditAnywhere, Category = General)
 	float TreeViewWidth;
 
+	UPROPERTY(config, EditAnywhere, Category = General)
+	FName ViewDensity;
+
 	/** The track filters that are enabled */
 	UPROPERTY(config, EditAnywhere, Category = General)
 	TArray<FString> TrackFilters;
diff --git a/Engine/Source/Editor/SequencerCore/Private/MVVM/Extensions/IOutlinerExtension.cpp b/Engine/Source/Editor/SequencerCore/Private/MVVM/Extensions/IOutlinerExtension.cpp
--- a/Engine/Source/Editor/SequencerCore/Private/MVVM/Extensions/IOutlinerExtension.cpp
+++ b/Engine/Source/Editor/SequencerCore/Private/MVVM/Extensions/IOutlinerExtension.cpp
@@ -9,10 +9,9 @@
 #include "Misc/StringBuilder.h"
 #include "Styling/AppStyle.h"
 #include "Styling/SlateColor.h"
+#include "Widgets/SNullWidget.h"
 
-namespace UE
-{
-namespace Sequencer
+namespace UE::Sequencer
 {
 
 void GetParentPathImpl(TParentModelIterator<IOutlinerExtension>& Iterator, FStringBuilderBase& OutString)
@@ -139,6 +138,15 @@ void IOutlinerExtension::GetPathName(const FViewModel& Item, FStringBuilderBase&
 	}
 }
 
-} // namespace Sequencer
-} // namespace UE
+TSharedRef<SWidget> IOutlinerExtension::CreateOutlinerView(const FCreateOutlinerViewParams& InParams)
+{
+	return SNullWidget::NullWidget;
+}
+
+TSharedPtr<SWidget> IOutlinerExtension::CreateOutlinerViewForColumn(const FCreateOutlinerViewParams& InParams, const FName& InColumnName)
+{
+	return nullptr;
+}
+
+} // namespace UE::Sequencer
 
diff --git a/Engine/Source/Editor/SequencerCore/Private/MVVM/ViewModelTypeIDs.cpp b/Engine/Source/Editor/SequencerCore/Private/MVVM/ViewModelTypeIDs.cpp
--- a/Engine/Source/Editor/SequencerCore/Private/MVVM/ViewModelTypeIDs.cpp
+++ b/Engine/Source/Editor/SequencerCore/Private/MVVM/ViewModelTypeIDs.cpp
@@ -33,6 +33,7 @@
 #include "MVVM/ViewModels/OutlinerViewModel.h"
 #include "MVVM/ViewModels/TrackAreaViewModel.h"
 #include "MVVM/ViewModels/ViewModel.h"
+#include "MVVM/ViewModels/EditorSharedViewModelData.h"
 #include "MVVM/Views/ITrackAreaHotspot.h"
 
 namespace UE
@@ -83,6 +84,8 @@ UE_SEQUENCER_DEFINE_CASTABLE(FOutlinerSpacer);
 UE_SEQUENCER_DEFINE_CASTABLE(FOutlinerViewModel);
 UE_SEQUENCER_DEFINE_CASTABLE(FTrackAreaViewModel);
 UE_SEQUENCER_DEFINE_CASTABLE(FViewModel);
+UE_SEQUENCER_DEFINE_CASTABLE(FSharedViewModelData);
+UE_SEQUENCER_DEFINE_CASTABLE(FEditorSharedViewModelData);
 
 // Views
 UE_SEQUENCER_DEFINE_CASTABLE(ITrackAreaHotspot);
diff --git a/Engine/Source/Editor/SequencerCore/Private/MVVM/ViewModels/EditorViewModel.cpp b/Engine/Source/Editor/SequencerCore/Private/MVVM/ViewModels/EditorViewModel.cpp
--- a/Engine/Source/Editor/SequencerCore/Private/MVVM/ViewModels/EditorViewModel.cpp
+++ b/Engine/Source/Editor/SequencerCore/Private/MVVM/ViewModels/EditorViewModel.cpp
@@ -121,6 +121,16 @@ bool FEditorViewModel::IsReadOnly() const
 	return false;
 }
 
+FViewDensityInfo FEditorViewModel::GetViewDensity() const
+{
+	return ViewDensity;
+}
+
+void FEditorViewModel::SetViewDensity(const FViewDensityInfo& InViewDensity)
+{
+	ViewDensity = InViewDensity;
+}
+
 } // namespace Sequencer
 } // namespace UE
 
diff --git a/Engine/Source/Editor/SequencerCore/Private/MVVM/ViewModels/OutlinerColumns/NavOutlinerColumn.cpp b/Engine/Source/Editor/SequencerCore/Private/MVVM/ViewModels/OutlinerColumns/NavOutlinerColumn.cpp
new file mode 100644
--- /dev/null
+++ b/Engine/Source/Editor/SequencerCore/Private/MVVM/ViewModels/OutlinerColumns/NavOutlinerColumn.cpp
@@ -0,0 +1,16 @@
+// Copyright Epic Games, Inc. All Rights Reserved.
+
+#include "MVVM/ViewModels/OutlinerColumns/NavOutlinerColumn.h"
+
+namespace UE::Sequencer
+{
+
+FNavOutlinerColumn::FNavOutlinerColumn()
+{
+	Name     = FCommonOutlinerNames::Nav;
+	Label    = NSLOCTEXT("FNavOutlinerColumn", "NavColumnNav", "Nav");
+	Position = FOutlinerColumnPosition{ 10, EOutlinerColumnGroup::RightGutter };
+	Layout   = FOutlinerColumnLayout{ 58.f, FMargin(8.f, 0.f), HAlign_Center, VAlign_Center, EOutlinerColumnSizeMode::Fixed, EOutlinerColumnFlags::OverflowSubsequentEmptyCells };
+}
+
+} // namespace UE::Sequencer
diff --git a/Engine/Source/Editor/SequencerCore/Private/MVVM/ViewModels/OutlinerSpacer.cpp b/Engine/Source/Editor/SequencerCore/Private/MVVM/ViewModels/OutlinerSpacer.cpp
--- a/Engine/Source/Editor/SequencerCore/Private/MVVM/ViewModels/OutlinerSpacer.cpp
+++ b/Engine/Source/Editor/SequencerCore/Private/MVVM/ViewModels/OutlinerSpacer.cpp
@@ -43,12 +43,9 @@ FName FOutlinerSpacer::GetIdentifier() const
 
 FOutlinerSizing FOutlinerSpacer::GetOutlinerSizing() const
 {
-	return FOutlinerSizing{ DesiredSpacerHeight };
-}
-
-TSharedRef<SWidget> FOutlinerSpacer::CreateOutlinerView(const FCreateOutlinerViewParams& InParams)
-{
-	return SNew(SBox).HeightOverride(DesiredSpacerHeight);
+	FOutlinerSizing Sizing{ DesiredSpacerHeight };
+	Sizing.Flags = EOutlinerSizingFlags::None;
+	return Sizing;
 }
 
 TSharedPtr<SWidget> FOutlinerSpacer::CreateContextMenuWidget(const FCreateOutlinerContextMenuWidgetParams& InParams)
diff --git a/Engine/Source/Editor/SequencerCore/Private/MVVM/ViewModels/OutlinerViewModel.cpp b/Engine/Source/Editor/SequencerCore/Private/MVVM/ViewModels/OutlinerViewModel.cpp
--- a/Engine/Source/Editor/SequencerCore/Private/MVVM/ViewModels/OutlinerViewModel.cpp
+++ b/Engine/Source/Editor/SequencerCore/Private/MVVM/ViewModels/OutlinerViewModel.cpp
@@ -100,6 +100,11 @@ TSharedPtr<SWidget> FOutlinerViewModel::CreateContextMenuWidget()
 	return nullptr;
 }
 
+void FOutlinerViewModel::RequestUpdate()
+{
+
+}
+
 FViewModelPtr FOutlinerViewModel::GetRootItem() const
 {
 	return WeakRootDataModel.Pin();
diff --git a/Engine/Source/Editor/SequencerCore/Private/MVVM/ViewModels/TrackAreaViewModel.cpp b/Engine/Source/Editor/SequencerCore/Private/MVVM/ViewModels/TrackAreaViewModel.cpp
--- a/Engine/Source/Editor/SequencerCore/Private/MVVM/ViewModels/TrackAreaViewModel.cpp
+++ b/Engine/Source/Editor/SequencerCore/Private/MVVM/ViewModels/TrackAreaViewModel.cpp
@@ -79,8 +79,6 @@ void FTrackAreaViewModel::SetHotspot(TSharedPtr<ITrackAreaHotspot> NewHotspot)
 
 void FTrackAreaViewModel::AddHotspot(TSharedPtr<ITrackAreaHotspot> NewHotspot)
 {
-	using namespace UE::Sequencer;
-
 	// We always set a hotspot if there is none
 	if (bHotspotLocked == false)
 	{
diff --git a/Engine/Source/Editor/SequencerCore/Private/MVVM/ViewModels/ViewDensity.cpp b/Engine/Source/Editor/SequencerCore/Private/MVVM/ViewModels/ViewDensity.cpp
new file mode 100644
--- /dev/null
+++ b/Engine/Source/Editor/SequencerCore/Private/MVVM/ViewModels/ViewDensity.cpp
@@ -0,0 +1,44 @@
+// Copyright Epic Games, Inc. All Rights Reserved.
+
+#include "MVVM/ViewModels/ViewDensity.h"
+#include "HAL/IConsoleManager.h"
+
+namespace UE::Sequencer
+{
+
+float GSequencerOutlinerRelaxedHeight = 22.f;
+FAutoConsoleVariableRef CVarSequencerOutlinerRelaxedHeight(
+	TEXT("Sequencer.Outliner.RelaxedHeight"),
+	GSequencerOutlinerRelaxedHeight,
+	TEXT("(Default: 22.f. Defines the height of outliner items when in relaxed mode.")
+	);
+float GSequencerOutlinerExpandedHeight = 28.f;
+FAutoConsoleVariableRef CVarSequencerOutlinerExpandedHeight(
+	TEXT("Sequencer.Outliner.ExpandedHeight"),
+	GSequencerOutlinerExpandedHeight,
+	TEXT("(Default: 28.f. Defines the height of outliner items when in expanded mode.")
+	);
+
+
+FViewDensityInfo::FViewDensityInfo()
+	: UniformHeight(GetUniformHeight(EViewDensity::Relaxed))
+	, Density(EViewDensity::Relaxed)
+{}
+
+FViewDensityInfo::FViewDensityInfo(EViewDensity InViewDensity)
+	: UniformHeight(GetUniformHeight(InViewDensity))
+	, Density(InViewDensity)
+{}
+
+TOptional<float> FViewDensityInfo::GetUniformHeight(EViewDensity Density)
+{
+	switch(Density)
+	{
+		case EViewDensity::Relaxed:  return GSequencerOutlinerRelaxedHeight;
+		case EViewDensity::Expanded: return GSequencerOutlinerExpandedHeight;
+		default:                     return TOptional<float>();
+	}
+}
+
+
+} // namespace UE::Sequencer
\ No newline at end of file
diff --git a/Engine/Source/Editor/SequencerCore/Private/MVVM/Views/SKeyNavigationButtons.cpp b/Engine/Source/Editor/SequencerCore/Private/MVVM/Views/SKeyNavigationButtons.cpp
--- a/Engine/Source/Editor/SequencerCore/Private/MVVM/Views/SKeyNavigationButtons.cpp
+++ b/Engine/Source/Editor/SequencerCore/Private/MVVM/Views/SKeyNavigationButtons.cpp
@@ -8,8 +8,6 @@
 #include "HAL/PlatformCrt.h"
 #include "Layout/Children.h"
 #include "Layout/Margin.h"
-#include "MVVM/Extensions/IHoveredExtension.h"
-#include "MVVM/ICastable.h"
 #include "MVVM/ViewModels/ViewModel.h"
 #include "Math/NumericLimits.h"
 #include "Misc/FrameNumber.h"
@@ -17,14 +15,9 @@
 #include "Misc/Optional.h"
 #include "SlotBase.h"
 #include "Styling/AppStyle.h"
-#include "Styling/ISlateStyle.h"
-#include "Styling/SlateColor.h"
-#include "Templates/TypeHash.h"
-#include "Types/SlateEnums.h"
-#include "Widgets/Input/SButton.h"
-#include "Widgets/Layout/SBorder.h"
 #include "Widgets/SBoxPanel.h"
-#include "Widgets/Text/STextBlock.h"
+#include "MVVM/Views/OutlinerColumns/SOutlinerColumnButton.h"
+
 
 struct FSlateBrush;
 
@@ -42,10 +35,6 @@ void SKeyNavigationButtons::Construct(const FArguments& InArgs, const TSharedPtr
 	SetTimeEvent = InArgs._OnSetTime;
 	TimeAttribute = InArgs._Time;
 
-	const FSlateBrush* NoBorder = FAppStyle::GetBrush( "NoBorder" );
-
-	TAttribute<FLinearColor> HoverTint(this, &SKeyNavigationButtons::GetHoverTint);
-
 	ChildSlot
 	[
 		SNew(SHorizontalBox)
@@ -54,85 +43,40 @@ void SKeyNavigationButtons::Construct(const FArguments& InArgs, const TSharedPtr
 		+ SHorizontalBox::Slot()
 		.VAlign(VAlign_Center)
 		.AutoWidth()
-		.Padding(3, 0, 0, 0)
 		[
-			SNew(SBorder)
-			.Padding(0)
-			.BorderImage(NoBorder)
-			.ColorAndOpacity(HoverTint)
-			[
-				SNew(SButton)
-				.ButtonStyle(FAppStyle::Get(), "HoverHintOnly")
-				.ToolTipText(InArgs._PreviousKeyToolTip)
-				.OnClicked(this, &SKeyNavigationButtons::OnPreviousKeyClicked)
-				.ForegroundColor( FSlateColor::UseForeground() )
-				.ContentPadding(2)
-				.IsFocusable(false)
-				[
-					SNew(STextBlock)
-					.Font(FAppStyle::Get().GetFontStyle("FontAwesome.7"))
-					.Text(FText::FromString(FString(TEXT("\xf060"))) /*fa-arrow-left*/)
-				]
-			]
+			SNew(SOutlinerColumnButton)
+			.IsFocusable(false)
+			.ToolTipText(InArgs._PreviousKeyToolTip)
+			.Image(FAppStyle::GetBrush("Sequencer.Outliner.PreviousKey"))
+			.OnClicked(this, &SKeyNavigationButtons::OnPreviousKeyClicked)
 		]
+
 		// Add key slot
 		+ SHorizontalBox::Slot()
 		.VAlign(VAlign_Center)
+		.Padding(8.f, 0.f)
 		.AutoWidth()
 		[
-			SNew(SBorder)
-			.Padding(0)
-			.BorderImage(NoBorder)
-			.ColorAndOpacity(HoverTint)
-			.IsEnabled(InArgs._IsEnabled)
-			[
-				SNew(SButton)
-				.ButtonStyle(FAppStyle::Get(), "HoverHintOnly")
-				.ToolTipText(InArgs._AddKeyToolTip)
-				.OnClicked(this, &SKeyNavigationButtons::OnAddKeyClicked)
-				.ForegroundColor( FSlateColor::UseForeground() )
-				.ContentPadding(2)
-				.IsFocusable(false)
-				[
-					SNew(STextBlock)
-					.Font(FAppStyle::Get().GetFontStyle("FontAwesome.7"))
-					.Text(FText::FromString(FString(TEXT("\xf055"))) /*fa-plus-circle*/)
-				]
-			]
+			SNew(SOutlinerColumnButton)
+			.IsFocusable(false)
+			.ToolTipText(InArgs._AddKeyToolTip)
+			.Image(FAppStyle::GetBrush("Sequencer.Outliner.AddKey"))
+			.OnClicked(this, &SKeyNavigationButtons::OnAddKeyClicked)
 		]
 		// Next key slot
 		+ SHorizontalBox::Slot()
 		.VAlign(VAlign_Center)
 		.AutoWidth()
 		[
-			SNew(SBorder)
-			.Padding(0)
-			.BorderImage(NoBorder)
-			.ColorAndOpacity(HoverTint)
-			[
-				SNew(SButton)
-				.ButtonStyle(FAppStyle::Get(), "HoverHintOnly")
-				.ToolTipText(InArgs._NextKeyToolTip)
-				.OnClicked(this, &SKeyNavigationButtons::OnNextKeyClicked)
-				.ContentPadding(2)
-				.ForegroundColor( FSlateColor::UseForeground() )
-				.IsFocusable(false)
-				[
-					SNew(STextBlock)
-					.Font(FAppStyle::Get().GetFontStyle("FontAwesome.7"))
-					.Text(FText::FromString(FString(TEXT("\xf061"))) /*fa-arrow-right*/)
-				]
-			]
+			SNew(SOutlinerColumnButton)
+			.IsFocusable(false)
+			.ToolTipText(InArgs._NextKeyToolTip)
+			.Image(FAppStyle::GetBrush("Sequencer.Outliner.NextKey"))
+			.OnClicked(this, &SKeyNavigationButtons::OnNextKeyClicked)
 		]
 	];
 }
 
-FLinearColor SKeyNavigationButtons::GetHoverTint() const
-{
-	IHoveredExtension* HoveredExtension = ICastable::CastWeakPtr<IHoveredExtension>(WeakModel);
-	return HoveredExtension && HoveredExtension->IsHovered() ? FLinearColor(1,1,1,0.9f) : FLinearColor(1,1,1,0.4f);
-}
-
 FReply SKeyNavigationButtons::OnPreviousKeyClicked()
 {
 	FFrameTime ClosestPreviousKeyDistance = FFrameTime(TNumericLimits<int32>::Max(), 0.99999f);
diff --git a/Engine/Source/Editor/SequencerCore/Private/MVVM/Views/SOutlinerItemViewBase.cpp b/Engine/Source/Editor/SequencerCore/Private/MVVM/Views/SOutlinerItemViewBase.cpp
--- a/Engine/Source/Editor/SequencerCore/Private/MVVM/Views/SOutlinerItemViewBase.cpp
+++ b/Engine/Source/Editor/SequencerCore/Private/MVVM/Views/SOutlinerItemViewBase.cpp
@@ -46,9 +46,7 @@ struct FPointerEvent;
 
 #define LOCTEXT_NAMESPACE "SOutlinerItemViewBase"
 
-namespace UE
-{
-namespace Sequencer
+namespace UE::Sequencer
 {
 
 void SOutlinerItemViewBase::Construct(
@@ -66,6 +64,11 @@ void SOutlinerItemViewBase::Construct(
 	IsRowHoveredAttribute  = MakeAttributeSP(&InTableRow.Get(), &ISequencerTreeViewRow::IsHovered);
 	IsRowSelectedAttribute = MakeAttributeSP(&InTableRow.Get(), &ISequencerTreeViewRow::IsItemSelected);
 
+	if (!IsReadOnlyAttribute.IsSet())
+	{
+		IsReadOnlyAttribute = MakeAttributeSP(InWeakEditor.Pin().ToSharedRef(), &FEditorViewModel::IsReadOnly);
+	}
+
 	TViewModelPtr<IOutlinerExtension> OutlinerExtension = WeakOutlinerExtension.Pin();
 	checkf(OutlinerExtension, TEXT("Attempting to create an outliner node from a null model"));
 
@@ -92,7 +95,7 @@ void SOutlinerItemViewBase::Construct(
 		InnerNodePadding = FMargin(0.f);
 	}
 
-	TableRowStyle = &FAppStyle::Get().GetWidgetStyle<FTableRowStyle>("Sequencer.TableView.Row");
+	TableRowStyle = &FAppStyle::Get().GetWidgetStyle<FTableRowStyle>("Sequencer.Outliner.Row");
 
 	TSharedPtr<SWidget> LabelContent;
 
@@ -425,8 +428,7 @@ FSlateColor SOutlinerItemViewBase::GetForegroundBasedOnSelection() const
 	return IsRowSelectedAttribute.Get() ? TableRowStyle->SelectedTextColor : TableRowStyle->TextColor;
 }
 
-} // namespace Sequencer
-} // namespace UE
+} // namespace UE::Sequencer
 
 #undef LOCTEXT_NAMESPACE
 
