Sequencer: Rerun Construction Scripts is now respected for any base blueprint class

diff --git a/Engine/Source/Editor/Sequencer/Private/Sequencer.cpp b/Engine/Source/Editor/Sequencer/Private/Sequencer.cpp
--- a/Engine/Source/Editor/Sequencer/Private/Sequencer.cpp
+++ b/Engine/Source/Editor/Sequencer/Private/Sequencer.cpp
@@ -1636,25 +1636,38 @@ void FSequencer::RerunConstructionScripts()
 
 void FSequencer::GetConstructionScriptActors(UMovieScene* MovieScene, FMovieSceneSequenceIDRef SequenceID, TSet<TWeakObjectPtr<AActor> >& BoundActors, TArray < TPair<FMovieSceneSequenceID, FGuid> >& BoundGuids)
 {
+	auto ShouldRerunConstructionScriptsForClass = [](UClass* Class)
+	{
+		// If any blueprints in the class hierarchy are marked as bRunConstructionScriptInSequencer, we need to rerun them
+		while (Class)
+		{
+			UBlueprint* Blueprint = Cast<UBlueprint>(Class->ClassGeneratedBy);
+			if (!Blueprint)
+			{
+				// As soon as we find a non BP class, we can't re-run construction scripts any more
+				return false;
+			}
+			else if (Blueprint->bRunConstructionScriptInSequencer)
+			{
+				return true;
+			}
+
+			Class = Class->GetSuperClass();
+		}
+		return false;
+	};
+
 	for (int32 Index = 0; Index < MovieScene->GetPossessableCount(); ++Index)
 	{
 		FGuid ThisGuid = MovieScene->GetPossessable(Index).GetGuid();
 
 		for (TWeakObjectPtr<> WeakObject : FindBoundObjects(ThisGuid, SequenceID))
 		{
-			if (WeakObject.IsValid())
+			AActor* Actor = Cast<AActor>(WeakObject.Get());
+			if (Actor && ShouldRerunConstructionScriptsForClass(Actor->GetClass()))
 			{
-				AActor* Actor = Cast<AActor>(WeakObject.Get());
-	
-				if (Actor)
-				{
-					UBlueprint* Blueprint = Cast<UBlueprint>(Actor->GetClass()->ClassGeneratedBy);
-					if (Blueprint && Blueprint->bRunConstructionScriptInSequencer)
-					{
-						BoundActors.Add(Actor);
-						BoundGuids.Add(TPair<FMovieSceneSequenceID, FGuid>(SequenceID, ThisGuid));
-					}
-				}
+				BoundActors.Add(Actor);
+				BoundGuids.Add(TPair<FMovieSceneSequenceID, FGuid>(SequenceID, ThisGuid));
 			}
 		}
 	}
@@ -1665,19 +1678,11 @@ void FSequencer::GetConstructionScriptActors(UMovieScene* MovieScene, FMovieScen
 
 		for (TWeakObjectPtr<> WeakObject : FindBoundObjects(ThisGuid, SequenceID))
 		{
-			if (WeakObject.IsValid())
+			AActor* Actor = Cast<AActor>(WeakObject.Get());
+			if (Actor && ShouldRerunConstructionScriptsForClass(Actor->GetClass()))
 			{
-				AActor* Actor = Cast<AActor>(WeakObject.Get());
-
-				if (Actor)
-				{
-					UBlueprint* Blueprint = Cast<UBlueprint>(Actor->GetClass()->ClassGeneratedBy);
-					if (Blueprint && Blueprint->bRunConstructionScriptInSequencer)
-					{
-						BoundActors.Add(Actor);
-						BoundGuids.Add(TPair<FMovieSceneSequenceID, FGuid>(SequenceID, ThisGuid));
-					}
-				}
+				BoundActors.Add(Actor);
+				BoundGuids.Add(TPair<FMovieSceneSequenceID, FGuid>(SequenceID, ThisGuid));
 			}
 		}
 	}
