Sequencer: Fix moving marked frames not moving the actual selected index

diff --git a/Engine/Source/Editor/Sequencer/Private/Tools/EditToolDragOperations.cpp b/Engine/Source/Editor/Sequencer/Private/Tools/EditToolDragOperations.cpp
--- a/Engine/Source/Editor/Sequencer/Private/Tools/EditToolDragOperations.cpp
+++ b/Engine/Source/Editor/Sequencer/Private/Tools/EditToolDragOperations.cpp
@@ -1522,8 +1522,9 @@ void FMoveKeysAndSections::HandleMarkedFrameMovement(TOptional<FFrameNumber> Max
 	UMovieScene* FocusedMovieScene = Sequencer.GetFocusedMovieSceneSequence()->GetMovieScene();
 	const TArray<FMovieSceneMarkedFrame>& AllMarkedFrames = FocusedMovieScene->GetMarkedFrames();
 
-	for (int32 MarkIndex = 0; MarkIndex < MarkedFrames.Num(); ++MarkIndex)
+	for (TSet<int32>::TConstIterator It = MarkedFrames.CreateConstIterator(); It; ++It)
 	{
+		const int32 MarkIndex = *It;
 		const FMovieSceneMarkedFrame& MarkedFrame = AllMarkedFrames[MarkIndex];
 		const FFrameNumber NewMarkTime = MarkedFrame.FrameNumber + EffectiveDelta;
 		FocusedMovieScene->SetMarkedFrame(MarkIndex, NewMarkTime);
