Sequencer: Add the ability to move all the selected marks either by dragging or by nudging (ctrl-arrow keys)

diff --git a/Engine/Source/Editor/Sequencer/Private/Sequencer.cpp b/Engine/Source/Editor/Sequencer/Private/Sequencer.cpp
--- a/Engine/Source/Editor/Sequencer/Private/Sequencer.cpp
+++ b/Engine/Source/Editor/Sequencer/Private/Sequencer.cpp
@@ -2036,6 +2036,8 @@ bool FSequencer::CanSnapToFrame() const
 
 void FSequencer::TransformSelectedKeysAndSections(FFrameTime InDeltaTime, float InScale)
 {
+	using namespace UE::Sequencer;
+
 	FScopedTransaction TransformKeysAndSectionsTransaction(NSLOCTEXT("Sequencer", "TransformKeysandSections_Transaction", "Transform Keys and Sections"));
 	bool bAnythingChanged = false;
 
@@ -2167,6 +2169,26 @@ void FSequencer::TransformSelectedKeysAndSections(FFrameTime InDeltaTime, float
 				}
 			}
 		}
+
+		// Marked frames
+		const FMarkedFrameSelection& SelectedMarkedFrames = ViewModel->GetSelection()->MarkedFrames;
+		if (SelectedMarkedFrames.Num() > 0)
+		{
+			bAnythingChanged = true;
+
+			UMovieScene* FocusedMovieScene = GetFocusedMovieSceneSequence()->GetMovieScene();
+			FocusedMovieScene->Modify();
+			
+			for (TSet<int32>::TConstIterator It = SelectedMarkedFrames.GetSelected(); It; ++It)
+			{
+				const int32 MarkIndex = *It;
+				FFrameNumber FrameNumber = FocusedMovieScene->GetMarkedFrames()[MarkIndex].FrameNumber;
+				FrameNumber = (OriginTime + InDeltaTime + (FrameNumber - OriginTime) * InScale).FloorToFrame();
+				FocusedMovieScene->SetMarkedFrame(MarkIndex, FrameNumber);
+			}
+
+			FocusedMovieScene->SortMarkedFrames();
+		}
 	}
 	
 	// Remove any null sections so we don't need a null check inside the loop.
diff --git a/Engine/Source/Editor/Sequencer/Private/SequencerTimeSliderController.cpp b/Engine/Source/Editor/Sequencer/Private/SequencerTimeSliderController.cpp
--- a/Engine/Source/Editor/Sequencer/Private/SequencerTimeSliderController.cpp
+++ b/Engine/Source/Editor/Sequencer/Private/SequencerTimeSliderController.cpp
@@ -64,9 +64,6 @@ FSequencerTimeSliderController::FSequencerTimeSliderController( const FTimeSlide
 	, bMouseDownInRegion(false)
 	, bPanning( false )
 	, HoverMarkIndex( INDEX_NONE )
-	, DragMarkIndex( INDEX_NONE )
-	, DragMarkReferenceFrameNumber( 0 )
-	, DragMarkCurrentFrameNumber( 0 )
 {
 	ScrubFillBrush              = FAppStyle::GetBrush( TEXT( "Sequencer.Timeline.ScrubFill" ) );
 	FrameBlockScrubHandleUpBrush   = FAppStyle::GetBrush( TEXT( "Sequencer.Timeline.FrameBlockScrubHandleUp" ) ); 
@@ -842,7 +839,6 @@ FReply FSequencerTimeSliderController::OnMouseButtonDown( SWidget& WidgetOwner,
 	MouseDownPosition[0] = MouseDownPosition[1] = MouseEvent.GetScreenSpacePosition();
 	MouseDownGeometry = MyGeometry;
 	bMouseDownInRegion = false;
-	DragMarkIndex = INDEX_NONE;
 
 	FVector2D CursorPos = MouseEvent.GetScreenSpacePosition();
 	FVector2D LocalPos = MouseDownGeometry.AbsoluteToLocal(CursorPos);
@@ -916,7 +912,6 @@ FReply FSequencerTimeSliderController::OnMouseButtonUp( SWidget& WidgetOwner, co
 		else if (MouseDragType == DRAG_MARK)
 		{
 			TimeSliderArgs.OnMarkEndDrag.ExecuteIfBound();
-			UpdateMarkSelection(DragMarkIndex, DragMarkCurrentFrameNumber);
 		}
 		else if (MouseDragType == DRAG_SETTING_RANGE && MouseDownPosition[0].IsSet())
 		{
@@ -997,7 +992,6 @@ FReply FSequencerTimeSliderController::OnMouseButtonUp( SWidget& WidgetOwner, co
 
 		MouseDragType = DRAG_NONE;
 		DistanceDragged = 0.f;
-		DragMarkIndex = INDEX_NONE;
 		bMouseDownInRegion = false;
 		MouseDownPosition[0].Reset();
 		MouseDownPosition[1].Reset();
@@ -1036,6 +1030,8 @@ FReply FSequencerTimeSliderController::OnTimeSliderMouseMove(SWidget& WidgetOwne
 
 FReply FSequencerTimeSliderController::OnMouseMoveImpl( SWidget& WidgetOwner, const FGeometry& MyGeometry, const FPointerEvent& MouseEvent, bool bFromTimeSlider )
 {
+	using namespace UE::Sequencer;
+
 	TSharedPtr<FSequencer> Sequencer = WeakSequencer.Pin();
 	if (!Sequencer.IsValid())
 	{
@@ -1049,6 +1045,7 @@ FReply FSequencerTimeSliderController::OnMouseMoveImpl( SWidget& WidgetOwner, co
 	const bool bLockedMarkedFrames = TimeSliderArgs.AreMarkedFramesLocked.Get();
 
 	HoverMarkIndex = INDEX_NONE;
+	int32 DragMarkIndex = INDEX_NONE;
 
 	if (bHandleRightMouseButton)
 	{
@@ -1143,11 +1140,21 @@ FReply FSequencerTimeSliderController::OnMouseMoveImpl( SWidget& WidgetOwner, co
 					MouseDragType = DRAG_PLAYBACK_START;
 					TimeSliderArgs.OnPlaybackRangeBeginDrag.ExecuteIfBound();
 				}
-				else if (!bLockedMarkedFrames && !bHitScrubber && HitTestMark(MyGeometry, RangeToScreen, MouseDownPixel, bFromTimeSlider, &DragMarkIndex, &DragMarkReferenceFrameNumber) && bHandleMiddleMouseButton == false)
+				else if (!bLockedMarkedFrames && !bHitScrubber && HitTestMark(MyGeometry, RangeToScreen, MouseDownPixel, bFromTimeSlider, &DragMarkIndex) && bHandleMiddleMouseButton == false)
 				{
 					MouseDragType = DRAG_MARK;
-					DragMarkCurrentFrameNumber = DragMarkReferenceFrameNumber;
 					HandleMarkSelection(DragMarkIndex);
+
+					DragMarkMap.Empty();
+					const FMarkedFrameSelection& SelectedMarkedFrames = Sequencer->GetViewModel()->GetSelection()->MarkedFrames;
+					const TArray<FMovieSceneMarkedFrame>& MarkedFrames = TimeSliderArgs.MarkedFrames.Get();
+					for (TSet<int32>::TConstIterator It = SelectedMarkedFrames.GetSelected(); It; ++It)
+					{
+						const int32 MarkIndex = *It;		
+						const FMovieSceneMarkedFrame& MarkedFrame = MarkedFrames[MarkIndex];
+						DragMarkMap.Add(MarkIndex, MarkedFrame.FrameNumber);
+					}
+
 					TimeSliderArgs.OnMarkBeginDrag.ExecuteIfBound();
 				}
 				else if (FSlateApplication::Get().GetModifierKeys().AreModifersDown(EModifierKey::Control) && bHandleMiddleMouseButton == false)
@@ -1221,8 +1228,7 @@ FReply FSequencerTimeSliderController::OnMouseMoveImpl( SWidget& WidgetOwner, co
 			}
 			else if (MouseDragType == DRAG_MARK)
 			{
-				DragMarkCurrentFrameNumber = DragMarkReferenceFrameNumber + DiffFrame;
-				SetMark(DragMarkIndex, DragMarkCurrentFrameNumber);
+				SetMark(DiffFrame);
 			}
 			else if (MouseDragType == DRAG_SCRUBBING_TIME)
 			{
@@ -2149,8 +2155,11 @@ void FSequencerTimeSliderController::HandleMarkSelection(int32 InMarkIndex)
 
 	if (!bToggleSelection && !bAddToSelection)
 	{
-		SequencerSelection.Empty();
-		SequencerSelection.MarkedFrames.Select(InMarkIndex);
+		if (!SequencerSelection.MarkedFrames.IsSelected(InMarkIndex))
+		{
+			SequencerSelection.Empty();
+			SequencerSelection.MarkedFrames.Select(InMarkIndex);
+		}
 	}
 	else if (bAddToSelection || !SequencerSelection.MarkedFrames.IsSelected(InMarkIndex))
 	{
@@ -2162,27 +2171,16 @@ void FSequencerTimeSliderController::HandleMarkSelection(int32 InMarkIndex)
 	}
 }
 
-void FSequencerTimeSliderController::UpdateMarkSelection(int32 InOldMarkIndex, FFrameNumber InMarkFrameNumber)
+void FSequencerTimeSliderController::SetMark(FFrameNumber DiffFrame)
 {
-	using namespace UE::Sequencer;
-
-	TSharedPtr<FSequencer> Sequencer = WeakSequencer.Pin();
-	FSequencerSelection& SequencerSelection = Sequencer->GetSelection();
-	SequencerSelection.MarkedFrames.Deselect(InOldMarkIndex);
-
-	UMovieScene* MovieScene = Sequencer->GetFocusedMovieSceneSequence()->GetMovieScene();
-	int32 NewMarkIndex = MovieScene->FindMarkedFrameByFrameNumber(InMarkFrameNumber);
-	if (ensure(NewMarkIndex != INDEX_NONE))
+	for (TMap<int32, FFrameNumber>::TConstIterator It = DragMarkMap.CreateConstIterator(); It; ++It)
 	{
-		SequencerSelection.MarkedFrames.Select(NewMarkIndex);
+		int32 MarkIndex = It.Key();
+		FFrameNumber FrameNumber = It.Value() + DiffFrame;
+		TimeSliderArgs.OnSetMarkedFrame.ExecuteIfBound(MarkIndex, FrameNumber);
 	}
 }
 
-void FSequencerTimeSliderController::SetMark(int32 InMarkIndex, FFrameNumber FrameNumber)
-{
-	TimeSliderArgs.OnSetMarkedFrame.ExecuteIfBound(InMarkIndex, FrameNumber);
-}
-
 void FSequencerTimeSliderController::AddMarkAtFrame(FFrameNumber FrameNumber)
 {
 	TimeSliderArgs.OnAddMarkedFrame.ExecuteIfBound(FrameNumber);
diff --git a/Engine/Source/Editor/Sequencer/Public/SequencerTimeSliderController.h b/Engine/Source/Editor/Sequencer/Public/SequencerTimeSliderController.h
--- a/Engine/Source/Editor/Sequencer/Public/SequencerTimeSliderController.h
+++ b/Engine/Source/Editor/Sequencer/Public/SequencerTimeSliderController.h
@@ -291,14 +291,13 @@ private:
 	void SetSelectionRangeStart(FFrameNumber NewStart);
 	void SetSelectionRangeEnd(FFrameNumber NewEnd);
 
-	void SetMark(int32 InMarkIndex, FFrameNumber NewFrame);
+	void SetMark(FFrameNumber DiffFrame);
 
 	TSharedRef<SWidget> OpenSetPlaybackRangeMenu(const FGeometry& MyGeometry, const FPointerEvent& MouseEvent);
 	FFrameTime ComputeScrubTimeFromMouse(const FGeometry& Geometry, const FPointerEvent& MouseEvent, FScrubRangeToScreen RangeToScreen) const;
 	FFrameTime ComputeFrameTimeFromMouse(const FGeometry& Geometry, FVector2D ScreenSpacePosition, FScrubRangeToScreen RangeToScreen, bool CheckSnapping = true) const;
 
 	void HandleMarkSelection(int32 InMarkIndex);
-	void UpdateMarkSelection(int32 InOldMarkIndex, FFrameNumber InMarkFrameNumber);
 	void AddMarkAtFrame(FFrameNumber FrameNumber);
 	void DeleteMarkAtIndex(int32 InMarkIndex);
 	void DeleteAllMarks();
@@ -381,14 +380,8 @@ private:
 	/** Index of mark being hovered */
 	int32 HoverMarkIndex;
 
-	/** Index of mark being edited */
-	int32 DragMarkIndex;
-
-	/** Marked frame time when the mouse is first pressed down */
-	FFrameNumber DragMarkReferenceFrameNumber;
-
-	/** Time of the mark being edited */
-	FFrameNumber DragMarkCurrentFrameNumber;
+	/** Map of the indices of the marks being edited and their initial frame numbers when pressed */
+	TMap<int32, FFrameNumber> DragMarkMap;
 
 	/** When > 0, we should not show context menus */
 	int32 ContextMenuSuppression;
