Sequencer: Fix vertical row dragging

#changelist validated
#virtualized

diff --git a/Engine/Source/Editor/Sequencer/Private/Tools/EditToolDragOperations.cpp b/Engine/Source/Editor/Sequencer/Private/Tools/EditToolDragOperations.cpp
--- a/Engine/Source/Editor/Sequencer/Private/Tools/EditToolDragOperations.cpp
+++ b/Engine/Source/Editor/Sequencer/Private/Tools/EditToolDragOperations.cpp
@@ -1229,7 +1229,7 @@ bool FMoveKeysAndSections::HandleSectionMovement(FFrameTime MouseTime, FVector2D
 				const float VirtualRowHeight = VirtualSectionHeight / NumRows;
 				const float MouseOffsetWithinRow = VirtualMousePos.Y - (VirtualGeometry.Top + (VirtualRowHeight * TargetRowIndex));
 
-				if (MouseOffsetWithinRow < VirtualRowHeight || MouseOffsetWithinRow > VirtualRowHeight)
+				if (MouseOffsetWithinRow <= VirtualRowHeight || MouseOffsetWithinRow >= VirtualRowHeight)
 				{
 					const int32 NewIndex = FMath::FloorToInt((VirtualMousePos.Y - VirtualGeometry.Top) / VirtualRowHeight);
 					TargetRowIndex = FMath::Clamp(NewIndex, 0, MaxRowIndex);
