Sequencer: Missing key button for track editors that don't have a key editor. They should show up for any keyable areas

diff --git a/Engine/Source/Editor/Sequencer/Private/MovieSceneTrackEditor.cpp b/Engine/Source/Editor/Sequencer/Private/MovieSceneTrackEditor.cpp
--- a/Engine/Source/Editor/Sequencer/Private/MovieSceneTrackEditor.cpp
+++ b/Engine/Source/Editor/Sequencer/Private/MovieSceneTrackEditor.cpp
@@ -32,7 +32,7 @@ namespace UE::Sequencer
 			{
 				if (TViewModelPtr<FChannelModel> Channel = WeakChannel.Pin())
 				{
-					if (Channel->GetKeyArea()->CanCreateKeyEditor())
+					if (Channel->GetKeyArea())
 					{
 						return true;
 					}
