Sequencer: Fix title position so that its centered only when there's enough space

diff --git a/Engine/Source/Editor/Sequencer/Private/SSequencerSection.cpp b/Engine/Source/Editor/Sequencer/Private/SSequencerSection.cpp
--- a/Engine/Source/Editor/Sequencer/Private/SSequencerSection.cpp
+++ b/Engine/Source/Editor/Sequencer/Private/SSequencerSection.cpp
@@ -1486,7 +1486,11 @@ int32 SSequencerSection::OnPaint( const FPaintArgs& Args, const FGeometry& Allot
 			FontInfo.Size = FMath::Max(FMath::FloorToInt(FontInfo.Size - 6.f), 7.f);
 		}
 
-		const float TitlePosition = (AllocatedFontHeight - FontInfo.Size) * 0.5f;
+		float TitlePosition = ContentPadding.Top;
+		if (FontInfo.Size + ContentPadding.Top + ContentPadding.Bottom > AllocatedFontHeight)
+		{
+			TitlePosition = (AllocatedFontHeight - FontInfo.Size) * 0.5f;
+		}
 
 		// Drop shadow
 		FSlateDrawElement::MakeText(
@@ -1554,7 +1558,6 @@ int32 SSequencerSection::OnPaint( const FPaintArgs& Args, const FGeometry& Allot
 	return LayerId;
 }
 
-
 void SSequencerSection::PaintEasingHandles( FSequencerSectionPainter& InPainter, FLinearColor SelectionColor, TSharedPtr<ITrackAreaHotspot> Hotspot ) const
 {
 	if (!SectionInterface->GetSectionObject()->GetBlendType().IsValid())
