Sequencer: Invalidate expired objects after replacing.

This fixes red/invalid component bindings in the UI after changing values.

diff --git a/Engine/Source/Editor/Sequencer/Private/Sequencer.cpp b/Engine/Source/Editor/Sequencer/Private/Sequencer.cpp
--- a/Engine/Source/Editor/Sequencer/Private/Sequencer.cpp
+++ b/Engine/Source/Editor/Sequencer/Private/Sequencer.cpp
@@ -440,9 +440,19 @@ void FSequencer::InitSequencer(const FSequencerInitParams& InitParams, const TSh
 			}
 
 			// Reset Bindings for replaced objects.
+			bool bAnythingReplaced = false;
 			for (TPair<UObject*, UObject*> ReplacedObject : ReplacementMap)
 			{
 				FGuid Guid = GetHandleToObject(ReplacedObject.Key, false);
+				if (Guid.IsValid())
+				{
+					bAnythingReplaced = true;
+				}
+			}
+
+			if (bAnythingReplaced)
+			{
+				State.InvalidateExpiredObjects();
 			}
 
 			// Replace pointers inside our pre-animate state storages.
