Sequencer: Set loop mode to loop selection range when initially setting the selection range (ie. if it's already been set and you explicitly change the loop mode to loop once, it won't change it back to loop selection mode)

diff --git a/Engine/Source/Editor/Sequencer/Private/Sequencer.cpp b/Engine/Source/Editor/Sequencer/Private/Sequencer.cpp
--- a/Engine/Source/Editor/Sequencer/Private/Sequencer.cpp
+++ b/Engine/Source/Editor/Sequencer/Private/Sequencer.cpp
@@ -2648,14 +2648,22 @@ void FSequencer::SetSelectionRange(TRange<FFrameNumber> Range)
 		return;
 	}
 
+	const bool bInitiallyEmpty = GetSelectionRange().IsEmpty();
+
 	const FScopedTransaction Transaction(LOCTEXT("SetSelectionRange_Transaction", "Set Selection Range"));
 	FocusedMovieScene->Modify();
 	FocusedMovieScene->SetSelectionRange(Range);
+
+	if (bInitiallyEmpty && GetLoopMode() != ESequencerLoopMode::SLM_LoopSelectionRange)
+	{
+		Settings->SetLoopMode(ESequencerLoopMode::SLM_LoopSelectionRange);
+	}
 }
 
 
 void FSequencer::SetSelectionRangeEnd(FFrameTime EndFrame)
 {
+	const bool bInitiallyEmpty = GetSelectionRange().IsEmpty();
 	const FFrameNumber LocalTime = EndFrame.FrameNumber;
 
 	if (GetSelectionRange().GetLowerBoundValue() >= LocalTime)
@@ -2666,11 +2674,17 @@ void FSequencer::SetSelectionRangeEnd(FFrameTime EndFrame)
 	{
 		SetSelectionRange(TRange<FFrameNumber>(GetSelectionRange().GetLowerBound(), LocalTime));
 	}
+
+	if (bInitiallyEmpty && GetLoopMode() != ESequencerLoopMode::SLM_LoopSelectionRange)
+	{
+		Settings->SetLoopMode(ESequencerLoopMode::SLM_LoopSelectionRange);
+	}
 }
 
 
 void FSequencer::SetSelectionRangeStart(FFrameTime StartFrame)
 {
+	const bool bInitiallyEmpty = GetSelectionRange().IsEmpty();
 	const FFrameNumber LocalTime = StartFrame.FrameNumber;
 
 	if (GetSelectionRange().GetUpperBoundValue() <= LocalTime)
@@ -2681,6 +2695,11 @@ void FSequencer::SetSelectionRangeStart(FFrameTime StartFrame)
 	{
 		SetSelectionRange(TRange<FFrameNumber>(LocalTime, GetSelectionRange().GetUpperBound()));
 	}
+
+	if (bInitiallyEmpty && GetLoopMode() != ESequencerLoopMode::SLM_LoopSelectionRange)
+	{
+		Settings->SetLoopMode(ESequencerLoopMode::SLM_LoopSelectionRange);
+	}
 }
 
 
