Sequencer: Always keep cursor in playback range when the option is enabled.

Previously, it was intentional that you could still set the time to outside the range by clicking (not scrubbing). We've decided to change the behavior.

diff --git a/Engine/Source/Editor/Sequencer/Private/SequencerTimeSliderController.cpp b/Engine/Source/Editor/Sequencer/Private/SequencerTimeSliderController.cpp
--- a/Engine/Source/Editor/Sequencer/Private/SequencerTimeSliderController.cpp
+++ b/Engine/Source/Editor/Sequencer/Private/SequencerTimeSliderController.cpp
@@ -107,22 +107,6 @@ FFrameTime FSequencerTimeSliderController::ComputeScrubTimeFromMouse(const FGeom
 		return ScrubTime;
 	}
 	
-	// Clamp first, snap to frame last
-	if (Sequencer->GetSequencerSettings()->ShouldKeepCursorInPlayRangeWhileScrubbing())
-	{
-		TOptional<TRange<FFrameNumber>> RangeValue;
-		RangeValue = TimeSliderArgs.SubSequenceRange.Get(RangeValue);
-
-		if (RangeValue.IsSet())
-		{
-			ScrubTime = UE::MovieScene::ClampToDiscreteRange(ScrubTime, RangeValue.GetValue());
-		}
-		else
-		{
-			ScrubTime = UE::MovieScene::ClampToDiscreteRange(ScrubTime, TimeSliderArgs.PlaybackRange.Get());
-		}
-	}
-
 	if ( Sequencer->GetSequencerSettings()->GetIsSnapEnabled() || MouseEvent.IsShiftDown() )
 	{
 		if (Sequencer->GetSequencerSettings()->GetSnapPlayTimeToInterval())
@@ -1265,23 +1249,41 @@ void FSequencerTimeSliderController::CommitScrubPosition( FFrameTime NewValue, b
 	bIsEvaluating = bEvaluate;
 	// The user can scrub past the viewing range of the time slider controller, so we clamp it to the view range.
 	TSharedPtr<FSequencer> Sequencer = WeakSequencer.Pin();
-	if(Sequencer.IsValid() && bIsScrubbing)
+	if(Sequencer.IsValid())
 	{
-		FAnimatedRange ViewRange = GetViewRange();
-		
-		FFrameRate DisplayRate = Sequencer->GetFocusedDisplayRate();
-		FFrameRate TickResolution = Sequencer->GetFocusedTickResolution();
+		if (Sequencer->GetSequencerSettings()->ShouldKeepCursorInPlayRangeWhileScrubbing())
+		{
+			TOptional<TRange<FFrameNumber>> RangeValue;
+			RangeValue = TimeSliderArgs.SubSequenceRange.Get(RangeValue);
 
-		FFrameTime LowerBound = (ViewRange.GetLowerBoundValue() * TickResolution).CeilToFrame();
-		FFrameTime UpperBound = (ViewRange.GetUpperBoundValue() * TickResolution).FloorToFrame();
+			if (RangeValue.IsSet())
+			{
+				NewValue = UE::MovieScene::ClampToDiscreteRange(NewValue, RangeValue.GetValue());
+			}
+			else
+			{
+				NewValue = UE::MovieScene::ClampToDiscreteRange(NewValue, TimeSliderArgs.PlaybackRange.Get());
+			}
+		}
 
-		if (Sequencer->GetSequencerSettings()->GetIsSnapEnabled() && Sequencer->GetSequencerSettings()->GetSnapPlayTimeToInterval())
+		if (bIsScrubbing)
 		{
-			LowerBound = FFrameRate::Snap(LowerBound, TickResolution, DisplayRate);
-			UpperBound = FFrameRate::Snap(UpperBound, TickResolution, DisplayRate);
-		}
+			FAnimatedRange ViewRange = GetViewRange();
+		
+			FFrameRate DisplayRate = Sequencer->GetFocusedDisplayRate();
+			FFrameRate TickResolution = Sequencer->GetFocusedTickResolution();
+
+			FFrameTime LowerBound = (ViewRange.GetLowerBoundValue() * TickResolution).CeilToFrame();
+			FFrameTime UpperBound = (ViewRange.GetUpperBoundValue() * TickResolution).FloorToFrame();
+
+			if (Sequencer->GetSequencerSettings()->GetIsSnapEnabled() && Sequencer->GetSequencerSettings()->GetSnapPlayTimeToInterval())
+			{
+				LowerBound = FFrameRate::Snap(LowerBound, TickResolution, DisplayRate);
+				UpperBound = FFrameRate::Snap(UpperBound, TickResolution, DisplayRate);
+			}
 
-		NewValue = FMath::Clamp(NewValue, LowerBound, UpperBound);
+			NewValue = FMath::Clamp(NewValue, LowerBound, UpperBound);
+		}
 	}
 
 	// Manage the scrub position ourselves if its not bound to a delegate
