Sequencer: Hidden columns are no longer toggleable on the Sequencer Columns toolbar

diff --git a/Engine/Source/Editor/Sequencer/Private/SSequencer.cpp b/Engine/Source/Editor/Sequencer/Private/SSequencer.cpp
--- a/Engine/Source/Editor/Sequencer/Private/SSequencer.cpp
+++ b/Engine/Source/Editor/Sequencer/Private/SSequencer.cpp
@@ -2417,6 +2417,11 @@ void SSequencer::FillColumnVisibilityMenu(FMenuBuilder& InMenuBuilder)
 
 	for (FSequencerOutlinerColumnVisibility& ColumnVisibility : OutlinerColumnVisibilities)
 	{
+		if (EnumHasAnyFlags(ColumnVisibility.Column->GetLayout().Flags, EOutlinerColumnFlags::Hidden))
+		{
+			continue;
+		}
+
 		auto ToggleVisibility = [this, &ColumnVisibility]
 		{
 			ColumnVisibility.bIsColumnVisible = !ColumnVisibility.bIsColumnVisible;
