Niagara: Fixed up Sequencer outliner entries for Niagara UI elements

 - Mute/Solo/Pin/Lock and key nav columns are hidden by default
 - Custom niagara buttons now appear correctly in the Edit column
 - Niagara section heights now respect view density
 - Key area switcher now reports itself as collapsed if there is now widget currently visible. This allows channel labels to correctly overflow by default where no key editor is visible.
 - Tweak stretch ratios for label and edit widgets to give more space for the edit column (25% ->40%). Really this needs to be user-customizable, but that is not a feature yet.

diff --git a/Engine/Config/BaseEditorPerProjectUserSettings.ini b/Engine/Config/BaseEditorPerProjectUserSettings.ini
--- a/Engine/Config/BaseEditorPerProjectUserSettings.ini
+++ b/Engine/Config/BaseEditorPerProjectUserSettings.ini
@@ -889,6 +889,16 @@ bShowRangeSlider=true
 LoopMode=SLM_Loop
 bCleanPlaybackMode=false
 bShowLayerBars=false
+ColumnVisibilitySettings=(ColumnName="Pin",bIsVisible=False)
++ColumnVisibilitySettings=(ColumnName="Lock",bIsVisible=False)
++ColumnVisibilitySettings=(ColumnName="Solo",bIsVisible=False)
++ColumnVisibilitySettings=(ColumnName="Mute",bIsVisible=False)
++ColumnVisibilitySettings=(ColumnName="Label",bIsVisible=True)
++ColumnVisibilitySettings=(ColumnName="Edit",bIsVisible=True)
++ColumnVisibilitySettings=(ColumnName="Add",bIsVisible=False)
++ColumnVisibilitySettings=(ColumnName="Nav",bIsVisible=False)
++ColumnVisibilitySettings=(ColumnName="KeyFrame",bIsVisible=False)
++ColumnVisibilitySettings=(ColumnName="ColorPicker",bIsVisible=True)
 
 [DaySequenceEditor SequencerSettings]
 bShowRangeSlider=true
diff --git a/Engine/Plugins/FX/Niagara/Source/NiagaraEditor/Private/Sequencer/NiagaraSequence/NiagaraEmitterTrackEditor.cpp b/Engine/Plugins/FX/Niagara/Source/NiagaraEditor/Private/Sequencer/NiagaraSequence/NiagaraEmitterTrackEditor.cpp
--- a/Engine/Plugins/FX/Niagara/Source/NiagaraEditor/Private/Sequencer/NiagaraSequence/NiagaraEmitterTrackEditor.cpp
+++ b/Engine/Plugins/FX/Niagara/Source/NiagaraEditor/Private/Sequencer/NiagaraSequence/NiagaraEmitterTrackEditor.cpp
@@ -22,6 +22,12 @@
 #include "Modules/ModuleManager.h"
 #include "ViewModels/Stack/NiagaraStackRendererItem.h"
 
+#include "MVVM/ViewModelPtr.h"
+#include "MVVM/Extensions/ITrackExtension.h"
+#include "MVVM/ViewModels/OutlinerColumns/OutlinerColumnTypes.h"
+#include "MVVM/Views/SOutlinerItemViewBase.h"
+#include "MVVM/Views/SOutlinerTrackColorPicker.h"
+
 #define LOCTEXT_NAMESPACE "NiagaraEmitterTrackEditor"
 
 class SEmitterTrackWidget : public SCompoundWidget
@@ -312,9 +318,34 @@ void FNiagaraEmitterTrackEditor::BuildTrackContextMenu( FMenuBuilder& MenuBuilde
 	}
 }
 
-TSharedPtr<SWidget> FNiagaraEmitterTrackEditor::BuildOutlinerEditWidget(const FGuid& ObjectBinding, UMovieSceneTrack* Track, const FBuildEditWidgetParams& Params)
+TSharedPtr<SWidget> FNiagaraEmitterTrackEditor::BuildOutlinerColumnWidget(const FBuildColumnWidgetParams& Params, const FName& ColumnName)
 {
-	return SNew(SEmitterTrackWidget, *CastChecked<UMovieSceneNiagaraEmitterTrack>(Track));
+	using namespace UE::Sequencer;
+
+	TViewModelPtr<IOutlinerExtension> OutlinerItem = Params.ViewModel.ImplicitCast();
+	if (!OutlinerItem)
+	{
+		return nullptr;
+	}
+
+	if (ColumnName == FCommonOutlinerNames::Edit)
+	{
+		UMovieSceneNiagaraEmitterTrack* Track = CastChecked<UMovieSceneNiagaraEmitterTrack>(Params.TrackModel->GetTrack());
+		return SNew(SEmitterTrackWidget, *Track);
+	}
+
+	if (ColumnName == FCommonOutlinerNames::Label)
+	{
+		return SNew(SOutlinerItemViewBase, OutlinerItem, Params.Editor, Params.TreeViewRow);
+	}
+
+	if (ColumnName == FCommonOutlinerNames::ColorPicker)
+	{
+		return SNew(SOutlinerTrackColorPicker, OutlinerItem, Params.Editor);
+	}
+
+
+	return SNullWidget::NullWidget;
 }
 
 #undef LOCTEXT_NAMESPACE
diff --git a/Engine/Plugins/FX/Niagara/Source/NiagaraEditor/Private/Sequencer/NiagaraSequence/NiagaraEmitterTrackEditor.h b/Engine/Plugins/FX/Niagara/Source/NiagaraEditor/Private/Sequencer/NiagaraSequence/NiagaraEmitterTrackEditor.h
--- a/Engine/Plugins/FX/Niagara/Source/NiagaraEditor/Private/Sequencer/NiagaraSequence/NiagaraEmitterTrackEditor.h
+++ b/Engine/Plugins/FX/Niagara/Source/NiagaraEditor/Private/Sequencer/NiagaraSequence/NiagaraEmitterTrackEditor.h
@@ -23,7 +23,7 @@ public:
 	virtual TSharedRef<ISequencerSection> MakeSectionInterface(UMovieSceneSection& SectionObject, UMovieSceneTrack& Track, FGuid ObjectBinding) override;
 	virtual bool HandleAssetAdded(UObject* Asset, const FGuid& TargetObjectGuid) override;
 	virtual void BuildTrackContextMenu( FMenuBuilder& MenuBuilder, UMovieSceneTrack* Track ) override; 
-	virtual TSharedPtr<SWidget> BuildOutlinerEditWidget(const FGuid& ObjectBinding, UMovieSceneTrack* Track, const FBuildEditWidgetParams& Params) override;
+	virtual TSharedPtr<SWidget> BuildOutlinerColumnWidget(const FBuildColumnWidgetParams& Params, const FName& ColumnName) override;
 
 private:
 	void ToggleEmitterIsolation(TSharedRef<FNiagaraEmitterHandleViewModel> EmitterToIsolate);
diff --git a/Engine/Plugins/FX/Niagara/Source/NiagaraEditor/Private/Sequencer/NiagaraSequence/Sections/NiagaraEmitterSection.cpp b/Engine/Plugins/FX/Niagara/Source/NiagaraEditor/Private/Sequencer/NiagaraSequence/Sections/NiagaraEmitterSection.cpp
--- a/Engine/Plugins/FX/Niagara/Source/NiagaraEditor/Private/Sequencer/NiagaraSequence/Sections/NiagaraEmitterSection.cpp
+++ b/Engine/Plugins/FX/Niagara/Source/NiagaraEditor/Private/Sequencer/NiagaraSequence/Sections/NiagaraEmitterSection.cpp
@@ -5,6 +5,7 @@
 #include "ViewModels/NiagaraEmitterHandleViewModel.h"
 #include "ViewModels/NiagaraEmitterViewModel.h"
 
+#include "MVVM/ViewModels/ViewDensity.h"
 #include "SequencerSectionPainter.h"
 #include "ISectionLayoutBuilder.h"
 
@@ -20,6 +21,11 @@ UMovieSceneSection* FNiagaraEmitterSection::GetSectionObject(void)
 	return EmitterSection;
 }
 
+float FNiagaraEmitterSection::GetSectionHeight(const UE::Sequencer::FViewDensityInfo& ViewDensity) const
+{
+	return ViewDensity.UniformHeight.Get(20.0f);
+}
+
 int32 FNiagaraEmitterSection::OnPaintSection(FSequencerSectionPainter& InPainter) const
 {
 	return InPainter.PaintSectionBackground();
diff --git a/Engine/Plugins/FX/Niagara/Source/NiagaraEditor/Private/Sequencer/NiagaraSequence/Sections/NiagaraEmitterSection.h b/Engine/Plugins/FX/Niagara/Source/NiagaraEditor/Private/Sequencer/NiagaraSequence/Sections/NiagaraEmitterSection.h
--- a/Engine/Plugins/FX/Niagara/Source/NiagaraEditor/Private/Sequencer/NiagaraSequence/Sections/NiagaraEmitterSection.h
+++ b/Engine/Plugins/FX/Niagara/Source/NiagaraEditor/Private/Sequencer/NiagaraSequence/Sections/NiagaraEmitterSection.h
@@ -18,7 +18,7 @@ public:
 	virtual UMovieSceneSection* GetSectionObject(void) override;
 	virtual int32 OnPaintSection(FSequencerSectionPainter& InPainter) const override;
 	virtual FText GetSectionTitle(void) const override;
-	virtual float GetSectionHeight() const override { return 20.0f; }
+	virtual float GetSectionHeight(const UE::Sequencer::FViewDensityInfo& ViewDensity) const override;
 
 private:
 	UMovieSceneNiagaraEmitterSection* EmitterSection;
diff --git a/Engine/Source/Editor/Sequencer/Private/SKeyAreaEditorSwitcher.cpp b/Engine/Source/Editor/Sequencer/Private/SKeyAreaEditorSwitcher.cpp
--- a/Engine/Source/Editor/Sequencer/Private/SKeyAreaEditorSwitcher.cpp
+++ b/Engine/Source/Editor/Sequencer/Private/SKeyAreaEditorSwitcher.cpp
@@ -20,6 +20,8 @@ void SKeyAreaEditorSwitcher::Construct(const FArguments& InArgs, TSharedPtr<FCha
 	WeakModel = InModel;
 	WeakEditorModel = InWeakEditorModel;
 	CachedChannelsSerialNumber = 0;
+
+	SetVisibility(MakeAttributeSP(this, &SKeyAreaEditorSwitcher::ComputeVisibility));
 }
 
 int32 SKeyAreaEditorSwitcher::GetWidgetIndex() const
@@ -27,6 +29,13 @@ int32 SKeyAreaEditorSwitcher::GetWidgetIndex() const
 	return VisibleIndex;
 }
 
+EVisibility SKeyAreaEditorSwitcher::ComputeVisibility() const
+{
+	return VisibleIndex != INDEX_NONE
+		? EVisibility::Visible
+		: EVisibility::Collapsed;
+}
+
 void SKeyAreaEditorSwitcher::Rebuild()
 {
 	TSharedPtr<FChannelGroupModel>        Model = WeakModel.Pin();
@@ -45,39 +54,21 @@ void SKeyAreaEditorSwitcher::Rebuild()
 	}
 
 	SetEnabled(MakeAttributeSP(Editor.ToSharedRef(), &FEditorViewModel::IsEditable));
-
-	// Index 0 is always the spacer node
-	VisibleIndex = 0;
+	VisibleIndex = INDEX_NONE;
 
 	TSharedRef<SWidgetSwitcher> Switcher = SNew(SWidgetSwitcher)
-		.WidgetIndex(this, &SKeyAreaEditorSwitcher::GetWidgetIndex)
-
-		+ SWidgetSwitcher::Slot()
-		[
-			SNullWidget::NullWidget
-		];
+		.WidgetIndex(this, &SKeyAreaEditorSwitcher::GetWidgetIndex);
 
 	TSharedPtr<IObjectBindingExtension> ParentObjectBinding = Model->FindAncestorOfType<IObjectBindingExtension>();
 	FGuid ObjectBindingID = ParentObjectBinding.IsValid() ? ParentObjectBinding->GetObjectGuid() : FGuid();
 
 	for (TSharedRef<IKeyArea> KeyArea : CachedKeyAreas)
 	{
-		if (!KeyArea->CanCreateKeyEditor())
-		{
-			// Always generate a slot so that indices line up correctly
-			Switcher->AddSlot()
-			[
-				SNullWidget::NullWidget
-			];
-		}
-		else
-		{
-			Switcher->AddSlot()
-			.HAlign(HAlign_Left)
-			[
-				KeyArea->CreateKeyEditor(Editor->GetSequencer(), ObjectBindingID)
-			];
-		}
+		Switcher->AddSlot()
+		.HAlign(HAlign_Left)
+		[
+			KeyArea->CreateKeyEditor(Editor->GetSequencer(), ObjectBindingID)
+		];
 	}
 
 	ChildSlot
@@ -109,6 +100,13 @@ void SKeyAreaEditorSwitcher::Tick( const FGeometry& AllottedGeometry, const doub
 		{
 			CachedChannelsSerialNumber = NewChannelsSerialNumber;
 			CachedKeyAreas = Model->GetAllKeyAreas();
+			for (int32 Index = CachedKeyAreas.Num()-1; Index >= 0; --Index)
+			{
+				if (!CachedKeyAreas[Index]->CanCreateKeyEditor())
+				{
+					CachedKeyAreas.RemoveAtSwap(Index);
+				}
+			}
 			Rebuild();
 		}
 
@@ -119,16 +117,7 @@ void SKeyAreaEditorSwitcher::Tick( const FGeometry& AllottedGeometry, const doub
 			AllSections.Add(KeyArea->GetOwningSection());
 		}
 
-		const int32 ActiveKeyArea = SequencerHelpers::GetSectionFromTime(AllSections, Editor->GetSequencer()->GetLocalTime().Time.FrameNumber);
-		if (ActiveKeyArea != INDEX_NONE)
-		{
-			// Index 0 is the spacer node, so add 1 to the key area index to get the widget index
-			VisibleIndex = 1 + ActiveKeyArea;
-		}
-		else
-		{
-			VisibleIndex = 0;
-		}
+		VisibleIndex = SequencerHelpers::GetSectionFromTime(AllSections, Editor->GetSequencer()->GetLocalTime().Time.FrameNumber);
 	}
 }
 
diff --git a/Engine/Source/Editor/Sequencer/Private/SKeyAreaEditorSwitcher.h b/Engine/Source/Editor/Sequencer/Private/SKeyAreaEditorSwitcher.h
--- a/Engine/Source/Editor/Sequencer/Private/SKeyAreaEditorSwitcher.h
+++ b/Engine/Source/Editor/Sequencer/Private/SKeyAreaEditorSwitcher.h
@@ -36,6 +36,8 @@ private:
 
 	int32 GetWidgetIndex() const;
 
+	EVisibility ComputeVisibility() const;
+
 private:
 
 	/** Our overlay widget */
diff --git a/Engine/Source/Editor/Sequencer/Public/MVVM/Views/SOutlinerTrackColorPicker.h b/Engine/Source/Editor/Sequencer/Public/MVVM/Views/SOutlinerTrackColorPicker.h
--- a/Engine/Source/Editor/Sequencer/Public/MVVM/Views/SOutlinerTrackColorPicker.h
+++ b/Engine/Source/Editor/Sequencer/Public/MVVM/Views/SOutlinerTrackColorPicker.h
@@ -21,7 +21,7 @@ public:
 	SLATE_BEGIN_ARGS(SOutlinerTrackColorPicker){}
 	SLATE_END_ARGS()
 
-	void Construct(const FArguments& InArgs, TWeakViewModelPtr<IOutlinerExtension> InWeakOutlinerExtension, const TSharedPtr<FEditorViewModel>& EditorViewModel);
+	SEQUENCER_API void Construct(const FArguments& InArgs, TWeakViewModelPtr<IOutlinerExtension> InWeakOutlinerExtension, const TSharedPtr<FEditorViewModel>& EditorViewModel);
 
 private:
 
diff --git a/Engine/Source/Editor/SequencerCore/Private/MVVM/ViewModels/OutlinerColumns/EditOutlinerColumn.cpp b/Engine/Source/Editor/SequencerCore/Private/MVVM/ViewModels/OutlinerColumns/EditOutlinerColumn.cpp
--- a/Engine/Source/Editor/SequencerCore/Private/MVVM/ViewModels/OutlinerColumns/EditOutlinerColumn.cpp
+++ b/Engine/Source/Editor/SequencerCore/Private/MVVM/ViewModels/OutlinerColumns/EditOutlinerColumn.cpp
@@ -12,7 +12,7 @@ FEditOutlinerColumn::FEditOutlinerColumn()
 	Name     = FCommonOutlinerNames::Edit;
 	Label    = NSLOCTEXT("FEditOutlinerColumn", "EditColumnEdit", "Edit");
 	Position = FOutlinerColumnPosition{ 10, EOutlinerColumnGroup::Center };
-	Layout   = FOutlinerColumnLayout{ 0.25f, FMargin(4.f, 1.f), HAlign_Fill, VAlign_Center, EOutlinerColumnSizeMode::Stretch, Flags };
+	Layout   = FOutlinerColumnLayout{ 0.4f, FMargin(4.f, 1.f), HAlign_Fill, VAlign_Center, EOutlinerColumnSizeMode::Stretch, Flags };
 }
 
 } // namespace UE::Sequencer
diff --git a/Engine/Source/Editor/SequencerCore/Private/MVVM/ViewModels/OutlinerColumns/LabelOutlinerColumn.cpp b/Engine/Source/Editor/SequencerCore/Private/MVVM/ViewModels/OutlinerColumns/LabelOutlinerColumn.cpp
--- a/Engine/Source/Editor/SequencerCore/Private/MVVM/ViewModels/OutlinerColumns/LabelOutlinerColumn.cpp
+++ b/Engine/Source/Editor/SequencerCore/Private/MVVM/ViewModels/OutlinerColumns/LabelOutlinerColumn.cpp
@@ -13,7 +13,7 @@ FLabelOutlinerColumn::FLabelOutlinerColumn()
 	Name     = FCommonOutlinerNames::Label;
 	Label    = NSLOCTEXT("FLabelOutlinerColumn", "LabelColumnLabel", "Label");
 	Position = FOutlinerColumnPosition{ 0, EOutlinerColumnGroup::Center };
-	Layout   = FOutlinerColumnLayout{ 0.75f, FMargin(4.f, 0.f), HAlign_Fill, VAlign_Center, EOutlinerColumnSizeMode::Stretch, Flags };
+	Layout   = FOutlinerColumnLayout{ 0.6f, FMargin(4.f, 0.f), HAlign_Fill, VAlign_Center, EOutlinerColumnSizeMode::Stretch, Flags };
 }
 
 TSharedPtr<SWidget> FLabelOutlinerColumn::CreateColumnWidget(const FCreateOutlinerColumnParams& InParams, const TSharedRef<ISequencerTreeViewRow>& TreeViewRow)
