Sequencer: Fixed pinned area using the wrong tree view for its scroll border

Also removed some superfluous slate elements

diff --git a/Engine/Source/Editor/Sequencer/Private/SSequencer.cpp b/Engine/Source/Editor/Sequencer/Private/SSequencer.cpp
--- a/Engine/Source/Editor/Sequencer/Private/SSequencer.cpp
+++ b/Engine/Source/Editor/Sequencer/Private/SSequencer.cpp
@@ -585,112 +585,92 @@ void SSequencer::Construct(const FArguments& InArgs, TSharedRef<FSequencer> InSe
 						+SVerticalBox::Slot()
 						.AutoHeight()
 						.MaxHeight(TAttribute<float>::Create(TAttribute<float>::FGetter::CreateSP(this, &SSequencer::GetPinnedAreaMaxHeight)))
+						.Padding(FMargin(0.f, 0.f, 0.f, CommonPadding))
 						[
-							SNew(SBorder)
+							SNew(SOverlay)
 							.Visibility(this, &SSequencer::GetPinnedAreaVisibility)
-							.Padding(FMargin(0.0f, 0.0f, 0.0f, CommonPadding))
-							[
-
-								SNew(SHorizontalBox)
 
-								+ SHorizontalBox::Slot()
+							+ SOverlay::Slot()
+							[
+								SNew(SScrollBorder, PinnedTreeView.ToSharedRef())
 								[
-									SNew(SOverlay)
+									SNew(SHorizontalBox)
 
-									+ SOverlay::Slot()
+									// outliner tree
+									+ SHorizontalBox::Slot()
+									.FillWidth( FillCoefficient_0 )
 									[
-										SNew(SScrollBorder, TreeView.ToSharedRef())
-										[
-											SNew(SHorizontalBox)
-
-											// outliner tree
-											+ SHorizontalBox::Slot()
-											.FillWidth( FillCoefficient_0 )
-											[
-												SNew(SBox)
-												[
-													PinnedTreeView.ToSharedRef()
-												]
-											]
-
-											// track area
-											+ SHorizontalBox::Slot()
-											.FillWidth( FillCoefficient_1 )
-											[
-												SNew(SBox)
-												.Padding(ResizeBarPadding)
-												.Clipping(EWidgetClipping::ClipToBounds)
-												[
-													PinnedTrackArea.ToSharedRef()
-												]
-											]
-										]
+										PinnedTreeView.ToSharedRef()
 									]
 
-									+ SOverlay::Slot()
-									.HAlign(HAlign_Right)
+									// track area
+									+ SHorizontalBox::Slot()
+									.FillWidth( FillCoefficient_1 )
 									[
-										PinnedAreaScrollBar
+										SNew(SBox)
+										.Padding(ResizeBarPadding)
+										.Clipping(EWidgetClipping::ClipToBounds)
+										[
+											PinnedTrackArea.ToSharedRef()
+										]
 									]
 								]
 							]
+
+							+ SOverlay::Slot()
+							.HAlign(HAlign_Right)
+							[
+								PinnedAreaScrollBar
+							]
 						]
 
 						+SVerticalBox::Slot()
 						[
-							SNew(SHorizontalBox)
+							SNew(SOverlay)
 
-							+ SHorizontalBox::Slot()
+							+ SOverlay::Slot()
 							[
-								SNew(SOverlay)
-
-								+ SOverlay::Slot()
+								SNew(SScrollBorder, TreeView.ToSharedRef())
 								[
-									SNew(SScrollBorder, TreeView.ToSharedRef())
-									[
-										SNew(SHorizontalBox)
+									SNew(SHorizontalBox)
 
-										// outliner tree
-										+ SHorizontalBox::Slot()
-										.FillWidth( FillCoefficient_0 )
-										[
-											SNew(SBox)
-											[
-												TreeView.ToSharedRef()
-											]
-										]
+									// outliner tree
+									+ SHorizontalBox::Slot()
+									.FillWidth( FillCoefficient_0 )
+									[
+										TreeView.ToSharedRef()
+									]
 
-										// track area
-										+ SHorizontalBox::Slot()
-										.FillWidth( FillCoefficient_1 )
+									// track area
+									+ SHorizontalBox::Slot()
+									.FillWidth( FillCoefficient_1 )
+									[
+										SNew(SBox)
+										.Padding(ResizeBarPadding)
+										.Clipping(EWidgetClipping::ClipToBounds)
 										[
-											SNew(SBox)
-											.Padding(ResizeBarPadding)
-											.Clipping(EWidgetClipping::ClipToBounds)
-											[
-												TrackArea.ToSharedRef()
-											]
+											TrackArea.ToSharedRef()
 										]
 									]
 								]
+							]
 
-								+ SOverlay::Slot()
-								.HAlign(HAlign_Right)
-								[
-									ScrollBar
-								]
+							+ SOverlay::Slot()
+							.HAlign(HAlign_Right)
+							[
+								ScrollBar
+							]
 
-								+ SOverlay::Slot()
-								.VAlign(VAlign_Bottom)
-								[
-									SNew(SHorizontalBox)
+							+ SOverlay::Slot()
+							.VAlign(VAlign_Bottom)
+							[
+								SNew(SHorizontalBox)
 
-									+ SHorizontalBox::Slot()
-									.AutoWidth()
-									[
-										SAssignNew(SequencerTreeFilterStatusBar, SSequencerTreeFilterStatusBar, InSequencer)
-										.Visibility(EVisibility::Hidden) // Initially hidden, visible on hover of the info button
-									]
+								+ SHorizontalBox::Slot()
+								.AutoWidth()
+								[
+									SAssignNew(SequencerTreeFilterStatusBar, SSequencerTreeFilterStatusBar, InSequencer)
+									.Visibility(EVisibility::Hidden) // Initially hidden, visible on hover of the info button
 								]
 							]
 						]
