Sequencer: Fixed OnActorAddedToSequencer being triggered multiple times for FSequencerUtilities methods

Calls to CreateBinding now always trigger this event if necessary

diff --git a/Engine/Source/Editor/Sequencer/Private/SequencerUtilities.cpp b/Engine/Source/Editor/Sequencer/Private/SequencerUtilities.cpp
--- a/Engine/Source/Editor/Sequencer/Private/SequencerUtilities.cpp
+++ b/Engine/Source/Editor/Sequencer/Private/SequencerUtilities.cpp
@@ -797,8 +797,6 @@ FGuid FSequencerUtilities::CreateCamera(TSharedRef<ISequencer> Sequencer, const
 		return CameraGuid;
 	}
 
-	Sequencer->OnActorAddedToSequencer().Broadcast(OutActor, CameraGuid);
-
 	NewCameraAdded(Sequencer, OutActor, CameraGuid);
 
 	return CameraGuid;
@@ -844,8 +842,6 @@ FGuid FSequencerUtilities::CreateCameraWithRig(TSharedRef<ISequencer> Sequencer,
 		OutActor->SetActorRotation(FRotator(0.f, -90.f, 0.f));
 	}
 
-	Sequencer->OnActorAddedToSequencer().Broadcast(OutActor, CameraGuid);
-
 	TRange<FFrameNumber> PlaybackRange = MovieScene->GetPlaybackRange();
 
 	if (bSpawnable)
@@ -965,8 +961,6 @@ TArray<FGuid> FSequencerUtilities::AddActors(TSharedRef<ISequencer> Sequencer, c
 				{
 					NewCameraAdded(Sequencer, CameraActor, PossessableGuid);
 				}
-
-				Sequencer->OnActorAddedToSequencer().Broadcast(Actor, PossessableGuid);
 			}
 		}
 	}
@@ -2559,7 +2553,7 @@ FGuid CreateGenericBinding(TSharedRef<ISequencer> Sequencer, UObject& InObject,
 
 		if (AActor* Actor = Cast<AActor>(CurrentObject))
 		{
-			Sequencer->OnActorAddedToSequencer().Broadcast(Actor, ParentID);
+			Sequencer->OnActorAddedToSequencer().Broadcast(Actor, NewID);
 		}
 
 		// If this is the last one
@@ -2642,7 +2636,14 @@ FGuid CreateImplementationDefinedBinding(TSharedRef<ISequencer> Sequencer, UObje
 	{
 		Sequencer->OnActorAddedToSequencer().Broadcast(ParentActorAdded, ParentGuid);
 	}
+
 	Sequencer->OnAddBinding(PossessableGuid, OwnerMovieScene);
+
+	if (Actor)
+	{
+		Sequencer->OnActorAddedToSequencer().Broadcast(Actor, PossessableGuid);
+	}
+
 	return PossessableGuid;
 }
 
