Sequencer: Renamed view density types

The new names are:
    Compact -> Variable
    Relaxed -> Compact
    Expanded -> Relaxed

This change breaks backwards compat with configs that were saved with the old names, but that is acceptible since this code hasn't been officially released yet.

diff --git a/Engine/Source/Editor/Sequencer/Private/SSequencer.cpp b/Engine/Source/Editor/Sequencer/Private/SSequencer.cpp
--- a/Engine/Source/Editor/Sequencer/Private/SSequencer.cpp
+++ b/Engine/Source/Editor/Sequencer/Private/SSequencer.cpp
@@ -2330,17 +2330,17 @@ void SSequencer::FillViewDensityMenu(FMenuBuilder& InMenuBuilder)
 
 		if (USequencerSettings* Settings = this->GetSequencerSettings())
 		{
-			if (InViewDensity == EViewDensity::Relaxed)
+			if (InViewDensity == EViewDensity::Compact)
 			{
-				Settings->SetViewDensity("Relaxed");
+				Settings->SetViewDensity("Compact");
 			}
-			else if (InViewDensity == EViewDensity::Expanded)
+			else if (InViewDensity == EViewDensity::Relaxed)
 			{
-				Settings->SetViewDensity("Expanded");
+				Settings->SetViewDensity("Relaxed");
 			}
 			else
 			{
-				Settings->SetViewDensity("Compact");
+				Settings->SetViewDensity("Variable");
 			}
 		}
 	};
@@ -2350,39 +2350,39 @@ void SSequencer::FillViewDensityMenu(FMenuBuilder& InMenuBuilder)
 	};
 
 	InMenuBuilder.AddMenuEntry(
-		LOCTEXT("CompactViewDensity_Label", "Compact"),
-		LOCTEXT("CompactViewDensity_Tooltip", "Change Sequencer to use a compact view mode"),
+		LOCTEXT("VariableViewDensity_Label", "Variable"),
+		LOCTEXT("VariableViewDensity_Tooltip", "Change Sequencer to use a variable height view mode withe inner items displaying more condensed than outer items"),
 		FSlateIcon(),
 		FUIAction(
-			FExecuteAction::CreateLambda(SetViewDensity, EViewDensity::Compact),
+			FExecuteAction::CreateLambda(SetViewDensity, EViewDensity::Variable),
 			FCanExecuteAction(),
-			FIsActionChecked::CreateLambda(IsCurrentViewDensity, EViewDensity::Compact)
+			FIsActionChecked::CreateLambda(IsCurrentViewDensity, EViewDensity::Variable)
 		),
 		NAME_None,
 		EUserInterfaceActionType::RadioButton
 	);
 
 	InMenuBuilder.AddMenuEntry(
-		LOCTEXT("RelaxedViewDensity_Label", "Relaxed"),
-		LOCTEXT("RelaxedViewDensity_Tooltip", "Change Sequencer to use a relaxed view mode with uniform track heights"),
+		LOCTEXT("CompactViewDensity_Label", "Compact"),
+		LOCTEXT("CompactViewDensity_Tooltip", "Change Sequencer to use a compact view mode with uniform track heights"),
 		FSlateIcon(),
 		FUIAction(
-			FExecuteAction::CreateLambda(SetViewDensity, EViewDensity::Relaxed),
+			FExecuteAction::CreateLambda(SetViewDensity, EViewDensity::Compact),
 			FCanExecuteAction(),
-			FIsActionChecked::CreateLambda(IsCurrentViewDensity, EViewDensity::Relaxed)
+			FIsActionChecked::CreateLambda(IsCurrentViewDensity, EViewDensity::Compact)
 		),
 		NAME_None,
 		EUserInterfaceActionType::RadioButton
 	);
 
 	InMenuBuilder.AddMenuEntry(
-		LOCTEXT("ExpandedViewDensity_Label", "Expanded"),
-		LOCTEXT("ExpandedViewDensity_Tooltip", "Change Sequencer to use an expanded view mode with larger uniform track heights"),
+		LOCTEXT("RelaxedViewDensity_Label", "Relaxed"),
+		LOCTEXT("RelaxedViewDensity_Tooltip", "Change Sequencer to use a relaxed view mode with larger uniform track heights"),
 		FSlateIcon(),
 		FUIAction(
-			FExecuteAction::CreateLambda(SetViewDensity, EViewDensity::Expanded),
+			FExecuteAction::CreateLambda(SetViewDensity, EViewDensity::Relaxed),
 			FCanExecuteAction(),
-			FIsActionChecked::CreateLambda(IsCurrentViewDensity, EViewDensity::Expanded)
+			FIsActionChecked::CreateLambda(IsCurrentViewDensity, EViewDensity::Relaxed)
 		),
 		NAME_None,
 		EUserInterfaceActionType::RadioButton
diff --git a/Engine/Source/Editor/Sequencer/Private/SequencerSettings.cpp b/Engine/Source/Editor/Sequencer/Private/SequencerSettings.cpp
--- a/Engine/Source/Editor/Sequencer/Private/SequencerSettings.cpp
+++ b/Engine/Source/Editor/Sequencer/Private/SequencerSettings.cpp
@@ -983,17 +983,17 @@ void USequencerSettings::SetTreeViewWidth(float InTreeViewWidth)
 
 UE::Sequencer::EViewDensity USequencerSettings::GetViewDensity() const
 {
+	static FName NAME_Compact("Compact");
 	static FName NAME_Relaxed("Relaxed");
-	static FName NAME_Expanded("Expanded");
-	if (ViewDensity == NAME_Relaxed)
+	if (ViewDensity == NAME_Compact)
 	{
-		return UE::Sequencer::EViewDensity::Relaxed;
+		return UE::Sequencer::EViewDensity::Compact;
 	}
-	if (ViewDensity == NAME_Expanded)
+	if (ViewDensity == NAME_Relaxed)
 	{
-		return UE::Sequencer::EViewDensity::Expanded;
+		return UE::Sequencer::EViewDensity::Relaxed;
 	}
-	return UE::Sequencer::EViewDensity::Compact;
+	return UE::Sequencer::EViewDensity::Variable;
 }
 
 void USequencerSettings::SetViewDensity(FName InViewDensity)
diff --git a/Engine/Source/Editor/SequencerCore/Private/MVVM/ViewModels/ViewDensity.cpp b/Engine/Source/Editor/SequencerCore/Private/MVVM/ViewModels/ViewDensity.cpp
--- a/Engine/Source/Editor/SequencerCore/Private/MVVM/ViewModels/ViewDensity.cpp
+++ b/Engine/Source/Editor/SequencerCore/Private/MVVM/ViewModels/ViewDensity.cpp
@@ -6,23 +6,23 @@
 namespace UE::Sequencer
 {
 
-float GSequencerOutlinerRelaxedHeight = 22.f;
+float GSequencerOutlinerCompactHeight = 22.f;
+FAutoConsoleVariableRef CVarSequencerOutlinerCompactHeight(
+	TEXT("Sequencer.Outliner.CompactHeight"),
+	GSequencerOutlinerCompactHeight,
+	TEXT("(Default: 22.f. Defines the height of outliner items when in compact mode.")
+	);
+float GSequencerOutlinerRelaxedHeight = 28.f;
 FAutoConsoleVariableRef CVarSequencerOutlinerRelaxedHeight(
 	TEXT("Sequencer.Outliner.RelaxedHeight"),
 	GSequencerOutlinerRelaxedHeight,
-	TEXT("(Default: 22.f. Defines the height of outliner items when in relaxed mode.")
-	);
-float GSequencerOutlinerExpandedHeight = 28.f;
-FAutoConsoleVariableRef CVarSequencerOutlinerExpandedHeight(
-	TEXT("Sequencer.Outliner.ExpandedHeight"),
-	GSequencerOutlinerExpandedHeight,
-	TEXT("(Default: 28.f. Defines the height of outliner items when in expanded mode.")
+	TEXT("(Default: 28.f. Defines the height of outliner items when in relaxed mode.")
 	);
 
 
 FViewDensityInfo::FViewDensityInfo()
-	: UniformHeight(GetUniformHeight(EViewDensity::Relaxed))
-	, Density(EViewDensity::Relaxed)
+	: UniformHeight(GetUniformHeight(EViewDensity::Compact))
+	, Density(EViewDensity::Compact)
 {}
 
 FViewDensityInfo::FViewDensityInfo(EViewDensity InViewDensity)
@@ -34,8 +34,8 @@ TOptional<float> FViewDensityInfo::GetUniformHeight(EViewDensity Density)
 {
 	switch(Density)
 	{
+		case EViewDensity::Compact:  return GSequencerOutlinerCompactHeight;
 		case EViewDensity::Relaxed:  return GSequencerOutlinerRelaxedHeight;
-		case EViewDensity::Expanded: return GSequencerOutlinerExpandedHeight;
 		default:                     return TOptional<float>();
 	}
 }
diff --git a/Engine/Source/Editor/SequencerCore/Public/MVVM/ViewModels/ViewDensity.h b/Engine/Source/Editor/SequencerCore/Public/MVVM/ViewModels/ViewDensity.h
--- a/Engine/Source/Editor/SequencerCore/Public/MVVM/ViewModels/ViewDensity.h
+++ b/Engine/Source/Editor/SequencerCore/Public/MVVM/ViewModels/ViewDensity.h
@@ -10,14 +10,14 @@ namespace UE::Sequencer
 
 enum class EViewDensity
 {
-	/** The most dense view type */
-	Compact,
+	/** Variable density where inner items are more condensed than top level or outer items */
+	Variable,
 
-	/** A relaxed view type with larger, uniform heights */
-	Relaxed,
+	/** A compact view type with uniform heights */
+	Compact,
 
-	/** An expanded view type with even larger, uniform heights and additional information */
-	Expanded
+	/** A relaxed view type with larger, uniform heights and additional information */
+	Relaxed
 };
 
 
