Sequencer: Don't set time if it's already the same.

#rnx

diff --git a/Engine/Source/Editor/Sequencer/Private/Sequencer.cpp b/Engine/Source/Editor/Sequencer/Private/Sequencer.cpp
--- a/Engine/Source/Editor/Sequencer/Private/Sequencer.cpp
+++ b/Engine/Source/Editor/Sequencer/Private/Sequencer.cpp
@@ -4304,10 +4304,18 @@ TSharedRef<SWidget> FSequencer::MakePlayTimeDisplay(const TSharedRef<INumericTyp
 			return GetLocalTime().Time.GetFrame().Value;
 		})
 		.OnValueChanged_Lambda([this](double InFrame) {
-			SequencerWidget->SetPlayTimeClampedByWorkingRange(InFrame);
+			double CurrentTime = GetLocalTime().Time.AsDecimal();
+			if (CurrentTime != InFrame)
+			{
+				SequencerWidget->SetPlayTimeClampedByWorkingRange(InFrame);
+			}
 		})
 		.OnValueCommitted_Lambda([this](double InFrame, ETextCommit::Type) {
-			SequencerWidget->SetPlayTime(InFrame);
+			double CurrentTime = GetLocalTime().Time.AsDecimal();
+			if (CurrentTime != InFrame)
+			{
+				SequencerWidget->SetPlayTime(InFrame);
+			}
 		})
 		.MinValue(TOptional<double>())
 		.MaxValue(TOptional<double>())
