Sequencer: Fix next/previous key not working properly for section bounds

In 30522626, next/previous key in SSequencer was consolidated to use FSequencerKeyCollection. However, SSequencer had explicit code that looked at section bounds (including those without channel groups). So that broke next/previous key for section bounds.

This change adds support for section bounds in FSequencerKeyCollection by traversing FSectionModels and not just FChannelGroupModels. There's now a map of SignatureToSectionBounds to cache the bounds.

Also, combined the functions to get next key and next section key. There's really no need to have separate functions to get next values and it just makes it more complicated.

#rnx

diff --git a/Engine/Plugins/MovieScene/LevelSequenceEditor/Source/LevelSequenceEditor/Private/CinematicViewport/SCinematicTransportRange.cpp b/Engine/Plugins/MovieScene/LevelSequenceEditor/Source/LevelSequenceEditor/Private/CinematicViewport/SCinematicTransportRange.cpp
--- a/Engine/Plugins/MovieScene/LevelSequenceEditor/Source/LevelSequenceEditor/Private/CinematicViewport/SCinematicTransportRange.cpp
+++ b/Engine/Plugins/MovieScene/LevelSequenceEditor/Source/LevelSequenceEditor/Private/CinematicViewport/SCinematicTransportRange.cpp
@@ -257,7 +257,7 @@ int32 SCinematicTransportRange::OnPaint(const FPaintArgs& Args, const FGeometry&
 			(ViewRange.GetUpperBoundValue() * TickResolution).CeilToFrame()
 			);
 
-		TArrayView<const FFrameNumber> Keys = ActiveKeyCollection->GetKeysInRange(VisibleFrameRange);
+		TArrayView<const FFrameNumber> Keys = ActiveKeyCollection->GetKeysInRange(VisibleFrameRange, EFindKeyType::FKT_All);
 
 		DrawKeys(AllottedGeometry, OutDrawElements, LayerId, bParentEnabled, Keys, TArrayView<FLinearColor>(), bPlayMarkerOnKey);
 	}
diff --git a/Engine/Source/Editor/Sequencer/Private/MVVM/ViewModels/SequencerOutlinerViewModel.cpp b/Engine/Source/Editor/Sequencer/Private/MVVM/ViewModels/SequencerOutlinerViewModel.cpp
--- a/Engine/Source/Editor/Sequencer/Private/MVVM/ViewModels/SequencerOutlinerViewModel.cpp
+++ b/Engine/Source/Editor/Sequencer/Private/MVVM/ViewModels/SequencerOutlinerViewModel.cpp
@@ -134,7 +134,7 @@ FFrameNumber FSequencerOutlinerViewModel::GetNextKeyInternal(const TArray<TShare
 
 		KeyCollection->Update(FSequencerKeyCollectionSignature::FromNodesRecursive(InNodes, ThresholdFrames));
 
-		TOptional<FFrameNumber> NextKey = KeyCollection->GetNextKey(FrameNumber, Direction, Range);
+		TOptional<FFrameNumber> NextKey = KeyCollection->GetNextKey(FrameNumber, Direction, Range, EFindKeyType::FKT_All);
 		if (NextKey.IsSet())
 		{
 			if (TimeUnit == EMovieSceneTimeUnit::DisplayRate)
diff --git a/Engine/Source/Editor/Sequencer/Private/Sequencer.cpp b/Engine/Source/Editor/Sequencer/Private/Sequencer.cpp
--- a/Engine/Source/Editor/Sequencer/Private/Sequencer.cpp
+++ b/Engine/Source/Editor/Sequencer/Private/Sequencer.cpp
@@ -4849,7 +4849,7 @@ FReply FSequencer::JumpToPreviousKey()
 	{
 		TRange<FFrameNumber> Range = GetTimeBounds();
 		FFrameNumber FrameNumber = GetLocalTime().Time.FloorToFrame();
-		TOptional<FFrameNumber> NewTime = SelectedKeyCollection->GetNextKey(FrameNumber, EFindKeyDirection::Backwards, Range);
+		TOptional<FFrameNumber> NewTime = SelectedKeyCollection->GetNextKey(FrameNumber, EFindKeyDirection::Backwards, Range, EFindKeyType::FKT_All);
 		if (NewTime.IsSet())
 		{
 			SetPlaybackStatus(EMovieScenePlayerStatus::Stepping);
@@ -4879,7 +4879,7 @@ FReply FSequencer::JumpToNextKey()
 	{
 		TRange<FFrameNumber> Range = GetTimeBounds();
 		FFrameNumber FrameNumber = GetLocalTime().Time.FloorToFrame();
-		TOptional<FFrameNumber> NewTime = SelectedKeyCollection->GetNextKey(FrameNumber, EFindKeyDirection::Forwards, Range);
+		TOptional<FFrameNumber> NewTime = SelectedKeyCollection->GetNextKey(FrameNumber, EFindKeyDirection::Forwards, Range, EFindKeyType::FKT_All);
 		if (NewTime.IsSet())
 		{
 			SetPlaybackStatus(EMovieScenePlayerStatus::Stepping);
@@ -5227,10 +5227,10 @@ FFrameNumber FSequencer::OnGetNearestKey(FFrameTime InTime, ENearestKeyOption Ne
 		TOptional<FFrameNumber> NearestKeyTime;
 
 		TRange<FFrameNumber> FindRangeBackwards(TRangeBound<FFrameNumber>::Open(), CurrentTime);
-		TOptional<FFrameNumber> NewTimeBackwards = SelectedKeyCollection->FindFirstKeyInRange(FindRangeBackwards, EFindKeyDirection::Backwards);
+		TOptional<FFrameNumber> NewTimeBackwards = SelectedKeyCollection->FindFirstKeyInRange(FindRangeBackwards, EFindKeyDirection::Backwards, EFindKeyType::FKT_Keys);
 
 		TRange<FFrameNumber> FindRangeForwards(CurrentTime, TRangeBound<FFrameNumber>::Open());
-		TOptional<FFrameNumber> NewTimeForwards = SelectedKeyCollection->FindFirstKeyInRange(FindRangeForwards, EFindKeyDirection::Forwards);
+		TOptional<FFrameNumber> NewTimeForwards = SelectedKeyCollection->FindFirstKeyInRange(FindRangeForwards, EFindKeyDirection::Forwards, EFindKeyType::FKT_Keys);
 		if (NewTimeForwards.IsSet())
 		{
 			if (NewTimeBackwards.IsSet())
@@ -5265,10 +5265,10 @@ FFrameNumber FSequencer::OnGetNearestKey(FFrameTime InTime, ENearestKeyOption Ne
 		TOptional<FFrameNumber> NearestSectionTime;
 
 		TRange<FFrameNumber> FindRangeBackwards(TRangeBound<FFrameNumber>::Open(), CurrentTime);
-		TOptional<FFrameNumber> NewTimeBackwards = SelectedKeyCollection->FindFirstSectionKeyInRange(FindRangeBackwards, EFindKeyDirection::Backwards);
+		TOptional<FFrameNumber> NewTimeBackwards = SelectedKeyCollection->FindFirstKeyInRange(FindRangeBackwards, EFindKeyDirection::Backwards, EFindKeyType::FKT_Sections);
 
 		TRange<FFrameNumber> FindRangeForwards(CurrentTime, TRangeBound<FFrameNumber>::Open());
-		TOptional<FFrameNumber> NewTimeForwards = SelectedKeyCollection->FindFirstSectionKeyInRange(FindRangeForwards, EFindKeyDirection::Forwards);
+		TOptional<FFrameNumber> NewTimeForwards = SelectedKeyCollection->FindFirstKeyInRange(FindRangeForwards, EFindKeyDirection::Forwards, EFindKeyType::FKT_Sections);
 		if (NewTimeForwards.IsSet())
 		{
 			if (NewTimeBackwards.IsSet())
diff --git a/Engine/Source/Editor/Sequencer/Private/SequencerKeyCollection.cpp b/Engine/Source/Editor/Sequencer/Private/SequencerKeyCollection.cpp
--- a/Engine/Source/Editor/Sequencer/Private/SequencerKeyCollection.cpp
+++ b/Engine/Source/Editor/Sequencer/Private/SequencerKeyCollection.cpp
@@ -10,6 +10,7 @@
 #include "MVVM/SharedList.h"
 #include "MVVM/ViewModelPtr.h"
 #include "MVVM/ViewModels/ChannelModel.h"
+#include "MVVM/ViewModels/SectionModel.h"
 #include "MVVM/ViewModels/ViewModel.h"
 #include "MVVM/ViewModels/ViewModelIterators.h"
 #include "Math/Range.h"
@@ -40,6 +41,15 @@ FSequencerKeyCollectionSignature FSequencerKeyCollectionSignature::FromNodes(con
 				}
 			}
 		}
+
+		for (TSharedPtr<FSectionModel> SectionModel : Node->GetChildrenOfType<FSectionModel>())
+		{
+			if (UMovieSceneSection* Section = SectionModel->GetSection())
+			{
+				TRange<FFrameNumber> SectionRange = SectionModel->GetRange();
+				Result.SignatureToSectionBounds.Add(Section->GetSignature(), SectionRange);
+			}
+		}
 	}
 
 	return Result;
@@ -65,6 +75,15 @@ FSequencerKeyCollectionSignature FSequencerKeyCollectionSignature::FromNodesRecu
 				AllKeyAreaNodes.Add(KeyAreaNode);
 			}
 		}
+
+		for (TSharedPtr<FSectionModel> SectionModel : Node->GetDescendantsOfType<FSectionModel>())
+		{
+			if (UMovieSceneSection* Section = SectionModel->GetSection())
+			{
+				TRange<FFrameNumber> SectionRange = SectionModel->GetRange();
+				Result.SignatureToSectionBounds.Add(Section->GetSignature(), SectionRange);
+			}
+		}
 	}
 
 	for (const TSharedPtr<FChannelGroupModel>& Node : AllKeyAreaNodes)
@@ -97,6 +116,15 @@ FSequencerKeyCollectionSignature FSequencerKeyCollectionSignature::FromNodeRecur
 		}
 	}
 
+	for (TSharedPtr<FSectionModel> SectionModel : InNode->GetChildrenOfType<FSectionModel>())
+	{
+		if (UMovieSceneSection* Section = SectionModel->GetSection())
+		{
+			TRange<FFrameNumber> SectionRange = SectionModel->GetRange();
+			Result.SignatureToSectionBounds.Add(Section->GetSignature(), SectionRange);
+		}
+	}
+
 	for (const TSharedPtr<FChannelGroupModel>& Node : AllKeyAreaNodes)
 	{
 		TSharedPtr<IKeyArea> KeyArea = Node->GetKeyArea(InSection);
@@ -118,6 +146,13 @@ bool FSequencerKeyCollectionSignature::HasUncachableContent() const
 			return true;
 		}
 	}
+	for (auto& Pair : SignatureToSectionBounds)
+	{
+		if (!Pair.Key.IsValid())
+		{
+			return true;
+		}
+	}
 	return false;
 }
 
@@ -128,7 +163,7 @@ bool operator!=(const FSequencerKeyCollectionSignature& A, const FSequencerKeyCo
 		return true;
 	}
 
-	if (A.KeyAreaToSignature.Num() != B.KeyAreaToSignature.Num() || A.DuplicateThresholdTime != B.DuplicateThresholdTime)
+	if (A.KeyAreaToSignature.Num() != B.KeyAreaToSignature.Num() || A.DuplicateThresholdTime != B.DuplicateThresholdTime || A.SignatureToSectionBounds.Num() != B.SignatureToSectionBounds.Num())
 	{
 		return true;
 	}
@@ -142,6 +177,15 @@ bool operator!=(const FSequencerKeyCollectionSignature& A, const FSequencerKeyCo
 		}
 	}
 
+	for (auto& Pair : A.SignatureToSectionBounds)
+	{
+		const TRange<FFrameNumber>* SectionBounds = B.SignatureToSectionBounds.Find(Pair.Key);
+		if (!SectionBounds || *SectionBounds != Pair.Value)
+		{
+			return true;
+		}
+	}
+
 	return false;
 }
 
@@ -152,7 +196,7 @@ bool operator==(const FSequencerKeyCollectionSignature& A, const FSequencerKeyCo
 		return false;
 	}
 
-	if (A.KeyAreaToSignature.Num() != B.KeyAreaToSignature.Num() || A.DuplicateThresholdTime != B.DuplicateThresholdTime)
+	if (A.KeyAreaToSignature.Num() != B.KeyAreaToSignature.Num() || A.DuplicateThresholdTime != B.DuplicateThresholdTime || A.SignatureToSectionBounds.Num() != B.SignatureToSectionBounds.Num())
 	{
 		return false;
 	}
@@ -166,6 +210,15 @@ bool operator==(const FSequencerKeyCollectionSignature& A, const FSequencerKeyCo
 		}
 	}
 
+	for (auto& Pair : A.SignatureToSectionBounds)
+	{
+		const TRange<FFrameNumber>* SectionBounds = B.SignatureToSectionBounds.Find(Pair.Key);
+		if (!SectionBounds || *SectionBounds != Pair.Value)
+		{
+			return false;
+		}
+	}
+
 	return true;
 }
 
@@ -185,16 +238,18 @@ bool FSequencerKeyCollection::Update(const FSequencerKeyCollectionSignature& InS
 		if (UMovieSceneSection* Section = Pair.Key->GetOwningSection())
 		{
 			Pair.Key->GetKeyTimes(AllTimes, Section->GetRange());
-
-			if (Section->HasStartFrame())
-			{
-				AllSectionTimes.Add(Section->GetInclusiveStartFrame());
-			}
+		}
+	}
 			
-			if (Section->HasEndFrame())
-			{
-				AllSectionTimes.Add(Section->GetExclusiveEndFrame());
-			}
+	for (auto& Pair : InSignature.GetSectionBounds())
+	{
+		if (Pair.Value.HasLowerBound())
+		{
+			AllSectionTimes.Add(Pair.Value.GetLowerBoundValue());
+		}
+		if (Pair.Value.HasUpperBound())
+		{
+			AllSectionTimes.Add(Pair.Value.GetUpperBoundValue());
 		}
 	}
 
@@ -227,6 +282,25 @@ bool FSequencerKeyCollection::Update(const FSequencerKeyCollectionSignature& InS
 	}
 	GroupedSectionTimes.Shrink();
 
+	AllTimes = GroupedTimes;
+	if (GroupedSectionTimes.Num() > 0)
+	{
+		AllTimes.Insert(GroupedSectionTimes, 0);
+		AllTimes.Sort();
+	}
+	AllGroupedTimes.Reset(AllTimes.Num());
+	Index = 0;
+	while (Index < AllTimes.Num())
+	{
+		FFrameNumber PredicateTime = AllTimes[Index];
+		AllGroupedTimes.Add(PredicateTime);
+		while (Index < AllTimes.Num() && FMath::Abs(AllTimes[Index] - PredicateTime) <= InSignature.GetDuplicateThreshold())
+		{
+			++Index;
+		}
+	}
+	AllGroupedTimes.Shrink();
+
 	Signature = InSignature;
 
 	return true;
@@ -292,32 +366,53 @@ TOptional<FFrameNumber> FindFirstKeyInRangeInternal(const TArray<FFrameNumber>&
 	return TOptional<FFrameNumber>();
 }
 
-TOptional<FFrameNumber> FSequencerKeyCollection::FindFirstKeyInRange(const TRange<FFrameNumber>& Range, EFindKeyDirection Direction) const
-{
-	return FindFirstKeyInRangeInternal(GroupedTimes, Range, Direction);
-}
-
-TOptional<FFrameNumber> FSequencerKeyCollection::FindFirstSectionKeyInRange(const TRange<FFrameNumber>& Range, EFindKeyDirection Direction) const
-{
-	return FindFirstKeyInRangeInternal(GroupedSectionTimes, Range, Direction);
-}
-
-TArrayView<const FFrameNumber> FSequencerKeyCollection::GetKeysInRange(const TRange<FFrameNumber>& Range) const
-{
-	return GetKeysInRangeInternal(GroupedTimes, Range);
-}
-
-TArrayView<const FFrameNumber> FSequencerKeyCollection::GetSectionKeysInRange(const TRange<FFrameNumber>& Range) const
+TOptional<FFrameNumber> FSequencerKeyCollection::FindFirstKeyInRange(const TRange<FFrameNumber>& Range, EFindKeyDirection Direction, EFindKeyType FindKeyType) const
 {
-	return GetKeysInRangeInternal(GroupedSectionTimes, Range);
+	if (FindKeyType == EFindKeyType::FKT_Keys)
+	{
+		return FindFirstKeyInRangeInternal(GroupedTimes, Range, Direction);
+	}
+	else if (FindKeyType == EFindKeyType::FKT_Sections)
+	{
+		return FindFirstKeyInRangeInternal(GroupedSectionTimes, Range, Direction);
+	}
+	else if (FindKeyType == EFindKeyType::FKT_All)
+	{
+		return FindFirstKeyInRangeInternal(AllGroupedTimes, Range, Direction);
+	}
+	return TOptional<FFrameNumber>();
 }
 
-TOptional<FFrameNumber> FSequencerKeyCollection::GetNextKey(FFrameNumber FrameNumber, EFindKeyDirection Direction, const TRange<FFrameNumber>& Range) const
+TArrayView<const FFrameNumber> FSequencerKeyCollection::GetKeysInRange(const TRange<FFrameNumber>& Range, EFindKeyType FindKeyType) const
 {
-	return GetNextKeyInternal(GroupedTimes, FrameNumber, Direction, Range);
+	if (FindKeyType == EFindKeyType::FKT_Keys)
+	{
+		return GetKeysInRangeInternal(GroupedTimes, Range);
+	}
+	else if (FindKeyType == EFindKeyType::FKT_Sections)
+	{
+		return GetKeysInRangeInternal(GroupedSectionTimes, Range);
+	}
+	else if (FindKeyType == EFindKeyType::FKT_All)
+	{
+		return GetKeysInRangeInternal(AllGroupedTimes, Range);
+	}
+	return TArrayView<const FFrameNumber>();
 }
 
-TOptional<FFrameNumber> FSequencerKeyCollection::GetNextSectionKey(FFrameNumber FrameNumber, EFindKeyDirection Direction, const TRange<FFrameNumber>& Range) const
+TOptional<FFrameNumber> FSequencerKeyCollection::GetNextKey(FFrameNumber FrameNumber, EFindKeyDirection Direction, const TRange<FFrameNumber>& Range, EFindKeyType FindKeyType) const
 {
-	return GetNextKeyInternal(GroupedSectionTimes, FrameNumber, Direction, Range);
+	if (FindKeyType == EFindKeyType::FKT_Keys)
+	{
+		return GetNextKeyInternal(GroupedTimes, FrameNumber, Direction, Range);
+	}
+	else if (FindKeyType == EFindKeyType::FKT_Sections)
+	{
+		return GetNextKeyInternal(GroupedSectionTimes, FrameNumber, Direction, Range);
+	}
+	else if (FindKeyType == EFindKeyType::FKT_All)
+	{
+		return GetNextKeyInternal(AllGroupedTimes, FrameNumber, Direction, Range);
+	}
+	return TOptional<FFrameNumber>();
 }
diff --git a/Engine/Source/Editor/Sequencer/Public/SequencerKeyCollection.h b/Engine/Source/Editor/Sequencer/Public/SequencerKeyCollection.h
--- a/Engine/Source/Editor/Sequencer/Public/SequencerKeyCollection.h
+++ b/Engine/Source/Editor/Sequencer/Public/SequencerKeyCollection.h
@@ -8,6 +8,7 @@
 #include "Containers/Map.h"
 #include "Containers/SparseArray.h"
 #include "CoreTypes.h"
+#include "Math/Range.h"
 #include "MVVM/ViewModels/ViewModelHierarchy.h"
 #include "Misc/FrameNumber.h"
 #include "Misc/Guid.h"
@@ -26,14 +27,19 @@ namespace Sequencer
 }
 }
 
-template<typename> class TRange;
-
 /** Enumeration used to define how to search for keys */
 enum class EFindKeyDirection
 {
 	Backwards, Forwards
 };
 
+enum class EFindKeyType : uint8
+{
+	FKT_Keys,
+	FKT_Sections,
+	FKT_All
+};
+
 struct FSequencerKeyCollectionSignature
 {
 	using FViewModel = UE::Sequencer::FViewModel;
@@ -62,6 +68,12 @@ struct FSequencerKeyCollectionSignature
 		return KeyAreaToSignature;
 	}
 
+	/** Access the map of signatures to section bounds that this signature was generated for */
+	const TMap<FGuid, TRange<FFrameNumber>>& GetSectionBounds() const
+	{
+		return SignatureToSectionBounds;
+	}
+
 	/** Access duplicate threshold that this signature was generated for */
 	FFrameNumber GetDuplicateThreshold() const
 	{
@@ -78,6 +90,9 @@ private:
 
 	/** Map of key areas to the section signature with with this signature was generated */
 	TMap<TSharedRef<IKeyArea>, FGuid> KeyAreaToSignature;
+
+	/** Map of the section signature to the section bounds */
+	TMap<FGuid, TRange<FFrameNumber> > SignatureToSectionBounds;
 };
 
 /**
@@ -94,8 +109,9 @@ public:
 	 * @param Direction  Whether to return the first or last key that reside in the given range
 	 * @return (Optional) the time of the key that matched the range
 	 */
-	SEQUENCER_API TOptional<FFrameNumber> FindFirstKeyInRange(const TRange<FFrameNumber>& Range, EFindKeyDirection Direction) const;
-	SEQUENCER_API TOptional<FFrameNumber> FindFirstSectionKeyInRange(const TRange<FFrameNumber>& Range, EFindKeyDirection Direction) const;
+	SEQUENCER_API TOptional<FFrameNumber> FindFirstKeyInRange(const TRange<FFrameNumber>& Range, EFindKeyDirection Direction, EFindKeyType FindKeyType = EFindKeyType::FKT_All) const;
+	UE_DEPRECATED(5.4, "Please use FindFirstKeyInRange which takes EFindKeyType, FindFirstSectionKeyInRange is no longer supported")
+	SEQUENCER_API TOptional<FFrameNumber> FindFirstSectionKeyInRange(const TRange<FFrameNumber>& Range, EFindKeyDirection Direction) const { return FindFirstKeyInRange(Range, Direction, EFindKeyType::FKT_Sections); }
 
 	/**
 	 * Get a view of all key times that reside within the specified range
@@ -103,8 +119,9 @@ public:
 	 * @param Range      The range to search within
 	 * @return A (possibly empty) array view of all the times that lie within the range
 	 */
-	SEQUENCER_API TArrayView<const FFrameNumber> GetKeysInRange(const TRange<FFrameNumber>& Range) const;
-	SEQUENCER_API TArrayView<const FFrameNumber> GetSectionKeysInRange(const TRange<FFrameNumber>& Range) const;
+	SEQUENCER_API TArrayView<const FFrameNumber> GetKeysInRange(const TRange<FFrameNumber>& Range, EFindKeyType FindKeyType = EFindKeyType::FKT_All) const;
+	UE_DEPRECATED(5.4, "Please use GetKeysInRange which takes EFindKeyType, GetSectionKeysInRange is no longer supported")
+	SEQUENCER_API TArrayView<const FFrameNumber> GetSectionKeysInRange(const TRange<FFrameNumber>& Range) const { return GetKeysInRange(Range, EFindKeyType::FKT_Sections); }
 
 	/**
 	* Search forwards or backwards for the next key from the specified frame number
@@ -112,8 +129,9 @@ public:
 	* @param Direction  Whether to return the next key or previous key from that time
 	* @return (Optional)  Frame number of the key that's next or previous from that time 
 	*/
-	SEQUENCER_API TOptional<FFrameNumber> GetNextKey(FFrameNumber FrameNumber, EFindKeyDirection Direction, const TRange<FFrameNumber>& Range) const;
-	SEQUENCER_API TOptional<FFrameNumber> GetNextSectionKey(FFrameNumber FrameNumber, EFindKeyDirection Direction, const TRange<FFrameNumber>& Range) const;
+	SEQUENCER_API TOptional<FFrameNumber> GetNextKey(FFrameNumber FrameNumber, EFindKeyDirection Direction, const TRange<FFrameNumber>& Range, EFindKeyType FindKeyType = EFindKeyType::FKT_All) const;
+	UE_DEPRECATED(5.4, "Please use GetNextKey which takes EFindKeyType, GetNextSectionKey is no longer supported")
+	SEQUENCER_API TOptional<FFrameNumber> GetNextSectionKey(FFrameNumber FrameNumber, EFindKeyDirection Direction, const TRange<FFrameNumber>& Range) const { return GetNextKey(FrameNumber, Direction, Range, EFindKeyType::FKT_Sections); }
 
 	/**
 	 * Access the signature this collection was generated with
@@ -137,6 +155,9 @@ public:
 
 private:
 
+	/** All keys and section times grouped by the supplied threshold */
+	TArray<FFrameNumber> AllGroupedTimes;
+
 	/** Times grouped by the supplied threshold */
 	TArray<FFrameNumber> GroupedTimes;
 
