Sequencer: Upper bound of the playback should always be exclusive.

Fix setting the upper bound to be exclusive when updating the playback range to the section bounds.

diff --git a/Engine/Source/Editor/Sequencer/Private/Sequencer.cpp b/Engine/Source/Editor/Sequencer/Private/Sequencer.cpp
--- a/Engine/Source/Editor/Sequencer/Private/Sequencer.cpp
+++ b/Engine/Source/Editor/Sequencer/Private/Sequencer.cpp
@@ -3070,6 +3070,11 @@ void FSequencer::UpdatePlaybackRange()
 		// When the playback range is determined by the section bounds, don't mark the change in the playback range otherwise the scene will be marked dirty
 		if (!NewBounds.IsDegenerate())
 		{
+			// Playback ranges should always have exclusive upper bounds
+			if (NewBounds.GetUpperBound().IsInclusive())
+			{
+				NewBounds.SetUpperBound(TRangeBound<FFrameNumber>::Exclusive(NewBounds.GetUpperBound().GetValue() + 1));
+			}
 			const bool bAlwaysMarkDirty = false;
 			FocusedMovieScene->SetPlaybackRange(NewBounds, bAlwaysMarkDirty);
 		}
