Sequencer: Numeric edits to key values now only occur on Enter and not when the focus changes.

This removes a hack that temporarily clears focus before changing time and then restoring focus. The reason that was necessary when key value changes were commited on focus change, the commit would happen after the time changed, which meant that the key would change on the time you're changing to, rather than the current time that you're trying affect. Regardless, this behavior is all removed now.

Tested the original jira bug and it no longer occurs with these changes.

diff --git a/Engine/Source/Editor/MovieSceneTools/Private/CurveKeyEditors/SNumericKeyEditor.h b/Engine/Source/Editor/MovieSceneTools/Private/CurveKeyEditors/SNumericKeyEditor.h
--- a/Engine/Source/Editor/MovieSceneTools/Private/CurveKeyEditors/SNumericKeyEditor.h
+++ b/Engine/Source/Editor/MovieSceneTools/Private/CurveKeyEditors/SNumericKeyEditor.h
@@ -113,7 +113,7 @@ private:
 
 	void OnValueCommitted(NumericType Value, ETextCommit::Type CommitInfo)
 	{
-		if (CommitInfo == ETextCommit::OnEnter || CommitInfo == ETextCommit::OnUserMovedFocus)
+		if (CommitInfo == ETextCommit::OnEnter)
 		{
 			const FScopedTransaction Transaction( LOCTEXT("SetNumericKey", "Set Key Value") );
 			KeyEditor.SetValueWithNotify(Value, EMovieSceneDataChangeType::TrackValueChangedRefreshImmediately);
diff --git a/Engine/Source/Editor/Sequencer/Private/Sequencer.cpp b/Engine/Source/Editor/Sequencer/Private/Sequencer.cpp
--- a/Engine/Source/Editor/Sequencer/Private/Sequencer.cpp
+++ b/Engine/Source/Editor/Sequencer/Private/Sequencer.cpp
@@ -3232,14 +3232,6 @@ void FSequencer::SetLocalTime( FFrameTime NewTime, ESnapTimeMode SnapTimeMode, b
 
 void FSequencer::SetLocalTimeDirectly(FFrameTime NewTime, bool bEvaluate)
 {
-	TWeakPtr<SWidget> PreviousFocusedWidget = FSlateApplication::Get().GetKeyboardFocusedWidget();
-
-	// Clear focus before setting time in case there's a key editor value selected that gets committed to a newly selected key on UserMovedFocus
-	if (GetPlaybackStatus() == EMovieScenePlayerStatus::Stopped)
-	{
-		FSlateApplication::Get().ClearKeyboardFocus(EFocusCause::Cleared);
-	}
-
 	// Special-case. If the RootToLocalTransform contains a zero timescale, then scrubbing in the subscene should not change the current time
 	// Inverting the zero time-scale time is technically non-deterministic, and so doing this scrub would just result in losing the current global time
 	// while maintaining the same sub-sequence time (as we're holding a single frame there). So the best case here is just preventing the scrub.
@@ -3248,11 +3240,6 @@ void FSequencer::SetLocalTimeDirectly(FFrameTime NewTime, bool bEvaluate)
 		// Transform the time to the root time-space
 		SetGlobalTime(NewTime * RootToLocalTransform.InverseFromLoop(RootToLocalLoopCounter), bEvaluate);
 	}
-
-	if (PreviousFocusedWidget.IsValid())
-	{
-		FSlateApplication::Get().SetKeyboardFocus(PreviousFocusedWidget.Pin());
-	}
 }
 
 
