Sequencer: Fix get nearest so that it doesn't compute anything is snapping isn't requested

#rnx

diff --git a/Engine/Source/Editor/Sequencer/Private/Sequencer.cpp b/Engine/Source/Editor/Sequencer/Private/Sequencer.cpp
--- a/Engine/Source/Editor/Sequencer/Private/Sequencer.cpp
+++ b/Engine/Source/Editor/Sequencer/Private/Sequencer.cpp
@@ -5203,7 +5203,13 @@ FFrameNumber FSequencer::OnGetNearestKey(FFrameTime InTime, ENearestKeyOption Ne
 {
 	const FFrameNumber CurrentTime = InTime.FloorToFrame();
 
-	if (EnumHasAnyFlags(NearestKeyOption, ENearestKeyOption::NKO_SearchAllTracks))
+	const bool bComputeNearest = EnumHasAnyFlags(NearestKeyOption, ENearestKeyOption::NKO_SearchKeys | ENearestKeyOption::NKO_SearchSections | ENearestKeyOption::NKO_SearchMarkers);
+	if (!bComputeNearest)
+	{
+		return CurrentTime;
+	}
+
+	if (ViewModel->GetSelection()->Outliner.Num() == 0)
 	{
 		GetAllKeys(SelectedKeyCollection, SMALL_NUMBER);
 	}
diff --git a/Engine/Source/Editor/Sequencer/Private/SequencerTimeSliderController.cpp b/Engine/Source/Editor/Sequencer/Private/SequencerTimeSliderController.cpp
--- a/Engine/Source/Editor/Sequencer/Private/SequencerTimeSliderController.cpp
+++ b/Engine/Source/Editor/Sequencer/Private/SequencerTimeSliderController.cpp
@@ -2060,13 +2060,6 @@ FFrameTime FSequencerTimeSliderController::SnapTimeToNearestKey(const FPointerEv
 	{
 		ENearestKeyOption NearestKeyOption = ENearestKeyOption::NKO_None;
 
-		// If there are any tracks selected we'll find the nearest key only on that track. If there are no keys selected,
-		// we will try to find the nearest keys on all tracks. This mirrors the behavior of the Jump to Next Keyframe commands.
-		if (WeakSequencer.Pin()->GetViewModel()->GetSelection()->Outliner.Num() == 0)
-		{
-			EnumAddFlags(NearestKeyOption, ENearestKeyOption::NKO_SearchAllTracks);
-		}
-
 		if (WeakSequencer.Pin()->GetSequencerSettings()->GetSnapPlayTimeToKeys() || MouseEvent.IsShiftDown())
 		{
 			EnumAddFlags(NearestKeyOption, ENearestKeyOption::NKO_SearchKeys);
diff --git a/Engine/Source/Editor/SequencerWidgets/Public/TimeSliderArgs.h b/Engine/Source/Editor/SequencerWidgets/Public/TimeSliderArgs.h
--- a/Engine/Source/Editor/SequencerWidgets/Public/TimeSliderArgs.h
+++ b/Engine/Source/Editor/SequencerWidgets/Public/TimeSliderArgs.h
@@ -33,7 +33,7 @@ enum class ENearestKeyOption : uint8
 	NKO_SearchMarkers = 0x04,
 
 	/** Search all tracks */
-	NKO_SearchAllTracks = 0x08
+	NKO_SearchAllTracks UE_DEPRECATED(5.4, "Search all tracks has been deprecated")  = 0x08
 };
 ENUM_CLASS_FLAGS(ENearestKeyOption);
 
