Sequencer: Add settings callbacks to SetSequencerSettings to fix other modules not responding to settings changes.

Take Recorder user settings defaults to showing time in timecode. This fixes a bug where the numeric inputs aren't set to timecode because the frame number interface is created before the settings have been initialized.

#rnx

diff --git a/Engine/Source/Editor/Sequencer/Private/Sequencer.cpp b/Engine/Source/Editor/Sequencer/Private/Sequencer.cpp
--- a/Engine/Source/Editor/Sequencer/Private/Sequencer.cpp
+++ b/Engine/Source/Editor/Sequencer/Private/Sequencer.cpp
@@ -669,6 +669,24 @@ void FSequencer::InitSequencer(const FSequencerInitParams& InitParams, const TSh
 	OnActivateSequenceEvent.Broadcast(ActiveTemplateIDs[0]);
 }
 
+void FSequencer::SetSequencerSettings(USequencerSettings* InSettings)
+{
+	if (Settings)
+	{
+		Settings->GetOnEvaluateSubSequencesInIsolationChanged().RemoveAll(this);
+		Settings->GetOnShowSelectedNodesOnlyChanged().RemoveAll(this);
+		Settings->GetOnTimeDisplayFormatChanged().RemoveAll(this);
+	}
+
+	Settings = InSettings;
+
+	Settings->GetOnEvaluateSubSequencesInIsolationChanged().AddSP(this, &FSequencer::RestorePreAnimatedState);
+	Settings->GetOnShowSelectedNodesOnlyChanged().AddSP(this, &FSequencer::OnSelectedNodesOnlyChanged);
+	Settings->GetOnTimeDisplayFormatChanged().AddSP(this, &FSequencer::OnTimeDisplayFormatChanged);
+
+	OnTimeDisplayFormatChanged();
+}
+
 void FSequencer::OnPlaybackContextChanged()
 {
 	RootTemplateInstance.PlaybackContextChanged(*this);
diff --git a/Engine/Source/Editor/Sequencer/Private/Sequencer.h b/Engine/Source/Editor/Sequencer/Private/Sequencer.h
--- a/Engine/Source/Editor/Sequencer/Private/Sequencer.h
+++ b/Engine/Source/Editor/Sequencer/Private/Sequencer.h
@@ -815,7 +815,7 @@ public:
 
 	/** Access the user-supplied settings object */
 	virtual USequencerSettings* GetSequencerSettings() const override { return Settings; }
-	virtual void SetSequencerSettings(USequencerSettings* InSettings) override { Settings = InSettings; }
+	virtual void SetSequencerSettings(USequencerSettings* InSettings) override;
 	virtual TSharedPtr<class ITimeSlider> GetTopTimeSliderWidget() const override;
 	virtual void ResetTimeController() override;
 	virtual void SetTrackFilterEnabled(const FText& InTrackFilterName, bool bEnabled) override;
