Sequencer: Clip to bounds the play time controls box so that you can collapse it entirely without seeing the contents within

#rnx

diff --git a/Engine/Source/Editor/Sequencer/Private/MVVM/CurveEditorExtension.cpp b/Engine/Source/Editor/Sequencer/Private/MVVM/CurveEditorExtension.cpp
--- a/Engine/Source/Editor/Sequencer/Private/MVVM/CurveEditorExtension.cpp
+++ b/Engine/Source/Editor/Sequencer/Private/MVVM/CurveEditorExtension.cpp
@@ -254,60 +254,66 @@ void FCurveEditorExtension::CreateCurveEditor(const FTimeSliderArgs& TimeSliderA
 			+ SVerticalBox::Slot()
 			.AutoHeight()
 			[
-				SNew(SHorizontalBox)
-
-				+ SHorizontalBox::Slot()
-				.AutoWidth()
-				.VAlign(VAlign_Center)
-				.HAlign(HAlign_Left)
+				SNew(SBorder)
+				.BorderImage(FAppStyle::GetBrush("ToolPanel.GroupBorder"))
+				.Clipping(EWidgetClipping::ClipToBounds)
 				[
-					SNew(SButton)
-					.VAlign(EVerticalAlignment::VAlign_Center)
-					.ButtonStyle(FAppStyle::Get(), "SimpleButton")
-					.ToolTipText_Lambda([this] { return LOCTEXT("ShowStatus", "Show Status"); })
-					.ContentPadding(FMargin(1, 0))
-					.OnHovered_Lambda([this] { CurveEditorTreeFilterStatusBar->ShowStatusBar(); })
-					.OnUnhovered_Lambda([this] { CurveEditorTreeFilterStatusBar->FadeOutStatusBar(); })
-					.OnClicked_Lambda([this] { CurveEditorTreeFilterStatusBar->HideStatusBar(); return FReply::Handled(); })
+
+					SNew(SHorizontalBox)
+
+					+ SHorizontalBox::Slot()
+					.AutoWidth()
+					.VAlign(VAlign_Center)
+					.HAlign(HAlign_Left)
 					[
-						SNew(SImage)
-						.ColorAndOpacity(FSlateColor::UseForeground())
-						.Image(FAppStyle::Get().GetBrush("Icons.Info.Small"))
+						SNew(SButton)
+						.VAlign(EVerticalAlignment::VAlign_Center)
+						.ButtonStyle(FAppStyle::Get(), "SimpleButton")
+						.ToolTipText_Lambda([this] { return LOCTEXT("ShowStatus", "Show Status"); })
+						.ContentPadding(FMargin(1, 0))
+						.OnHovered_Lambda([this] { CurveEditorTreeFilterStatusBar->ShowStatusBar(); })
+						.OnUnhovered_Lambda([this] { CurveEditorTreeFilterStatusBar->FadeOutStatusBar(); })
+						.OnClicked_Lambda([this] { CurveEditorTreeFilterStatusBar->HideStatusBar(); return FReply::Handled(); })
+						[
+							SNew(SImage)
+							.ColorAndOpacity(FSlateColor::UseForeground())
+							.Image(FAppStyle::Get().GetBrush("Icons.Info.Small"))
+						]
 					]
-				]
 
-				+ SHorizontalBox::Slot()
-				[
-					SNew(SBorder)
-					.BorderImage(FAppStyle::GetBrush("ToolPanel.GroupBorder"))
-					.HAlign(HAlign_Center)
+					+ SHorizontalBox::Slot()
 					[
-						Sequencer->MakeTransportControls(true)
+						SNew(SBorder)
+						.BorderImage(FAppStyle::GetBrush("ToolPanel.GroupBorder"))
+						.HAlign(HAlign_Center)
+						[
+							Sequencer->MakeTransportControls(true)
+						]
 					]
-				]
 
-				+ SHorizontalBox::Slot()
-				.AutoWidth()
-				.VAlign(VAlign_Center)
-				.HAlign(HAlign_Right)
-				[
-					SNew(SButton)
-					.VAlign(EVerticalAlignment::VAlign_Center)
-					.ButtonStyle(FAppStyle::Get(), "NoBorder")
-					.ContentPadding(FMargin(1, 0))
+					+ SHorizontalBox::Slot()
+					.AutoWidth()
+					.VAlign(VAlign_Center)
+					.HAlign(HAlign_Right)
 					[
-						SNew(SHorizontalBox)
-
-						+ SHorizontalBox::Slot()
-						.AutoWidth()
-						.VAlign(VAlign_Center)
-						.HAlign(HAlign_Right)
-						.Padding(FMargin(3.f, 0.f, 0.f, 0.f))
+						SNew(SButton)
+						.VAlign(EVerticalAlignment::VAlign_Center)
+						.ButtonStyle(FAppStyle::Get(), "NoBorder")
+						.ContentPadding(FMargin(1, 0))
 						[
-							SNew(SBorder)
-							.BorderImage(nullptr)
+							SNew(SHorizontalBox)
+
+							+ SHorizontalBox::Slot()
+							.AutoWidth()
+							.VAlign(VAlign_Center)
+							.HAlign(HAlign_Right)
+							.Padding(FMargin(3.f, 0.f, 0.f, 0.f))
 							[
-								PlayTimeDisplay.ToSharedRef()
+								SNew(SBorder)
+								.BorderImage(nullptr)
+								[
+									PlayTimeDisplay.ToSharedRef()
+								]
 							]
 						]
 					]
diff --git a/Engine/Source/Editor/Sequencer/Private/SSequencer.cpp b/Engine/Source/Editor/Sequencer/Private/SSequencer.cpp
--- a/Engine/Source/Editor/Sequencer/Private/SSequencer.cpp
+++ b/Engine/Source/Editor/Sequencer/Private/SSequencer.cpp
@@ -679,80 +679,85 @@ void SSequencer::Construct(const FArguments& InArgs, TSharedRef<FSequencer> InSe
 					// Info Button, Transport Controls and Current Frame
 					+ SGridPanel::Slot( Column0, Row4, SGridPanel::Layer(10) )
 					[
-						SNew(SHorizontalBox)
-						
-						+ SHorizontalBox::Slot()
-						.AutoWidth()
-						.VAlign(VAlign_Center)
-						.HAlign(HAlign_Left)
+						SNew(SBorder)
+						.BorderImage(FAppStyle::GetBrush("ToolPanel.GroupBorder"))
+						.Clipping(EWidgetClipping::ClipToBounds)
 						[
-							SNew(SButton)
-							.VAlign(EVerticalAlignment::VAlign_Center)
-							.ButtonStyle(FAppStyle::Get(), "SimpleButton")
-							.ToolTipText_Lambda([this] { return LOCTEXT("ShowStatus", "Show Status"); })
-							.ContentPadding(FMargin(1, 0))
-							.Visibility(this, &SSequencer::GetInfoButtonVisibility)
-							.OnHovered_Lambda([this] { SequencerTreeFilterStatusBar->ShowStatusBar(); })
-							.OnUnhovered_Lambda([this] { SequencerTreeFilterStatusBar->FadeOutStatusBar(); })
-							.OnClicked_Lambda([this] { SequencerTreeFilterStatusBar->HideStatusBar(); return FReply::Handled(); })
+							SNew(SHorizontalBox)
+						
+							+ SHorizontalBox::Slot()
+							.AutoWidth()
+							.VAlign(VAlign_Center)
+							.HAlign(HAlign_Left)
 							[
-								SNew(SImage)
-								.ColorAndOpacity(FSlateColor::UseForeground())
-								.Image(FAppStyle::Get().GetBrush("Icons.Info.Small"))
+								SNew(SButton)
+								.VAlign(EVerticalAlignment::VAlign_Center)
+								.ButtonStyle(FAppStyle::Get(), "SimpleButton")
+								.ToolTipText_Lambda([this] { return LOCTEXT("ShowStatus", "Show Status"); })
+								.ContentPadding(FMargin(1, 0))
+								.Visibility(this, &SSequencer::GetInfoButtonVisibility)
+								.OnHovered_Lambda([this] { SequencerTreeFilterStatusBar->ShowStatusBar(); })
+								.OnUnhovered_Lambda([this] { SequencerTreeFilterStatusBar->FadeOutStatusBar(); })
+								.OnClicked_Lambda([this] { SequencerTreeFilterStatusBar->HideStatusBar(); return FReply::Handled(); })
+								[
+									SNew(SImage)
+									.ColorAndOpacity(FSlateColor::UseForeground())
+									.Image(FAppStyle::Get().GetBrush("Icons.Info.Small"))
+								]
 							]
-						]
 
-						+ SHorizontalBox::Slot()
-						[
-							SNew(SBorder)
-							.BorderImage(FAppStyle::GetBrush("ToolPanel.GroupBorder"))
-							.HAlign(HAlign_Center)
+							+ SHorizontalBox::Slot()
 							[
-								SequencerPtr.Pin()->MakeTransportControls(true)
+								SNew(SBorder)
+								.BorderImage(FAppStyle::GetBrush("ToolPanel.GroupBorder"))
+								.HAlign(HAlign_Center)
+								[
+									SequencerPtr.Pin()->MakeTransportControls(true)
+								]
 							]
-						]
 
-						+ SHorizontalBox::Slot()
-						.AutoWidth()
-						.VAlign(VAlign_Center)
-						.HAlign(HAlign_Right)
-						[
-							SNew(SButton)
-							.VAlign(EVerticalAlignment::VAlign_Center)
-							.ButtonStyle(FAppStyle::Get(), "NoBorder")
-							.ContentPadding(FMargin(1, 0))
+							+ SHorizontalBox::Slot()
+							.AutoWidth()
+							.VAlign(VAlign_Center)
+							.HAlign(HAlign_Right)
 							[
-								SNew(SHorizontalBox)
-
-								+ SHorizontalBox::Slot()
-								.AutoWidth()
-								.VAlign(VAlign_Center)
-								.HAlign(HAlign_Right)
-								.Padding(FMargin(CommonPadding, 0.f, 0.f, 0.f))
+								SNew(SButton)
+								.VAlign(EVerticalAlignment::VAlign_Center)
+								.ButtonStyle(FAppStyle::Get(), "NoBorder")
+								.ContentPadding(FMargin(1, 0))
 								[
-									SNew(SBorder)
-									.BorderImage(nullptr)
+									SNew(SHorizontalBox)
+
+									+ SHorizontalBox::Slot()
+									.AutoWidth()
+									.VAlign(VAlign_Center)
+									.HAlign(HAlign_Right)
+									.Padding(FMargin(CommonPadding, 0.f, 0.f, 0.f))
 									[
-										PlayTimeDisplay.ToSharedRef()
+										SNew(SBorder)
+										.BorderImage(nullptr)
+										[
+											PlayTimeDisplay.ToSharedRef()
+										]
 									]
-								]
 
-								+ SHorizontalBox::Slot()
-								.AutoWidth()
-								.VAlign(VAlign_Center)
-								.HAlign(HAlign_Right)
-								.Padding(FMargin(CommonPadding, 0.f, 0.f, 0.f))
-								[
-									SNew(SBorder)
-									.BorderImage(nullptr)
+									+ SHorizontalBox::Slot()
+									.AutoWidth()
+									.VAlign(VAlign_Center)
+									.HAlign(HAlign_Right)
+									.Padding(FMargin(CommonPadding, 0.f, 0.f, 0.f))
 									[
-										// Current loop index, if any
-										SAssignNew(LoopIndexDisplay, STextBlock)
-										.Font(FCoreStyle::GetDefaultFontStyle("Regular", 8))
-										.Text_Lambda([this]() -> FText {
-											uint32 LoopIndex = SequencerPtr.Pin()->GetLocalLoopIndex();
-											return (LoopIndex != FMovieSceneTimeWarping::InvalidWarpCount) ? FText::AsNumber(LoopIndex + 1) : FText();
-										})
+										SNew(SBorder)
+										.BorderImage(nullptr)
+										[
+											// Current loop index, if any
+											SAssignNew(LoopIndexDisplay, STextBlock)
+											.Font(FCoreStyle::GetDefaultFontStyle("Regular", 8))
+											.Text_Lambda([this]() -> FText {
+												uint32 LoopIndex = SequencerPtr.Pin()->GetLocalLoopIndex();
+												return (LoopIndex != FMovieSceneTimeWarping::InvalidWarpCount) ? FText::AsNumber(LoopIndex + 1) : FText();
+											})
+										]
 									]
 								]
 							]
