Sequencer: Clip to bounds the play time controls box so that you can collapse it entirely without seeing the contents within

#rnx

diff --git a/Engine/Source/Editor/Sequencer/Private/MVVM/CurveEditorExtension.cpp b/Engine/Source/Editor/Sequencer/Private/MVVM/CurveEditorExtension.cpp
--- a/Engine/Source/Editor/Sequencer/Private/MVVM/CurveEditorExtension.cpp
+++ b/Engine/Source/Editor/Sequencer/Private/MVVM/CurveEditorExtension.cpp
@@ -222,7 +222,12 @@ void FCurveEditorExtension::CreateCurveEditor(const FTimeSliderArgs& TimeSliderA
 			+ SVerticalBox::Slot()
 			.AutoHeight()
 			[
-				SAssignNew(CurveEditorSearchBox, SCurveEditorTreeTextFilter, CurveEditorModel)
+				SNew(SBorder)
+				.BorderImage(FAppStyle::GetBrush("ToolPanel.GroupBorder"))
+				.Clipping(EWidgetClipping::ClipToBounds)
+				[
+					SAssignNew(CurveEditorSearchBox, SCurveEditorTreeTextFilter, CurveEditorModel)
+				]
 			]
 
 			+ SVerticalBox::Slot()
