Sequencer: Skip deprecated properties when looking for properties that are animatable

SkeletalMeshAsset was still showing up, despite being deprecated. Instead only SkeletalMesh should appear.

diff --git a/Engine/Source/Editor/Sequencer/Private/SequencerObjectChangeListener.cpp b/Engine/Source/Editor/Sequencer/Private/SequencerObjectChangeListener.cpp
--- a/Engine/Source/Editor/Sequencer/Private/SequencerObjectChangeListener.cpp
+++ b/Engine/Source/Editor/Sequencer/Private/SequencerObjectChangeListener.cpp
@@ -212,13 +212,11 @@ const FOnAnimatablePropertyChanged* FSequencerObjectChangeListener::FindProperty
 
 	if (const FOnAnimatablePropertyChanged* DelegatePtr = PropertyChangedEventMap.Find(PropertyKey))
 	{
-		FString PropertyVarName = PropertyOrContainer->GetName();
-
-		// If this is a bool property, strip off the 'b' so that the "Set" functions to be 
-		// found are, for example, "SetHidden" instead of "SetbHidden"
-		if (PropertyKey.PropertyTypeName == "BoolProperty")
+		// Skip deprecated properties
+		static const FName NAME_DeprecatedProperty = TEXT("DeprecatedProperty");
+		if (PropertyOrContainer->HasAnyPropertyFlags(CPF_Deprecated) || PropertyOrContainer->HasMetaData(NAME_DeprecatedProperty))
 		{
-			PropertyVarName.RemoveFromStart("b", ESearchCase::CaseSensitive);
+			return nullptr;
 		}
 
 		// Interp properties are always keyable
@@ -247,6 +245,15 @@ const FOnAnimatablePropertyChanged* FSequencerObjectChangeListener::FindProperty
 				static const FString Set(TEXT("Set"));
 				static const FName DeprecatedFunctionName(TEXT("DeprecatedFunction"));
 
+				FString PropertyVarName = PropertyOrContainer->GetName();
+
+				// If this is a bool property, strip off the 'b' so that the "Set" functions to be 
+				// found are, for example, "SetHidden" instead of "SetbHidden"
+				if (PropertyKey.PropertyTypeName == "BoolProperty")
+				{
+					PropertyVarName.RemoveFromStart("b", ESearchCase::CaseSensitive);
+				}
+
 				FName FunctionName = FName(*(Set + PropertyVarName));
 				UFunction* Function = Class->FindFunctionByName(FunctionName);
 
