Sequencer- Refactor Bindings to support inherited Custom Bindings.

diff --git a/Engine/Source/Editor/Sequencer/Private/Sequencer.cpp b/Engine/Source/Editor/Sequencer/Private/Sequencer.cpp
--- a/Engine/Source/Editor/Sequencer/Private/Sequencer.cpp
+++ b/Engine/Source/Editor/Sequencer/Private/Sequencer.cpp
@@ -164,6 +164,9 @@
 #include "EngineModule.h"
 #include "Tracks/MovieSceneBindingLifetimeTrack.h"
 #include "Sections/MovieSceneBindingLifetimeSection.h"
+#include "Bindings/MovieSceneSpawnableBinding.h"
+#include "Bindings/MovieSceneSpawnableActorBinding.h"
+#include "Bindings/MovieSceneReplaceableActorBinding.h"
 
 #define LOCTEXT_NAMESPACE "Sequencer"
 
@@ -432,13 +435,6 @@ void FSequencer::InitSequencer(const FSequencerInitParams& InitParams, const TSh
 				}
 			}
 
-			 /*If we are mid-resolve, we may be loading objects and trying to get a handle here may go into infinite recursion trying 
-			 to resolve the same binding.*/
-			if (State.IsResolvingObject())
-			{
-				return;
-			}
-
 			// Reset Bindings for replaced objects.
 			bool bAnythingReplaced = false;
 			for (TPair<UObject*, UObject*> ReplacedObject : ReplacementMap)
@@ -1307,9 +1303,9 @@ void FSequencer::SuppressAutoEvaluation(UMovieSceneSequence* Sequence, const FGu
 	SuppressAutoEvalSignature = MakeTuple(MakeWeakObjectPtr(Sequence), InSequenceSignature);
 }
 
-FGuid FSequencer::CreateBinding(UObject& InObject, const FString& InName)
+FGuid FSequencer::CreateBinding(UObject& InObject, const UE::Sequencer::FCreateBindingParams& InParams)
 {
-    return FSequencerUtilities::CreateBinding(AsShared(), InObject, UE::Sequencer::FCreateBindingParams().Name(CopyTemp(InName)));
+    return FSequencerUtilities::CreateBinding(AsShared(), InObject, InParams);
 }
 
 UObject* FSequencer::GetPlaybackContext() const
@@ -1615,6 +1611,22 @@ void FSequencer::UpdateSequencerCustomizations(const UMovieSceneSequence* Previo
 			OnPaste.Add(CustomizationInfo.OnPaste);
 		}
 	}
+
+	// TODO: Allow customization of which custom bindings are allowed? For now just iterate over all subclasses and sort by priority
+
+	// Cache custom spawnable types
+	SupportedCustomBindingTypes.Empty();
+	for (TObjectIterator<UClass> It; It; ++It)
+	{
+		if (It->IsChildOf(UMovieSceneCustomBinding::StaticClass()) && !It->HasAnyClassFlags(CLASS_Abstract))
+		{
+			SupportedCustomBindingTypes.Add(*It);
+		}
+	}
+
+	// Sort by spawner priority to allow disambiguation for similar object types
+	SupportedCustomBindingTypes.Sort([](const TSubclassOf<UMovieSceneCustomBinding>& A, const TSubclassOf<UMovieSceneCustomBinding>& B) {
+		return A && B && A->GetDefaultObject<UMovieSceneCustomBinding>()->GetCustomBindingPriority() > B->GetDefaultObject<UMovieSceneCustomBinding>()->GetCustomBindingPriority(); });
 }
 
 void FSequencer::RerunConstructionScripts()
@@ -3825,43 +3837,63 @@ TSharedRef<SWidget> FSequencer::MakeTimeRange(const TSharedRef<SWidget>& InnerCo
 FGuid FSequencer::FindUnspawnedObjectGuid(UObject& InObject)
 {
 	if (UMovieSceneSequence* FocusedMovieSceneSequence = GetFocusedMovieSceneSequence())
-{
+	{
 		UMovieScene* MovieScene = FocusedMovieSceneSequence->GetMovieScene();
 
-	// If the object is an archetype, the it relates to an unspawned spawnable.
+		// If the object is an archetype, the it relates to an unspawned spawnable.
 		UObject* ParentObject = FocusedMovieSceneSequence->GetParentObject(&InObject);
-	if (ParentObject && FMovieSceneSpawnable::IsSpawnableTemplate(*ParentObject))
-	{
-		FMovieSceneSpawnable* ParentSpawnable = MovieScene->FindSpawnable([&](FMovieSceneSpawnable& InSpawnable){
-			return InSpawnable.GetObjectTemplate() == ParentObject;
-		});
-
-		if (ParentSpawnable)
+		if (ParentObject && FMovieSceneSpawnable::IsSpawnableTemplate(*ParentObject))
 		{
-			// The only way to find the object now is to resolve all the child bindings, and see if they are the same
-			for (const FGuid& ChildGuid : ParentSpawnable->GetChildPossessables())
+			FMovieSceneSpawnable* ParentSpawnable = MovieScene->FindSpawnable([&](FMovieSceneSpawnable& InSpawnable){
+				return InSpawnable.GetObjectTemplate() == ParentObject;
+			});
+
+			if (ParentSpawnable)
 			{
-					TArrayView<TWeakObjectPtr<>> BoundObjects = State.FindBoundObjects(ChildGuid, GetFocusedTemplateID(), GetSharedPlaybackState());
-					const bool bHasObject = BoundObjects.Contains(&InObject);
-				if (bHasObject)
+				// TODO: Won't this always fail to find anything if the spawnable actor is unspawned?
+
+				// The only way to find the object now is to resolve all the child bindings, and see if they are the same
+				for (const FGuid& ChildGuid : ParentSpawnable->GetChildPossessables())
 				{
-					return ChildGuid;
+						TArrayView<TWeakObjectPtr<>> BoundObjects = State.FindBoundObjects(ChildGuid, GetFocusedTemplateID(), GetSharedPlaybackState());
+						const bool bHasObject = BoundObjects.Contains(&InObject);
+					if (bHasObject)
+					{
+						return ChildGuid;
+					}
 				}
 			}
 		}
-	}
-	else if (FMovieSceneSpawnable::IsSpawnableTemplate(InObject))
-	{
-		FMovieSceneSpawnable* SpawnableByArchetype = MovieScene->FindSpawnable([&](FMovieSceneSpawnable& InSpawnable){
-			return InSpawnable.GetObjectTemplate() == &InObject;
-		});
-
-		if (SpawnableByArchetype)
+		else if (FMovieSceneSpawnable::IsSpawnableTemplate(InObject))
 		{
-			return SpawnableByArchetype->GetGuid();
+			FMovieSceneSpawnable* SpawnableByArchetype = MovieScene->FindSpawnable([&](FMovieSceneSpawnable& InSpawnable){
+				return InSpawnable.GetObjectTemplate() == &InObject;
+			});
+
+			if (SpawnableByArchetype)
+			{
+				return SpawnableByArchetype->GetGuid();
+			}
+
+			// Also check for custom spawnables
+			if (const FMovieSceneBindingReferences* BindingReferences = FocusedMovieSceneSequence->GetBindingReferences())
+			{
+				for (const FMovieSceneBindingReference& BindingReference : BindingReferences->GetAllReferences())
+				{
+					if (BindingReference.CustomBinding)
+					{
+						if (UMovieSceneSpawnableBindingBase* SpawnableBinding = BindingReference.CustomBinding->AsSpawnable(GetSharedPlaybackState()))
+						{
+							if (SpawnableBinding->SupportsObjectTemplates() && SpawnableBinding->GetObjectTemplate() == &InObject)
+							{
+								return BindingReference.ID;
+							}
+						}
+					}
+				}
+			}
 		}
 	}
-	}
 
 	return FGuid();
 }
@@ -3945,9 +3977,8 @@ FGuid FSequencer::GetHandleToObject( UObject* Object, bool bCreateHandleIfMissin
 
 	if (ObjectGuid.IsValid())
 	{
-		// Check here for spawnable otherwise spawnables get recreated as possessables, which doesn't make sense
-		FMovieSceneSpawnable* Spawnable = FocusedMovieScene->FindSpawnable(ObjectGuid);
-		if (Spawnable)
+		// Check here for spawnable and custom spawnables otherwise spawnables get recreated as possessables, which doesn't make sense
+		if (MovieSceneHelpers::IsBoundToAnySpawnable(FocusedMovieSceneSequence, ObjectGuid, GetSharedPlaybackState()))
 		{
 			return ObjectGuid;
 		}
@@ -4597,27 +4628,24 @@ UObject* FSequencer::FindSpawnedObjectOrTemplate(const FGuid& BindingId)
 	if (Possessable && Possessable->GetParent().IsValid())
 	{
 		// If we're a spawnable and we don't have the object, use the default object to build up the track menu
-		FMovieSceneSpawnable* ParentSpawnable = FocusedMovieScene->FindSpawnable(Possessable->GetParent());
-		if (ParentSpawnable)
+		UObject* ParentObject = MovieSceneHelpers::GetObjectTemplate(Sequence, Possessable->GetParent(), GetSharedPlaybackState());
+
+		if (ParentObject)
 		{
-			UObject* ParentObject = ParentSpawnable->GetObjectTemplate();
-			if (ParentObject)
+			TArrayView<TWeakObjectPtr<>> BoundObjects = State.FindBoundObjects(BindingId, GetFocusedTemplateID(), GetSharedPlaybackState());
+			for (TWeakObjectPtr<> WeakObj : BoundObjects)
 			{
-				TArrayView<TWeakObjectPtr<>> BoundObjects = State.FindBoundObjects(BindingId, GetFocusedTemplateID(), GetSharedPlaybackState());
-				for (TWeakObjectPtr<> WeakObj : BoundObjects)
-				{
-					if (UObject* Obj = WeakObj.Get())
+				if (UObject* Obj = WeakObj.Get())
 				{
 					return Obj;
 				}
 			}
 		}
 	}
-	}
 	// If we're a spawnable and we don't have the object, use the default object to build up the track menu
-	else if (FMovieSceneSpawnable* Spawnable = FocusedMovieScene->FindSpawnable(BindingId))
+	else if (UObject* Template = MovieSceneHelpers::GetObjectTemplate(Sequence, BindingId, GetSharedPlaybackState()))
 	{
-		return Spawnable->GetObjectTemplate();
+		return Template;
 	}
 
 	return nullptr;
@@ -5739,8 +5767,9 @@ void FSequencer::OnNewActorsDropped(const TArray<UObject*>& DroppedObjects, cons
 {
 	bool bAddSpawnable = FSlateApplication::Get().GetModifierKeys().IsShiftDown();
 	bool bAddPossessable = FSlateApplication::Get().GetModifierKeys().IsControlDown();
+	bool bAddReplaceable = FSlateApplication::Get().GetModifierKeys().IsAltDown();
 
-	if (bAddSpawnable || bAddPossessable)
+	if (bAddSpawnable || bAddPossessable || bAddReplaceable)
 	{
 		TArray<AActor*> SpawnedActors;
 
@@ -5766,12 +5795,15 @@ void FSequencer::OnNewActorsDropped(const TArray<UObject*>& DroppedObjects, cons
 
 				OnActorAddedToSequencerEvent.Broadcast(NewActor, PossessableGuid);
 
-				if (bAddSpawnable)
+				if (bAddSpawnable || bAddReplaceable)
 				{
-					TArray<FMovieSceneSpawnable*> Spawnables = FSequencerUtilities::ConvertToSpawnable(AsShared(), PossessableGuid);
-					if (Spawnables.Num() > 0)
+					TSubclassOf<UMovieSceneCustomBinding> CustomBindingClass = bAddSpawnable ? UMovieSceneSpawnableActorBinding::StaticClass() : UMovieSceneReplaceableActorBinding::StaticClass();
+
+					FMovieScenePossessable* NewPossessable = FSequencerUtilities::ConvertToCustomBinding(AsShared(), NewGuid, CustomBindingClass);
+
+					if (NewPossessable)
 					{
-						for (TWeakObjectPtr<> WeakObject : FindBoundObjects(Spawnables[0]->GetGuid(), ActiveTemplateIDs.Top()))
+						for (TWeakObjectPtr<> WeakObject : FindBoundObjects(NewPossessable->GetGuid(), ActiveTemplateIDs.Top()))
 						{
 							AActor* SpawnedActor = Cast<AActor>(WeakObject.Get());
 							if (SpawnedActor)
@@ -5780,7 +5812,8 @@ void FSequencer::OnNewActorsDropped(const TArray<UObject*>& DroppedObjects, cons
 								NewActor = SpawnedActor;
 							}
 						}
-						NewGuid = Spawnables[0]->GetGuid();
+
+						NewGuid = NewPossessable->GetGuid();
 					}
 				}
 
@@ -6360,7 +6393,7 @@ void FSequencer::SynchronizeExternalSelectionWithSequencerSelection()
 			: SelectedItem->FindAncestorOfType<IObjectBindingExtension>(true);
 
 		// If the closest node is an object node, try to get the actor/component nodes from it.
-		if (ObjectBinding)
+	if (ObjectBinding)
 		{
 			for (auto RuntimeObject : FindBoundObjects(ObjectBinding->GetObjectGuid(), ActiveTemplateIDs.Top()) )
 			{
@@ -8539,7 +8572,8 @@ void FSequencer::SaveSelectedNodesSpawnableState()
 {
 	using namespace UE::Sequencer;
 
-	UMovieScene* MovieScene = GetFocusedMovieSceneSequence()->GetMovieScene();
+	UMovieSceneSequence* Sequence = GetFocusedMovieSceneSequence();
+	UMovieScene* MovieScene = Sequence->GetMovieScene();
 
 	if (MovieScene->IsReadOnly())
 	{
@@ -8551,33 +8585,22 @@ void FSequencer::SaveSelectedNodesSpawnableState()
 
 	MovieScene->Modify();
 
-	TArray<FMovieSceneSpawnable*> Spawnables;
-
-	for (TViewModelPtr<IObjectBindingExtension> ObjectBindingNode : ViewModel->GetSelection()->Outliner.Filter<IObjectBindingExtension>())
-	{
-		FMovieSceneSpawnable* Spawnable = MovieScene->FindSpawnable(ObjectBindingNode->GetObjectGuid());
-		if (Spawnable)
-		{
-			Spawnables.Add(Spawnable);
-		}
-	}
-
-	FScopedSlowTask SlowTask(Spawnables.Num(), LOCTEXT("SaveSpawnableStateProgress", "Saving selected spawnables"));
+	FScopedSlowTask SlowTask(ViewModel->GetSelection()->Outliner.Num(), LOCTEXT("SaveSpawnableStateProgress", "Saving selected spawnables"));
 	SlowTask.MakeDialog(true);
 
-	TArray<AActor*> PossessedActors;
-	for (FMovieSceneSpawnable* Spawnable : Spawnables)
+	for (TViewModelPtr<IOutlinerExtension> OutlinerSelectionNode : ViewModel->GetSelection()->Outliner)
 	{
 		SlowTask.EnterProgressFrame();
-		
-		SpawnRegister->SaveDefaultSpawnableState(*Spawnable, ActiveTemplateIDs.Top(), *this);
-
-		if (GWarn->ReceivedUserCancel())
+		if (TViewModelPtr<IObjectBindingExtension> ObjectBindingNode = OutlinerSelectionNode.ImplicitCast())
 		{
-			break;
+			SpawnRegister->SaveDefaultSpawnableState(ObjectBindingNode->GetObjectGuid(), ActiveTemplateIDs.Top(), GetSharedPlaybackState());
+
+			if (GWarn->ReceivedUserCancel())
+			{
+				break;
+			}
 		}
 	}
-
 	NotifyMovieSceneDataChanged(EMovieSceneDataChangeType::MovieSceneStructureItemsChanged);
 }
 
@@ -8744,8 +8767,7 @@ void FSequencer::ConvertSelectedNodesToPossessables()
 
 	for (TViewModelPtr<IObjectBindingExtension> ObjectBindingNode : ViewModel->GetSelection()->Outliner.Filter<IObjectBindingExtension>())
 	{
-		FMovieSceneSpawnable* Spawnable = MovieScene->FindSpawnable(ObjectBindingNode->GetObjectGuid());
-		if (Spawnable && SpawnRegister->CanConvertSpawnableToPossessable(*Spawnable))
+		if (SpawnRegister->CanConvertToPossessable(ObjectBindingNode->GetObjectGuid(), GetFocusedTemplateID(), GetSharedPlaybackState()))
 		{
 			ObjectBindingNodes.Add(ObjectBindingNode.Get());
 		}
@@ -8764,19 +8786,15 @@ void FSequencer::ConvertSelectedNodesToPossessables()
 		{
 			SlowTask.EnterProgressFrame();
 
-			FMovieSceneSpawnable* Spawnable = MovieScene->FindSpawnable(ObjectBindingNode->GetObjectGuid());
-			if (Spawnable)
-			{
-				FMovieScenePossessable* Possessable = FSequencerUtilities::ConvertToPossessable(AsShared(), Spawnable->GetGuid());
+			FMovieScenePossessable* Possessable = FSequencerUtilities::ConvertToPossessable(AsShared(), ObjectBindingNode->GetObjectGuid());
 
-				ForceEvaluate();
+			ForceEvaluate();
 
-				for (TWeakObjectPtr<> WeakObject : FindBoundObjects(Possessable->GetGuid(), ActiveTemplateIDs.Top()))
+			for (TWeakObjectPtr<> WeakObject : FindBoundObjects(Possessable->GetGuid(), ActiveTemplateIDs.Top()))
+			{
+				if (AActor* PossessedActor = Cast<AActor>(WeakObject.Get()))
 				{
-					if (AActor* PossessedActor = Cast<AActor>(WeakObject.Get()))
-					{
-						PossessedActors.Add(PossessedActor);
-					}
+					PossessedActors.Add(PossessedActor);
 				}
 			}
 
@@ -10716,19 +10734,14 @@ void FSequencer::BindCommands()
 		}
 		return false;
 	};
-	SequencerCommandBindings->MapAction(
-		FSequencerCommands::Get().ConvertToSpawnable,
-		FExecuteAction::CreateSP(this, &FSequencer::ConvertSelectedNodesToSpawnables),
-		FCanExecuteAction::CreateLambda(CanConvertToSpawnables)
-	);
 
-	auto AreConvertableSpawnablesSelected = [this] {
+
+	auto CanSelectedNodesBeConvertedToPossessables = [this] {
 		UMovieScene* MovieScene = GetFocusedMovieSceneSequence()->GetMovieScene();
 
 		for (TViewModelPtr<IObjectBindingExtension> ObjectBindingNode : ViewModel->GetSelection()->Outliner.Filter<IObjectBindingExtension>())
 		{
-			FMovieSceneSpawnable* Spawnable = MovieScene->FindSpawnable(ObjectBindingNode->GetObjectGuid());
-			if (Spawnable && SpawnRegister->CanConvertSpawnableToPossessable(*Spawnable))
+			if (SpawnRegister->CanConvertToPossessable(ObjectBindingNode->GetObjectGuid(), GetFocusedTemplateID(), GetSharedPlaybackState()))
 			{
 				return true;
 			}
@@ -10739,16 +10752,15 @@ void FSequencer::BindCommands()
 	SequencerCommandBindings->MapAction(
 		FSequencerCommands::Get().ConvertToPossessable,
 		FExecuteAction::CreateSP(this, &FSequencer::ConvertSelectedNodesToPossessables),
-		FCanExecuteAction::CreateLambda(AreConvertableSpawnablesSelected)
+		FCanExecuteAction::CreateLambda(CanSelectedNodesBeConvertedToPossessables)
 	);
 
 	auto AreSpawnablesSelected = [this] {
-		UMovieScene* MovieScene = GetFocusedMovieSceneSequence()->GetMovieScene();
+		UMovieSceneSequence* MovieSceneSequence = GetFocusedMovieSceneSequence();
 
 		for (TViewModelPtr<IObjectBindingExtension> ObjectBindingNode : ViewModel->GetSelection()->Outliner.Filter<IObjectBindingExtension>())
 		{
-			FMovieSceneSpawnable* Spawnable = MovieScene->FindSpawnable(ObjectBindingNode->GetObjectGuid());
-			if (Spawnable)
+			if (MovieSceneHelpers::IsBoundToAnySpawnable(MovieSceneSequence, ObjectBindingNode->GetObjectGuid(), GetSharedPlaybackState()))
 			{
 				return true;
 			}
diff --git a/Engine/Source/Editor/Sequencer/Private/Sequencer.h b/Engine/Source/Editor/Sequencer/Private/Sequencer.h
--- a/Engine/Source/Editor/Sequencer/Private/Sequencer.h
+++ b/Engine/Source/Editor/Sequencer/Private/Sequencer.h
@@ -67,6 +67,8 @@ class UMovieSceneCopyableBinding;
 class UMovieSceneCompiledDataManager;
 class UMovieSceneCopyableTrack;
 class UMovieSceneNodeGroup;
+template <class TClass> class TSubclassOf;
+class UMovieSceneCustomBinding;
 
 struct FMovieSceneTimeController;
 struct FMovieSceneSequencePlaybackParams;
@@ -791,7 +793,7 @@ public:
 	virtual FOnCurveDisplayChanged& GetCurveDisplayChanged() override { return OnCurveDisplayChanged; }
 	virtual FOnSelectionChangedSections& GetSelectionChangedSections() override { return OnSelectionChangedSectionsDelegate; }
 	virtual FOnTreeViewChanged& OnTreeViewChanged() override { return OnTreeViewChangedDelegate; }
-	virtual FGuid CreateBinding(UObject& InObject, const FString& InName) override;
+	virtual FGuid CreateBinding(UObject& InObject, const UE::Sequencer::FCreateBindingParams& InParams) override;
 	virtual UObject* GetPlaybackContext() const override;
 	virtual IMovieScenePlaybackClient* GetPlaybackClient() override;
 	virtual TArray<UObject*> GetEventContexts() const override; 
@@ -852,6 +854,8 @@ public:
 	 */
 	TRange<FFrameNumber> GetRootTimeBounds() const;
 
+	virtual TArrayView<const TSubclassOf<UMovieSceneCustomBinding>> GetSupportedCustomBindingTypes() const override { return SupportedCustomBindingTypes; }
+
 protected:
 
 	// FCameraCutPlaybackCapability interface
@@ -1488,4 +1492,7 @@ private:
 	/** Cache of all bound cameras in the sequence hierarchy */
 	TMap<AActor*, FGuid> CachedCameraActors;
 	uint32 LastKnownStateSerial = 0;
+
+	/* Cached list of supported custom binding types*/
+	TArray<TSubclassOf<UMovieSceneCustomBinding>> SupportedCustomBindingTypes;
 };
diff --git a/Engine/Source/Editor/Sequencer/Private/SequencerCommands.cpp b/Engine/Source/Editor/Sequencer/Private/SequencerCommands.cpp
--- a/Engine/Source/Editor/Sequencer/Private/SequencerCommands.cpp
+++ b/Engine/Source/Editor/Sequencer/Private/SequencerCommands.cpp
@@ -166,7 +166,7 @@ void FSequencerCommands::RegisterCommands()
 	UI_COMMAND( PasteFromHistory, "Paste From History", "Paste from the sequencer clipboard history", EUserInterfaceActionType::Button, FInputChord(EModifierKey::Control | EModifierKey::Shift, EKeys::V) );
 
 	UI_COMMAND( ConvertToSpawnable, "Convert to Spawnable", "Make the specified possessed objects spawnable from sequencer. This will allow sequencer to have control over the lifetime of the object.", EUserInterfaceActionType::Button, FInputChord() );
-	UI_COMMAND( ConvertToPossessable, "Convert to Possessable", "Make the specified spawned objects possessed by sequencer.", EUserInterfaceActionType::Button, FInputChord() );
+	UI_COMMAND( ConvertToPossessable, "Convert to Possessable", "Make the specified object bindings into level objects possessed by sequencer.", EUserInterfaceActionType::Button, FInputChord() );
 
 	UI_COMMAND( SaveCurrentSpawnableState, "Save Default State", "Save the current state of this spawnable object as its default properties.", EUserInterfaceActionType::Button, FInputChord() );
 	UI_COMMAND( RestoreAnimatedState, "Restore Pre-Animated State", "Restore any objects that have been animated by sequencer back to their original state.", EUserInterfaceActionType::Button, FInputChord(EModifierKey::Control, EKeys::R) );
diff --git a/Engine/Source/Editor/Sequencer/Private/SequencerSpawnableColumn.cpp b/Engine/Source/Editor/Sequencer/Private/SequencerSpawnableColumn.cpp
--- a/Engine/Source/Editor/Sequencer/Private/SequencerSpawnableColumn.cpp
+++ b/Engine/Source/Editor/Sequencer/Private/SequencerSpawnableColumn.cpp
@@ -61,19 +61,23 @@ const TSharedRef< SWidget > FSequencerSpawnableColumn::ConstructRowWidget(FScene
 const FSlateBrush* FSequencerSpawnableColumn::GetSpawnableIcon( FSceneOutlinerTreeItemRef TreeItem ) const
 {
 	static const FSlateBrush* SpawnedBrush  = FAppStyle::Get().GetBrush("Sequencer.SpawnableIconOverlay");
+	static const FSlateBrush* ReplaceableBrush = FAppStyle::Get().GetBrush("Sequencer.ReplaceableIconOverlay");
 	static const FSlateBrush* NoSpawnedBrush  = FAppStyle::Get().GetBrush("NoBrush");
 
 	bool bIsSpawned = false;
+	bool bIsPreviewActor = false;
 	{
 		static const FName SequencerActorTag(TEXT("SequencerActor"));
+		static const FName SequencerPreviewActorTag(TEXT("SequencerPreviewActor"));
 		if (const FActorTreeItem* ActorItem = (*TreeItem).CastTo<FActorTreeItem>())
 		{
 			AActor* Actor = ActorItem->Actor.Get();
 			bIsSpawned = Actor && Actor->ActorHasTag(SequencerActorTag);
+			bIsPreviewActor = Actor && Actor->ActorHasTag(SequencerPreviewActorTag);
 		}
 	}
 
-	return bIsSpawned ? SpawnedBrush : NoSpawnedBrush;
+	return bIsPreviewActor ? ReplaceableBrush : bIsSpawned ? SpawnedBrush : NoSpawnedBrush;
 }
 
 }// End Sequencer namespace
diff --git a/Engine/Source/Editor/Sequencer/Private/SequencerUtilities.cpp b/Engine/Source/Editor/Sequencer/Private/SequencerUtilities.cpp
--- a/Engine/Source/Editor/Sequencer/Private/SequencerUtilities.cpp
+++ b/Engine/Source/Editor/Sequencer/Private/SequencerUtilities.cpp
@@ -62,6 +62,10 @@
 #include "MVVM/Views/ViewUtilities.h"
 #include "Framework/Notifications/NotificationManager.h"
 #include "Widgets/Notifications/SNotificationList.h"
+#include "Bindings/MovieSceneCustomBinding.h"
+#include "Bindings/MovieSceneSpawnableBinding.h"
+#include "Bindings/MovieSceneReplaceableBinding.h"
+#include "ActorFactories/ActorFactory.h"
 
 #define LOCTEXT_NAMESPACE "FSequencerUtilities"
 
@@ -685,19 +689,13 @@ FGuid FSequencerUtilities::MakeNewSpawnable(TSharedRef<ISequencer> Sequencer, UO
 		return FGuid();
 	}
 
-	FMovieSceneSpawnable* Spawnable = MovieScene->FindSpawnable(NewGuid);
-	if (!Spawnable)
-	{
-		return FGuid();
-	}
-
 	// Spawn the object so we can position it correctly, it's going to get spawned anyway since things default to spawned.
 	UObject* SpawnedObject = Sequencer->GetSpawnRegister().SpawnObject(NewGuid, *MovieScene, Sequencer->GetFocusedTemplateID(), Sequencer.Get());
 
 	if (bSetupDefaults)
 	{
 		FTransformData TransformData;
-		Sequencer->GetSpawnRegister().SetupDefaultsForSpawnable(SpawnedObject, Spawnable->GetGuid(), TransformData, Sequencer, Sequencer->GetSequencerSettings());
+		Sequencer->GetSpawnRegister().SetupDefaultsForSpawnable(SpawnedObject, NewGuid, TransformData, Sequencer, Sequencer->GetSequencerSettings());
 	}
 
 	if (ACameraActor* NewCamera = Cast<ACameraActor>(SpawnedObject))
@@ -762,13 +760,11 @@ FGuid FSequencerUtilities::CreateCamera(TSharedRef<ISequencer> Sequencer, const
 	{
 		FString NewName = MovieSceneHelpers::MakeUniqueSpawnableName(MovieScene, FName::NameToDisplayString(ACineCameraActor::StaticClass()->GetFName().ToString(), false));
 
-		CameraGuid = MakeNewSpawnable(Sequencer, *OutActor);
-		Spawnable = MovieScene->FindSpawnable(CameraGuid);
-
-		if (ensure(Spawnable))
-		{
-			Spawnable->SetName(NewName);
-		}
+		UE::Sequencer::FCreateBindingParams CreateBindingParams;
+		CreateBindingParams.BindingNameOverride = NewName;
+		CreateBindingParams.bSpawnable = true;
+		CameraGuid = CreateBinding(Sequencer, *OutActor, CreateBindingParams);
+		ensure(CameraGuid.IsValid());
 
 		// Destroy the old actor
 		World->EditorDestroyActor(OutActor, false);
@@ -835,7 +831,13 @@ FGuid FSequencerUtilities::CreateCameraWithRig(TSharedRef<ISequencer> Sequencer,
 	// Create a cine camera actor
 	UWorld* PlaybackContext = Sequencer->GetPlaybackContext()->GetWorld();
 	OutActor = PlaybackContext->SpawnActor<ACineCameraActor>();
-	CameraGuid = CreateBinding(Sequencer, *OutActor);
+
+	FString NewCameraName = MovieSceneHelpers::MakeUniqueSpawnableName(MovieScene, FName::NameToDisplayString(ACineCameraActor::StaticClass()->GetFName().ToString(), false));
+	UE::Sequencer::FCreateBindingParams CreateBindingParams;
+	CreateBindingParams.BindingNameOverride = NewCameraName;
+	CreateBindingParams.bSpawnable = bSpawnable;
+
+	CameraGuid = CreateBinding(Sequencer, *OutActor, CreateBindingParams);
 
 	if (RailActor)
 	{
@@ -846,12 +848,7 @@ FGuid FSequencerUtilities::CreateCameraWithRig(TSharedRef<ISequencer> Sequencer,
 
 	if (bSpawnable)
 	{
-		FString NewCameraName = MovieSceneHelpers::MakeUniqueSpawnableName(MovieScene, FName::NameToDisplayString(ACineCameraActor::StaticClass()->GetFName().ToString(), false));
-
-		FMovieSceneSpawnable* Spawnable = ConvertToSpawnable(Sequencer, CameraGuid)[0];
-		Spawnable->SetName(NewCameraName);
-
-		for (TWeakObjectPtr<> WeakObject : Sequencer->FindBoundObjects(Spawnable->GetGuid(), Sequencer->GetFocusedTemplateID()))
+		for (TWeakObjectPtr<> WeakObject : Sequencer->FindBoundObjects(CameraGuid, Sequencer->GetFocusedTemplateID()))
 		{
 			OutActor = Cast<ACineCameraActor>(WeakObject.Get());
 			if (OutActor)
@@ -862,8 +859,6 @@ FGuid FSequencerUtilities::CreateCameraWithRig(TSharedRef<ISequencer> Sequencer,
 
 		OutActor->SetActorLabel(NewCameraName, false);
 
-		CameraGuid = Spawnable->GetGuid();
-
 		// Create an attach track
 		UMovieScene3DAttachTrack* AttachTrack = Cast<UMovieScene3DAttachTrack>(MovieScene->AddTrack(UMovieScene3DAttachTrack::StaticClass(), CameraGuid));
 
@@ -1076,7 +1071,7 @@ TArray<FMovieSceneSpawnable*> FSequencerUtilities::ConvertToSpawnable(TSharedRef
 	return CreatedSpawnables;
 }
 
-FMovieScenePossessable* FSequencerUtilities::ConvertToPossessable(TSharedRef<ISequencer> Sequencer, FGuid SpawnableGuid)
+FMovieScenePossessable* FSequencerUtilities::ConvertToPossessable(TSharedRef<ISequencer> Sequencer, FGuid BindingGuid, int32 BindingIndex/*=0*/)
 {
 	FMovieScenePossessable* CreatedPossessable = nullptr;
 
@@ -1098,94 +1093,125 @@ FMovieScenePossessable* FSequencerUtilities::ConvertToPossessable(TSharedRef<ISe
 		return CreatedPossessable;
 	}
 
-	// Find the object in the environment
-	FMovieSceneSpawnable* Spawnable = MovieScene->FindSpawnable(SpawnableGuid);
-	if (!Spawnable || !Spawnable->GetObjectTemplate())
+	TArrayView<TWeakObjectPtr<>> FoundObjects = Sequencer->FindBoundObjects(BindingGuid, Sequencer->GetFocusedTemplateID());
+
+	UObject* ObjectToConvert = nullptr;
+
+	if (FoundObjects.IsValidIndex(BindingIndex))
 	{
-		return CreatedPossessable;
+		ObjectToConvert = FoundObjects[BindingIndex].Get();
+	}
+
+	// If we have an old-style spawnable, use the template as the object to convert instead.
+	bool bConvertFromSpawnable = MovieSceneHelpers::IsBoundToSpawnable(Sequence, BindingGuid, Sequencer->GetSharedPlaybackState(), BindingIndex);
+	if (bConvertFromSpawnable)
+	{
+		ObjectToConvert = MovieSceneHelpers::GetObjectTemplate(Sequence, BindingGuid, Sequencer->GetSharedPlaybackState(), BindingIndex);
 	}
 
-	AActor* SpawnableActorTemplate = Cast<AActor>(Spawnable->GetObjectTemplate());
-	if (!SpawnableActorTemplate)
+	if (!ObjectToConvert)
 	{
 		return CreatedPossessable;
 	}
 
+	AActor* SpawnableActorTemplate = Cast<AActor>(ObjectToConvert);
+
 	TMap<TWeakObjectPtr<AActor>, FTransform> AttachedChildTransforms;
-	FTransform DefaultTransform = SpawnableActorTemplate->GetActorTransform();
-	for (TWeakObjectPtr<> RuntimeObject : Sequencer->FindBoundObjects(SpawnableGuid, Sequencer->GetFocusedTemplateID()))
+	FTransform DefaultTransform = SpawnableActorTemplate ? SpawnableActorTemplate->GetActorTransform() : FTransform();
+	// Prefer the transform at the current time over the spawnable actor template's transform because that's most likely 0. 
+	// This makes it so that the object will return to the current position on restore state.
+	AActor* Actor = Cast<AActor>(ObjectToConvert);
+	if (Actor)
 	{
-		// Prefer the transform at the current time over the spawnable actor template's transform because that's most likely 0. 
-		// This makes it so that the object will return to the current position on restore state.
-		AActor* Actor = Cast<AActor>(RuntimeObject.Get());
-		if (Actor)
+		if (Actor->GetRootComponent())
 		{
-			if (Actor->GetRootComponent())
-			{
-				DefaultTransform = Actor->GetRootComponent()->GetRelativeTransform();
-			}
+			DefaultTransform = Actor->GetRootComponent()->GetRelativeTransform();
+		}
 
-			// Removing a parent will compensate the children at their world transform. We don't want that since we'll be replacing that parent right away.
-			// To negate that, we store the relative transform of these children and reset it after the parent is replaced with the new possessable.
-			TArray<AActor*> AttachedActors;
-			Actor->GetAttachedActors(AttachedActors);
-			for (AActor* ChildActor : AttachedActors)
+		// Removing a parent will compensate the children at their world transform. We don't want that since we'll be replacing that parent right away.
+		// To negate that, we store the relative transform of these children and reset it after the parent is replaced with the new possessable.
+		TArray<AActor*> AttachedActors;
+		Actor->GetAttachedActors(AttachedActors);
+		for (AActor* ChildActor : AttachedActors)
+		{
+			if (ChildActor && ChildActor->GetRootComponent())
 			{
-				if (ChildActor && ChildActor->GetRootComponent())
+				// Only do this for child actors that Sequencer is controlling
+				FGuid ExistingID = Sequencer->FindObjectId(*ChildActor, Sequencer->GetFocusedTemplateID());
+				if (ExistingID.IsValid())
 				{
-					// Only do this for child actors that Sequencer is controlling
-					FGuid ExistingID = Sequencer->FindObjectId(*ChildActor, Sequencer->GetFocusedTemplateID());
-					if (ExistingID.IsValid())
-					{
-						AttachedChildTransforms.Add(ChildActor);
-						AttachedChildTransforms[ChildActor] = ChildActor->GetRootComponent()->GetRelativeTransform();
-					}
+					AttachedChildTransforms.Add(ChildActor);
+					AttachedChildTransforms[ChildActor] = ChildActor->GetRootComponent()->GetRelativeTransform();
 				}
 			}
-			break;
 		}
 	}
 
-	Sequence->Modify();
-	MovieScene->Modify();
-
-	// Delete the spawn track
-	UMovieSceneSpawnTrack* SpawnTrack = Cast<UMovieSceneSpawnTrack>(MovieScene->FindTrack(UMovieSceneSpawnTrack::StaticClass(), SpawnableGuid, NAME_None));
-	if (SpawnTrack)
+	FMovieSceneSpawnable* Spawnable = MovieScene->FindSpawnable(BindingGuid);
+	FMovieScenePossessable* ExistingPossessable = MovieScene->FindPossessable(BindingGuid);
+	// TODO: How to convert to possessable of non-actor type? Presumably we need to generalize the 'creation' step here for now.
+	UObject* PossessedObject = nullptr;
+	if (Actor)
 	{
-		MovieScene->RemoveTrack(*SpawnTrack);
+		FActorSpawnParameters SpawnInfo;
+		SpawnInfo.bDeferConstruction = true;
+		SpawnInfo.Template = SpawnableActorTemplate;
+
+		UWorld* PlaybackContext = Sequencer->GetPlaybackContext()->GetWorld();
+		AActor* PossessedActor = PlaybackContext->SpawnActor(ObjectToConvert->GetClass(), &DefaultTransform, SpawnInfo);
+
+		if (!PossessedActor)
+		{
+			return nullptr;
+		}
+
+		PossessedActor->SetActorLabel(Spawnable ? Spawnable->GetName() : ExistingPossessable ? ExistingPossessable->GetName() : ObjectToConvert->GetName());
+
+		const bool bIsDefaultTransform = true;
+		PossessedActor->FinishSpawning(DefaultTransform, bIsDefaultTransform);
+
+		// The transform needs to be set again for deferred construction and dynamic root components. Until the fix for: UE-67537
+		PossessedActor->SetActorTransform(DefaultTransform);
+
+		PossessedObject = PossessedActor;
 	}
 
-	FActorSpawnParameters SpawnInfo;
-	SpawnInfo.bDeferConstruction = true;
-	SpawnInfo.Template = SpawnableActorTemplate;
+	if (!PossessedObject)
+	{
+		return CreatedPossessable;
+	}
 
-	UWorld* PlaybackContext = Sequencer->GetPlaybackContext()->GetWorld();
-	AActor* PossessedActor = PlaybackContext->SpawnActor(Spawnable->GetObjectTemplate()->GetClass(), &DefaultTransform, SpawnInfo);
+	Sequence->Modify();
+	MovieScene->Modify();
 
-	if (!PossessedActor)
+	// If we're converting from an old-style spawnable, we'll need to remove the spawn track
+	if (Spawnable)
 	{
-		return nullptr;
+		// Delete the spawn track
+		UMovieSceneSpawnTrack* SpawnTrack = Cast<UMovieSceneSpawnTrack>(MovieScene->FindTrack(UMovieSceneSpawnTrack::StaticClass(), BindingGuid, NAME_None));
+		if (SpawnTrack)
+		{
+			MovieScene->RemoveTrack(*SpawnTrack);
+		}
 	}
 
-	PossessedActor->SetActorLabel(Spawnable->GetName());
 
-	const bool bIsDefaultTransform = true;
-	PossessedActor->FinishSpawning(DefaultTransform, bIsDefaultTransform);
+	UE::Sequencer::FCreateBindingParams CreateBindingParams;
+	CreateBindingParams.ReplacementGuid = BindingGuid;
+	CreateBindingParams.BindingIndex = BindingIndex;
+	CreateBindingParams.bAllowCustomBinding = false;
 
-	// The transform needs to be set again for deferred construction and dynamic root components. Until the fix for: UE-67537
-	PossessedActor->SetActorTransform(DefaultTransform);
+	// Create or replace the binding
+	FGuid NewPossessableGuid = CreateBinding(Sequencer, *PossessedObject, CreateBindingParams);
 
-	const FGuid NewPossessableGuid = CreateBinding(Sequencer, *PossessedActor);
-	const FGuid OldSpawnableGuid = Spawnable->GetGuid();
 
 	FMovieScenePossessable* Possessable = MovieScene->FindPossessable(NewPossessableGuid);
-	if (Possessable)
+	if (Spawnable)
 	{
 		// Remap all the spawnable's tracks and child bindings onto the new possessable
-		MovieScene->MoveBindingContents(OldSpawnableGuid, NewPossessableGuid);
+		MovieScene->MoveBindingContents(BindingGuid, NewPossessableGuid);
 
-		FMovieSceneBinding* SpawnableBinding = MovieScene->FindBinding(OldSpawnableGuid);
+		FMovieSceneBinding* SpawnableBinding = MovieScene->FindBinding(BindingGuid);
 		check(SpawnableBinding);
 
 		for (UMovieSceneFolder* Folder : MovieScene->GetRootFolders())
@@ -1195,24 +1221,32 @@ FMovieScenePossessable* FSequencerUtilities::ConvertToPossessable(TSharedRef<ISe
 				break;
 			}
 		}
-
 		int32 SortingOrder = SpawnableBinding->GetSortingOrder();
 
 		// Remove the spawnable and all it's sub tracks
-		if (MovieScene->RemoveSpawnable(OldSpawnableGuid))
+		if (MovieScene->RemoveSpawnable(BindingGuid))
 		{
-			Sequencer->GetSpawnRegister().DestroySpawnedObject(OldSpawnableGuid, Sequencer->GetFocusedTemplateID(), Sequencer.Get());
+			UpdateBindingIDs(Sequencer, BindingGuid, NewPossessableGuid);
 
 			FMovieSceneBinding* PossessableBinding = MovieScene->FindBinding(NewPossessableGuid);
 			check(PossessableBinding);
 
 			PossessableBinding->SetSortingOrder(SortingOrder);
 		}
+	}
+
+	// If we previously had an old-style spawnable or a spawnable custom binding, destroy the old spawned object
+	if (bConvertFromSpawnable)
+	{
+		Sequencer->GetSpawnRegister().DestroySpawnedObject(BindingGuid, Sequencer->GetFocusedTemplateID(), Sequencer->GetSharedPlaybackState(), BindingIndex);
+	}
 
+	if (AActor* PossessedActor = Cast<AActor>(PossessedObject))
+	{
 		static const FName SequencerActorTag(TEXT("SequencerActor"));
+		static const FName SequencerPreviewActorTag(TEXT("SequencerPreviewActor"));
 		PossessedActor->Tags.Remove(SequencerActorTag);
-
-		UpdateBindingIDs(Sequencer, OldSpawnableGuid, NewPossessableGuid);
+		PossessedActor->Tags.Remove(SequencerPreviewActorTag);
 
 		GEditor->SelectActor(PossessedActor, false, true);
 
@@ -1226,11 +1260,174 @@ FMovieScenePossessable* FSequencerUtilities::ConvertToPossessable(TSharedRef<ISe
 				}
 			}
 		}
+	}
+
+	Sequencer->ForceEvaluate();
+
+	return Possessable;
+}
+
+FMovieScenePossessable* FSequencerUtilities::ConvertToCustomBinding(TSharedRef<ISequencer> Sequencer, FGuid BindingGuid, TSubclassOf<UMovieSceneCustomBinding> CustomBindingType, int32 BindingIndex/*=0*/)
+{
+	FMovieScenePossessable* CreatedPossessable = nullptr;
+
+	UMovieSceneSequence* Sequence = Sequencer->GetFocusedMovieSceneSequence();
+	if (!Sequence)
+	{
+		return CreatedPossessable;
+	}
+
+	UMovieScene* MovieScene = Sequence->GetMovieScene();
+	if (!MovieScene)
+	{
+		return CreatedPossessable;
+	}
+
+	if (MovieScene->IsReadOnly())
+	{
+		ShowReadOnlyError();
+		return CreatedPossessable;
+	}
+
+	const FMovieSceneBindingReferences* BindingReferences = Sequence->GetBindingReferences();
+	if (!BindingReferences)
+	{
+		// Not supported with this sequence type- show an error?
+		return CreatedPossessable;
+	}
+
+	if (!CustomBindingType)
+	{
+		return CreatedPossessable;
+	}
+
+	TArrayView<TWeakObjectPtr<>> FoundObjects = Sequencer->FindBoundObjects(BindingGuid, Sequencer->GetFocusedTemplateID());
+
+	UObject* ObjectToConvert = nullptr;
+
+	if (FoundObjects.IsValidIndex(BindingIndex))
+	{
+		ObjectToConvert = FoundObjects[BindingIndex].Get();
+	}
+
+	bool bConvertFromSpawnable = false;
+
+	// If we have an old-style spawnable, use the template as the object to convert instead.
+	FMovieSceneSpawnable* Spawnable = MovieScene->FindSpawnable(BindingGuid);
+	if (Spawnable)
+	{
+		ObjectToConvert = Spawnable->GetObjectTemplate();
+		bConvertFromSpawnable = true;
+	}
+	else if (const FMovieSceneBindingReference* Reference = BindingReferences->GetReference(BindingGuid, BindingIndex))
+	{
+		if (const UMovieSceneCustomBinding* CustomBinding = Reference->CustomBinding)
+		{
+			bConvertFromSpawnable = CustomBinding->WillSpawnObject(Sequencer->GetSharedPlaybackState());
+		}
+	}
+
+	bool bConvertFromPossessable = !bConvertFromSpawnable
+		&& !Algo::AnyOf(BindingReferences->GetReferences(BindingGuid), [](const FMovieSceneBindingReference& Reference) { return Reference.CustomBinding; });
+
+	if (!ObjectToConvert)
+	{
+		return CreatedPossessable;
+	}
+
+	// Create the new custom binding and then we'll slot it in.
+	UMovieSceneCustomBinding* NewCustomBinding = CustomBindingType->GetDefaultObject<UMovieSceneCustomBinding>()->CreateNewCustomBinding(ObjectToConvert, *MovieScene);
+	if (!NewCustomBinding)
+	{
+		return CreatedPossessable;
+	}
+
+	Sequence->Modify();
+	MovieScene->Modify();
+
+
+	// If we're converting from an old-style spawnable, we'll need to do some extra work and also create a new FMovieScenePossessable*.
+	if (Spawnable)
+	{
+		// Delete the spawn track
+		UMovieSceneSpawnTrack* SpawnTrack = Cast<UMovieSceneSpawnTrack>(MovieScene->FindTrack(UMovieSceneSpawnTrack::StaticClass(), BindingGuid, NAME_None));
+		if (SpawnTrack)
+		{
+			MovieScene->RemoveTrack(*SpawnTrack);
+		}
+	}
+
+	UE::Sequencer::FCreateBindingParams CreateBindingParams;
+	CreateBindingParams.ReplacementGuid = BindingGuid;
+	CreateBindingParams.BindingIndex = BindingIndex;
+	CreateBindingParams.BindingNameOverride = NewCustomBinding->GetDesiredBindingName();
+	CreateBindingParams.CustomBinding = NewCustomBinding;
+
+	// Create or replace the binding
+	FGuid NewPossessableGuid = CreateBinding(Sequencer, *ObjectToConvert, CreateBindingParams);
+
+	CreatedPossessable = MovieScene->FindPossessable(NewPossessableGuid);
+
+	// If we previously had an old-style spawnable, we need to move over bindings
+	if (Spawnable)
+	{
+		// Remap all the spawnable's tracks and child bindings onto the new possessable
+		MovieScene->MoveBindingContents(BindingGuid, NewPossessableGuid);
+
+		FMovieSceneBinding* SpawnableBinding = MovieScene->FindBinding(BindingGuid);
+		check(SpawnableBinding);
+
+		for (UMovieSceneFolder* Folder : MovieScene->GetRootFolders())
+		{
+			if (UpdateFolderBindingID(Folder, Spawnable->GetGuid(), NewPossessableGuid))
+			{
+				break;
+			}
+		}
+
+		int32 SortingOrder = SpawnableBinding->GetSortingOrder();
+
+		// Remove the spawnable and all its' sub tracks
+		if (MovieScene->RemoveSpawnable(BindingGuid))
+		{
+			FMovieSceneBinding* PossessableBinding = MovieScene->FindBinding(NewPossessableGuid);
+			check(PossessableBinding);
+
+			PossessableBinding->SetSortingOrder(SortingOrder);
+		}
+
+		UpdateBindingIDs(Sequencer, BindingGuid, NewPossessableGuid);
 
 		Sequencer->ForceEvaluate();
 	}
 
-	return Possessable;
+	if (bConvertFromSpawnable)
+	{
+		Sequencer->GetSpawnRegister().DestroySpawnedObject(BindingGuid, Sequencer->GetFocusedTemplateID(), Sequencer->GetSharedPlaybackState(), BindingIndex);
+	}
+	else if (bConvertFromPossessable)
+	{
+		// We have an old possessable to destroy
+		TOptional<FTransformData> TransformData;
+		Sequencer->GetSpawnRegister().HandleConvertPossessableToSpawnable(ObjectToConvert, *Sequencer, TransformData);
+	}
+
+	// If this is a new spawnable or replaceable binding, we need to set up some defaults
+	if (NewCustomBinding->WillSpawnObject(Sequencer->GetSharedPlaybackState()))
+	{
+		// Spawn the object so we can position it correctly, it's going to get spawned anyway since things default to spawned.
+		UObject* SpawnedObject = Sequencer->GetSpawnRegister().SpawnObject(NewPossessableGuid, *MovieScene, Sequencer->GetFocusedTemplateID(), Sequencer->GetSharedPlaybackState(), BindingIndex);
+
+		// Allow the binding to set up any necessary defaults
+		NewCustomBinding->SetupDefaults(SpawnedObject, NewPossessableGuid, *MovieScene, Sequencer->GetSharedPlaybackState());
+
+		FTransformData TransformData;
+		Sequencer->GetSpawnRegister().SetupDefaultsForSpawnable(SpawnedObject, NewPossessableGuid, TransformData, Sequencer, Sequencer->GetSequencerSettings());
+	}
+
+	//Sequencer->State.Invalidate(NewPossessableGuid, Sequencer->GetFocusedTemplateID());
+
+	return CreatedPossessable;
 }
 
 void ExportObjectsToText(const TArray<UObject*>& ObjectsToExport, FString& ExportedText)
@@ -1862,7 +2059,7 @@ bool FSequencerUtilities::CanPasteSections(const FString& TextToImport)
  *
  */
 
-void ExportObjectBindingsToText(const TArray<UMovieSceneCopyableBinding*>& ObjectsToExport, FString& ExportedText)
+void ExportObjectBindingsToText(const TArray<UMovieSceneCopyableBinding*>& ObjectsToExport, FString& ExportedText, TSharedRef<UE::MovieScene::FSharedPlaybackState> SharedPlaybackState)
 {
 	// Clear the mark state for saving.
 	UnMarkAllObjects(EObjectMark(OBJECTMARK_TagExp | OBJECTMARK_TagImp));
@@ -1884,19 +2081,56 @@ void ExportObjectBindingsToText(const TArray<UMovieSceneCopyableBinding*>& Objec
 		// serialized manually into the archive, as the auto-serialization will only store a reference (to a privately owned object) which creates issues on deserialization. Attempting 
 		// to deserialize these private objects throws a superflous error in the console that makes it look like things went wrong when they're actually OK and expected.
 		TArray<UMovieSceneTrack*> OldTracks = ObjectToExport->Binding.StealTracks(nullptr);
-		UObject* OldSpawnableTemplate = ObjectToExport->Spawnable.GetObjectTemplate();
-		ObjectToExport->Spawnable.SetObjectTemplate(nullptr);
+
+		TArray<UObject*, TInlineAllocator<1>> OldObjectTemplates;
+		if (ObjectToExport->Spawnable.GetGuid().IsValid())
+		{
+			OldObjectTemplates.Add(ObjectToExport->Spawnable.GetObjectTemplate());
+			ObjectToExport->Spawnable.SetObjectTemplate(nullptr);
+		}
+		else
+		{
+			for (UMovieSceneCustomBinding* CustomBinding : ObjectToExport->CustomBindings)
+			{
+				if (UMovieSceneSpawnableBindingBase* SpawnableBinding = CustomBinding->AsSpawnable(SharedPlaybackState))
+				{
+					if (SpawnableBinding->SupportsObjectTemplates())
+					{
+						OldObjectTemplates.Add(SpawnableBinding->GetObjectTemplate());
+						SpawnableBinding->SetObjectTemplate(nullptr);
+					}
+				}
+			}
+		}
 
 		UExporter::ExportToOutputDevice(&Context, ObjectToExport, nullptr, Archive, TEXT("copy"), 0, PPF_ExportsNotFullyQualified | PPF_Copy | PPF_Delimited, false, ThisOuter);
 
 		// Restore the references (as we don't want to modify the original in the event of a copy operation!)
 		ObjectToExport->Binding.SetTracks(MoveTemp(OldTracks), nullptr);
-		ObjectToExport->Spawnable.SetObjectTemplate(OldSpawnableTemplate);
 
-		// We manually export the object template for the same private-ownership reason as above. Templates need to be re-created anyways as each Spawnable contains its own copy of the template.
-		if (ObjectToExport->SpawnableObjectTemplate)
+		int32 ObjectTemplateIndex = 0;
+		if (ObjectToExport->Spawnable.GetGuid().IsValid())
+		{
+			ObjectToExport->Spawnable.SetObjectTemplate(OldObjectTemplates[ObjectTemplateIndex++]);
+		}
+		else
+		{
+			for (UMovieSceneCustomBinding* CustomBinding : ObjectToExport->CustomBindings)
+			{
+				if (UMovieSceneSpawnableBindingBase* SpawnableBinding = CustomBinding->AsSpawnable(SharedPlaybackState))
+				{
+					if (SpawnableBinding->SupportsObjectTemplates())
+					{
+						SpawnableBinding->SetObjectTemplate(OldObjectTemplates[ObjectTemplateIndex++]);
+					}
+				}
+			}
+		}
+
+		// We manually export the object templates for the same private-ownership reason as above. Templates need to be re-created anyways as each Spawnable contains its own copy of the template.
+		for (UObject* ObjectTemplate : ObjectToExport->SpawnableObjectTemplates)
 		{
-			UExporter::ExportToOutputDevice(&Context, ObjectToExport->SpawnableObjectTemplate, nullptr, Archive, TEXT("copy"), 0, PPF_ExportsNotFullyQualified | PPF_Copy | PPF_Delimited);
+			UExporter::ExportToOutputDevice(&Context, ObjectTemplate, nullptr, Archive, TEXT("copy"), 0, PPF_ExportsNotFullyQualified | PPF_Copy | PPF_Delimited);
 		}
 	}
 
@@ -1930,6 +2164,26 @@ void FSequencerUtilities::CopyBindings(TSharedRef<ISequencer> Sequencer, const T
 					CopyableBinding->BoundObjectNames.Add(Actor->GetPathName());
 				}
 			}
+
+			// Store any custom bindings
+			if (FMovieSceneBindingReferences* BindingReferences = ObjectBinding.Sequence->GetBindingReferences())
+			{
+				for (const FMovieSceneBindingReference& BindingReference : BindingReferences->GetReferences(ObjectBinding.BindingID))
+				{
+					if (UMovieSceneCustomBinding* CustomBinding = BindingReference.CustomBinding)
+					{
+						if (UMovieSceneSpawnableBindingBase* SpawnableBinding = CustomBinding->AsSpawnable(Sequencer->GetSharedPlaybackState()))
+						{
+							if (SpawnableBinding->SupportsObjectTemplates())
+							{
+								// We manually serialize the spawnable object template so that it's not a reference to a privately owned object. Spawnables all have unique copies of their template objects anyways.
+								// Object Templates are re-created on paste (based on these templates) with the correct ownership set up.
+								CopyableBinding->SpawnableObjectTemplates.Add(SpawnableBinding->GetObjectTemplate());
+							}
+						}
+					}
+				}
+			}
 		}
 		else
 		{
@@ -1940,7 +2194,7 @@ void FSequencerUtilities::CopyBindings(TSharedRef<ISequencer> Sequencer, const T
 
 				// We manually serialize the spawnable object template so that it's not a reference to a privately owned object. Spawnables all have unique copies of their template objects anyways.
 				// Object Templates are re-created on paste (based on these templates) with the correct ownership set up.
-				CopyableBinding->SpawnableObjectTemplate = Spawnable->GetObjectTemplate();
+				CopyableBinding->SpawnableObjectTemplates.Add(Spawnable->GetObjectTemplate());
 			}
 		}
 
@@ -1978,7 +2232,7 @@ void FSequencerUtilities::CopyBindings(TSharedRef<ISequencer> Sequencer, const T
 		}
 	}
 
-	ExportObjectBindingsToText(Objects, /*out*/ ExportedText);
+	ExportObjectBindingsToText(Objects, /*out*/ ExportedText, Sequencer->GetSharedPlaybackState());
 }
 
 class FObjectBindingTextFactory : public FCustomizableTextObjectFactory
@@ -2041,10 +2295,23 @@ void ImportObjectBindingsFromText(ISequencer& InSequencer, const FString& TextTo
 	int32 SpawnableObjectTemplateIndex = 0;
 	for (auto ImportedObject : ImportedObjects)
 	{
-		if (ImportedObject->Spawnable.GetGuid().IsValid() && SpawnableObjectTemplateIndex < Factory.NewSpawnableObjectTemplates.Num())
+		if (ImportedObject->Spawnable.GetGuid().IsValid())
 		{
 			// This Spawnable Object Template is owned by our transient package, so you'll need to change the owner if you want to keep it later.
-			ImportedObject->SpawnableObjectTemplate = Factory.NewSpawnableObjectTemplates[SpawnableObjectTemplateIndex++];
+			ImportedObject->SpawnableObjectTemplates.Add(Factory.NewSpawnableObjectTemplates[SpawnableObjectTemplateIndex++]);
+		}
+		else if (ImportedObject->CustomBindings.Num() > 0 && SpawnableObjectTemplateIndex < Factory.NewSpawnableObjectTemplates.Num())
+		{
+			for (UMovieSceneCustomBinding* CustomBinding : ImportedObject->CustomBindings)
+			{
+				if (UMovieSceneSpawnableBindingBase* SpawnableBinding = CustomBinding->AsSpawnable(InSequencer.GetSharedPlaybackState()))
+				{
+					if (SpawnableBinding->SupportsObjectTemplates())
+					{
+						ImportedObject->SpawnableObjectTemplates.Add(Factory.NewSpawnableObjectTemplates[SpawnableObjectTemplateIndex++]);
+					}
+				}
+			}
 		}
 	}
 
@@ -2052,54 +2319,363 @@ void ImportObjectBindingsFromText(ISequencer& InSequencer, const FString& TextTo
 	TempPackage->RemoveFromRoot();
 }
 
-bool FSequencerUtilities::PasteBindings(const FString& TextToImport, TSharedRef<ISequencer> Sequencer, FMovieScenePasteBindingsParams PasteBindingsParams, TArray<FMovieSceneBindingProxy>& OutBindings, TArray<FNotificationInfo>& OutErrors)
+FGuid TryCreateCustomBinding(TSharedRef<ISequencer> Sequencer, UObject* CustomBindingObject, FMovieSceneBindingReferences* BindingReferences, const UE::Sequencer::FCreateBindingParams& InParams, UMovieScene* OwnerMovieScene, bool bSpawnable, bool bReplaceable)
 {
-	UMovieSceneSequence* Sequence = Sequencer->GetFocusedMovieSceneSequence();
-	if (!Sequence)
+	UMovieSceneCustomBinding * NewCustomBinding = nullptr;
+	if (InParams.CustomBinding)
 	{
-		return false;
+		// We've been provided a custom binding pre-created. Ensure it supports the object given
+		if (CustomBindingObject == nullptr || InParams.CustomBinding->SupportsBindingCreationFromObject(CustomBindingObject))
+		{
+			NewCustomBinding = InParams.CustomBinding;
+		}
 	}
-
-	UMovieScene* MovieScene = Sequence->GetMovieScene();
-	if (!MovieScene)
+	else
 	{
-		return false;
+		TArrayView<const TSubclassOf<UMovieSceneCustomBinding>> PrioritySortedCustomBindingTypes = Sequencer->GetSupportedCustomBindingTypes();
+		for (const TSubclassOf<UMovieSceneCustomBinding>& CustomBindingType : PrioritySortedCustomBindingTypes)
+		{
+			// If 'spawnable' has been passed in, we only want to use children of UMovieSceneSpawnableBindingBase and vice versa
+			if ((!bSpawnable != CustomBindingType->IsChildOf<UMovieSceneSpawnableBindingBase>()) && (!bReplaceable != CustomBindingType->IsChildOf<UMovieSceneReplaceableBindingBase>()))
+			{
+				if (UMovieSceneCustomBinding* CustomBindingCDO = CustomBindingType ? CustomBindingType->GetDefaultObject<UMovieSceneCustomBinding>() : nullptr)
+				{
+					if (CustomBindingObject && CustomBindingCDO->SupportsBindingCreationFromObject(CustomBindingObject))
+					{
+						// Create a custom binding from this Object
+						NewCustomBinding = CustomBindingCDO->CreateNewCustomBinding(CustomBindingObject, *OwnerMovieScene);
+						if (NewCustomBinding)
+						{
+							break;
+						}
+					}
+				}
+			}
+		}
 	}
 
-	UMovieScene* RootMovieScene = Sequencer->GetRootMovieSceneSequence()->GetMovieScene();
+	if (NewCustomBinding)
+	{
+		FString DesiredBindingName = NewCustomBinding->GetDesiredBindingName();
+		FString CurrentName = DesiredBindingName.IsEmpty() ? InParams.BindingNameOverride : DesiredBindingName;
 
-	UWorld* World = Sequencer->GetPlaybackContext()->GetWorld();
+		FMovieScenePossessable* NewPossessable = nullptr;
+		FGuid NewID;
+		if (InParams.ReplacementGuid.IsValid())
+		{
+			NewID = InParams.ReplacementGuid;
+			NewPossessable = OwnerMovieScene->FindPossessable(InParams.ReplacementGuid);
+		}
+		if (!NewPossessable)
+		{
+			// Add a possessable binding track- we will use these even if the custom binding is a 'spawnable' one
+			NewID = OwnerMovieScene->AddPossessable(CurrentName, NewCustomBinding->GetBoundObjectClass());
+			NewPossessable = OwnerMovieScene->FindPossessable(NewID);
+		}
 
-	const FScopedTransaction Transaction(LOCTEXT("PasteBindings", "Paste Bindings"));
+		// Add the custom binding
+		BindingReferences->AddOrReplaceBinding(NewID, NewCustomBinding, InParams.BindingIndex);
 
-	TMap<FGuid, FGuid> OldToNewGuidMap;
-	TArray<FGuid> PossessableGuids;
-	TArray<TArray<FString> > PossessableObjectNames;
-	TArray<FGuid> SpawnableGuids;
-	TMap<FGuid, UMovieSceneFolder*> GuidToFolderMap;
+		UObject* SpawnedObject = nullptr;
+		// If this is a spawnable or replaceable binding, we need to set up some defaults
+		if (NewCustomBinding->WillSpawnObject(Sequencer->GetSharedPlaybackState()))
+		{
+			// Spawn the object so we can position it correctly, it's going to get spawned anyway since things default to spawned.
+			SpawnedObject = Sequencer->GetSpawnRegister().SpawnObject(NewID, *OwnerMovieScene, Sequencer->GetFocusedTemplateID(), Sequencer->GetSharedPlaybackState(), 0);
 
-	TArray<FMovieSceneBinding> BindingsPasted;
+			// Allow the binding to set up any necessary defaults
+			NewCustomBinding->SetupDefaults(SpawnedObject, NewID, *OwnerMovieScene, Sequencer->GetSharedPlaybackState());
 
-	const int NumTargets = FMath::Max(1, PasteBindingsParams.Bindings.Num());
+			FTransformData TransformData;
+			Sequencer->GetSpawnRegister().SetupDefaultsForSpawnable(SpawnedObject, NewID, TransformData, Sequencer, Sequencer->GetSequencerSettings());
+		}
 
-	for (int32 TargetIndex = 0; TargetIndex < NumTargets; ++TargetIndex)
-	{
-		TArray<UMovieSceneCopyableBinding*> ImportedBindings;
-		ImportObjectBindingsFromText(Sequencer.Get(), TextToImport, ImportedBindings);
+		Sequencer->State.Invalidate(NewID, Sequencer->GetFocusedTemplateID());
+		Sequencer->ForceEvaluate();
 
-		if (ImportedBindings.Num() == 0)
+		if (AActor* Actor = Cast<AActor>(SpawnedObject))
 		{
-			return false;
+			Sequencer->OnActorAddedToSequencer().Broadcast(Actor, NewID);
 		}
 
-		TArray<UObject*> SectionSubObjects;
-		for (UMovieSceneCopyableBinding* CopyableBinding : ImportedBindings)
-		{
-			// Clear transient flags on the imported tracks
-			for (UMovieSceneTrack* CopiedTrack : CopyableBinding->Tracks)
-			{
-				ResetCopiedTracksFlags(CopiedTrack);
-			}
+		Sequencer->OnAddBinding(NewID, OwnerMovieScene);
+
+		return NewID;
+	}
+	return FGuid();
+}
+
+FGuid CreateGenericBinding(TSharedRef<ISequencer> Sequencer, UObject* InObject, FMovieSceneBindingReferences* BindingReferences, const UE::Sequencer::FCreateBindingParams& InParams)
+{
+	using namespace UE::Sequencer;
+
+	UMovieSceneSequence* OwnerSequence = Sequencer->GetFocusedMovieSceneSequence();
+	UMovieScene* OwnerMovieScene = OwnerSequence->GetMovieScene();
+
+	ISequencerModule& Module = FModuleManager::Get().LoadModuleChecked<ISequencerModule>("Sequencer");
+	bool bSpawnable = InParams.bSpawnable && OwnerSequence->AllowsSpawnableObjects();
+	bool bAllowCustom = InParams.bAllowCustomBinding && OwnerSequence->AllowsCustomBindings();
+	bool bReplaceable = InParams.bReplaceable && bAllowCustom;
+	FGuid NewBindingID;
+	// First see if any custom bindings support creation from this object type directly. 
+
+	if (bAllowCustom)
+	{
+		
+
+		NewBindingID = TryCreateCustomBinding(Sequencer, InObject, BindingReferences, InParams, OwnerMovieScene, bSpawnable, bReplaceable);
+		if (NewBindingID.IsValid())
+		{
+			return NewBindingID;
+		}
+		else if (InObject && !InObject->IsA<AActor>())
+		{
+			// If the passed in object is not an actor, see if we can create an Actor from it, and if so, if that Actor type has a custom binding that supports it
+			UActorFactory* FactoryToUse = InParams.ActorFactory ? InParams.ActorFactory.Get() : FActorFactoryAssetProxy::GetFactoryForAssetObject(InObject);
+
+			if (FactoryToUse)
+			{
+				FText ErrorMessage;
+				if (FactoryToUse->CanCreateActorFrom(FAssetData(InObject), ErrorMessage))
+				{
+					UWorld* World = GCurrentLevelEditingViewportClient ? GCurrentLevelEditingViewportClient->GetWorld() : nullptr;
+					if (World)
+					{
+						const FName ActorName = MakeUniqueObjectName(World->PersistentLevel, FactoryToUse->NewActorClass->StaticClass(), *InParams.BindingNameOverride);
+
+						FActorSpawnParameters SpawnParams;
+						SpawnParams.ObjectFlags = RF_Transient | RF_Transactional;
+						SpawnParams.Name = ActorName;
+
+						AActor* Instance = FactoryToUse->CreateActor(InObject, World->PersistentLevel, FTransform(), SpawnParams);
+						if (Instance)
+						{
+							Instance->SetActorLabel(MovieSceneHelpers::MakeUniqueBindingName(OwnerMovieScene, FName::NameToDisplayString(InObject->GetName(), false)));
+							Instance->bIsEditorPreviewActor = false;
+
+							// Can we create a custom binding from this factory created actor? In all cases, we expect the binding creation process to duplicate
+							// the actor, so we should be safe to destroy the actor after this process.
+							NewBindingID = TryCreateCustomBinding(Sequencer, Instance, BindingReferences, InParams, OwnerMovieScene, bSpawnable, bReplaceable);
+
+							const bool bNetForce = false;
+							const bool bShouldModifyLevel = false;
+							World->DestroyActor(Instance, bNetForce, bShouldModifyLevel);
+
+							if (NewBindingID.IsValid())
+							{
+								return NewBindingID;
+							}
+						}
+					}
+				}
+			}
+		}
+	}
+
+	if (!InObject)
+	{
+		return FGuid();
+	}
+
+
+	// If no custom bindings support this object type, but InParams.bSpawnable is true, attempt to make an old-style spawnable.
+	if (bSpawnable)
+	{
+		NewBindingID = FSequencerUtilities::MakeNewSpawnable(Sequencer, *InObject, InParams.ActorFactory);
+		if (NewBindingID.IsValid())
+		{
+			return NewBindingID;
+		}
+	}
+
+	// Otherwise, create a possessable.
+
+	TArray<TPair<UObject*, FString>> ObjectsToPossess;
+
+	// Build up the list of child->parent bindings required for this object
+	{
+		UObject* CurrentObject = InObject;
+		while (CurrentObject)
+		{
+			TSharedPtr<IObjectSchema> Schema = Module.FindObjectSchema(CurrentObject);
+			if (Schema)
+			{
+				if (ObjectsToPossess.Num() == 0 && InParams.BindingNameOverride.Len() != 0)
+				{
+					ObjectsToPossess.Add(MakeTuple(CurrentObject, InParams.BindingNameOverride));
+				}
+				else
+				{
+					ObjectsToPossess.Add(MakeTuple(CurrentObject, Schema->GetPrettyName(CurrentObject).ToString()));
+				}
+
+				CurrentObject = Schema->GetParentObject(CurrentObject);
+			}
+			else
+			{
+				break;
+			}
+		}
+	}
+
+	// Nothing to possess?
+	if (ObjectsToPossess.Num() == 0)
+	{
+		// We've failed to find a custom binding type 
+
+		return FGuid();
+	}
+
+	const bool bParentContextsAreSignificant = OwnerSequence->AreParentContextsSignificant();
+
+	UObject* Context = Sequencer->GetPlaybackContext();
+
+	FGuid ParentID;
+
+	// Iterate in reverse (parent -> child)
+	for (int32 Index = ObjectsToPossess.Num() - 1; Index >= 0; --Index)
+	{
+		UObject* CurrentObject = ObjectsToPossess[Index].Key;
+
+		// If we're not purposefully replacing a binding, then check to see if we already have one, and use that
+		if (!InParams.ReplacementGuid.IsValid())
+		{
+			FGuid    ObjectGuid = Sequencer->GetHandleToObject(CurrentObject, false);
+
+			// If the object already has a binding, use that and move on
+			if (ObjectGuid.IsValid())
+			{
+				ParentID = ObjectGuid;
+				if (bParentContextsAreSignificant)
+				{
+					Context = CurrentObject;
+				}
+				continue;
+			}
+		}
+
+		// Create a new binding for this object
+		FString CurrentName = MoveTemp(ObjectsToPossess[Index].Value);
+
+		FMovieScenePossessable* NewPossessable = nullptr;
+		FGuid NewID;
+		if (InParams.ReplacementGuid.IsValid() && !ParentID.IsValid())
+		{
+			NewID = InParams.ReplacementGuid;
+			NewPossessable = OwnerMovieScene->FindPossessable(InParams.ReplacementGuid);
+		}
+		if (!NewPossessable)
+		{
+			NewID = OwnerMovieScene->AddPossessable(CurrentName, CurrentObject->GetClass());
+			NewPossessable = OwnerMovieScene->FindPossessable(NewID);
+		}
+
+		// If we're not trying to replace a binding, and the object is a spawnable, try and bind to that first
+		if (InParams.ReplacementGuid.IsValid() || !NewPossessable->BindSpawnableObject(Sequencer->GetFocusedTemplateID(), CurrentObject, Sequencer->GetSharedPlaybackState()))
+		{
+			FUniversalObjectLocator Locator;
+			if (!OwnerSequence->MakeLocatorForObject(CurrentObject, Context, Locator) || Locator.IsEmpty())
+			{
+				// Unable to possess this object
+				return FGuid();
+			}
+
+			if (InParams.ReplacementGuid.IsValid() && !ParentID.IsValid())
+			{
+				BindingReferences->AddOrReplaceBinding(NewID, MoveTemp(Locator), InParams.BindingIndex);
+				Sequencer->State.Invalidate(NewID, Sequencer->GetFocusedTemplateID());
+			}
+			else
+			{
+				BindingReferences->AddBinding(NewID, MoveTemp(Locator));
+			}
+		}
+
+		if (ParentID.IsValid())
+		{
+			NewPossessable->SetParent(ParentID, OwnerMovieScene);
+
+			FMovieSceneSpawnable* ParentSpawnable = OwnerMovieScene->FindSpawnable(ParentID);
+			if (ParentSpawnable)
+			{
+				ParentSpawnable->AddChildPossessable(NewID);
+			}
+		}
+
+		ParentID = NewID;
+
+		if (AActor* Actor = Cast<AActor>(CurrentObject))
+		{
+			Sequencer->OnActorAddedToSequencer().Broadcast(Actor, NewID);
+		}
+
+		// If this is the last one
+		if (Index == 0)
+		{
+			Sequencer->OnAddBinding(NewID, OwnerMovieScene);
+			return NewID;
+		}
+
+		if (bParentContextsAreSignificant)
+		{
+			Context = CurrentObject;
+		}
+	}
+
+	// Should never get here - we should always hit the Index == 0 condition inside the loop
+	return FGuid();
+}
+
+bool FSequencerUtilities::PasteBindings(const FString& TextToImport, TSharedRef<ISequencer> Sequencer, FMovieScenePasteBindingsParams PasteBindingsParams, TArray<FMovieSceneBindingProxy>& OutBindings, TArray<FNotificationInfo>& OutErrors)
+{
+	UMovieSceneSequence* Sequence = Sequencer->GetFocusedMovieSceneSequence();
+	if (!Sequence)
+	{
+		return false;
+	}
+
+	FMovieSceneBindingReferences* BindingReferences = Sequence->GetBindingReferences();
+
+	UMovieScene* MovieScene = Sequence->GetMovieScene();
+	if (!MovieScene)
+	{
+		return false;
+	}
+
+	UMovieScene* RootMovieScene = Sequencer->GetRootMovieSceneSequence()->GetMovieScene();
+
+	UWorld* World = Sequencer->GetPlaybackContext()->GetWorld();
+
+	const FScopedTransaction Transaction(LOCTEXT("PasteBindings", "Paste Bindings"));
+
+	TMap<FGuid, FGuid> OldToNewGuidMap;
+	TArray<FGuid> PossessableGuids;
+	TArray<TArray<FString> > PossessableObjectNames;
+	TArray<FGuid> SpawnableGuids;
+	TMap<FGuid, UMovieSceneFolder*> GuidToFolderMap;
+
+	TArray<FMovieSceneBinding> BindingsPasted;
+
+	const int NumTargets = FMath::Max(1, PasteBindingsParams.Bindings.Num());
+
+	for (int32 TargetIndex = 0; TargetIndex < NumTargets; ++TargetIndex)
+	{
+		TArray<UMovieSceneCopyableBinding*> ImportedBindings;
+		ImportObjectBindingsFromText(Sequencer.Get(), TextToImport, ImportedBindings);
+
+		if (ImportedBindings.Num() == 0)
+		{
+			return false;
+		}
+
+		TArray<UObject*> SectionSubObjects;
+		for (UMovieSceneCopyableBinding* CopyableBinding : ImportedBindings)
+		{
+			// Clear transient flags on the imported tracks
+			for (UMovieSceneTrack* CopiedTrack : CopyableBinding->Tracks)
+			{
+				ResetCopiedTracksFlags(CopiedTrack);
+			}
 
 			UMovieSceneFolder* ParentFolder = PasteBindingsParams.ParentFolder;
 
@@ -2110,14 +2686,61 @@ bool FSequencerUtilities::PasteBindings(const FString& TextToImport, TSharedRef<
 
 			if (CopyableBinding->Possessable.GetGuid().IsValid())
 			{
-				FGuid NewGuid = FGuid::NewGuid();
+				// TODO: We likely need additional work here for possessable bindings using locators other than actor locators.
+				// For now, we'll at least handle the custom binding case
+
+				// If we have a custom binding, we need to let the sequence create it, especially since it could have a spawnable template.
+				// However, making a new custom spawnable also creates the binding for us - this is a problem
+				// because we need to use our binding (which has tracks associated with it). To solve this, we let it create
+				// an object template based off of our (transient package owned) template, then find the newly created binding
+				// and update it.
 
+				FGuid NewGuid;
+				if (!CopyableBinding->CustomBindings.IsEmpty())
+				{
+					if (BindingReferences)
+					{
+						int32 BindingIndex = 0;
+						for (UMovieSceneCustomBinding* CustomBinding : CopyableBinding->CustomBindings)
+						{
+							if (CustomBinding)
+							{
+								// We need to duplicate this as the one in CopyableBinding will be transient package owned
+								UMovieSceneCustomBinding* NewCustomBinding = Cast<UMovieSceneCustomBinding>(StaticDuplicateObject(CustomBinding, MovieScene));
+
+								// This will either add a brand new possessable and binding (if one doesn't exist for that guid), or just add a new binding to that same possessable
+								UE::Sequencer::FCreateBindingParams CreateBindingParams;
+								CreateBindingParams.ReplacementGuid = NewGuid;
+								CreateBindingParams.BindingIndex = BindingIndex++;
+								CreateBindingParams.bAllowCustomBinding = true;
+								CreateBindingParams.CustomBinding = NewCustomBinding;
+								NewGuid = CreateGenericBinding(Sequencer, nullptr, BindingReferences, CreateBindingParams);
+							}
+						}
+					}
+				}
+				else
+				{
+					FMovieScenePossessable NewPossessable = CopyableBinding->Possessable;
+					NewPossessable.SetGuid(FGuid::NewGuid());
+					MovieScene->AddPossessable(NewPossessable, FMovieSceneBinding(NewPossessable.GetGuid(), NewPossessable.GetName()));
+					NewGuid = NewPossessable.GetGuid();
+				}
+					
 				FMovieSceneBinding NewBinding(NewGuid, CopyableBinding->Binding.GetName(), CopyableBinding->Tracks);
+				FMovieScenePossessable* Possessable = MovieScene->FindPossessable(NewGuid);
 
-				FMovieScenePossessable NewPossessable = CopyableBinding->Possessable;
-				NewPossessable.SetGuid(NewGuid);
+				// Copy the name of the original possessable too.
+				Possessable->SetName(CopyableBinding->Possessable.GetName());
 
-				MovieScene->AddPossessable(NewPossessable, NewBinding);
+				// Clear the transient flags on the copyable binding before assigning to the new possessable
+				for (UMovieSceneTrack* Track : NewBinding.GetTracks())
+				{
+					ResetCopiedTracksFlags(Track);
+				}
+
+				// Replace the auto-generated binding with our deserialized bindings (which has our tracks)
+				MovieScene->ReplaceBinding(NewGuid, NewBinding);
 
 				OldToNewGuidMap.Add(CopyableBinding->Possessable.GetGuid(), NewGuid);
 
@@ -2144,18 +2767,15 @@ bool FSequencerUtilities::PasteBindings(const FString& TextToImport, TSharedRef<
 				const UClass* PossessedObjectClass = CopyableBinding->Possessable.GetPossessedObjectClass();
 				if (PossessedObjectClass && !PossessedObjectClass->IsChildOf(AActor::StaticClass()))
 				{
-					if (FMovieScenePossessable* Possessable = MovieScene->FindPossessable(NewGuid))
+					if (TargetIndex < PasteBindingsParams.Bindings.Num())
 					{
-						if (TargetIndex < PasteBindingsParams.Bindings.Num())
-						{
-							Possessable->SetParent(PasteBindingsParams.Bindings[TargetIndex].BindingID, MovieScene);
-						}
+						Possessable->SetParent(PasteBindingsParams.Bindings[TargetIndex].BindingID, MovieScene);
 					}
 				}
 
 				// Find the actors that this pasted binding should bind to
 				TArray<AActor*> ActorsToRebind;
-				if (World)
+				if (World && CopyableBinding->CustomBindings.IsEmpty())
 				{
 					for (TActorIterator<AActor> ActorItr(World); ActorItr; ++ActorItr)
 					{
@@ -2225,9 +2845,9 @@ bool FSequencerUtilities::PasteBindings(const FString& TextToImport, TSharedRef<
 				// and update it.
 
 				FGuid NewGuid;
-				if (CopyableBinding->SpawnableObjectTemplate)
+				if (!CopyableBinding->SpawnableObjectTemplates.IsEmpty())
 				{
-					NewGuid = MakeNewSpawnable(Sequencer, *CopyableBinding->SpawnableObjectTemplate, nullptr, false, FName(*CopyableBinding->Spawnable.GetName()));
+					NewGuid = MakeNewSpawnable(Sequencer, *CopyableBinding->SpawnableObjectTemplates[0], nullptr, false, FName(*CopyableBinding->Spawnable.GetName()));
 				}
 				else
 				{
@@ -2282,6 +2902,10 @@ bool FSequencerUtilities::PasteBindings(const FString& TextToImport, TSharedRef<
 	// Fix possessable actor bindings
 	for (int32 PossessableGuidIndex = 0; PossessableGuidIndex < PossessableGuids.Num(); ++PossessableGuidIndex)
 	{
+		if (BindingReferences && Algo::AnyOf(BindingReferences->GetReferences(PossessableGuids[PossessableGuidIndex]), [](const FMovieSceneBindingReference& Reference) {return Reference.CustomBinding; }))
+		{
+			continue;
+		}
 		FMovieScenePossessable* Possessable = MovieScene->FindPossessable(PossessableGuids[PossessableGuidIndex]);
 		UWorld* PlaybackContext = Sequencer->GetPlaybackContext()->GetWorld();
 		if (Possessable && PlaybackContext)
@@ -2453,126 +3077,6 @@ TArray<FString> FSequencerUtilities::GetPasteBindingsObjectNames(TSharedRef<ISeq
 	return ObjectNames;
 }
 
-FGuid CreateGenericBinding(TSharedRef<ISequencer> Sequencer, UObject& InObject, FMovieSceneBindingReferences* BindingReferences, const UE::Sequencer::FCreateBindingParams& InParams)
-{
-	using namespace UE::Sequencer;
-
-	UMovieSceneSequence* OwnerSequence   = Sequencer->GetFocusedMovieSceneSequence();
-	UMovieScene*         OwnerMovieScene = OwnerSequence->GetMovieScene();
-
-	ISequencerModule& Module = FModuleManager::Get().LoadModuleChecked<ISequencerModule>("Sequencer");
-
-	TArray<TPair<UObject*, FString>> ObjectsToPossess;
-
-	// Build up the list of child->parent bindings required for this object
-	{
-		UObject* CurrentObject = &InObject;
-		while (CurrentObject)
-		{
-			TSharedPtr<IObjectSchema> Schema = Module.FindObjectSchema(CurrentObject);
-			if (Schema)
-			{
-				if (ObjectsToPossess.Num() == 0 && InParams.BindingNameOverride.Len() != 0)
-				{
-					ObjectsToPossess.Add(MakeTuple(CurrentObject, InParams.BindingNameOverride));
-				}
-				else
-				{
-					ObjectsToPossess.Add(MakeTuple(CurrentObject, Schema->GetPrettyName(CurrentObject).ToString()));
-				}
-				
-				CurrentObject = Schema->GetParentObject(CurrentObject);
-			}
-			else
-			{
-				break;
-			}
-		}
-	}
-
-	// Nothing to possess?
-	if (ObjectsToPossess.Num() == 0)
-	{
-		return FGuid();
-	}
-
-	const bool bParentContextsAreSignificant = OwnerSequence->AreParentContextsSignificant();
-
-	UObject* Context = Sequencer->GetPlaybackContext();
-
-	FGuid ParentID;
-
-	// Iterate in reverse (parent -> child)
-	for (int32 Index = ObjectsToPossess.Num()-1; Index >= 0; --Index)
-	{
-		UObject* CurrentObject = ObjectsToPossess[Index].Key;
-		FGuid    ObjectGuid    = Sequencer->GetHandleToObject(CurrentObject, false);
-
-		// If the object already has a binding, use that and move on
-		if (ObjectGuid.IsValid())
-		{
-			ParentID = ObjectGuid;
-			if (bParentContextsAreSignificant)
-			{
-				Context = CurrentObject;
-			}
-			continue;
-		}
-
-		// Create a new binding for this object
-		FString CurrentName = MoveTemp(ObjectsToPossess[Index].Value);
-		FGuid   NewID       = OwnerMovieScene->AddPossessable(CurrentName, CurrentObject->GetClass());
-
-		FMovieScenePossessable* NewPossessable = OwnerMovieScene->FindPossessable(NewID);
-
-		// If the object is a spawnable, try and bind to that first
-		if (!NewPossessable->BindSpawnableObject(Sequencer->GetFocusedTemplateID(), CurrentObject, Sequencer->GetSharedPlaybackState()))
-		{
-			FUniversalObjectLocator Locator;
-			if (!OwnerSequence->MakeLocatorForObject(CurrentObject, Context, Locator) || Locator.IsEmpty())
-			{
-				// Unable to possess this object
-				return FGuid();
-			}
-
-			BindingReferences->AddBinding(NewID, MoveTemp(Locator));
-		}
-
-		if (ParentID.IsValid())
-		{
-			NewPossessable->SetParent(ParentID, OwnerMovieScene);
-
-			FMovieSceneSpawnable* ParentSpawnable = OwnerMovieScene->FindSpawnable(ParentID);
-			if (ParentSpawnable)
-			{
-				ParentSpawnable->AddChildPossessable(NewID);
-			}
-		}
-
-		ParentID = NewID;
-
-		if (AActor* Actor = Cast<AActor>(CurrentObject))
-		{
-			Sequencer->OnActorAddedToSequencer().Broadcast(Actor, NewID);
-		}
-
-		// If this is the last one
-		if (Index == 0)
-		{
-			Sequencer->OnAddBinding(NewID, OwnerMovieScene);
-			return NewID;
-		}
-
-		if (bParentContextsAreSignificant)
-		{
-			Context = CurrentObject;
-		}
-	}
-
-	// Should never get here - we should always hit the Index == 0 condition inside the loop
-	return FGuid();
-}
-
 FGuid CreateImplementationDefinedBinding(TSharedRef<ISequencer> Sequencer, UObject& InObject, const UE::Sequencer::FCreateBindingParams& InParams)
 {
 	UMovieSceneSequence* OwnerSequence   = Sequencer->GetFocusedMovieSceneSequence();
@@ -2657,23 +3161,34 @@ FGuid FSequencerUtilities::CreateBinding(TSharedRef<ISequencer> Sequencer, UObje
 	OwnerSequence->Modify();
 	OwnerMovieScene->Modify();
 
+	FGuid BindingGuid;
 	FMovieSceneBindingReferences* BindingReferences = OwnerSequence->GetBindingReferences();
+	if (BindingReferences)
+	{
+		BindingGuid = CreateGenericBinding(Sequencer, &InObject, BindingReferences, InParams);
+	}
+	else if (InParams.bSpawnable)
+	{
+		// Create an old-style spawnable
+		BindingGuid = MakeNewSpawnable(Sequencer, InObject, InParams.ActorFactory.Get());
+	}
+	else
+	{
+		BindingGuid = CreateImplementationDefinedBinding(Sequencer, InObject, InParams);
+	}
 
-	FGuid PossessableGuid = BindingReferences
-		? CreateGenericBinding(Sequencer, InObject, BindingReferences, InParams)
-		: CreateImplementationDefinedBinding(Sequencer, InObject, InParams);
-
-	if (!PossessableGuid.IsValid())
+	if (!BindingGuid.IsValid())
 	{
 		return FGuid();
 	}
 
 	if (InParams.DesiredFolder != NAME_None)
 	{
-		// Find the outermost object and put it in a folder of the specified name
-		FGuid RootObjectGuid = PossessableGuid;
+		// Find the outermost object and put it in a folder of the specified name.
+		FGuid RootObjectGuid = BindingGuid;
 		while (true)
 		{
+			// This only applies to possessables/custom bindings, as old-style spawnables will not have parents.
 			FMovieScenePossessable* Possessable = OwnerMovieScene->FindPossessable(RootObjectGuid);
 			if (!Possessable || !Possessable->GetParent().IsValid())
 			{
@@ -2707,8 +3222,8 @@ FGuid FSequencerUtilities::CreateBinding(TSharedRef<ISequencer> Sequencer, UObje
 		}
 	}
 
-	Sequencer->OnAddBinding(PossessableGuid, OwnerMovieScene);
-	return PossessableGuid;
+	Sequencer->OnAddBinding(BindingGuid, OwnerMovieScene);
+	return BindingGuid;
 }
 
 void FSequencerUtilities::UpdateBindingIDs(TSharedRef<ISequencer> Sequencer, FGuid OldGuid, FGuid NewGuid)
diff --git a/Engine/Source/Editor/Sequencer/Public/ISequencer.h b/Engine/Source/Editor/Sequencer/Public/ISequencer.h
--- a/Engine/Source/Editor/Sequencer/Public/ISequencer.h
+++ b/Engine/Source/Editor/Sequencer/Public/ISequencer.h
@@ -35,6 +35,7 @@
 #if UE_ENABLE_INCLUDE_ORDER_DEPRECATED_IN_5_3
 #include "ITimeSlider.h"
 #endif
+#include "Bindings/MovieSceneCustomBinding.h"
 
 #include "ISequencer.generated.h"
 
@@ -81,6 +82,7 @@ namespace Sequencer
 {
 
 class FSequencerEditorViewModel;
+struct FCreateBindingParams;
 
 } // namespace Sequencer
 } // namespace UE
@@ -296,9 +298,12 @@ public:
 	virtual void SuppressAutoEvaluation(UMovieSceneSequence* Sequence, const FGuid& InSignature) = 0;
 
 	/**
-	 * Create a new binding for the specified object
-	 */
-	virtual FGuid CreateBinding(UObject& InObject, const FString& InName) = 0;
+	* Create a new binding for the specified object
+	*/
+	SEQUENCER_API virtual FGuid CreateBinding(UObject& InObject, const UE::Sequencer::FCreateBindingParams& InParams) = 0;
+
+	// Override from IMovieScenePlayer
+	FGuid CreateBinding(UMovieSceneSequence* InSequence, UObject* InObject) override;
 
 	/**
 	 * Attempts to add a new spawnable to the MovieScene for the specified object (asset, class or actor instance)
@@ -833,6 +838,12 @@ public:
 	SEQUENCER_API FFrameRate GetFocusedDisplayRate() const;
 
 
+	/**
+	 * Create a new binding for the specified object, lightly deprecated and no longer virtual in favor of the overload with binding parameters
+	 */
+	SEQUENCER_API FGuid CreateBinding(UObject& InObject, const FString& InName);
+
+
 	/**
 	* Get the Display Name of the Object Binding Track.
 	* @param InBinding the Binding of the Object
@@ -846,6 +857,12 @@ public:
 	* @param InDisplayName The new name of the object binding track.
 	*/
 	virtual void SetDisplayName(FGuid InBinding, const FText& InDisplayName) = 0;
+
+	/*
+	* Returns priority-sorted list of custom binding types supported by this Sequencer. 
+	*/
+	virtual TArrayView<const TSubclassOf<UMovieSceneCustomBinding>> GetSupportedCustomBindingTypes() const { static TArray<TSubclassOf<UMovieSceneCustomBinding>> EmptyArray; return EmptyArray; }
+
 protected:
 	FOnInitializeDetailsPanel InitializeDetailsPanelEvent;
 	FOnCameraAddedToSequencer CameraAddedToSequencer;
diff --git a/Engine/Source/Editor/Sequencer/Public/SequencerUtilities.h b/Engine/Source/Editor/Sequencer/Public/SequencerUtilities.h
--- a/Engine/Source/Editor/Sequencer/Public/SequencerUtilities.h
+++ b/Engine/Source/Editor/Sequencer/Public/SequencerUtilities.h
@@ -23,6 +23,7 @@ struct FMovieScenePossessable;
 struct FMovieSceneSpawnable;
 struct FNotificationInfo;
 class ULevelSequence;
+class UMovieSceneCustomBinding;
 enum class EMovieSceneBlendType : uint8;
 
 namespace UE::Sequencer
@@ -50,7 +51,23 @@ struct FCreateBindingParams
 	}
 
 	FString BindingNameOverride;
+	
 	FName DesiredFolder;
+	
+	/* If true, will prefer the creation of a custom or regular Spawnable binding, unless such is incompatible with the passed in object.*/
+	bool bSpawnable = false;
+	/* If true, will prefer the creation of a custom Replaceable binding, unless such is incompatible with the passed in object.*/
+	bool bReplaceable = false;
+	/* If true, will allow the creation of custom bindings if they support the object type.*/
+	bool bAllowCustomBinding = true;
+	/* If set, will attempt to replace any existing possessable binding at the provided guid and binding index */
+	FGuid ReplacementGuid;
+	/* Optional BindingIndex used if bReplace is true to replace a specific possessable binding*/
+	int32 BindingIndex = 0;
+	/* Optional pre-created custom binding to use when creating the binding.*/
+	TObjectPtr<UMovieSceneCustomBinding> CustomBinding = nullptr;
+	/* May be used depending on options if an asset UObject is passed in to create a custom or regular spawnable actor binding*/
+	TObjectPtr<UActorFactory> ActorFactory = nullptr;
 };
 
 } // namespace UE::Sequencer
@@ -184,11 +201,14 @@ struct SEQUENCER_API FSequencerUtilities
 
 	static FGuid MakeNewSpawnable(TSharedRef<ISequencer> Sequencer, UObject& SourceObject, UActorFactory* ActorFactory = nullptr, bool bSetupDefaults = true, FName SpawnableName = NAME_None);
 
-	/** Convert the requested possessable to spawnable. If there are multiple objects assigned to the possessable, multiple spawnables will be created */
+	/** Convert the requested object binding to old-style spawnable. If there are multiple objects assigned to the possessable, multiple spawnables will be created */
 	static TArray<FMovieSceneSpawnable*> ConvertToSpawnable(TSharedRef<ISequencer> Sequencer, FGuid PossessableGuid);
 	
-	/** Convert the requested spawnable to possessable */
-	static FMovieScenePossessable* ConvertToPossessable(TSharedRef<ISequencer> Sequencer, FGuid SpawnableGuid);
+	/** Convert the requested object binding and object binding index to a possessable */
+	static FMovieScenePossessable* ConvertToPossessable(TSharedRef<ISequencer> Sequencer, FGuid BindingGuid, int32 BindingIndex=0);
+
+	/** Convert the selected object binding and object binding index to a custom binding of the chosen type.*/
+	static FMovieScenePossessable* ConvertToCustomBinding(TSharedRef<ISequencer> Sequencer, FGuid BindingGuid, TSubclassOf<UMovieSceneCustomBinding> CustomBindingType, int32 BindingIndex = 0);
 
 	/** Copy/paste folders */
 	static void CopyFolders(const TArray<UMovieSceneFolder*>& Folders, FString& ExportedText);
diff --git a/Engine/Source/Editor/UniversalObjectLocatorEditor/Private/ActorLocatorEditor.cpp b/Engine/Source/Editor/UniversalObjectLocatorEditor/Private/ActorLocatorEditor.cpp
--- a/Engine/Source/Editor/UniversalObjectLocatorEditor/Private/ActorLocatorEditor.cpp
+++ b/Engine/Source/Editor/UniversalObjectLocatorEditor/Private/ActorLocatorEditor.cpp
@@ -86,7 +86,7 @@ TSharedPtr<SWidget> FActorLocatorEditor::MakeEditUI(TSharedPtr<IUniversalObjectL
 	.AllowedClass(AActor::StaticClass())
 	.OnObjectChanged(this, &FActorLocatorEditor::OnSetObject, TWeakPtr<IUniversalObjectLocatorCustomization>(Customization))
 	.AllowClear(true)
-	.DisplayUseSelected(true)
+	.DisplayUseSelected(false)
 	.DisplayBrowse(true)
 	.DisplayThumbnail(true);
 
diff --git a/Engine/Source/Runtime/Engine/Private/UniversalObjectLocators/ActorLocatorFragment.cpp b/Engine/Source/Runtime/Engine/Private/UniversalObjectLocators/ActorLocatorFragment.cpp
--- a/Engine/Source/Runtime/Engine/Private/UniversalObjectLocators/ActorLocatorFragment.cpp
+++ b/Engine/Source/Runtime/Engine/Private/UniversalObjectLocators/ActorLocatorFragment.cpp
@@ -112,6 +112,13 @@ UE::UniversalObjectLocator::FResolveResult FActorLocatorFragment::Resolve(const
 	// Finally fallback to just trying to resolve the path directly
 	if (!Result)
 	{
+		FSoftObjectPath TempPath = Path;
+
+		// Soft Object Paths don't follow asset redirectors when attempting to call ResolveObject or TryLoad.
+		// We want to follow the asset redirector so that maps that have been renamed (from Untitled to their first asset name)
+		// properly resolve. This fixes Possessable bindings losing their references the first time you save a map.
+		TempPath.PreSavePath();
+
 #if WITH_EDITORONLY_DATA
 		UPackage* ContextPackage = Params.Context ? Params.Context->GetOutermost() : nullptr;
 		if (ContextPackage)
@@ -120,7 +127,7 @@ UE::UniversalObjectLocator::FResolveResult FActorLocatorFragment::Resolve(const
 			const int32 PIEInstanceID = ContextPackage->GetPIEInstanceID();
 			if (PIEInstanceID != INDEX_NONE)
 			{
-				FSoftObjectPath PIEPath = Path;
+				FSoftObjectPath PIEPath = TempPath;
 				PIEPath.FixupForPIE(PIEInstanceID);
 
 				Result = PIEPath.ResolveObject();
@@ -129,7 +136,7 @@ UE::UniversalObjectLocator::FResolveResult FActorLocatorFragment::Resolve(const
 		}
 #endif
 
-		Result = Path.ResolveObject();
+		Result = TempPath.ResolveObject();
 	}
 
 	return FResolveResultData(Result);
diff --git a/Engine/Source/Runtime/LevelSequence/Private/LevelSequence.cpp b/Engine/Source/Runtime/LevelSequence/Private/LevelSequence.cpp
--- a/Engine/Source/Runtime/LevelSequence/Private/LevelSequence.cpp
+++ b/Engine/Source/Runtime/LevelSequence/Private/LevelSequence.cpp
@@ -539,6 +539,11 @@ bool ULevelSequence::AllowsSpawnableObjects() const
 	return true;
 }
 
+bool ULevelSequence::AllowsCustomBindings() const
+{
+	return true;
+}
+
 bool ULevelSequence::CanRebindPossessable(const FMovieScenePossessable& InPossessable) const
 {
 	return !InPossessable.GetParent().IsValid();
diff --git a/Engine/Source/Runtime/LevelSequence/Private/LevelSequencePlayer.cpp b/Engine/Source/Runtime/LevelSequence/Private/LevelSequencePlayer.cpp
--- a/Engine/Source/Runtime/LevelSequence/Private/LevelSequencePlayer.cpp
+++ b/Engine/Source/Runtime/LevelSequence/Private/LevelSequencePlayer.cpp
@@ -117,7 +117,7 @@ void ULevelSequencePlayer::SetSourceActorContext(UWorld* InStreamingWorld, FActo
 	SourceAssetPath = InSourceAssetPath;
 }
 
-void ULevelSequencePlayer::ResolveBoundObjects(UE::UniversalObjectLocator::FResolveParams& ResolveParams, const FGuid& InBindingId, FMovieSceneSequenceID SequenceID, UMovieSceneSequence& InSequence, TArray<UObject*, TInlineAllocator<1>>& OutObjects) const
+void ULevelSequencePlayer::ResolveBoundObjects(UE::UniversalObjectLocator::FResolveParams& LocatorResolveParams, const FGuid& InBindingId, FMovieSceneSequenceID SequenceID, UMovieSceneSequence& InSequence, TArray<UObject*, TInlineAllocator<1>>& OutObjects) const
 {
 	using namespace UE::UniversalObjectLocator;
 	using namespace UE::MovieScene;
@@ -126,19 +126,40 @@ void ULevelSequencePlayer::ResolveBoundObjects(UE::UniversalObjectLocator::FReso
 
 	if (bAllowDefault)
 	{
-		if (ResolveParams.ParameterBuffer == nullptr)
+		if (const FMovieSceneBindingReferences* BindingReferences = InSequence.GetBindingReferences())
 		{
-			// Allocate temporary local buffer for this
-			TInlineResolveParameterBuffer<128> Buffer;
-			ResolveParams.ParameterBuffer = &Buffer;
-			ResolveParams.ParameterBuffer->AddParameter(FActorLocatorFragmentResolveParameter::ParameterType, WeakStreamingWorld.Get(), ContainerID, SourceAssetPath);
-			InSequence.LocateBoundObjects(InBindingId, ResolveParams, OutObjects);
-			ResolveParams.ParameterBuffer = nullptr;
+			FMovieSceneBindingResolveParams BindingResolveParams{ &InSequence, InBindingId, SequenceID };
+			if (LocatorResolveParams.ParameterBuffer == nullptr)
+			{
+				// Allocate temporary local buffer for this
+				TInlineResolveParameterBuffer<128> Buffer;
+				LocatorResolveParams.ParameterBuffer = &Buffer;
+				LocatorResolveParams.ParameterBuffer->AddParameter(FActorLocatorFragmentResolveParameter::ParameterType, WeakStreamingWorld.Get(), ContainerID, SourceAssetPath);
+				BindingReferences->ResolveBinding(BindingResolveParams, LocatorResolveParams, ConstCastSharedPtr<const UE::MovieScene::FSharedPlaybackState>(const_cast<ULevelSequencePlayer*>(this)->FindSharedPlaybackState()), OutObjects);
+				LocatorResolveParams.ParameterBuffer = nullptr;
+			}
+			else
+			{
+				LocatorResolveParams.ParameterBuffer->AddParameter(FActorLocatorFragmentResolveParameter::ParameterType, WeakStreamingWorld.Get(), ContainerID, SourceAssetPath);
+				BindingReferences->ResolveBinding(BindingResolveParams, LocatorResolveParams, ConstCastSharedPtr<const UE::MovieScene::FSharedPlaybackState>(const_cast<ULevelSequencePlayer*>(this)->FindSharedPlaybackState()), OutObjects);
+			}
 		}
 		else
 		{
-			ResolveParams.ParameterBuffer->AddParameter(FActorLocatorFragmentResolveParameter::ParameterType, WeakStreamingWorld.Get(), ContainerID, SourceAssetPath);
-			InSequence.LocateBoundObjects(InBindingId, ResolveParams, OutObjects);
+			if (LocatorResolveParams.ParameterBuffer == nullptr)
+			{
+				// Allocate temporary local buffer for this
+				TInlineResolveParameterBuffer<128> Buffer;
+				LocatorResolveParams.ParameterBuffer = &Buffer;
+				LocatorResolveParams.ParameterBuffer->AddParameter(FActorLocatorFragmentResolveParameter::ParameterType, WeakStreamingWorld.Get(), ContainerID, SourceAssetPath);
+				InSequence.LocateBoundObjects(InBindingId, LocatorResolveParams, OutObjects);
+				LocatorResolveParams.ParameterBuffer = nullptr;
+			}
+			else
+			{
+				LocatorResolveParams.ParameterBuffer->AddParameter(FActorLocatorFragmentResolveParameter::ParameterType, WeakStreamingWorld.Get(), ContainerID, SourceAssetPath);
+				InSequence.LocateBoundObjects(InBindingId, LocatorResolveParams, OutObjects);
+			}
 		}
 	}
 }
diff --git a/Engine/Source/Runtime/LevelSequence/Private/LevelSequenceSpawnRegister.cpp b/Engine/Source/Runtime/LevelSequence/Private/LevelSequenceSpawnRegister.cpp
--- a/Engine/Source/Runtime/LevelSequence/Private/LevelSequenceSpawnRegister.cpp
+++ b/Engine/Source/Runtime/LevelSequence/Private/LevelSequenceSpawnRegister.cpp
@@ -7,6 +7,7 @@
 #include "LevelSequenceModule.h"
 #include "IMovieSceneObjectSpawner.h"
 #include "Modules/ModuleManager.h"
+#include "Bindings/MovieSceneSpawnableBinding.h"
 
 FLevelSequenceSpawnRegister::FLevelSequenceSpawnRegister()
 {
@@ -31,18 +32,25 @@ UObject* FLevelSequenceSpawnRegister::SpawnObject(FMovieSceneSpawnable& Spawnabl
 	return nullptr;
 }
 
-void FLevelSequenceSpawnRegister::DestroySpawnedObject(UObject& Object)
+void FLevelSequenceSpawnRegister::DestroySpawnedObject(UObject& Object, UMovieSceneSpawnableBindingBase* CustomSpawnableBinding)
 {
-	for (TSharedRef<IMovieSceneObjectSpawner> MovieSceneObjectSpawner : MovieSceneObjectSpawners)
+	if (CustomSpawnableBinding)
+	{
+		CustomSpawnableBinding->DestroySpawnedObject(&Object);
+	}
+	else
 	{
-		if (Object.IsA(MovieSceneObjectSpawner->GetSupportedTemplateType()))
+		for (TSharedRef<IMovieSceneObjectSpawner> MovieSceneObjectSpawner : MovieSceneObjectSpawners)
 		{
-			MovieSceneObjectSpawner->DestroySpawnedObject(Object);
-			return;
+			if (Object.IsA(MovieSceneObjectSpawner->GetSupportedTemplateType()))
+			{
+				MovieSceneObjectSpawner->DestroySpawnedObject(Object);
+				return;
+			}
 		}
-	}
 
-	checkf(false, TEXT("No valid object spawner found to destroy spawned object of type %s"), *Object.GetClass()->GetName());
+		checkf(false, TEXT("No valid object spawner found to destroy spawned object of type %s"), *Object.GetClass()->GetName());
+	}
 }
 
 #if WITH_EDITOR
diff --git a/Engine/Source/Runtime/LevelSequence/Public/LevelSequence.h b/Engine/Source/Runtime/LevelSequence/Public/LevelSequence.h
--- a/Engine/Source/Runtime/LevelSequence/Public/LevelSequence.h
+++ b/Engine/Source/Runtime/LevelSequence/Public/LevelSequence.h
@@ -50,6 +50,7 @@ public:
 	LEVELSEQUENCE_API virtual void UnbindObjects(const FGuid& ObjectId, const TArray<UObject*>& InObjects, UObject* InContext) override;
 	LEVELSEQUENCE_API virtual void UnbindInvalidObjects(const FGuid& ObjectId, UObject* InContext) override;
 	LEVELSEQUENCE_API virtual bool AllowsSpawnableObjects() const override;
+	LEVELSEQUENCE_API virtual bool AllowsCustomBindings() const override;
 	LEVELSEQUENCE_API virtual bool CanRebindPossessable(const FMovieScenePossessable& InPossessable) const override;
 	LEVELSEQUENCE_API virtual UObject* MakeSpawnableTemplateFromInstance(UObject& InSourceObject, FName ObjectName) override;
 	LEVELSEQUENCE_API virtual bool CanAnimateObject(UObject& InObject) const override;
diff --git a/Engine/Source/Runtime/LevelSequence/Public/LevelSequenceSpawnRegister.h b/Engine/Source/Runtime/LevelSequence/Public/LevelSequenceSpawnRegister.h
--- a/Engine/Source/Runtime/LevelSequence/Public/LevelSequenceSpawnRegister.h
+++ b/Engine/Source/Runtime/LevelSequence/Public/LevelSequenceSpawnRegister.h
@@ -14,6 +14,7 @@ class IMovieScenePlayer;
 class UClass;
 class UObject;
 struct FMovieSceneSpawnable;
+class UMovieSceneSpawnableBindingBase;
 
 /** Movie scene spawn register that knows how to handle spawning objects (actors) for a level sequence  */
 class FLevelSequenceSpawnRegister : public FMovieSceneSpawnRegister
@@ -24,7 +25,7 @@ public:
 protected:
 	/** ~ FMovieSceneSpawnRegister interface */
 	LEVELSEQUENCE_API virtual UObject* SpawnObject(FMovieSceneSpawnable& Spawnable, FMovieSceneSequenceIDRef TemplateID, TSharedRef<const FSharedPlaybackState> SharedPlaybackState) override;
-	LEVELSEQUENCE_API virtual void DestroySpawnedObject(UObject& Object) override;
+	LEVELSEQUENCE_API virtual void DestroySpawnedObject(UObject& Object, UMovieSceneSpawnableBindingBase* CustomSpawnableBinding) override;
 
 #if WITH_EDITOR
 	LEVELSEQUENCE_API virtual bool CanSpawnObject(UClass* InClass) const override;
diff --git a/Engine/Source/Runtime/MovieScene/Private/Bindings/MovieSceneSpawnableBinding.cpp b/Engine/Source/Runtime/MovieScene/Private/Bindings/MovieSceneSpawnableBinding.cpp
new file mode 100644
--- /dev/null
+++ b/Engine/Source/Runtime/MovieScene/Private/Bindings/MovieSceneSpawnableBinding.cpp
@@ -0,0 +1,141 @@
+// Copyright Epic Games, Inc. All Rights Reserved.
+
+#include "Bindings/MovieSceneSpawnableBinding.h"
+#include "MovieSceneSpawnRegister.h"
+#include "EntitySystem/MovieSceneSharedPlaybackState.h"
+#include "MovieScene.h"
+#include "Engine/World.h"
+#include "Engine/Level.h"
+#include "MovieSceneBindingReferences.h"
+#include "Tracks/MovieSceneBindingLifetimeTrack.h"
+#include "Sections/MovieSceneBindingLifetimeSection.h"
+#include "Styling/SlateBrush.h"
+#include "Styling/AppStyle.h"
+#include "Internationalization/Internationalization.h"
+
+#define LOCTEXT_NAMESPACE "FPossessableModel"
+
+
+UObject* UMovieSceneSpawnableBindingBase::SpawnObject(const FGuid& BindingId, int32 BindingIndex, UMovieScene& MovieScene, FMovieSceneSequenceIDRef TemplateID, TSharedRef<const UE::MovieScene::FSharedPlaybackState> SharedPlaybackState)
+{
+	UWorld* WorldContext = GetWorldContext(SharedPlaybackState);
+
+	if (WorldContext == nullptr)
+	{
+		UE_LOG(LogMovieScene, Warning, TEXT("Can't find world to spawn '%s' into, defaulting to Persistent level"), *MovieScene.GetName());
+
+		WorldContext = GWorld;
+	}
+
+	FName SpawnName = GetSpawnName(BindingId, MovieScene, TemplateID, SharedPlaybackState);
+
+	// If there's an object that already exists with the requested name, it needs to be renamed (it's probably pending kill)
+	if (!SpawnName.IsNone())
+	{
+		UObject* ExistingObject = StaticFindObjectFast(nullptr, WorldContext->PersistentLevel.Get(), SpawnName);
+		if (ExistingObject)
+		{
+			FName DefunctName = MakeUniqueObjectName(WorldContext->PersistentLevel.Get(), ExistingObject->GetClass());
+			ExistingObject->Rename(*DefunctName.ToString(), nullptr, REN_ForceNoResetLoaders);
+		}
+	}
+
+	// Spawn Object Internal
+
+	UObject* SpawnedObject = SpawnObjectInternal(WorldContext, SpawnName, BindingId, BindingIndex, MovieScene, TemplateID, SharedPlaybackState);
+
+	if (!SpawnedObject)
+	{
+		return nullptr;
+	}
+
+#if WITH_EDITOR
+	if (GIsEditor)
+	{
+		// Explicitly set RF_Transactional on spawned objects so we can undo/redo properties on them.
+		SpawnedObject->SetFlags(RF_Transactional);
+	}
+#endif
+
+	// Allows derived classes to perform post-spawn logic such as mesh setup on actors.
+	PostSpawnObject(SpawnedObject, WorldContext, BindingId, BindingIndex, MovieScene, TemplateID, SharedPlaybackState);
+
+	return SpawnedObject;
+}
+
+void UMovieSceneSpawnableBindingBase::DestroySpawnedObject(UObject* Object)
+{
+	if (!Object)
+	{
+		return;
+	}
+#if WITH_EDITOR
+	if (GIsEditor)
+	{
+		// Explicitly remove RF_Transactional on spawned objects since we don't want to transact spawn/destroy events
+		Object->ClearFlags(RF_Transactional);
+	}
+#endif
+
+	DestroySpawnedObjectInternal(Object);
+}
+
+#if WITH_EDITOR
+void UMovieSceneSpawnableBindingBase::SetupDefaults(UObject* SpawnedObject, FGuid ObjectBindingId, UMovieScene& OwnerMovieScene, TSharedRef<const UE::MovieScene::FSharedPlaybackState> SharedPlaybackState)
+{
+	if (!SpawnedObject)
+	{
+		return;
+	}
+	Super::SetupDefaults(SpawnedObject, ObjectBindingId, OwnerMovieScene, SharedPlaybackState);
+	// Ensure it has a binding lifetime track
+	UMovieSceneBindingLifetimeTrack* BindingLifetimeTrack = Cast<UMovieSceneBindingLifetimeTrack>(OwnerMovieScene.FindTrack(UMovieSceneBindingLifetimeTrack::StaticClass(), ObjectBindingId, NAME_None));
+	if (!BindingLifetimeTrack)
+	{
+		BindingLifetimeTrack = Cast<UMovieSceneBindingLifetimeTrack>(OwnerMovieScene.AddTrack(UMovieSceneBindingLifetimeTrack::StaticClass(), ObjectBindingId));
+	}
+
+	if (BindingLifetimeTrack && BindingLifetimeTrack->GetAllSections().IsEmpty())
+	{
+		UMovieSceneBindingLifetimeSection* BindingLifetimeSection = Cast<UMovieSceneBindingLifetimeSection>(BindingLifetimeTrack->CreateNewSection());
+		BindingLifetimeSection->SetRange(TRange<FFrameNumber>::All());
+		BindingLifetimeTrack->AddSection(*BindingLifetimeSection);
+	}
+}
+
+const FSlateBrush* UMovieSceneSpawnableBindingBase::GetBindingTrackCustomIconOverlay() const
+{
+	return FAppStyle::GetBrush("Sequencer.SpawnableIconOverlay");
+}
+
+FText UMovieSceneSpawnableBindingBase::GetBindingTrackIconTooltip() const
+{
+	return LOCTEXT("CustomSpawnableTooltip", "This item is spawned by sequencer by a custom spawnable binding according to this object's binding lifetime track.");
+}
+
+#endif
+
+UWorld* UMovieSceneSpawnableBindingBase::GetWorldContext(TSharedRef<const UE::MovieScene::FSharedPlaybackState> SharedPlaybackState) const
+{
+	UObject* PlaybackContext = SharedPlaybackState->GetPlaybackContext();
+	return PlaybackContext ? PlaybackContext->GetWorld() : nullptr;
+}
+
+FMovieSceneBindingResolveResult UMovieSceneSpawnableBindingBase::ResolveBinding(const FMovieSceneBindingResolveParams& ResolveParams, int32 BindingIndex, TSharedRef<const UE::MovieScene::FSharedPlaybackState> SharedPlaybackState) const
+{
+	FMovieSceneBindingResolveResult Result;
+	const FMovieSceneSpawnRegister* SpawnRegister = SharedPlaybackState->FindCapability<FMovieSceneSpawnRegister>();
+	UObject* SpawnedObject = SpawnRegister ? SpawnRegister->FindSpawnedObject(ResolveParams.ObjectBindingID, ResolveParams.SequenceID, BindingIndex).Get() : nullptr;
+	if (SpawnedObject)
+	{
+		Result.Object = SpawnedObject;
+	}
+	return Result;
+}
+
+const UMovieSceneSpawnableBindingBase* UMovieSceneSpawnableBindingBase::AsSpawnable(TSharedRef<const UE::MovieScene::FSharedPlaybackState> SharedPlaybackState) const
+{
+	return Cast<UMovieSceneSpawnableBindingBase>(this);
+}
+
+#undef LOCTEXT_NAMESPACE
\ No newline at end of file
diff --git a/Engine/Source/Runtime/MovieScene/Private/EntitySystem/BuiltInComponentTypes.cpp b/Engine/Source/Runtime/MovieScene/Private/EntitySystem/BuiltInComponentTypes.cpp
--- a/Engine/Source/Runtime/MovieScene/Private/EntitySystem/BuiltInComponentTypes.cpp
+++ b/Engine/Source/Runtime/MovieScene/Private/EntitySystem/BuiltInComponentTypes.cpp
@@ -412,6 +412,7 @@ FBuiltInComponentTypes::FBuiltInComponentTypes()
 	Tags.PreRoll                 = ComponentRegistry->NewTag(TEXT("Pre Roll"));
 	Tags.SectionPreRoll          = ComponentRegistry->NewTag(TEXT("Section Pre Roll"));
 	Tags.AlwaysCacheInitialValue = ComponentRegistry->NewTag(TEXT("Always Cache Initial Value"));
+	Tags.OldStyleSpawnable		= ComponentRegistry->NewTag(TEXT("Old Style Spawnable"));
 
 	SymbolicTags.CreatesEntities = ComponentRegistry->NewTag(TEXT("~~ SYMBOLIC ~~ Creates Entities"));
 
