Seuquencer: Smart Auto hotkey missing in Sequencer

diff --git a/Engine/Source/Editor/Sequencer/Private/Sequencer.cpp b/Engine/Source/Editor/Sequencer/Private/Sequencer.cpp
--- a/Engine/Source/Editor/Sequencer/Private/Sequencer.cpp
+++ b/Engine/Source/Editor/Sequencer/Private/Sequencer.cpp
@@ -11355,7 +11355,7 @@ void FSequencer::BindCommands()
 		EUIActionRepeatMode::RepeatEnabled);
 
 	SequencerCommandBindings->MapAction(
-		Commands.SetInterpolationCubicAuto,
+		Commands.SetInterpolationCubicSmartAuto,
 		FExecuteAction::CreateSP(this, &FSequencer::SetInterpTangentMode, ERichCurveInterpMode::RCIM_Cubic, ERichCurveTangentMode::RCTM_SmartAuto));
 
 	SequencerCommandBindings->MapAction(
