Sequencer- Return to using the spawn track (for now) for custom spawnable bindings. We still use the binding lifetime track for custom replaceable bindings. Ideally we will change  to using the binding lifetime track for everything once we do some UX work on sections to allow easier splitting and trimming operations.

diff --git a/Engine/Source/Editor/Sequencer/Private/MVVM/ViewModels/PossessableModel.cpp b/Engine/Source/Editor/Sequencer/Private/MVVM/ViewModels/PossessableModel.cpp
--- a/Engine/Source/Editor/Sequencer/Private/MVVM/ViewModels/PossessableModel.cpp
+++ b/Engine/Source/Editor/Sequencer/Private/MVVM/ViewModels/PossessableModel.cpp
@@ -225,10 +225,18 @@ FSlateColor FPossessableModel::GetInvalidBindingLabelColor() const
 		{
 			for (const FMovieSceneBindingReference& BindingReference : BindingReferences->GetReferences(ObjectBindingID))
 			{
-				if (BindingReference.Locator.IsEmpty() && BindingReference.CustomBinding == nullptr)
+				if (BindingReference.Locator.IsEmpty())
 				{
-					// Show empty bindings as yellow rather than red
-					return FLinearColor::Yellow;
+					TSharedPtr<ISequencer> Sequencer = OwnerModel->GetSequencer();
+					if (BindingReference.CustomBinding == nullptr)
+					{
+						// Show empty bindings as yellow rather than red
+						return FLinearColor::Yellow;
+					}
+					else if (BindingReference.CustomBinding->WillSpawnObject(Sequencer->GetSharedPlaybackState()))
+					{
+						return FSlateColor::UseSubduedForeground();
+					}
 				}
 			}
 
diff --git a/Engine/Source/Editor/Sequencer/Private/SequencerUtilities.cpp b/Engine/Source/Editor/Sequencer/Private/SequencerUtilities.cpp
--- a/Engine/Source/Editor/Sequencer/Private/SequencerUtilities.cpp
+++ b/Engine/Source/Editor/Sequencer/Private/SequencerUtilities.cpp
@@ -66,6 +66,7 @@
 #include "Bindings/MovieSceneSpawnableBinding.h"
 #include "Bindings/MovieSceneReplaceableBinding.h"
 #include "ActorFactories/ActorFactory.h"
+#include "Tracks/MovieSceneBindingLifetimeTrack.h"
 
 #define LOCTEXT_NAMESPACE "FSequencerUtilities"
 
@@ -1184,8 +1185,8 @@ FMovieScenePossessable* FSequencerUtilities::ConvertToPossessable(TSharedRef<ISe
 	Sequence->Modify();
 	MovieScene->Modify();
 
-	// If we're converting from an old-style spawnable, we'll need to remove the spawn track
-	if (Spawnable)
+	// If we're converting from a spawnable, we'll need to remove the spawn track
+	if (bConvertFromSpawnable)
 	{
 		// Delete the spawn track
 		UMovieSceneSpawnTrack* SpawnTrack = Cast<UMovieSceneSpawnTrack>(MovieScene->FindTrack(UMovieSceneSpawnTrack::StaticClass(), BindingGuid, NAME_None));
@@ -1315,6 +1316,7 @@ FMovieScenePossessable* FSequencerUtilities::ConvertToCustomBinding(TSharedRef<I
 	// If we have an old-style spawnable, use the template as the object to convert instead.
 	FMovieSceneSpawnable* Spawnable = MovieScene->FindSpawnable(BindingGuid);
 	const FMovieSceneBindingReference* PreviousBindingReference = BindingReferences->GetReference(BindingGuid, BindingIndex);
+	const UMovieSceneCustomBinding* PreviousCustomBinding = nullptr;
 	if (Spawnable)
 	{
 		ObjectToConvert = Spawnable->GetObjectTemplate();
@@ -1324,7 +1326,8 @@ FMovieScenePossessable* FSequencerUtilities::ConvertToCustomBinding(TSharedRef<I
 	{
 		if (const UMovieSceneCustomBinding* CustomBinding = PreviousBindingReference->CustomBinding)
 		{
-			bConvertFromSpawnable = CustomBinding->WillSpawnObject(Sequencer->GetSharedPlaybackState());
+			PreviousCustomBinding = CustomBinding;
+			bConvertFromSpawnable = PreviousCustomBinding->WillSpawnObject(Sequencer->GetSharedPlaybackState());
 		}
 	}
 
@@ -1356,8 +1359,8 @@ FMovieScenePossessable* FSequencerUtilities::ConvertToCustomBinding(TSharedRef<I
 	MovieScene->Modify();
 
 
-	// If we're converting from an old-style spawnable, we'll need to do some extra work and also create a new FMovieScenePossessable*.
-	if (Spawnable)
+	// If we're converting from a spawnable and the new custom binding isn't a spawnable, remove the spawn track
+	if (PreviousCustomBinding->IsA<UMovieSceneSpawnableBindingBase>() && !NewCustomBinding->IsA<UMovieSceneSpawnableBindingBase>())
 	{
 		// Delete the spawn track
 		UMovieSceneSpawnTrack* SpawnTrack = Cast<UMovieSceneSpawnTrack>(MovieScene->FindTrack(UMovieSceneSpawnTrack::StaticClass(), BindingGuid, NAME_None));
@@ -1366,6 +1369,15 @@ FMovieScenePossessable* FSequencerUtilities::ConvertToCustomBinding(TSharedRef<I
 			MovieScene->RemoveTrack(*SpawnTrack);
 		}
 	}
+	else if (PreviousCustomBinding->IsA<UMovieSceneReplaceableBindingBase>() && !NewCustomBinding->IsA<UMovieSceneReplaceableBindingBase>())
+	{
+		// Delete the binding lifetime track
+		UMovieSceneBindingLifetimeTrack* BindingLifetimeTrack = Cast<UMovieSceneBindingLifetimeTrack>(MovieScene->FindTrack(UMovieSceneBindingLifetimeTrack::StaticClass(), BindingGuid, NAME_None));
+		if (BindingLifetimeTrack)
+		{
+			MovieScene->RemoveTrack(*BindingLifetimeTrack);
+		}
+	}
 
 	UE::Sequencer::FCreateBindingParams CreateBindingParams;
 	CreateBindingParams.ReplacementGuid = BindingGuid;
diff --git a/Engine/Source/Runtime/MovieScene/Private/Bindings/MovieSceneSpawnableBinding.cpp b/Engine/Source/Runtime/MovieScene/Private/Bindings/MovieSceneSpawnableBinding.cpp
--- a/Engine/Source/Runtime/MovieScene/Private/Bindings/MovieSceneSpawnableBinding.cpp
+++ b/Engine/Source/Runtime/MovieScene/Private/Bindings/MovieSceneSpawnableBinding.cpp
@@ -12,6 +12,8 @@
 #include "Styling/SlateBrush.h"
 #include "Styling/AppStyle.h"
 #include "Internationalization/Internationalization.h"
+#include "Tracks/MovieSceneSpawnTrack.h"
+#include "Sections/MovieSceneBoolSection.h"
 
 #define LOCTEXT_NAMESPACE "FPossessableModel"
 
@@ -88,18 +90,37 @@ void UMovieSceneSpawnableBindingBase::SetupDefaults(UObject* SpawnedObject, FGui
 		return;
 	}
 	Super::SetupDefaults(SpawnedObject, ObjectBindingId, OwnerMovieScene, SharedPlaybackState);
-	// Ensure it has a binding lifetime track
-	UMovieSceneBindingLifetimeTrack* BindingLifetimeTrack = Cast<UMovieSceneBindingLifetimeTrack>(OwnerMovieScene.FindTrack(UMovieSceneBindingLifetimeTrack::StaticClass(), ObjectBindingId, NAME_None));
-	if (!BindingLifetimeTrack)
+
+	// TODO: For now we are not using binding lifetime track for this, though it will support it. We continue to use spawn track until we improve UX of splitting sections
+	// // 
+	// 
+	//// Ensure it has a binding lifetime track
+	//UMovieSceneBindingLifetimeTrack* BindingLifetimeTrack = Cast<UMovieSceneBindingLifetimeTrack>(OwnerMovieScene.FindTrack(UMovieSceneBindingLifetimeTrack::StaticClass(), ObjectBindingId, NAME_None));
+	//if (!BindingLifetimeTrack)
+	//{
+	//	BindingLifetimeTrack = Cast<UMovieSceneBindingLifetimeTrack>(OwnerMovieScene.AddTrack(UMovieSceneBindingLifetimeTrack::StaticClass(), ObjectBindingId));
+	//}
+
+	//if (BindingLifetimeTrack && BindingLifetimeTrack->GetAllSections().IsEmpty())
+	//{
+	//	UMovieSceneBindingLifetimeSection* BindingLifetimeSection = Cast<UMovieSceneBindingLifetimeSection>(BindingLifetimeTrack->CreateNewSection());
+	//	BindingLifetimeSection->SetRange(TRange<FFrameNumber>::All());
+	//	BindingLifetimeTrack->AddSection(*BindingLifetimeSection);
+	//}
+
+	UMovieSceneSpawnTrack* SpawnTrack = Cast<UMovieSceneSpawnTrack>(OwnerMovieScene.FindTrack(UMovieSceneSpawnTrack::StaticClass(), ObjectBindingId, NAME_None));
+	if (!SpawnTrack)
 	{
-		BindingLifetimeTrack = Cast<UMovieSceneBindingLifetimeTrack>(OwnerMovieScene.AddTrack(UMovieSceneBindingLifetimeTrack::StaticClass(), ObjectBindingId));
+		SpawnTrack = Cast<UMovieSceneSpawnTrack>(OwnerMovieScene.AddTrack(UMovieSceneSpawnTrack::StaticClass(), ObjectBindingId));
 	}
 
-	if (BindingLifetimeTrack && BindingLifetimeTrack->GetAllSections().IsEmpty())
+	if (SpawnTrack)
 	{
-		UMovieSceneBindingLifetimeSection* BindingLifetimeSection = Cast<UMovieSceneBindingLifetimeSection>(BindingLifetimeTrack->CreateNewSection());
-		BindingLifetimeSection->SetRange(TRange<FFrameNumber>::All());
-		BindingLifetimeTrack->AddSection(*BindingLifetimeSection);
+		UMovieSceneBoolSection* SpawnSection = Cast<UMovieSceneBoolSection>(SpawnTrack->CreateNewSection());
+		SpawnSection->GetChannel().SetDefault(true);
+		SpawnSection->SetRange(TRange<FFrameNumber>::All());
+		SpawnTrack->AddSection(*SpawnSection);
+		SpawnTrack->SetObjectId(ObjectBindingId);
 	}
 }
 
diff --git a/Engine/Source/Runtime/MovieScene/Private/Sections/MovieSceneSpawnSection.cpp b/Engine/Source/Runtime/MovieScene/Private/Sections/MovieSceneSpawnSection.cpp
--- a/Engine/Source/Runtime/MovieScene/Private/Sections/MovieSceneSpawnSection.cpp
+++ b/Engine/Source/Runtime/MovieScene/Private/Sections/MovieSceneSpawnSection.cpp
@@ -4,6 +4,7 @@
 #include "UObject/SequencerObjectVersion.h"
 
 #include "MovieScene.h"
+#include "MovieSceneSequence.h"
 #include "MovieSceneTimeHelpers.h"
 #include "Evaluation/MovieSceneEvaluationField.h"
 #include "EntitySystem/MovieSceneEntityManager.h"
@@ -11,6 +12,8 @@
 #include "EntitySystem/MovieSceneEntityBuilder.h"
 #include "EntitySystem/BuiltInComponentTypes.h"
 #include "EntitySystem/MovieSceneSpawnablesSystem.h"
+#include "MovieSceneBindingReferences.h"
+#include "Bindings/MovieSceneSpawnableBinding.h"
 
 #include UE_INLINE_GENERATED_CPP_BY_NAME(MovieSceneSpawnSection)
 
@@ -33,7 +36,7 @@ void UMovieSceneSpawnSection::ImportEntityImpl(UMovieSceneEntitySystemLinker* En
 	OutImportedEntity->AddBuilder(
 		FEntityBuilder()
 		.Add(FBuiltInComponentTypes::Get()->SpawnableBinding, Params.GetObjectBindingID())
-		.AddTag(FBuiltInComponentTypes::Get()->Tags.OldStyleSpawnable)
+		.AddTagConditional(FBuiltInComponentTypes::Get()->Tags.OldStyleSpawnable, GetTypedOuter<UMovieScene>()->FindSpawnable(Params.GetObjectBindingID()) != nullptr)
 	);
 }
 
@@ -41,8 +44,31 @@ bool UMovieSceneSpawnSection::PopulateEvaluationFieldImpl(const TRange<FFrameNum
 {
 	FGuid ObjectBindingID = OutFieldBuilder->GetSharedMetaData().ObjectBindingID;
 
+	UMovieSceneSequence* ParentSequence = GetTypedOuter<UMovieSceneSequence>();
 	UMovieScene* ParentMovieScene = GetTypedOuter<UMovieScene>();
-	if (ParentMovieScene->FindPossessable(ObjectBindingID))
+	
+	bool bSpawnable = false;
+	
+	if (ParentMovieScene && ParentSequence)
+	{
+		if (ParentMovieScene->FindSpawnable(ObjectBindingID))
+		{
+			bSpawnable = true;
+		}
+		else if (const FMovieSceneBindingReferences* BindingReferences = ParentSequence->GetBindingReferences())
+		{
+			if (BindingReferences)
+			{
+				if (Algo::AnyOf(BindingReferences->GetReferences(ObjectBindingID), [](const FMovieSceneBindingReference& BindingReference) { return BindingReference.CustomBinding && BindingReference.CustomBinding->IsA<UMovieSceneSpawnableBindingBase>(); }))
+				{
+					bSpawnable = true;
+				}
+			}
+		}
+
+	}
+
+	if (!bSpawnable)
 	{
 		return true;
 	}
