Sequencer- Fix to static analysis warning in SequencerUtilities.cpp

diff --git a/Engine/Source/Editor/Sequencer/Private/SequencerUtilities.cpp b/Engine/Source/Editor/Sequencer/Private/SequencerUtilities.cpp
--- a/Engine/Source/Editor/Sequencer/Private/SequencerUtilities.cpp
+++ b/Engine/Source/Editor/Sequencer/Private/SequencerUtilities.cpp
@@ -1360,7 +1360,7 @@ FMovieScenePossessable* FSequencerUtilities::ConvertToCustomBinding(TSharedRef<I
 
 
 	// If we're converting from a spawnable and the new custom binding isn't a spawnable, remove the spawn track
-	if (PreviousCustomBinding->IsA<UMovieSceneSpawnableBindingBase>() && !NewCustomBinding->IsA<UMovieSceneSpawnableBindingBase>())
+	if (PreviousCustomBinding && PreviousCustomBinding->IsA<UMovieSceneSpawnableBindingBase>() && !NewCustomBinding->IsA<UMovieSceneSpawnableBindingBase>())
 	{
 		// Delete the spawn track
 		UMovieSceneSpawnTrack* SpawnTrack = Cast<UMovieSceneSpawnTrack>(MovieScene->FindTrack(UMovieSceneSpawnTrack::StaticClass(), BindingGuid, NAME_None));
@@ -1369,7 +1369,7 @@ FMovieScenePossessable* FSequencerUtilities::ConvertToCustomBinding(TSharedRef<I
 			MovieScene->RemoveTrack(*SpawnTrack);
 		}
 	}
-	else if (PreviousCustomBinding->IsA<UMovieSceneReplaceableBindingBase>() && !NewCustomBinding->IsA<UMovieSceneReplaceableBindingBase>())
+	else if (PreviousCustomBinding && PreviousCustomBinding->IsA<UMovieSceneReplaceableBindingBase>() && !NewCustomBinding->IsA<UMovieSceneReplaceableBindingBase>())
 	{
 		// Delete the binding lifetime track
 		UMovieSceneBindingLifetimeTrack* BindingLifetimeTrack = Cast<UMovieSceneBindingLifetimeTrack>(MovieScene->FindTrack(UMovieSceneBindingLifetimeTrack::StaticClass(), BindingGuid, NAME_None));
