Fix raw pointer usage for MovieSceneCopyableBinding after addressing a bug in the enforcement of TObjectPtr usage within engine code.

diff --git a/Engine/Source/Editor/Sequencer/Private/MovieSceneCopyableBinding.h b/Engine/Source/Editor/Sequencer/Private/MovieSceneCopyableBinding.h
--- a/Engine/Source/Editor/Sequencer/Private/MovieSceneCopyableBinding.h
+++ b/Engine/Source/Editor/Sequencer/Private/MovieSceneCopyableBinding.h
@@ -54,5 +54,5 @@ public:
 	TArray<FName> Tags;
 
 	UPROPERTY()
-	TArray<UMovieSceneCustomBinding*> CustomBindings;
+	TArray<TObjectPtr<UMovieSceneCustomBinding>> CustomBindings;
 };
