Sequencer: Add a default marked frame color, if bUseCustomColor is true, then use the custom color as opposed to the default color.

Upgrade old colors by toggling the bUseCustomColor if the deprecated color doesn't match the original color.

Added a toggle to show/hide marked frames.

diff --git a/Engine/Plugins/MovieScene/LevelSequenceEditor/Source/LevelSequenceEditor/Private/CinematicViewport/SCinematicTransportRange.cpp b/Engine/Plugins/MovieScene/LevelSequenceEditor/Source/LevelSequenceEditor/Private/CinematicViewport/SCinematicTransportRange.cpp
--- a/Engine/Plugins/MovieScene/LevelSequenceEditor/Source/LevelSequenceEditor/Private/CinematicViewport/SCinematicTransportRange.cpp
+++ b/Engine/Plugins/MovieScene/LevelSequenceEditor/Source/LevelSequenceEditor/Private/CinematicViewport/SCinematicTransportRange.cpp
@@ -263,6 +263,7 @@ int32 SCinematicTransportRange::OnPaint(const FPaintArgs& Args, const FGeometry&
 	}
 
 	// Draw the marked frames
+	const FLinearColor DefaultMarkedFrameColor = Sequencer->GetSequencerSettings()->GetMarkedFrameColor();
 	TArray<FMovieSceneMarkedFrame> MarkedFrames = Sequencer->GetMarkedFrames();
 	if (MarkedFrames.Num())
 	{
@@ -287,8 +288,9 @@ int32 SCinematicTransportRange::OnPaint(const FPaintArgs& Args, const FGeometry&
 		while ( Index < MarkedFrames.Num() )
 		{
 			FFrameNumber PredicateTime = MarkedFrames[Index].FrameNumber;
+			FLinearColor MarkedFrameColor = MarkedFrames[Index].bUseCustomColor ? MarkedFrames[Index].CustomColor : DefaultMarkedFrameColor;
 			GroupedTimes.Add(PredicateTime);
-			KeyColors.Add(MarkedFrames[Index].Color);
+			KeyColors.Add(MarkedFrameColor);
 			KeyColors[KeyColors.Num()-1].A = 0.8f; // make the alpha consistent across all markers
 
 			while (Index < MarkedFrames.Num() && FMath::Abs(MarkedFrames[Index].FrameNumber - PredicateTime) <= ThresholdFrames)
diff --git a/Engine/Source/Editor/Sequencer/Private/SSequencer.cpp b/Engine/Source/Editor/Sequencer/Private/SSequencer.cpp
--- a/Engine/Source/Editor/Sequencer/Private/SSequencer.cpp
+++ b/Engine/Source/Editor/Sequencer/Private/SSequencer.cpp
@@ -2233,6 +2233,7 @@ TSharedRef<SWidget> SSequencer::MakeViewMenu()
 
 	MenuBuilder.BeginSection("MarkedFrames", LOCTEXT("MarkedFramesHeader", "Marked Frames"));
 	{
+		MenuBuilder.AddMenuEntry(FSequencerCommands::Get().ToggleShowMarkedFrames);
 		MenuBuilder.AddMenuEntry(FSequencerCommands::Get().ToggleShowMarkedFramesGlobally);
 		MenuBuilder.AddMenuEntry(FSequencerCommands::Get().ClearGlobalMarkedFrames);
 	}
diff --git a/Engine/Source/Editor/Sequencer/Private/Sequencer.cpp b/Engine/Source/Editor/Sequencer/Private/Sequencer.cpp
--- a/Engine/Source/Editor/Sequencer/Private/Sequencer.cpp
+++ b/Engine/Source/Editor/Sequencer/Private/Sequencer.cpp
@@ -10629,6 +10629,12 @@ void FSequencer::BindCommands()
 		Commands.SortAllNodesAndDescendants,
 		FExecuteAction::CreateSP(this, &FSequencer::SortAllNodesAndDescendants));
 
+	SequencerCommandBindings->MapAction(
+		Commands.ToggleShowMarkedFrames,
+		FExecuteAction::CreateLambda([this] { Settings->SetShowMarkedFrames(!Settings->GetShowMarkedFrames()); }),
+		FCanExecuteAction::CreateLambda([] { return true; }),
+		FIsActionChecked::CreateLambda([this] { return Settings->GetShowMarkedFrames(); }));
+
 	SequencerCommandBindings->MapAction(
 		Commands.ToggleShowMarkedFramesGlobally,
 		FExecuteAction::CreateSP(this, &FSequencer::ToggleShowMarkedFramesGlobally),
diff --git a/Engine/Source/Editor/Sequencer/Private/SequencerCommands.cpp b/Engine/Source/Editor/Sequencer/Private/SequencerCommands.cpp
--- a/Engine/Source/Editor/Sequencer/Private/SequencerCommands.cpp
+++ b/Engine/Source/Editor/Sequencer/Private/SequencerCommands.cpp
@@ -113,6 +113,7 @@ void FSequencerCommands::RegisterCommands()
 	UI_COMMAND( SetKeyGroup, "Key Group", "Key the groups channels/properties when only one of them changes. ie. Keys all three translation channels when only translation Y changes", EUserInterfaceActionType::RadioButton, FInputChord());
 	UI_COMMAND( SetKeyAll, "Key All", "Key all channels/properties when only one of them changes. ie. Keys all translation, rotation, scale channels when only translation Y changes", EUserInterfaceActionType::RadioButton, FInputChord());
 
+	UI_COMMAND( ToggleShowMarkedFrames, "Show Marked Frames", "Show marked frames.", EUserInterfaceActionType::ToggleButton, FInputChord());
 	UI_COMMAND( ToggleMarkAtPlayPosition, "Toggle Mark", "Sets or clears a mark at the current play position.", EUserInterfaceActionType::Button, FInputChord(EKeys::M) );
 	UI_COMMAND( StepToNextMark, "Step to Next Marked Frame", "Step to the next marked frame", EUserInterfaceActionType::Button, FInputChord(EModifierKey::Control | EModifierKey::Shift, EKeys::Period) );
 	UI_COMMAND( StepToPreviousMark, "Step to Previous Marked Frame", "Step to the previous marked frame", EUserInterfaceActionType::Button, FInputChord(EModifierKey::Control | EModifierKey::Shift, EKeys::Comma) );
diff --git a/Engine/Source/Editor/Sequencer/Private/SequencerSettings.cpp b/Engine/Source/Editor/Sequencer/Private/SequencerSettings.cpp
--- a/Engine/Source/Editor/Sequencer/Private/SequencerSettings.cpp
+++ b/Engine/Source/Editor/Sequencer/Private/SequencerSettings.cpp
@@ -55,6 +55,7 @@ USequencerSettings::USequencerSettings( const FObjectInitializer& ObjectInitiali
 	KeyAreaHeightWithCurves = SequencerLayoutConstants::KeyAreaHeight;
 	bDeleteKeysWhenTrimming = true;
 	bDisableSectionsAfterBaking = true;
+	MarkedFrameColor = FLinearColor(0.f, 1.f, 1.f, 0.4f);
 	bCleanPlaybackMode = true;
 	bActivateRealtimeViewports = true;
 	bEvaluateSubSequencesInIsolation = false;
@@ -67,6 +68,7 @@ USequencerSettings::USequencerSettings( const FObjectInitializer& ObjectInitiali
 	TreeViewWidth = 0.3f;
 	bShowTickLines = true;
 	bShowSequencerToolbar = true;
+	bShowMarkedFrames = true;
 	ViewDensity = "Relaxed";
 
 	SectionColorTints.Add(FColor(88, 102, 142, 255)); // blue
@@ -668,6 +670,20 @@ void USequencerSettings::SetShowSequencerToolbar(bool bInShowSequencerToolbar)
 	}
 }
 
+bool USequencerSettings::GetShowMarkedFrames() const
+{
+	return bShowMarkedFrames;
+}
+
+void USequencerSettings::SetShowMarkedFrames(bool bInShowMarkedFrames)
+{
+	if (bShowMarkedFrames != bInShowMarkedFrames)
+	{
+		bShowMarkedFrames = bInShowMarkedFrames;
+		SaveConfig();
+	}
+}
+
 bool USequencerSettings::HasKeyAreaCurveExtents(const FString& ChannelName) const
 {
 	TArray<FString> ChannelsArray;
@@ -796,6 +812,20 @@ void USequencerSettings::SetDisableSectionsAfterBaking(bool bInDisableSectionsAf
 	}
 }
 
+FLinearColor USequencerSettings::GetMarkedFrameColor() const
+{
+	return MarkedFrameColor;
+}
+
+void USequencerSettings::SetMarkedFrameColor(const FLinearColor& InMarkedFrameColor)
+{
+	if (MarkedFrameColor != InMarkedFrameColor)
+	{
+		MarkedFrameColor = InMarkedFrameColor;
+		SaveConfig();
+	}
+}
+
 TArray<FColor> USequencerSettings::GetSectionColorTints() const
 {
 	return SectionColorTints;
diff --git a/Engine/Source/Editor/Sequencer/Private/SequencerTimeSliderController.cpp b/Engine/Source/Editor/Sequencer/Private/SequencerTimeSliderController.cpp
--- a/Engine/Source/Editor/Sequencer/Private/SequencerTimeSliderController.cpp
+++ b/Engine/Source/Editor/Sequencer/Private/SequencerTimeSliderController.cpp
@@ -325,6 +325,11 @@ void FSequencerTimeSliderController::DrawTicks( FSlateWindowElementList& OutDraw
 
 int32 FSequencerTimeSliderController::DrawMarkedFrames( const FGeometry& AllottedGeometry, const FScrubRangeToScreen& RangeToScreen, FSlateWindowElementList& OutDrawElements, int32 LayerId, const ESlateDrawEffect& DrawEffects, const FWidgetStyle& InWidgetStyle, bool bDrawLabels ) const
 {
+	if (!WeakSequencer.Pin()->GetSequencerSettings()->GetShowMarkedFrames())
+	{
+		return LayerId;
+	}
+
 	using namespace UE::Sequencer;
 
 	const TArray<FMovieSceneMarkedFrame> & MarkedFrames = TimeSliderArgs.MarkedFrames.Get();
@@ -337,6 +342,7 @@ int32 FSequencerTimeSliderController::DrawMarkedFrames( const FGeometry& Allotte
 	const FMarkedFrameSelection& SelectedMarkedFrames = WeakSequencer.Pin()->GetViewModel()->GetSelection()->MarkedFrames;
 	//const FLinearColor SelectedColor = FLinearColor::White; //FAppStyle::GetSlateColor("SelectionColor").GetColor(InWidgetStyle);
 	const FLinearColor WhiteColorHSV = FLinearColor::White.LinearRGBToHSV();
+	const FLinearColor DefaultMarkedFrameColor = WeakSequencer.Pin()->GetSequencerSettings()->GetMarkedFrameColor();
 
 	auto DrawFrameMarkers = ([=, this](const TArray<FMovieSceneMarkedFrame> & InMarkedFrames, FSlateWindowElementList& DrawElements, bool bIsGlobal)
 	{
@@ -350,7 +356,8 @@ int32 FSequencerTimeSliderController::DrawMarkedFrames( const FGeometry& Allotte
 
 			// Get a selected color that's the marked frame color but at full opacity, full brightness, and a bit desaturated if it's
 			// already bright to begin with.
-			const FLinearColor MarkedFrameColorHSV = MarkedFrame.Color.CopyWithNewOpacity(1.f).LinearRGBToHSV();
+			const FLinearColor MarkedFrameColor = MarkedFrame.bUseCustomColor ? MarkedFrame.CustomColor : DefaultMarkedFrameColor;
+			const FLinearColor MarkedFrameColorHSV = MarkedFrameColor.CopyWithNewOpacity(1.f).LinearRGBToHSV();
 			const FLinearColor SelectedColor = FLinearColor::LerpUsingHSV(
 				MarkedFrameColorHSV, 
 				FLinearColor(
@@ -360,7 +367,7 @@ int32 FSequencerTimeSliderController::DrawMarkedFrames( const FGeometry& Allotte
 				0.5f)
 				.HSVToLinearRGB();
 
-			FLinearColor DrawColor = bIsGlobal ? MarkedFrame.Color.Desaturate(0.25f) : (bIsSelected ? SelectedColor : MarkedFrame.Color);
+			FLinearColor DrawColor = bIsGlobal ? MarkedFrameColor.Desaturate(0.25f) : (bIsSelected ? SelectedColor : MarkedFrameColor);
 			const float  LinePos = RangeToScreen.InputToLocalX(Seconds);
 			TArray<FVector2D> LinePoints;
 			LinePoints.AddUninitialized(2);
diff --git a/Engine/Source/Editor/Sequencer/Public/SequencerCommands.h b/Engine/Source/Editor/Sequencer/Public/SequencerCommands.h
--- a/Engine/Source/Editor/Sequencer/Public/SequencerCommands.h
+++ b/Engine/Source/Editor/Sequencer/Public/SequencerCommands.h
@@ -286,6 +286,9 @@ public:
 	/** Set mode to key all. */
 	TSharedPtr< FUICommandInfo > SetKeyAll;
 
+	/** Toggle showing marked frames */
+	TSharedPtr< FUICommandInfo> ToggleShowMarkedFrames;
+
 	/** Toggle on/off a Mark at the current time **/
 	TSharedPtr< FUICommandInfo> ToggleMarkAtPlayPosition;
 
diff --git a/Engine/Source/Editor/Sequencer/Public/SequencerSettings.h b/Engine/Source/Editor/Sequencer/Public/SequencerSettings.h
--- a/Engine/Source/Editor/Sequencer/Public/SequencerSettings.h
+++ b/Engine/Source/Editor/Sequencer/Public/SequencerSettings.h
@@ -342,6 +342,11 @@ public:
 	/** Set whether to show sequencer toolbar bar */
 	void SetShowSequencerToolbar(bool bInDrawTickLines);
 
+	/** @return true if showing marked frames */
+	bool GetShowMarkedFrames() const;
+	/** Set whether to show marked frames */
+	void SetShowMarkedFrames(bool bShowMarkedFrames);
+
 	/** @return Whether the given channel has curve extents */
 	bool HasKeyAreaCurveExtents(const FString& ChannelName) const;
 	/** @ Remove curve extents for the given channel */
@@ -371,6 +376,11 @@ public:
 	/** Set whether to disable sections when baking, as opposed to deleting */
 	void SetDisableSectionsAfterBaking(bool bInDisableSectionsAfterBaking);
 
+	/** @return the default marked frame color */
+	FLinearColor GetMarkedFrameColor() const;
+	/** Set the default marked frame color */
+	void SetMarkedFrameColor(const FLinearColor& InColor);
+
 	/** @return the section color tints */
 	TArray<FColor> GetSectionColorTints() const;
 	/** Set the section color tints */
@@ -650,6 +660,10 @@ protected:
 	UPROPERTY(config, EditAnywhere, Category = Timeline)
 	bool bShowSequencerToolbar;
 
+	/** Enable or disable showing marked frames */
+	UPROPERTY(config, EditAnywhere, Category = Timeline)
+	bool bShowMarkedFrames;
+
 	/** The key area curve extents, stored per channel name */
 	UPROPERTY(config, EditAnywhere, Category = Timeline)
 	FString KeyAreaCurveExtents;
@@ -669,6 +683,10 @@ protected:
 	/** Whether to disable sections after baking as opposed to deleting. */
 	UPROPERTY(config, EditAnywhere, Category = Timeline)
 	bool bDisableSectionsAfterBaking;
+	
+	/** Default marked frame color */
+	UPROPERTY(config, EditAnywhere, Category = Timeline)
+	FLinearColor MarkedFrameColor;
 
 	/** Section color tints */
 	UPROPERTY(config, EditAnywhere, Category = General)
diff --git a/Engine/Source/Runtime/MovieScene/Private/MovieScene.cpp b/Engine/Source/Runtime/MovieScene/Private/MovieScene.cpp
--- a/Engine/Source/Runtime/MovieScene/Private/MovieScene.cpp
+++ b/Engine/Source/Runtime/MovieScene/Private/MovieScene.cpp
@@ -105,6 +105,19 @@ void UMovieScene::PostLoad()
 {
 	SortMarkedFrames();
 
+#if WITH_EDITORONLY_DATA
+	for (FMovieSceneMarkedFrame& MarkedFrame : MarkedFrames)
+	{
+		const FLinearColor DefaultDeprecatedColor(0.f, 1.f, 1.f, 0.4f);
+
+		if (MarkedFrame.Color_DEPRECATED != DefaultDeprecatedColor)
+		{
+			MarkedFrame.bUseCustomColor = true;
+			MarkedFrame.CustomColor = MarkedFrame.Color_DEPRECATED;
+		}
+	}
+#endif
+
 	Super::PostLoad();
 }
 
diff --git a/Engine/Source/Runtime/MovieScene/Public/MovieSceneMarkedFrame.h b/Engine/Source/Runtime/MovieScene/Public/MovieSceneMarkedFrame.h
--- a/Engine/Source/Runtime/MovieScene/Public/MovieSceneMarkedFrame.h
+++ b/Engine/Source/Runtime/MovieScene/Public/MovieSceneMarkedFrame.h
@@ -17,7 +17,7 @@ struct FMovieSceneMarkedFrame
 		: Label(FString())
 #if WITH_EDITORONLY_DATA
 		, Comment(FString())
-		, Color(0.f, 1.f, 1.f, 0.4f)
+		, CustomColor(0.f, 1.f, 1.f, 0.4f)	
 #endif
 		, bIsDeterminismFence(false)
 	{}
@@ -27,7 +27,7 @@ struct FMovieSceneMarkedFrame
 		, Label(FString())
 #if WITH_EDITORONLY_DATA
 		, Comment(FString())
-		, Color(0.f, 1.f, 1.f, 0.4f)
+		, CustomColor(0.f, 1.f, 1.f, 0.4f)
 #endif
 		, bIsDeterminismFence(false)
 	{}
@@ -42,8 +42,14 @@ struct FMovieSceneMarkedFrame
 	UPROPERTY(EditAnywhere, Category = "Marked Frame")
 	FString Comment;
 
+	UPROPERTY(EditAnywhere, Category = "Marked Frame", DisplayName = "Color", meta=(EditCondition = bUseCustomColor))
+	FLinearColor CustomColor;
+
 	UPROPERTY(EditAnywhere, Category = "Marked Frame")
-	FLinearColor Color;
+	bool bUseCustomColor = false;
+
+	UPROPERTY()
+	FLinearColor Color_DEPRECATED = FLinearColor(0.f, 1.f, 1.f, 0.4f);
 #endif
 
 	UPROPERTY(EditAnywhere, BlueprintReadWrite, Category = "Marked Frame", DisplayName="Is Determinism Fence?", meta=(Tooltip="When checked, treat this mark as a fence for evaluation purposes. Fences cannot be crossed in a single evaluation, and force the evaluation to be split into 2 separate parts."))
