Sequencer: Remove selection changed

Speculative fix for this crash which seems to come in from FSelectionEventSuppressor which broadcasts a selection changed when the Sequencer panel is destroyed.

diff --git a/Engine/Source/Editor/Sequencer/Private/Sequencer.cpp b/Engine/Source/Editor/Sequencer/Private/Sequencer.cpp
--- a/Engine/Source/Editor/Sequencer/Private/Sequencer.cpp
+++ b/Engine/Source/Editor/Sequencer/Private/Sequencer.cpp
@@ -744,6 +744,11 @@ FSequencer::~FSequencer()
 
 	AcquiredResources.Release();
 
+	if (ViewModel && ViewModel->GetSelection())
+	{
+		ViewModel->GetSelection()->OnChanged.RemoveAll(this);
+	}
+
 	ViewModel.Reset();
 	SequencerWidget.Reset();
 
