Sequencer- UX Improvements for the Binding Properties menu.
* You can now convert individual bindings inside a single object binding track back and forth from spawnable, replaceable, possessable, etc.
* You can save default state for individual spawnables and replacables within a single object binding track
* You can now change class for individual spawnable/replaceables within a single object binding track
* Fixed a few bugs found during this refactor.

diff --git a/Engine/Plugins/MovieScene/LevelSequenceEditor/Source/LevelSequenceEditor/Private/LevelSequenceEditorSubsystem.cpp b/Engine/Plugins/MovieScene/LevelSequenceEditor/Source/LevelSequenceEditor/Private/LevelSequenceEditorSubsystem.cpp
--- a/Engine/Plugins/MovieScene/LevelSequenceEditor/Source/LevelSequenceEditor/Private/LevelSequenceEditorSubsystem.cpp
+++ b/Engine/Plugins/MovieScene/LevelSequenceEditor/Source/LevelSequenceEditor/Private/LevelSequenceEditorSubsystem.cpp
@@ -60,6 +60,13 @@
 #include "Bindings/MovieSceneSpawnableActorBinding.h"
 #include "Bindings/MovieSceneSpawnableActorBindingCustomization.h"
 #include "Bindings/MovieSceneCustomBinding.h"
+#include "DetailWidgetRow.h"
+#include "IPropertyTypeCustomization.h"
+#include "IDetailChildrenBuilder.h"
+#include "Widgets/Input/SComboBox.h"
+#include "DetailLayoutBuilder.h"
+#include "IPropertyTypeCustomization.h"
+#include "IPropertyUtilities.h"
 
 #include UE_INLINE_GENERATED_CPP_BY_NAME(LevelSequenceEditorSubsystem)
 
@@ -67,6 +74,259 @@ DEFINE_LOG_CATEGORY(LogLevelSequenceEditor);
 
 #define LOCTEXT_NAMESPACE "LevelSequenceEditor"
 
+
+class FMovieSceneBindingPropertyInfoDetailCustomization : public IPropertyTypeCustomization
+{
+public:
+	FMovieSceneBindingPropertyInfoDetailCustomization(TWeakPtr<ISequencer> InSequencer, UMovieScene* InMovieScene, FGuid InBindingGuid, ULevelSequenceEditorSubsystem* InLevelSequenceEditorSubsystem);
+
+	// Begin IDetailCustomization interface
+	virtual void CustomizeHeader(TSharedRef<IPropertyHandle> PropertyHandle, FDetailWidgetRow& HeaderRow, IPropertyTypeCustomizationUtils& CustomizationUtils) override;
+	virtual void CustomizeChildren(TSharedRef<IPropertyHandle> PropertyHandle, IDetailChildrenBuilder& ChildBuilder, IPropertyTypeCustomizationUtils& CustomizationUtils) override;
+	//~ End IDetailCustomization interface
+
+private:
+
+	TSharedRef<SWidget> OnGetConvertMenuContent(IDetailChildrenBuilder& StructBuilder, IPropertyTypeCustomizationUtils& CustomizationUtils);
+	TSharedRef<SWidget> OnGetChangeClassMenuContent(IDetailChildrenBuilder& StructBuilder, IPropertyTypeCustomizationUtils& CustomizationUtils);
+
+	TWeakPtr<ISequencer> SequencerPtr;
+	TObjectPtr<UMovieScene> MovieScene = nullptr;
+	TObjectPtr<ULevelSequenceEditorSubsystem> LevelSequenceEditorSubsystem;
+	FGuid BindingGuid;
+	int32 BindingIndex = 0;
+
+	TSharedPtr<IPropertyHandle> StructPropertyHandle;
+
+	TArray<TSharedPtr<FText>> BindingTypeNames;
+};
+
+FMovieSceneBindingPropertyInfoDetailCustomization::FMovieSceneBindingPropertyInfoDetailCustomization(TWeakPtr<ISequencer> InSequencer, UMovieScene* InMovieScene, FGuid InBindingGuid, ULevelSequenceEditorSubsystem* InLevelSequenceEditorSubsystem)
+	: SequencerPtr(InSequencer)
+	, MovieScene(InMovieScene)
+	, LevelSequenceEditorSubsystem(InLevelSequenceEditorSubsystem)
+	, BindingGuid(InBindingGuid)
+{
+
+}
+
+void FMovieSceneBindingPropertyInfoDetailCustomization::CustomizeHeader(TSharedRef<IPropertyHandle> InStructPropertyHandle, FDetailWidgetRow& HeaderRow, IPropertyTypeCustomizationUtils& CustomizationUtils)
+{
+	TSharedPtr<IPropertyUtilities> PropertyUtils = CustomizationUtils.GetPropertyUtilities();
+
+	HeaderRow.NameContent()
+		[
+			InStructPropertyHandle->CreatePropertyNameWidget()
+		]
+	.ValueContent()
+		[
+			InStructPropertyHandle->CreatePropertyValueWidget()
+		];
+
+	StructPropertyHandle = InStructPropertyHandle;
+}
+
+void FMovieSceneBindingPropertyInfoDetailCustomization::CustomizeChildren(TSharedRef<IPropertyHandle> InStructPropertyHandle, IDetailChildrenBuilder& StructBuilder, IPropertyTypeCustomizationUtils& CustomizationUtils)
+{
+	TSharedPtr<ISequencer> Sequencer = SequencerPtr.Pin();
+	if (Sequencer.IsValid() && InStructPropertyHandle->IsValidHandle())
+	{
+		UMovieSceneSequence* Sequence = MovieScene->GetTypedOuter<UMovieSceneSequence>();
+		TArray<void*> StructPtrs;
+		StructPropertyHandle->AccessRawData(StructPtrs);
+		FMovieSceneBindingPropertyInfo* BindingPropertyInfo = (StructPtrs.Num() == 1) ? reinterpret_cast<FMovieSceneBindingPropertyInfo*>(StructPtrs[0]) : nullptr;
+		if (BindingPropertyInfo)
+		{
+			BindingIndex = InStructPropertyHandle->GetArrayIndex();
+			int32 InitialIndex = 0;
+			BindingTypeNames.Add(MakeShared<FText>(LOCTEXT("BindingType_Possessable", "Possessable")));
+			TArrayView<const TSubclassOf<UMovieSceneCustomBinding>> SupportedBindingTypes = Sequencer->GetSupportedCustomBindingTypes();
+			int32 BindingTypeIndex = 1;
+			for (const TSubclassOf<UMovieSceneCustomBinding>& CustomBindingType : SupportedBindingTypes)
+			{
+				if (CustomBindingType)
+				{
+					if (BindingPropertyInfo->CustomBinding && BindingPropertyInfo->CustomBinding->GetClass() == CustomBindingType)
+					{
+						InitialIndex = BindingTypeIndex;
+					}
+					BindingTypeNames.Add(MakeShared<FText>(CustomBindingType->GetDefaultObject<UMovieSceneCustomBinding>()->GetBindingTypePrettyName()));
+				}
+				BindingTypeIndex++;
+			}
+
+			FDetailWidgetRow& BindingTypeRow = StructBuilder.AddCustomRow(LOCTEXT("BindingPropertyType", "Binding Type"));
+			BindingTypeRow.NameContent()
+				[
+					SNew(STextBlock)
+					.Text(LOCTEXT("BindingPropertyType", "Binding Type"))
+					.ToolTipText(LOCTEXT("BindingPropertyType_Tooltip", "The type of binding for this object binding track entry"))
+					.Font(IDetailLayoutBuilder::GetDetailFont())
+				];
+			BindingTypeRow.ValueContent()
+				[
+					SNew(SComboBox<TSharedPtr<FText>>)
+					.OptionsSource(&BindingTypeNames)
+					.InitiallySelectedItem(BindingTypeNames[InitialIndex])
+					.OnSelectionChanged_Lambda([this, InStructPropertyHandle, SupportedBindingTypes, BindingPropertyInfo, &StructBuilder, &CustomizationUtils](TSharedPtr<FText> NewValue, ESelectInfo::Type)
+					{
+						int32 CurrentIndex = BindingTypeNames.IndexOfByPredicate([&](const TSharedPtr<FText>& In) {
+							return In == NewValue;
+							});
+						if (CurrentIndex == 0)
+						{
+							// Possessable- clear everything
+							BindingPropertyInfo->ResolveFlags = ELocatorResolveFlags::None;
+							BindingPropertyInfo->Locator = FUniversalObjectLocator();
+							BindingPropertyInfo->CustomBinding = nullptr;
+						}
+						else
+						{
+							// Custom binding
+							UClass* BindingTypeClass = SupportedBindingTypes[CurrentIndex - 1];
+							BindingPropertyInfo->Locator = FUniversalObjectLocator();
+							BindingPropertyInfo->ResolveFlags = ELocatorResolveFlags::None;
+							BindingPropertyInfo->CustomBinding = NewObject<UMovieSceneCustomBinding>(MovieScene.Get(), BindingTypeClass, NAME_None, EObjectFlags::RF_Transient);
+						}
+
+						// Force reset the view
+						TSharedPtr<IPropertyHandle> LocatorProperty = StructPropertyHandle->GetChildHandle(GET_MEMBER_NAME_CHECKED(FMovieSceneBindingPropertyInfo, Locator));
+						if (LocatorProperty.IsValid())
+						{
+							CustomizationUtils.GetPropertyUtilities()->NotifyFinishedChangingProperties(FPropertyChangedEvent(LocatorProperty->GetProperty()));
+						}
+						StructBuilder.GetParentCategory().GetParentLayout().ForceRefreshDetails();
+
+					})
+				.OnGenerateWidget_Lambda([](TSharedPtr<FText> Item)
+					{
+						return SNew(STextBlock).Text(*Item.Get());
+					})
+				.InitiallySelectedItem(BindingTypeNames[InitialIndex])
+				[
+					SNew(STextBlock)
+					.Margin(FMargin(0.0f, 2.0f))
+					.Text_Lambda([this, InitialIndex]() { return *BindingTypeNames[InitialIndex].Get(); })
+				]
+			];
+
+			// Only show certain menus if we have a currently bound object
+			
+			if (MovieSceneHelpers::GetSingleBoundObject(Sequence, BindingGuid, Sequencer->GetSharedPlaybackState(), BindingIndex) != nullptr)
+			{
+				FDetailWidgetRow& ConvertToRow = StructBuilder.AddCustomRow(LOCTEXT("ConvertBindingTo", "Convert Binding To..."));
+				ConvertToRow.WholeRowContent()
+					[
+						SNew(SComboButton)
+						.OnGetMenuContent(FOnGetContent::CreateLambda([this, &StructBuilder, &CustomizationUtils]() { return OnGetConvertMenuContent(StructBuilder, CustomizationUtils); }))
+					.ComboButtonStyle(FAppStyle::Get(), "SimpleComboButton")
+					.ButtonContent()
+					[
+						SNew(STextBlock).Text(LOCTEXT("ConvertBindingTo", "Convert Binding To..."))
+					]
+					];
+
+				// Show Change class and save default state menus for spawnables
+				if (MovieSceneHelpers::IsBoundToSpawnable(Sequence, BindingGuid, Sequencer->GetSharedPlaybackState(), BindingIndex))
+				{
+					FDetailWidgetRow& ChangeClassRow = StructBuilder.AddCustomRow(LOCTEXT("ChangeClass", "Change Class..."));
+					ChangeClassRow.WholeRowContent()
+						[
+							SNew(SComboButton)
+							.OnGetMenuContent(FOnGetContent::CreateLambda([this, &StructBuilder, &CustomizationUtils]() { return OnGetChangeClassMenuContent(StructBuilder, CustomizationUtils); }))
+						.ComboButtonStyle(FAppStyle::Get(), "SimpleComboButton")
+						.ButtonContent()
+						[
+							SNew(STextBlock).Text(LOCTEXT("ChangeClass", "Change Class..."))
+						]
+						];
+
+					// Save Default State
+					FDetailWidgetRow& SaveDefaultStateRow = StructBuilder.AddCustomRow(LOCTEXT("SaveDefaultState", "Save Default State"));
+					SaveDefaultStateRow.WholeRowContent()
+						[
+							SNew(SButton)
+							.Text(LOCTEXT("SaveDefaultState", "Save Default State"))
+						.ToolTipText(LOCTEXT("SaveDefaultState_Tooltip", "Save the current state of this spawnable as default properties"))
+						.OnClicked_Lambda([this]()
+							{
+								TSharedPtr<ISequencer> Sequencer = SequencerPtr.Pin();
+								if (Sequencer.IsValid())
+								{
+									Sequencer->GetSpawnRegister().SaveDefaultSpawnableState(BindingGuid, BindingIndex, Sequencer->GetFocusedTemplateID(), Sequencer->GetSharedPlaybackState());
+								}
+								return FReply::Handled();
+							})
+						];
+				}
+			}
+
+
+
+			if (InitialIndex == 0)
+			{
+				// Show locator property
+				TSharedPtr<IPropertyHandle> LocatorProperty = StructPropertyHandle->GetChildHandle(GET_MEMBER_NAME_CHECKED(FMovieSceneBindingPropertyInfo, Locator));
+				if (LocatorProperty.IsValid())
+				{
+					StructBuilder.AddProperty(LocatorProperty.ToSharedRef());
+				}
+			}
+			else
+			{
+				// Show instanced binding type property
+				TSharedPtr<IPropertyHandle> CustomBindingProperty = StructPropertyHandle->GetChildHandle(GET_MEMBER_NAME_CHECKED(FMovieSceneBindingPropertyInfo, CustomBinding));
+				if (CustomBindingProperty.IsValid())
+				{
+					StructBuilder.AddProperty(CustomBindingProperty.ToSharedRef()).CustomWidget(true)
+						.NameContent()
+						[
+							SNew(STextBlock)
+							.Text(LOCTEXT("BindingProperties", "Binding Properties"))
+							.Font(IDetailLayoutBuilder::GetDetailFont())
+						];
+				}
+			}
+		}
+	}
+}
+
+TSharedRef<SWidget> FMovieSceneBindingPropertyInfoDetailCustomization::OnGetConvertMenuContent(IDetailChildrenBuilder& StructBuilder, IPropertyTypeCustomizationUtils& CustomizationUtils)
+{
+	FMenuBuilder MenuBuilder(true, nullptr, nullptr, true);
+
+	TSharedPtr<ISequencer> Sequencer = SequencerPtr.Pin();
+	if (Sequencer.IsValid() && StructPropertyHandle->IsValidHandle())
+	{
+		FSequencerUtilities::AddConvertBindingMenu(MenuBuilder, Sequencer.ToSharedRef(), BindingGuid, BindingIndex, [this, &StructBuilder, &CustomizationUtils]()
+			{
+				if (IDetailsView* DetailsView = StructBuilder.GetParentCategory().GetParentLayout().GetDetailsView())
+				{
+					LevelSequenceEditorSubsystem->RefreshBindingDetails(DetailsView, BindingGuid);
+				}				
+			});
+	}
+	return MenuBuilder.MakeWidget();
+}
+
+TSharedRef<SWidget> FMovieSceneBindingPropertyInfoDetailCustomization::OnGetChangeClassMenuContent(IDetailChildrenBuilder& StructBuilder, IPropertyTypeCustomizationUtils& CustomizationUtils)
+{
+	FMenuBuilder MenuBuilder(true, nullptr, nullptr, true);
+
+	TSharedPtr<ISequencer> Sequencer = SequencerPtr.Pin();
+	if (Sequencer.IsValid() && StructPropertyHandle->IsValidHandle())
+	{
+		FSequencerUtilities::AddChangeClassMenu(MenuBuilder, Sequencer.ToSharedRef(), BindingGuid, BindingIndex, [this, &StructBuilder, &CustomizationUtils]()
+			{
+				if (IDetailsView* DetailsView = StructBuilder.GetParentCategory().GetParentLayout().GetDetailsView())
+				{
+					LevelSequenceEditorSubsystem->RefreshBindingDetails(DetailsView, BindingGuid);
+				}
+			});
+	}
+	return MenuBuilder.MakeWidget();
+}
+
 void ULevelSequenceEditorSubsystem::Initialize(FSubsystemCollectionBase& Collection)
 {
 	UE_LOG(LogLevelSequenceEditor, Log, TEXT("LevelSequenceEditor subsystem initialized."));
@@ -263,6 +523,11 @@ void ULevelSequenceEditorSubsystem::AddBindingDetailCustomizations(TSharedRef<ID
 		if (MovieScene)
 		{
 			FPropertyEditorModule& PropertyEditor = FModuleManager::Get().LoadModuleChecked<FPropertyEditorModule>(TEXT("PropertyEditor"));
+			DetailsView->RegisterInstancedCustomPropertyTypeLayout(FMovieSceneBindingPropertyInfo::StaticStruct()->GetFName(), FOnGetPropertyTypeCustomizationInstance::CreateLambda([](TWeakPtr<ISequencer> InSequencer, UMovieScene* InMovieScene, FGuid InBindingGuid, ULevelSequenceEditorSubsystem* LevelSequenceEditorSubsystem)
+				{
+					return MakeShared<FMovieSceneBindingPropertyInfoDetailCustomization>(InSequencer, InMovieScene, InBindingGuid, LevelSequenceEditorSubsystem);
+				}, ActiveSequencer.ToWeakPtr(), MovieScene, BindingGuid, this));
+			
 			DetailsView->RegisterInstancedCustomPropertyLayout(UMovieSceneSpawnableActorBinding::StaticClass(), FOnGetDetailCustomizationInstance::CreateStatic(&FMovieSceneSpawnableActorBindingBaseCustomization::MakeInstance, ActiveSequencer.ToWeakPtr(), MovieScene, BindingGuid));
 		}
 	}
@@ -1874,18 +2139,6 @@ void ULevelSequenceEditorSubsystem::AddBindingPropertiesMenu(FMenuBuilder& MenuB
 	}
 	if (FMovieSceneBindingReferences* BindingReferences = Sequence->GetBindingReferences())
 	{
-		BindingPropertyInfoList = NewObject<UMovieSceneBindingPropertyInfoList>();
-
-		Algo::Transform(BindingReferences->GetReferences(ObjectBindings[0]), BindingPropertyInfoList->Bindings, [this](const FMovieSceneBindingReference& Reference)
-			{ 
-				UMovieSceneCustomBinding* CopiedBinding = nullptr;
-				if (Reference.CustomBinding)
-				{
-					CopiedBinding = Cast<UMovieSceneCustomBinding>(StaticDuplicateObject(Reference.CustomBinding, BindingPropertyInfoList.Get()));
-				}
-				return FMovieSceneBindingPropertyInfo{ Reference.Locator, Reference.ResolveFlags, Reference.CustomBinding };
-			});
-
 		MenuBuilder.AddMenuSeparator();
 
 		NotifyHook = FBindingPropertiesNotifyHook(Sequence);
@@ -1906,12 +2159,11 @@ void ULevelSequenceEditorSubsystem::AddBindingPropertiesMenu(FMenuBuilder& MenuB
 			DetailsViewArgs.NotifyHook = &NotifyHook;
 		}
 
-
 		TSharedRef<IDetailsView> DetailsView = FModuleManager::GetModuleChecked<FPropertyEditorModule>("PropertyEditor").CreateDetailView(DetailsViewArgs);
 
 		AddBindingDetailCustomizations(DetailsView, Sequencer, ObjectBindings[0]);
 
-		DetailsView->SetObject(BindingPropertyInfoList.Get(), true);
+		RefreshBindingDetails(&DetailsView.Get(), ObjectBindings[0]);
 		DetailsView->OnFinishedChangingProperties().AddUObject(this, &ULevelSequenceEditorSubsystem::OnFinishedChangingLocators, DetailsView, ObjectBindings[0]);
 
 		MenuBuilder.AddWidget(DetailsView, FText::GetEmpty(), true);
@@ -2240,4 +2492,38 @@ void ULevelSequenceEditorSubsystem::RebindComponent(const TArray<FMovieSceneBind
 	}
 }
 
+void ULevelSequenceEditorSubsystem::RefreshBindingDetails(IDetailsView* DetailsView, FGuid ObjectBindingID)
+{
+	if (DetailsView == nullptr)
+	{
+		return;
+	}
+	TSharedPtr<ISequencer> Sequencer = GetActiveSequencer();
+	if (Sequencer == nullptr)
+	{
+		return;
+	}
+
+	UMovieSceneSequence* Sequence = Sequencer->GetFocusedMovieSceneSequence();
+	if (!Sequence)
+	{
+		return;
+	}
+	if (FMovieSceneBindingReferences* BindingReferences = Sequence->GetBindingReferences())
+	{
+		BindingPropertyInfoList = NewObject<UMovieSceneBindingPropertyInfoList>();
+		Algo::Transform(BindingReferences->GetReferences(ObjectBindingID), BindingPropertyInfoList->Bindings, [this](const FMovieSceneBindingReference& Reference)
+			{
+				UMovieSceneCustomBinding* CopiedBinding = nullptr;
+				if (Reference.CustomBinding)
+				{
+					CopiedBinding = Cast<UMovieSceneCustomBinding>(StaticDuplicateObject(Reference.CustomBinding, BindingPropertyInfoList.Get()));
+				}
+				return FMovieSceneBindingPropertyInfo{ Reference.Locator, Reference.ResolveFlags, Reference.CustomBinding };
+			});
+
+		DetailsView->SetObject(BindingPropertyInfoList.Get(), true);
+	}
+}
+
 #undef LOCTEXT_NAMESPACE
\ No newline at end of file
diff --git a/Engine/Plugins/MovieScene/LevelSequenceEditor/Source/LevelSequenceEditor/Private/Misc/LevelSequenceCustomization.cpp b/Engine/Plugins/MovieScene/LevelSequenceEditor/Source/LevelSequenceEditor/Private/Misc/LevelSequenceCustomization.cpp
--- a/Engine/Plugins/MovieScene/LevelSequenceEditor/Source/LevelSequenceEditor/Private/Misc/LevelSequenceCustomization.cpp
+++ b/Engine/Plugins/MovieScene/LevelSequenceEditor/Source/LevelSequenceEditor/Private/Misc/LevelSequenceCustomization.cpp
@@ -200,7 +200,7 @@ void FLevelSequenceCustomization::ExtendObjectBindingContextMenu(FMenuBuilder& M
 		MenuBuilder.AddSubMenu(
 			LOCTEXT("ChangeClassLabel", "Change Class"),
 			LOCTEXT("ChangeClassTooltip", "Change the class (object template) that this spawns from"),
-			FNewMenuDelegate::CreateRaw(this, &FLevelSequenceCustomization::AddChangeClassMenu, ObjectBindingModel));
+			FNewMenuDelegate::CreateLambda([=](FMenuBuilder& MenuBuilder) { FSequencerUtilities::AddChangeClassMenu(MenuBuilder, Sequencer.ToSharedRef(), ObjectBindingID, 0, TFunction<void()>()); }));
 
 		MenuBuilder.AddMenuEntry(
 			LOCTEXT("ContinuouslyRespawn", "Continuously Respawn"),
@@ -229,7 +229,11 @@ void FLevelSequenceCustomization::ExtendObjectBindingContextMenu(FMenuBuilder& M
 		);
 
 		MenuBuilder.AddMenuEntry(FSequencerCommands::Get().SaveCurrentSpawnableState);
-		MenuBuilder.AddMenuEntry(FSequencerCommands::Get().ConvertToPossessable);
+
+		MenuBuilder.AddSubMenu(
+			LOCTEXT("ConvertBindingLabel", "Convert Binding To..."),
+			LOCTEXT("ConvertBindingLabelTooltip", "Convert this binding into another binding type"),
+			FNewMenuDelegate::CreateLambda([=](FMenuBuilder& MenuBuilder) { FSequencerUtilities::AddConvertBindingMenu(MenuBuilder, Sequencer.ToSharedRef(), ObjectBindingID, 0, TFunction<void()>()); }));
 
 		MenuBuilder.AddSubMenu(
 				LOCTEXT("DynamicSpawn", "Dynamic Spawn"),
@@ -248,6 +252,14 @@ void FLevelSequenceCustomization::ExtendObjectBindingContextMenu(FMenuBuilder& M
 			bMultipleBindings = BindingReferences->GetReferences(ObjectBindingID).Num() > 1;
 		}
 
+		if (!bMultipleBindings)
+		{
+			MenuBuilder.AddSubMenu(
+				LOCTEXT("ConvertBindingLabel", "Convert Binding To..."),
+				LOCTEXT("ConvertBindingLabelTooltip", "Convert this binding into another binding type"),
+				FNewMenuDelegate::CreateLambda([=](FMenuBuilder& MenuBuilder) {FSequencerUtilities::AddConvertBindingMenu(MenuBuilder, Sequencer.ToSharedRef(), ObjectBindingID, 0, TFunction<void()>()); }));
+		}
+
 		// Regular possessable
 		if (!bCustomBinding)
 		{
@@ -262,34 +274,6 @@ void FLevelSequenceCustomization::ExtendObjectBindingContextMenu(FMenuBuilder& M
 						LOCTEXT("DynamicPossessionTooltip", "Specify a Blueprint method that will find a compatible actor for this binding"),
 						FNewMenuDelegate::CreateRaw(this, &FLevelSequenceCustomization::AddDynamicPossessionMenu, ObjectBindingModel));
 				}
-				// Custom bindings
-				TArrayView<const FMovieSceneBindingReference> BindingReferencesList = Sequencer->GetFocusedMovieSceneSequence()->GetBindingReferences()->GetReferences(ObjectBindingID);
-				if (BindingReferencesList.Num() == 1)
-				{
-					const FMovieSceneBindingReference& CurrentBindingReference = BindingReferencesList[0];
-					TArrayView<TWeakObjectPtr<>> BoundObjects = Sequencer->FindBoundObjects(ObjectBindingID, Sequencer->GetFocusedTemplateID());
-					if (BoundObjects.Num() == 1)
-					{
-						if (UObject* CurrentBoundObject = BoundObjects[0].Get())
-						{
-							TArrayView<const TSubclassOf<UMovieSceneCustomBinding>> PrioritySortedCustomBindingTypes = Sequencer->GetSupportedCustomBindingTypes();
-							for (const TSubclassOf<UMovieSceneCustomBinding>& CustomBindingType : PrioritySortedCustomBindingTypes)
-							{
-								if (CustomBindingType
-									&& (!CurrentBindingReference.CustomBinding
-										|| CurrentBindingReference.CustomBinding->GetClass() != CustomBindingType)
-									&& CustomBindingType->GetDefaultObject<UMovieSceneCustomBinding>()->SupportsConversionFromBinding(CurrentBindingReference, CurrentBoundObject))
-								{
-									MenuBuilder.AddMenuEntry(
-										FText::Format(LOCTEXT("ConvertToCustomBinding", "Convert to {0}"), CustomBindingType->GetDefaultObject<UMovieSceneCustomBinding>()->GetBindingTypePrettyName()),
-										FText::Format(LOCTEXT("ConvertToCustomBindingTooltip", "Convert selected binding to {0}"), CustomBindingType->GetDefaultObject<UMovieSceneCustomBinding>()->GetBindingTypePrettyName()),
-										FSlateIcon(),
-										FUIAction(FExecuteAction::CreateRaw(this, &FLevelSequenceCustomization::ConvertToCustomBinding, ObjectBindingModel, CustomBindingType)));
-								}
-							}
-						}
-					}
-				}
 			}
 			MenuBuilder.EndSection();
 		}
@@ -308,37 +292,13 @@ void FLevelSequenceCustomization::ExtendObjectBindingContextMenu(FMenuBuilder& M
 					MenuBuilder.AddSubMenu(
 						LOCTEXT("ChangeClassLabel", "Change Class"),
 						LOCTEXT("ChangeClassTooltip", "Change the class (object template) that this spawns from"),
-						FNewMenuDelegate::CreateRaw(this, &FLevelSequenceCustomization::AddChangeClassMenu, ObjectBindingModel));
-
-					// Binding conversions
-					MenuBuilder.AddMenuEntry(FSequencerCommands::Get().ConvertToPossessable);
-
-					// Custom bindings
-					TArrayView<const FMovieSceneBindingReference> BindingReferencesList = Sequencer->GetFocusedMovieSceneSequence()->GetBindingReferences()->GetReferences(ObjectBindingID);
-					if (BindingReferencesList.Num() == 1)
-					{
-						const FMovieSceneBindingReference& CurrentBindingReference = BindingReferencesList[0];
-						UObject* CurrentBoundObject = Sequencer->FindSpawnedObjectOrTemplate(ObjectBindingID);
-						if (CurrentBoundObject)
-						{
-							TArrayView<const TSubclassOf<UMovieSceneCustomBinding>> PrioritySortedCustomBindingTypes = Sequencer->GetSupportedCustomBindingTypes();
-							for (const TSubclassOf<UMovieSceneCustomBinding>& CustomBindingType : PrioritySortedCustomBindingTypes)
-							{
-								if (CustomBindingType
-									&& CurrentBindingReference.CustomBinding
-									&& CurrentBindingReference.CustomBinding->GetClass() != CustomBindingType
-									&& CustomBindingType->GetDefaultObject<UMovieSceneCustomBinding>()->SupportsConversionFromBinding(CurrentBindingReference, CurrentBoundObject))
-								{
-									MenuBuilder.AddMenuEntry(
-										FText::Format(LOCTEXT("ConvertToCustomBinding", "Convert to {0}"), CustomBindingType->GetDefaultObject<UMovieSceneCustomBinding>()->GetBindingTypePrettyName()),
-										FText::Format(LOCTEXT("ConvertToCustomBindingTooltip", "Convert selected binding to {0}"), CustomBindingType->GetDefaultObject<UMovieSceneCustomBinding>()->GetBindingTypePrettyName()),
-										FSlateIcon(),
-										FUIAction(FExecuteAction::CreateRaw(this, &FLevelSequenceCustomization::ConvertToCustomBinding, ObjectBindingModel, CustomBindingType)));
-								}
-							}
-						}
-					}
-				}
+						FNewMenuDelegate::CreateLambda([=](FMenuBuilder& MenuBuilder) { FSequencerUtilities::AddChangeClassMenu(MenuBuilder, Sequencer.ToSharedRef(), ObjectBindingID, 0, TFunction<void()>()); }));
+				}		
+				
+				MenuBuilder.AddSubMenu(
+					LOCTEXT("DynamicSpawn", "Dynamic Spawn"),
+					LOCTEXT("DynamicSpawnTooltip", "Specify a Blueprint method that will spawn or otherwise acquire a compatible actor for this binding"),
+					FNewMenuDelegate::CreateRaw(this, &FLevelSequenceCustomization::AddDynamicSpawnMenu, ObjectBindingModel));
 			}
 			else
 			{
@@ -547,111 +507,6 @@ void FLevelSequenceCustomization::SetSelectedNodesSpawnableLevel(TSharedPtr<FObj
 	}
 }
 
-void FLevelSequenceCustomization::AddChangeClassMenu(FMenuBuilder& MenuBuilder, TSharedPtr<FObjectBindingModel> ObjectBindingModel)
-{
-	TSharedPtr<ISequencer> Sequencer = WeakSequencer.Pin();
-	TSharedPtr<FSequencerEditorViewModel> EditorViewModel = Sequencer->GetViewModel();
-
-	FGuid ObjectBindingID = ObjectBindingModel->GetObjectGuid();
-	UMovieScene* MovieScene = Sequencer->GetFocusedMovieSceneSequence()->GetMovieScene();
-
-	FClassViewerInitializationOptions Options;
-	Options.Mode = EClassViewerMode::ClassPicker;
-	Options.bIsPlaceableOnly = true;
-	
-	FMovieSceneSpawnable* Spawnable = MovieScene->FindSpawnable(ObjectBindingID);
-	if (Spawnable)
-	{
-		Options.bIsActorsOnly = true;
-	}
-	else if (const FMovieSceneBindingReferences* BindingReferences = Sequencer->GetFocusedMovieSceneSequence()->GetBindingReferences())
-	{
-		TArrayView<const FMovieSceneBindingReference> BindingReferencesList = BindingReferences->GetReferences(ObjectBindingID);
-		if (BindingReferencesList.Num() == 1 && BindingReferencesList[0].CustomBinding && BindingReferencesList[0].CustomBinding->WillSpawnObject(Sequencer->GetSharedPlaybackState()))
-		{
-			// Class filter for the custom binding type
-			class FCustomBindingClassFilter : public IClassViewerFilter
-			{
-			public:
-				bool IsClassAllowed(const FClassViewerInitializationOptions& InInitOptions, const UClass* InClass, TSharedRef< FClassViewerFilterFuncs > InFilterFuncs) override
-				{
-					return CustomBinding && InClass && CustomBinding->SupportsBindingCreationFromObject(InClass->GetDefaultObject());
-				}
-
-				virtual bool IsUnloadedClassAllowed(const FClassViewerInitializationOptions& InInitOptions, const TSharedRef< const IUnloadedBlueprintData > InClass, TSharedRef< FClassViewerFilterFuncs > InFilterFuncs) override
-				{
-					if (const UClass* ClassWithin = InClass->GetClassWithin())
-					{
-						return IsClassAllowed(InInitOptions, ClassWithin, InFilterFuncs);
-					}
-					return false;
-				}
-
-				TObjectPtr<UMovieSceneCustomBinding> CustomBinding;
-			};
-
-			TSharedRef<FCustomBindingClassFilter> ClassFilter = MakeShared<FCustomBindingClassFilter>();
-			ClassFilter->CustomBinding = BindingReferencesList[0].CustomBinding;
-			Options.ClassFilters.Add(ClassFilter);
-		}
-		else
-		{
-			return;
-		}
-	}
-	else
-	{
-		return;
-	}
-
-	FClassViewerModule& ClassViewerModule = FModuleManager::LoadModuleChecked<FClassViewerModule>("ClassViewer");
-
-
-	const UClass* ClassForObjectBinding = ObjectBindingModel->FindObjectClass();
-	if (ClassForObjectBinding)
-	{
-		Options.ViewerTitleString = FText::FromString(TEXT("Change from: ") + ClassForObjectBinding->GetFName().ToString());
-	}
-	else
-	{
-		Options.ViewerTitleString = FText::FromString(TEXT("Change from: (empty)"));
-	}
-
-	MenuBuilder.AddWidget(
-		SNew(SBox)
-		.MinDesiredWidth(300.0f)
-		.MaxDesiredHeight(400.0f)
-		[
-			ClassViewerModule.CreateClassViewer(Options, FOnClassPicked::CreateRaw(
-				this, &FLevelSequenceCustomization::HandleTemplateActorClassPicked, ObjectBindingModel))
-		],
-		FText(), true, false
-	);
-}
-
-void FLevelSequenceCustomization::HandleTemplateActorClassPicked(UClass* ChosenClass, TSharedPtr<FObjectBindingModel> ObjectBindingModel)
-{
-	TSharedPtr<ISequencer> Sequencer = WeakSequencer.Pin();
-	UMovieScene* MovieScene = Sequencer->GetFocusedMovieSceneSequence()->GetMovieScene();
-
-	FGuid ObjectBindingID = ObjectBindingModel->GetObjectGuid();
-
-	FSlateApplication::Get().DismissAllMenus();
-
-	FScopedTransaction Transaction(LOCTEXT("ChangeClass", "Change Class"));
-
-	MovieScene->Modify();
-
-	TValueOrError<FNewSpawnable, FText> Result = Sequencer->GetSpawnRegister().CreateNewSpawnableType(*ChosenClass, *MovieScene, nullptr);
-	if (Result.IsValid())
-	{
-		MovieSceneHelpers::SetObjectTemplate(Sequencer->GetFocusedMovieSceneSequence(), ObjectBindingID, Result.GetValue().ObjectTemplate, Sequencer->GetSharedPlaybackState());
-
-		Sequencer->GetSpawnRegister().DestroySpawnedObject(ObjectBindingID, Sequencer->GetFocusedTemplateID(), *Sequencer.Get());
-		Sequencer->ForceEvaluate();
-	}
-}
-
 void FLevelSequenceCustomization::AddDynamicSpawnMenu(FMenuBuilder& MenuBuilder, TSharedPtr<FObjectBindingModel> ObjectBindingModel)
 {
 	TSharedPtr<ISequencer> Sequencer = WeakSequencer.Pin();
@@ -682,14 +537,6 @@ void FLevelSequenceCustomization::AddDynamicPossessionMenu(FMenuBuilder& MenuBui
 	ObjectBindingModel->AddDynamicBindingMenu(MenuBuilder, Possessable->DynamicBinding);
 }
 
-void FLevelSequenceCustomization::ConvertToCustomBinding(TSharedPtr<FObjectBindingModel> ObjectBindingModel, TSubclassOf<UMovieSceneCustomBinding> CustomBindingType)
-{
-	if (WeakSequencer.IsValid() && ObjectBindingModel.IsValid() && CustomBindingType)
-	{
-		FSequencerUtilities::ConvertToCustomBinding(WeakSequencer.Pin().ToSharedRef(), ObjectBindingModel->GetObjectGuid(), CustomBindingType);
-	}
-}
-
 } // namespace UE::Sequencer
 
 #undef LOCTEXT_NAMESPACE
diff --git a/Engine/Plugins/MovieScene/LevelSequenceEditor/Source/LevelSequenceEditor/Private/Misc/LevelSequenceCustomization.h b/Engine/Plugins/MovieScene/LevelSequenceEditor/Source/LevelSequenceEditor/Private/Misc/LevelSequenceCustomization.h
--- a/Engine/Plugins/MovieScene/LevelSequenceEditor/Source/LevelSequenceEditor/Private/Misc/LevelSequenceCustomization.h
+++ b/Engine/Plugins/MovieScene/LevelSequenceEditor/Source/LevelSequenceEditor/Private/Misc/LevelSequenceCustomization.h
@@ -37,12 +37,8 @@ private:
 	void AddSpawnOwnershipMenu(FMenuBuilder& MenuBuilder, TSharedPtr<FObjectBindingModel> ObjectBindingModel);
 	void AddSpawnLevelMenu(FMenuBuilder& MenuBuilder, TSharedPtr<FObjectBindingModel> ObjectBindingModel);
 	void SetSelectedNodesSpawnableLevel(TSharedPtr<FObjectBindingModel> ObjectBindingModel, FName InLevelName);
-	void AddChangeClassMenu(FMenuBuilder& MenuBuilder, TSharedPtr<FObjectBindingModel> ObjectBindingModel);
-	void HandleTemplateActorClassPicked(UClass* ChosenClass, TSharedPtr<FObjectBindingModel> ObjectBindingModel);
-
 	void AddDynamicSpawnMenu(FMenuBuilder& MenuBuilder, TSharedPtr<FObjectBindingModel> ObjectBindingModel);
 	void AddDynamicPossessionMenu(FMenuBuilder& MenuBuilder, TSharedPtr<FObjectBindingModel> ObjectBindingModel);
-	void ConvertToCustomBinding(TSharedPtr<FObjectBindingModel> ObjectBindingModel, TSubclassOf<UMovieSceneCustomBinding> CustomBindingType);
 
 private:
 
diff --git a/Engine/Plugins/MovieScene/LevelSequenceEditor/Source/LevelSequenceEditor/Private/Misc/LevelSequenceEditorSpawnRegister.cpp b/Engine/Plugins/MovieScene/LevelSequenceEditor/Source/LevelSequenceEditor/Private/Misc/LevelSequenceEditorSpawnRegister.cpp
--- a/Engine/Plugins/MovieScene/LevelSequenceEditor/Source/LevelSequenceEditor/Private/Misc/LevelSequenceEditorSpawnRegister.cpp
+++ b/Engine/Plugins/MovieScene/LevelSequenceEditor/Source/LevelSequenceEditor/Private/Misc/LevelSequenceEditorSpawnRegister.cpp
@@ -140,13 +140,13 @@ void FLevelSequenceEditorSpawnRegister::PreDestroyObject(UObject& Object, const
 }
 
 #if WITH_EDITOR
-void FLevelSequenceEditorSpawnRegister::SaveDefaultSpawnableState(const FGuid& BindingId, FMovieSceneSequenceIDRef TemplateID, TSharedRef<const FSharedPlaybackState> SharedPlaybackState)
+void FLevelSequenceEditorSpawnRegister::SaveDefaultSpawnableState(const FGuid& BindingId, int32 BindingIndex, FMovieSceneSequenceIDRef TemplateID, TSharedRef<const FSharedPlaybackState> SharedPlaybackState)
 {
 	if (UMovieSceneSequence* Sequence = SharedPlaybackState->GetSequence(TemplateID))
 	{
-		if (UObject* Object = FindSpawnedObject(BindingId, TemplateID, 0).Get())
+		if (UObject* Object = FindSpawnedObject(BindingId, TemplateID, BindingIndex).Get())
 		{
-			SaveDefaultSpawnableStateImpl(BindingId, 0, Sequence, Object, SharedPlaybackState);
+			SaveDefaultSpawnableStateImpl(BindingId, BindingIndex, Sequence, Object, SharedPlaybackState);
 			Sequence->MarkPackageDirty();
 		}
 	}
@@ -177,7 +177,7 @@ void FLevelSequenceEditorSpawnRegister::SaveDefaultSpawnableStateImpl(const FGui
 
 
 	// Copy the template
-	MovieSceneHelpers::CopyObjectTemplate(Sequence, BindingId ,SpawnedObject, SharedPlaybackState);
+	MovieSceneHelpers::CopyObjectTemplate(Sequence, BindingId ,SpawnedObject, SharedPlaybackState, BindingIndex);
 
 	if (FTrackedObjectState* TrackedState = TrackedObjects.Find(SpawnedObject))
 	{
diff --git a/Engine/Plugins/MovieScene/LevelSequenceEditor/Source/LevelSequenceEditor/Public/LevelSequenceEditorSubsystem.h b/Engine/Plugins/MovieScene/LevelSequenceEditor/Source/LevelSequenceEditor/Public/LevelSequenceEditorSubsystem.h
--- a/Engine/Plugins/MovieScene/LevelSequenceEditor/Source/LevelSequenceEditor/Public/LevelSequenceEditorSubsystem.h
+++ b/Engine/Plugins/MovieScene/LevelSequenceEditor/Source/LevelSequenceEditor/Public/LevelSequenceEditorSubsystem.h
@@ -55,14 +55,14 @@ struct FMovieSceneBindingPropertyInfo
 	GENERATED_BODY()
 
 	// Locator for the entry
-	UPROPERTY(EditAnywhere, Category = "Default", meta=(AllowedLocators="Actor"))
+	UPROPERTY(EditAnywhere, Category = "Default", meta=(AllowedLocators="Actor", DisplayName="Actor"))
 	FUniversalObjectLocator Locator;
 
 	// Flags for how to resolve the locator
 	UPROPERTY()
 	ELocatorResolveFlags ResolveFlags = ELocatorResolveFlags::None;
 
-	UPROPERTY(Instanced, EditAnywhere, Category = "Default", meta=(AllowEditInlineCustomization))
+	UPROPERTY(Instanced, VisibleAnywhere, Category = "Default", meta=(EditInline, AllowEditInlineCustomization, DisplayName="Custom Binding Type"))
 	UMovieSceneCustomBinding* CustomBinding = nullptr;
 };
 
@@ -229,6 +229,9 @@ public:
 	UFUNCTION(BlueprintCallable, Category = "Level Sequence Editor")
 	void RebindComponent(const TArray<FMovieSceneBindingProxy>& ComponentBindings, const FName& ComponentName);
 
+	// Refreshes the binding details when the bindings change in the menu
+	void RefreshBindingDetails(IDetailsView* DetailsView, FGuid ObjectBindingID);
+
 private:
 	/** Used by Baking transforms*/
 	struct FBakeData
diff --git a/Engine/Plugins/MovieScene/LevelSequenceEditor/Source/LevelSequenceEditor/Public/Misc/LevelSequenceEditorSpawnRegister.h b/Engine/Plugins/MovieScene/LevelSequenceEditor/Source/LevelSequenceEditor/Public/Misc/LevelSequenceEditorSpawnRegister.h
--- a/Engine/Plugins/MovieScene/LevelSequenceEditor/Source/LevelSequenceEditor/Public/Misc/LevelSequenceEditorSpawnRegister.h
+++ b/Engine/Plugins/MovieScene/LevelSequenceEditor/Source/LevelSequenceEditor/Public/Misc/LevelSequenceEditorSpawnRegister.h
@@ -39,7 +39,7 @@ public:
 	virtual UObject* SpawnObject(const FGuid& BindingId, UMovieScene& MovieScene, FMovieSceneSequenceIDRef Template, TSharedRef<const FSharedPlaybackState> SharedPlaybackState, int32 BindingIndex);
 	virtual void PreDestroyObject(UObject& Object, const FGuid& BindingId, int32 BindingIndex, FMovieSceneSequenceIDRef TemplateID) override;
 #if WITH_EDITOR
-	virtual void SaveDefaultSpawnableState(const FGuid& BindingId, FMovieSceneSequenceIDRef TemplateID, TSharedRef<const FSharedPlaybackState> SharedPlaybackState) override;
+	virtual void SaveDefaultSpawnableState(const FGuid& BindingId, int32 BindingIndex, FMovieSceneSequenceIDRef TemplateID, TSharedRef<const FSharedPlaybackState> SharedPlaybackState) override;
 	virtual TValueOrError<FNewSpawnable, FText> CreateNewSpawnableType(UObject& SourceObject, UMovieScene& OwnerMovieScene, UActorFactory* ActorFactory = nullptr) override;
 	virtual void SetupDefaultsForSpawnable(UObject* SpawnedObject, const FGuid& Guid, const TOptional<FTransformData>& TransformData, TSharedRef<ISequencer> Sequencer, USequencerSettings* Settings) override;
 	virtual void HandleConvertPossessableToSpawnable(UObject* OldObject, TSharedRef<const FSharedPlaybackState> SharedPlaybackState, TOptional<FTransformData>& OutTransformData) override;
diff --git a/Engine/Source/Editor/Sequencer/Private/Sequencer.cpp b/Engine/Source/Editor/Sequencer/Private/Sequencer.cpp
--- a/Engine/Source/Editor/Sequencer/Private/Sequencer.cpp
+++ b/Engine/Source/Editor/Sequencer/Private/Sequencer.cpp
@@ -8594,7 +8594,21 @@ void FSequencer::SaveSelectedNodesSpawnableState()
 		SlowTask.EnterProgressFrame();
 		if (TViewModelPtr<IObjectBindingExtension> ObjectBindingNode = OutlinerSelectionNode.ImplicitCast())
 		{
-			SpawnRegister->SaveDefaultSpawnableState(ObjectBindingNode->GetObjectGuid(), ActiveTemplateIDs.Top(), GetSharedPlaybackState());
+			if (const FMovieSceneSpawnable* Spawnable = MovieScene->FindSpawnable(ObjectBindingNode->GetObjectGuid()))
+			{
+				SpawnRegister->SaveDefaultSpawnableState(ObjectBindingNode->GetObjectGuid(), 0, ActiveTemplateIDs.Top(), GetSharedPlaybackState());
+			}
+			else if (const FMovieSceneBindingReferences* BindingReferences = Sequence->GetBindingReferences())
+			{
+				TArrayView<const FMovieSceneBindingReference> AllReferences = BindingReferences->GetReferences(ObjectBindingNode->GetObjectGuid());
+				for (int32 Index = 0; Index < AllReferences.Num(); ++Index)
+				{
+					if (AllReferences[Index].CustomBinding && AllReferences[Index].CustomBinding->WillSpawnObject(GetSharedPlaybackState()))
+					{
+						SpawnRegister->SaveDefaultSpawnableState(ObjectBindingNode->GetObjectGuid(), Index, ActiveTemplateIDs.Top(), GetSharedPlaybackState());
+					}
+				}
+			}
 
 			if (GWarn->ReceivedUserCancel())
 			{
diff --git a/Engine/Source/Editor/Sequencer/Private/SequencerUtilities.cpp b/Engine/Source/Editor/Sequencer/Private/SequencerUtilities.cpp
--- a/Engine/Source/Editor/Sequencer/Private/SequencerUtilities.cpp
+++ b/Engine/Source/Editor/Sequencer/Private/SequencerUtilities.cpp
@@ -67,6 +67,10 @@
 #include "Bindings/MovieSceneReplaceableBinding.h"
 #include "ActorFactories/ActorFactory.h"
 #include "Tracks/MovieSceneBindingLifetimeTrack.h"
+#include "ClassViewerModule.h"
+#include "ClassViewerFilter.h"
+#include "Framework/Application/SlateApplication.h"
+#include "SequencerCommands.h"
 
 #define LOCTEXT_NAMESPACE "FSequencerUtilities"
 
@@ -1094,14 +1098,7 @@ FMovieScenePossessable* FSequencerUtilities::ConvertToPossessable(TSharedRef<ISe
 		return CreatedPossessable;
 	}
 
-	TArrayView<TWeakObjectPtr<>> FoundObjects = Sequencer->FindBoundObjects(BindingGuid, Sequencer->GetFocusedTemplateID());
-
-	UObject* ObjectToConvert = nullptr;
-
-	if (FoundObjects.IsValidIndex(BindingIndex))
-	{
-		ObjectToConvert = FoundObjects[BindingIndex].Get();
-	}
+	UObject* ObjectToConvert = MovieSceneHelpers::GetSingleBoundObject(Sequence, BindingGuid, Sequencer->GetSharedPlaybackState(), BindingIndex);
 
 	// If we have an old-style spawnable, use the template as the object to convert instead.
 	bool bConvertFromSpawnable = MovieSceneHelpers::IsBoundToSpawnable(Sequence, BindingGuid, Sequencer->GetSharedPlaybackState(), BindingIndex);
@@ -1166,7 +1163,24 @@ FMovieScenePossessable* FSequencerUtilities::ConvertToPossessable(TSharedRef<ISe
 			return nullptr;
 		}
 
-		PossessedActor->SetActorLabel(Spawnable ? Spawnable->GetName() : ExistingPossessable ? ExistingPossessable->GetName() : ObjectToConvert->GetName());
+		FString ActorLabel = Actor->GetActorLabel();
+		if (Spawnable)
+		{
+			ActorLabel = Spawnable->GetName();
+		}
+		else if (ExistingPossessable)
+		{
+			if (const FMovieSceneBindingReferences* BindingReferences = Sequence->GetBindingReferences())
+			{
+				// If we don't have multiple bound objects, use the Possessable name instead of the template label
+				if (BindingReferences->GetReferences(BindingGuid).Num() == 1)
+				{
+					ActorLabel = ExistingPossessable->GetName();
+				}
+			}
+		}
+
+		PossessedActor->SetActorLabel(ActorLabel);
 
 		const bool bIsDefaultTransform = true;
 		PossessedActor->FinishSpawning(DefaultTransform, bIsDefaultTransform);
@@ -1185,8 +1199,19 @@ FMovieScenePossessable* FSequencerUtilities::ConvertToPossessable(TSharedRef<ISe
 	Sequence->Modify();
 	MovieScene->Modify();
 
-	// If we're converting from a spawnable, we'll need to remove the spawn track
-	if (bConvertFromSpawnable)
+	UE::Sequencer::FCreateBindingParams CreateBindingParams;
+	CreateBindingParams.ReplacementGuid = BindingGuid;
+	CreateBindingParams.BindingIndex = BindingIndex;
+	CreateBindingParams.bAllowCustomBinding = false;
+
+	// Create or replace the binding
+	FGuid NewPossessableGuid = CreateBinding(Sequencer, *PossessedObject, CreateBindingParams);
+
+	TArrayView<const FMovieSceneBindingReference> BindingReferences = Sequence->GetBindingReferences()->GetReferences(BindingGuid);
+	bool bAnySpawnablesLeft = Algo::AnyOf(BindingReferences, [](const FMovieSceneBindingReference& BindingReference) { return BindingReference.CustomBinding && BindingReference.CustomBinding->IsA<UMovieSceneSpawnableBindingBase>(); });
+
+	// If we're converting from a spawnable and none of the other bindings on the guid are spawnable, we'll need to remove the spawn track
+	if (bConvertFromSpawnable && !bAnySpawnablesLeft)
 	{
 		// Delete the spawn track
 		UMovieSceneSpawnTrack* SpawnTrack = Cast<UMovieSceneSpawnTrack>(MovieScene->FindTrack(UMovieSceneSpawnTrack::StaticClass(), BindingGuid, NAME_None));
@@ -1197,15 +1222,6 @@ FMovieScenePossessable* FSequencerUtilities::ConvertToPossessable(TSharedRef<ISe
 	}
 
 
-	UE::Sequencer::FCreateBindingParams CreateBindingParams;
-	CreateBindingParams.ReplacementGuid = BindingGuid;
-	CreateBindingParams.BindingIndex = BindingIndex;
-	CreateBindingParams.bAllowCustomBinding = false;
-
-	// Create or replace the binding
-	FGuid NewPossessableGuid = CreateBinding(Sequencer, *PossessedObject, CreateBindingParams);
-
-
 	FMovieScenePossessable* Possessable = MovieScene->FindPossessable(NewPossessableGuid);
 	if (Spawnable)
 	{
@@ -1302,14 +1318,7 @@ FMovieScenePossessable* FSequencerUtilities::ConvertToCustomBinding(TSharedRef<I
 		return CreatedPossessable;
 	}
 
-	TArrayView<TWeakObjectPtr<>> FoundObjects = Sequencer->FindBoundObjects(BindingGuid, Sequencer->GetFocusedTemplateID());
-
-	UObject* ObjectToConvert = nullptr;
-
-	if (FoundObjects.IsValidIndex(BindingIndex))
-	{
-		ObjectToConvert = FoundObjects[BindingIndex].Get();
-	}
+	UObject* ObjectToConvert = MovieSceneHelpers::GetSingleBoundObject(Sequence, BindingGuid, Sequencer->GetSharedPlaybackState(), BindingIndex);
 
 	bool bConvertFromSpawnable = false;
 
@@ -1331,8 +1340,7 @@ FMovieScenePossessable* FSequencerUtilities::ConvertToCustomBinding(TSharedRef<I
 		}
 	}
 
-	bool bConvertFromPossessable = !bConvertFromSpawnable
-		&& !Algo::AnyOf(BindingReferences->GetReferences(BindingGuid), [](const FMovieSceneBindingReference& Reference) { return Reference.CustomBinding; });
+	bool bConvertFromPossessable = !bConvertFromSpawnable && !BindingReferences->GetCustomBinding(BindingGuid, BindingIndex);
 
 	if (!ObjectToConvert)
 	{
@@ -1359,8 +1367,21 @@ FMovieScenePossessable* FSequencerUtilities::ConvertToCustomBinding(TSharedRef<I
 	MovieScene->Modify();
 
 
+	UE::Sequencer::FCreateBindingParams CreateBindingParams;
+	CreateBindingParams.ReplacementGuid = BindingGuid;
+	CreateBindingParams.BindingIndex = BindingIndex;
+	CreateBindingParams.BindingNameOverride = NewCustomBinding->GetDesiredBindingName();
+	CreateBindingParams.CustomBinding = NewCustomBinding;
+
+	// Create or replace the binding
+	FGuid NewPossessableGuid = CreateBinding(Sequencer, *ObjectToConvert, CreateBindingParams);
+
+	TArrayView<const FMovieSceneBindingReference> BindingReferencesForGuid = BindingReferences->GetReferences(BindingGuid);
+	bool bAnySpawnablesLeft = Algo::AnyOf(BindingReferencesForGuid, [](const FMovieSceneBindingReference& BindingReference) { return BindingReference.CustomBinding && BindingReference.CustomBinding->IsA<UMovieSceneSpawnableBindingBase>(); });
+	bool bAnyReplaceablesLeft = Algo::AnyOf(BindingReferencesForGuid, [](const FMovieSceneBindingReference& BindingReference) { return BindingReference.CustomBinding && BindingReference.CustomBinding->IsA<UMovieSceneReplaceableBindingBase>(); });
+
 	// If we're converting from a spawnable and the new custom binding isn't a spawnable, remove the spawn track
-	if (PreviousCustomBinding && PreviousCustomBinding->IsA<UMovieSceneSpawnableBindingBase>() && !NewCustomBinding->IsA<UMovieSceneSpawnableBindingBase>())
+	if (PreviousCustomBinding && PreviousCustomBinding->IsA<UMovieSceneSpawnableBindingBase>() && !bAnySpawnablesLeft)
 	{
 		// Delete the spawn track
 		UMovieSceneSpawnTrack* SpawnTrack = Cast<UMovieSceneSpawnTrack>(MovieScene->FindTrack(UMovieSceneSpawnTrack::StaticClass(), BindingGuid, NAME_None));
@@ -1369,7 +1390,7 @@ FMovieScenePossessable* FSequencerUtilities::ConvertToCustomBinding(TSharedRef<I
 			MovieScene->RemoveTrack(*SpawnTrack);
 		}
 	}
-	else if (PreviousCustomBinding && PreviousCustomBinding->IsA<UMovieSceneReplaceableBindingBase>() && !NewCustomBinding->IsA<UMovieSceneReplaceableBindingBase>())
+	else if (PreviousCustomBinding && PreviousCustomBinding->IsA<UMovieSceneReplaceableBindingBase>() && !bAnyReplaceablesLeft)
 	{
 		// Delete the binding lifetime track
 		UMovieSceneBindingLifetimeTrack* BindingLifetimeTrack = Cast<UMovieSceneBindingLifetimeTrack>(MovieScene->FindTrack(UMovieSceneBindingLifetimeTrack::StaticClass(), BindingGuid, NAME_None));
@@ -1379,15 +1400,6 @@ FMovieScenePossessable* FSequencerUtilities::ConvertToCustomBinding(TSharedRef<I
 		}
 	}
 
-	UE::Sequencer::FCreateBindingParams CreateBindingParams;
-	CreateBindingParams.ReplacementGuid = BindingGuid;
-	CreateBindingParams.BindingIndex = BindingIndex;
-	CreateBindingParams.BindingNameOverride = NewCustomBinding->GetDesiredBindingName();
-	CreateBindingParams.CustomBinding = NewCustomBinding;
-
-	// Create or replace the binding
-	FGuid NewPossessableGuid = CreateBinding(Sequencer, *ObjectToConvert, CreateBindingParams);
-
 	CreatedPossessable = MovieScene->FindPossessable(NewPossessableGuid);
 
 	// If we previously had an old-style spawnable, we need to move over bindings
@@ -3716,4 +3728,191 @@ TRange<FFrameNumber> FSequencerUtilities::GetTimeBounds(TSharedRef<ISequencer> S
 	return StaticCastSharedPtr<FSequencer>(Sequencer.ToSharedPtr())->GetTimeBounds();
 }
 
+void FSequencerUtilities::AddChangeClassMenu(FMenuBuilder& MenuBuilder, TSharedRef<ISequencer> Sequencer, FGuid ObjectBindingID, int32 BindingIndex, TFunction<void()> OnBindingChanged)
+{
+	UMovieSceneSequence* Sequence = Sequencer->GetFocusedMovieSceneSequence();
+	if (!Sequence)
+	{
+		return;
+	}
+	UMovieScene* MovieScene = Sequence->GetMovieScene();
+
+	FClassViewerInitializationOptions Options;
+	Options.Mode = EClassViewerMode::ClassPicker;
+	Options.bIsPlaceableOnly = true;
+
+	FMovieSceneSpawnable* Spawnable = MovieScene->FindSpawnable(ObjectBindingID);
+	if (Spawnable)
+	{
+		Options.bIsActorsOnly = true;
+	}
+	else if (const FMovieSceneBindingReferences* BindingReferences = Sequence->GetBindingReferences())
+	{
+		TArrayView<const FMovieSceneBindingReference> BindingReferencesList = BindingReferences->GetReferences(ObjectBindingID);
+		if (BindingReferencesList.IsValidIndex(BindingIndex) && BindingReferencesList[BindingIndex].CustomBinding && BindingReferencesList[BindingIndex].CustomBinding->WillSpawnObject(Sequencer->GetSharedPlaybackState()))
+		{
+			// Class filter for the custom binding type
+			class FCustomBindingClassFilter : public IClassViewerFilter
+			{
+			public:
+				bool IsClassAllowed(const FClassViewerInitializationOptions& InInitOptions, const UClass* InClass, TSharedRef< FClassViewerFilterFuncs > InFilterFuncs) override
+				{
+					return CustomBinding && InClass && CustomBinding->SupportsBindingCreationFromObject(InClass->GetDefaultObject());
+				}
+
+				virtual bool IsUnloadedClassAllowed(const FClassViewerInitializationOptions& InInitOptions, const TSharedRef< const IUnloadedBlueprintData > InClass, TSharedRef< FClassViewerFilterFuncs > InFilterFuncs) override
+				{
+					if (const UClass* ClassWithin = InClass->GetClassWithin())
+					{
+						return IsClassAllowed(InInitOptions, ClassWithin, InFilterFuncs);
+					}
+					return false;
+				}
+
+				TObjectPtr<UMovieSceneCustomBinding> CustomBinding;
+			};
+
+			TSharedRef<FCustomBindingClassFilter> ClassFilter = MakeShared<FCustomBindingClassFilter>();
+			ClassFilter->CustomBinding = BindingReferencesList[0].CustomBinding;
+			Options.ClassFilters.Add(ClassFilter);
+		}
+		else
+		{
+			return;
+		}
+	}
+	else
+	{
+		return;
+	}
+
+	FClassViewerModule& ClassViewerModule = FModuleManager::LoadModuleChecked<FClassViewerModule>("ClassViewer");
+
+	
+	const UClass* ClassForObjectBinding = MovieSceneHelpers::GetBoundObjectClass(Sequence, ObjectBindingID, BindingIndex);
+	if (ClassForObjectBinding)
+	{
+		Options.ViewerTitleString = FText::FromString(TEXT("Change from: ") + ClassForObjectBinding->GetFName().ToString());
+	}
+	else
+	{
+		Options.ViewerTitleString = FText::FromString(TEXT("Change from: (empty)"));
+	}
+
+	MenuBuilder.AddWidget(
+		SNew(SBox)
+		.MinDesiredWidth(300.0f)
+		.MaxDesiredHeight(400.0f)
+		[
+			ClassViewerModule.CreateClassViewer(Options, FOnClassPicked::CreateLambda([Sequencer, ObjectBindingID, BindingIndex, OnBindingChanged](UClass* Class) { FSequencerUtilities::HandleTemplateActorClassPicked(Class, Sequencer, ObjectBindingID, BindingIndex, OnBindingChanged); }))
+		],
+		FText(), true, false
+	);
+}
+
+void FSequencerUtilities::HandleTemplateActorClassPicked(UClass* ChosenClass, TSharedRef<ISequencer> Sequencer, FGuid ObjectBindingID, int32 BindingIndex, TFunction<void()> OnBindingChanged)
+{
+	UMovieScene* MovieScene = Sequencer->GetFocusedMovieSceneSequence()->GetMovieScene();
+
+	FScopedTransaction Transaction(LOCTEXT("ChangeClass", "Change Class"));
+
+	MovieScene->Modify();
+
+	TValueOrError<FNewSpawnable, FText> Result = Sequencer->GetSpawnRegister().CreateNewSpawnableType(*ChosenClass, *MovieScene, nullptr);
+	if (Result.IsValid())
+	{
+		MovieSceneHelpers::SetObjectTemplate(Sequencer->GetFocusedMovieSceneSequence(), ObjectBindingID, Result.GetValue().ObjectTemplate, Sequencer->GetSharedPlaybackState(), BindingIndex);
+
+		Sequencer->GetSpawnRegister().DestroySpawnedObject(ObjectBindingID, Sequencer->GetFocusedTemplateID(), Sequencer->GetSharedPlaybackState(), BindingIndex);
+		Sequencer->ForceEvaluate();
+	}
+
+	if (OnBindingChanged)
+	{
+		OnBindingChanged();
+	}
+}
+
+void FSequencerUtilities::AddConvertBindingMenu(FMenuBuilder& MenuBuilder, TSharedRef<ISequencer> Sequencer, FGuid ObjectBindingID, int32 BindingIndex, TFunction<void()> OnBindingChanged)
+{
+	UMovieSceneSequence* Sequence = Sequencer->GetFocusedMovieSceneSequence();
+	UMovieScene* MovieScene = Sequence ? Sequence->GetMovieScene() : nullptr;
+
+	if (!MovieScene || !ObjectBindingID.IsValid())
+	{
+		return;
+	}
+
+	FMovieSceneSpawnable* Spawnable = MovieScene->FindSpawnable(ObjectBindingID);
+
+	if (Spawnable)
+	{
+		MenuBuilder.AddMenuEntry(FSequencerCommands::Get().ConvertToPossessable, NAME_None, LOCTEXT("ConvertToPossessable", "Possessable"));
+		MenuBuilder.AddMenuEntry(
+			LOCTEXT("ConvertToPossessable", "Possessable"),
+			LOCTEXT("ConvertToPossessableTooltip", "Convert selected binding to a possessable"),
+			FSlateIcon(),
+			FUIAction(FExecuteAction::CreateLambda([Sequencer, ObjectBindingID, BindingIndex, OnBindingChanged]() 
+				{ 
+					FSequencerUtilities::ConvertToPossessable(Sequencer, ObjectBindingID, BindingIndex); 
+					if (OnBindingChanged != nullptr)
+					{
+						OnBindingChanged();
+					}
+				})));
+	}
+	else if (const FMovieSceneBindingReferences* BindingReferences = Sequencer->GetFocusedMovieSceneSequence()->GetBindingReferences())
+	{
+		TArrayView<const FMovieSceneBindingReference> BindingReferencesList = Sequencer->GetFocusedMovieSceneSequence()->GetBindingReferences()->GetReferences(ObjectBindingID);
+		if (BindingReferencesList.IsValidIndex(BindingIndex))
+		{
+			if (BindingReferencesList[BindingIndex].CustomBinding != nullptr)
+			{
+				MenuBuilder.AddMenuEntry(
+					LOCTEXT("ConvertToPossessable", "Possessable"),
+					LOCTEXT("ConvertToPossessableTooltip", "Convert selected binding to a possessable"),
+					FSlateIcon(),
+					FUIAction(FExecuteAction::CreateLambda([Sequencer, ObjectBindingID, BindingIndex, OnBindingChanged]() 
+						{ 
+							FSequencerUtilities::ConvertToPossessable(Sequencer, ObjectBindingID, BindingIndex);
+							if (OnBindingChanged != nullptr)
+							{
+								OnBindingChanged();
+							}
+						})));
+			}
+
+			const FMovieSceneBindingReference& CurrentBindingReference = BindingReferencesList[BindingIndex];
+			UE::UniversalObjectLocator::FResolveParams LocatorResolveParams(Sequencer->GetSharedPlaybackState()->GetPlaybackContext());
+			FMovieSceneBindingResolveParams BindingResolveParams{ Sequence, ObjectBindingID, Sequencer->GetFocusedTemplateID()};
+			if (UObject* CurrentBoundObject = BindingReferences->ResolveSingleBinding(BindingResolveParams, BindingIndex, LocatorResolveParams, Sequencer->GetSharedPlaybackState()))
+			{
+				TArrayView<const TSubclassOf<UMovieSceneCustomBinding>> PrioritySortedCustomBindingTypes = Sequencer->GetSupportedCustomBindingTypes();
+				for (const TSubclassOf<UMovieSceneCustomBinding>& CustomBindingType : PrioritySortedCustomBindingTypes)
+				{
+					if (CustomBindingType
+						&& (!CurrentBindingReference.CustomBinding
+							|| CurrentBindingReference.CustomBinding->GetClass() != CustomBindingType)
+						&& CustomBindingType->GetDefaultObject<UMovieSceneCustomBinding>()->SupportsConversionFromBinding(CurrentBindingReference, CurrentBoundObject))
+					{
+						MenuBuilder.AddMenuEntry(
+							CustomBindingType->GetDefaultObject<UMovieSceneCustomBinding>()->GetBindingTypePrettyName(),
+							FText::Format(LOCTEXT("ConvertToCustomBindingTooltip", "Convert selected binding to {0}"), CustomBindingType->GetDefaultObject<UMovieSceneCustomBinding>()->GetBindingTypePrettyName()),
+							FSlateIcon(),
+							FUIAction(FExecuteAction::CreateLambda([Sequencer, ObjectBindingID, CustomBindingType, BindingIndex, OnBindingChanged]() 
+								{ 
+									FSequencerUtilities::ConvertToCustomBinding(Sequencer, ObjectBindingID, CustomBindingType, BindingIndex); 
+									if (OnBindingChanged != nullptr)
+									{
+										OnBindingChanged();
+									}
+								})));
+					}
+				}
+			}
+		}
+	}
+}
+
+
 #undef LOCTEXT_NAMESPACE
diff --git a/Engine/Source/Editor/Sequencer/Public/SequencerUtilities.h b/Engine/Source/Editor/Sequencer/Public/SequencerUtilities.h
--- a/Engine/Source/Editor/Sequencer/Public/SequencerUtilities.h
+++ b/Engine/Source/Editor/Sequencer/Public/SequencerUtilities.h
@@ -250,4 +250,9 @@ struct SEQUENCER_API FSequencerUtilities
 	static void SaveCurrentMovieSceneAs(TSharedRef<ISequencer> Sequencer);
 	static void SynchronizeExternalSelectionWithSequencerSelection (TSharedRef<ISequencer> Sequencer);
 	static TRange<FFrameNumber> GetTimeBounds(TSharedRef<ISequencer> Sequencer);
+
+	// Functions allowing menus to be built for modifying bindings
+	static void AddChangeClassMenu(FMenuBuilder& MenuBuilder, TSharedRef<ISequencer> Sequencer, FGuid Guid, int32 BindingIndex, TFunction<void()> OnBindingChanged);
+	static void HandleTemplateActorClassPicked(UClass* ChosenClass, TSharedRef<ISequencer> Sequencer, FGuid Guid, int32 BindingIndex, TFunction<void()> OnBindingChanged);
+	static void AddConvertBindingMenu(FMenuBuilder& MenuBuilder, TSharedRef<ISequencer> Sequencer, FGuid Guid, int32 BindingIndex, TFunction<void()> OnBindingChanged);
 };
diff --git a/Engine/Source/Runtime/MovieScene/Private/MovieSceneBindingReferences.cpp b/Engine/Source/Runtime/MovieScene/Private/MovieSceneBindingReferences.cpp
--- a/Engine/Source/Runtime/MovieScene/Private/MovieSceneBindingReferences.cpp
+++ b/Engine/Source/Runtime/MovieScene/Private/MovieSceneBindingReferences.cpp
@@ -158,7 +158,8 @@ void FMovieSceneBindingReferences::RemoveBinding(const FGuid& ObjectId)
 	}
 }
 
-void FMovieSceneBindingReferences::ResolveBindingFromLocator(int32 Index, const UE::UniversalObjectLocator::FResolveParams& ResolveParams, TArray<UObject*, TInlineAllocator<1>>& OutObjects) const
+
+UObject* FMovieSceneBindingReferences::ResolveBindingFromLocator(int32 Index, const UE::UniversalObjectLocator::FResolveParams& ResolveParams) const
 {
 	// Add our resolve param flags
 	if (ResolveParams.Context)
@@ -170,10 +171,21 @@ void FMovieSceneBindingReferences::ResolveBindingFromLocator(int32 Index, const
 	}
 
 	UObject* ResolvedObject = SortedReferences[Index].Locator.Resolve(ResolveParams).SyncGet().Object;
-	ResolvedObject = UE::MovieScene::FindBoundObjectProxy(ResolvedObject);
-	if (ResolvedObject)
+	return UE::MovieScene::FindBoundObjectProxy(ResolvedObject);
+}
+
+UObject* FMovieSceneBindingReferences::ResolveBindingInternal(const FMovieSceneBindingResolveParams& BindingResolveParams, const UE::UniversalObjectLocator::FResolveParams& LocatorResolveParams, int32 BindingIndex, int32 InternalIndex, TSharedPtr<const UE::MovieScene::FSharedPlaybackState> SharedPlaybackState) const
+{
+	// If a custom binding is present and we have valid shared playback state, resolve the custom binding
+	if (SortedReferences[InternalIndex].CustomBinding && SharedPlaybackState.IsValid())
 	{
-		OutObjects.Add(ResolvedObject);
+		UObject* ResolvedObject = SortedReferences[InternalIndex].CustomBinding->ResolveBinding(BindingResolveParams, BindingIndex, SharedPlaybackState.ToSharedRef()).Object;
+		return UE::MovieScene::FindBoundObjectProxy(ResolvedObject);
+	}
+	else
+	{
+		// Otherwise, attempt to resolve via the locator
+		return ResolveBindingFromLocator(InternalIndex, LocatorResolveParams);
 	}
 }
 
@@ -193,7 +205,11 @@ void FMovieSceneBindingReferences::ResolveBinding(const FGuid& ObjectId, const U
 
 	for (int32 Index = StartIndex; Index < Num && SortedReferences[Index].ID == ObjectId; ++Index)
 	{
-		ResolveBindingFromLocator(Index, ResolveParams, OutObjects);
+		UObject* ResolvedObject = ResolveBindingFromLocator(Index, ResolveParams);
+		if (ResolvedObject)
+		{
+			OutObjects.Add(ResolvedObject);
+		}
 	}
 }
 
@@ -204,24 +220,24 @@ void FMovieSceneBindingReferences::ResolveBinding(const FMovieSceneBindingResolv
 
 	for (int32 Index = StartIndex; Index < Num && SortedReferences[Index].ID == BindingResolveParams.ObjectBindingID; ++Index)
 	{
-		// If a custom binding is present and we have valid shared playback state, resolve the custom binding
-		if (SortedReferences[Index].CustomBinding && SharedPlaybackState.IsValid())
+		UObject* ResolvedObject = ResolveBindingInternal(BindingResolveParams, LocatorResolveParams, Index - StartIndex, Index, SharedPlaybackState);
+		if (ResolvedObject)
 		{
-			UObject* ResolvedObject = SortedReferences[Index].CustomBinding->ResolveBinding(BindingResolveParams, Index - StartIndex, SharedPlaybackState.ToSharedRef()).Object;
-			ResolvedObject = UE::MovieScene::FindBoundObjectProxy(ResolvedObject);
-			if (ResolvedObject)
-			{
-				OutObjects.Add(ResolvedObject);
-			}
-		}
-		else
-		{
-			// Otherwise, attempt to resolve via the locator
-			ResolveBindingFromLocator(Index, LocatorResolveParams, OutObjects);
+			OutObjects.Add(ResolvedObject);
 		}
 	}
 }
 
+UObject* FMovieSceneBindingReferences::ResolveSingleBinding(const FMovieSceneBindingResolveParams& BindingResolveParams, int32 BindingIndex, const UE::UniversalObjectLocator::FResolveParams& LocatorResolveParams, TSharedPtr<const UE::MovieScene::FSharedPlaybackState> SharedPlaybackState) const
+{
+	const int32 Index = Algo::LowerBoundBy(SortedReferences, BindingResolveParams.ObjectBindingID, &FMovieSceneBindingReference::ID) + BindingIndex;
+	if (SortedReferences.IsValidIndex(Index) && SortedReferences[Index].ID == BindingResolveParams.ObjectBindingID)
+	{
+		return ResolveBindingInternal(BindingResolveParams, LocatorResolveParams, BindingIndex, Index, SharedPlaybackState);
+	}
+	return nullptr;
+}
+
 void FMovieSceneBindingReferences::RemoveObjects(const FGuid& ObjectId, const TArray<UObject*>& InObjects, UObject* InContext)
 {
 	const int32 StartIndex = Algo::LowerBoundBy(SortedReferences, ObjectId, &FMovieSceneBindingReference::ID);
diff --git a/Engine/Source/Runtime/MovieScene/Private/MovieSceneCommonHelpers.cpp b/Engine/Source/Runtime/MovieScene/Private/MovieSceneCommonHelpers.cpp
--- a/Engine/Source/Runtime/MovieScene/Private/MovieSceneCommonHelpers.cpp
+++ b/Engine/Source/Runtime/MovieScene/Private/MovieSceneCommonHelpers.cpp
@@ -18,6 +18,7 @@
 #include "UObject/Package.h"
 #include "MovieSceneBindingReferences.h"
 #include "Bindings/MovieSceneSpawnableBinding.h"
+#include "Evaluation/MovieSceneEvaluationState.h"
 
 bool MovieSceneHelpers::IsSectionKeyable(const UMovieSceneSection* Section)
 {
@@ -710,6 +711,39 @@ bool MovieSceneHelpers::IsBoundToSpawnable(UMovieSceneSequence* Sequence, const
 	return false;
 }
 
+
+UObject* MovieSceneHelpers::GetSingleBoundObject(UMovieSceneSequence* Sequence, const FGuid& ObjectId, TSharedRef<const UE::MovieScene::FSharedPlaybackState> SharedPlaybackState, int32 BindingIndex)
+{
+	if (Sequence)
+	{
+		if (UMovieScene* MovieScene = Sequence->GetMovieScene())
+		{
+			if (FMovieSceneEvaluationState* EvaluationState = SharedPlaybackState->FindCapability<FMovieSceneEvaluationState>())
+			{
+				FMovieSceneSequenceIDRef SequenceID = EvaluationState->FindSequenceId(Sequence);
+
+				if (MovieScene->FindSpawnable(ObjectId))
+				{
+					TArrayView<TWeakObjectPtr<>> BoundObjects = EvaluationState->FindBoundObjects(FMovieSceneEvaluationOperand(SequenceID, ObjectId), SharedPlaybackState);
+					if (BoundObjects.Num() > 0)
+					{
+						return BoundObjects[0].Get();
+					}
+				}
+
+				const FMovieSceneBindingReferences* Refs = Sequence->GetBindingReferences();
+				if (Refs)
+				{
+					UE::UniversalObjectLocator::FResolveParams LocatorResolveParams(SharedPlaybackState->GetPlaybackContext());
+					FMovieSceneBindingResolveParams BindingResolveParams{ Sequence, ObjectId, SequenceID };
+					return Refs->ResolveSingleBinding(BindingResolveParams, BindingIndex, LocatorResolveParams, SharedPlaybackState);
+				}
+			}
+		}
+	}
+	return nullptr;
+}
+
 UObject* MovieSceneHelpers::GetObjectTemplate(UMovieSceneSequence* Sequence, const FGuid& ObjectId, TSharedRef<const UE::MovieScene::FSharedPlaybackState> SharedPlaybackState, int32 BindingIndex)
 {
 	if (Sequence)
diff --git a/Engine/Source/Runtime/MovieScene/Public/MovieSceneBindingReferences.h b/Engine/Source/Runtime/MovieScene/Public/MovieSceneBindingReferences.h
--- a/Engine/Source/Runtime/MovieScene/Public/MovieSceneBindingReferences.h
+++ b/Engine/Source/Runtime/MovieScene/Public/MovieSceneBindingReferences.h
@@ -181,6 +181,17 @@ struct FMovieSceneBindingReferences
 	 */
 	MOVIESCENE_API void ResolveBinding(const FMovieSceneBindingResolveParams& BindingResolveParams, const UE::UniversalObjectLocator::FResolveParams& LocatorResolveParams, TSharedPtr<const UE::MovieScene::FSharedPlaybackState> SharedPlaybackState, TArray<UObject*, TInlineAllocator<1>>& OutObjects) const;
 
+	/**
+	 * Resolve a binding for the specified ID and BindingIndex using a given context
+	 *
+	 * @param ObjectId					The ID to associate the object with
+	 * @param BindingResolveParams		Resolve parameters
+	 * @param LocatorResolveParams      Locator-specific resolve parameters
+	 * @param SharedPlaybackState       Shared Playback State from the Sequencer
+	 * @return The single bound objects if it exists
+	 * */
+	MOVIESCENE_API UObject* ResolveSingleBinding(const FMovieSceneBindingResolveParams& BindingResolveParams, int32 BindingIndex, const UE::UniversalObjectLocator::FResolveParams& LocatorResolveParams, TSharedPtr<const UE::MovieScene::FSharedPlaybackState> SharedPlaybackState) const;
+
 	/**
 	 * Resolve a binding for the specified ID using a given context
 	 *
@@ -209,7 +220,8 @@ struct FMovieSceneBindingReferences
 
 private:
 
-	void ResolveBindingFromLocator(int32 Index, const UE::UniversalObjectLocator::FResolveParams& LocatorResolveParams, TArray<UObject*, TInlineAllocator<1>>& OutObjects) const;
+	UObject* ResolveBindingInternal(const FMovieSceneBindingResolveParams& BindingResolveParams, const UE::UniversalObjectLocator::FResolveParams& LocatorResolveParams, int32 BindingIndex, int32 InternalIndex, TSharedPtr<const UE::MovieScene::FSharedPlaybackState> SharedPlaybackState) const;
+	UObject* ResolveBindingFromLocator(int32 Index, const UE::UniversalObjectLocator::FResolveParams& LocatorResolveParams) const;
 
 	/** The map from object binding ID to an array of references that pertain to that ID */
 	UPROPERTY()
diff --git a/Engine/Source/Runtime/MovieScene/Public/MovieSceneCommonHelpers.h b/Engine/Source/Runtime/MovieScene/Public/MovieSceneCommonHelpers.h
--- a/Engine/Source/Runtime/MovieScene/Public/MovieSceneCommonHelpers.h
+++ b/Engine/Source/Runtime/MovieScene/Public/MovieSceneCommonHelpers.h
@@ -226,6 +226,11 @@ public:
 	*/
 	static MOVIESCENE_API bool IsBoundToSpawnable(UMovieSceneSequence* Sequence, const FGuid& ObjectId, TSharedRef<const UE::MovieScene::FSharedPlaybackState> SharedPlaybackState, int32 BindingIndex = 0);
 
+	/*
+	* Returns the single bound object currently bound to the given objectid and binding index (optional).
+	*/
+	static MOVIESCENE_API UObject* GetSingleBoundObject(UMovieSceneSequence* Sequence, const FGuid& ObjectId, TSharedRef<const UE::MovieScene::FSharedPlaybackState> SharedPlaybackState, int32 BindingIndex = 0);
+
 	/*
 	* If the binding for the given ObjectId supports object templates, returns the template, otherwise returns nullptr
 	*/
diff --git a/Engine/Source/Runtime/MovieScene/Public/MovieSceneSpawnRegister.h b/Engine/Source/Runtime/MovieScene/Public/MovieSceneSpawnRegister.h
--- a/Engine/Source/Runtime/MovieScene/Public/MovieSceneSpawnRegister.h
+++ b/Engine/Source/Runtime/MovieScene/Public/MovieSceneSpawnRegister.h
@@ -176,7 +176,7 @@ public:
 	/**
 	 * Called to save the default state of the specified spawnable
 	 */
-	virtual void SaveDefaultSpawnableState(const FGuid& Guid, FMovieSceneSequenceIDRef TemplateID, TSharedRef<const FSharedPlaybackState> SharedPlaybackState) {}
+	virtual void SaveDefaultSpawnableState(const FGuid& Guid, int32 BindingIndex, FMovieSceneSequenceIDRef TemplateID, TSharedRef<const FSharedPlaybackState> SharedPlaybackState) {}
 
 	/**
 	 * Setup a new spawnable object with some default tracks and keys
@@ -210,6 +210,7 @@ public:
 
 	// Backwards compatible API, to be deprecated later
 #if WITH_EDITOR
+	virtual void SaveDefaultSpawnableState(const FGuid& Guid, FMovieSceneSequenceIDRef TemplateID, TSharedRef<const FSharedPlaybackState> SharedPlaybackState) { SaveDefaultSpawnableState(Guid, 0, TemplateID, SharedPlaybackState); }
 	MOVIESCENE_API virtual void SaveDefaultSpawnableState(FMovieSceneSpawnable& Spawnable, FMovieSceneSequenceIDRef TemplateID, TSharedRef<const FSharedPlaybackState> SharedPlaybackState);
 	MOVIESCENE_API void SaveDefaultSpawnableState(FMovieSceneSpawnable& Spawnable, FMovieSceneSequenceIDRef TemplateID, IMovieScenePlayer& Player);
 #endif
diff --git a/Engine/Source/Runtime/MovieSceneTracks/Private/Bindings/MovieSceneSpawnableActorBinding.cpp b/Engine/Source/Runtime/MovieSceneTracks/Private/Bindings/MovieSceneSpawnableActorBinding.cpp
--- a/Engine/Source/Runtime/MovieSceneTracks/Private/Bindings/MovieSceneSpawnableActorBinding.cpp
+++ b/Engine/Source/Runtime/MovieSceneTracks/Private/Bindings/MovieSceneSpawnableActorBinding.cpp
@@ -128,9 +128,22 @@ UObject* UMovieSceneSpawnableActorBindingBase::SpawnObjectInternal(UWorld* World
 	{
 		FString BindingName = GetDesiredBindingName();
 		FString ActorLabel = !BindingName.IsEmpty() ? BindingName : SpawnName.ToString();
+
 		if (FMovieScenePossessable* Possessable = MovieScene.FindPossessable(BindingId))
 		{
 			ActorLabel = Possessable->GetName();
+
+			if (UMovieSceneSequence* Sequence = MovieScene.GetTypedOuter<UMovieSceneSequence>())
+			{
+				if (const FMovieSceneBindingReferences* BindingReferences = Sequence->GetBindingReferences())
+				{
+					if (BindingReferences->GetReferences(BindingId).Num() > 1)
+					{
+						// If there are multiple bound objects, use the Object Template actor label instead of the possessable name
+						ActorLabel = ActorTemplate->GetActorLabel();
+					}
+				}
+			}
 		}
 		SpawnedActor->SetActorLabel(ActorLabel);
 	}
diff --git a/Engine/Source/Runtime/MovieSceneTracks/Public/Bindings/MovieSceneReplaceableActorBinding.h b/Engine/Source/Runtime/MovieSceneTracks/Public/Bindings/MovieSceneReplaceableActorBinding.h
--- a/Engine/Source/Runtime/MovieSceneTracks/Public/Bindings/MovieSceneReplaceableActorBinding.h
+++ b/Engine/Source/Runtime/MovieSceneTracks/Public/Bindings/MovieSceneReplaceableActorBinding.h
@@ -11,7 +11,7 @@
 * An implementation of UMovieSceneReplaceableBindingBase that uses UMovieSceneSpawnableActorBinding as the preview spawnable,
 * and has no implementation of ResolveRuntimeBindingInternal, relying instead of Sequencer's built in BindingOverride mechanism for binding at runtime.
 */
-UCLASS(BlueprintType, MinimalAPI, EditInlineNew, DefaultToInstanced)
+UCLASS(BlueprintType, MinimalAPI, EditInlineNew, DefaultToInstanced, Meta=(DisplayName="Replaceable Actor"))
 class UMovieSceneReplaceableActorBinding
 	: public UMovieSceneReplaceableBindingBase
 {
diff --git a/Engine/Source/Runtime/MovieSceneTracks/Public/Bindings/MovieSceneSpawnableActorBinding.h b/Engine/Source/Runtime/MovieSceneTracks/Public/Bindings/MovieSceneSpawnableActorBinding.h
--- a/Engine/Source/Runtime/MovieSceneTracks/Public/Bindings/MovieSceneSpawnableActorBinding.h
+++ b/Engine/Source/Runtime/MovieSceneTracks/Public/Bindings/MovieSceneSpawnableActorBinding.h
@@ -108,7 +108,7 @@ private:
 * An implementation of UMovieSceneSpawnableActorBindingBase that matches the old FMovieSceneSpawnable spawnable implementation, allowing the spawning 
 * of Actors from a UObject template which is serialized inside the Sequence.
 */
-UCLASS(BlueprintType, MinimalAPI, EditInlineNew, DefaultToInstanced)
+UCLASS(BlueprintType, MinimalAPI, EditInlineNew, DefaultToInstanced, Meta=(DisplayName="Spawnable Actor"))
 class UMovieSceneSpawnableActorBinding
 	: public UMovieSceneSpawnableActorBindingBase
 {
