Sequencer- fix issue where if you change the payload on dynamic possession, it will not refresh the binding automatically.

diff --git a/Engine/Source/Editor/Sequencer/Private/MVVM/ViewModels/ObjectBindingModel.cpp b/Engine/Source/Editor/Sequencer/Private/MVVM/ViewModels/ObjectBindingModel.cpp
--- a/Engine/Source/Editor/Sequencer/Private/MVVM/ViewModels/ObjectBindingModel.cpp
+++ b/Engine/Source/Editor/Sequencer/Private/MVVM/ViewModels/ObjectBindingModel.cpp
@@ -1088,13 +1088,20 @@ void FObjectBindingModel::OnFinishedChangingDynamicBindingProperties(const FProp
 	if (Possessable)
 	{
 		Possessable->DynamicBinding = Container->DynamicBinding;
-		return;
 	}
 	FMovieSceneSpawnable* Spawnable = MovieScene->FindSpawnable(ObjectBindingID);
 	if (Spawnable)
 	{
 		Spawnable->DynamicBinding = Container->DynamicBinding;
-		return;
+	}
+	// Force refresh the binding
+	TSharedPtr<FSequencer> Sequencer = OwnerModel->GetSequencerImpl();
+	if (Sequencer.IsValid())
+	{
+		if (FMovieSceneEvaluationState* EvaluationState = Sequencer->GetSharedPlaybackState()->FindCapability<FMovieSceneEvaluationState>())
+		{
+			EvaluationState->Invalidate(GetObjectGuid(), Sequencer->GetFocusedTemplateID());
+		}
 	}
 }
 
