Sequencer- Fix 3 issues with copy/pasting object bindings and binding naming
* Custom spawnable copy/paste between movie scenes was holding a reference to the previous scene.
* On paste default object tracks and subcomponent tracks were being created, which means you can get double tracks.
* Fix issue when creating camera bindings for new sequence and shots, the cameras had no name.

diff --git a/Engine/Plugins/MovieScene/LevelSequenceEditor/Source/LevelSequenceEditor/Private/LevelSequenceEditorToolkit.cpp b/Engine/Plugins/MovieScene/LevelSequenceEditor/Source/LevelSequenceEditor/Private/LevelSequenceEditorToolkit.cpp
--- a/Engine/Plugins/MovieScene/LevelSequenceEditor/Source/LevelSequenceEditor/Private/LevelSequenceEditorToolkit.cpp
+++ b/Engine/Plugins/MovieScene/LevelSequenceEditor/Source/LevelSequenceEditor/Private/LevelSequenceEditorToolkit.cpp
@@ -743,8 +743,10 @@ void FLevelSequenceEditorToolkit::AddShot(UMovieSceneCinematicShotTrack* ShotTra
 		FGuid CameraGuid;
 		if (bCreateSpawnableCamera)
 		{
+			FString NewName = MovieSceneHelpers::MakeUniqueBindingName(ShotSequence->GetMovieScene(), FName::NameToDisplayString(ACineCameraActor::StaticClass()->GetFName().ToString(), false));
 			UE::Sequencer::FCreateBindingParams CreateBindingParams;
 			CreateBindingParams.bSpawnable = true;
+			CreateBindingParams.BindingNameOverride = NewName;
 			CameraGuid = FSequencerUtilities::CreateBinding(GetSequencer().ToSharedRef(), *NewCamera, CreateBindingParams);
 
 			UObject* SpawnedCamera = GetSequencer()->FindSpawnedObjectOrTemplate(CameraGuid);
diff --git a/Engine/Source/Editor/Sequencer/Private/SequencerUtilities.cpp b/Engine/Source/Editor/Sequencer/Private/SequencerUtilities.cpp
--- a/Engine/Source/Editor/Sequencer/Private/SequencerUtilities.cpp
+++ b/Engine/Source/Editor/Sequencer/Private/SequencerUtilities.cpp
@@ -2395,7 +2395,9 @@ FGuid TryCreateCustomBinding(TSharedRef<ISequencer> Sequencer, UObject* CustomBi
 	if (NewCustomBinding)
 	{
 		FString DesiredBindingName = NewCustomBinding->GetDesiredBindingName();
-		FString CurrentName = DesiredBindingName.IsEmpty() ? InParams.BindingNameOverride : DesiredBindingName;
+		FString CurrentName = DesiredBindingName.IsEmpty() ? 
+			(InParams.BindingNameOverride.IsEmpty() && CustomBindingObject ? FName::NameToDisplayString(CustomBindingObject->GetName(), false) : InParams.BindingNameOverride)
+			: DesiredBindingName;
 		CurrentName = MovieSceneHelpers::MakeUniqueBindingName(OwnerMovieScene, CurrentName);
 
 		FMovieScenePossessable* NewPossessable = nullptr;
@@ -2422,22 +2424,29 @@ FGuid TryCreateCustomBinding(TSharedRef<ISequencer> Sequencer, UObject* CustomBi
 			// Spawn the object so we can position it correctly, it's going to get spawned anyway since things default to spawned.
 			SpawnedObject = Sequencer->GetSpawnRegister().SpawnObject(NewID, *OwnerMovieScene, Sequencer->GetFocusedTemplateID(), Sequencer->GetSharedPlaybackState(), 0);
 
-			// Allow the binding to set up any necessary defaults
-			NewCustomBinding->SetupDefaults(SpawnedObject, NewID, *OwnerMovieScene, Sequencer->GetSharedPlaybackState());
+			if (InParams.bSetupDefaults)
+			{
+				// Allow the binding to set up any necessary defaults
+				NewCustomBinding->SetupDefaults(SpawnedObject, NewID, *OwnerMovieScene, Sequencer->GetSharedPlaybackState());
 
-			FTransformData TransformData;
-			Sequencer->GetSpawnRegister().SetupDefaultsForSpawnable(SpawnedObject, NewID, TransformData, Sequencer, Sequencer->GetSequencerSettings());
+				FTransformData TransformData;
+				Sequencer->GetSpawnRegister().SetupDefaultsForSpawnable(SpawnedObject, NewID, TransformData, Sequencer, Sequencer->GetSequencerSettings());
+			}
 		}
 
 		Sequencer->State.Invalidate(NewID, Sequencer->GetFocusedTemplateID());
 		Sequencer->ForceEvaluate();
 
-		if (AActor* Actor = Cast<AActor>(SpawnedObject))
+		// We don't call these events in the case bSetupDefaults is false because they may add tracks.
+		if (InParams.bSetupDefaults)
 		{
-			Sequencer->OnActorAddedToSequencer().Broadcast(Actor, NewID);
-		}
+			if (AActor* Actor = Cast<AActor>(SpawnedObject))
+			{
+				Sequencer->OnActorAddedToSequencer().Broadcast(Actor, NewID);
+			}
 
-		Sequencer->OnAddBinding(NewID, OwnerMovieScene);
+			Sequencer->OnAddBinding(NewID, OwnerMovieScene);
+		}
 
 		return NewID;
 	}
@@ -2738,19 +2747,29 @@ bool FSequencerUtilities::PasteBindings(const FString& TextToImport, TSharedRef<
 					if (BindingReferences)
 					{
 						int32 BindingIndex = 0;
+						int32 SpawnableBindingIndex = 0;
 						for (UMovieSceneCustomBinding* CustomBinding : CopyableBinding->CustomBindings)
 						{
 							if (CustomBinding)
 							{
-								// We need to duplicate this as the one in CopyableBinding will be transient package owned
 								UMovieSceneCustomBinding* NewCustomBinding = Cast<UMovieSceneCustomBinding>(StaticDuplicateObject(CustomBinding, MovieScene));
 
+								// Need to re-copy the object template to avoid private object issues
+								if (UMovieSceneSpawnableBindingBase* SpawnableBinding = NewCustomBinding->AsSpawnable(Sequencer->GetSharedPlaybackState()))
+								{
+									if (CopyableBinding->SpawnableObjectTemplates.IsValidIndex(SpawnableBindingIndex))
+									{
+										SpawnableBinding->CopyObjectTemplate(CopyableBinding->SpawnableObjectTemplates[SpawnableBindingIndex++], *Sequence);
+									}
+								}
+
 								// This will either add a brand new possessable and binding (if one doesn't exist for that guid), or just add a new binding to that same possessable
 								UE::Sequencer::FCreateBindingParams CreateBindingParams;
 								CreateBindingParams.ReplacementGuid = NewGuid;
 								CreateBindingParams.BindingIndex = BindingIndex++;
 								CreateBindingParams.bAllowCustomBinding = true;
 								CreateBindingParams.CustomBinding = NewCustomBinding;
+								CreateBindingParams.bSetupDefaults = false;
 								NewGuid = CreateGenericBinding(Sequencer, nullptr, BindingReferences, CreateBindingParams);
 							}
 						}
diff --git a/Engine/Source/Editor/Sequencer/Public/SequencerUtilities.h b/Engine/Source/Editor/Sequencer/Public/SequencerUtilities.h
--- a/Engine/Source/Editor/Sequencer/Public/SequencerUtilities.h
+++ b/Engine/Source/Editor/Sequencer/Public/SequencerUtilities.h
@@ -68,6 +68,8 @@ struct FCreateBindingParams
 	TObjectPtr<UMovieSceneCustomBinding> CustomBinding = nullptr;
 	/* May be used depending on options if an asset UObject is passed in to create a custom or regular spawnable actor binding*/
 	TObjectPtr<UActorFactory> ActorFactory = nullptr;
+	/* Whether to set up default tracks and child components for a new binding. May be set false for example by a copy/paste which won't want that.*/
+	bool bSetupDefaults = true;
 };
 
 } // namespace UE::Sequencer
