Sequencer- Fix copy/paste binding crash and naming issue with spawnables using custom binding system.

diff --git a/Engine/Source/Editor/Sequencer/Private/MovieSceneCopyableBinding.h b/Engine/Source/Editor/Sequencer/Private/MovieSceneCopyableBinding.h
--- a/Engine/Source/Editor/Sequencer/Private/MovieSceneCopyableBinding.h
+++ b/Engine/Source/Editor/Sequencer/Private/MovieSceneCopyableBinding.h
@@ -27,6 +27,9 @@ public:
 	UPROPERTY(TextExportTransient)
 	TArray<TObjectPtr<UObject>> SpawnableObjectTemplates;
 
+	UPROPERTY()
+	int32 NumSpawnableObjectTemplates;
+
 	/**
 	 * Tracks are also owned by the owning Movie Sequence. We manually copy the tracks out of a binding when we copy,
 	 * because the binding stores them as a reference to a privately owned object. We store these copied tracks here,
@@ -53,6 +56,9 @@ public:
 	UPROPERTY()
 	TArray<FName> Tags;
 
+	UPROPERTY(TextExportTransient)
+	TArray<TObjectPtr<UMovieSceneCustomBinding>> CustomBindings;	
+	
 	UPROPERTY()
-	TArray<TObjectPtr<UMovieSceneCustomBinding>> CustomBindings;
+	int32 NumCustomBindings;
 };
diff --git a/Engine/Source/Editor/Sequencer/Private/SequencerUtilities.cpp b/Engine/Source/Editor/Sequencer/Private/SequencerUtilities.cpp
--- a/Engine/Source/Editor/Sequencer/Private/SequencerUtilities.cpp
+++ b/Engine/Source/Editor/Sequencer/Private/SequencerUtilities.cpp
@@ -2117,6 +2117,7 @@ void ExportObjectBindingsToText(const TArray<UMovieSceneCopyableBinding*>& Objec
 		TArray<UMovieSceneTrack*> OldTracks = ObjectToExport->Binding.StealTracks(nullptr);
 
 		TArray<UObject*, TInlineAllocator<1>> OldObjectTemplates;
+		TArray<UMovieSceneCustomBinding*, TInlineAllocator<1>> OldCustomBindings;
 		if (ObjectToExport->Spawnable.GetGuid().IsValid())
 		{
 			OldObjectTemplates.Add(ObjectToExport->Spawnable.GetObjectTemplate());
@@ -2134,14 +2135,26 @@ void ExportObjectBindingsToText(const TArray<UMovieSceneCopyableBinding*>& Objec
 						SpawnableBinding->SetObjectTemplate(nullptr);
 					}
 				}
+				OldCustomBindings.Add(CustomBinding);
 			}
+			ObjectToExport->CustomBindings.Empty();
 		}
 
+		ObjectToExport->NumCustomBindings = OldCustomBindings.Num();
+		ObjectToExport->NumSpawnableObjectTemplates = OldObjectTemplates.Num();
+
 		UExporter::ExportToOutputDevice(&Context, ObjectToExport, nullptr, Archive, TEXT("copy"), 0, PPF_ExportsNotFullyQualified | PPF_Copy | PPF_Delimited, false, ThisOuter);
 
 		// Restore the references (as we don't want to modify the original in the event of a copy operation!)
 		ObjectToExport->Binding.SetTracks(MoveTemp(OldTracks), nullptr);
 
+		ObjectToExport->CustomBindings.Append(OldCustomBindings);
+
+		for (UMovieSceneCustomBinding* CustomBinding : ObjectToExport->CustomBindings)
+		{
+			UExporter::ExportToOutputDevice(&Context, CustomBinding, nullptr, Archive, TEXT("copy"), 0, PPF_ExportsNotFullyQualified | PPF_Copy | PPF_Delimited);
+		}
+
 		int32 ObjectTemplateIndex = 0;
 		if (ObjectToExport->Spawnable.GetGuid().IsValid())
 		{
@@ -2286,8 +2299,14 @@ public:
 		{
 			return true;
 		}
-
-		return Sequencer->GetSpawnRegister().CanSpawnObject(InObjectClass);
+		else if (InObjectClass->IsChildOf<UMovieSceneCustomBinding>())
+		{
+			return true;
+		}
+		else
+		{
+			return Sequencer->GetSpawnRegister().CanSpawnObject(InObjectClass);
+		}
 	}
 
 
@@ -2300,6 +2319,10 @@ public:
 			UMovieSceneCopyableBinding* CopyableBinding = Cast<UMovieSceneCopyableBinding>(NewObject);
 			NewCopyableBindings.Add(CopyableBinding);
 		}
+		else if (NewObject->IsA<UMovieSceneCustomBinding>())
+		{
+			NewCustomBindings.Add(Cast<UMovieSceneCustomBinding>(NewObject));
+		}
 		else
 		{
 			NewSpawnableObjectTemplates.Add(NewObject);
@@ -2309,6 +2332,7 @@ public:
 public:
 	TArray<UMovieSceneCopyableBinding*> NewCopyableBindings;
 	TArray<UObject*> NewSpawnableObjectTemplates;
+	TArray<UMovieSceneCustomBinding*> NewCustomBindings;
 
 private:
 	ISequencer* Sequencer;
@@ -2328,6 +2352,7 @@ void ImportObjectBindingsFromText(ISequencer& InSequencer, const FString& TextTo
 	// and match them up with their MovieSceneCopyableBinding again.
 
 	int32 SpawnableObjectTemplateIndex = 0;
+	int32 CustomBindingIndex = 0;
 	for (auto ImportedObject : ImportedObjects)
 	{
 		if (ImportedObject->Spawnable.GetGuid().IsValid())
@@ -2335,15 +2360,23 @@ void ImportObjectBindingsFromText(ISequencer& InSequencer, const FString& TextTo
 			// This Spawnable Object Template is owned by our transient package, so you'll need to change the owner if you want to keep it later.
 			ImportedObject->SpawnableObjectTemplates.Add(Factory.NewSpawnableObjectTemplates[SpawnableObjectTemplateIndex++]);
 		}
-		else if (ImportedObject->CustomBindings.Num() > 0 && SpawnableObjectTemplateIndex < Factory.NewSpawnableObjectTemplates.Num())
+		else if (CustomBindingIndex < Factory.NewCustomBindings.Num())
 		{
-			for (UMovieSceneCustomBinding* CustomBinding : ImportedObject->CustomBindings)
+			for(int32 Index = 0; Index < ImportedObject->NumCustomBindings; ++Index)
+			{
+				ImportedObject->CustomBindings.Add(Factory.NewCustomBindings[CustomBindingIndex++]);
+			}
+
+			if (ImportedObject->CustomBindings.Num() > 0 && SpawnableObjectTemplateIndex < Factory.NewSpawnableObjectTemplates.Num())
 			{
-				if (UMovieSceneSpawnableBindingBase* SpawnableBinding = CustomBinding->AsSpawnable(InSequencer.GetSharedPlaybackState()))
+				for (UMovieSceneCustomBinding* CustomBinding : ImportedObject->CustomBindings)
 				{
-					if (SpawnableBinding->SupportsObjectTemplates())
+					if (UMovieSceneSpawnableBindingBase* SpawnableBinding = CustomBinding->AsSpawnable(InSequencer.GetSharedPlaybackState()))
 					{
-						ImportedObject->SpawnableObjectTemplates.Add(Factory.NewSpawnableObjectTemplates[SpawnableObjectTemplateIndex++]);
+						if (SpawnableBinding->SupportsObjectTemplates())
+						{
+							ImportedObject->SpawnableObjectTemplates.Add(Factory.NewSpawnableObjectTemplates[SpawnableObjectTemplateIndex++]);
+						}
 					}
 				}
 			}
@@ -2759,7 +2792,8 @@ bool FSequencerUtilities::PasteBindings(const FString& TextToImport, TSharedRef<
 								{
 									if (CopyableBinding->SpawnableObjectTemplates.IsValidIndex(SpawnableBindingIndex))
 									{
-										SpawnableBinding->CopyObjectTemplate(CopyableBinding->SpawnableObjectTemplates[SpawnableBindingIndex++], *Sequence);
+										UObject* NewObjectTemplate = StaticDuplicateObject(CopyableBinding->SpawnableObjectTemplates[SpawnableBindingIndex++], MovieScene);
+										SpawnableBinding->SetObjectTemplate(NewObjectTemplate);
 									}
 								}
 
@@ -2770,6 +2804,7 @@ bool FSequencerUtilities::PasteBindings(const FString& TextToImport, TSharedRef<
 								CreateBindingParams.bAllowCustomBinding = true;
 								CreateBindingParams.CustomBinding = NewCustomBinding;
 								CreateBindingParams.bSetupDefaults = false;
+								CreateBindingParams.BindingNameOverride = CopyableBinding->Possessable.GetName();
 								NewGuid = CreateGenericBinding(Sequencer, nullptr, BindingReferences, CreateBindingParams);
 							}
 						}
@@ -2786,9 +2821,6 @@ bool FSequencerUtilities::PasteBindings(const FString& TextToImport, TSharedRef<
 				FMovieSceneBinding NewBinding(NewGuid, CopyableBinding->Binding.GetName(), CopyableBinding->Tracks);
 				FMovieScenePossessable* Possessable = MovieScene->FindPossessable(NewGuid);
 
-				// Copy the name of the original possessable too.
-				Possessable->SetName(CopyableBinding->Possessable.GetName());
-
 				// Clear the transient flags on the copyable binding before assigning to the new possessable
 				for (UMovieSceneTrack* Track : NewBinding.GetTracks())
 				{
diff --git a/Engine/Source/Runtime/MovieSceneTracks/Private/Bindings/MovieSceneSpawnableActorBinding.cpp b/Engine/Source/Runtime/MovieSceneTracks/Private/Bindings/MovieSceneSpawnableActorBinding.cpp
--- a/Engine/Source/Runtime/MovieSceneTracks/Private/Bindings/MovieSceneSpawnableActorBinding.cpp
+++ b/Engine/Source/Runtime/MovieSceneTracks/Private/Bindings/MovieSceneSpawnableActorBinding.cpp
@@ -475,7 +475,7 @@ UMovieSceneCustomBinding* UMovieSceneSpawnableActorBinding::CreateNewCustomBindi
 
 		if (NewCustomBinding)
 		{
-			AActor* SpawnedActor = Cast<AActor>(StaticDuplicateObject(Actor, NewCustomBinding, TemplateName, RF_AllFlags & ~RF_Transient));
+			AActor* SpawnedActor = Cast<AActor>(StaticDuplicateObject(Actor, &OwnerMovieScene, TemplateName, RF_AllFlags & ~RF_Transient));
 			SpawnedActor->DetachFromActor(FDetachmentTransformRules(EDetachmentRule::KeepRelative, false));
 #if WITH_EDITORONLY_DATA
 			SpawnedActor->bIsEditorPreviewActor = false;
