Sequencer: Spawnable and possessable names are synced up to the actor label when changed. But don't do this for actors in the PIE world.

diff --git a/Engine/Source/Editor/Sequencer/Private/LevelEditorSequencerIntegration.cpp b/Engine/Source/Editor/Sequencer/Private/LevelEditorSequencerIntegration.cpp
--- a/Engine/Source/Editor/Sequencer/Private/LevelEditorSequencerIntegration.cpp
+++ b/Engine/Source/Editor/Sequencer/Private/LevelEditorSequencerIntegration.cpp
@@ -367,6 +367,13 @@ void RenameBindingRecursive(FSequencer* Sequencer, UMovieScene* MovieScene, FMov
 
 void FLevelEditorSequencerIntegration::OnActorLabelChanged(AActor* ChangedActor)
 {
+	// Sync up the spawnable or possessable name with the actor label if it is changed in editor (but not in PIE)
+	const bool bPIEWorld = ChangedActor->GetWorld()->IsPlayInEditor();
+	if (bPIEWorld)
+	{
+		return;
+	}
+
 	for (const FSequencerAndOptions& SequencerAndOptions : BoundSequencers)
 	{
 		if (!SequencerAndOptions.Options.bSyncBindingsToActorLabels)
