Sequencer: Fix regression where multiple selected bindings could not all be converted to a different binding type at once.

diff --git a/Engine/Plugins/MovieScene/LevelSequenceEditor/Source/LevelSequenceEditor/Private/LevelSequenceEditorSubsystem.cpp b/Engine/Plugins/MovieScene/LevelSequenceEditor/Source/LevelSequenceEditor/Private/LevelSequenceEditorSubsystem.cpp
--- a/Engine/Plugins/MovieScene/LevelSequenceEditor/Source/LevelSequenceEditor/Private/LevelSequenceEditorSubsystem.cpp
+++ b/Engine/Plugins/MovieScene/LevelSequenceEditor/Source/LevelSequenceEditor/Private/LevelSequenceEditorSubsystem.cpp
@@ -298,7 +298,9 @@ TSharedRef<SWidget> FMovieSceneBindingPropertyInfoDetailCustomization::OnGetConv
 	TSharedPtr<ISequencer> Sequencer = SequencerPtr.Pin();
 	if (Sequencer.IsValid() && StructPropertyHandle->IsValidHandle())
 	{
-		FSequencerUtilities::AddConvertBindingMenu(MenuBuilder, Sequencer.ToSharedRef(), BindingGuid, BindingIndex, [this, &StructBuilder, &CustomizationUtils]()
+		TArray<FSequencerConvertBindingInfo> Bindings;
+		Bindings.Add({BindingGuid, BindingIndex});
+		FSequencerUtilities::AddConvertBindingMenu(MenuBuilder, Sequencer.ToSharedRef(), Bindings, [this, &StructBuilder, &CustomizationUtils]()
 			{
 				if (IDetailsView* DetailsView = StructBuilder.GetParentCategory().GetParentLayout().GetDetailsView())
 				{
@@ -422,8 +424,8 @@ void ULevelSequenceEditorSubsystem::Initialize(FSubsystemCollectionBase& Collect
 
 	AssignActorMenuExtender = MakeShareable(new FExtender);
 	AssignActorMenuExtender->AddMenuExtension("Possessable", EExtensionHook::First, CommandList, FMenuExtensionDelegate::CreateLambda([this](FMenuBuilder& MenuBuilder) {
-		// Only add menu entries where the focused sequence is a ULevelSequence
-		if (!GetActiveSequencer())
+
+		if (!IsSelectedBindingRootPossessable())
 		{
 			return;
 		}
@@ -440,8 +442,8 @@ void ULevelSequenceEditorSubsystem::Initialize(FSubsystemCollectionBase& Collect
 	// For now we have the binding properties being a separate menu. When the UX is worked out we will likely merge the AssignActor menu away.
 	BindingPropertiesMenuExtender = MakeShareable(new FExtender);
 	BindingPropertiesMenuExtender->AddMenuExtension("Possessable", EExtensionHook::First, CommandList, FMenuExtensionDelegate::CreateLambda([this](FMenuBuilder& MenuBuilder) {
-		// Only add menu entries where the focused sequence is a ULevelSequence
-		if (!GetActiveSequencer())
+		
+		if (!IsSelectedBindingRootPossessable())
 		{
 			return;
 		}
@@ -477,7 +479,7 @@ void ULevelSequenceEditorSubsystem::Initialize(FSubsystemCollectionBase& Collect
 
 	BindingPropertiesMenuExtender->AddMenuExtension("CustomBinding", EExtensionHook::First, CommandList, FMenuExtensionDelegate::CreateLambda([this](FMenuBuilder& MenuBuilder) {
 		// Only add menu entries where the focused sequence is a ULevelSequence
-		if (!GetActiveSequencer())
+		if (!IsSelectedBindingRootPossessable())
 		{
 			return;
 		}
@@ -2492,6 +2494,32 @@ void ULevelSequenceEditorSubsystem::RebindComponent(const TArray<FMovieSceneBind
 	}
 }
 
+bool ULevelSequenceEditorSubsystem::IsSelectedBindingRootPossessable()
+{
+	TSharedPtr<ISequencer> Sequencer = GetActiveSequencer();
+	if (Sequencer)
+	{
+		TArray<FGuid> ObjectBindings;
+		Sequencer->GetSelectedObjects(ObjectBindings);
+		if (ObjectBindings.Num() >= 0)
+		{
+			UMovieSceneSequence* Sequence = Sequencer->GetFocusedMovieSceneSequence();
+			UMovieScene* MovieScene = Sequence ? Sequence->GetMovieScene() : nullptr;
+			if (MovieScene)
+			{
+				if (FMovieScenePossessable* Possessable = MovieScene->FindPossessable(ObjectBindings[0]))
+				{
+					if (!Possessable->GetParent().IsValid())
+					{
+						return true;
+					}
+				}
+			}
+		}
+	}
+	return false;
+}
+
 void ULevelSequenceEditorSubsystem::RefreshBindingDetails(IDetailsView* DetailsView, FGuid ObjectBindingID)
 {
 	if (DetailsView == nullptr)
diff --git a/Engine/Plugins/MovieScene/LevelSequenceEditor/Source/LevelSequenceEditor/Private/Misc/LevelSequenceCustomization.cpp b/Engine/Plugins/MovieScene/LevelSequenceEditor/Source/LevelSequenceEditor/Private/Misc/LevelSequenceCustomization.cpp
--- a/Engine/Plugins/MovieScene/LevelSequenceEditor/Source/LevelSequenceEditor/Private/Misc/LevelSequenceCustomization.cpp
+++ b/Engine/Plugins/MovieScene/LevelSequenceEditor/Source/LevelSequenceEditor/Private/Misc/LevelSequenceCustomization.cpp
@@ -172,12 +172,14 @@ void FLevelSequenceCustomization::ExtendObjectBindingContextMenu(FMenuBuilder& M
 	TSharedPtr<FSequencerEditorViewModel> EditorViewModel = Sequencer->GetViewModel();
 
 	FGuid ObjectBindingID = ObjectBindingModel->GetObjectGuid();
-	UMovieScene* MovieScene = Sequencer->GetFocusedMovieSceneSequence()->GetMovieScene();
+	UMovieSceneSequence* Sequence = Sequencer->GetFocusedMovieSceneSequence();
+	UMovieScene* MovieScene = Sequence->GetMovieScene();
 
 	if (!MovieScene || !ObjectBindingID.IsValid())
 	{
 		return;
 	}
+	bool bShowConvert = true;
 
 	FMovieSceneSpawnable* Spawnable = MovieScene->FindSpawnable(ObjectBindingID);
 
@@ -230,11 +232,6 @@ void FLevelSequenceCustomization::ExtendObjectBindingContextMenu(FMenuBuilder& M
 
 		MenuBuilder.AddMenuEntry(FSequencerCommands::Get().SaveCurrentSpawnableState);
 
-		MenuBuilder.AddSubMenu(
-			LOCTEXT("ConvertBindingLabel", "Convert Binding To..."),
-			LOCTEXT("ConvertBindingLabelTooltip", "Convert this binding into another binding type"),
-			FNewMenuDelegate::CreateLambda([=](FMenuBuilder& MenuBuilder) { FSequencerUtilities::AddConvertBindingMenu(MenuBuilder, Sequencer.ToSharedRef(), ObjectBindingID, 0, TFunction<void()>()); }));
-
 		MenuBuilder.AddSubMenu(
 				LOCTEXT("DynamicSpawn", "Dynamic Spawn"),
 				LOCTEXT("DynamicSpawnTooltip", "Specify a Blueprint method that will spawn or otherwise acquire a compatible actor for this binding"),
@@ -242,8 +239,9 @@ void FLevelSequenceCustomization::ExtendObjectBindingContextMenu(FMenuBuilder& M
 
 		MenuBuilder.EndSection();
 	}
-	else
+	else if (FMovieScenePossessable* Possessable = MovieScene->FindPossessable(ObjectBindingID))
 	{
+		bShowConvert = !Possessable->GetParent().IsValid();
 		bool bCustomBinding = false;
 		bool bMultipleBindings = false;
 		if (const FMovieSceneBindingReferences* BindingReferences = Sequencer->GetFocusedMovieSceneSequence()->GetBindingReferences())
@@ -252,14 +250,6 @@ void FLevelSequenceCustomization::ExtendObjectBindingContextMenu(FMenuBuilder& M
 			bMultipleBindings = BindingReferences->GetReferences(ObjectBindingID).Num() > 1;
 		}
 
-		if (!bMultipleBindings)
-		{
-			MenuBuilder.AddSubMenu(
-				LOCTEXT("ConvertBindingLabel", "Convert Binding To..."),
-				LOCTEXT("ConvertBindingLabelTooltip", "Convert this binding into another binding type"),
-				FNewMenuDelegate::CreateLambda([=](FMenuBuilder& MenuBuilder) {FSequencerUtilities::AddConvertBindingMenu(MenuBuilder, Sequencer.ToSharedRef(), ObjectBindingID, 0, TFunction<void()>()); }));
-		}
-
 		// Regular possessable
 		if (!bCustomBinding)
 		{
@@ -314,6 +304,29 @@ void FLevelSequenceCustomization::ExtendObjectBindingContextMenu(FMenuBuilder& M
 			MenuBuilder.EndSection();
 		}
 	}
+	
+	if (bShowConvert)
+	{
+		// Binding conversion
+		MenuBuilder.AddSubMenu(
+			LOCTEXT("ConvertBindingLabel", "Convert Selected Binding(s) To..."),
+			LOCTEXT("ConvertBindingLabelTooltip", "Convert selected bindings into another binding type"),
+			FNewMenuDelegate::CreateLambda([Sequencer, Sequence](FMenuBuilder& MenuBuilder) 
+			{ 
+					TArray<FSequencerConvertBindingInfo> Bindings;
+					const FMovieSceneBindingReferences* BindingReferences = Sequence->GetBindingReferences();
+					for (TViewModelPtr<IObjectBindingExtension> ObjectBindingNode : Sequencer->GetViewModel()->GetSelection()->Outliner.Filter<IObjectBindingExtension>())
+					{
+						int32 BindingIndex = 0;
+						for (const FMovieSceneBindingReference& Reference : BindingReferences->GetReferences(ObjectBindingNode->GetObjectGuid()))
+						{
+							Bindings.Add({ Reference.ID, BindingIndex++ });
+						}
+					}
+
+				FSequencerUtilities::AddConvertBindingMenu(MenuBuilder, Sequencer.ToSharedRef(), Bindings, TFunction<void()>()); 
+			}));
+	}
 
 	MenuBuilder.BeginSection("Import/Export", LOCTEXT("ImportExportMenuSectionName", "Import/Export"));
 
diff --git a/Engine/Plugins/MovieScene/LevelSequenceEditor/Source/LevelSequenceEditor/Public/LevelSequenceEditorSubsystem.h b/Engine/Plugins/MovieScene/LevelSequenceEditor/Source/LevelSequenceEditor/Public/LevelSequenceEditorSubsystem.h
--- a/Engine/Plugins/MovieScene/LevelSequenceEditor/Source/LevelSequenceEditor/Public/LevelSequenceEditorSubsystem.h
+++ b/Engine/Plugins/MovieScene/LevelSequenceEditor/Source/LevelSequenceEditor/Public/LevelSequenceEditorSubsystem.h
@@ -282,6 +282,8 @@ private:
 	void GetRebindComponentNames(TArray<FName>& OutComponentNames);
 	void RebindComponentMenu(FMenuBuilder& MenuBuilder);
 
+	bool IsSelectedBindingRootPossessable();
+
 	FDelegateHandle OnSequencerCreatedHandle;
 
 	/* List of sequencers that have been created */
diff --git a/Engine/Source/Editor/MovieSceneTools/Private/Bindings/MovieSceneSpawnableBindingCustomization.cpp b/Engine/Source/Editor/MovieSceneTools/Private/Bindings/MovieSceneSpawnableBindingCustomization.cpp
--- a/Engine/Source/Editor/MovieSceneTools/Private/Bindings/MovieSceneSpawnableBindingCustomization.cpp
+++ b/Engine/Source/Editor/MovieSceneTools/Private/Bindings/MovieSceneSpawnableBindingCustomization.cpp
@@ -17,6 +17,7 @@
 #include "MovieScene.h"
 #include "MovieSceneSection.h"
 #include "Tracks/MovieSceneBindingLifetimeTrack.h"
+#include "Tracks/MovieSceneSpawnTrack.h"
 
 TSharedRef<IDetailCustomization> FMovieSceneSpawnableBindingCustomization::MakeInstance(UMovieScene* InMovieScene, FGuid InBindingGuid)
 {
@@ -45,7 +46,7 @@ void FMovieSceneSpawnableBindingCustomization::OnSpawnOwnershipChanged()
 			// Overwrite the completion state for all spawn sections to ensure the expected behaviour.
 			EMovieSceneCompletionMode NewCompletionMode = *SpawnOwnership == ESpawnOwnership::InnerSequence ? EMovieSceneCompletionMode::RestoreState : EMovieSceneCompletionMode::KeepState;
 
-			// Make all binding lifetime sections retain state
+			// Make all binding lifetime and spawn track sections retain state
 			UMovieSceneBindingLifetimeTrack* BindingLifetimeTrack = MovieScene->FindTrack<UMovieSceneBindingLifetimeTrack>(BindingGuid);
 			if (BindingLifetimeTrack)
 			{
@@ -55,6 +56,16 @@ void FMovieSceneSpawnableBindingCustomization::OnSpawnOwnershipChanged()
 					Section->EvalOptions.CompletionMode = NewCompletionMode;
 				}
 			}
+
+			UMovieSceneSpawnTrack* SpawnTrack = MovieScene->FindTrack<UMovieSceneSpawnTrack>(BindingGuid);
+			if (SpawnTrack)
+			{
+				for (UMovieSceneSection* Section : SpawnTrack->GetAllSections())
+				{
+					Section->Modify();
+					Section->EvalOptions.CompletionMode = NewCompletionMode;
+				}
+			}
 		}
 	}
 }
diff --git a/Engine/Source/Editor/Sequencer/Private/SequencerUtilities.cpp b/Engine/Source/Editor/Sequencer/Private/SequencerUtilities.cpp
--- a/Engine/Source/Editor/Sequencer/Private/SequencerUtilities.cpp
+++ b/Engine/Source/Editor/Sequencer/Private/SequencerUtilities.cpp
@@ -71,6 +71,9 @@
 #include "ClassViewerFilter.h"
 #include "Framework/Application/SlateApplication.h"
 #include "SequencerCommands.h"
+#include "MVVM/Selection/Selection.h"
+#include "UnrealEdGlobals.h"
+#include "Misc/FeedbackContext.h"
 
 #define LOCTEXT_NAMESPACE "FSequencerUtilities"
 
@@ -3884,83 +3887,154 @@ void FSequencerUtilities::HandleTemplateActorClassPicked(UClass* ChosenClass, TS
 	}
 }
 
-void FSequencerUtilities::AddConvertBindingMenu(FMenuBuilder& MenuBuilder, TSharedRef<ISequencer> Sequencer, FGuid ObjectBindingID, int32 BindingIndex, TFunction<void()> OnBindingChanged)
+void FSequencerUtilities::AddConvertBindingMenu(FMenuBuilder& MenuBuilder, TSharedRef<ISequencer> Sequencer, const TArray<FSequencerConvertBindingInfo>& BindingsToConvert, TFunction<void()> OnBindingChanged)
 {
 	UMovieSceneSequence* Sequence = Sequencer->GetFocusedMovieSceneSequence();
 	UMovieScene* MovieScene = Sequence ? Sequence->GetMovieScene() : nullptr;
 
-	if (!MovieScene || !ObjectBindingID.IsValid())
+	if (!MovieScene || BindingsToConvert.IsEmpty())
 	{
 		return;
 	}
 
-	FMovieSceneSpawnable* Spawnable = MovieScene->FindSpawnable(ObjectBindingID);
+	auto ConvertBindings = [Sequencer, MovieScene, BindingsToConvert, OnBindingChanged](TFunction<FMovieScenePossessable*(FGuid,int32)> DoConvert) {
+		using namespace UE::Sequencer;
 
-	if (Spawnable)
+		if (MovieScene->IsReadOnly())
+		{
+			FSequencerUtilities::ShowReadOnlyError();
+			return;
+		}
+
+		if (BindingsToConvert.Num() > 0)
+		{
+			const FScopedTransaction Transaction(LOCTEXT("ConvertSelectedNodesPossessable", "Convert Selected Nodes to Possessables"));
+			MovieScene->Modify();
+
+			FScopedSlowTask SlowTask(BindingsToConvert.Num(), LOCTEXT("ConvertPossessablesProgress", "Converting Selected Spawnable Nodes to Possessables"));
+			SlowTask.MakeDialog(true);
+
+			TArray<AActor*> PossessedActors;
+			for (const FSequencerConvertBindingInfo& BindingInfo : BindingsToConvert)
+			{
+				SlowTask.EnterProgressFrame();
+
+				FMovieScenePossessable* Possessable = DoConvert(BindingInfo.BindingID, BindingInfo.BindingIndex);
+
+				Sequencer->ForceEvaluate();
+
+				for (TWeakObjectPtr<> WeakObject : Sequencer->FindBoundObjects(Possessable->GetGuid(), Sequencer->GetFocusedTemplateID()))
+				{
+					if (AActor* PossessedActor = Cast<AActor>(WeakObject.Get()))
+					{
+						PossessedActors.Add(PossessedActor);
+					}
+				}
+
+				if (GWarn->ReceivedUserCancel())
+				{
+					break;
+				}
+			}
+
+			if (PossessedActors.Num())
+			{
+				const bool bNotifySelectionChanged = true;
+				const bool bDeselectBSP = true;
+				const bool bWarnAboutTooManyActors = false;
+				const bool bSelectEvenIfHidden = false;
+
+				GEditor->GetSelectedActors()->Modify();
+				GEditor->GetSelectedActors()->BeginBatchSelectOperation();
+				GEditor->SelectNone(bNotifySelectionChanged, bDeselectBSP, bWarnAboutTooManyActors);
+				for (auto PossessedActor : PossessedActors)
+				{
+					GEditor->SelectActor(PossessedActor, true, bNotifySelectionChanged, bSelectEvenIfHidden);
+				}
+				GEditor->GetSelectedActors()->EndBatchSelectOperation();
+				GEditor->NoteSelectionChange();
+
+				Sequencer->NotifyMovieSceneDataChanged(EMovieSceneDataChangeType::MovieSceneStructureItemsChanged);
+			}
+		}
+		if (OnBindingChanged != nullptr)
+		{
+			OnBindingChanged();
+		}
+	};
+
+	// Can convert to possessable
+	if (Algo::AllOf(BindingsToConvert, [&Sequence, &MovieScene](const FSequencerConvertBindingInfo& BindingInfo) {
+		if (FMovieSceneSpawnable* Spawnable = MovieScene->FindSpawnable(BindingInfo.BindingID))
+		{
+			return true;
+		}
+		else if (const FMovieSceneBindingReferences* BindingReferences = Sequence->GetBindingReferences())
+		{
+			TArrayView<const FMovieSceneBindingReference> BindingReferencesList = BindingReferences->GetReferences(BindingInfo.BindingID);
+			if (BindingReferencesList.IsValidIndex(BindingInfo.BindingIndex) && BindingReferencesList[BindingInfo.BindingIndex].CustomBinding != nullptr)
+			{
+				return true;
+			}
+		}
+		return false;
+	}))
 	{
-		MenuBuilder.AddMenuEntry(FSequencerCommands::Get().ConvertToPossessable, NAME_None, LOCTEXT("ConvertToPossessable", "Possessable"));
 		MenuBuilder.AddMenuEntry(
 			LOCTEXT("ConvertToPossessable", "Possessable"),
-			LOCTEXT("ConvertToPossessableTooltip", "Convert selected binding to a possessable"),
+			LOCTEXT("ConvertToPossessableTooltip", "Convert selected binding(s) to a possessable"),
 			FSlateIcon(),
-			FUIAction(FExecuteAction::CreateLambda([Sequencer, ObjectBindingID, BindingIndex, OnBindingChanged]() 
-				{ 
-					FSequencerUtilities::ConvertToPossessable(Sequencer, ObjectBindingID, BindingIndex); 
-					if (OnBindingChanged != nullptr)
-					{
-						OnBindingChanged();
-					}
-				})));
+			FUIAction(FExecuteAction::CreateLambda([Sequencer, MovieScene, BindingsToConvert, ConvertBindings, OnBindingChanged]()
+			{
+				ConvertBindings([Sequencer](FGuid BindingID, int32 BindingIndex){return FSequencerUtilities::ConvertToPossessable(Sequencer, BindingID, BindingIndex);});
+			})));
 	}
-	else if (const FMovieSceneBindingReferences* BindingReferences = Sequencer->GetFocusedMovieSceneSequence()->GetBindingReferences())
+	
+	// Can convert to custom bindings
+
+	TArrayView<const TSubclassOf<UMovieSceneCustomBinding>> PrioritySortedCustomBindingTypes = Sequencer->GetSupportedCustomBindingTypes();
+	for (const TSubclassOf<UMovieSceneCustomBinding>& CustomBindingType : PrioritySortedCustomBindingTypes)
 	{
-		TArrayView<const FMovieSceneBindingReference> BindingReferencesList = Sequencer->GetFocusedMovieSceneSequence()->GetBindingReferences()->GetReferences(ObjectBindingID);
-		if (BindingReferencesList.IsValidIndex(BindingIndex))
+		if (Algo::AllOf(BindingsToConvert, [&Sequence, &MovieScene, &CustomBindingType, Sequencer](const FSequencerConvertBindingInfo& BindingInfo) 
 		{
-			if (BindingReferencesList[BindingIndex].CustomBinding != nullptr)
+			if (FMovieSceneSpawnable* Spawnable = MovieScene->FindSpawnable(BindingInfo.BindingID))
 			{
-				MenuBuilder.AddMenuEntry(
-					LOCTEXT("ConvertToPossessable", "Possessable"),
-					LOCTEXT("ConvertToPossessableTooltip", "Convert selected binding to a possessable"),
-					FSlateIcon(),
-					FUIAction(FExecuteAction::CreateLambda([Sequencer, ObjectBindingID, BindingIndex, OnBindingChanged]() 
-						{ 
-							FSequencerUtilities::ConvertToPossessable(Sequencer, ObjectBindingID, BindingIndex);
-							if (OnBindingChanged != nullptr)
-							{
-								OnBindingChanged();
-							}
-						})));
+				if (UObject* CurrentBoundObject = Spawnable->GetObjectTemplate())
+				{
+					return CustomBindingType && CustomBindingType->GetDefaultObject<UMovieSceneCustomBinding>()->SupportsBindingCreationFromObject(CurrentBoundObject);
+				}
 			}
-
-			const FMovieSceneBindingReference& CurrentBindingReference = BindingReferencesList[BindingIndex];
-			UE::UniversalObjectLocator::FResolveParams LocatorResolveParams(Sequencer->GetSharedPlaybackState()->GetPlaybackContext());
-			FMovieSceneBindingResolveParams BindingResolveParams{ Sequence, ObjectBindingID, Sequencer->GetFocusedTemplateID()};
-			if (UObject* CurrentBoundObject = BindingReferences->ResolveSingleBinding(BindingResolveParams, BindingIndex, LocatorResolveParams, Sequencer->GetSharedPlaybackState()))
+			else if (const FMovieSceneBindingReferences* BindingReferences = Sequence->GetBindingReferences())
 			{
-				TArrayView<const TSubclassOf<UMovieSceneCustomBinding>> PrioritySortedCustomBindingTypes = Sequencer->GetSupportedCustomBindingTypes();
-				for (const TSubclassOf<UMovieSceneCustomBinding>& CustomBindingType : PrioritySortedCustomBindingTypes)
+				TArrayView<const FMovieSceneBindingReference> BindingReferencesList = Sequencer->GetFocusedMovieSceneSequence()->GetBindingReferences()->GetReferences(BindingInfo.BindingID);
+
+				if (const FMovieSceneBindingReference* CurrentBindingReference = BindingReferences->GetReference(BindingInfo.BindingID, BindingInfo.BindingIndex))
 				{
-					if (CustomBindingType
-						&& (!CurrentBindingReference.CustomBinding
-							|| CurrentBindingReference.CustomBinding->GetClass() != CustomBindingType)
-						&& CustomBindingType->GetDefaultObject<UMovieSceneCustomBinding>()->SupportsConversionFromBinding(CurrentBindingReference, CurrentBoundObject))
+					UE::UniversalObjectLocator::FResolveParams LocatorResolveParams(Sequencer->GetSharedPlaybackState()->GetPlaybackContext());
+					FMovieSceneBindingResolveParams BindingResolveParams{ Sequence, BindingInfo.BindingID, Sequencer->GetFocusedTemplateID() };
+					if (UObject* CurrentBoundObject = BindingReferences->ResolveSingleBinding(BindingResolveParams, BindingInfo.BindingIndex, LocatorResolveParams, Sequencer->GetSharedPlaybackState()))
 					{
-						MenuBuilder.AddMenuEntry(
-							CustomBindingType->GetDefaultObject<UMovieSceneCustomBinding>()->GetBindingTypePrettyName(),
-							FText::Format(LOCTEXT("ConvertToCustomBindingTooltip", "Convert selected binding to {0}"), CustomBindingType->GetDefaultObject<UMovieSceneCustomBinding>()->GetBindingTypePrettyName()),
-							FSlateIcon(),
-							FUIAction(FExecuteAction::CreateLambda([Sequencer, ObjectBindingID, CustomBindingType, BindingIndex, OnBindingChanged]() 
-								{ 
-									FSequencerUtilities::ConvertToCustomBinding(Sequencer, ObjectBindingID, CustomBindingType, BindingIndex); 
-									if (OnBindingChanged != nullptr)
-									{
-										OnBindingChanged();
-									}
-								})));
+						if (CustomBindingType
+							&& (!CurrentBindingReference->CustomBinding
+								|| CurrentBindingReference->CustomBinding->GetClass() != CustomBindingType)
+							&& CustomBindingType->GetDefaultObject<UMovieSceneCustomBinding>()->SupportsConversionFromBinding(*CurrentBindingReference, CurrentBoundObject))
+						{
+							return true;
+						}
 					}
 				}
 			}
+			return false;
+		}))
+		{
+			MenuBuilder.AddMenuEntry(
+				CustomBindingType->GetDefaultObject<UMovieSceneCustomBinding>()->GetBindingTypePrettyName(),
+				FText::Format(LOCTEXT("ConvertToCustomBindingTooltip", "Convert selected binding to {0}"), CustomBindingType->GetDefaultObject<UMovieSceneCustomBinding>()->GetBindingTypePrettyName()),
+				FSlateIcon(),
+				FUIAction(FExecuteAction::CreateLambda([Sequencer, MovieScene, BindingsToConvert, ConvertBindings, CustomBindingType, OnBindingChanged]()
+				{
+					ConvertBindings([Sequencer, CustomBindingType](FGuid BindingID, int32 BindingIndex){return FSequencerUtilities::ConvertToCustomBinding(Sequencer, BindingID, CustomBindingType, BindingIndex);});
+				})));
 		}
 	}
 }
diff --git a/Engine/Source/Editor/Sequencer/Public/SequencerUtilities.h b/Engine/Source/Editor/Sequencer/Public/SequencerUtilities.h
--- a/Engine/Source/Editor/Sequencer/Public/SequencerUtilities.h
+++ b/Engine/Source/Editor/Sequencer/Public/SequencerUtilities.h
@@ -165,6 +165,18 @@ struct FMovieScenePasteBindingsParams
 	bool bDuplicateExistingActors;
 };
 
+USTRUCT()
+struct FSequencerConvertBindingInfo
+{
+	GENERATED_BODY()
+
+	UPROPERTY()
+	FGuid BindingID;
+
+	UPROPERTY()
+	int32 BindingIndex;
+};
+
 struct SEQUENCER_API FSequencerUtilities
 {
 	/* Creates a button (used for +Section) that opens a ComboButton with a user-defined sub-menu content. */
@@ -256,5 +268,5 @@ struct SEQUENCER_API FSequencerUtilities
 	// Functions allowing menus to be built for modifying bindings
 	static void AddChangeClassMenu(FMenuBuilder& MenuBuilder, TSharedRef<ISequencer> Sequencer, FGuid Guid, int32 BindingIndex, TFunction<void()> OnBindingChanged);
 	static void HandleTemplateActorClassPicked(UClass* ChosenClass, TSharedRef<ISequencer> Sequencer, FGuid Guid, int32 BindingIndex, TFunction<void()> OnBindingChanged);
-	static void AddConvertBindingMenu(FMenuBuilder& MenuBuilder, TSharedRef<ISequencer> Sequencer, FGuid Guid, int32 BindingIndex, TFunction<void()> OnBindingChanged);
+	static void AddConvertBindingMenu(FMenuBuilder& MenuBuilder, TSharedRef<ISequencer> Sequencer, const TArray<FSequencerConvertBindingInfo>& BindingsToConvert, TFunction<void()> OnBindingChanged);
 };
