Sequencer: Get the editor world when dealing with actors in the editor world

This fixes a crash trying to copy/paste a niagara section

#rnx

diff --git a/Engine/Source/Editor/Sequencer/Private/SequencerUtilities.cpp b/Engine/Source/Editor/Sequencer/Private/SequencerUtilities.cpp
--- a/Engine/Source/Editor/Sequencer/Private/SequencerUtilities.cpp
+++ b/Engine/Source/Editor/Sequencer/Private/SequencerUtilities.cpp
@@ -837,8 +837,8 @@ FGuid FSequencerUtilities::CreateCameraWithRig(TSharedRef<ISequencer> Sequencer,
 	}
 
 	// Create a cine camera actor
-	UWorld* PlaybackContext = Sequencer->GetPlaybackContext()->GetWorld();
-	OutActor = PlaybackContext->SpawnActor<ACineCameraActor>();
+	UWorld* World = GCurrentLevelEditingViewportClient ? GCurrentLevelEditingViewportClient->GetWorld() : nullptr;
+	OutActor = World->SpawnActor<ACineCameraActor>();
 
 	FString NewCameraName = MovieSceneHelpers::MakeUniqueSpawnableName(MovieScene, FName::NameToDisplayString(ACineCameraActor::StaticClass()->GetFName().ToString(), false));
 	UE::Sequencer::FCreateBindingParams CreateBindingParams;
@@ -1158,8 +1158,8 @@ FMovieScenePossessable* FSequencerUtilities::ConvertToPossessable(TSharedRef<ISe
 		SpawnInfo.bDeferConstruction = true;
 		SpawnInfo.Template = SpawnableActorTemplate;
 
-		UWorld* PlaybackContext = Sequencer->GetPlaybackContext()->GetWorld();
-		AActor* PossessedActor = PlaybackContext->SpawnActor(ObjectToConvert->GetClass(), &DefaultTransform, SpawnInfo);
+		UWorld* World = GCurrentLevelEditingViewportClient ? GCurrentLevelEditingViewportClient->GetWorld() : nullptr;
+		AActor* PossessedActor = World->SpawnActor(ObjectToConvert->GetClass(), &DefaultTransform, SpawnInfo);
 
 		if (!PossessedActor)
 		{
@@ -2726,7 +2726,7 @@ bool FSequencerUtilities::PasteBindings(const FString& TextToImport, TSharedRef<
 
 	UMovieScene* RootMovieScene = Sequencer->GetRootMovieSceneSequence()->GetMovieScene();
 
-	UWorld* World = Sequencer->GetPlaybackContext()->GetWorld();
+	UWorld* World = GCurrentLevelEditingViewportClient ? GCurrentLevelEditingViewportClient->GetWorld() : nullptr;
 
 	const FScopedTransaction Transaction(LOCTEXT("PasteBindings", "Paste Bindings"));
 
@@ -2998,10 +2998,9 @@ bool FSequencerUtilities::PasteBindings(const FString& TextToImport, TSharedRef<
 			continue;
 		}
 		FMovieScenePossessable* Possessable = MovieScene->FindPossessable(PossessableGuids[PossessableGuidIndex]);
-		UWorld* PlaybackContext = Sequencer->GetPlaybackContext()->GetWorld();
-		if (Possessable && PlaybackContext)
+		if (Possessable && World)
 		{
-			for (TActorIterator<AActor> ActorItr(PlaybackContext); ActorItr; ++ActorItr)
+			for (TActorIterator<AActor> ActorItr(World); ActorItr; ++ActorItr)
 			{
 				AActor* Actor = *ActorItr;
 				if (Actor && PossessableGuidIndex < PossessableObjectNames.Num() && PossessableObjectNames[PossessableGuidIndex].Contains(Actor->GetName()))
