Sequencer: Fix initialization

#rnx

diff --git a/Engine/Source/Editor/Sequencer/Public/SequencerUtilities.h b/Engine/Source/Editor/Sequencer/Public/SequencerUtilities.h
--- a/Engine/Source/Editor/Sequencer/Public/SequencerUtilities.h
+++ b/Engine/Source/Editor/Sequencer/Public/SequencerUtilities.h
@@ -174,7 +174,7 @@ struct FSequencerConvertBindingInfo
 	FGuid BindingID;
 
 	UPROPERTY()
-	int32 BindingIndex;
+	int32 BindingIndex = -1;
 };
 
 struct SEQUENCER_API FSequencerUtilities
