Sequencer: Spawnable and possessable names are synced up to the actor label when changed. But don't do this for actors in the PIE world.

Follow up fix for valid world

diff --git a/Engine/Source/Editor/Sequencer/Private/LevelEditorSequencerIntegration.cpp b/Engine/Source/Editor/Sequencer/Private/LevelEditorSequencerIntegration.cpp
--- a/Engine/Source/Editor/Sequencer/Private/LevelEditorSequencerIntegration.cpp
+++ b/Engine/Source/Editor/Sequencer/Private/LevelEditorSequencerIntegration.cpp
@@ -368,7 +368,7 @@ void RenameBindingRecursive(FSequencer* Sequencer, UMovieScene* MovieScene, FMov
 void FLevelEditorSequencerIntegration::OnActorLabelChanged(AActor* ChangedActor)
 {
 	// Sync up the spawnable or possessable name with the actor label if it is changed in editor (but not in PIE)
-	const bool bPIEWorld = ChangedActor->GetWorld()->IsPlayInEditor();
+	const bool bPIEWorld = ChangedActor != nullptr && ChangedActor->GetWorld() != nullptr && ChangedActor->GetWorld()->IsPlayInEditor();
 	if (bPIEWorld)
 	{
 		return;
