Sequencer: Possessables can possess a spawnable, as defined by SpawnableObjectBindingID. If the user changes that spawnable's class/object template, look for any possessables that possess it and also change the possessable's PossessedObjectClass.

This fixes a bug where doing a Change Class on a spawnable will fail to update any possessable's PossessedObjectClass.

diff --git a/Engine/Source/Editor/Sequencer/Private/SequencerUtilities.cpp b/Engine/Source/Editor/Sequencer/Private/SequencerUtilities.cpp
--- a/Engine/Source/Editor/Sequencer/Private/SequencerUtilities.cpp
+++ b/Engine/Source/Editor/Sequencer/Private/SequencerUtilities.cpp
@@ -3863,6 +3863,38 @@ void FSequencerUtilities::AddChangeClassMenu(FMenuBuilder& MenuBuilder, TSharedR
 	);
 }
 
+void UpdatePossessedClasses(UMovieScene* MovieScene, FMovieSceneSequenceIDRef SequenceID, const FMovieSceneSequenceHierarchy* Hierarchy, FGuid ObjectBindingID, UClass* ChosenClass)
+{
+	for (int32 Index = 0; Index < MovieScene->GetPossessableCount(); ++Index)
+	{
+		FMovieScenePossessable& Possessable = MovieScene->GetPossessable(Index);
+		if (Possessable.GetSpawnableObjectBindingID().GetGuid() == ObjectBindingID && Possessable.GetPossessedObjectClass() != ChosenClass)
+		{
+			MovieScene->Modify();
+
+			Possessable.SetPossessedObjectClass(ChosenClass);
+		}
+	}
+
+	if (const FMovieSceneSequenceHierarchyNode* Node = Hierarchy->FindNode(SequenceID))
+	{
+		for (FMovieSceneSequenceIDRef ChildID : Node->Children)
+		{
+			const FMovieSceneSubSequenceData* SubData = Hierarchy->FindSubData(ChildID);
+			if (SubData)
+			{
+				UMovieSceneSequence* SubSequence = SubData->GetSequence();
+				UMovieScene* SubMovieScene = SubSequence ? SubSequence->GetMovieScene() : nullptr;
+
+				if (SubMovieScene)
+				{
+					UpdatePossessedClasses(SubMovieScene, ChildID, Hierarchy, ObjectBindingID, ChosenClass);
+				}
+			}
+		}
+	}
+}
+
 void FSequencerUtilities::HandleTemplateActorClassPicked(UClass* ChosenClass, TSharedRef<ISequencer> Sequencer, FGuid ObjectBindingID, int32 BindingIndex, TFunction<void()> OnBindingChanged)
 {
 	UMovieScene* MovieScene = Sequencer->GetFocusedMovieSceneSequence()->GetMovieScene();
@@ -3874,6 +3906,10 @@ void FSequencerUtilities::HandleTemplateActorClassPicked(UClass* ChosenClass, TS
 	TValueOrError<FNewSpawnable, FText> Result = Sequencer->GetSpawnRegister().CreateNewSpawnableType(*ChosenClass, *MovieScene, nullptr);
 	if (Result.IsValid())
 	{
+		FMovieSceneRootEvaluationTemplateInstance& RootInstance = Sequencer->GetEvaluationTemplate();
+		const FMovieSceneSequenceHierarchy* Hierarchy = RootInstance.GetCompiledDataManager()->FindHierarchy(RootInstance.GetCompiledDataID());
+		UpdatePossessedClasses(Sequencer->GetRootMovieSceneSequence()->GetMovieScene(), MovieSceneSequenceID::Root, Hierarchy, ObjectBindingID, ChosenClass);
+
 		MovieSceneHelpers::SetObjectTemplate(Sequencer->GetFocusedMovieSceneSequence(), ObjectBindingID, Result.GetValue().ObjectTemplate, Sequencer->GetSharedPlaybackState(), BindingIndex);
 
 		Sequencer->GetSpawnRegister().DestroySpawnedObject(ObjectBindingID, Sequencer->GetFocusedTemplateID(), Sequencer->GetSharedPlaybackState(), BindingIndex);
