Take Recorder: Skip UpdateCachedFlags while recording

diff --git a/Engine/Source/Editor/Sequencer/Private/MVVM/ViewModels/SequenceModel.cpp b/Engine/Source/Editor/Sequencer/Private/MVVM/ViewModels/SequenceModel.cpp
--- a/Engine/Source/Editor/Sequencer/Private/MVVM/ViewModels/SequenceModel.cpp
+++ b/Engine/Source/Editor/Sequencer/Private/MVVM/ViewModels/SequenceModel.cpp
@@ -357,17 +357,31 @@ void FSequenceModel::PerformDrop(const FViewModelPtr& TargetModel, const FDragDr
 
 void FSequenceModel::OnModifiedIndirectly(UMovieSceneSignedObject*)
 {
-	if (FOutlinerCacheExtension* OutlinerCache = GetSharedData()->CastThis<FOutlinerCacheExtension>())
+	TSharedPtr<ISequencer> Sequencer = GetEditor()->GetSequencer();
+	const bool bRecording = Sequencer && Sequencer->OnGetIsRecording().IsBound() && Sequencer->OnGetIsRecording().Execute();
+
+	// Disregard updating cached flags during recording
+	if (!bRecording)
 	{
-		OutlinerCache->UpdateCachedFlags();
+		if (FOutlinerCacheExtension* OutlinerCache = GetSharedData()->CastThis<FOutlinerCacheExtension>())
+		{
+			OutlinerCache->UpdateCachedFlags();
+		}
 	}
 }
 
 void FSequenceModel::OnModifiedDirectly(UMovieSceneSignedObject*)
 {
-	if (FOutlinerCacheExtension* OutlinerCache = GetSharedData()->CastThis<FOutlinerCacheExtension>())
+	TSharedPtr<ISequencer> Sequencer = GetEditor()->GetSequencer();
+	const bool bRecording = Sequencer && Sequencer->OnGetIsRecording().IsBound() && Sequencer->OnGetIsRecording().Execute();
+
+	// Disregard updating cached flags during recording
+	if (!bRecording)
 	{
-		OutlinerCache->UpdateCachedFlags();
+		if (FOutlinerCacheExtension* OutlinerCache = GetSharedData()->CastThis<FOutlinerCacheExtension>())
+		{
+			OutlinerCache->UpdateCachedFlags();
+		}
 	}
 }
 
