Sequencer: Nullptr check for movie scene

#rnx

diff --git a/Engine/Source/Editor/Sequencer/Private/Sequencer.cpp b/Engine/Source/Editor/Sequencer/Private/Sequencer.cpp
--- a/Engine/Source/Editor/Sequencer/Private/Sequencer.cpp
+++ b/Engine/Source/Editor/Sequencer/Private/Sequencer.cpp
@@ -3959,7 +3959,7 @@ FGuid FSequencer::GetHandleToObject( UObject* Object, bool bCreateHandleIfMissin
 	}
 
 	UMovieSceneSequence* FocusedMovieSceneSequence = GetFocusedMovieSceneSequence();
-	UMovieScene* FocusedMovieScene = FocusedMovieSceneSequence->GetMovieScene();
+	UMovieScene* FocusedMovieScene = FocusedMovieSceneSequence ? FocusedMovieSceneSequence->GetMovieScene() : nullptr;
 	
 	if (!FocusedMovieScene)
 	{
