Sequencer: Add new section now always creates an infinite section instead of one that starts at the current frame. This is more consistent with how new tracks/sections are created.

diff --git a/Engine/Source/Editor/Sequencer/Private/SequencerUtilities.cpp b/Engine/Source/Editor/Sequencer/Private/SequencerUtilities.cpp
--- a/Engine/Source/Editor/Sequencer/Private/SequencerUtilities.cpp
+++ b/Engine/Source/Editor/Sequencer/Private/SequencerUtilities.cpp
@@ -108,9 +108,6 @@ void FSequencerUtilities::CreateNewSection(UMovieSceneTrack* InTrack, TWeakPtr<I
 		return;
 	}
 
-	FQualifiedFrameTime CurrentTime = Sequencer->GetLocalTime();
-	FFrameNumber PlaybackEnd = UE::MovieScene::DiscreteExclusiveUpper(Sequencer->GetFocusedMovieSceneSequence()->GetMovieScene()->GetPlaybackRange());
-
 	FScopedTransaction Transaction(LOCTEXT("AddSectionTransactionText", "Add Section"));
 	if (UMovieSceneSection* NewSection = InTrack->CreateNewSection())
 	{
@@ -128,7 +125,11 @@ void FSequencerUtilities::CreateNewSection(UMovieSceneTrack* InTrack, TWeakPtr<I
 
 		InTrack->Modify();
 
-		NewSection->SetRange(TRange<FFrameNumber>(CurrentTime.Time.FrameNumber, PlaybackEnd));
+		if (Sequencer->GetInfiniteKeyAreas())
+		{
+			NewSection->SetRange(TRange<FFrameNumber>::All());
+		}
+
 		NewSection->SetOverlapPriority(OverlapPriority);
 		NewSection->SetRowIndex(InRowIndex);
 		NewSection->SetBlendType(InBlendType);
@@ -162,9 +163,6 @@ void FSequencerUtilities::PopulateMenu_CreateNewSection(FMenuBuilder& MenuBuilde
 			return;
 		}
 
-		FQualifiedFrameTime CurrentTime = Sequencer->GetLocalTime();
-		FFrameNumber PlaybackEnd = UE::MovieScene::DiscreteExclusiveUpper(Sequencer->GetFocusedMovieSceneSequence()->GetMovieScene()->GetPlaybackRange());
-
 		FScopedTransaction Transaction(LOCTEXT("AddSectionTransactionText", "Add Section"));
 		if (UMovieSceneSection* NewSection = Track->CreateNewSection())
 		{
@@ -189,15 +187,11 @@ void FSequencerUtilities::PopulateMenu_CreateNewSection(FMenuBuilder& MenuBuilde
 
 			Track->OnRowIndicesChanged(NewToOldRowIndices);
 
-			FFrameNumber NewSectionRangeEnd = PlaybackEnd;
-			if (PlaybackEnd <= CurrentTime.Time.FrameNumber)
+			if (Sequencer->GetInfiniteKeyAreas())
 			{
-				const FAnimatedRange ViewRange = Sequencer->GetViewRange();
-				const FFrameRate TickResolution = Sequencer->GetFocusedTickResolution();
-				NewSectionRangeEnd = (ViewRange.GetUpperBoundValue() * TickResolution).FloorToFrame();
+				NewSection->SetRange(TRange<FFrameNumber>::All());
 			}
-			
-			NewSection->SetRange(TRange<FFrameNumber>(CurrentTime.Time.FrameNumber, NewSectionRangeEnd));
+
 			NewSection->SetOverlapPriority(OverlapPriority);
 			NewSection->SetRowIndex(RowIndex);			
 			NewSection->SetBlendType(BlendType);
@@ -222,7 +216,7 @@ void FSequencerUtilities::PopulateMenu_CreateNewSection(FMenuBuilder& MenuBuilde
 	};
 
 	FText NameOverride		= Track->GetSupportedBlendTypes().Num() == 1 ? LOCTEXT("AddSectionText", "Add New Section") : FText();
-	FText TooltipOverride	= Track->GetSupportedBlendTypes().Num() == 1 ? LOCTEXT("AddSectionToolTip", "Adds a new section at the current time") : FText();
+	FText TooltipOverride	= Track->GetSupportedBlendTypes().Num() == 1 ? LOCTEXT("AddSectionToolTip", "Adds a new section") : FText();
 
 	const UEnum* MovieSceneBlendType = FindObjectChecked<UEnum>(nullptr, TEXT("/Script/MovieScene.EMovieSceneBlendType"));
 	for (EMovieSceneBlendType BlendType : Track->GetSupportedBlendTypes())
@@ -231,7 +225,7 @@ void FSequencerUtilities::PopulateMenu_CreateNewSection(FMenuBuilder& MenuBuilde
 		FName EnumValueName = MovieSceneBlendType->GetNameByValue((int64)BlendType);
 		MenuBuilder.AddMenuEntry(
 			NameOverride.IsEmpty() ? DisplayName : NameOverride,
-			TooltipOverride.IsEmpty() ? FText::Format(LOCTEXT("AddSectionFormatToolTip", "Adds a new {0} section at the current time"), DisplayName) : TooltipOverride,
+			TooltipOverride.IsEmpty() ? FText::Format(LOCTEXT("AddSectionFormatToolTip", "Adds a new {0} section"), DisplayName) : TooltipOverride,
 			FSlateIcon(FAppStyle::GetAppStyleSetName(), EnumValueName),
 			FUIAction(
 				FExecuteAction::CreateLambda(CreateNewSection, BlendType),
