Sequencer: Packages to save to be unique in the given case (e.g. Motion Design) where root sequences and sub-sequences are within a same package

diff --git a/Engine/Source/Editor/Sequencer/Private/Sequencer.cpp b/Engine/Source/Editor/Sequencer/Private/Sequencer.cpp
--- a/Engine/Source/Editor/Sequencer/Private/Sequencer.cpp
+++ b/Engine/Source/Editor/Sequencer/Private/Sequencer.cpp
@@ -5930,12 +5930,12 @@ void FSequencer::SaveCurrentMovieScene()
 	TArray<UPackage*> PackagesToSave;
 	TArray<UMovieScene*> MovieScenesToSave;
 	MovieSceneHelpers::GetDescendantMovieScenes(GetRootMovieSceneSequence(), MovieScenesToSave);
-	for (auto MovieSceneToSave : MovieScenesToSave)
+	for (UMovieScene* MovieSceneToSave : MovieScenesToSave)
 	{
 		UPackage* MovieScenePackageToSave = MovieSceneToSave->GetOuter()->GetOutermost();
 		if (MovieScenePackageToSave->IsDirty())
 		{
-			PackagesToSave.Add(MovieScenePackageToSave);
+			PackagesToSave.AddUnique(MovieScenePackageToSave);
 		}
 	}
 
