Sequencer: Fixed subobjects not always having RF_Transient removed when pasting tracks

diff --git a/Engine/Source/Editor/Sequencer/Private/SequencerUtilities.cpp b/Engine/Source/Editor/Sequencer/Private/SequencerUtilities.cpp
--- a/Engine/Source/Editor/Sequencer/Private/SequencerUtilities.cpp
+++ b/Engine/Source/Editor/Sequencer/Private/SequencerUtilities.cpp
@@ -79,11 +79,14 @@
 
 static void ResetCopiedTracksFlags(UMovieSceneTrack* Track)
 {
-	TArray<UObject*> SectionSubObjects;
 	Track->ClearFlags(RF_Transient);
+
+	ForEachObjectWithOuter(Track, [](UObject* InObject) {
+		InObject->ClearFlags(RF_Transient);
+	});
+
 	for (UMovieSceneSection* Section : Track->GetAllSections())
 	{
-		Section->ClearFlags(RF_Transient);
 		Section->PostPaste();
 	}
 }
