Sequencer: Remove restriction on auto expanding if there's only 1 item selected

This allows components in an actor to be expanded since the actor and the component are both selected in Sequencer

diff --git a/Engine/Source/Editor/Sequencer/Private/Sequencer.cpp b/Engine/Source/Editor/Sequencer/Private/Sequencer.cpp
--- a/Engine/Source/Editor/Sequencer/Private/Sequencer.cpp
+++ b/Engine/Source/Editor/Sequencer/Private/Sequencer.cpp
@@ -6670,19 +6670,19 @@ void FSequencer::SynchronizeSequencerSelectionWithExternalSelection()
 		}
 		
 		TSharedPtr<SOutlinerView> TreeView = SequencerWidget->GetTreeView();
-		if (Selection->Outliner.Num() == 1)
+		bool bScrolledIntoView = false;
+		for (TViewModelPtr<IOutlinerExtension> Node : Selection->Outliner)
 		{
-			for (TViewModelPtr<IOutlinerExtension> Node : Selection->Outliner)
+			for (TViewModelPtr<IOutlinerExtension> Parent : Node.AsModel()->GetAncestorsOfType<IOutlinerExtension>())
 			{
-				for (TViewModelPtr<IOutlinerExtension> Parent : Node.AsModel()->GetAncestorsOfType<IOutlinerExtension>())
-				{
-					TreeView->SetItemExpansion(Parent, true);
-						TreeView->SetItemExpansion(Parent, true);
-						Parent = Parent.AsModel()->FindAncestorOfType<IOutlinerExtension>();
-				}
+				TreeView->SetItemExpansion(Parent, true);
+				Parent = Parent.AsModel()->FindAncestorOfType<IOutlinerExtension>();
+			}
 
+			if (!bScrolledIntoView)
+			{
+				bScrolledIntoView = true;
 				TreeView->RequestScrollIntoView(Node);
-				break;
 			}
 		}
 	}
