More non-unity compile error fixes

diff --git a/Engine/Source/Editor/Sequencer/Private/SequencerTimeChangeUndoRedoProxy.h b/Engine/Source/Editor/Sequencer/Private/SequencerTimeChangeUndoRedoProxy.h
--- a/Engine/Source/Editor/Sequencer/Private/SequencerTimeChangeUndoRedoProxy.h
+++ b/Engine/Source/Editor/Sequencer/Private/SequencerTimeChangeUndoRedoProxy.h
@@ -8,6 +8,7 @@
 #include "EventHandlers/ISignedObjectEventHandler.h"
 #include "Delegates/IDelegateInstance.h"
 #include "Misc/QualifiedFrameTime.h"
+#include "MovieSceneSequenceID.h"
 #include "SequencerTimeChangeUndoRedoProxy.generated.h"
 
 class FSequencer;
diff --git a/Engine/Source/Runtime/MovieScene/Private/MovieSceneFrameMigration.cpp b/Engine/Source/Runtime/MovieScene/Private/MovieSceneFrameMigration.cpp
--- a/Engine/Source/Runtime/MovieScene/Private/MovieSceneFrameMigration.cpp
+++ b/Engine/Source/Runtime/MovieScene/Private/MovieSceneFrameMigration.cpp
@@ -3,6 +3,7 @@
 #include "MovieSceneFrameMigration.h"
 #include "UObject/PropertyPortFlags.h"
 #include "MovieSceneFwd.h"
+#include "Misc/FrameRate.h"
 
 #include UE_INLINE_GENERATED_CPP_BY_NAME(MovieSceneFrameMigration)
 
diff --git a/Engine/Source/Runtime/MovieScene/Public/EventHandlers/MovieSceneDataEventContainer.h b/Engine/Source/Runtime/MovieScene/Public/EventHandlers/MovieSceneDataEventContainer.h
--- a/Engine/Source/Runtime/MovieScene/Public/EventHandlers/MovieSceneDataEventContainer.h
+++ b/Engine/Source/Runtime/MovieScene/Public/EventHandlers/MovieSceneDataEventContainer.h
@@ -4,6 +4,7 @@
 
 #include "MovieSceneFwd.h"
 #include "Delegates/Delegate.h"
+#include "Containers/List.h"
 
 namespace UE
 {
