Sequencer: Fix snapping to last frame so that it now uses the generic SnapSequencerTime

diff --git a/Engine/Source/Editor/Sequencer/Private/SequencerTimeSliderController.cpp b/Engine/Source/Editor/Sequencer/Private/SequencerTimeSliderController.cpp
--- a/Engine/Source/Editor/Sequencer/Private/SequencerTimeSliderController.cpp
+++ b/Engine/Source/Editor/Sequencer/Private/SequencerTimeSliderController.cpp
@@ -1258,20 +1258,7 @@ void FSequencerTimeSliderController::CommitScrubPosition( FFrameTime NewValue, b
 	TSharedPtr<FSequencer> Sequencer = WeakSequencer.Pin();
 	if(Sequencer.IsValid())
 	{
-		if (Sequencer->GetSequencerSettings()->ShouldKeepCursorInPlayRangeWhileScrubbing())
-		{
-			TOptional<TRange<FFrameNumber>> RangeValue;
-			RangeValue = TimeSliderArgs.SubSequenceRange.Get(RangeValue);
-
-			if (RangeValue.IsSet())
-			{
-				NewValue = UE::MovieScene::ClampToDiscreteRange(NewValue, RangeValue.GetValue());
-			}
-			else
-			{
-				NewValue = UE::MovieScene::ClampToDiscreteRange(NewValue, TimeSliderArgs.PlaybackRange.Get());
-			}
-		}
+		Sequencer->SnapSequencerTime(NewValue);
 
 		if (bIsScrubbing)
 		{
