Sequencer- Dynamic Binding Refactor

diff --git a/Engine/Plugins/VirtualProduction/Takes/Source/TakeRecorderSources/Private/TakeRecorderActorSource.cpp b/Engine/Plugins/VirtualProduction/Takes/Source/TakeRecorderSources/Private/TakeRecorderActorSource.cpp
--- a/Engine/Plugins/VirtualProduction/Takes/Source/TakeRecorderSources/Private/TakeRecorderActorSource.cpp
+++ b/Engine/Plugins/VirtualProduction/Takes/Source/TakeRecorderSources/Private/TakeRecorderActorSource.cpp
@@ -46,6 +46,7 @@
 #if WITH_EDITOR
 #include "Tracks/MovieSceneSpawnTrack.h"
 #endif
+#include "SequencerUtilities.h"
 
 #include UE_INLINE_GENERATED_CPP_BY_NAME(TakeRecorderActorSource)
 
@@ -256,13 +257,16 @@ TArray<UTakeRecorderSource*> UTakeRecorderActorSource::PreRecording(ULevelSequen
 	{
 		// We need to store the object template in the Movie Scene (because it's a complex UObject)
 		// instead of trying to place this data into the non-UObject safe data stream.
-		FName UniqueTemplateName = MakeUniqueObjectName(TargetLevelSequence, ActorToRecord->GetClass(), NAME_None);
-		Header.TemplateName = UniqueTemplateName.ToString();
-		CachedObjectTemplate = CastChecked<AActor>(TargetLevelSequence->MakeSpawnableTemplateFromInstance(*ActorToRecord, UniqueTemplateName));
-		CachedObjectBindingGuid = MovieScene->AddSpawnable(ActorToRecord->GetActorLabel(), *CachedObjectTemplate);
+		UE::Sequencer::FCreateBindingParams CreateBindingParams;
+		CreateBindingParams.bSpawnable = true;
+		CreateBindingParams.bAllowCustomBinding = true;
+		CreateBindingParams.BindingNameOverride = ActorToRecord->GetActorLabel();
+		CachedObjectBindingGuid = FSequencerUtilities::CreateOrReplaceBinding(nullptr, InSequence, ActorToRecord, CreateBindingParams);
+		CachedObjectTemplate = CastChecked<AActor>(MovieSceneHelpers::GetObjectTemplate(InSequence, CachedObjectBindingGuid, MovieSceneHelpers::CreateTransientSharedPlaybackState(ActorToRecord, InSequence)));
 		
 		if (CachedObjectTemplate.IsValid())
 		{
+			Header.TemplateName = CachedObjectTemplate->GetName();
 			PostProcessCreatedObjectTemplateImpl(CachedObjectTemplate.Get());
 		}
 	}
@@ -416,12 +420,6 @@ void UTakeRecorderActorSource::CreateSectionRecordersRecursive(UObject* ObjectTo
 			ChildPossessable->SetParent(CachedObjectBindingGuid, MovieScene);
 		}
 
-		FMovieSceneSpawnable* ParentSpawnable = MovieScene->FindSpawnable(CachedObjectBindingGuid);
-		if (ParentSpawnable)
-		{
-			ParentSpawnable->AddChildPossessable(Guid);
-		}
-
 		// Bindings are stored relative to their context outer. Newly duplicated components have a different outer
 		TargetLevelSequence->BindPossessableObject(Guid, *ComponentToRecord, bNewComponentAdded ? CachedObjectTemplate.Get() : Target.Get());
 
diff --git a/Engine/Plugins/VirtualProduction/Takes/Source/TakeRecorderSources/TakeRecorderSources.Build.cs b/Engine/Plugins/VirtualProduction/Takes/Source/TakeRecorderSources/TakeRecorderSources.Build.cs
--- a/Engine/Plugins/VirtualProduction/Takes/Source/TakeRecorderSources/TakeRecorderSources.Build.cs
+++ b/Engine/Plugins/VirtualProduction/Takes/Source/TakeRecorderSources/TakeRecorderSources.Build.cs
@@ -30,6 +30,7 @@ public class TakeRecorderSources : ModuleRules
 				"TakeRecorder",
 				"TakeMovieScene",
 				"UnrealEd",
+				"Sequencer"
 			}
 		);
 
diff --git a/Engine/Plugins/VirtualProduction/Takes/Source/TakeTrackRecorders/Private/MovieScene3DTransformTrackRecorder.cpp b/Engine/Plugins/VirtualProduction/Takes/Source/TakeTrackRecorders/Private/MovieScene3DTransformTrackRecorder.cpp
--- a/Engine/Plugins/VirtualProduction/Takes/Source/TakeTrackRecorders/Private/MovieScene3DTransformTrackRecorder.cpp
+++ b/Engine/Plugins/VirtualProduction/Takes/Source/TakeTrackRecorders/Private/MovieScene3DTransformTrackRecorder.cpp
@@ -15,6 +15,7 @@
 #include "Components/SkeletalMeshComponent.h"
 #include "SequenceRecorderUtils.h"
 #include "Animation/AnimData/IAnimationDataModel.h"
+#include "Evaluation/MovieSceneEvaluationState.h"
 
 #include UE_INLINE_GENERATED_CPP_BY_NAME(MovieScene3DTransformTrackRecorder)
 
@@ -298,15 +299,24 @@ void UMovieScene3DTransformTrackRecorder::FinalizeTrackImpl()
  
  	SlowTask.EnterProgressFrame();
  
- 	// If recording a spawnable, update the spawnable object template to the first keyframe
- 	if (MovieScene.IsValid() && ObjectGuid.IsValid())
- 	{
- 		FMovieSceneSpawnable* Spawnable = MovieScene->FindSpawnable(ObjectGuid);
- 		if (Spawnable)
- 		{
- 			Spawnable->SpawnTransform = FirstTransform;
- 		}
- 	}
+	// If recording a spawnable, update the spawnable object template to the first keyframe
+	if (MovieScene.IsValid() && ObjectGuid.IsValid())
+	{
+		UMovieSceneSequence* ThisSequence = MovieScene->GetTypedOuter<UMovieSceneSequence>();
+		UObject* WorldContext = ObjectToRecord.Get() ? ObjectToRecord.Get() : GWorld;
+		if (WorldContext)
+		{ 
+			TSharedRef<UE::MovieScene::FSharedPlaybackState> TransientPlaybackState = MovieSceneHelpers::CreateTransientSharedPlaybackState(WorldContext, ThisSequence);
+
+			if (AActor* ActorTemplatePtr = Cast<AActor>(MovieSceneHelpers::GetObjectTemplate(ThisSequence, ObjectGuid, TransientPlaybackState, 0)))
+			{
+				if (USceneComponent* RootComponent = ActorTemplatePtr->GetRootComponent())
+				{
+					RootComponent->SetRelativeTransform(FirstTransform);
+				}
+			}
+		}
+	}
 }
 
 void UMovieScene3DTransformTrackRecorder::RecordSampleImpl(const FQualifiedFrameTime& CurrentTime)
diff --git a/Engine/Plugins/VirtualProduction/Takes/Source/TakeTrackRecorders/Private/MovieSceneAnimationTrackRecorder.cpp b/Engine/Plugins/VirtualProduction/Takes/Source/TakeTrackRecorders/Private/MovieSceneAnimationTrackRecorder.cpp
--- a/Engine/Plugins/VirtualProduction/Takes/Source/TakeTrackRecorders/Private/MovieSceneAnimationTrackRecorder.cpp
+++ b/Engine/Plugins/VirtualProduction/Takes/Source/TakeTrackRecorders/Private/MovieSceneAnimationTrackRecorder.cpp
@@ -92,16 +92,6 @@ static FGuid GetActorInSequence(AActor* InActor, UMovieScene* MovieScene)
 {
 	FString ActorTargetName = InActor->GetActorLabel();
 
-
-	for (int32 SpawnableCount = 0; SpawnableCount < MovieScene->GetSpawnableCount(); ++SpawnableCount)
-	{
-		const FMovieSceneSpawnable& Spawnable = MovieScene->GetSpawnable(SpawnableCount);
-		if (Spawnable.GetName() == ActorTargetName || Spawnable.Tags.Contains(*ActorTargetName))
-		{
-			return Spawnable.GetGuid();
-		}
-	}
-
 	for (int32 PossessableCount = 0; PossessableCount < MovieScene->GetPossessableCount(); ++PossessableCount)
 	{
 		const FMovieScenePossessable& Possessable = MovieScene->GetPossessable(PossessableCount);
diff --git a/Engine/Plugins/VirtualProduction/Takes/Source/TakeTrackRecorders/Private/MovieSceneSpawnTrackRecorder.cpp b/Engine/Plugins/VirtualProduction/Takes/Source/TakeTrackRecorders/Private/MovieSceneSpawnTrackRecorder.cpp
--- a/Engine/Plugins/VirtualProduction/Takes/Source/TakeTrackRecorders/Private/MovieSceneSpawnTrackRecorder.cpp
+++ b/Engine/Plugins/VirtualProduction/Takes/Source/TakeTrackRecorders/Private/MovieSceneSpawnTrackRecorder.cpp
@@ -7,6 +7,7 @@
 #include "Sections/MovieSceneBoolSection.h"
 #include "Tracks/MovieSceneSpawnTrack.h"
 #include "Channels/MovieSceneChannelProxy.h"
+#include "MovieSceneSequence.h"
 
 #include UE_INLINE_GENERATED_CPP_BY_NAME(MovieSceneSpawnTrackRecorder)
 
@@ -25,7 +26,12 @@ bool FMovieSceneSpawnTrackRecorderFactory::CanRecordObject(UObject* InObjectToRe
 
 void UMovieSceneSpawnTrackRecorder::CreateTrackImpl()
 {
-	if (MovieScene->FindPossessable(ObjectGuid))
+	UMovieSceneSequence* Sequence = MovieScene->GetTypedOuter<UMovieSceneSequence>();
+	if (!Sequence)
+	{
+		return;
+	}
+	if (!MovieSceneHelpers::IsBoundToAnySpawnable(Sequence, ObjectGuid, MovieSceneHelpers::CreateTransientSharedPlaybackState(this, Sequence)))
 	{
 		return;
 	}
diff --git a/Engine/Source/Editor/MovieSceneTools/Private/MovieSceneDynamicBindingBlueprintExtension.cpp b/Engine/Source/Editor/MovieSceneTools/Private/MovieSceneDynamicBindingBlueprintExtension.cpp
--- a/Engine/Source/Editor/MovieSceneTools/Private/MovieSceneDynamicBindingBlueprintExtension.cpp
+++ b/Engine/Source/Editor/MovieSceneTools/Private/MovieSceneDynamicBindingBlueprintExtension.cpp
@@ -113,7 +113,7 @@ void UMovieSceneDynamicBindingBlueprintExtension::HandleGenerateFunctionGraphs(F
 	};
 	FMovieSceneDynamicBindingUtils::IterateDynamicBindings(MovieScene, GenerateFunctionGraphs);
 
-	// This callback sets the generated function calls back onto the spawnables/possessables in the sequence,
+	// This callback sets the generated function calls back onto the custom director blueprint bindings in the sequence,
 	// and keeps a pointer to any needed arguments on this function, so that we can pass special values later
 	// at runtime.
 	TWeakObjectPtr<UMovieSceneSequence> WeakMovieSceneSequence(MovieSceneSequence);
diff --git a/Engine/Source/Editor/MovieSceneTools/Private/MovieSceneDynamicBindingCustomization.cpp b/Engine/Source/Editor/MovieSceneTools/Private/MovieSceneDynamicBindingCustomization.cpp
--- a/Engine/Source/Editor/MovieSceneTools/Private/MovieSceneDynamicBindingCustomization.cpp
+++ b/Engine/Source/Editor/MovieSceneTools/Private/MovieSceneDynamicBindingCustomization.cpp
@@ -41,6 +41,23 @@ TSharedRef<IPropertyTypeCustomization> FMovieSceneDynamicBindingCustomization::M
 	return Instance;
 }
 
+void FMovieSceneDynamicBindingCustomization::CustomizeChildren(TSharedRef<IPropertyHandle> InPropertyHandle, IDetailChildrenBuilder& ChildBuilder, IPropertyTypeCustomizationUtils& CustomizationUtils)
+{
+	FMovieSceneDirectorBlueprintEndpointCustomization::CustomizeChildren(InPropertyHandle, ChildBuilder, CustomizationUtils);
+	FStringView PropertyPath = InPropertyHandle->GetPropertyPath();
+	BindingIndex = 0;
+	if (PropertyPath.StartsWith(TEXT("Bindings[")))
+	{
+		int32 RightBounds;
+		if (PropertyPath.FindChar(']', RightBounds))
+		{
+			FStringView BindingIndexStr = PropertyPath.Mid(9, RightBounds);
+			// Grab the BindingIndex out of the PropertyPath
+			BindingIndex = FCString::Atoi(BindingIndexStr.GetData());
+		}
+	}
+}
+
 void FMovieSceneDynamicBindingCustomization::GetPayloadVariables(UObject* EditObject, void* RawData, FPayloadVariableMap& OutPayloadVariables) const
 {
 	const FMovieSceneDynamicBinding* DynamicBinding = static_cast<FMovieSceneDynamicBinding*>(RawData);
@@ -148,10 +165,10 @@ FMovieSceneDirectorBlueprintEndpointDefinition FMovieSceneDynamicBindingCustomiz
 	if (FMovieScenePossessable* Possessable = MovieScene->FindPossessable(ObjectBinding))
 	{
 		Definition.EndpointName = Possessable->GetName() + "_DynamicBinding";
-	}
-	else if (FMovieSceneSpawnable* Spawnable = MovieScene->FindSpawnable(ObjectBinding))
-	{
-		Definition.EndpointName = Spawnable->GetName() + "_DynamicBinding";
+		if (BindingIndex > 0)
+		{
+			Definition.EndpointName += "_" + FString::FromInt(BindingIndex);
+		}
 	}
 
 	return Definition;
@@ -171,6 +188,12 @@ void FMovieSceneDynamicBindingCustomization::OnCreateEndpoint(UMovieSceneSequenc
 
 		FMovieSceneDynamicBinding* DynamicBinding = static_cast<FMovieSceneDynamicBinding*>(RawData[Index]);
 
+		// Default call in editor to true
+		if (UK2Node_FunctionEntry* CallFunction = Cast<UK2Node_FunctionEntry>(NewEndpoint))
+		{
+			CallFunction->MetaData.bCallInEditor = true;
+		}
+
 		FMovieSceneDynamicBindingUtils::SetEndpoint(MovieScene, DynamicBinding, NewEndpoint);
 	}
 
diff --git a/Engine/Source/Editor/MovieSceneTools/Private/MovieSceneToolsModule.cpp b/Engine/Source/Editor/MovieSceneTools/Private/MovieSceneToolsModule.cpp
--- a/Engine/Source/Editor/MovieSceneTools/Private/MovieSceneToolsModule.cpp
+++ b/Engine/Source/Editor/MovieSceneTools/Private/MovieSceneToolsModule.cpp
@@ -117,6 +117,11 @@
 #include "TransformableHandle.h"
 #include "Constraints/ComponentConstraintChannelInterface.h"
 #include "Constraints/TransformConstraintChannelInterface.h"
+#include "Bindings/MovieSceneReplaceableDirectorBlueprintBinding.h"
+#include "Bindings/MovieSceneSpawnableDirectorBlueprintBinding.h"
+#include "EntitySystem/MovieSceneSharedPlaybackState.h"
+#include "MovieSceneDynamicBindingUtils.h"
+#include "Kismet2/KismetEditorUtilities.h"
 
 #define LOCTEXT_NAMESPACE "FMovieSceneToolsModule"
 
@@ -250,6 +255,8 @@ void FMovieSceneToolsModule::StartupModule()
 	UMovieSceneEventSectionBase::RemoveForCookEvent.BindStatic(RemoveForCookEventSection);
 	UMovieScene::IsTrackClassAllowedEvent.BindStatic(IsTrackClassAllowed);
 	ULevelSequence::PostDuplicateEvent.BindStatic(PostDuplicateEvent);
+	FixupDynamicBindingsHandle = ULevelSequence::FixupDynamicBindingsEvent.AddStatic(FixupDynamicBindingsEvent);
+
 
 	auto OnObjectsReplaced = [](const TMap<UObject*, UObject*>& ReplacedObjects)
 	{
@@ -287,6 +294,7 @@ void FMovieSceneToolsModule::ShutdownModule()
 	UMovieSceneEventSectionBase::RemoveForCookEvent = UMovieSceneEventSectionBase::FRemoveForCookEvent();
 	UMovieScene::IsTrackClassAllowedEvent = UMovieScene::FIsTrackClassAllowedEvent();
 	ULevelSequence::PostDuplicateEvent = ULevelSequence::FPostDuplicateEvent();
+	ULevelSequence::FixupDynamicBindingsEvent.Remove(FixupDynamicBindingsHandle);
 
 	if (ICurveEditorModule* CurveEditorModule = FModuleManager::GetModulePtr<ICurveEditorModule>("CurveEditor"))
 	{
@@ -440,6 +448,19 @@ void FMovieSceneToolsModule::PostDuplicateEvent(ULevelSequence* LevelSequence)
 	}
 }
 
+
+void FMovieSceneToolsModule::FixupDynamicBindingsEvent(ULevelSequence* LevelSequence)
+{
+	if (LevelSequence)
+	{
+		if (UBlueprint* DirectorBlueprint = LevelSequence->GetDirectorBlueprint())
+		{
+			FMovieSceneDynamicBindingUtils::EnsureBlueprintExtensionCreated(LevelSequence, DirectorBlueprint);
+			FKismetEditorUtilities::CompileBlueprint(DirectorBlueprint);
+		}
+	}
+}
+
 bool FMovieSceneToolsModule::UpgradeLegacyEventEndpointForSection(UMovieSceneEventSectionBase* Section)
 {
 	UMovieSceneSequence*       Sequence           = Section->GetTypedOuter<UMovieSceneSequence>();
@@ -583,6 +604,8 @@ void FMovieSceneToolsModule::FixupPayloadParameterNameForSection(UMovieSceneEven
 
 void FMovieSceneToolsModule::FixupPayloadParameterNameForDynamicBinding(UMovieScene* MovieScene, UK2Node* InNode, FName OldPinName, FName NewPinName)
 {
+	using namespace UE::MovieScene;
+
 	check(MovieScene);
 
 	auto FixupPayloadParameterName = [InNode, OldPinName, NewPinName](FMovieSceneDynamicBinding& DynamicBinding)
@@ -597,14 +620,26 @@ void FMovieSceneToolsModule::FixupPayloadParameterNameForDynamicBinding(UMovieSc
 		}
 	};
 
-	for (int32 Index = 0, PossessableCount = MovieScene->GetPossessableCount(); Index < PossessableCount; ++Index)
-	{
-		FixupPayloadParameterName(MovieScene->GetPossessable(Index).DynamicBinding);
-	}
+	UMovieSceneSequence* ThisSequence = MovieScene->GetTypedOuter<UMovieSceneSequence>();
+	TSharedRef<UE::MovieScene::FSharedPlaybackState> TransientPlaybackState = MovieSceneHelpers::CreateTransientSharedPlaybackState(GEditor->GetEditorWorldContext().World(), ThisSequence);
 
-	for (int32 Index = 0, SpawnableCount = MovieScene->GetSpawnableCount(); Index < SpawnableCount; ++Index)
+	if (FMovieSceneBindingReferences* BindingReferences = MovieScene->GetTypedOuter<UMovieSceneSequence>()->GetBindingReferences())
 	{
-		FixupPayloadParameterName(MovieScene->GetSpawnable(Index).DynamicBinding);
+		for (FMovieSceneBindingReference& BindingReference : BindingReferences->GetAllReferences())
+		{
+			if (BindingReference.CustomBinding)
+			{
+				if (UMovieSceneReplaceableDirectorBlueprintBinding* ReplaceableDirectorBlueprintBinding = Cast<UMovieSceneReplaceableDirectorBlueprintBinding>(BindingReference.CustomBinding))
+				{
+					FixupPayloadParameterName(ReplaceableDirectorBlueprintBinding->DynamicBinding);
+				}
+
+				if (UMovieSceneSpawnableDirectorBlueprintBinding* SpawnableDirectorBlueprintBinding = Cast<UMovieSceneSpawnableDirectorBlueprintBinding>(BindingReference.CustomBinding->AsSpawnable(TransientPlaybackState)))
+				{
+					FixupPayloadParameterName(SpawnableDirectorBlueprintBinding->DynamicBinding);
+				}
+			}
+		}
 	}
 }
 
diff --git a/Engine/Source/Editor/MovieSceneTools/Public/MovieSceneDynamicBindingCustomization.h b/Engine/Source/Editor/MovieSceneTools/Public/MovieSceneDynamicBindingCustomization.h
--- a/Engine/Source/Editor/MovieSceneTools/Public/MovieSceneDynamicBindingCustomization.h
+++ b/Engine/Source/Editor/MovieSceneTools/Public/MovieSceneDynamicBindingCustomization.h
@@ -28,6 +28,8 @@ public:
 	static TSharedRef<IPropertyTypeCustomization> MakeInstance();
 	static TSharedRef<IPropertyTypeCustomization> MakeInstance(UMovieScene* InMovieScene, FGuid InObjectBinding);
 
+	virtual void CustomizeChildren(TSharedRef<IPropertyHandle> InPropertyHandle, IDetailChildrenBuilder& ChildBuilder, IPropertyTypeCustomizationUtils& CustomizationUtils) override;
+
 protected:
 
 	virtual void GetPayloadVariables(UObject* EditObject, void* RawData, FPayloadVariableMap& OutPayloadVariables) const override;
@@ -53,5 +55,6 @@ private:
 
 	UMovieScene* EditedMovieScene;
 	FGuid ObjectBinding;
+	int32 BindingIndex = 0;
 };
 
diff --git a/Engine/Source/Editor/MovieSceneTools/Public/MovieSceneDynamicBindingUtils.h b/Engine/Source/Editor/MovieSceneTools/Public/MovieSceneDynamicBindingUtils.h
--- a/Engine/Source/Editor/MovieSceneTools/Public/MovieSceneDynamicBindingUtils.h
+++ b/Engine/Source/Editor/MovieSceneTools/Public/MovieSceneDynamicBindingUtils.h
@@ -5,8 +5,13 @@
 #include "CoreTypes.h"
 #include "MovieScene.h"
 #include "MovieSceneDynamicBinding.h"
-#include "MovieScenePossessable.h"
-#include "MovieSceneSpawnable.h"
+#include "Bindings/MovieSceneReplaceableDirectorBlueprintBinding.h"
+#include "Bindings/MovieSceneSpawnableDirectorBlueprintBinding.h"
+#include "EntitySystem/MovieSceneSharedPlaybackState.h"
+#include "Evaluation/MovieSceneEvaluationState.h"
+#include "Editor.h"
+#include "MovieSceneBindingReferences.h"
+#include "MovieSceneCommonHelpers.h"
 
 #include "MovieSceneDynamicBindingUtils.generated.h"
 
@@ -15,6 +20,8 @@
  */
 struct MOVIESCENETOOLS_API FMovieSceneDynamicBindingUtils
 {
+	using FSharedPlaybackState = UE::MovieScene::FSharedPlaybackState;
+
 	/**
 	 * Set an endpoint on the given dynamic binding.
 	 */
@@ -31,19 +38,9 @@ struct MOVIESCENETOOLS_API FMovieSceneDynamicBindingUtils
 	template<typename Callback>
 	static void IterateDynamicBindings(UMovieScene* InMovieScene, Callback&& InCallback)
 	{
-		for (int32 Index = 0, PossessableCount = InMovieScene->GetPossessableCount(); Index < PossessableCount; ++Index)
-		{
-			FMovieScenePossessable& Possessable = InMovieScene->GetPossessable(Index);
-			FMovieSceneDynamicBinding& DynamicBinding = Possessable.DynamicBinding;
-			InCallback(Possessable.GetGuid(), DynamicBinding);
-		}
-
-		for (int32 Index = 0, SpawnableCount = InMovieScene->GetSpawnableCount(); Index < SpawnableCount; ++Index)
-		{
-			FMovieSceneSpawnable& Spawnable = InMovieScene->GetSpawnable(Index);
-			FMovieSceneDynamicBinding& DynamicBinding = Spawnable.DynamicBinding;
-			InCallback(Spawnable.GetGuid(), DynamicBinding);
-		}
+		UMovieSceneSequence* ThisSequence = InMovieScene->GetTypedOuter<UMovieSceneSequence>();
+		TSharedRef<UE::MovieScene::FSharedPlaybackState> TransientPlaybackState = MovieSceneHelpers::CreateTransientSharedPlaybackState(GEditor->GetEditorWorldContext().World(), ThisSequence);
+		ThisSequence->IterateDynamicBindings(TransientPlaybackState, InCallback);
 	}
 
 	/**
diff --git a/Engine/Source/Editor/MovieSceneTools/Public/MovieSceneToolsModule.h b/Engine/Source/Editor/MovieSceneTools/Public/MovieSceneToolsModule.h
--- a/Engine/Source/Editor/MovieSceneTools/Public/MovieSceneToolsModule.h
+++ b/Engine/Source/Editor/MovieSceneTools/Public/MovieSceneToolsModule.h
@@ -96,6 +96,7 @@ private:
 	static void RemoveForCookEventSection(UMovieSceneEventSectionBase* Section);
 	static bool IsTrackClassAllowed(UClass* InClass);
 	static void PostDuplicateEvent(ULevelSequence* LevelSequence);
+	static void FixupDynamicBindingsEvent(ULevelSequence* LevelSequence);
 
 private:
 
@@ -148,6 +149,7 @@ private:
 	FDelegateHandle FixupDynamicBindingPayloadParameterNameHandle;
 	FDelegateHandle FixupEventSectionPayloadParameterNameHandle;
 	FDelegateHandle UpgradeLegacyEventEndpointHandle;
+	FDelegateHandle FixupDynamicBindingsHandle;
 
 	FDelegateHandle OnObjectsReplacedHandle;
 
diff --git a/Engine/Source/Editor/SequenceRecorder/Private/ActorRecording.cpp b/Engine/Source/Editor/SequenceRecorder/Private/ActorRecording.cpp
--- a/Engine/Source/Editor/SequenceRecorder/Private/ActorRecording.cpp
+++ b/Engine/Source/Editor/SequenceRecorder/Private/ActorRecording.cpp
@@ -24,6 +24,10 @@
 #include "IAssetTools.h"
 #include "AssetToolsModule.h"
 #include "Subsystems/AssetEditorSubsystem.h"
+#include "SequencerUtilities.h"
+#include "EntitySystem/MovieSceneSharedPlaybackState.h"
+#include "Evaluation/MovieSceneEvaluationState.h"
+#include "MovieSceneCommonHelpers.h"
 
 static const FName SequencerActorTag(TEXT("SequencerActor"));
 static const FName MovieSceneSectionRecorderFactoryName("MovieSceneSectionRecorderFactory");
@@ -174,26 +178,6 @@ bool UActorRecording::StartRecording(ULevelSequence* CurrentSequence, float Curr
 	return true;
 }
 
-static FString GetUniqueSpawnableName(UMovieScene* MovieScene, const FString& BaseName)
-{
-	FString BlueprintName = BaseName;
-	auto DuplName = [&](FMovieSceneSpawnable& InSpawnable)
-	{
-		return InSpawnable.GetName() == BlueprintName;
-	};
-
-	int32 Index = 2;
-	FString UniqueString;
-	while (MovieScene->FindSpawnable(DuplName))
-	{
-		BlueprintName.RemoveFromEnd(UniqueString);
-		UniqueString = FString::Printf(TEXT(" (%d)"), Index++);
-		BlueprintName += UniqueString;
-	}
-
-	return BlueprintName;
-}
-
 void UActorRecording::GetNonSceneActorComponents(TArray<UActorComponent*>& OutArray)
 {
 	if (GetActorToRecord() != nullptr)
@@ -435,15 +419,6 @@ FGuid UActorRecording::GetActorInSequence(AActor* InActor, ULevelSequence* Curre
 
 	UMovieScene* MovieScene = CurrentSequence->GetMovieScene();
 	
-	for (int32 SpawnableCount = 0; SpawnableCount < MovieScene->GetSpawnableCount(); ++SpawnableCount)
-	{
-		const FMovieSceneSpawnable& Spawnable = MovieScene->GetSpawnable(SpawnableCount);
-		if (Spawnable.GetName() == ActorTargetName || Spawnable.Tags.Contains(*InActor->GetActorLabel()))
-		{
-			return Spawnable.GetGuid();
-		}
-	}
-
 	for (int32 PossessableCount = 0; PossessableCount < MovieScene->GetPossessableCount(); ++PossessableCount)
 	{
 		const FMovieScenePossessable& Possessable = MovieScene->GetPossessable(PossessableCount);
@@ -482,7 +457,7 @@ void UActorRecording::StartRecordingActorProperties(ULevelSequence* CurrentSeque
 			}
 			else
 			{
-				FString TemplateName = GetUniqueSpawnableName(MovieScene, Actor->GetName());
+				FString TemplateName = MovieSceneHelpers::MakeUniqueBindingName(MovieScene, Actor->GetName());
 
 				AActor* ObjectTemplate = CastChecked<AActor>(CurrentSequence->MakeSpawnableTemplateFromInstance(*Actor, *TemplateName));
 
@@ -506,7 +481,11 @@ void UActorRecording::StartRecordingActorProperties(ULevelSequence* CurrentSeque
 						Pawn->AutoPossessPlayer = EAutoReceiveInput::Disabled;
 					}
 
-					Guid = MovieScene->AddSpawnable(ObjectBindingName, *ObjectTemplate);
+					UE::Sequencer::FCreateBindingParams CreateBindingParams;
+					CreateBindingParams.bAllowCustomBinding = true;
+					CreateBindingParams.bSpawnable = true;
+					CreateBindingParams.BindingNameOverride = ObjectBindingName;
+					Guid = FSequencerUtilities::CreateOrReplaceBinding(nullptr, CurrentSequence, ObjectTemplate, CreateBindingParams);
 				}
 			}
 		}
@@ -524,15 +503,6 @@ void UActorRecording::StartRecordingActorProperties(ULevelSequence* CurrentSeque
 				Possessable->SetName(ObjectBindingName);
 			}
 
-			if (FMovieSceneSpawnable* Spawnable = MovieScene->FindSpawnable(Guid))
-			{
-				if (!Spawnable->Tags.Contains(*Actor->GetActorLabel()))
-				{
-					Spawnable->Tags.AddUnique(FName(*Actor->GetActorLabel()));
-				}
-				Spawnable->SetName(ObjectBindingName);
-			}
-
 			// add our folder
 			FindOrAddFolder(MovieScene);
 
@@ -646,12 +616,6 @@ TSharedPtr<FMovieSceneAnimationSectionRecorder> UActorRecording::StartRecordingC
 		ChildPossessable->SetParent(Guid, OwnerMovieScene);
 	}
 
-	FMovieSceneSpawnable* ParentSpawnable = OwnerMovieScene->FindSpawnable(Guid);
-	if (ParentSpawnable)
-	{
-		ParentSpawnable->AddChildPossessable(PossessableGuid);
-	}
-
 	CurrentSequence->BindPossessableObject(PossessableGuid, *ActorComponent, BindingContext);
 
 	const USequenceRecorderSettings* Settings = GetDefault<USequenceRecorderSettings>();
@@ -749,10 +713,10 @@ bool UActorRecording::StopRecording(ULevelSequence* OriginalSequence, float Curr
 		UMovieScene* MovieScene = CurrentSequence->GetMovieScene();
 		check(MovieScene);
 
-		FMovieSceneSpawnable* Spawnable = MovieScene->FindSpawnable(Guid);
-		if(Spawnable)
+		FMovieScenePossessable* Possessable = MovieScene->FindPossessable(Guid);
+		if(Possessable)
 		{
-			ActorName = Spawnable->GetName();
+			ActorName = Possessable->GetName();
 		}
 	}
 
@@ -1045,10 +1009,9 @@ void UActorRecording::StartRecordingNewComponents(ULevelSequence* CurrentSequenc
 		const USequenceRecorderSettings* Settings = GetDefault<USequenceRecorderSettings>();
 		if (!bRecordToPossessable)
 		{
-			FMovieSceneSpawnable* Spawnable = MovieScene->FindSpawnable(Guid);
-			check(Spawnable);
+			TSharedRef<UE::MovieScene::FSharedPlaybackState> TransientPlaybackState = MovieSceneHelpers::CreateTransientSharedPlaybackState(GetActorToRecord()->GetWorld(), CurrentSequence);
 
-			AActor* ObjectTemplate = CastChecked<AActor>(Spawnable->GetObjectTemplate());
+			AActor* ObjectTemplate = Cast<AActor>(MovieSceneHelpers::GetObjectTemplate(CurrentSequence, Guid, TransientPlaybackState, 0));
 
 			for (UActorComponent* ActorComponent : NewComponents)
 			{
diff --git a/Engine/Source/Editor/SequenceRecorder/Private/Sections/MovieScene3DTransformSectionRecorder.cpp b/Engine/Source/Editor/SequenceRecorder/Private/Sections/MovieScene3DTransformSectionRecorder.cpp
--- a/Engine/Source/Editor/SequenceRecorder/Private/Sections/MovieScene3DTransformSectionRecorder.cpp
+++ b/Engine/Source/Editor/SequenceRecorder/Private/Sections/MovieScene3DTransformSectionRecorder.cpp
@@ -14,6 +14,7 @@
 #include "Channels/MovieSceneChannelProxy.h"
 #include "Animation/AnimData/IAnimationDataModel.h"
 #include "Animation/AnimSequence.h"
+#include "Evaluation/MovieSceneEvaluationState.h"
 
 TSharedPtr<IMovieSceneSectionRecorder> FMovieScene3DTransformSectionRecorderFactory::CreateSectionRecorder(const FActorRecordingSettings& InActorRecordingSettings) const
 {
@@ -318,10 +319,15 @@ void FMovieScene3DTransformSectionRecorder::FinalizeSection(float CurrentTime)
 	// If recording a spawnable, update the spawnable object template to the first keyframe
 	if (MovieScene.IsValid() && Guid.IsValid())
 	{
-		FMovieSceneSpawnable* Spawnable = MovieScene->FindSpawnable(Guid);
-		if (Spawnable)
+		UMovieSceneSequence* ThisSequence = MovieScene->GetTypedOuter<UMovieSceneSequence>();
+		TSharedRef<UE::MovieScene::FSharedPlaybackState> TransientPlaybackState = MovieSceneHelpers::CreateTransientSharedPlaybackState(GetSourceObject(), ThisSequence);
+
+		if (AActor* ActorTemplatePtr = Cast<AActor>(MovieSceneHelpers::GetObjectTemplate(ThisSequence, Guid, TransientPlaybackState, 0)))
 		{
-			Spawnable->SpawnTransform = FirstTransform;
+			if (USceneComponent* RootComponent = ActorTemplatePtr->GetRootComponent())
+			{
+				RootComponent->SetRelativeTransform(FirstTransform);
+			}
 		}
 	}
 
diff --git a/Engine/Source/Editor/SequenceRecorder/SequenceRecorder.Build.cs b/Engine/Source/Editor/SequenceRecorder/SequenceRecorder.Build.cs
--- a/Engine/Source/Editor/SequenceRecorder/SequenceRecorder.Build.cs
+++ b/Engine/Source/Editor/SequenceRecorder/SequenceRecorder.Build.cs
@@ -41,6 +41,7 @@ namespace UnrealBuildTool.Rules
                     "Kismet",
                     "LiveLinkInterface",
 					"SceneOutliner",
+					"Sequencer"
                 }
                 );
 
diff --git a/Engine/Source/Editor/Sequencer/Private/MVVM/ViewModels/ObjectBindingModel.cpp b/Engine/Source/Editor/Sequencer/Private/MVVM/ViewModels/ObjectBindingModel.cpp
--- a/Engine/Source/Editor/Sequencer/Private/MVVM/ViewModels/ObjectBindingModel.cpp
+++ b/Engine/Source/Editor/Sequencer/Private/MVVM/ViewModels/ObjectBindingModel.cpp
@@ -1026,85 +1026,6 @@ void FObjectBindingModel::BuildOrganizeContextMenu(FMenuBuilder& MenuBuilder)
 	FOutlinerItemModel::BuildOrganizeContextMenu(MenuBuilder);
 }
 
-void FObjectBindingModel::AddDynamicBindingMenu(FMenuBuilder& MenuBuilder, FMovieSceneDynamicBinding& DynamicBinding)
-{
-	FDetailsViewArgs DetailsViewArgs;
-	{
-		DetailsViewArgs.bAllowSearch = false;
-		DetailsViewArgs.bCustomFilterAreaLocation = true;
-		DetailsViewArgs.bCustomNameAreaLocation = true;
-		DetailsViewArgs.bHideSelectionTip = true;
-		DetailsViewArgs.bLockable = false;
-		DetailsViewArgs.bSearchInitialKeyFocus = true;
-		DetailsViewArgs.bUpdatesFromSelection = false;
-		DetailsViewArgs.bShowOptions = false;
-		DetailsViewArgs.bShowModifiedPropertiesOption = false;
-		DetailsViewArgs.bShowScrollBar = false;
-	}
-
-	FStructureDetailsViewArgs StructureViewArgs;
-	{
-		StructureViewArgs.bShowObjects = false;
-		StructureViewArgs.bShowAssets = true;
-		StructureViewArgs.bShowClasses = true;
-		StructureViewArgs.bShowInterfaces = false;
-	}
-
-	TSharedRef<IStructureDetailsView> StructureDetailsView = FModuleManager::GetModuleChecked<FPropertyEditorModule>("PropertyEditor")
-		.CreateStructureDetailView(DetailsViewArgs, StructureViewArgs, nullptr);
-
-	// Register details customizations for this instance
-	TSharedPtr<FSequencer> Sequencer = OwnerModel->GetSequencerImpl();
-	UMovieSceneSequence* Sequence = Sequencer->GetFocusedMovieSceneSequence();
-	StructureDetailsView->GetDetailsView()->RegisterInstancedCustomPropertyTypeLayout(
-		FMovieSceneDynamicBinding::StaticStruct()->GetFName(),
-		FOnGetPropertyTypeCustomizationInstance::CreateStatic(&FMovieSceneDynamicBindingCustomization::MakeInstance, Sequence->GetMovieScene(), ObjectBindingID));
-
-	// We can't just show the FMovieSceneDynamicBinding struct in the details view, because Slate only uses
-	// the above details view customization for *properties* (not for the root object). So here we put a copy of
-	// our dynamic binding struct inside a container, and when the details view is done setting values on it,
-	// we copy these values back to the original dynamic binding.
-	TSharedPtr<FStructOnScope> StructOnScope = MakeShared<FStructOnScope>(FMovieSceneDynamicBindingContainer::StaticStruct());
-	FMovieSceneDynamicBindingContainer* BufferContainer = (FMovieSceneDynamicBindingContainer*)StructOnScope->GetStructMemory();
-	BufferContainer->DynamicBinding = DynamicBinding;
-	StructureDetailsView->SetStructureData(StructOnScope);
-
-	StructureDetailsView->GetOnFinishedChangingPropertiesDelegate().AddSP(this, &FObjectBindingModel::OnFinishedChangingDynamicBindingProperties, StructOnScope);
-
-	MenuBuilder.BeginSection(NAME_None, LOCTEXT("DynamicBindingHeader", "Dynamic Binding"));
-	{
-		TSharedRef<SWidget> Widget = StructureDetailsView->GetWidget().ToSharedRef();
-		MenuBuilder.AddWidget(Widget, FText());
-	}
-	MenuBuilder.EndSection();
-}
-
-void FObjectBindingModel::OnFinishedChangingDynamicBindingProperties(const FPropertyChangedEvent& ChangeEvent, TSharedPtr<FStructOnScope> ValueStruct)
-{
-	auto* Container = (FMovieSceneDynamicBindingContainer*)ValueStruct->GetStructMemory();
-
-	UMovieScene* MovieScene = OwnerModel->GetMovieScene();
-	FMovieScenePossessable* Possessable = MovieScene->FindPossessable(ObjectBindingID);
-	if (Possessable)
-	{
-		Possessable->DynamicBinding = Container->DynamicBinding;
-	}
-	FMovieSceneSpawnable* Spawnable = MovieScene->FindSpawnable(ObjectBindingID);
-	if (Spawnable)
-	{
-		Spawnable->DynamicBinding = Container->DynamicBinding;
-	}
-	// Force refresh the binding
-	TSharedPtr<FSequencer> Sequencer = OwnerModel->GetSequencerImpl();
-	if (Sequencer.IsValid())
-	{
-		if (FMovieSceneEvaluationState* EvaluationState = Sequencer->GetSharedPlaybackState()->FindCapability<FMovieSceneEvaluationState>())
-		{
-			EvaluationState->Invalidate(GetObjectGuid(), Sequencer->GetFocusedTemplateID());
-		}
-	}
-}
-
 void FObjectBindingModel::AddTagMenu(FMenuBuilder& MenuBuilder)
 {
 	MenuBuilder.AddMenuEntry(FSequencerCommands::Get().OpenTaggedBindingManager);
diff --git a/Engine/Source/Editor/Sequencer/Private/MVVM/ViewModels/PossessableModel.cpp b/Engine/Source/Editor/Sequencer/Private/MVVM/ViewModels/PossessableModel.cpp
--- a/Engine/Source/Editor/Sequencer/Private/MVVM/ViewModels/PossessableModel.cpp
+++ b/Engine/Source/Editor/Sequencer/Private/MVVM/ViewModels/PossessableModel.cpp
@@ -132,16 +132,6 @@ const FSlateBrush* FPossessableModel::GetIconOverlayBrush() const
 		}
 	}
 
-	UMovieScene*            MovieScene  = OwnerModel ? OwnerModel->GetMovieScene() : nullptr;
-	FMovieScenePossessable* Possessable = MovieScene ? MovieScene->FindPossessable(ObjectBindingID) : nullptr;
-	if (Possessable && Possessable->DynamicBinding.WeakEndpoint.IsValid())
-	{
-		if (Possessable->DynamicBinding.WeakEndpoint.IsValid())
-		{
-			return FAppStyle::GetBrush("Sequencer.DynamicBindingIconOverlay");
-		}
-	}
-
 	UMovieSceneSequence* Sequence = OwnerModel ? OwnerModel->GetSequence() : nullptr;
 	if (Sequence)
 	{
@@ -155,6 +145,10 @@ const FSlateBrush* FPossessableModel::GetIconOverlayBrush() const
 				}
 			}
 		}
+		else if (const FSlateBrush* CustomBrush = Sequence->GetCustomBrushForBinding(ObjectBindingID))
+		{
+			return CustomBrush;
+		}
 	}
 
 	return nullptr;
diff --git a/Engine/Source/Editor/Sequencer/Private/MVVM/ViewModels/SpawnableModel.cpp b/Engine/Source/Editor/Sequencer/Private/MVVM/ViewModels/SpawnableModel.cpp
--- a/Engine/Source/Editor/Sequencer/Private/MVVM/ViewModels/SpawnableModel.cpp
+++ b/Engine/Source/Editor/Sequencer/Private/MVVM/ViewModels/SpawnableModel.cpp
@@ -86,13 +86,6 @@ FText FSpawnableModel::GetIconToolTipText() const
 
 const FSlateBrush* FSpawnableModel::GetIconOverlayBrush() const
 {
-	UMovieScene*          MovieScene  = OwnerModel ? OwnerModel->GetMovieScene() : nullptr;
-	FMovieSceneSpawnable* Spawnable  = MovieScene ? MovieScene->FindSpawnable(ObjectBindingID) : nullptr;
-	if (Spawnable && Spawnable->DynamicBinding.WeakEndpoint.IsValid())
-	{
-		return FAppStyle::GetBrush("Sequencer.SpawnableDynamicBindingIconOverlay");
-	}
-
 	return FAppStyle::GetBrush("Sequencer.SpawnableIconOverlay");
 }
 
diff --git a/Engine/Source/Editor/Sequencer/Private/MovieSceneCopyableBinding.h b/Engine/Source/Editor/Sequencer/Private/MovieSceneCopyableBinding.h
--- a/Engine/Source/Editor/Sequencer/Private/MovieSceneCopyableBinding.h
+++ b/Engine/Source/Editor/Sequencer/Private/MovieSceneCopyableBinding.h
@@ -61,4 +61,7 @@ public:
 	
 	UPROPERTY()
 	int32 NumCustomBindings;
+
+	UPROPERTY()
+	TArray<int32> PreviewSpawnableBindings;
 };
diff --git a/Engine/Source/Editor/Sequencer/Private/Sequencer.cpp b/Engine/Source/Editor/Sequencer/Private/Sequencer.cpp
--- a/Engine/Source/Editor/Sequencer/Private/Sequencer.cpp
+++ b/Engine/Source/Editor/Sequencer/Private/Sequencer.cpp
@@ -1616,18 +1616,7 @@ void FSequencer::UpdateSequencerCustomizations(const UMovieSceneSequence* Previo
 	// TODO: Allow customization of which custom bindings are allowed? For now just iterate over all subclasses and sort by priority
 
 	// Cache custom spawnable types
-	SupportedCustomBindingTypes.Empty();
-	for (TObjectIterator<UClass> It; It; ++It)
-	{
-		if (It->IsChildOf(UMovieSceneCustomBinding::StaticClass()) && !It->HasAnyClassFlags(CLASS_Abstract))
-		{
-			SupportedCustomBindingTypes.Add(*It);
-		}
-	}
-
-	// Sort by spawner priority to allow disambiguation for similar object types
-	SupportedCustomBindingTypes.Sort([](const TSubclassOf<UMovieSceneCustomBinding>& A, const TSubclassOf<UMovieSceneCustomBinding>& B) {
-		return A && B && A->GetDefaultObject<UMovieSceneCustomBinding>()->GetCustomBindingPriority() > B->GetDefaultObject<UMovieSceneCustomBinding>()->GetCustomBindingPriority(); });
+	MovieSceneHelpers::GetPrioritySortedCustomBindingTypes(SupportedCustomBindingTypes);
 }
 
 void FSequencer::RerunConstructionScripts()
diff --git a/Engine/Source/Editor/Sequencer/Private/Sequencer.h b/Engine/Source/Editor/Sequencer/Private/Sequencer.h
--- a/Engine/Source/Editor/Sequencer/Private/Sequencer.h
+++ b/Engine/Source/Editor/Sequencer/Private/Sequencer.h
@@ -1488,5 +1488,5 @@ private:
 	uint32 LastKnownStateSerial = 0;
 
 	/* Cached list of supported custom binding types*/
-	TArray<TSubclassOf<UMovieSceneCustomBinding>> SupportedCustomBindingTypes;
+	TArray<const TSubclassOf<UMovieSceneCustomBinding>> SupportedCustomBindingTypes;
 };
diff --git a/Engine/Source/Editor/Sequencer/Private/SequencerUtilities.cpp b/Engine/Source/Editor/Sequencer/Private/SequencerUtilities.cpp
--- a/Engine/Source/Editor/Sequencer/Private/SequencerUtilities.cpp
+++ b/Engine/Source/Editor/Sequencer/Private/SequencerUtilities.cpp
@@ -74,6 +74,8 @@
 #include "MVVM/Selection/Selection.h"
 #include "UnrealEdGlobals.h"
 #include "Misc/FeedbackContext.h"
+#include "UObject/UObjectIterator.h"
+#include "ObjectTools.h"
 
 #define LOCTEXT_NAMESPACE "FSequencerUtilities"
 
@@ -759,8 +761,6 @@ FGuid FSequencerUtilities::CreateCamera(TSharedRef<ISequencer> Sequencer, const
 	OutActor->SetActorRotation(GCurrentLevelEditingViewportClient->GetViewRotation());
 	//OutActor->CameraComponent->FieldOfView = ViewportClient->ViewFOV; //@todo set the focal length from this field of view
 
-	FMovieSceneSpawnable* Spawnable = nullptr;
-
 	if (bSpawnable)
 	{
 		FString NewName = MovieSceneHelpers::MakeUniqueSpawnableName(MovieScene, FName::NameToDisplayString(ACineCameraActor::StaticClass()->GetFName().ToString(), false));
@@ -1098,11 +1098,24 @@ FMovieScenePossessable* FSequencerUtilities::ConvertToPossessable(TSharedRef<ISe
 		return CreatedPossessable;
 	}
 
+	FMovieScenePossessable* ExistingPossessable = MovieScene->FindPossessable(BindingGuid);
+	if (ExistingPossessable)
+	{
+		if (const FMovieSceneBindingReference* ExistingReference = Sequence->GetBindingReferences()->GetReference(BindingGuid, BindingIndex))
+		{
+			if (!ExistingReference->CustomBinding)
+			{
+				// Already a possessable, just return
+				return ExistingPossessable;
+			}
+		}
+	}
+
 	UObject* ObjectToConvert = MovieSceneHelpers::GetSingleBoundObject(Sequence, BindingGuid, Sequencer->GetSharedPlaybackState(), BindingIndex);
 
 	// If we have an old-style spawnable, use the template as the object to convert instead.
 	bool bConvertFromSpawnable = MovieSceneHelpers::IsBoundToSpawnable(Sequence, BindingGuid, Sequencer->GetSharedPlaybackState(), BindingIndex);
-	if (bConvertFromSpawnable)
+	if (bConvertFromSpawnable && MovieSceneHelpers::SupportsObjectTemplate(Sequence, BindingGuid, Sequencer->GetSharedPlaybackState(), BindingIndex))
 	{
 		ObjectToConvert = MovieSceneHelpers::GetObjectTemplate(Sequence, BindingGuid, Sequencer->GetSharedPlaybackState(), BindingIndex);
 	}
@@ -1146,7 +1159,6 @@ FMovieScenePossessable* FSequencerUtilities::ConvertToPossessable(TSharedRef<ISe
 	}
 
 	FMovieSceneSpawnable* Spawnable = MovieScene->FindSpawnable(BindingGuid);
-	FMovieScenePossessable* ExistingPossessable = MovieScene->FindPossessable(BindingGuid);
 	// TODO: How to convert to possessable of non-actor type? Presumably we need to generalize the 'creation' step here for now.
 	UObject* PossessedObject = nullptr;
 	if (Actor)
@@ -1372,6 +1384,7 @@ FMovieScenePossessable* FSequencerUtilities::ConvertToCustomBinding(TSharedRef<I
 	CreateBindingParams.BindingIndex = BindingIndex;
 	CreateBindingParams.BindingNameOverride = NewCustomBinding->GetDesiredBindingName();
 	CreateBindingParams.CustomBinding = NewCustomBinding;
+	CreateBindingParams.bSetupDefaults = false;
 
 	// Create or replace the binding
 	FGuid NewPossessableGuid = CreateBinding(Sequencer, *ObjectToConvert, CreateBindingParams);
@@ -2093,12 +2106,11 @@ bool FSequencerUtilities::CanPasteSections(const FString& TextToImport)
  *
  */
 
-void ExportObjectBindingsToText(const TArray<UMovieSceneCopyableBinding*>& ObjectsToExport, FString& ExportedText, TSharedRef<UE::MovieScene::FSharedPlaybackState> SharedPlaybackState)
+void ExportObjectBindingsToText(const TArray<UMovieSceneCopyableBinding*>& ObjectsToExport, FOutputDevice& Archive, TSharedRef<UE::MovieScene::FSharedPlaybackState> SharedPlaybackState)
 {
 	// Clear the mark state for saving.
 	UnMarkAllObjects(EObjectMark(OBJECTMARK_TagExp | OBJECTMARK_TagImp));
 
-	FStringOutputDevice Archive;
 	const FExportObjectInnerContext Context;
 
 	// Export each of the selected nodes
@@ -2118,6 +2130,7 @@ void ExportObjectBindingsToText(const TArray<UMovieSceneCopyableBinding*>& Objec
 
 		TArray<UObject*, TInlineAllocator<1>> OldObjectTemplates;
 		TArray<UMovieSceneCustomBinding*, TInlineAllocator<1>> OldCustomBindings;
+		TMap<int32, UMovieSceneSpawnableBindingBase*> OldPreviewSpawnables;
 		if (ObjectToExport->Spawnable.GetGuid().IsValid())
 		{
 			OldObjectTemplates.Add(ObjectToExport->Spawnable.GetObjectTemplate());
@@ -2125,9 +2138,11 @@ void ExportObjectBindingsToText(const TArray<UMovieSceneCopyableBinding*>& Objec
 		}
 		else
 		{
-			for (UMovieSceneCustomBinding* CustomBinding : ObjectToExport->CustomBindings)
+			
+			for (int32 CustomBindingIndex = 0; CustomBindingIndex < ObjectToExport->CustomBindings.Num(); ++CustomBindingIndex)
 			{
-				if (UMovieSceneSpawnableBindingBase* SpawnableBinding = CustomBinding->AsSpawnable(SharedPlaybackState))
+				UMovieSceneCustomBinding* CustomBinding = ObjectToExport->CustomBindings[CustomBindingIndex];
+				if (UMovieSceneSpawnableBindingBase* SpawnableBinding = Cast<UMovieSceneSpawnableBindingBase>(CustomBinding))
 				{
 					if (SpawnableBinding->SupportsObjectTemplates())
 					{
@@ -2135,6 +2150,17 @@ void ExportObjectBindingsToText(const TArray<UMovieSceneCopyableBinding*>& Objec
 						SpawnableBinding->SetObjectTemplate(nullptr);
 					}
 				}
+				else if (UMovieSceneReplaceableBindingBase * ReplaceableBinding = Cast<UMovieSceneReplaceableBindingBase>(CustomBinding))
+				{
+					// Prevent inner references here during export
+					if (ReplaceableBinding->PreviewSpawnable)
+					{ 
+						OldPreviewSpawnables.Add(CustomBindingIndex, ReplaceableBinding->PreviewSpawnable);
+						// The Preview Spawnable is next in the CustomBindings list
+						ObjectToExport->PreviewSpawnableBindings.Add(CustomBindingIndex+1);
+						ReplaceableBinding->PreviewSpawnable = nullptr;
+					}
+				}
 				OldCustomBindings.Add(CustomBinding);
 			}
 			ObjectToExport->CustomBindings.Empty();
@@ -2155,6 +2181,15 @@ void ExportObjectBindingsToText(const TArray<UMovieSceneCopyableBinding*>& Objec
 			UExporter::ExportToOutputDevice(&Context, CustomBinding, nullptr, Archive, TEXT("copy"), 0, PPF_ExportsNotFullyQualified | PPF_Copy | PPF_Delimited);
 		}
 
+		// Restore replaceable references now that we've exported them
+		for (const TPair<int32, UMovieSceneSpawnableBindingBase*>& PreviewSpawnable : OldPreviewSpawnables)
+		{
+			if (UMovieSceneReplaceableBindingBase* ReplaceableBinding = Cast<UMovieSceneReplaceableBindingBase>(ObjectToExport->CustomBindings[PreviewSpawnable.Key].Get()))
+			{
+				ReplaceableBinding->PreviewSpawnable = PreviewSpawnable.Value;
+			}
+		}
+
 		int32 ObjectTemplateIndex = 0;
 		if (ObjectToExport->Spawnable.GetGuid().IsValid())
 		{
@@ -2166,7 +2201,8 @@ void ExportObjectBindingsToText(const TArray<UMovieSceneCopyableBinding*>& Objec
 			{
 				if (UMovieSceneSpawnableBindingBase* SpawnableBinding = CustomBinding->AsSpawnable(SharedPlaybackState))
 				{
-					if (SpawnableBinding->SupportsObjectTemplates())
+					// Ignore bindings with their template already set, which is possible for Replaceables since they'll show up twice in the list.
+					if (SpawnableBinding->SupportsObjectTemplates() && !SpawnableBinding->GetObjectTemplate())
 					{
 						SpawnableBinding->SetObjectTemplate(OldObjectTemplates[ObjectTemplateIndex++]);
 					}
@@ -2180,11 +2216,16 @@ void ExportObjectBindingsToText(const TArray<UMovieSceneCopyableBinding*>& Objec
 			UExporter::ExportToOutputDevice(&Context, ObjectTemplate, nullptr, Archive, TEXT("copy"), 0, PPF_ExportsNotFullyQualified | PPF_Copy | PPF_Delimited);
 		}
 	}
+}
 
+void FSequencerUtilities::CopyBindings(TSharedRef<ISequencer> Sequencer, const TArray<FMovieSceneBindingProxy>& Bindings, const TArray<UMovieSceneFolder*>& InFolders, FString& ExportedText)
+{
+	FStringOutputDevice Archive;
+	CopyBindings(Sequencer, Bindings, InFolders, (FOutputDevice&)Archive);
 	ExportedText = Archive;
 }
 
-void FSequencerUtilities::CopyBindings(TSharedRef<ISequencer> Sequencer, const TArray<FMovieSceneBindingProxy>& Bindings, const TArray<UMovieSceneFolder*>& InFolders, FString& ExportedText)
+void FSequencerUtilities::CopyBindings(TSharedRef<ISequencer> Sequencer, const TArray<FMovieSceneBindingProxy>& Bindings, const TArray<UMovieSceneFolder*>& InFolders, FOutputDevice& Ar)
 {
 	TArray<UMovieSceneCopyableBinding*> Objects;
 	for (const FMovieSceneBindingProxy& ObjectBinding : Bindings)
@@ -2201,24 +2242,24 @@ void FSequencerUtilities::CopyBindings(TSharedRef<ISequencer> Sequencer, const T
 		FMovieScenePossessable* Possessable = MovieScene->FindPossessable(ObjectBinding.BindingID);
 		if (Possessable)
 		{
-			CopyableBinding->Possessable = *Possessable;
+			UObject* ResolutionContext = FindResolutionContext(Sequencer
+				, *MovieScene->GetTypedOuter<UMovieSceneSequence>()
+				, *MovieScene
+				, Possessable->GetParent()
+				, Sequencer->GetPlaybackContext());
 
-			// Store the names of the bound objects so that they can be found on paste
-			for (TWeakObjectPtr<> RuntimeObject : Sequencer->FindBoundObjects(CopyableBinding->Possessable.GetGuid(), Sequencer->GetFocusedTemplateID()))
-			{
-				if (AActor* Actor = Cast<AActor>(RuntimeObject.Get()))
-				{
-					CopyableBinding->BoundObjectNames.Add(Actor->GetPathName());
-				}
-			}
+			CopyableBinding->Possessable = *Possessable;
 
 			// Store any custom bindings
 			if (FMovieSceneBindingReferences* BindingReferences = ObjectBinding.Sequence->GetBindingReferences())
 			{
+				int32 BindingIndex = 0;
 				for (const FMovieSceneBindingReference& BindingReference : BindingReferences->GetReferences(ObjectBinding.BindingID))
 				{
 					if (UMovieSceneCustomBinding* CustomBinding = BindingReference.CustomBinding)
 					{
+						CopyableBinding->CustomBindings.Add(CustomBinding);
+
 						if (UMovieSceneSpawnableBindingBase* SpawnableBinding = CustomBinding->AsSpawnable(Sequencer->GetSharedPlaybackState()))
 						{
 							if (SpawnableBinding->SupportsObjectTemplates())
@@ -2227,9 +2268,28 @@ void FSequencerUtilities::CopyBindings(TSharedRef<ISequencer> Sequencer, const T
 								// Object Templates are re-created on paste (based on these templates) with the correct ownership set up.
 								CopyableBinding->SpawnableObjectTemplates.Add(SpawnableBinding->GetObjectTemplate());
 							}
+
+							// This is the inner spawnable of a replaceable and is always placed after the replaceable in the list
+							if (SpawnableBinding != CustomBinding)
+							{
+								CopyableBinding->CustomBindings.Add(SpawnableBinding);
+							}
 						}
-						CopyableBinding->CustomBindings.Add(CustomBinding);
 					}
+					else if (UObject* RuntimeObject = MovieSceneHelpers::GetSingleBoundObject(ObjectBinding.Sequence, ObjectBinding.BindingID, Sequencer->GetSharedPlaybackState(), BindingIndex))
+					{
+						CopyableBinding->BoundObjectNames.Add(RuntimeObject->GetPathName(ResolutionContext));
+					}
+					BindingIndex++;
+				}
+			}
+			else
+			{
+				// Store the names of the bound objects so that they can be found on paste
+
+				for (TWeakObjectPtr<> RuntimeObject : Sequencer->FindBoundObjects(CopyableBinding->Possessable.GetGuid(), Sequencer->GetFocusedTemplateID()))
+				{
+					CopyableBinding->BoundObjectNames.Add(RuntimeObject->GetPathName(ResolutionContext));
 				}
 			}
 		}
@@ -2280,7 +2340,7 @@ void FSequencerUtilities::CopyBindings(TSharedRef<ISequencer> Sequencer, const T
 		}
 	}
 
-	ExportObjectBindingsToText(Objects, /*out*/ ExportedText, Sequencer->GetSharedPlaybackState());
+	ExportObjectBindingsToText(Objects, /*out*/ Ar, Sequencer->GetSharedPlaybackState());
 }
 
 class FObjectBindingTextFactory : public FCustomizableTextObjectFactory
@@ -2371,7 +2431,7 @@ void ImportObjectBindingsFromText(ISequencer& InSequencer, const FString& TextTo
 			{
 				for (UMovieSceneCustomBinding* CustomBinding : ImportedObject->CustomBindings)
 				{
-					if (UMovieSceneSpawnableBindingBase* SpawnableBinding = CustomBinding->AsSpawnable(InSequencer.GetSharedPlaybackState()))
+					if (UMovieSceneSpawnableBindingBase* SpawnableBinding = Cast<UMovieSceneSpawnableBindingBase>(CustomBinding))
 					{
 						if (SpawnableBinding->SupportsObjectTemplates())
 						{
@@ -2387,7 +2447,7 @@ void ImportObjectBindingsFromText(ISequencer& InSequencer, const FString& TextTo
 	TempPackage->RemoveFromRoot();
 }
 
-FGuid TryCreateCustomBinding(TSharedRef<ISequencer> Sequencer, UObject* CustomBindingObject, FMovieSceneBindingReferences* BindingReferences, const UE::Sequencer::FCreateBindingParams& InParams, UMovieScene* OwnerMovieScene, bool bSpawnable, bool bReplaceable)
+FGuid TryCreateCustomBinding(TSharedPtr<ISequencer> Sequencer, UObject* CustomBindingObject, AActor* FactoryCreatedActor, FMovieSceneBindingReferences* BindingReferences, const UE::Sequencer::FCreateBindingParams& InParams, UMovieScene* OwnerMovieScene, bool bSpawnable, bool bReplaceable)
 {
 	UMovieSceneCustomBinding * NewCustomBinding = nullptr;
 	if (InParams.CustomBinding)
@@ -2403,7 +2463,23 @@ FGuid TryCreateCustomBinding(TSharedRef<ISequencer> Sequencer, UObject* CustomBi
 	}
 	else
 	{
-		TArrayView<const TSubclassOf<UMovieSceneCustomBinding>> PrioritySortedCustomBindingTypes = Sequencer->GetSupportedCustomBindingTypes();
+		TArrayView<const TSubclassOf<UMovieSceneCustomBinding>> PrioritySortedCustomBindingTypes;
+		if (Sequencer)
+		{
+			PrioritySortedCustomBindingTypes = Sequencer->GetSupportedCustomBindingTypes();
+		}
+		else
+		{
+			static TArray<const TSubclassOf<UMovieSceneCustomBinding>> CachedCustomBindingTypes;
+			static bool CustomBindingTypesCached = false;
+			if (!CustomBindingTypesCached)
+			{
+				CustomBindingTypesCached = true;
+				MovieSceneHelpers::GetPrioritySortedCustomBindingTypes(CachedCustomBindingTypes);
+			}
+			PrioritySortedCustomBindingTypes = CachedCustomBindingTypes;
+		}
+
 		for (const TSubclassOf<UMovieSceneCustomBinding>& CustomBindingType : PrioritySortedCustomBindingTypes)
 		{
 			// If 'spawnable' has been passed in, we only want to use children of UMovieSceneSpawnableBindingBase and vice versa
@@ -2420,6 +2496,16 @@ FGuid TryCreateCustomBinding(TSharedRef<ISequencer> Sequencer, UObject* CustomBi
 							break;
 						}
 					}
+
+					if (!NewCustomBinding && FactoryCreatedActor && CustomBindingCDO->SupportsBindingCreationFromObject(FactoryCreatedActor))
+					{
+						// Create a custom binding from the factory created actor
+						NewCustomBinding = CustomBindingCDO->CreateNewCustomBinding(FactoryCreatedActor, *OwnerMovieScene);
+						if (NewCustomBinding)
+						{
+							break;
+						}
+					}
 				}
 			}
 		}
@@ -2452,33 +2538,36 @@ FGuid TryCreateCustomBinding(TSharedRef<ISequencer> Sequencer, UObject* CustomBi
 
 		UObject* SpawnedObject = nullptr;
 		// If this is a spawnable or replaceable binding, we need to set up some defaults
-		if (NewCustomBinding->WillSpawnObject(Sequencer->GetSharedPlaybackState()))
+		if (Sequencer)
 		{
-			// Spawn the object so we can position it correctly, it's going to get spawned anyway since things default to spawned.
-			SpawnedObject = Sequencer->GetSpawnRegister().SpawnObject(NewID, *OwnerMovieScene, Sequencer->GetFocusedTemplateID(), Sequencer->GetSharedPlaybackState(), 0);
-
-			if (InParams.bSetupDefaults)
+			if (NewCustomBinding->WillSpawnObject(Sequencer->GetSharedPlaybackState()))
 			{
-				// Allow the binding to set up any necessary defaults
-				NewCustomBinding->SetupDefaults(SpawnedObject, NewID, *OwnerMovieScene, Sequencer->GetSharedPlaybackState());
+				// Spawn the object so we can position it correctly, it's going to get spawned anyway since things default to spawned.
+				SpawnedObject = Sequencer->GetSpawnRegister().SpawnObject(NewID, *OwnerMovieScene, Sequencer->GetFocusedTemplateID(), Sequencer->GetSharedPlaybackState(), 0);
+
+				if (InParams.bSetupDefaults)
+				{
+					// Allow the binding to set up any necessary defaults
+					NewCustomBinding->SetupDefaults(SpawnedObject, NewID, *OwnerMovieScene, Sequencer->GetSharedPlaybackState());
 
-				FTransformData TransformData;
-				Sequencer->GetSpawnRegister().SetupDefaultsForSpawnable(SpawnedObject, NewID, TransformData, Sequencer, Sequencer->GetSequencerSettings());
+					FTransformData TransformData;
+					Sequencer->GetSpawnRegister().SetupDefaultsForSpawnable(SpawnedObject, NewID, TransformData, Sequencer.ToSharedRef(), Sequencer->GetSequencerSettings());
+				}
 			}
-		}
 
-		Sequencer->State.Invalidate(NewID, Sequencer->GetFocusedTemplateID());
-		Sequencer->ForceEvaluate();
+			Sequencer->State.Invalidate(NewID, Sequencer->GetFocusedTemplateID());
+			Sequencer->ForceEvaluate();
 
-		// We don't call these events in the case bSetupDefaults is false because they may add tracks.
-		if (InParams.bSetupDefaults)
-		{
-			if (AActor* Actor = Cast<AActor>(SpawnedObject))
+			// We don't call these events in the case bSetupDefaults is false because they may add tracks.
+			if (InParams.bSetupDefaults)
 			{
-				Sequencer->OnActorAddedToSequencer().Broadcast(Actor, NewID);
-			}
+				if (AActor* Actor = Cast<AActor>(SpawnedObject))
+				{
+					Sequencer->OnActorAddedToSequencer().Broadcast(Actor, NewID);
+				}
 
-			Sequencer->OnAddBinding(NewID, OwnerMovieScene);
+				Sequencer->OnAddBinding(NewID, OwnerMovieScene);
+			}
 		}
 
 		return NewID;
@@ -2486,11 +2575,15 @@ FGuid TryCreateCustomBinding(TSharedRef<ISequencer> Sequencer, UObject* CustomBi
 	return FGuid();
 }
 
-FGuid CreateGenericBinding(TSharedRef<ISequencer> Sequencer, UObject* InObject, FMovieSceneBindingReferences* BindingReferences, const UE::Sequencer::FCreateBindingParams& InParams)
+FGuid CreateGenericBinding(TSharedPtr<ISequencer> Sequencer, UMovieSceneSequence* OwnerSequence, UObject* InObject, FMovieSceneBindingReferences* BindingReferences, const UE::Sequencer::FCreateBindingParams& InParams)
 {
+	if (!OwnerSequence)
+	{
+		return FGuid();
+	}
+
 	using namespace UE::Sequencer;
 
-	UMovieSceneSequence* OwnerSequence = Sequencer->GetFocusedMovieSceneSequence();
 	UMovieScene* OwnerMovieScene = OwnerSequence->GetMovieScene();
 
 	ISequencerModule& Module = FModuleManager::Get().LoadModuleChecked<ISequencerModule>("Sequencer");
@@ -2502,12 +2595,10 @@ FGuid CreateGenericBinding(TSharedRef<ISequencer> Sequencer, UObject* InObject,
 
 	if (bAllowCustom)
 	{
-		NewBindingID = TryCreateCustomBinding(Sequencer, InObject, BindingReferences, InParams, OwnerMovieScene, bSpawnable, bReplaceable);
-		if (NewBindingID.IsValid())
-		{
-			return NewBindingID;
-		}
-		else if (InObject && !InObject->IsA<AActor>())
+		// In addition to the raw object, we also try spawning an actor from an actor factory if relevant, to give the custom binding an option to create from that as well
+		AActor* FactoryCreatedActorInstance = nullptr; 
+		UWorld* World = GCurrentLevelEditingViewportClient ? GCurrentLevelEditingViewportClient->GetWorld() : nullptr;
+		if (InObject && !InObject->IsA<AActor>())
 		{
 			// If the passed in object is not an actor, see if we can create an Actor from it, and if so, if that Actor type has a custom binding that supports it
 			UActorFactory* FactoryToUse = InParams.ActorFactory ? InParams.ActorFactory.Get() : FActorFactoryAssetProxy::GetFactoryForAssetObject(InObject);
@@ -2517,7 +2608,6 @@ FGuid CreateGenericBinding(TSharedRef<ISequencer> Sequencer, UObject* InObject,
 				FText ErrorMessage;
 				if (FactoryToUse->CanCreateActorFrom(FAssetData(InObject), ErrorMessage))
 				{
-					UWorld* World = GCurrentLevelEditingViewportClient ? GCurrentLevelEditingViewportClient->GetWorld() : nullptr;
 					if (World)
 					{
 						const FName ActorName = MakeUniqueObjectName(World->PersistentLevel, FactoryToUse->NewActorClass->StaticClass(), *InParams.BindingNameOverride);
@@ -2526,29 +2616,28 @@ FGuid CreateGenericBinding(TSharedRef<ISequencer> Sequencer, UObject* InObject,
 						SpawnParams.ObjectFlags = RF_Transient | RF_Transactional;
 						SpawnParams.Name = ActorName;
 
-						AActor* Instance = FactoryToUse->CreateActor(InObject, World->PersistentLevel, FTransform(), SpawnParams);
-						if (Instance)
+						FactoryCreatedActorInstance = FactoryToUse->CreateActor(InObject, World->PersistentLevel, FTransform(), SpawnParams);
+						if (FactoryCreatedActorInstance)
 						{
-							Instance->SetActorLabel(MovieSceneHelpers::MakeUniqueBindingName(OwnerMovieScene, FName::NameToDisplayString(InObject->GetName(), false)));
-							Instance->bIsEditorPreviewActor = false;
-
-							// Can we create a custom binding from this factory created actor? In all cases, we expect the binding creation process to duplicate
-							// the actor, so we should be safe to destroy the actor after this process.
-							NewBindingID = TryCreateCustomBinding(Sequencer, Instance, BindingReferences, InParams, OwnerMovieScene, bSpawnable, bReplaceable);
-
-							const bool bNetForce = false;
-							const bool bShouldModifyLevel = false;
-							World->DestroyActor(Instance, bNetForce, bShouldModifyLevel);
-
-							if (NewBindingID.IsValid())
-							{
-								return NewBindingID;
-							}
+							FactoryCreatedActorInstance->SetActorLabel(MovieSceneHelpers::MakeUniqueBindingName(OwnerMovieScene, FName::NameToDisplayString(InObject->GetName(), false)));
+							FactoryCreatedActorInstance->bIsEditorPreviewActor = false;
 						}
 					}
 				}
 			}
 		}
+		NewBindingID = TryCreateCustomBinding(Sequencer, InObject, FactoryCreatedActorInstance, BindingReferences, InParams, OwnerMovieScene, bSpawnable, bReplaceable);
+		if (FactoryCreatedActorInstance)
+		{
+			const bool bNetForce = false;
+			const bool bShouldModifyLevel = false;
+			World->DestroyActor(FactoryCreatedActorInstance, bNetForce, bShouldModifyLevel);
+			FactoryCreatedActorInstance = nullptr;
+		}
+		if (NewBindingID.IsValid())
+		{
+			return NewBindingID;
+		}
 	}
 
 	if (!InObject)
@@ -2558,9 +2647,9 @@ FGuid CreateGenericBinding(TSharedRef<ISequencer> Sequencer, UObject* InObject,
 
 
 	// If no custom bindings support this object type, but InParams.bSpawnable is true, attempt to make an old-style spawnable.
-	if (bSpawnable)
+	if (bSpawnable && Sequencer)
 	{
-		NewBindingID = FSequencerUtilities::MakeNewSpawnable(Sequencer, *InObject, InParams.ActorFactory);
+		NewBindingID = FSequencerUtilities::MakeNewSpawnable(Sequencer.ToSharedRef(), *InObject, InParams.ActorFactory);
 		if (NewBindingID.IsValid())
 		{
 			return NewBindingID;
@@ -2607,7 +2696,8 @@ FGuid CreateGenericBinding(TSharedRef<ISequencer> Sequencer, UObject* InObject,
 
 	const bool bParentContextsAreSignificant = OwnerSequence->AreParentContextsSignificant();
 
-	UObject* Context = Sequencer->GetPlaybackContext();
+
+	UObject* Context = Sequencer ? Sequencer->GetPlaybackContext() : nullptr;
 
 	FGuid ParentID;
 
@@ -2617,7 +2707,7 @@ FGuid CreateGenericBinding(TSharedRef<ISequencer> Sequencer, UObject* InObject,
 		UObject* CurrentObject = ObjectsToPossess[Index].Key;
 
 		// If we're not purposefully replacing a binding, then check to see if we already have one, and use that
-		if (!InParams.ReplacementGuid.IsValid())
+		if (Sequencer && !InParams.ReplacementGuid.IsValid())
 		{
 			FGuid    ObjectGuid = Sequencer->GetHandleToObject(CurrentObject, false);
 
@@ -2650,7 +2740,7 @@ FGuid CreateGenericBinding(TSharedRef<ISequencer> Sequencer, UObject* InObject,
 		}
 
 		// If we're not trying to replace a binding, and the object is a spawnable, try and bind to that first
-		if (InParams.ReplacementGuid.IsValid() || !NewPossessable->BindSpawnableObject(Sequencer->GetFocusedTemplateID(), CurrentObject, Sequencer->GetSharedPlaybackState()))
+		if (InParams.ReplacementGuid.IsValid() || (Sequencer && !NewPossessable->BindSpawnableObject(Sequencer->GetFocusedTemplateID(), CurrentObject, Sequencer->GetSharedPlaybackState())))
 		{
 			FUniversalObjectLocator Locator;
 			if (!OwnerSequence->MakeLocatorForObject(CurrentObject, Context, Locator) || Locator.IsEmpty())
@@ -2662,7 +2752,10 @@ FGuid CreateGenericBinding(TSharedRef<ISequencer> Sequencer, UObject* InObject,
 			if (InParams.ReplacementGuid.IsValid() && !ParentID.IsValid())
 			{
 				BindingReferences->AddOrReplaceBinding(NewID, MoveTemp(Locator), InParams.BindingIndex);
-				Sequencer->State.Invalidate(NewID, Sequencer->GetFocusedTemplateID());
+				if (Sequencer)
+				{
+					Sequencer->State.Invalidate(NewID, Sequencer->GetFocusedTemplateID());
+				}
 			}
 			else
 			{
@@ -2682,16 +2775,21 @@ FGuid CreateGenericBinding(TSharedRef<ISequencer> Sequencer, UObject* InObject,
 		}
 
 		ParentID = NewID;
-
-		if (AActor* Actor = Cast<AActor>(CurrentObject))
+		if (Sequencer)
 		{
-			Sequencer->OnActorAddedToSequencer().Broadcast(Actor, NewID);
+			if (AActor* Actor = Cast<AActor>(CurrentObject))
+			{
+				Sequencer->OnActorAddedToSequencer().Broadcast(Actor, NewID);
+			}
 		}
 
 		// If this is the last one
 		if (Index == 0)
 		{
-			Sequencer->OnAddBinding(NewID, OwnerMovieScene);
+			if (Sequencer)
+			{
+				Sequencer->OnAddBinding(NewID, OwnerMovieScene);
+			}
 			return NewID;
 		}
 
@@ -2779,33 +2877,47 @@ bool FSequencerUtilities::PasteBindings(const FString& TextToImport, TSharedRef<
 				{
 					if (BindingReferences)
 					{
-						int32 BindingIndex = 0;
 						int32 SpawnableBindingIndex = 0;
-						for (UMovieSceneCustomBinding* CustomBinding : CopyableBinding->CustomBindings)
+						UMovieSceneCustomBinding* PreviousBinding = nullptr;
+						for(int32 BindingIndex = 0; BindingIndex < CopyableBinding->CustomBindings.Num(); ++BindingIndex)
 						{
+							UMovieSceneCustomBinding* CustomBinding = CopyableBinding->CustomBindings[BindingIndex];
 							if (CustomBinding)
 							{
 								UMovieSceneCustomBinding* NewCustomBinding = Cast<UMovieSceneCustomBinding>(StaticDuplicateObject(CustomBinding, MovieScene));
 
 								// Need to re-copy the object template to avoid private object issues
-								if (UMovieSceneSpawnableBindingBase* SpawnableBinding = NewCustomBinding->AsSpawnable(Sequencer->GetSharedPlaybackState()))
+								if (UMovieSceneSpawnableBindingBase* SpawnableBinding = Cast<UMovieSceneSpawnableBindingBase>(NewCustomBinding))
 								{
 									if (CopyableBinding->SpawnableObjectTemplates.IsValidIndex(SpawnableBindingIndex))
 									{
 										UObject* NewObjectTemplate = StaticDuplicateObject(CopyableBinding->SpawnableObjectTemplates[SpawnableBindingIndex++], MovieScene);
 										SpawnableBinding->SetObjectTemplate(NewObjectTemplate);
 									}
+
+									// If this is a preview spawnable, find the just added replaceable and link with that rather than creating a new binding
+									if (CopyableBinding->PreviewSpawnableBindings.Contains(BindingIndex))
+									{
+										if (UMovieSceneReplaceableBindingBase* PreviousReplaceableBinding = Cast<UMovieSceneReplaceableBindingBase>(PreviousBinding))
+										{
+											PreviousReplaceableBinding->PreviewSpawnable = SpawnableBinding;
+											PreviousBinding = NewCustomBinding;
+											continue;
+										}
+									}
 								}
 
 								// This will either add a brand new possessable and binding (if one doesn't exist for that guid), or just add a new binding to that same possessable
 								UE::Sequencer::FCreateBindingParams CreateBindingParams;
 								CreateBindingParams.ReplacementGuid = NewGuid;
-								CreateBindingParams.BindingIndex = BindingIndex++;
+								CreateBindingParams.BindingIndex = BindingIndex;
 								CreateBindingParams.bAllowCustomBinding = true;
 								CreateBindingParams.CustomBinding = NewCustomBinding;
 								CreateBindingParams.bSetupDefaults = false;
 								CreateBindingParams.BindingNameOverride = CopyableBinding->Possessable.GetName();
-								NewGuid = CreateGenericBinding(Sequencer, nullptr, BindingReferences, CreateBindingParams);
+								NewGuid = CreateGenericBinding(Sequencer, Sequence, nullptr, BindingReferences, CreateBindingParams);
+
+								PreviousBinding = NewCustomBinding;
 							}
 						}
 					}
@@ -2861,9 +2973,10 @@ bool FSequencerUtilities::PasteBindings(const FString& TextToImport, TSharedRef<
 					}
 				}
 
-				// Find the actors that this pasted binding should bind to
-				TArray<AActor*> ActorsToRebind;
-				if (World && CopyableBinding->CustomBindings.IsEmpty())
+				// Find the objects that this pasted binding should bind to
+				TArray<UObject*> ObjectsToBind;
+				UObject* ResolutionContext = Sequencer->GetPlaybackContext();
+				if (World)
 				{
 					for (TActorIterator<AActor> ActorItr(World); ActorItr; ++ActorItr)
 					{
@@ -2876,20 +2989,25 @@ bool FSequencerUtilities::PasteBindings(const FString& TextToImport, TSharedRef<
 								continue;
 							}
 
-							ActorsToRebind.Add(Actor);
+							ObjectsToBind.Add(Actor);
 							CopyableBinding->BoundObjectNames.Remove(Actor->GetPathName());
 						}
 					}
 				}
 
-				if (ActorsToRebind.Num() != 0)
+				if (ObjectsToBind.Num() != 0)
 				{
 					if (PasteBindingsParams.bDuplicateExistingActors)
 					{
 						GEditor->SelectNone(false, true);
-						for (AActor* ActorToRebind : ActorsToRebind)
+						TArray<UObject*> SelectedObjects;
+						for (UObject* ObjectToBind : ObjectsToBind)
 						{
-							GEditor->SelectActor(ActorToRebind, true, false, false);
+							if (AActor* Actor = Cast<AActor>(ObjectToBind))
+							{
+								GEditor->SelectActor(Actor, true, false, false);
+								SelectedObjects.Add(Actor);
+							}
 						}
 							
 						// Duplicate the bound actors
@@ -2900,14 +3018,14 @@ bool FSequencerUtilities::PasteBindings(const FString& TextToImport, TSharedRef<
 						// of the clipboard after duplicating the actor
 						FPlatformApplicationMisc::ClipboardCopy(*TextToImport);
 
-						ActorsToRebind.Empty();
+						ObjectsToBind.RemoveAll([&SelectedObjects](UObject* Object) { return SelectedObjects.Contains(Object);});
 						USelection* ActorSelection = GEditor->GetSelectedActors();
 						for (FSelectionIterator Iter(*ActorSelection); Iter; ++Iter)
 						{
 							AActor* Actor = Cast<AActor>(*Iter);
 							if (Actor)
 							{
-								ActorsToRebind.Add(Actor);
+								ObjectsToBind.Add(Actor);
 
 								CopyableBinding->BoundObjectNames.Add(Actor->GetPathName());
 							}
@@ -2915,9 +3033,9 @@ bool FSequencerUtilities::PasteBindings(const FString& TextToImport, TSharedRef<
 					}
 
 					// Bind the actors
-					if (ActorsToRebind.Num())
+					if (ObjectsToBind.Num())
 					{
-						AddActorsToBinding(Sequencer, ActorsToRebind, FMovieSceneBindingProxy(NewGuid, Sequence));
+						AddObjectsToBinding(Sequencer, ObjectsToBind, FMovieSceneBindingProxy(NewGuid, Sequence), ResolutionContext);
 					}
 				}
 
@@ -2997,44 +3115,72 @@ bool FSequencerUtilities::PasteBindings(const FString& TextToImport, TSharedRef<
 		FMovieScenePossessable* Possessable = MovieScene->FindPossessable(PossessableGuids[PossessableGuidIndex]);
 		if (Possessable && World)
 		{
-			for (TActorIterator<AActor> ActorItr(World); ActorItr; ++ActorItr)
+			auto AddActor = [&](AActor* Actor)
 			{
-				AActor* Actor = *ActorItr;
-				if (Actor && PossessableGuidIndex < PossessableObjectNames.Num() && PossessableObjectNames[PossessableGuidIndex].Contains(Actor->GetName()))
+				FGuid ExistingGuid = Sequencer->FindObjectId(*Actor, Sequencer->GetFocusedTemplateID());
+
+				if (!ExistingGuid.IsValid())
 				{
-					FGuid ExistingGuid = Sequencer->FindObjectId(*Actor, Sequencer->GetFocusedTemplateID());
+					FGuid NewGuid = AssignActor(Sequencer, Actor, Possessable->GetGuid());
 
-					if (!ExistingGuid.IsValid())
+					// If assigning produces a new guid, update the possessable guids and the bindings pasted data
+					if (NewGuid.IsValid())
 					{
-						FGuid NewGuid = AssignActor(Sequencer, Actor, Possessable->GetGuid());
-
-						// If assigning produces a new guid, update the possessable guids and the bindings pasted data
-						if (NewGuid.IsValid())
+						for (auto BindingPasted : BindingsPasted)
 						{
-							for (auto BindingPasted : BindingsPasted)
+							if (BindingPasted.GetObjectGuid() == PossessableGuids[PossessableGuidIndex])
 							{
-								if (BindingPasted.GetObjectGuid() == PossessableGuids[PossessableGuidIndex])
-								{
-									BindingPasted.SetObjectGuid(NewGuid);
-								}
+								BindingPasted.SetObjectGuid(NewGuid);
 							}
+						}
 
-							if (GuidToFolderMap.Contains(PossessableGuids[PossessableGuidIndex]))
-							{
-								GuidToFolderMap.Add(NewGuid, GuidToFolderMap[PossessableGuids[PossessableGuidIndex]]);
-								GuidToFolderMap.Remove(PossessableGuids[PossessableGuidIndex]);
-							}
+						if (GuidToFolderMap.Contains(PossessableGuids[PossessableGuidIndex]))
+						{
+							GuidToFolderMap.Add(NewGuid, GuidToFolderMap[PossessableGuids[PossessableGuidIndex]]);
+							GuidToFolderMap.Remove(PossessableGuids[PossessableGuidIndex]);
+						}
 
-							for (TPair<FGuid, FGuid>& OldToNewGuid : OldToNewGuidMap)
+						for (TPair<FGuid, FGuid>& OldToNewGuid : OldToNewGuidMap)
+						{
+							if (OldToNewGuid.Value == PossessableGuids[PossessableGuidIndex])
 							{
-								if (OldToNewGuid.Value == PossessableGuids[PossessableGuidIndex])
-								{
-									OldToNewGuid.Value = NewGuid;
-								}
+								OldToNewGuid.Value = NewGuid;
 							}
+						}
 
-							PossessableGuids[PossessableGuidIndex] = NewGuid;
+						PossessableGuids[PossessableGuidIndex] = NewGuid;
+					}
+				}
+			};
+			
+			for (TActorIterator<AActor> ActorItr(World); ActorItr; ++ActorItr)
+			{
+				AActor* Actor = *ActorItr;
+				FString ActorName = Actor ? Actor->GetName() : TEXT("");
+				if (Actor && PossessableGuidIndex < PossessableObjectNames.Num() && PossessableObjectNames[PossessableGuidIndex].Contains(Actor->GetName()))
+				{
+					AddActor(Actor);
+				}
+			}
+
+			// If pasted actors have been provided, go through those as well
+			for (const TPair<FName, TObjectPtr<AActor>>& PastedActorPair : PasteBindingsParams.PastedActors)
+			{
+				if (PastedActorPair.Value && PossessableGuidIndex < PossessableObjectNames.Num())
+				{
+					FString PastedActorPairName = ObjectTools::SanitizeObjectPath(PastedActorPair.Key.ToString());
+					if (Algo::AnyOf(PossessableObjectNames[PossessableGuidIndex], [&PastedActorPairName] (const FString& PathName) 
+					{ 
+						FString SanitizedObjectPath = ObjectTools::SanitizeObjectPath(PathName);
+						int32 LastPeriod = INDEX_NONE;
+						if (SanitizedObjectPath.FindLastChar('.', LastPeriod))
+						{
+							return SanitizedObjectPath.RightChop(LastPeriod+1) == PastedActorPairName;
 						}
+						return SanitizedObjectPath == PastedActorPairName;
+					}))
+					{
+						AddActor(PastedActorPair.Value.Get());
 					}
 				}
 			}
@@ -3080,9 +3226,10 @@ bool FSequencerUtilities::PasteBindings(const FString& TextToImport, TSharedRef<
 	// Refresh all immediately so that spawned actors will be generated immediately
 	Sequencer->ForceEvaluate();
 
-	// Fix possessable component bindings
-	for (auto PossessableGuid : PossessableGuids)
+	// Fix possessable subobject bindings
+	for (int32 PossessableGuidIndex = 0; PossessableGuidIndex < PossessableGuids.Num(); ++PossessableGuidIndex)
 	{
+		FGuid PossessableGuid = PossessableGuids[PossessableGuidIndex];
 		// If a possessable guid does not have any bound objects, they might be 
 		// possessable components for spawnables, so they need to be remapped
 		if (Sequencer->FindBoundObjects(PossessableGuid, Sequencer->GetFocusedTemplateID()).Num() == 0)
@@ -3091,6 +3238,7 @@ bool FSequencerUtilities::PasteBindings(const FString& TextToImport, TSharedRef<
 			if (Possessable)
 			{
 				FGuid ParentGuid = Possessable->GetParent();
+				bool bBound = false;
 				for (TWeakObjectPtr<> WeakObject : Sequencer->FindBoundObjects(ParentGuid, Sequencer->GetFocusedTemplateID()))
 				{
 					if (AActor* SpawnedActor = Cast<AActor>(WeakObject.Get()))
@@ -3100,11 +3248,29 @@ bool FSequencerUtilities::PasteBindings(const FString& TextToImport, TSharedRef<
 							if (Component->GetName() == Possessable->GetName())
 							{
 								Sequence->BindPossessableObject(PossessableGuid, *Component, SpawnedActor);
+								bBound = true;
 								break;
 							}
 						}
 					}
+
+					if (!bBound)
+					{
+						if (PossessableGuidIndex < PossessableObjectNames.Num())
+						{
+							for (const FString& BoundObjectPath : PossessableObjectNames[PossessableGuidIndex])
+							{
+								if (UObject* FoundObject = FindObject<UObject>(WeakObject.Get(), *BoundObjectPath))
+								{
+									Sequence->BindPossessableObject(PossessableGuid, *FoundObject, WeakObject.Get());
+									bBound = true;
+									break;
+								}
+							}
+						}
+					}
 				}
+				
 
 				// If the parent doesn't actually exist, clear it.
 				FMovieScenePossessable* PossessableParent = MovieScene->FindPossessable(ParentGuid);
@@ -3238,11 +3404,65 @@ FGuid CreateImplementationDefinedBinding(TSharedRef<ISequencer> Sequencer, UObje
 	return PossessableGuid;
 }
 
+UObject* FSequencerUtilities::FindResolutionContext(TSharedRef<ISequencer> InSequencer
+	, UMovieSceneSequence& InSequence
+	, UMovieScene& InMovieScene
+	, const FGuid& InParentGuid
+	, UObject* InPlaybackContext)
+{
+	if (!InPlaybackContext || !InParentGuid.IsValid() || !InSequence.AreParentContextsSignificant())
+	{
+		return InPlaybackContext;
+	}
+
+	UObject* ResolutionContext = nullptr;
+
+	// Recursive call up the hierarchy
+	if (FMovieScenePossessable* const ParentPossessable = InMovieScene.FindPossessable(InParentGuid))
+	{
+		ResolutionContext = FSequencerUtilities::FindResolutionContext(InSequencer
+			, InSequence
+			, InMovieScene
+			, ParentPossessable->GetParent()
+			, InPlaybackContext);
+	}
+
+	if (!ResolutionContext)
+	{
+		ResolutionContext = InPlaybackContext;
+	}
+
+	TArray<UObject*, TInlineAllocator<1>> FoundObjects;
+	for (TWeakObjectPtr<> WeakObj : InSequencer->FindBoundObjects(InParentGuid, InSequencer->GetFocusedTemplateID()))
+	{
+		FoundObjects.Add(WeakObj.Get());
+	}
+
+	if (FoundObjects.IsEmpty())
+	{
+		return ResolutionContext;
+	}
+
+	return FoundObjects[0] ? FoundObjects[0] : ResolutionContext;
+}
+
 FGuid FSequencerUtilities::CreateBinding(TSharedRef<ISequencer> Sequencer, UObject& InObject, const UE::Sequencer::FCreateBindingParams& InParams)
 {
+	return CreateOrReplaceBinding(Sequencer.ToSharedPtr(), Sequencer->GetFocusedMovieSceneSequence(), &InObject, InParams);
+}
+
+FGuid FSequencerUtilities::CreateOrReplaceBinding(TSharedRef<ISequencer> Sequencer, UObject* InObject, const UE::Sequencer::FCreateBindingParams& InParams)
+{
+	return CreateOrReplaceBinding(Sequencer.ToSharedPtr(), Sequencer->GetFocusedMovieSceneSequence(), InObject, InParams);
+}
+
+FGuid FSequencerUtilities::CreateOrReplaceBinding(TSharedPtr<ISequencer> Sequencer, UMovieSceneSequence* OwnerSequence, UObject* InObject, const UE::Sequencer::FCreateBindingParams& InParams)
+{
+	if (!OwnerSequence)
+		return FGuid();
+	
 	const FScopedTransaction Transaction(LOCTEXT("CreateBinding", "Create New Binding"));
 
-	UMovieSceneSequence* OwnerSequence = Sequencer->GetFocusedMovieSceneSequence();
 	UMovieScene* OwnerMovieScene = OwnerSequence->GetMovieScene();
 
 	OwnerSequence->Modify();
@@ -3252,16 +3472,16 @@ FGuid FSequencerUtilities::CreateBinding(TSharedRef<ISequencer> Sequencer, UObje
 	FMovieSceneBindingReferences* BindingReferences = OwnerSequence->GetBindingReferences();
 	if (BindingReferences)
 	{
-		BindingGuid = CreateGenericBinding(Sequencer, &InObject, BindingReferences, InParams);
+		BindingGuid = CreateGenericBinding(Sequencer, OwnerSequence, InObject, BindingReferences, InParams);
 	}
-	else if (InParams.bSpawnable)
+	else if (Sequencer && InParams.bSpawnable && InObject)
 	{
 		// Create an old-style spawnable
-		BindingGuid = MakeNewSpawnable(Sequencer, InObject, InParams.ActorFactory.Get());
+		BindingGuid = MakeNewSpawnable(Sequencer.ToSharedRef(), *InObject, InParams.ActorFactory.Get());
 	}
-	else
+	else if (Sequencer && InObject)
 	{
-		BindingGuid = CreateImplementationDefinedBinding(Sequencer, InObject, InParams);
+		BindingGuid = CreateImplementationDefinedBinding(Sequencer.ToSharedRef(), *InObject, InParams);
 	}
 
 	if (!BindingGuid.IsValid())
@@ -3308,11 +3528,14 @@ FGuid FSequencerUtilities::CreateBinding(TSharedRef<ISequencer> Sequencer, UObje
 			DestinationFolder->AddChildObjectBinding(RootObjectGuid);
 		}
 	}
-
-	Sequencer->OnAddBinding(BindingGuid, OwnerMovieScene);
+	if (Sequencer)
+	{
+		Sequencer->OnAddBinding(BindingGuid, OwnerMovieScene);
+	}
 	return BindingGuid;
 }
 
+
 void FSequencerUtilities::UpdateBindingIDs(TSharedRef<ISequencer> Sequencer, FGuid OldGuid, FGuid NewGuid)
 {
 	UMovieSceneCompiledDataManager* CompiledDataManager = FindObject<UMovieSceneCompiledDataManager>(GetTransientPackage(), TEXT("SequencerCompiledDataManager"));
@@ -3543,97 +3766,132 @@ FGuid FSequencerUtilities::AssignActor(TSharedRef<ISequencer> Sequencer, AActor*
 
 void FSequencerUtilities::AddActorsToBinding(TSharedRef<ISequencer> Sequencer, const TArray<AActor*>& Actors, const FMovieSceneBindingProxy& ObjectBinding)
 {
-	if (!Actors.Num())
-	{
-		return;
-	}
+	TArray<UObject*> ObjectsToAdd;
+	Algo::Transform(Actors, ObjectsToAdd, [](AActor* Actor) { return Actor;});
+	AddObjectsToBinding(Sequencer, ObjectsToAdd, ObjectBinding, Sequencer->GetPlaybackContext());
+}
 
-	UMovieSceneSequence* Sequence = Sequencer->GetFocusedMovieSceneSequence();
+void FSequencerUtilities::AddObjectsToBinding(TSharedRef<ISequencer> InSequencer, const TArray<UObject*>& InObjectsToAdd, const FMovieSceneBindingProxy& InObjectBinding, UObject* InResolutionContext)
+{
+	UMovieSceneSequence* Sequence = InObjectBinding.Sequence;
 	if (!Sequence)
 	{
 		return;
 	}
 
-	UMovieScene* MovieScene = Sequence->GetMovieScene();
-	if (!MovieScene)
+	UMovieScene* const MovieScene = Sequence->GetMovieScene();
+	if (!MovieScene || InObjectsToAdd.IsEmpty())
 	{
 		return;
 	}
 
-	UClass* ActorClass = nullptr;
-	int32 NumRuntimeObjects = 0;
+	UClass* ObjectClass = nullptr;
+	int32 ValidObjectCount = 0;
 
-	FGuid Guid = ObjectBinding.BindingID;
-	TArrayView<TWeakObjectPtr<>> ObjectsInCurrentSequence = Sequencer->FindObjectsInCurrentSequence(Guid);
+	FGuid Guid = InObjectBinding.BindingID;
+
+	TArrayView<TWeakObjectPtr<>> ObjectsInCurrentSequence = InSequencer->FindObjectsInCurrentSequence(Guid);
 
 	for (TWeakObjectPtr<> Ptr : ObjectsInCurrentSequence)
 	{
-		if (const AActor* Actor = Cast<AActor>(Ptr.Get()))
+		if (const UObject* Object = Cast<AActor>(Ptr.Get()))
 		{
-			ActorClass = Actor->GetClass();
-			++NumRuntimeObjects;
+			ObjectClass = Object->GetClass();
+			++ValidObjectCount;
 		}
 	}
 
-	FScopedTransaction AddSelectedToBinding(LOCTEXT("AddSelectedToBinding", "Add Selected to Binding"));
-
 	Sequence->Modify();
 	MovieScene->Modify();
 
-	// Bind objects
-	int32 NumObjectsAdded = 0;
-	for (AActor* ActorToAdd : Actors)
+	TArray<UObject*> AddedObjects;
+	AddedObjects.Reserve(InObjectsToAdd.Num());
+
+	for (UObject* ObjectToAdd : InObjectsToAdd)
 	{
-		if (!ObjectsInCurrentSequence.Contains(ActorToAdd))
+		// Skip invalid objects or objects already in the sequence
+		if (!ObjectToAdd || ObjectsInCurrentSequence.Contains(ObjectToAdd))
 		{
-			if (ActorClass == nullptr || UClass::FindCommonBase(ActorToAdd->GetClass(), ActorClass) != nullptr)
-			{
-				if (ActorClass == nullptr)
-				{
-					ActorClass = ActorToAdd->GetClass();
-				}
+			continue;
+		}
 
-				ActorToAdd->Modify();
-				if (!MovieScene->FindPossessable(Guid)->BindSpawnableObject(Sequencer->GetFocusedTemplateID(), ActorToAdd, Sequencer->GetSharedPlaybackState()))
-				{
-					Sequence->BindPossessableObject(Guid, *ActorToAdd, Sequencer->GetPlaybackContext());
-				}
-				++NumObjectsAdded;
+		// Skip if the object has no common class with the objects already in the binding
+		if (ObjectClass && !UClass::FindCommonBase(ObjectToAdd->GetClass(), ObjectClass))
+		{
+			continue;
+		}
+
+		// if no objects are in the binding, set the class to this object's
+		if (!ObjectClass)
+		{
+			ObjectClass = ObjectToAdd->GetClass();
+		}
+
+		FMovieScenePossessable* Possessable = MovieScene->FindPossessable(Guid);
+		if (!ensureAlways(Possessable))
+		{
+			continue;
+		}
+
+		ObjectToAdd->Modify();
+		if (!Possessable->BindSpawnableObject(InSequencer->GetFocusedTemplateID(), ObjectToAdd, InSequencer->GetSharedPlaybackState()))
+		{
+			Sequence->BindPossessableObject(Guid, *ObjectToAdd, InResolutionContext);
+		}
+
+		if (ObjectClass == nullptr || UClass::FindCommonBase(ObjectToAdd->GetClass(), ObjectClass) != nullptr)
+		{
+			if (ObjectClass == nullptr)
+			{
+				ObjectClass = ObjectToAdd->GetClass();
 			}
-			else
+
+			ObjectToAdd->Modify();
+			if (!MovieScene->FindPossessable(Guid)->BindSpawnableObject(InSequencer->GetFocusedTemplateID(), ObjectToAdd, InSequencer->GetSharedPlaybackState()))
 			{
-				const FText NotificationText = FText::Format(LOCTEXT("UnableToAssignObject", "Cannot assign object {0}. Expected class {1}"), FText::FromString(ActorToAdd->GetPathName()), FText::FromString(ActorClass->GetName()));
-				FNotificationInfo Info(NotificationText);
-				Info.ExpireDuration = 3.f;
-				Info.bUseLargeFont = false;
-				FSlateNotificationManager::Get().AddNotification(Info);
+				Sequence->BindPossessableObject(Guid, *ObjectToAdd, InResolutionContext);
 			}
+			AddedObjects.Add(ObjectToAdd);
+		}
+		else
+		{
+			const FText NotificationText = FText::Format(LOCTEXT("UnableToAssignObject", "Cannot assign object {0}. Expected class {1}"), FText::FromString(ObjectToAdd->GetPathName()), FText::FromString(ObjectClass->GetName()));
+			FNotificationInfo Info(NotificationText);
+			Info.ExpireDuration = 3.f;
+			Info.bUseLargeFont = false;
+			FSlateNotificationManager::Get().AddNotification(Info);
 		}
 	}
 
-	// Update label
-	if (NumRuntimeObjects + NumObjectsAdded > 0)
+	// Update Labels
+	if (ValidObjectCount + AddedObjects.Num() > 0)
 	{
-		FMovieScenePossessable* Possessable = MovieScene->FindPossessable(Guid);
-		if (Possessable && ActorClass != nullptr)
+		FMovieScenePossessable* const Possessable = MovieScene->FindPossessable(Guid);
+		if (Possessable && ObjectClass)
 		{
-			if (NumRuntimeObjects + NumObjectsAdded > 1)
+			// If there are multiple objects within the same possessable, name possessable as "ClassName (Count)"
+			if (ValidObjectCount + AddedObjects.Num() > 1)
 			{
-				FString NewLabel = ActorClass->GetName() + FString::Printf(TEXT(" (%d)"), NumRuntimeObjects + NumObjectsAdded);
-				Possessable->SetName(NewLabel);
+				Possessable->SetName(FString::Printf(TEXT("%s (%d)")
+					, *ObjectClass->GetName()
+					, ValidObjectCount + AddedObjects.Num()));
 			}
-			else if (NumObjectsAdded > 0 && Actors.Num() > 0)
+			else if (!AddedObjects.IsEmpty())
 			{
-				Possessable->SetName(Actors[0]->GetActorLabel());
+				FString PossessableName = AddedObjects[0]->GetName();
+				if (AActor* const Actor = Cast<AActor>(AddedObjects[0]))
+				{
+					PossessableName = Actor->GetActorLabel();
+				}
+				Possessable->SetName(PossessableName);
 			}
-
-			Possessable->SetPossessedObjectClass(ActorClass);
+			Possessable->SetPossessedObjectClass(ObjectClass);
 		}
 	}
 
-	Sequencer->RestorePreAnimatedState();
+	InSequencer->RestorePreAnimatedState();
 
-	Sequencer->NotifyMovieSceneDataChanged(EMovieSceneDataChangeType::MovieSceneStructureItemsChanged);
+	InSequencer->NotifyMovieSceneDataChanged(EMovieSceneDataChangeType::MovieSceneStructureItemsChanged);
 }
 
 void FSequencerUtilities::ReplaceBindingWithActors(TSharedRef<ISequencer> Sequencer, const TArray<AActor*>& Actors, const FMovieSceneBindingProxy& ObjectBinding)
@@ -3919,6 +4177,74 @@ void FSequencerUtilities::HandleTemplateActorClassPicked(UClass* ChosenClass, TS
 	}
 }
 
+bool FSequencerUtilities::CanConvertToPossessable(TSharedRef<ISequencer> Sequencer, FGuid BindingGuid, int32 BindingIndex)
+{
+	UMovieSceneSequence* Sequence = Sequencer->GetFocusedMovieSceneSequence();
+	UMovieScene* MovieScene = Sequence ? Sequence->GetMovieScene() : nullptr;
+	if (!MovieScene)
+	{
+		return false;
+	}
+
+	if (FMovieSceneSpawnable* Spawnable = MovieScene->FindSpawnable(BindingGuid))
+	{
+		return true;
+	}
+	else if (const FMovieSceneBindingReferences* BindingReferences = Sequence->GetBindingReferences())
+	{
+		TArrayView<const FMovieSceneBindingReference> BindingReferencesList = BindingReferences->GetReferences(BindingGuid);
+		if (BindingReferencesList.IsValidIndex(BindingIndex) && BindingReferencesList[BindingIndex].CustomBinding != nullptr)
+		{
+			return true;
+		}
+	}
+	return false;
+}
+
+bool FSequencerUtilities::CanConvertToCustomBinding(TSharedRef<ISequencer> Sequencer, FGuid BindingGuid, TSubclassOf<UMovieSceneCustomBinding> CustomBindingType, int32 BindingIndex)
+{
+	UMovieSceneSequence* Sequence = Sequencer->GetFocusedMovieSceneSequence();
+	UMovieScene* MovieScene = Sequence ? Sequence->GetMovieScene() : nullptr;
+	if (!MovieScene)
+	{
+		return false;
+	}
+	if (FMovieSceneSpawnable* Spawnable = MovieScene->FindSpawnable(BindingGuid))
+	{
+		if (UObject* CurrentBoundObject = Spawnable->GetObjectTemplate())
+		{
+			return CustomBindingType && CustomBindingType->GetDefaultObject<UMovieSceneCustomBinding>()->SupportsBindingCreationFromObject(CurrentBoundObject);
+		}
+	}
+	else if (FMovieScenePossessable* Possessable = MovieScene->FindPossessable(BindingGuid))
+	{
+		if (const FMovieSceneBindingReferences* BindingReferences = Sequence->GetBindingReferences())
+		{
+			UObject* ResolutionContext = MovieSceneHelpers::GetResolutionContext(Sequence, BindingGuid, Sequencer->GetFocusedTemplateID(), Sequencer->GetSharedPlaybackState());
+
+			TArrayView<const FMovieSceneBindingReference> BindingReferencesList = Sequencer->GetFocusedMovieSceneSequence()->GetBindingReferences()->GetReferences(BindingGuid);
+
+			if (const FMovieSceneBindingReference* CurrentBindingReference = BindingReferences->GetReference(BindingGuid, BindingIndex))
+			{
+				UE::UniversalObjectLocator::FResolveParams LocatorResolveParams(ResolutionContext);
+				FMovieSceneBindingResolveParams BindingResolveParams{ Sequence, BindingGuid, Sequencer->GetFocusedTemplateID(), ResolutionContext };
+				if (UObject* CurrentBoundObject = BindingReferences->ResolveSingleBinding(BindingResolveParams, BindingIndex, LocatorResolveParams, Sequencer->GetSharedPlaybackState()))
+				{
+					if (CustomBindingType
+						&& (!CurrentBindingReference->CustomBinding
+							|| CurrentBindingReference->CustomBinding->GetClass() != CustomBindingType)
+						&& CustomBindingType->GetDefaultObject<UMovieSceneCustomBinding>()->SupportsConversionFromBinding(*CurrentBindingReference, CurrentBoundObject))
+					{
+						return true;
+					}
+				}
+			}
+		}
+	}
+	return false;
+}
+
+
 void FSequencerUtilities::AddConvertBindingMenu(FMenuBuilder& MenuBuilder, TSharedRef<ISequencer> Sequencer, const TArray<FSequencerConvertBindingInfo>& BindingsToConvert, TFunction<void()> OnBindingChanged)
 {
 	UMovieSceneSequence* Sequence = Sequencer->GetFocusedMovieSceneSequence();
@@ -3951,21 +4277,22 @@ void FSequencerUtilities::AddConvertBindingMenu(FMenuBuilder& MenuBuilder, TShar
 			{
 				SlowTask.EnterProgressFrame();
 
-				FMovieScenePossessable* Possessable = DoConvert(BindingInfo.BindingID, BindingInfo.BindingIndex);
-
-				Sequencer->ForceEvaluate();
-
-				for (TWeakObjectPtr<> WeakObject : Sequencer->FindBoundObjects(Possessable->GetGuid(), Sequencer->GetFocusedTemplateID()))
+				if (FMovieScenePossessable* Possessable = DoConvert(BindingInfo.BindingID, BindingInfo.BindingIndex))
 				{
-					if (AActor* PossessedActor = Cast<AActor>(WeakObject.Get()))
+					Sequencer->ForceEvaluate();
+
+					for (TWeakObjectPtr<> WeakObject : Sequencer->FindBoundObjects(Possessable->GetGuid(), Sequencer->GetFocusedTemplateID()))
 					{
-						PossessedActors.Add(PossessedActor);
+						if (AActor* PossessedActor = Cast<AActor>(WeakObject.Get()))
+						{
+							PossessedActors.Add(PossessedActor);
+						}
 					}
-				}
 
-				if (GWarn->ReceivedUserCancel())
-				{
-					break;
+					if (GWarn->ReceivedUserCancel())
+					{
+						break;
+					}
 				}
 			}
 
@@ -3996,20 +4323,8 @@ void FSequencerUtilities::AddConvertBindingMenu(FMenuBuilder& MenuBuilder, TShar
 	};
 
 	// Can convert to possessable
-	if (Algo::AllOf(BindingsToConvert, [&Sequence, &MovieScene](const FSequencerConvertBindingInfo& BindingInfo) {
-		if (FMovieSceneSpawnable* Spawnable = MovieScene->FindSpawnable(BindingInfo.BindingID))
-		{
-			return true;
-		}
-		else if (const FMovieSceneBindingReferences* BindingReferences = Sequence->GetBindingReferences())
-		{
-			TArrayView<const FMovieSceneBindingReference> BindingReferencesList = BindingReferences->GetReferences(BindingInfo.BindingID);
-			if (BindingReferencesList.IsValidIndex(BindingInfo.BindingIndex) && BindingReferencesList[BindingInfo.BindingIndex].CustomBinding != nullptr)
-			{
-				return true;
-			}
-		}
-		return false;
+	if (Algo::AllOf(BindingsToConvert, [Sequencer, &Sequence, &MovieScene](const FSequencerConvertBindingInfo& BindingInfo) {
+		return CanConvertToPossessable(Sequencer, BindingInfo.BindingID, BindingInfo.BindingIndex );
 	}))
 	{
 		MenuBuilder.AddMenuEntry(
@@ -4029,34 +4344,7 @@ void FSequencerUtilities::AddConvertBindingMenu(FMenuBuilder& MenuBuilder, TShar
 	{
 		if (Algo::AllOf(BindingsToConvert, [&Sequence, &MovieScene, &CustomBindingType, Sequencer](const FSequencerConvertBindingInfo& BindingInfo) 
 		{
-			if (FMovieSceneSpawnable* Spawnable = MovieScene->FindSpawnable(BindingInfo.BindingID))
-			{
-				if (UObject* CurrentBoundObject = Spawnable->GetObjectTemplate())
-				{
-					return CustomBindingType && CustomBindingType->GetDefaultObject<UMovieSceneCustomBinding>()->SupportsBindingCreationFromObject(CurrentBoundObject);
-				}
-			}
-			else if (const FMovieSceneBindingReferences* BindingReferences = Sequence->GetBindingReferences())
-			{
-				TArrayView<const FMovieSceneBindingReference> BindingReferencesList = Sequencer->GetFocusedMovieSceneSequence()->GetBindingReferences()->GetReferences(BindingInfo.BindingID);
-
-				if (const FMovieSceneBindingReference* CurrentBindingReference = BindingReferences->GetReference(BindingInfo.BindingID, BindingInfo.BindingIndex))
-				{
-					UE::UniversalObjectLocator::FResolveParams LocatorResolveParams(Sequencer->GetSharedPlaybackState()->GetPlaybackContext());
-					FMovieSceneBindingResolveParams BindingResolveParams{ Sequence, BindingInfo.BindingID, Sequencer->GetFocusedTemplateID() };
-					if (UObject* CurrentBoundObject = BindingReferences->ResolveSingleBinding(BindingResolveParams, BindingInfo.BindingIndex, LocatorResolveParams, Sequencer->GetSharedPlaybackState()))
-					{
-						if (CustomBindingType
-							&& (!CurrentBindingReference->CustomBinding
-								|| CurrentBindingReference->CustomBinding->GetClass() != CustomBindingType)
-							&& CustomBindingType->GetDefaultObject<UMovieSceneCustomBinding>()->SupportsConversionFromBinding(*CurrentBindingReference, CurrentBoundObject))
-						{
-							return true;
-						}
-					}
-				}
-			}
-			return false;
+			return CanConvertToCustomBinding(Sequencer, BindingInfo.BindingID, CustomBindingType, BindingInfo.BindingIndex);
 		}))
 		{
 			MenuBuilder.AddMenuEntry(
@@ -4070,6 +4358,4 @@ void FSequencerUtilities::AddConvertBindingMenu(FMenuBuilder& MenuBuilder, TShar
 		}
 	}
 }
-
-
 #undef LOCTEXT_NAMESPACE
diff --git a/Engine/Source/Editor/Sequencer/Public/MVVM/ViewModels/ObjectBindingModel.h b/Engine/Source/Editor/Sequencer/Public/MVVM/ViewModels/ObjectBindingModel.h
--- a/Engine/Source/Editor/Sequencer/Public/MVVM/ViewModels/ObjectBindingModel.h
+++ b/Engine/Source/Editor/Sequencer/Public/MVVM/ViewModels/ObjectBindingModel.h
@@ -23,7 +23,6 @@ class UMovieSceneTrack;
 class FMenuBuilder;
 class FPropertyPath;
 class FStructOnScope;
-struct FMovieSceneDynamicBinding;
 enum class ECheckBoxState : uint8;
 
 namespace UE
@@ -125,9 +124,6 @@ public:
 
 	TSharedRef<SWidget> GetAddTrackMenuContent();
 
-	/** Build a sub-menu for editing the given dynamic binding */
-	void AddDynamicBindingMenu(FMenuBuilder& MenuBuilder, FMovieSceneDynamicBinding& DynamicBinding);
-
 protected:
 
 	/*~ FViewModel interface */
diff --git a/Engine/Source/Editor/Sequencer/Public/SequencerUtilities.h b/Engine/Source/Editor/Sequencer/Public/SequencerUtilities.h
--- a/Engine/Source/Editor/Sequencer/Public/SequencerUtilities.h
+++ b/Engine/Source/Editor/Sequencer/Public/SequencerUtilities.h
@@ -163,6 +163,9 @@ struct FMovieScenePasteBindingsParams
 
 	UPROPERTY(EditAnywhere, BlueprintReadWrite, Category = "Movie Scene")
 	bool bDuplicateExistingActors;
+
+	UPROPERTY(EditAnywhere, BlueprintReadWrite, Category = "Movie Scene")
+	TMap<FName, TObjectPtr<AActor>> PastedActors;
 };
 
 USTRUCT()
@@ -219,9 +222,11 @@ struct SEQUENCER_API FSequencerUtilities
 	static TArray<FMovieSceneSpawnable*> ConvertToSpawnable(TSharedRef<ISequencer> Sequencer, FGuid PossessableGuid);
 	
 	/** Convert the requested object binding and object binding index to a possessable */
+	static bool CanConvertToPossessable(TSharedRef<ISequencer> Sequencer, FGuid BindingGuid, int32 BindingIndex = 0);
 	static FMovieScenePossessable* ConvertToPossessable(TSharedRef<ISequencer> Sequencer, FGuid BindingGuid, int32 BindingIndex=0);
 
 	/** Convert the selected object binding and object binding index to a custom binding of the chosen type.*/
+	static bool CanConvertToCustomBinding(TSharedRef<ISequencer> Sequencer, FGuid BindingGuid, TSubclassOf<UMovieSceneCustomBinding> CustomBindingType, int32 BindingIndex = 0);
 	static FMovieScenePossessable* ConvertToCustomBinding(TSharedRef<ISequencer> Sequencer, FGuid BindingGuid, TSubclassOf<UMovieSceneCustomBinding> CustomBindingType, int32 BindingIndex = 0);
 
 	/** Copy/paste folders */
@@ -241,15 +246,28 @@ struct SEQUENCER_API FSequencerUtilities
 
 	/** Copy/paste object bindings */
 	static void CopyBindings(TSharedRef<ISequencer> Sequencer, const TArray<FMovieSceneBindingProxy>& Bindings, const TArray<UMovieSceneFolder*>& InFolders, FString& ExportedText);
+	static void CopyBindings(TSharedRef<ISequencer> Sequencer, const TArray<FMovieSceneBindingProxy>& Bindings, const TArray<UMovieSceneFolder*>& InFolders, FOutputDevice& Ar);
 	static bool PasteBindings(const FString& TextToImport, TSharedRef<ISequencer> Sequencer, FMovieScenePasteBindingsParams PasteBindingsParams, TArray<FMovieSceneBindingProxy>& OutBindings, TArray<FNotificationInfo>& OutErrors);
 	static bool CanPasteBindings(TSharedRef<ISequencer> Sequencer, const FString& TextToImport);
 	static TArray<FString> GetPasteBindingsObjectNames(TSharedRef<ISequencer> Sequencer, const FString& TextToImport);
+	/**
+	 * Recursively finds the most appropriate Resolution Context for a given Parent Guid of a Possessable
+	 */
+
+	static UObject* FindResolutionContext(TSharedRef<ISequencer> Sequencer
+		, UMovieSceneSequence& InSequence
+		, UMovieScene& InMovieScene
+		, const FGuid& InParentGuid
+		, UObject* InPlaybackContext);
 
 	/** Utility functions for managing bindings */
 	static FGuid CreateBinding(TSharedRef<ISequencer> Sequencer, UObject& InObject, const UE::Sequencer::FCreateBindingParams& Params = UE::Sequencer::FCreateBindingParams());
+	static FGuid CreateOrReplaceBinding(TSharedRef<ISequencer> Sequencer, UObject* Object, const UE::Sequencer::FCreateBindingParams& Params = UE::Sequencer::FCreateBindingParams());
+	static FGuid CreateOrReplaceBinding(TSharedPtr<ISequencer> Sequencer, UMovieSceneSequence* Sequence, UObject* Object, const UE::Sequencer::FCreateBindingParams& Params = UE::Sequencer::FCreateBindingParams());
 	static void UpdateBindingIDs(TSharedRef<ISequencer> Sequencer, FGuid OldGuid, FGuid NewGuid);
 	static FGuid AssignActor(TSharedRef<ISequencer> Sequencer, AActor* Actor, FGuid InObjectBinding);
 	static void AddActorsToBinding(TSharedRef<ISequencer> Sequencer, const TArray<AActor*>& Actors, const FMovieSceneBindingProxy& ObjectBinding);
+	static void AddObjectsToBinding(TSharedRef<ISequencer> Sequencer, const TArray<UObject*>& Objects, const FMovieSceneBindingProxy& ObjectBinding, UObject* ResolutionContext);
 	static void ReplaceBindingWithActors(TSharedRef<ISequencer> Sequencer, const TArray<AActor*>& Actors, const FMovieSceneBindingProxy& ObjectBinding);
 	static void RemoveActorsFromBinding(TSharedRef<ISequencer> Sequencer, const TArray<AActor*>& Actors, const FMovieSceneBindingProxy& ObjectBinding);
 
