[Backout] - CL34065553

diff --git a/Engine/Source/Editor/Sequencer/Private/Sequencer.h b/Engine/Source/Editor/Sequencer/Private/Sequencer.h
--- a/Engine/Source/Editor/Sequencer/Private/Sequencer.h
+++ b/Engine/Source/Editor/Sequencer/Private/Sequencer.h
@@ -1488,5 +1488,5 @@ private:
 	uint32 LastKnownStateSerial = 0;
 
 	/* Cached list of supported custom binding types*/
-	TArray<const TSubclassOf<UMovieSceneCustomBinding>> SupportedCustomBindingTypes;
+	TArray<TSubclassOf<UMovieSceneCustomBinding>> SupportedCustomBindingTypes;
 };
diff --git a/Engine/Source/Editor/Sequencer/Private/SequencerUtilities.cpp b/Engine/Source/Editor/Sequencer/Private/SequencerUtilities.cpp
--- a/Engine/Source/Editor/Sequencer/Private/SequencerUtilities.cpp
+++ b/Engine/Source/Editor/Sequencer/Private/SequencerUtilities.cpp
@@ -74,8 +74,6 @@
 #include "MVVM/Selection/Selection.h"
 #include "UnrealEdGlobals.h"
 #include "Misc/FeedbackContext.h"
-#include "UObject/UObjectIterator.h"
-#include "ObjectTools.h"
 
 #define LOCTEXT_NAMESPACE "FSequencerUtilities"
 
@@ -761,6 +759,8 @@ FGuid FSequencerUtilities::CreateCamera(TSharedRef<ISequencer> Sequencer, const
 	OutActor->SetActorRotation(GCurrentLevelEditingViewportClient->GetViewRotation());
 	//OutActor->CameraComponent->FieldOfView = ViewportClient->ViewFOV; //@todo set the focal length from this field of view
 
+	FMovieSceneSpawnable* Spawnable = nullptr;
+
 	if (bSpawnable)
 	{
 		FString NewName = MovieSceneHelpers::MakeUniqueSpawnableName(MovieScene, FName::NameToDisplayString(ACineCameraActor::StaticClass()->GetFName().ToString(), false));
@@ -1098,24 +1098,11 @@ FMovieScenePossessable* FSequencerUtilities::ConvertToPossessable(TSharedRef<ISe
 		return CreatedPossessable;
 	}
 
-	FMovieScenePossessable* ExistingPossessable = MovieScene->FindPossessable(BindingGuid);
-	if (ExistingPossessable)
-	{
-		if (const FMovieSceneBindingReference* ExistingReference = Sequence->GetBindingReferences()->GetReference(BindingGuid, BindingIndex))
-		{
-			if (!ExistingReference->CustomBinding)
-			{
-				// Already a possessable, just return
-				return ExistingPossessable;
-			}
-		}
-	}
-
 	UObject* ObjectToConvert = MovieSceneHelpers::GetSingleBoundObject(Sequence, BindingGuid, Sequencer->GetSharedPlaybackState(), BindingIndex);
 
 	// If we have an old-style spawnable, use the template as the object to convert instead.
 	bool bConvertFromSpawnable = MovieSceneHelpers::IsBoundToSpawnable(Sequence, BindingGuid, Sequencer->GetSharedPlaybackState(), BindingIndex);
-	if (bConvertFromSpawnable && MovieSceneHelpers::SupportsObjectTemplate(Sequence, BindingGuid, Sequencer->GetSharedPlaybackState(), BindingIndex))
+	if (bConvertFromSpawnable)
 	{
 		ObjectToConvert = MovieSceneHelpers::GetObjectTemplate(Sequence, BindingGuid, Sequencer->GetSharedPlaybackState(), BindingIndex);
 	}
@@ -1159,6 +1146,7 @@ FMovieScenePossessable* FSequencerUtilities::ConvertToPossessable(TSharedRef<ISe
 	}
 
 	FMovieSceneSpawnable* Spawnable = MovieScene->FindSpawnable(BindingGuid);
+	FMovieScenePossessable* ExistingPossessable = MovieScene->FindPossessable(BindingGuid);
 	// TODO: How to convert to possessable of non-actor type? Presumably we need to generalize the 'creation' step here for now.
 	UObject* PossessedObject = nullptr;
 	if (Actor)
@@ -1384,7 +1372,6 @@ FMovieScenePossessable* FSequencerUtilities::ConvertToCustomBinding(TSharedRef<I
 	CreateBindingParams.BindingIndex = BindingIndex;
 	CreateBindingParams.BindingNameOverride = NewCustomBinding->GetDesiredBindingName();
 	CreateBindingParams.CustomBinding = NewCustomBinding;
-	CreateBindingParams.bSetupDefaults = false;
 
 	// Create or replace the binding
 	FGuid NewPossessableGuid = CreateBinding(Sequencer, *ObjectToConvert, CreateBindingParams);
@@ -2106,11 +2093,12 @@ bool FSequencerUtilities::CanPasteSections(const FString& TextToImport)
  *
  */
 
-void ExportObjectBindingsToText(const TArray<UMovieSceneCopyableBinding*>& ObjectsToExport, FOutputDevice& Archive, TSharedRef<UE::MovieScene::FSharedPlaybackState> SharedPlaybackState)
+void ExportObjectBindingsToText(const TArray<UMovieSceneCopyableBinding*>& ObjectsToExport, FString& ExportedText, TSharedRef<UE::MovieScene::FSharedPlaybackState> SharedPlaybackState)
 {
 	// Clear the mark state for saving.
 	UnMarkAllObjects(EObjectMark(OBJECTMARK_TagExp | OBJECTMARK_TagImp));
 
+	FStringOutputDevice Archive;
 	const FExportObjectInnerContext Context;
 
 	// Export each of the selected nodes
@@ -2130,7 +2118,6 @@ void ExportObjectBindingsToText(const TArray<UMovieSceneCopyableBinding*>& Objec
 
 		TArray<UObject*, TInlineAllocator<1>> OldObjectTemplates;
 		TArray<UMovieSceneCustomBinding*, TInlineAllocator<1>> OldCustomBindings;
-		TMap<int32, UMovieSceneSpawnableBindingBase*> OldPreviewSpawnables;
 		if (ObjectToExport->Spawnable.GetGuid().IsValid())
 		{
 			OldObjectTemplates.Add(ObjectToExport->Spawnable.GetObjectTemplate());
@@ -2138,11 +2125,9 @@ void ExportObjectBindingsToText(const TArray<UMovieSceneCopyableBinding*>& Objec
 		}
 		else
 		{
-			
-			for (int32 CustomBindingIndex = 0; CustomBindingIndex < ObjectToExport->CustomBindings.Num(); ++CustomBindingIndex)
+			for (UMovieSceneCustomBinding* CustomBinding : ObjectToExport->CustomBindings)
 			{
-				UMovieSceneCustomBinding* CustomBinding = ObjectToExport->CustomBindings[CustomBindingIndex];
-				if (UMovieSceneSpawnableBindingBase* SpawnableBinding = Cast<UMovieSceneSpawnableBindingBase>(CustomBinding))
+				if (UMovieSceneSpawnableBindingBase* SpawnableBinding = CustomBinding->AsSpawnable(SharedPlaybackState))
 				{
 					if (SpawnableBinding->SupportsObjectTemplates())
 					{
@@ -2150,17 +2135,6 @@ void ExportObjectBindingsToText(const TArray<UMovieSceneCopyableBinding*>& Objec
 						SpawnableBinding->SetObjectTemplate(nullptr);
 					}
 				}
-				else if (UMovieSceneReplaceableBindingBase * ReplaceableBinding = Cast<UMovieSceneReplaceableBindingBase>(CustomBinding))
-				{
-					// Prevent inner references here during export
-					if (ReplaceableBinding->PreviewSpawnable)
-					{ 
-						OldPreviewSpawnables.Add(CustomBindingIndex, ReplaceableBinding->PreviewSpawnable);
-						// The Preview Spawnable is next in the CustomBindings list
-						ObjectToExport->PreviewSpawnableBindings.Add(CustomBindingIndex+1);
-						ReplaceableBinding->PreviewSpawnable = nullptr;
-					}
-				}
 				OldCustomBindings.Add(CustomBinding);
 			}
 			ObjectToExport->CustomBindings.Empty();
@@ -2181,15 +2155,6 @@ void ExportObjectBindingsToText(const TArray<UMovieSceneCopyableBinding*>& Objec
 			UExporter::ExportToOutputDevice(&Context, CustomBinding, nullptr, Archive, TEXT("copy"), 0, PPF_ExportsNotFullyQualified | PPF_Copy | PPF_Delimited);
 		}
 
-		// Restore replaceable references now that we've exported them
-		for (const TPair<int32, UMovieSceneSpawnableBindingBase*>& PreviewSpawnable : OldPreviewSpawnables)
-		{
-			if (UMovieSceneReplaceableBindingBase* ReplaceableBinding = Cast<UMovieSceneReplaceableBindingBase>(ObjectToExport->CustomBindings[PreviewSpawnable.Key].Get()))
-			{
-				ReplaceableBinding->PreviewSpawnable = PreviewSpawnable.Value;
-			}
-		}
-
 		int32 ObjectTemplateIndex = 0;
 		if (ObjectToExport->Spawnable.GetGuid().IsValid())
 		{
@@ -2201,8 +2166,7 @@ void ExportObjectBindingsToText(const TArray<UMovieSceneCopyableBinding*>& Objec
 			{
 				if (UMovieSceneSpawnableBindingBase* SpawnableBinding = CustomBinding->AsSpawnable(SharedPlaybackState))
 				{
-					// Ignore bindings with their template already set, which is possible for Replaceables since they'll show up twice in the list.
-					if (SpawnableBinding->SupportsObjectTemplates() && !SpawnableBinding->GetObjectTemplate())
+					if (SpawnableBinding->SupportsObjectTemplates())
 					{
 						SpawnableBinding->SetObjectTemplate(OldObjectTemplates[ObjectTemplateIndex++]);
 					}
@@ -2216,16 +2180,11 @@ void ExportObjectBindingsToText(const TArray<UMovieSceneCopyableBinding*>& Objec
 			UExporter::ExportToOutputDevice(&Context, ObjectTemplate, nullptr, Archive, TEXT("copy"), 0, PPF_ExportsNotFullyQualified | PPF_Copy | PPF_Delimited);
 		}
 	}
-}
 
-void FSequencerUtilities::CopyBindings(TSharedRef<ISequencer> Sequencer, const TArray<FMovieSceneBindingProxy>& Bindings, const TArray<UMovieSceneFolder*>& InFolders, FString& ExportedText)
-{
-	FStringOutputDevice Archive;
-	CopyBindings(Sequencer, Bindings, InFolders, (FOutputDevice&)Archive);
 	ExportedText = Archive;
 }
 
-void FSequencerUtilities::CopyBindings(TSharedRef<ISequencer> Sequencer, const TArray<FMovieSceneBindingProxy>& Bindings, const TArray<UMovieSceneFolder*>& InFolders, FOutputDevice& Ar)
+void FSequencerUtilities::CopyBindings(TSharedRef<ISequencer> Sequencer, const TArray<FMovieSceneBindingProxy>& Bindings, const TArray<UMovieSceneFolder*>& InFolders, FString& ExportedText)
 {
 	TArray<UMovieSceneCopyableBinding*> Objects;
 	for (const FMovieSceneBindingProxy& ObjectBinding : Bindings)
@@ -2242,24 +2201,24 @@ void FSequencerUtilities::CopyBindings(TSharedRef<ISequencer> Sequencer, const T
 		FMovieScenePossessable* Possessable = MovieScene->FindPossessable(ObjectBinding.BindingID);
 		if (Possessable)
 		{
-			UObject* ResolutionContext = FindResolutionContext(Sequencer
-				, *MovieScene->GetTypedOuter<UMovieSceneSequence>()
-				, *MovieScene
-				, Possessable->GetParent()
-				, Sequencer->GetPlaybackContext());
-
 			CopyableBinding->Possessable = *Possessable;
 
+			// Store the names of the bound objects so that they can be found on paste
+			for (TWeakObjectPtr<> RuntimeObject : Sequencer->FindBoundObjects(CopyableBinding->Possessable.GetGuid(), Sequencer->GetFocusedTemplateID()))
+			{
+				if (AActor* Actor = Cast<AActor>(RuntimeObject.Get()))
+				{
+					CopyableBinding->BoundObjectNames.Add(Actor->GetPathName());
+				}
+			}
+
 			// Store any custom bindings
 			if (FMovieSceneBindingReferences* BindingReferences = ObjectBinding.Sequence->GetBindingReferences())
 			{
-				int32 BindingIndex = 0;
 				for (const FMovieSceneBindingReference& BindingReference : BindingReferences->GetReferences(ObjectBinding.BindingID))
 				{
 					if (UMovieSceneCustomBinding* CustomBinding = BindingReference.CustomBinding)
 					{
-						CopyableBinding->CustomBindings.Add(CustomBinding);
-
 						if (UMovieSceneSpawnableBindingBase* SpawnableBinding = CustomBinding->AsSpawnable(Sequencer->GetSharedPlaybackState()))
 						{
 							if (SpawnableBinding->SupportsObjectTemplates())
@@ -2268,28 +2227,9 @@ void FSequencerUtilities::CopyBindings(TSharedRef<ISequencer> Sequencer, const T
 								// Object Templates are re-created on paste (based on these templates) with the correct ownership set up.
 								CopyableBinding->SpawnableObjectTemplates.Add(SpawnableBinding->GetObjectTemplate());
 							}
-
-							// This is the inner spawnable of a replaceable and is always placed after the replaceable in the list
-							if (SpawnableBinding != CustomBinding)
-							{
-								CopyableBinding->CustomBindings.Add(SpawnableBinding);
-							}
 						}
+						CopyableBinding->CustomBindings.Add(CustomBinding);
 					}
-					else if (UObject* RuntimeObject = MovieSceneHelpers::GetSingleBoundObject(ObjectBinding.Sequence, ObjectBinding.BindingID, Sequencer->GetSharedPlaybackState(), BindingIndex))
-					{
-						CopyableBinding->BoundObjectNames.Add(RuntimeObject->GetPathName(ResolutionContext));
-					}
-					BindingIndex++;
-				}
-			}
-			else
-			{
-				// Store the names of the bound objects so that they can be found on paste
-
-				for (TWeakObjectPtr<> RuntimeObject : Sequencer->FindBoundObjects(CopyableBinding->Possessable.GetGuid(), Sequencer->GetFocusedTemplateID()))
-				{
-					CopyableBinding->BoundObjectNames.Add(RuntimeObject->GetPathName(ResolutionContext));
 				}
 			}
 		}
@@ -2340,7 +2280,7 @@ void FSequencerUtilities::CopyBindings(TSharedRef<ISequencer> Sequencer, const T
 		}
 	}
 
-	ExportObjectBindingsToText(Objects, /*out*/ Ar, Sequencer->GetSharedPlaybackState());
+	ExportObjectBindingsToText(Objects, /*out*/ ExportedText, Sequencer->GetSharedPlaybackState());
 }
 
 class FObjectBindingTextFactory : public FCustomizableTextObjectFactory
@@ -2431,7 +2371,7 @@ void ImportObjectBindingsFromText(ISequencer& InSequencer, const FString& TextTo
 			{
 				for (UMovieSceneCustomBinding* CustomBinding : ImportedObject->CustomBindings)
 				{
-					if (UMovieSceneSpawnableBindingBase* SpawnableBinding = Cast<UMovieSceneSpawnableBindingBase>(CustomBinding))
+					if (UMovieSceneSpawnableBindingBase* SpawnableBinding = CustomBinding->AsSpawnable(InSequencer.GetSharedPlaybackState()))
 					{
 						if (SpawnableBinding->SupportsObjectTemplates())
 						{
@@ -2447,7 +2387,7 @@ void ImportObjectBindingsFromText(ISequencer& InSequencer, const FString& TextTo
 	TempPackage->RemoveFromRoot();
 }
 
-FGuid TryCreateCustomBinding(TSharedPtr<ISequencer> Sequencer, UObject* CustomBindingObject, AActor* FactoryCreatedActor, FMovieSceneBindingReferences* BindingReferences, const UE::Sequencer::FCreateBindingParams& InParams, UMovieScene* OwnerMovieScene, bool bSpawnable, bool bReplaceable)
+FGuid TryCreateCustomBinding(TSharedRef<ISequencer> Sequencer, UObject* CustomBindingObject, FMovieSceneBindingReferences* BindingReferences, const UE::Sequencer::FCreateBindingParams& InParams, UMovieScene* OwnerMovieScene, bool bSpawnable, bool bReplaceable)
 {
 	UMovieSceneCustomBinding * NewCustomBinding = nullptr;
 	if (InParams.CustomBinding)
@@ -2463,23 +2403,7 @@ FGuid TryCreateCustomBinding(TSharedPtr<ISequencer> Sequencer, UObject* CustomBi
 	}
 	else
 	{
-		TArrayView<const TSubclassOf<UMovieSceneCustomBinding>> PrioritySortedCustomBindingTypes;
-		if (Sequencer)
-		{
-			PrioritySortedCustomBindingTypes = Sequencer->GetSupportedCustomBindingTypes();
-		}
-		else
-		{
-			static TArray<const TSubclassOf<UMovieSceneCustomBinding>> CachedCustomBindingTypes;
-			static bool CustomBindingTypesCached = false;
-			if (!CustomBindingTypesCached)
-			{
-				CustomBindingTypesCached = true;
-				MovieSceneHelpers::GetPrioritySortedCustomBindingTypes(CachedCustomBindingTypes);
-			}
-			PrioritySortedCustomBindingTypes = CachedCustomBindingTypes;
-		}
-
+		TArrayView<const TSubclassOf<UMovieSceneCustomBinding>> PrioritySortedCustomBindingTypes = Sequencer->GetSupportedCustomBindingTypes();
 		for (const TSubclassOf<UMovieSceneCustomBinding>& CustomBindingType : PrioritySortedCustomBindingTypes)
 		{
 			// If 'spawnable' has been passed in, we only want to use children of UMovieSceneSpawnableBindingBase and vice versa
@@ -2496,16 +2420,6 @@ FGuid TryCreateCustomBinding(TSharedPtr<ISequencer> Sequencer, UObject* CustomBi
 							break;
 						}
 					}
-
-					if (!NewCustomBinding && FactoryCreatedActor && CustomBindingCDO->SupportsBindingCreationFromObject(FactoryCreatedActor))
-					{
-						// Create a custom binding from the factory created actor
-						NewCustomBinding = CustomBindingCDO->CreateNewCustomBinding(FactoryCreatedActor, *OwnerMovieScene);
-						if (NewCustomBinding)
-						{
-							break;
-						}
-					}
 				}
 			}
 		}
@@ -2538,36 +2452,33 @@ FGuid TryCreateCustomBinding(TSharedPtr<ISequencer> Sequencer, UObject* CustomBi
 
 		UObject* SpawnedObject = nullptr;
 		// If this is a spawnable or replaceable binding, we need to set up some defaults
-		if (Sequencer)
+		if (NewCustomBinding->WillSpawnObject(Sequencer->GetSharedPlaybackState()))
 		{
-			if (NewCustomBinding->WillSpawnObject(Sequencer->GetSharedPlaybackState()))
-			{
-				// Spawn the object so we can position it correctly, it's going to get spawned anyway since things default to spawned.
-				SpawnedObject = Sequencer->GetSpawnRegister().SpawnObject(NewID, *OwnerMovieScene, Sequencer->GetFocusedTemplateID(), Sequencer->GetSharedPlaybackState(), 0);
+			// Spawn the object so we can position it correctly, it's going to get spawned anyway since things default to spawned.
+			SpawnedObject = Sequencer->GetSpawnRegister().SpawnObject(NewID, *OwnerMovieScene, Sequencer->GetFocusedTemplateID(), Sequencer->GetSharedPlaybackState(), 0);
 
-				if (InParams.bSetupDefaults)
-				{
-					// Allow the binding to set up any necessary defaults
-					NewCustomBinding->SetupDefaults(SpawnedObject, NewID, *OwnerMovieScene, Sequencer->GetSharedPlaybackState());
+			if (InParams.bSetupDefaults)
+			{
+				// Allow the binding to set up any necessary defaults
+				NewCustomBinding->SetupDefaults(SpawnedObject, NewID, *OwnerMovieScene, Sequencer->GetSharedPlaybackState());
 
-					FTransformData TransformData;
-					Sequencer->GetSpawnRegister().SetupDefaultsForSpawnable(SpawnedObject, NewID, TransformData, Sequencer.ToSharedRef(), Sequencer->GetSequencerSettings());
-				}
+				FTransformData TransformData;
+				Sequencer->GetSpawnRegister().SetupDefaultsForSpawnable(SpawnedObject, NewID, TransformData, Sequencer, Sequencer->GetSequencerSettings());
 			}
+		}
 
-			Sequencer->State.Invalidate(NewID, Sequencer->GetFocusedTemplateID());
-			Sequencer->ForceEvaluate();
+		Sequencer->State.Invalidate(NewID, Sequencer->GetFocusedTemplateID());
+		Sequencer->ForceEvaluate();
 
-			// We don't call these events in the case bSetupDefaults is false because they may add tracks.
-			if (InParams.bSetupDefaults)
+		// We don't call these events in the case bSetupDefaults is false because they may add tracks.
+		if (InParams.bSetupDefaults)
+		{
+			if (AActor* Actor = Cast<AActor>(SpawnedObject))
 			{
-				if (AActor* Actor = Cast<AActor>(SpawnedObject))
-				{
-					Sequencer->OnActorAddedToSequencer().Broadcast(Actor, NewID);
-				}
-
-				Sequencer->OnAddBinding(NewID, OwnerMovieScene);
+				Sequencer->OnActorAddedToSequencer().Broadcast(Actor, NewID);
 			}
+
+			Sequencer->OnAddBinding(NewID, OwnerMovieScene);
 		}
 
 		return NewID;
@@ -2575,15 +2486,11 @@ FGuid TryCreateCustomBinding(TSharedPtr<ISequencer> Sequencer, UObject* CustomBi
 	return FGuid();
 }
 
-FGuid CreateGenericBinding(TSharedPtr<ISequencer> Sequencer, UMovieSceneSequence* OwnerSequence, UObject* InObject, FMovieSceneBindingReferences* BindingReferences, const UE::Sequencer::FCreateBindingParams& InParams)
+FGuid CreateGenericBinding(TSharedRef<ISequencer> Sequencer, UObject* InObject, FMovieSceneBindingReferences* BindingReferences, const UE::Sequencer::FCreateBindingParams& InParams)
 {
-	if (!OwnerSequence)
-	{
-		return FGuid();
-	}
-
 	using namespace UE::Sequencer;
 
+	UMovieSceneSequence* OwnerSequence = Sequencer->GetFocusedMovieSceneSequence();
 	UMovieScene* OwnerMovieScene = OwnerSequence->GetMovieScene();
 
 	ISequencerModule& Module = FModuleManager::Get().LoadModuleChecked<ISequencerModule>("Sequencer");
@@ -2595,10 +2502,12 @@ FGuid CreateGenericBinding(TSharedPtr<ISequencer> Sequencer, UMovieSceneSequence
 
 	if (bAllowCustom)
 	{
-		// In addition to the raw object, we also try spawning an actor from an actor factory if relevant, to give the custom binding an option to create from that as well
-		AActor* FactoryCreatedActorInstance = nullptr; 
-		UWorld* World = GCurrentLevelEditingViewportClient ? GCurrentLevelEditingViewportClient->GetWorld() : nullptr;
-		if (InObject && !InObject->IsA<AActor>())
+		NewBindingID = TryCreateCustomBinding(Sequencer, InObject, BindingReferences, InParams, OwnerMovieScene, bSpawnable, bReplaceable);
+		if (NewBindingID.IsValid())
+		{
+			return NewBindingID;
+		}
+		else if (InObject && !InObject->IsA<AActor>())
 		{
 			// If the passed in object is not an actor, see if we can create an Actor from it, and if so, if that Actor type has a custom binding that supports it
 			UActorFactory* FactoryToUse = InParams.ActorFactory ? InParams.ActorFactory.Get() : FActorFactoryAssetProxy::GetFactoryForAssetObject(InObject);
@@ -2608,6 +2517,7 @@ FGuid CreateGenericBinding(TSharedPtr<ISequencer> Sequencer, UMovieSceneSequence
 				FText ErrorMessage;
 				if (FactoryToUse->CanCreateActorFrom(FAssetData(InObject), ErrorMessage))
 				{
+					UWorld* World = GCurrentLevelEditingViewportClient ? GCurrentLevelEditingViewportClient->GetWorld() : nullptr;
 					if (World)
 					{
 						const FName ActorName = MakeUniqueObjectName(World->PersistentLevel, FactoryToUse->NewActorClass->StaticClass(), *InParams.BindingNameOverride);
@@ -2616,28 +2526,29 @@ FGuid CreateGenericBinding(TSharedPtr<ISequencer> Sequencer, UMovieSceneSequence
 						SpawnParams.ObjectFlags = RF_Transient | RF_Transactional;
 						SpawnParams.Name = ActorName;
 
-						FactoryCreatedActorInstance = FactoryToUse->CreateActor(InObject, World->PersistentLevel, FTransform(), SpawnParams);
-						if (FactoryCreatedActorInstance)
+						AActor* Instance = FactoryToUse->CreateActor(InObject, World->PersistentLevel, FTransform(), SpawnParams);
+						if (Instance)
 						{
-							FactoryCreatedActorInstance->SetActorLabel(MovieSceneHelpers::MakeUniqueBindingName(OwnerMovieScene, FName::NameToDisplayString(InObject->GetName(), false)));
-							FactoryCreatedActorInstance->bIsEditorPreviewActor = false;
+							Instance->SetActorLabel(MovieSceneHelpers::MakeUniqueBindingName(OwnerMovieScene, FName::NameToDisplayString(InObject->GetName(), false)));
+							Instance->bIsEditorPreviewActor = false;
+
+							// Can we create a custom binding from this factory created actor? In all cases, we expect the binding creation process to duplicate
+							// the actor, so we should be safe to destroy the actor after this process.
+							NewBindingID = TryCreateCustomBinding(Sequencer, Instance, BindingReferences, InParams, OwnerMovieScene, bSpawnable, bReplaceable);
+
+							const bool bNetForce = false;
+							const bool bShouldModifyLevel = false;
+							World->DestroyActor(Instance, bNetForce, bShouldModifyLevel);
+
+							if (NewBindingID.IsValid())
+							{
+								return NewBindingID;
+							}
 						}
 					}
 				}
 			}
 		}
-		NewBindingID = TryCreateCustomBinding(Sequencer, InObject, FactoryCreatedActorInstance, BindingReferences, InParams, OwnerMovieScene, bSpawnable, bReplaceable);
-		if (FactoryCreatedActorInstance)
-		{
-			const bool bNetForce = false;
-			const bool bShouldModifyLevel = false;
-			World->DestroyActor(FactoryCreatedActorInstance, bNetForce, bShouldModifyLevel);
-			FactoryCreatedActorInstance = nullptr;
-		}
-		if (NewBindingID.IsValid())
-		{
-			return NewBindingID;
-		}
 	}
 
 	if (!InObject)
@@ -2647,9 +2558,9 @@ FGuid CreateGenericBinding(TSharedPtr<ISequencer> Sequencer, UMovieSceneSequence
 
 
 	// If no custom bindings support this object type, but InParams.bSpawnable is true, attempt to make an old-style spawnable.
-	if (bSpawnable && Sequencer)
+	if (bSpawnable)
 	{
-		NewBindingID = FSequencerUtilities::MakeNewSpawnable(Sequencer.ToSharedRef(), *InObject, InParams.ActorFactory);
+		NewBindingID = FSequencerUtilities::MakeNewSpawnable(Sequencer, *InObject, InParams.ActorFactory);
 		if (NewBindingID.IsValid())
 		{
 			return NewBindingID;
@@ -2696,8 +2607,7 @@ FGuid CreateGenericBinding(TSharedPtr<ISequencer> Sequencer, UMovieSceneSequence
 
 	const bool bParentContextsAreSignificant = OwnerSequence->AreParentContextsSignificant();
 
-
-	UObject* Context = Sequencer ? Sequencer->GetPlaybackContext() : nullptr;
+	UObject* Context = Sequencer->GetPlaybackContext();
 
 	FGuid ParentID;
 
@@ -2707,7 +2617,7 @@ FGuid CreateGenericBinding(TSharedPtr<ISequencer> Sequencer, UMovieSceneSequence
 		UObject* CurrentObject = ObjectsToPossess[Index].Key;
 
 		// If we're not purposefully replacing a binding, then check to see if we already have one, and use that
-		if (Sequencer && !InParams.ReplacementGuid.IsValid())
+		if (!InParams.ReplacementGuid.IsValid())
 		{
 			FGuid    ObjectGuid = Sequencer->GetHandleToObject(CurrentObject, false);
 
@@ -2740,7 +2650,7 @@ FGuid CreateGenericBinding(TSharedPtr<ISequencer> Sequencer, UMovieSceneSequence
 		}
 
 		// If we're not trying to replace a binding, and the object is a spawnable, try and bind to that first
-		if (InParams.ReplacementGuid.IsValid() || (Sequencer && !NewPossessable->BindSpawnableObject(Sequencer->GetFocusedTemplateID(), CurrentObject, Sequencer->GetSharedPlaybackState())))
+		if (InParams.ReplacementGuid.IsValid() || !NewPossessable->BindSpawnableObject(Sequencer->GetFocusedTemplateID(), CurrentObject, Sequencer->GetSharedPlaybackState()))
 		{
 			FUniversalObjectLocator Locator;
 			if (!OwnerSequence->MakeLocatorForObject(CurrentObject, Context, Locator) || Locator.IsEmpty())
@@ -2752,10 +2662,7 @@ FGuid CreateGenericBinding(TSharedPtr<ISequencer> Sequencer, UMovieSceneSequence
 			if (InParams.ReplacementGuid.IsValid() && !ParentID.IsValid())
 			{
 				BindingReferences->AddOrReplaceBinding(NewID, MoveTemp(Locator), InParams.BindingIndex);
-				if (Sequencer)
-				{
-					Sequencer->State.Invalidate(NewID, Sequencer->GetFocusedTemplateID());
-				}
+				Sequencer->State.Invalidate(NewID, Sequencer->GetFocusedTemplateID());
 			}
 			else
 			{
@@ -2775,21 +2682,16 @@ FGuid CreateGenericBinding(TSharedPtr<ISequencer> Sequencer, UMovieSceneSequence
 		}
 
 		ParentID = NewID;
-		if (Sequencer)
+
+		if (AActor* Actor = Cast<AActor>(CurrentObject))
 		{
-			if (AActor* Actor = Cast<AActor>(CurrentObject))
-			{
-				Sequencer->OnActorAddedToSequencer().Broadcast(Actor, NewID);
-			}
+			Sequencer->OnActorAddedToSequencer().Broadcast(Actor, NewID);
 		}
 
 		// If this is the last one
 		if (Index == 0)
 		{
-			if (Sequencer)
-			{
-				Sequencer->OnAddBinding(NewID, OwnerMovieScene);
-			}
+			Sequencer->OnAddBinding(NewID, OwnerMovieScene);
 			return NewID;
 		}
 
@@ -2877,47 +2779,33 @@ bool FSequencerUtilities::PasteBindings(const FString& TextToImport, TSharedRef<
 				{
 					if (BindingReferences)
 					{
+						int32 BindingIndex = 0;
 						int32 SpawnableBindingIndex = 0;
-						UMovieSceneCustomBinding* PreviousBinding = nullptr;
-						for(int32 BindingIndex = 0; BindingIndex < CopyableBinding->CustomBindings.Num(); ++BindingIndex)
+						for (UMovieSceneCustomBinding* CustomBinding : CopyableBinding->CustomBindings)
 						{
-							UMovieSceneCustomBinding* CustomBinding = CopyableBinding->CustomBindings[BindingIndex];
 							if (CustomBinding)
 							{
 								UMovieSceneCustomBinding* NewCustomBinding = Cast<UMovieSceneCustomBinding>(StaticDuplicateObject(CustomBinding, MovieScene));
 
 								// Need to re-copy the object template to avoid private object issues
-								if (UMovieSceneSpawnableBindingBase* SpawnableBinding = Cast<UMovieSceneSpawnableBindingBase>(NewCustomBinding))
+								if (UMovieSceneSpawnableBindingBase* SpawnableBinding = NewCustomBinding->AsSpawnable(Sequencer->GetSharedPlaybackState()))
 								{
 									if (CopyableBinding->SpawnableObjectTemplates.IsValidIndex(SpawnableBindingIndex))
 									{
 										UObject* NewObjectTemplate = StaticDuplicateObject(CopyableBinding->SpawnableObjectTemplates[SpawnableBindingIndex++], MovieScene);
 										SpawnableBinding->SetObjectTemplate(NewObjectTemplate);
 									}
-
-									// If this is a preview spawnable, find the just added replaceable and link with that rather than creating a new binding
-									if (CopyableBinding->PreviewSpawnableBindings.Contains(BindingIndex))
-									{
-										if (UMovieSceneReplaceableBindingBase* PreviousReplaceableBinding = Cast<UMovieSceneReplaceableBindingBase>(PreviousBinding))
-										{
-											PreviousReplaceableBinding->PreviewSpawnable = SpawnableBinding;
-											PreviousBinding = NewCustomBinding;
-											continue;
-										}
-									}
 								}
 
 								// This will either add a brand new possessable and binding (if one doesn't exist for that guid), or just add a new binding to that same possessable
 								UE::Sequencer::FCreateBindingParams CreateBindingParams;
 								CreateBindingParams.ReplacementGuid = NewGuid;
-								CreateBindingParams.BindingIndex = BindingIndex;
+								CreateBindingParams.BindingIndex = BindingIndex++;
 								CreateBindingParams.bAllowCustomBinding = true;
 								CreateBindingParams.CustomBinding = NewCustomBinding;
 								CreateBindingParams.bSetupDefaults = false;
 								CreateBindingParams.BindingNameOverride = CopyableBinding->Possessable.GetName();
-								NewGuid = CreateGenericBinding(Sequencer, Sequence, nullptr, BindingReferences, CreateBindingParams);
-
-								PreviousBinding = NewCustomBinding;
+								NewGuid = CreateGenericBinding(Sequencer, nullptr, BindingReferences, CreateBindingParams);
 							}
 						}
 					}
@@ -2973,10 +2861,9 @@ bool FSequencerUtilities::PasteBindings(const FString& TextToImport, TSharedRef<
 					}
 				}
 
-				// Find the objects that this pasted binding should bind to
-				TArray<UObject*> ObjectsToBind;
-				UObject* ResolutionContext = Sequencer->GetPlaybackContext();
-				if (World)
+				// Find the actors that this pasted binding should bind to
+				TArray<AActor*> ActorsToRebind;
+				if (World && CopyableBinding->CustomBindings.IsEmpty())
 				{
 					for (TActorIterator<AActor> ActorItr(World); ActorItr; ++ActorItr)
 					{
@@ -2989,25 +2876,20 @@ bool FSequencerUtilities::PasteBindings(const FString& TextToImport, TSharedRef<
 								continue;
 							}
 
-							ObjectsToBind.Add(Actor);
+							ActorsToRebind.Add(Actor);
 							CopyableBinding->BoundObjectNames.Remove(Actor->GetPathName());
 						}
 					}
 				}
 
-				if (ObjectsToBind.Num() != 0)
+				if (ActorsToRebind.Num() != 0)
 				{
 					if (PasteBindingsParams.bDuplicateExistingActors)
 					{
 						GEditor->SelectNone(false, true);
-						TArray<UObject*> SelectedObjects;
-						for (UObject* ObjectToBind : ObjectsToBind)
+						for (AActor* ActorToRebind : ActorsToRebind)
 						{
-							if (AActor* Actor = Cast<AActor>(ObjectToBind))
-							{
-								GEditor->SelectActor(Actor, true, false, false);
-								SelectedObjects.Add(Actor);
-							}
+							GEditor->SelectActor(ActorToRebind, true, false, false);
 						}
 							
 						// Duplicate the bound actors
@@ -3018,14 +2900,14 @@ bool FSequencerUtilities::PasteBindings(const FString& TextToImport, TSharedRef<
 						// of the clipboard after duplicating the actor
 						FPlatformApplicationMisc::ClipboardCopy(*TextToImport);
 
-						ObjectsToBind.RemoveAll([&SelectedObjects](UObject* Object) { return SelectedObjects.Contains(Object);});
+						ActorsToRebind.Empty();
 						USelection* ActorSelection = GEditor->GetSelectedActors();
 						for (FSelectionIterator Iter(*ActorSelection); Iter; ++Iter)
 						{
 							AActor* Actor = Cast<AActor>(*Iter);
 							if (Actor)
 							{
-								ObjectsToBind.Add(Actor);
+								ActorsToRebind.Add(Actor);
 
 								CopyableBinding->BoundObjectNames.Add(Actor->GetPathName());
 							}
@@ -3033,9 +2915,9 @@ bool FSequencerUtilities::PasteBindings(const FString& TextToImport, TSharedRef<
 					}
 
 					// Bind the actors
-					if (ObjectsToBind.Num())
+					if (ActorsToRebind.Num())
 					{
-						AddObjectsToBinding(Sequencer, ObjectsToBind, FMovieSceneBindingProxy(NewGuid, Sequence), ResolutionContext);
+						AddActorsToBinding(Sequencer, ActorsToRebind, FMovieSceneBindingProxy(NewGuid, Sequence));
 					}
 				}
 
@@ -3115,72 +2997,44 @@ bool FSequencerUtilities::PasteBindings(const FString& TextToImport, TSharedRef<
 		FMovieScenePossessable* Possessable = MovieScene->FindPossessable(PossessableGuids[PossessableGuidIndex]);
 		if (Possessable && World)
 		{
-			auto AddActor = [&](AActor* Actor)
+			for (TActorIterator<AActor> ActorItr(World); ActorItr; ++ActorItr)
 			{
-				FGuid ExistingGuid = Sequencer->FindObjectId(*Actor, Sequencer->GetFocusedTemplateID());
-
-				if (!ExistingGuid.IsValid())
+				AActor* Actor = *ActorItr;
+				if (Actor && PossessableGuidIndex < PossessableObjectNames.Num() && PossessableObjectNames[PossessableGuidIndex].Contains(Actor->GetName()))
 				{
-					FGuid NewGuid = AssignActor(Sequencer, Actor, Possessable->GetGuid());
+					FGuid ExistingGuid = Sequencer->FindObjectId(*Actor, Sequencer->GetFocusedTemplateID());
 
-					// If assigning produces a new guid, update the possessable guids and the bindings pasted data
-					if (NewGuid.IsValid())
+					if (!ExistingGuid.IsValid())
 					{
-						for (auto BindingPasted : BindingsPasted)
+						FGuid NewGuid = AssignActor(Sequencer, Actor, Possessable->GetGuid());
+
+						// If assigning produces a new guid, update the possessable guids and the bindings pasted data
+						if (NewGuid.IsValid())
 						{
-							if (BindingPasted.GetObjectGuid() == PossessableGuids[PossessableGuidIndex])
+							for (auto BindingPasted : BindingsPasted)
 							{
-								BindingPasted.SetObjectGuid(NewGuid);
+								if (BindingPasted.GetObjectGuid() == PossessableGuids[PossessableGuidIndex])
+								{
+									BindingPasted.SetObjectGuid(NewGuid);
+								}
 							}
-						}
 
-						if (GuidToFolderMap.Contains(PossessableGuids[PossessableGuidIndex]))
-						{
-							GuidToFolderMap.Add(NewGuid, GuidToFolderMap[PossessableGuids[PossessableGuidIndex]]);
-							GuidToFolderMap.Remove(PossessableGuids[PossessableGuidIndex]);
-						}
-
-						for (TPair<FGuid, FGuid>& OldToNewGuid : OldToNewGuidMap)
-						{
-							if (OldToNewGuid.Value == PossessableGuids[PossessableGuidIndex])
+							if (GuidToFolderMap.Contains(PossessableGuids[PossessableGuidIndex]))
 							{
-								OldToNewGuid.Value = NewGuid;
+								GuidToFolderMap.Add(NewGuid, GuidToFolderMap[PossessableGuids[PossessableGuidIndex]]);
+								GuidToFolderMap.Remove(PossessableGuids[PossessableGuidIndex]);
 							}
-						}
 
-						PossessableGuids[PossessableGuidIndex] = NewGuid;
-					}
-				}
-			};
-			
-			for (TActorIterator<AActor> ActorItr(World); ActorItr; ++ActorItr)
-			{
-				AActor* Actor = *ActorItr;
-				FString ActorName = Actor ? Actor->GetName() : TEXT("");
-				if (Actor && PossessableGuidIndex < PossessableObjectNames.Num() && PossessableObjectNames[PossessableGuidIndex].Contains(Actor->GetName()))
-				{
-					AddActor(Actor);
-				}
-			}
+							for (TPair<FGuid, FGuid>& OldToNewGuid : OldToNewGuidMap)
+							{
+								if (OldToNewGuid.Value == PossessableGuids[PossessableGuidIndex])
+								{
+									OldToNewGuid.Value = NewGuid;
+								}
+							}
 
-			// If pasted actors have been provided, go through those as well
-			for (const TPair<FName, TObjectPtr<AActor>>& PastedActorPair : PasteBindingsParams.PastedActors)
-			{
-				if (PastedActorPair.Value && PossessableGuidIndex < PossessableObjectNames.Num())
-				{
-					FString PastedActorPairName = ObjectTools::SanitizeObjectPath(PastedActorPair.Key.ToString());
-					if (Algo::AnyOf(PossessableObjectNames[PossessableGuidIndex], [&PastedActorPairName] (const FString& PathName) 
-					{ 
-						FString SanitizedObjectPath = ObjectTools::SanitizeObjectPath(PathName);
-						int32 LastPeriod = INDEX_NONE;
-						if (SanitizedObjectPath.FindLastChar('.', LastPeriod))
-						{
-							return SanitizedObjectPath.RightChop(LastPeriod+1) == PastedActorPairName;
+							PossessableGuids[PossessableGuidIndex] = NewGuid;
 						}
-						return SanitizedObjectPath == PastedActorPairName;
-					}))
-					{
-						AddActor(PastedActorPair.Value.Get());
 					}
 				}
 			}
@@ -3226,10 +3080,9 @@ bool FSequencerUtilities::PasteBindings(const FString& TextToImport, TSharedRef<
 	// Refresh all immediately so that spawned actors will be generated immediately
 	Sequencer->ForceEvaluate();
 
-	// Fix possessable subobject bindings
-	for (int32 PossessableGuidIndex = 0; PossessableGuidIndex < PossessableGuids.Num(); ++PossessableGuidIndex)
+	// Fix possessable component bindings
+	for (auto PossessableGuid : PossessableGuids)
 	{
-		FGuid PossessableGuid = PossessableGuids[PossessableGuidIndex];
 		// If a possessable guid does not have any bound objects, they might be 
 		// possessable components for spawnables, so they need to be remapped
 		if (Sequencer->FindBoundObjects(PossessableGuid, Sequencer->GetFocusedTemplateID()).Num() == 0)
@@ -3238,7 +3091,6 @@ bool FSequencerUtilities::PasteBindings(const FString& TextToImport, TSharedRef<
 			if (Possessable)
 			{
 				FGuid ParentGuid = Possessable->GetParent();
-				bool bBound = false;
 				for (TWeakObjectPtr<> WeakObject : Sequencer->FindBoundObjects(ParentGuid, Sequencer->GetFocusedTemplateID()))
 				{
 					if (AActor* SpawnedActor = Cast<AActor>(WeakObject.Get()))
@@ -3248,29 +3100,11 @@ bool FSequencerUtilities::PasteBindings(const FString& TextToImport, TSharedRef<
 							if (Component->GetName() == Possessable->GetName())
 							{
 								Sequence->BindPossessableObject(PossessableGuid, *Component, SpawnedActor);
-								bBound = true;
 								break;
 							}
 						}
 					}
-
-					if (!bBound)
-					{
-						if (PossessableGuidIndex < PossessableObjectNames.Num())
-						{
-							for (const FString& BoundObjectPath : PossessableObjectNames[PossessableGuidIndex])
-							{
-								if (UObject* FoundObject = FindObject<UObject>(WeakObject.Get(), *BoundObjectPath))
-								{
-									Sequence->BindPossessableObject(PossessableGuid, *FoundObject, WeakObject.Get());
-									bBound = true;
-									break;
-								}
-							}
-						}
-					}
 				}
-				
 
 				// If the parent doesn't actually exist, clear it.
 				FMovieScenePossessable* PossessableParent = MovieScene->FindPossessable(ParentGuid);
@@ -3404,65 +3238,11 @@ FGuid CreateImplementationDefinedBinding(TSharedRef<ISequencer> Sequencer, UObje
 	return PossessableGuid;
 }
 
-UObject* FSequencerUtilities::FindResolutionContext(TSharedRef<ISequencer> InSequencer
-	, UMovieSceneSequence& InSequence
-	, UMovieScene& InMovieScene
-	, const FGuid& InParentGuid
-	, UObject* InPlaybackContext)
-{
-	if (!InPlaybackContext || !InParentGuid.IsValid() || !InSequence.AreParentContextsSignificant())
-	{
-		return InPlaybackContext;
-	}
-
-	UObject* ResolutionContext = nullptr;
-
-	// Recursive call up the hierarchy
-	if (FMovieScenePossessable* const ParentPossessable = InMovieScene.FindPossessable(InParentGuid))
-	{
-		ResolutionContext = FSequencerUtilities::FindResolutionContext(InSequencer
-			, InSequence
-			, InMovieScene
-			, ParentPossessable->GetParent()
-			, InPlaybackContext);
-	}
-
-	if (!ResolutionContext)
-	{
-		ResolutionContext = InPlaybackContext;
-	}
-
-	TArray<UObject*, TInlineAllocator<1>> FoundObjects;
-	for (TWeakObjectPtr<> WeakObj : InSequencer->FindBoundObjects(InParentGuid, InSequencer->GetFocusedTemplateID()))
-	{
-		FoundObjects.Add(WeakObj.Get());
-	}
-
-	if (FoundObjects.IsEmpty())
-	{
-		return ResolutionContext;
-	}
-
-	return FoundObjects[0] ? FoundObjects[0] : ResolutionContext;
-}
-
 FGuid FSequencerUtilities::CreateBinding(TSharedRef<ISequencer> Sequencer, UObject& InObject, const UE::Sequencer::FCreateBindingParams& InParams)
 {
-	return CreateOrReplaceBinding(Sequencer.ToSharedPtr(), Sequencer->GetFocusedMovieSceneSequence(), &InObject, InParams);
-}
-
-FGuid FSequencerUtilities::CreateOrReplaceBinding(TSharedRef<ISequencer> Sequencer, UObject* InObject, const UE::Sequencer::FCreateBindingParams& InParams)
-{
-	return CreateOrReplaceBinding(Sequencer.ToSharedPtr(), Sequencer->GetFocusedMovieSceneSequence(), InObject, InParams);
-}
-
-FGuid FSequencerUtilities::CreateOrReplaceBinding(TSharedPtr<ISequencer> Sequencer, UMovieSceneSequence* OwnerSequence, UObject* InObject, const UE::Sequencer::FCreateBindingParams& InParams)
-{
-	if (!OwnerSequence)
-		return FGuid();
-	
 	const FScopedTransaction Transaction(LOCTEXT("CreateBinding", "Create New Binding"));
 
+	UMovieSceneSequence* OwnerSequence = Sequencer->GetFocusedMovieSceneSequence();
 	UMovieScene* OwnerMovieScene = OwnerSequence->GetMovieScene();
 
 	OwnerSequence->Modify();
@@ -3472,16 +3252,16 @@ FGuid FSequencerUtilities::CreateOrReplaceBinding(TSharedPtr<ISequencer> Sequenc
 	FMovieSceneBindingReferences* BindingReferences = OwnerSequence->GetBindingReferences();
 	if (BindingReferences)
 	{
-		BindingGuid = CreateGenericBinding(Sequencer, OwnerSequence, InObject, BindingReferences, InParams);
+		BindingGuid = CreateGenericBinding(Sequencer, &InObject, BindingReferences, InParams);
 	}
-	else if (Sequencer && InParams.bSpawnable && InObject)
+	else if (InParams.bSpawnable)
 	{
 		// Create an old-style spawnable
-		BindingGuid = MakeNewSpawnable(Sequencer.ToSharedRef(), *InObject, InParams.ActorFactory.Get());
+		BindingGuid = MakeNewSpawnable(Sequencer, InObject, InParams.ActorFactory.Get());
 	}
-	else if (Sequencer && InObject)
+	else
 	{
-		BindingGuid = CreateImplementationDefinedBinding(Sequencer.ToSharedRef(), *InObject, InParams);
+		BindingGuid = CreateImplementationDefinedBinding(Sequencer, InObject, InParams);
 	}
 
 	if (!BindingGuid.IsValid())
@@ -3528,14 +3308,11 @@ FGuid FSequencerUtilities::CreateOrReplaceBinding(TSharedPtr<ISequencer> Sequenc
 			DestinationFolder->AddChildObjectBinding(RootObjectGuid);
 		}
 	}
-	if (Sequencer)
-	{
-		Sequencer->OnAddBinding(BindingGuid, OwnerMovieScene);
-	}
+
+	Sequencer->OnAddBinding(BindingGuid, OwnerMovieScene);
 	return BindingGuid;
 }
 
-
 void FSequencerUtilities::UpdateBindingIDs(TSharedRef<ISequencer> Sequencer, FGuid OldGuid, FGuid NewGuid)
 {
 	UMovieSceneCompiledDataManager* CompiledDataManager = FindObject<UMovieSceneCompiledDataManager>(GetTransientPackage(), TEXT("SequencerCompiledDataManager"));
@@ -3766,132 +3543,97 @@ FGuid FSequencerUtilities::AssignActor(TSharedRef<ISequencer> Sequencer, AActor*
 
 void FSequencerUtilities::AddActorsToBinding(TSharedRef<ISequencer> Sequencer, const TArray<AActor*>& Actors, const FMovieSceneBindingProxy& ObjectBinding)
 {
-	TArray<UObject*> ObjectsToAdd;
-	Algo::Transform(Actors, ObjectsToAdd, [](AActor* Actor) { return Actor;});
-	AddObjectsToBinding(Sequencer, ObjectsToAdd, ObjectBinding, Sequencer->GetPlaybackContext());
-}
+	if (!Actors.Num())
+	{
+		return;
+	}
 
-void FSequencerUtilities::AddObjectsToBinding(TSharedRef<ISequencer> InSequencer, const TArray<UObject*>& InObjectsToAdd, const FMovieSceneBindingProxy& InObjectBinding, UObject* InResolutionContext)
-{
-	UMovieSceneSequence* Sequence = InObjectBinding.Sequence;
+	UMovieSceneSequence* Sequence = Sequencer->GetFocusedMovieSceneSequence();
 	if (!Sequence)
 	{
 		return;
 	}
 
-	UMovieScene* const MovieScene = Sequence->GetMovieScene();
-	if (!MovieScene || InObjectsToAdd.IsEmpty())
+	UMovieScene* MovieScene = Sequence->GetMovieScene();
+	if (!MovieScene)
 	{
 		return;
 	}
 
-	UClass* ObjectClass = nullptr;
-	int32 ValidObjectCount = 0;
-
-	FGuid Guid = InObjectBinding.BindingID;
+	UClass* ActorClass = nullptr;
+	int32 NumRuntimeObjects = 0;
 
-	TArrayView<TWeakObjectPtr<>> ObjectsInCurrentSequence = InSequencer->FindObjectsInCurrentSequence(Guid);
+	FGuid Guid = ObjectBinding.BindingID;
+	TArrayView<TWeakObjectPtr<>> ObjectsInCurrentSequence = Sequencer->FindObjectsInCurrentSequence(Guid);
 
 	for (TWeakObjectPtr<> Ptr : ObjectsInCurrentSequence)
 	{
-		if (const UObject* Object = Cast<AActor>(Ptr.Get()))
+		if (const AActor* Actor = Cast<AActor>(Ptr.Get()))
 		{
-			ObjectClass = Object->GetClass();
-			++ValidObjectCount;
+			ActorClass = Actor->GetClass();
+			++NumRuntimeObjects;
 		}
 	}
 
+	FScopedTransaction AddSelectedToBinding(LOCTEXT("AddSelectedToBinding", "Add Selected to Binding"));
+
 	Sequence->Modify();
 	MovieScene->Modify();
 
-	TArray<UObject*> AddedObjects;
-	AddedObjects.Reserve(InObjectsToAdd.Num());
-
-	for (UObject* ObjectToAdd : InObjectsToAdd)
+	// Bind objects
+	int32 NumObjectsAdded = 0;
+	for (AActor* ActorToAdd : Actors)
 	{
-		// Skip invalid objects or objects already in the sequence
-		if (!ObjectToAdd || ObjectsInCurrentSequence.Contains(ObjectToAdd))
+		if (!ObjectsInCurrentSequence.Contains(ActorToAdd))
 		{
-			continue;
-		}
-
-		// Skip if the object has no common class with the objects already in the binding
-		if (ObjectClass && !UClass::FindCommonBase(ObjectToAdd->GetClass(), ObjectClass))
-		{
-			continue;
-		}
-
-		// if no objects are in the binding, set the class to this object's
-		if (!ObjectClass)
-		{
-			ObjectClass = ObjectToAdd->GetClass();
-		}
-
-		FMovieScenePossessable* Possessable = MovieScene->FindPossessable(Guid);
-		if (!ensureAlways(Possessable))
-		{
-			continue;
-		}
-
-		ObjectToAdd->Modify();
-		if (!Possessable->BindSpawnableObject(InSequencer->GetFocusedTemplateID(), ObjectToAdd, InSequencer->GetSharedPlaybackState()))
-		{
-			Sequence->BindPossessableObject(Guid, *ObjectToAdd, InResolutionContext);
-		}
-
-		if (ObjectClass == nullptr || UClass::FindCommonBase(ObjectToAdd->GetClass(), ObjectClass) != nullptr)
-		{
-			if (ObjectClass == nullptr)
+			if (ActorClass == nullptr || UClass::FindCommonBase(ActorToAdd->GetClass(), ActorClass) != nullptr)
 			{
-				ObjectClass = ObjectToAdd->GetClass();
-			}
+				if (ActorClass == nullptr)
+				{
+					ActorClass = ActorToAdd->GetClass();
+				}
 
-			ObjectToAdd->Modify();
-			if (!MovieScene->FindPossessable(Guid)->BindSpawnableObject(InSequencer->GetFocusedTemplateID(), ObjectToAdd, InSequencer->GetSharedPlaybackState()))
+				ActorToAdd->Modify();
+				if (!MovieScene->FindPossessable(Guid)->BindSpawnableObject(Sequencer->GetFocusedTemplateID(), ActorToAdd, Sequencer->GetSharedPlaybackState()))
+				{
+					Sequence->BindPossessableObject(Guid, *ActorToAdd, Sequencer->GetPlaybackContext());
+				}
+				++NumObjectsAdded;
+			}
+			else
 			{
-				Sequence->BindPossessableObject(Guid, *ObjectToAdd, InResolutionContext);
+				const FText NotificationText = FText::Format(LOCTEXT("UnableToAssignObject", "Cannot assign object {0}. Expected class {1}"), FText::FromString(ActorToAdd->GetPathName()), FText::FromString(ActorClass->GetName()));
+				FNotificationInfo Info(NotificationText);
+				Info.ExpireDuration = 3.f;
+				Info.bUseLargeFont = false;
+				FSlateNotificationManager::Get().AddNotification(Info);
 			}
-			AddedObjects.Add(ObjectToAdd);
-		}
-		else
-		{
-			const FText NotificationText = FText::Format(LOCTEXT("UnableToAssignObject", "Cannot assign object {0}. Expected class {1}"), FText::FromString(ObjectToAdd->GetPathName()), FText::FromString(ObjectClass->GetName()));
-			FNotificationInfo Info(NotificationText);
-			Info.ExpireDuration = 3.f;
-			Info.bUseLargeFont = false;
-			FSlateNotificationManager::Get().AddNotification(Info);
 		}
 	}
 
-	// Update Labels
-	if (ValidObjectCount + AddedObjects.Num() > 0)
+	// Update label
+	if (NumRuntimeObjects + NumObjectsAdded > 0)
 	{
-		FMovieScenePossessable* const Possessable = MovieScene->FindPossessable(Guid);
-		if (Possessable && ObjectClass)
+		FMovieScenePossessable* Possessable = MovieScene->FindPossessable(Guid);
+		if (Possessable && ActorClass != nullptr)
 		{
-			// If there are multiple objects within the same possessable, name possessable as "ClassName (Count)"
-			if (ValidObjectCount + AddedObjects.Num() > 1)
+			if (NumRuntimeObjects + NumObjectsAdded > 1)
 			{
-				Possessable->SetName(FString::Printf(TEXT("%s (%d)")
-					, *ObjectClass->GetName()
-					, ValidObjectCount + AddedObjects.Num()));
+				FString NewLabel = ActorClass->GetName() + FString::Printf(TEXT(" (%d)"), NumRuntimeObjects + NumObjectsAdded);
+				Possessable->SetName(NewLabel);
 			}
-			else if (!AddedObjects.IsEmpty())
+			else if (NumObjectsAdded > 0 && Actors.Num() > 0)
 			{
-				FString PossessableName = AddedObjects[0]->GetName();
-				if (AActor* const Actor = Cast<AActor>(AddedObjects[0]))
-				{
-					PossessableName = Actor->GetActorLabel();
-				}
-				Possessable->SetName(PossessableName);
+				Possessable->SetName(Actors[0]->GetActorLabel());
 			}
-			Possessable->SetPossessedObjectClass(ObjectClass);
+
+			Possessable->SetPossessedObjectClass(ActorClass);
 		}
 	}
 
-	InSequencer->RestorePreAnimatedState();
+	Sequencer->RestorePreAnimatedState();
 
-	InSequencer->NotifyMovieSceneDataChanged(EMovieSceneDataChangeType::MovieSceneStructureItemsChanged);
+	Sequencer->NotifyMovieSceneDataChanged(EMovieSceneDataChangeType::MovieSceneStructureItemsChanged);
 }
 
 void FSequencerUtilities::ReplaceBindingWithActors(TSharedRef<ISequencer> Sequencer, const TArray<AActor*>& Actors, const FMovieSceneBindingProxy& ObjectBinding)
@@ -4177,74 +3919,6 @@ void FSequencerUtilities::HandleTemplateActorClassPicked(UClass* ChosenClass, TS
 	}
 }
 
-bool FSequencerUtilities::CanConvertToPossessable(TSharedRef<ISequencer> Sequencer, FGuid BindingGuid, int32 BindingIndex)
-{
-	UMovieSceneSequence* Sequence = Sequencer->GetFocusedMovieSceneSequence();
-	UMovieScene* MovieScene = Sequence ? Sequence->GetMovieScene() : nullptr;
-	if (!MovieScene)
-	{
-		return false;
-	}
-
-	if (FMovieSceneSpawnable* Spawnable = MovieScene->FindSpawnable(BindingGuid))
-	{
-		return true;
-	}
-	else if (const FMovieSceneBindingReferences* BindingReferences = Sequence->GetBindingReferences())
-	{
-		TArrayView<const FMovieSceneBindingReference> BindingReferencesList = BindingReferences->GetReferences(BindingGuid);
-		if (BindingReferencesList.IsValidIndex(BindingIndex) && BindingReferencesList[BindingIndex].CustomBinding != nullptr)
-		{
-			return true;
-		}
-	}
-	return false;
-}
-
-bool FSequencerUtilities::CanConvertToCustomBinding(TSharedRef<ISequencer> Sequencer, FGuid BindingGuid, TSubclassOf<UMovieSceneCustomBinding> CustomBindingType, int32 BindingIndex)
-{
-	UMovieSceneSequence* Sequence = Sequencer->GetFocusedMovieSceneSequence();
-	UMovieScene* MovieScene = Sequence ? Sequence->GetMovieScene() : nullptr;
-	if (!MovieScene)
-	{
-		return false;
-	}
-	if (FMovieSceneSpawnable* Spawnable = MovieScene->FindSpawnable(BindingGuid))
-	{
-		if (UObject* CurrentBoundObject = Spawnable->GetObjectTemplate())
-		{
-			return CustomBindingType && CustomBindingType->GetDefaultObject<UMovieSceneCustomBinding>()->SupportsBindingCreationFromObject(CurrentBoundObject);
-		}
-	}
-	else if (FMovieScenePossessable* Possessable = MovieScene->FindPossessable(BindingGuid))
-	{
-		if (const FMovieSceneBindingReferences* BindingReferences = Sequence->GetBindingReferences())
-		{
-			UObject* ResolutionContext = MovieSceneHelpers::GetResolutionContext(Sequence, BindingGuid, Sequencer->GetFocusedTemplateID(), Sequencer->GetSharedPlaybackState());
-
-			TArrayView<const FMovieSceneBindingReference> BindingReferencesList = Sequencer->GetFocusedMovieSceneSequence()->GetBindingReferences()->GetReferences(BindingGuid);
-
-			if (const FMovieSceneBindingReference* CurrentBindingReference = BindingReferences->GetReference(BindingGuid, BindingIndex))
-			{
-				UE::UniversalObjectLocator::FResolveParams LocatorResolveParams(ResolutionContext);
-				FMovieSceneBindingResolveParams BindingResolveParams{ Sequence, BindingGuid, Sequencer->GetFocusedTemplateID(), ResolutionContext };
-				if (UObject* CurrentBoundObject = BindingReferences->ResolveSingleBinding(BindingResolveParams, BindingIndex, LocatorResolveParams, Sequencer->GetSharedPlaybackState()))
-				{
-					if (CustomBindingType
-						&& (!CurrentBindingReference->CustomBinding
-							|| CurrentBindingReference->CustomBinding->GetClass() != CustomBindingType)
-						&& CustomBindingType->GetDefaultObject<UMovieSceneCustomBinding>()->SupportsConversionFromBinding(*CurrentBindingReference, CurrentBoundObject))
-					{
-						return true;
-					}
-				}
-			}
-		}
-	}
-	return false;
-}
-
-
 void FSequencerUtilities::AddConvertBindingMenu(FMenuBuilder& MenuBuilder, TSharedRef<ISequencer> Sequencer, const TArray<FSequencerConvertBindingInfo>& BindingsToConvert, TFunction<void()> OnBindingChanged)
 {
 	UMovieSceneSequence* Sequence = Sequencer->GetFocusedMovieSceneSequence();
@@ -4277,23 +3951,22 @@ void FSequencerUtilities::AddConvertBindingMenu(FMenuBuilder& MenuBuilder, TShar
 			{
 				SlowTask.EnterProgressFrame();
 
-				if (FMovieScenePossessable* Possessable = DoConvert(BindingInfo.BindingID, BindingInfo.BindingIndex))
-				{
-					Sequencer->ForceEvaluate();
+				FMovieScenePossessable* Possessable = DoConvert(BindingInfo.BindingID, BindingInfo.BindingIndex);
 
-					for (TWeakObjectPtr<> WeakObject : Sequencer->FindBoundObjects(Possessable->GetGuid(), Sequencer->GetFocusedTemplateID()))
-					{
-						if (AActor* PossessedActor = Cast<AActor>(WeakObject.Get()))
-						{
-							PossessedActors.Add(PossessedActor);
-						}
-					}
+				Sequencer->ForceEvaluate();
 
-					if (GWarn->ReceivedUserCancel())
+				for (TWeakObjectPtr<> WeakObject : Sequencer->FindBoundObjects(Possessable->GetGuid(), Sequencer->GetFocusedTemplateID()))
+				{
+					if (AActor* PossessedActor = Cast<AActor>(WeakObject.Get()))
 					{
-						break;
+						PossessedActors.Add(PossessedActor);
 					}
 				}
+
+				if (GWarn->ReceivedUserCancel())
+				{
+					break;
+				}
 			}
 
 			if (PossessedActors.Num())
@@ -4323,8 +3996,20 @@ void FSequencerUtilities::AddConvertBindingMenu(FMenuBuilder& MenuBuilder, TShar
 	};
 
 	// Can convert to possessable
-	if (Algo::AllOf(BindingsToConvert, [Sequencer, &Sequence, &MovieScene](const FSequencerConvertBindingInfo& BindingInfo) {
-		return CanConvertToPossessable(Sequencer, BindingInfo.BindingID, BindingInfo.BindingIndex );
+	if (Algo::AllOf(BindingsToConvert, [&Sequence, &MovieScene](const FSequencerConvertBindingInfo& BindingInfo) {
+		if (FMovieSceneSpawnable* Spawnable = MovieScene->FindSpawnable(BindingInfo.BindingID))
+		{
+			return true;
+		}
+		else if (const FMovieSceneBindingReferences* BindingReferences = Sequence->GetBindingReferences())
+		{
+			TArrayView<const FMovieSceneBindingReference> BindingReferencesList = BindingReferences->GetReferences(BindingInfo.BindingID);
+			if (BindingReferencesList.IsValidIndex(BindingInfo.BindingIndex) && BindingReferencesList[BindingInfo.BindingIndex].CustomBinding != nullptr)
+			{
+				return true;
+			}
+		}
+		return false;
 	}))
 	{
 		MenuBuilder.AddMenuEntry(
@@ -4344,7 +4029,34 @@ void FSequencerUtilities::AddConvertBindingMenu(FMenuBuilder& MenuBuilder, TShar
 	{
 		if (Algo::AllOf(BindingsToConvert, [&Sequence, &MovieScene, &CustomBindingType, Sequencer](const FSequencerConvertBindingInfo& BindingInfo) 
 		{
-			return CanConvertToCustomBinding(Sequencer, BindingInfo.BindingID, CustomBindingType, BindingInfo.BindingIndex);
+			if (FMovieSceneSpawnable* Spawnable = MovieScene->FindSpawnable(BindingInfo.BindingID))
+			{
+				if (UObject* CurrentBoundObject = Spawnable->GetObjectTemplate())
+				{
+					return CustomBindingType && CustomBindingType->GetDefaultObject<UMovieSceneCustomBinding>()->SupportsBindingCreationFromObject(CurrentBoundObject);
+				}
+			}
+			else if (const FMovieSceneBindingReferences* BindingReferences = Sequence->GetBindingReferences())
+			{
+				TArrayView<const FMovieSceneBindingReference> BindingReferencesList = Sequencer->GetFocusedMovieSceneSequence()->GetBindingReferences()->GetReferences(BindingInfo.BindingID);
+
+				if (const FMovieSceneBindingReference* CurrentBindingReference = BindingReferences->GetReference(BindingInfo.BindingID, BindingInfo.BindingIndex))
+				{
+					UE::UniversalObjectLocator::FResolveParams LocatorResolveParams(Sequencer->GetSharedPlaybackState()->GetPlaybackContext());
+					FMovieSceneBindingResolveParams BindingResolveParams{ Sequence, BindingInfo.BindingID, Sequencer->GetFocusedTemplateID() };
+					if (UObject* CurrentBoundObject = BindingReferences->ResolveSingleBinding(BindingResolveParams, BindingInfo.BindingIndex, LocatorResolveParams, Sequencer->GetSharedPlaybackState()))
+					{
+						if (CustomBindingType
+							&& (!CurrentBindingReference->CustomBinding
+								|| CurrentBindingReference->CustomBinding->GetClass() != CustomBindingType)
+							&& CustomBindingType->GetDefaultObject<UMovieSceneCustomBinding>()->SupportsConversionFromBinding(*CurrentBindingReference, CurrentBoundObject))
+						{
+							return true;
+						}
+					}
+				}
+			}
+			return false;
 		}))
 		{
 			MenuBuilder.AddMenuEntry(
@@ -4358,4 +4070,6 @@ void FSequencerUtilities::AddConvertBindingMenu(FMenuBuilder& MenuBuilder, TShar
 		}
 	}
 }
+
+
 #undef LOCTEXT_NAMESPACE
diff --git a/Engine/Source/Editor/Sequencer/Public/SequencerUtilities.h b/Engine/Source/Editor/Sequencer/Public/SequencerUtilities.h
--- a/Engine/Source/Editor/Sequencer/Public/SequencerUtilities.h
+++ b/Engine/Source/Editor/Sequencer/Public/SequencerUtilities.h
@@ -163,9 +163,6 @@ struct FMovieScenePasteBindingsParams
 
 	UPROPERTY(EditAnywhere, BlueprintReadWrite, Category = "Movie Scene")
 	bool bDuplicateExistingActors;
-
-	UPROPERTY(EditAnywhere, BlueprintReadWrite, Category = "Movie Scene")
-	TMap<FName, TObjectPtr<AActor>> PastedActors;
 };
 
 USTRUCT()
@@ -222,11 +219,9 @@ struct SEQUENCER_API FSequencerUtilities
 	static TArray<FMovieSceneSpawnable*> ConvertToSpawnable(TSharedRef<ISequencer> Sequencer, FGuid PossessableGuid);
 	
 	/** Convert the requested object binding and object binding index to a possessable */
-	static bool CanConvertToPossessable(TSharedRef<ISequencer> Sequencer, FGuid BindingGuid, int32 BindingIndex = 0);
 	static FMovieScenePossessable* ConvertToPossessable(TSharedRef<ISequencer> Sequencer, FGuid BindingGuid, int32 BindingIndex=0);
 
 	/** Convert the selected object binding and object binding index to a custom binding of the chosen type.*/
-	static bool CanConvertToCustomBinding(TSharedRef<ISequencer> Sequencer, FGuid BindingGuid, TSubclassOf<UMovieSceneCustomBinding> CustomBindingType, int32 BindingIndex = 0);
 	static FMovieScenePossessable* ConvertToCustomBinding(TSharedRef<ISequencer> Sequencer, FGuid BindingGuid, TSubclassOf<UMovieSceneCustomBinding> CustomBindingType, int32 BindingIndex = 0);
 
 	/** Copy/paste folders */
@@ -246,28 +241,15 @@ struct SEQUENCER_API FSequencerUtilities
 
 	/** Copy/paste object bindings */
 	static void CopyBindings(TSharedRef<ISequencer> Sequencer, const TArray<FMovieSceneBindingProxy>& Bindings, const TArray<UMovieSceneFolder*>& InFolders, FString& ExportedText);
-	static void CopyBindings(TSharedRef<ISequencer> Sequencer, const TArray<FMovieSceneBindingProxy>& Bindings, const TArray<UMovieSceneFolder*>& InFolders, FOutputDevice& Ar);
 	static bool PasteBindings(const FString& TextToImport, TSharedRef<ISequencer> Sequencer, FMovieScenePasteBindingsParams PasteBindingsParams, TArray<FMovieSceneBindingProxy>& OutBindings, TArray<FNotificationInfo>& OutErrors);
 	static bool CanPasteBindings(TSharedRef<ISequencer> Sequencer, const FString& TextToImport);
 	static TArray<FString> GetPasteBindingsObjectNames(TSharedRef<ISequencer> Sequencer, const FString& TextToImport);
-	/**
-	 * Recursively finds the most appropriate Resolution Context for a given Parent Guid of a Possessable
-	 */
-
-	static UObject* FindResolutionContext(TSharedRef<ISequencer> Sequencer
-		, UMovieSceneSequence& InSequence
-		, UMovieScene& InMovieScene
-		, const FGuid& InParentGuid
-		, UObject* InPlaybackContext);
 
 	/** Utility functions for managing bindings */
 	static FGuid CreateBinding(TSharedRef<ISequencer> Sequencer, UObject& InObject, const UE::Sequencer::FCreateBindingParams& Params = UE::Sequencer::FCreateBindingParams());
-	static FGuid CreateOrReplaceBinding(TSharedRef<ISequencer> Sequencer, UObject* Object, const UE::Sequencer::FCreateBindingParams& Params = UE::Sequencer::FCreateBindingParams());
-	static FGuid CreateOrReplaceBinding(TSharedPtr<ISequencer> Sequencer, UMovieSceneSequence* Sequence, UObject* Object, const UE::Sequencer::FCreateBindingParams& Params = UE::Sequencer::FCreateBindingParams());
 	static void UpdateBindingIDs(TSharedRef<ISequencer> Sequencer, FGuid OldGuid, FGuid NewGuid);
 	static FGuid AssignActor(TSharedRef<ISequencer> Sequencer, AActor* Actor, FGuid InObjectBinding);
 	static void AddActorsToBinding(TSharedRef<ISequencer> Sequencer, const TArray<AActor*>& Actors, const FMovieSceneBindingProxy& ObjectBinding);
-	static void AddObjectsToBinding(TSharedRef<ISequencer> Sequencer, const TArray<UObject*>& Objects, const FMovieSceneBindingProxy& ObjectBinding, UObject* ResolutionContext);
 	static void ReplaceBindingWithActors(TSharedRef<ISequencer> Sequencer, const TArray<AActor*>& Actors, const FMovieSceneBindingProxy& ObjectBinding);
 	static void RemoveActorsFromBinding(TSharedRef<ISequencer> Sequencer, const TArray<AActor*>& Actors, const FMovieSceneBindingProxy& ObjectBinding);
 
diff --git a/Engine/Source/Editor/UMGEditor/Private/UMGEditorModule.cpp b/Engine/Source/Editor/UMGEditor/Private/UMGEditorModule.cpp
--- a/Engine/Source/Editor/UMGEditor/Private/UMGEditorModule.cpp
+++ b/Engine/Source/Editor/UMGEditor/Private/UMGEditorModule.cpp
@@ -44,8 +44,6 @@
 #include "WidgetBlueprintThumbnailRenderer.h"
 #include "Customizations/WidgetThumbnailCustomization.h"
 #include "Widgets/SBindWidgetView.h"
-#include "MovieSceneDynamicBindingUtils.h"
-#include "Kismet2/KismetEditorUtilities.h"
 
 #define LOCTEXT_NAMESPACE "UMG"
 
@@ -137,9 +135,6 @@ public:
 		FEdGraphUtilities::RegisterVisualPinFactory(GraphPanelPinFactory);
 
 		CVarThumbnailRenderEnable->AsVariable()->SetOnChangedCallback(FConsoleVariableDelegate::CreateStatic(&FUMGEditorModule::ThumbnailRenderingEnabled));
-
-		FixupDynamicBindingPayloadParameterNameHandle = UMovieScene::FixupDynamicBindingPayloadParameterNameEvent.AddStatic(FixupPayloadParameterNameForDynamicBinding);
-		FixupWidgetDynamicBindingsHandle = UWidgetAnimation::FixupWidgetDynamicBindingsEvent.AddStatic(FixupWidgetDynamicBindings);
 	}
 
 	/** Called before the module is unloaded, right before the module object is destroyed. */
@@ -220,9 +215,6 @@ public:
 		//	SettingsModule->UnregisterSettings("Editor", "ContentEditors", "WidgetDesigner");
 		//	SettingsModule->UnregisterSettings("Project", "Editor", "UMGEditor");
 		//}
-
-		UMovieScene::FixupDynamicBindingPayloadParameterNameEvent.Remove(FixupDynamicBindingPayloadParameterNameHandle);
-		UWidgetAnimation::FixupWidgetDynamicBindingsEvent.Remove(FixupWidgetDynamicBindingsHandle);
 	}
 
 	/** Gets the extensibility managers for outside entities to extend gui page editor's menus and toolbars */
@@ -442,58 +434,6 @@ private:
 		}
 	}
 
-	static void FixupPayloadParameterNameForDynamicBinding(UMovieScene* MovieScene, UK2Node* InNode, FName OldPinName, FName NewPinName)
-	{
-		using namespace UE::MovieScene;
-
-		check(MovieScene);
-
-		auto FixupPayloadParameterName = [InNode, OldPinName, NewPinName](FMovieSceneDynamicBinding& DynamicBinding)
-		{
-			if (DynamicBinding.WeakEndpoint.Get() == InNode)
-			{
-				if (FMovieSceneDynamicBindingPayloadVariable* Variable = DynamicBinding.PayloadVariables.Find(OldPinName))
-				{
-					DynamicBinding.PayloadVariables.Add(NewPinName, MoveTemp(*Variable));
-					DynamicBinding.PayloadVariables.Remove(OldPinName);
-				}
-			}
-		};
-
-		UMovieSceneSequence* ThisSequence = MovieScene->GetTypedOuter<UMovieSceneSequence>();
-		TSharedRef<UE::MovieScene::FSharedPlaybackState> TransientPlaybackState = MovieSceneHelpers::CreateTransientSharedPlaybackState(GEditor->GetEditorWorldContext().World(), ThisSequence);
-
-		if (UWidgetAnimation* WidgetAnimation = MovieScene->GetTypedOuter<UWidgetAnimation>())
-		{
-			for (FWidgetAnimationBinding& WidgetAnimationBinding : WidgetAnimation->AnimationBindings)
-			{
-				FixupPayloadParameterName(WidgetAnimationBinding.DynamicBinding);
-			}
-		}
-	}
-
-	static void FixupWidgetDynamicBindings(UWidgetAnimation* WidgetAnimation)
-	{
-		if (WidgetAnimation)
-		{
-			FMovieSceneSequenceEditor* SequenceEditor = FMovieSceneSequenceEditor::Find(WidgetAnimation);
-			if (!SequenceEditor)
-			{
-				return;
-			}
-
-			UBlueprint* SequenceDirectorBP = SequenceEditor->GetOrCreateDirectorBlueprint(WidgetAnimation);
-			if (!SequenceDirectorBP)
-			{
-				return;
-			}
-
-			FMovieSceneDynamicBindingUtils::EnsureBlueprintExtensionCreated(WidgetAnimation, SequenceDirectorBP);
-			FKismetEditorUtilities::CompileBlueprint(SequenceDirectorBP);
-		}
-	}
-
-
 private:
 	TSharedPtr<FExtensibilityManager> MenuExtensibilityManager;
 	TSharedPtr<FExtensibilityManager> ToolBarExtensibilityManager;
@@ -536,8 +476,6 @@ private:
 	FDelegateHandle BlueprintVariableCustomizationHandle;
 	/** Handle for FBlueprintEditorModule::RegisterFunctionCustomization */
 	FDelegateHandle BlueprintFunctionCustomizationHandle;
-	FDelegateHandle FixupDynamicBindingPayloadParameterNameHandle;
-	FDelegateHandle FixupWidgetDynamicBindingsHandle;
 
 	bool bThumbnailRenderersRegistered;
 	bool bOnPostEngineInitHandled;
diff --git a/Engine/Source/Editor/UMGEditor/Private/WidgetBlueprintCompiler.cpp b/Engine/Source/Editor/UMGEditor/Private/WidgetBlueprintCompiler.cpp
--- a/Engine/Source/Editor/UMGEditor/Private/WidgetBlueprintCompiler.cpp
+++ b/Engine/Source/Editor/UMGEditor/Private/WidgetBlueprintCompiler.cpp
@@ -1243,12 +1243,12 @@ void FWidgetBlueprintCompilerContext::ValidateWidgetAnimations()
 		for (const FWidgetAnimationBinding& Binding : InAnimation->AnimationBindings)
 		{
 			// Look for the object bindings within the widget
-			UObject* FoundObject = Binding.FindRuntimeObject(*LatestWidgetTree, *UserWidget, InAnimation, nullptr);
+			UObject* FoundObject = Binding.FindRuntimeObject(*LatestWidgetTree, *UserWidget);
 
 			// If any of the FoundObjects is null, we do not play the animation.
 			if (FoundObject == nullptr)
 			{
-				FoundObject = Binding.FindRuntimeObject(*WidgetBP->WidgetTree, *UserWidget, InAnimation, nullptr);
+				FoundObject = Binding.FindRuntimeObject(*WidgetBP->WidgetTree, *UserWidget);
 				if (FoundObject == nullptr)
 				{
 					// Notify the user of the null track in the editor
diff --git a/Engine/Source/Editor/UMGEditor/Private/WidgetBlueprintEditor.cpp b/Engine/Source/Editor/UMGEditor/Private/WidgetBlueprintEditor.cpp
--- a/Engine/Source/Editor/UMGEditor/Private/WidgetBlueprintEditor.cpp
+++ b/Engine/Source/Editor/UMGEditor/Private/WidgetBlueprintEditor.cpp
@@ -89,7 +89,6 @@
 #include "Preferences/UnrealEdOptions.h"
 #include "UnrealEdGlobals.h"
 #include "GraphEditorActions.h"
-#include "MovieSceneDynamicBindingCustomization.h"
 
 #define LOCTEXT_NAMESPACE "UMG"
 
@@ -661,8 +660,7 @@ bool FWidgetBlueprintEditor::IsBindingSelected(const FMovieSceneBinding& InBindi
 	TSharedPtr<ISequencer>& ActiveSequencer = GetSequencer();
 	UMovieSceneSequence* AnimationSequence = ActiveSequencer->GetFocusedMovieSceneSequence();
 	UObject* BindingContext = GetAnimationPlaybackContext();
-	TArray<UObject*, TInlineAllocator<1>> BoundObjects;
-	AnimationSequence->LocateBoundObjects(InBinding.GetObjectGuid(), UE::UniversalObjectLocator::FResolveParams(BindingContext), ActiveSequencer->GetSharedPlaybackState(), BoundObjects);
+	TArray<UObject*, TInlineAllocator<1>> BoundObjects = AnimationSequence->LocateBoundObjects(InBinding.GetObjectGuid(), BindingContext);
 
 	if (BoundObjects.Num() == 0)
 	{
@@ -2297,7 +2295,7 @@ void FWidgetBlueprintEditor::RemoveMissingWidgetsFromTrack(FGuid ObjectId)
 	for (int32 Index = WidgetAnimation->AnimationBindings.Num() - 1; Index >= 0; --Index)
 	{
 		const FWidgetAnimationBinding& Binding = WidgetAnimation->AnimationBindings[Index];
-		if (Binding.AnimationGuid == ObjectId && Binding.FindRuntimeObject(*PreviewRoot->WidgetTree, *PreviewRoot, WidgetAnimation, ActiveSequencer->GetSharedPlaybackState()) == nullptr)
+		if (Binding.AnimationGuid == ObjectId && Binding.FindRuntimeObject(*PreviewRoot->WidgetTree, *PreviewRoot) == nullptr)
 		{
 			WidgetAnimation->AnimationBindings.RemoveAt(Index, EAllowShrinking::No);
 		}
@@ -2371,118 +2369,27 @@ void FWidgetBlueprintEditor::AddDynamicPossessionMenu(FMenuBuilder& MenuBuilder,
 
 	TSharedPtr<ISequencer>& ActiveSequencer = GetSequencer();
 
-	UWidgetAnimation* WidgetAnimation = Cast<UWidgetAnimation>(ActiveSequencer->GetFocusedMovieSceneSequence());
-	UMovieScene* MovieScene = WidgetAnimation->GetMovieScene();
-
+	UMovieScene* MovieScene = ActiveSequencer->GetFocusedMovieSceneSequence()->GetMovieScene();
 	FMovieScenePossessable* Possessable = MovieScene->FindPossessable(ObjectId);
 	if (!Possessable)
 	{
 		return;
 	}
 
-	FWidgetAnimationBinding* WidgetBinding = nullptr;
-
-	for (int32 Index = 0; Index < WidgetAnimation->AnimationBindings.Num(); ++Index)
-	{
-		FWidgetAnimationBinding& Binding = WidgetAnimation->AnimationBindings[Index];
-		if (Binding.AnimationGuid == ObjectId)
-		{
-			WidgetBinding = &Binding;
-			break;
-		}
-	}
-
-	if (WidgetBinding)
-	{
-		FDetailsViewArgs DetailsViewArgs;
-		{
-			DetailsViewArgs.bAllowSearch = false;
-			DetailsViewArgs.bCustomFilterAreaLocation = true;
-			DetailsViewArgs.bCustomNameAreaLocation = true;
-			DetailsViewArgs.bHideSelectionTip = true;
-			DetailsViewArgs.bLockable = false;
-			DetailsViewArgs.bSearchInitialKeyFocus = true;
-			DetailsViewArgs.bUpdatesFromSelection = false;
-			DetailsViewArgs.bShowOptions = false;
-			DetailsViewArgs.bShowModifiedPropertiesOption = false;
-			DetailsViewArgs.bShowScrollBar = false;
-		}
-
-		FStructureDetailsViewArgs StructureViewArgs;
-		{
-			StructureViewArgs.bShowObjects = false;
-			StructureViewArgs.bShowAssets = true;
-			StructureViewArgs.bShowClasses = true;
-			StructureViewArgs.bShowInterfaces = false;
-		}
-
-		TSharedRef<IStructureDetailsView> StructureDetailsView = FModuleManager::GetModuleChecked<FPropertyEditorModule>("PropertyEditor")
-			.CreateStructureDetailView(DetailsViewArgs, StructureViewArgs, nullptr);
-
-		// Register details customizations for this instance
-		StructureDetailsView->GetDetailsView()->RegisterInstancedCustomPropertyTypeLayout(
-			FMovieSceneDynamicBinding::StaticStruct()->GetFName(),
-			FOnGetPropertyTypeCustomizationInstance::CreateStatic(&FMovieSceneDynamicBindingCustomization::MakeInstance, MovieScene, ObjectId));
-
-		// We can't just show the FMovieSceneDynamicBinding struct in the details view, because Slate only uses
-		// the above details view customization for *properties* (not for the root object). So here we put a copy of
-		// our dynamic binding struct inside a container, and when the details view is done setting values on it,
-		// we copy these values back to the original dynamic binding.
-		TSharedPtr<FStructOnScope> StructOnScope = MakeShared<FStructOnScope>(FMovieSceneDynamicBindingContainer::StaticStruct());
-		FMovieSceneDynamicBindingContainer* BufferContainer = (FMovieSceneDynamicBindingContainer*)StructOnScope->GetStructMemory();
-		BufferContainer->DynamicBinding = WidgetBinding->DynamicBinding;
-		StructureDetailsView->SetStructureData(StructOnScope);
-
-		StructureDetailsView->GetOnFinishedChangingPropertiesDelegate().AddSP(this, &FWidgetBlueprintEditor::OnFinishedChangingDynamicBindingProperties, StructOnScope, ObjectId);
-
-		MenuBuilder.BeginSection(NAME_None, LOCTEXT("DynamicBindingHeader", "Dynamic Binding"));
-		{
-			TSharedRef<SWidget> Widget = StructureDetailsView->GetWidget().ToSharedRef();
-			MenuBuilder.AddWidget(Widget, FText());
-		}
-		MenuBuilder.EndSection();
-	}
-}
-
-void FWidgetBlueprintEditor::OnFinishedChangingDynamicBindingProperties(const FPropertyChangedEvent& ChangeEvent, TSharedPtr<FStructOnScope> ValueStruct, FGuid ObjectId)
-{
-	auto* Container = (FMovieSceneDynamicBindingContainer*)ValueStruct->GetStructMemory();
-
-	using namespace UE::Sequencer;
-
-	TSharedPtr<ISequencer>& ActiveSequencer = GetSequencer();
-
-	UWidgetAnimation* WidgetAnimation = Cast<UWidgetAnimation>(ActiveSequencer->GetFocusedMovieSceneSequence());
-	UMovieScene* MovieScene = WidgetAnimation->GetMovieScene();
-
-	FMovieScenePossessable* Possessable = MovieScene->FindPossessable(ObjectId);
-	if (!Possessable)
+	TSharedPtr<FSequencerEditorViewModel> SequencerViewModel = ActiveSequencer->GetViewModel();
+	FObjectBindingModelStorageExtension* ObjectStorage = SequencerViewModel->GetRootModel()->CastDynamic<FObjectBindingModelStorageExtension>();
+	if (!ObjectStorage)
 	{
 		return;
 	}
 
-	FWidgetAnimationBinding* WidgetBinding = nullptr;
-
-	for (int32 Index = 0; Index < WidgetAnimation->AnimationBindings.Num(); ++Index)
-	{
-		FWidgetAnimationBinding& Binding = WidgetAnimation->AnimationBindings[Index];
-		if (Binding.AnimationGuid == ObjectId)
-		{
-			WidgetBinding = &Binding;
-			break;
-		}
-	}
-
-	if (WidgetBinding)
+	TSharedPtr<FObjectBindingModel> ObjectBindingModel = ObjectStorage->FindModelForObjectBinding(ObjectId);
+	if (!ObjectBindingModel)
 	{
-		WidgetBinding->DynamicBinding = Container->DynamicBinding;
+		return;
 	}
 
-	// Force refresh the binding
-	if (FMovieSceneEvaluationState* EvaluationState = ActiveSequencer->GetSharedPlaybackState()->FindCapability<FMovieSceneEvaluationState>())
-	{
-		EvaluationState->Invalidate(ObjectId, ActiveSequencer->GetFocusedTemplateID());
-	}
+	ObjectBindingModel->AddDynamicBindingMenu(MenuBuilder, Possessable->DynamicBinding);
 }
 
 void FWidgetBlueprintEditor::AddSlotTrack( UPanelSlot* Slot )
@@ -2607,8 +2514,7 @@ void FWidgetBlueprintEditor::SyncSelectedWidgetsWithSequencerSelection(TArray<FG
 	TSet<FWidgetReference> SequencerSelectedWidgets;
 	for (FGuid Guid : ObjectGuids)
 	{
-		TArray<UObject*, TInlineAllocator<1>> BoundObjects;
-		AnimationSequence->LocateBoundObjects(Guid, UE::UniversalObjectLocator::FResolveParams(BindingContext), ActiveSequencer->GetSharedPlaybackState(), BoundObjects);
+		TArray<UObject*, TInlineAllocator<1>> BoundObjects = AnimationSequence->LocateBoundObjects(Guid, BindingContext);
 		if (BoundObjects.Num() == 0)
 		{
 			continue;
@@ -2682,7 +2588,7 @@ void FWidgetBlueprintEditor::UpdateTrackName(FGuid ObjectId)
 		}
 
 		TArray<UObject*, TInlineAllocator<1>> BoundObjects;
-		WidgetAnimation->LocateBoundObjects(ObjectId, UE::UniversalObjectLocator::FResolveParams(BindingContext), ActiveSequencer->GetSharedPlaybackState(), BoundObjects);
+		WidgetAnimation->LocateBoundObjects(ObjectId, BindingContext, BoundObjects);
 
 		if (BoundObjects.Num() > 0)
 		{
diff --git a/Engine/Source/Editor/UMGEditor/Public/WidgetBlueprint.h b/Engine/Source/Editor/UMGEditor/Public/WidgetBlueprint.h
--- a/Engine/Source/Editor/UMGEditor/Public/WidgetBlueprint.h
+++ b/Engine/Source/Editor/UMGEditor/Public/WidgetBlueprint.h
@@ -173,6 +173,32 @@ struct UMGEDITOR_API FDelegateEditorBinding
 	FDelegateRuntimeBinding ToRuntimeBinding(class UWidgetBlueprint* Blueprint) const;
 };
 
+
+/** Struct used only for loading old animations */
+USTRUCT()
+struct FWidgetAnimation_DEPRECATED
+{
+	GENERATED_USTRUCT_BODY()
+
+	UPROPERTY()
+	TObjectPtr<UMovieScene> MovieScene = nullptr;
+
+	UPROPERTY()
+	TArray<FWidgetAnimationBinding> AnimationBindings;
+
+	bool SerializeFromMismatchedTag(struct FPropertyTag const& Tag, FStructuredArchive::FSlot Slot);
+
+};
+
+template<>
+struct TStructOpsTypeTraits<FWidgetAnimation_DEPRECATED> : public TStructOpsTypeTraitsBase2<FWidgetAnimation_DEPRECATED>
+{
+	enum
+	{
+		WithStructuredSerializeFromMismatchedTag = true,
+	};
+};
+
 UENUM()
 enum class EWidgetSupportsDynamicCreation : uint8
 {
@@ -226,6 +252,9 @@ public:
 	UPROPERTY()
 	TArray< FDelegateEditorBinding > Bindings;
 
+	UPROPERTY()
+	TArray<FWidgetAnimation_DEPRECATED> AnimationData_DEPRECATED;
+
 	UPROPERTY()
 	TArray<TObjectPtr<UWidgetAnimation>> Animations;
 
diff --git a/Engine/Source/Editor/UMGEditor/Public/WidgetBlueprintEditor.h b/Engine/Source/Editor/UMGEditor/Public/WidgetBlueprintEditor.h
--- a/Engine/Source/Editor/UMGEditor/Public/WidgetBlueprintEditor.h
+++ b/Engine/Source/Editor/UMGEditor/Public/WidgetBlueprintEditor.h
@@ -435,9 +435,6 @@ private:
 	/** Update the name of a track to reflect changes in bindings */
 	void UpdateTrackName(FGuid ObjectId);
 
-	/* Called after changing dynamic binding properties*/
-	void OnFinishedChangingDynamicBindingProperties(const FPropertyChangedEvent& ChangeEvent, TSharedPtr<FStructOnScope> ValueStruct, FGuid ObjectId);
-
 private:
 	/** The preview scene that owns the preview GUI */
 	FPreviewScene PreviewScene;
diff --git a/Engine/Source/Runtime/LevelSequence/Private/LevelSequence.cpp b/Engine/Source/Runtime/LevelSequence/Private/LevelSequence.cpp
--- a/Engine/Source/Runtime/LevelSequence/Private/LevelSequence.cpp
+++ b/Engine/Source/Runtime/LevelSequence/Private/LevelSequence.cpp
@@ -49,12 +49,6 @@
 #include "Engine/AssetUserData.h"
 #include "Misc/App.h"
 #include "Misc/DateTime.h"
-#include "UObject/FortniteMainBranchObjectVersion.h"
-#include "Bindings/MovieSceneSpawnableDirectorBlueprintBinding.h"
-#include "Bindings/MovieSceneReplaceableDirectorBlueprintBinding.h"
-#include "Bindings/MovieSceneSpawnableActorBinding.h"
-#include "MovieSceneFolder.h"
-#include "Sections/MovieSceneBindingLifetimeSection.h"
 
 #include UE_INLINE_GENERATED_CPP_BY_NAME(LevelSequence)
 
@@ -64,7 +58,6 @@
 	#include "UObject/ObjectRedirector.h"
 
 ULevelSequence::FPostDuplicateEvent ULevelSequence::PostDuplicateEvent;
-ULevelSequence::FFixupDynamicBindingsEvent ULevelSequence::FixupDynamicBindingsEvent;
 
 #endif
 
@@ -420,161 +413,9 @@ void ULevelSequence::PostLoad()
 	}
 	ObjectReferences_DEPRECATED.Map.Empty();
 
-	if (GetLinkerCustomVersion(FFortniteMainBranchObjectVersion::GUID) < FFortniteMainBranchObjectVersion::LevelSequenceUpgradeDynamicBindings)
-	{
-		bool bConvertedDynamicBinding = ConvertOldSpawnables();
-
-		for (int32 Index = 0; Index < MovieScene->GetPossessableCount(); ++Index)
-		{
-			FMovieScenePossessable& Possessable = MovieScene->GetPossessable(Index);
-			if (Possessable.DynamicBinding_DEPRECATED.Function)
-			{
-				bConvertedDynamicBinding = true;
-				ConvertDynamicBindingPossessable(Possessable);
-			}
-		}
-		if (bConvertedDynamicBinding && FixupDynamicBindingsEvent.IsBound())
-		{
-			FixupDynamicBindingsEvent.Broadcast(this);
-		}
-	}
-
 #endif
 }
 
-#if WITH_EDITOR
-
-bool ULevelSequence::ConvertOldSpawnables()
-{
-	bool bConvertedDynamicBinding = false;
-	while (MovieScene->GetSpawnableCount() > 0)
-	{
-		FMovieSceneSpawnable& Spawnable = MovieScene->GetSpawnable(0);
-
-		FMovieScenePossessable* CreatedPossessable = nullptr;
-
-		UObject* ObjectToConvert = Spawnable.GetObjectTemplate();
-
-		UClass* CustomBindingType = nullptr;
-		if (Spawnable.DynamicBinding_DEPRECATED.Function)
-		{
-			CustomBindingType = UMovieSceneSpawnableDirectorBlueprintBinding::StaticClass();
-		}
-		else
-		{
-			// Search through custom binding types to find one that best supports the template type
-			static TArray<const TSubclassOf<UMovieSceneCustomBinding>> CachedCustomBindingTypes;
-			static bool CustomBindingTypesCached = false;
-			if (!CustomBindingTypesCached)
-			{
-				CustomBindingTypesCached = true;
-				MovieSceneHelpers::GetPrioritySortedCustomBindingTypes(CachedCustomBindingTypes);
-			}
-
-			for (const TSubclassOf<UMovieSceneCustomBinding>& CandidateCustomBindingType : CachedCustomBindingTypes)
-			{
-				if (CandidateCustomBindingType && CandidateCustomBindingType->IsChildOf(UMovieSceneSpawnableBindingBase::StaticClass()) && CandidateCustomBindingType->GetDefaultObject<UMovieSceneCustomBinding>()->SupportsBindingCreationFromObject(Spawnable.GetObjectTemplate()))
-				{
-					CustomBindingType = CandidateCustomBindingType;
-					break;
-				}
-			}
-		}
-
-		if (!CustomBindingType)
-		{
-			UE_LOG(LogLevelSequence, Warning, TEXT("Could not upgrade Spawnable '%s' with ID '%s'"), *Spawnable.GetName(), *Spawnable.GetGuid().ToString());
-			break;
-		}
-
-		UMovieSceneCustomBinding* NewCustomBinding = CustomBindingType->GetDefaultObject<UMovieSceneCustomBinding>()->CreateNewCustomBinding(ObjectToConvert, *MovieScene);
-
-		if (!NewCustomBinding)
-		{
-			UE_LOG(LogLevelSequence, Warning, TEXT("Could not upgrade Spawnable '%s' with ID '%s'"), *Spawnable.GetName(), *Spawnable.GetGuid().ToString());
-			break;
-		}
-
-		if (UMovieSceneSpawnableDirectorBlueprintBinding* DirectorBlueprintBinding = Cast<UMovieSceneSpawnableDirectorBlueprintBinding>(NewCustomBinding))
-		{
-			// Copy over the binding info
-			DirectorBlueprintBinding->DynamicBinding = Spawnable.DynamicBinding_DEPRECATED;
-			bConvertedDynamicBinding = true;
-		}
-		else if (UMovieSceneSpawnableActorBinding* SpawnableActorBinding = Cast<UMovieSceneSpawnableActorBinding>(NewCustomBinding))
-		{
-			SpawnableActorBinding->bNetAddressableName = Spawnable.bNetAddressableName;
-			SpawnableActorBinding->LevelName = Spawnable.LevelName;
-			SpawnableActorBinding->bContinuouslyRespawn = Spawnable.bContinuouslyRespawn;
-			SpawnableActorBinding->SpawnOwnership = Spawnable.GetSpawnOwnership();
-		}
-
-		FString PossessableName = Spawnable.GetName();
-		FGuid SpawnableGuid = Spawnable.GetGuid();
-
-		FMovieScenePossessable NewPossessable(PossessableName, NewCustomBinding->GetBoundObjectClass());
-		// Steal guid
-		NewPossessable.SetGuid(SpawnableGuid);
-
-		if (FMovieSceneBinding* SpawnableBinding = MovieScene->FindBinding(SpawnableGuid))
-		{
-			// Copy binding and track references to be tied to the new possessable
-			FMovieSceneBinding PossessableBinding = *SpawnableBinding;
-			// Add the custom binding. We use the spawnable binding here since it won't have a binding reference yet, and we need to steal the id
-			BindingReferences.AddOrReplaceBinding(SpawnableGuid, NewCustomBinding, 0);
-
-			// Remove the spawnable and all its' sub tracks
-			if (MovieScene->RemoveSpawnable(SpawnableGuid))
-			{
-				// Add the new possessable with the copied binding
-				MovieScene->AddPossessable(NewPossessable, PossessableBinding);
-			}
-		}
-	}
-	return bConvertedDynamicBinding;
-}
-
-void ULevelSequence::ConvertDynamicBindingPossessable(FMovieScenePossessable& Possessable)
-{
-	UMovieSceneReplaceableDirectorBlueprintBinding* NewCustomBinding = nullptr;
-	
-	const FName InstancedBindingName = MakeUniqueObjectName(MovieScene, UObject::StaticClass(), *FString(Possessable.GetName() + TEXT("_CustomBinding")));
-	NewCustomBinding = NewObject<UMovieSceneReplaceableDirectorBlueprintBinding>(MovieScene, UMovieSceneReplaceableDirectorBlueprintBinding::StaticClass(), InstancedBindingName, RF_Transactional);
-	
-	if (!NewCustomBinding)
-	{
-		return;
-	}
-
-	if (UMovieSceneReplaceableDirectorBlueprintBinding* DirectorBlueprintBinding = Cast<UMovieSceneReplaceableDirectorBlueprintBinding>(NewCustomBinding))
-	{
-		// Copy over the binding info
-		DirectorBlueprintBinding->DynamicBinding = Possessable.DynamicBinding_DEPRECATED;
-		Possessable.DynamicBinding_DEPRECATED = FMovieSceneDynamicBinding();
-	}
-
-	// Replace the current binding with the new one. We call RemoveBinding first because if there were multiple bindings for this track,
-	// they would have been overridden with the Dynamic Binding anyway, and so we ensure that stays the same by keeping only one binding
-	BindingReferences.RemoveBinding(Possessable.GetGuid());
-	BindingReferences.AddOrReplaceBinding(Possessable.GetGuid(), NewCustomBinding, 0);
-
-	// Add a binding lifetime track if not present
-	UMovieSceneBindingLifetimeTrack* BindingLifetimeTrack = Cast<UMovieSceneBindingLifetimeTrack>(MovieScene->FindTrack(UMovieSceneBindingLifetimeTrack::StaticClass(), Possessable.GetGuid(), NAME_None));
-	if (!BindingLifetimeTrack)
-	{
-		BindingLifetimeTrack = Cast<UMovieSceneBindingLifetimeTrack>(MovieScene->AddTrack(UMovieSceneBindingLifetimeTrack::StaticClass(), Possessable.GetGuid()));
-	}
-
-	if (BindingLifetimeTrack && BindingLifetimeTrack->GetAllSections().IsEmpty())
-	{
-		UMovieSceneBindingLifetimeSection* BindingLifetimeSection = Cast<UMovieSceneBindingLifetimeSection>(BindingLifetimeTrack->CreateNewSection());
-		BindingLifetimeSection->SetRange(TRange<FFrameNumber>::All());
-		BindingLifetimeTrack->AddSection(*BindingLifetimeSection);
-	}
-}
-
-#endif
-
 #if WITH_EDITORONLY_DATA
 void ULevelSequence::DeclareConstructClasses(TArray<FTopLevelAssetPath>& OutConstructClasses, const UClass* SpecificSubclass)
 {
@@ -627,7 +468,7 @@ void ULevelSequence::LocateBoundObjects(const FGuid& ObjectId, UObject* Context,
 		InResolveBindingParams.WorldPartitionResolveData ? InResolveBindingParams.WorldPartitionResolveData->SourceWorldAssetPath : InResolveBindingParams.StreamedLevelAssetPath
 		);
 
-	LocateBoundObjects(ObjectId, ResolveParams, nullptr, OutObjects);
+	LocateBoundObjects(ObjectId, ResolveParams, OutObjects);
 }
 
 FGuid ULevelSequence::FindBindingFromObject(UObject* InObject, UObject* Context) const
@@ -840,6 +681,12 @@ FGuid ULevelSequence::FindOrAddBinding(UObject* InObject)
 		{
 			ChildPossessable->SetParent(ParentGuid, MovieScene);
 		}
+
+		FMovieSceneSpawnable* ParentSpawnable = MovieScene->FindSpawnable(ParentGuid);
+		if (ParentSpawnable)
+		{
+			ParentSpawnable->AddChildPossessable(NewGuid);
+		}
 	}
 
 	BindPossessableObject(NewGuid, *InObject, BindingContext);
@@ -860,16 +707,36 @@ FGuid ULevelSequence::CreateSpawnable(UObject* ObjectToSpawn)
 		return FGuid();
 	}
 
-	FGuid NewGuid = MovieSceneHelpers::TryCreateCustomSpawnableBinding(this, ObjectToSpawn);
+	TArray<TSharedRef<IMovieSceneObjectSpawner>> ObjectSpawners;
+
+	// In order to create a spawnable, we have to instantiate all the relevant object spawners for level sequences, and try to create a spawnable from each
+	FLevelSequenceModule& LevelSequenceModule = FModuleManager::LoadModuleChecked<FLevelSequenceModule>("LevelSequence");
+	LevelSequenceModule.GenerateObjectSpawners(ObjectSpawners);
 
-	UMovieSceneSpawnTrack* NewSpawnTrack = MovieScene->AddTrack<UMovieSceneSpawnTrack>(NewGuid);
-	if (NewSpawnTrack)
+	// The first object spawner to return a valid result will win
+	for (TSharedRef<IMovieSceneObjectSpawner> Spawner : ObjectSpawners)
 	{
-		NewSpawnTrack->Modify();
+		TValueOrError<FNewSpawnable, FText> Result = Spawner->CreateNewSpawnableType(*ObjectToSpawn, *MovieScene, nullptr);
+		if (Result.IsValid())
+		{
+			FNewSpawnable& NewSpawnable = Result.GetValue();
+
+			NewSpawnable.Name = MovieSceneHelpers::MakeUniqueSpawnableName(MovieScene, NewSpawnable.Name);			
+
+			FGuid NewGuid = MovieScene->AddSpawnable(NewSpawnable.Name, *NewSpawnable.ObjectTemplate);
 
-		NewSpawnTrack->AddSection(*NewSpawnTrack->CreateNewSection());
+			UMovieSceneSpawnTrack* NewSpawnTrack = MovieScene->AddTrack<UMovieSceneSpawnTrack>(NewGuid);
+			if (NewSpawnTrack)
+			{
+				NewSpawnTrack->Modify();
+
+				NewSpawnTrack->AddSection(*NewSpawnTrack->CreateNewSection());
+			}
+			return NewGuid;
+		}
 	}
-	return NewGuid;
+
+	return FGuid();
 }
 
 #endif // WITH_EDITOR
@@ -959,21 +826,3 @@ const TArray<UAssetUserData*>* ULevelSequence::GetAssetUserDataArray() const
 	return &ToRawPtrTArrayUnsafe(AssetUserData);
 }
 
-void ULevelSequence::IterateDynamicBindings(const TSharedRef<UE::MovieScene::FSharedPlaybackState> SharedPlaybackState, TFunction<void(const FGuid&, FMovieSceneDynamicBinding&)> InCallback)
-{
-	for (FMovieSceneBindingReference& BindingReference : BindingReferences.GetAllReferences())
-	{
-		if (BindingReference.CustomBinding)
-		{
-			if (UMovieSceneReplaceableDirectorBlueprintBinding* ReplaceableDirectorBlueprintBinding = Cast<UMovieSceneReplaceableDirectorBlueprintBinding>(BindingReference.CustomBinding))
-			{
-				InCallback(BindingReference.ID, ReplaceableDirectorBlueprintBinding->DynamicBinding);
-			}
-
-			if (UMovieSceneSpawnableDirectorBlueprintBinding* SpawnableDirectorBlueprintBinding = Cast<UMovieSceneSpawnableDirectorBlueprintBinding>(BindingReference.CustomBinding->AsSpawnable(SharedPlaybackState)))
-			{
-				InCallback(BindingReference.ID, SpawnableDirectorBlueprintBinding->DynamicBinding);
-			}
-		}
-	}
-}
diff --git a/Engine/Source/Runtime/LevelSequence/Private/LevelSequencePlayer.cpp b/Engine/Source/Runtime/LevelSequence/Private/LevelSequencePlayer.cpp
--- a/Engine/Source/Runtime/LevelSequence/Private/LevelSequencePlayer.cpp
+++ b/Engine/Source/Runtime/LevelSequence/Private/LevelSequencePlayer.cpp
@@ -128,7 +128,7 @@ void ULevelSequencePlayer::ResolveBoundObjects(UE::UniversalObjectLocator::FReso
 	{
 		if (const FMovieSceneBindingReferences* BindingReferences = InSequence.GetBindingReferences())
 		{
-			FMovieSceneBindingResolveParams BindingResolveParams{ &InSequence, InBindingId, SequenceID, LocatorResolveParams.Context };
+			FMovieSceneBindingResolveParams BindingResolveParams{ &InSequence, InBindingId, SequenceID };
 			if (LocatorResolveParams.ParameterBuffer == nullptr)
 			{
 				// Allocate temporary local buffer for this
@@ -152,13 +152,13 @@ void ULevelSequencePlayer::ResolveBoundObjects(UE::UniversalObjectLocator::FReso
 				TInlineResolveParameterBuffer<128> Buffer;
 				LocatorResolveParams.ParameterBuffer = &Buffer;
 				LocatorResolveParams.ParameterBuffer->AddParameter(FActorLocatorFragmentResolveParameter::ParameterType, WeakStreamingWorld.Get(), ContainerID, SourceAssetPath);
-				InSequence.LocateBoundObjects(InBindingId, LocatorResolveParams, GetSharedPlaybackState(), OutObjects);
+				InSequence.LocateBoundObjects(InBindingId, LocatorResolveParams, OutObjects);
 				LocatorResolveParams.ParameterBuffer = nullptr;
 			}
 			else
 			{
 				LocatorResolveParams.ParameterBuffer->AddParameter(FActorLocatorFragmentResolveParameter::ParameterType, WeakStreamingWorld.Get(), ContainerID, SourceAssetPath);
-				InSequence.LocateBoundObjects(InBindingId, LocatorResolveParams, GetSharedPlaybackState(), OutObjects);
+				InSequence.LocateBoundObjects(InBindingId, LocatorResolveParams, OutObjects);
 			}
 		}
 	}
diff --git a/Engine/Source/Runtime/LevelSequence/Public/LevelSequence.h b/Engine/Source/Runtime/LevelSequence/Public/LevelSequence.h
--- a/Engine/Source/Runtime/LevelSequence/Public/LevelSequence.h
+++ b/Engine/Source/Runtime/LevelSequence/Public/LevelSequence.h
@@ -85,11 +85,6 @@ public:
 
 	DECLARE_DELEGATE_RetVal_OneParam(void, FPostDuplicateEvent, ULevelSequence*);
 	static LEVELSEQUENCE_API FPostDuplicateEvent PostDuplicateEvent;
-
-	DECLARE_MULTICAST_DELEGATE_OneParam(FFixupDynamicBindingsEvent, ULevelSequence*);
-	static LEVELSEQUENCE_API FFixupDynamicBindingsEvent FixupDynamicBindingsEvent;
-
-
 #endif
 
 	LEVELSEQUENCE_API virtual void PostDuplicate(bool bDuplicateForPIE) override;
@@ -102,8 +97,6 @@ public:
 	UE_DEPRECATED(5.4, "Use the base class LocateBoundObjects()")
 	LEVELSEQUENCE_API void LocateBoundObjects(const FGuid& ObjectId, UObject* Context, const FLevelSequenceBindingReference::FResolveBindingParams& InResolveBindingParams, TArray<UObject*, TInlineAllocator<1>>& OutObjects) const;
 
-	LEVELSEQUENCE_API virtual void IterateDynamicBindings(const TSharedRef<UE::MovieScene::FSharedPlaybackState> SharedPlaybackState, TFunction<void(const FGuid&, FMovieSceneDynamicBinding&)> InCallback) override;
-
 #if WITH_EDITOR
 
 
@@ -132,14 +125,6 @@ protected:
 	 */
 	LEVELSEQUENCE_API void OnDirectorRecompiled(UBlueprint*);
 
-private:
-#if WITH_EDITOR
-	// Used to convert any old spawnables in this movie scene to use the custom binding system
-	bool ConvertOldSpawnables();
-	// Used to convert possessables using deprecated 'DynamicBinding' to a custom director blueprint binding
-	void ConvertDynamicBindingPossessable(FMovieScenePossessable& Possessable);
-#endif
-
 #endif // WITH_EDITOR
 
 protected:
diff --git a/Engine/Source/Runtime/MovieScene/Private/Bindings/MovieSceneReplaceableBinding.cpp b/Engine/Source/Runtime/MovieScene/Private/Bindings/MovieSceneReplaceableBinding.cpp
--- a/Engine/Source/Runtime/MovieScene/Private/Bindings/MovieSceneReplaceableBinding.cpp
+++ b/Engine/Source/Runtime/MovieScene/Private/Bindings/MovieSceneReplaceableBinding.cpp
@@ -47,7 +47,7 @@ const FSlateBrush* UMovieSceneReplaceableBindingBase::GetBindingTrackCustomIconO
 
 FText UMovieSceneReplaceableBindingBase::GetBindingTrackIconTooltip() const
 {
-	return LOCTEXT("CustomReplaceableTooltip", "This item is dynamically bound at runtime, and may spawn a preview object in Editor within Sequencer");
+	return LOCTEXT("CustomReplaceableTooltip", "This item is dynamically bound at runtime, and spawns a preview object in Editor within Sequencer");
 }
 
 
@@ -156,9 +156,14 @@ UMovieSceneSpawnableBindingBase* UMovieSceneReplaceableBindingBase::CreateInnerS
 
 UMovieSceneCustomBinding* UMovieSceneReplaceableBindingBase::CreateNewCustomBinding(UObject* SourceObject, UMovieScene& OwnerMovieScene)
 {
+	if (!SourceObject)
+	{
+		return nullptr;
+	}
+
 	UMovieSceneReplaceableBindingBase* NewCustomBinding = nullptr;
 
-	const FName TemplateName = MakeUniqueObjectName(&OwnerMovieScene, UObject::StaticClass(), SourceObject ? SourceObject->GetFName() : GetClass()->GetFName());
+	const FName TemplateName = MakeUniqueObjectName(&OwnerMovieScene, UObject::StaticClass(), SourceObject->GetFName());
 	const FName InstancedBindingName = MakeUniqueObjectName(&OwnerMovieScene, UObject::StaticClass(), *FString(TemplateName.ToString() + TEXT("_CustomBinding")));
 
 	NewCustomBinding = NewObject<UMovieSceneReplaceableBindingBase>(&OwnerMovieScene, GetClass(), InstancedBindingName, RF_Transactional);
diff --git a/Engine/Source/Runtime/MovieScene/Private/Bindings/MovieSceneSpawnableBinding.cpp b/Engine/Source/Runtime/MovieScene/Private/Bindings/MovieSceneSpawnableBinding.cpp
--- a/Engine/Source/Runtime/MovieScene/Private/Bindings/MovieSceneSpawnableBinding.cpp
+++ b/Engine/Source/Runtime/MovieScene/Private/Bindings/MovieSceneSpawnableBinding.cpp
@@ -14,10 +14,10 @@
 #include "Internationalization/Internationalization.h"
 #include "Tracks/MovieSceneSpawnTrack.h"
 #include "Sections/MovieSceneBoolSection.h"
-#include "MovieSceneSequence.h"
 
 #define LOCTEXT_NAMESPACE "FPossessableModel"
 
+
 UObject* UMovieSceneSpawnableBindingBase::SpawnObject(const FGuid& BindingId, int32 BindingIndex, UMovieScene& MovieScene, FMovieSceneSequenceIDRef TemplateID, TSharedRef<const UE::MovieScene::FSharedPlaybackState> SharedPlaybackState)
 {
 	UWorld* WorldContext = GetWorldContext(SharedPlaybackState);
@@ -59,65 +59,6 @@ UObject* UMovieSceneSpawnableBindingBase::SpawnObject(const FGuid& BindingId, in
 	}
 #endif
 
-	// If have spawned an actor, do some actor-specific setup
-	if (AActor* SpawnedActor = Cast<AActor>(SpawnedObject))
-	{
-		// Ensure this spawnable is not a preview actor. Preview actors will not have BeginPlay() called on them.
-#if WITH_EDITOR
-		SpawnedActor->bIsEditorPreviewActor = false;
-#endif
-
-		static const FName SequencerActorTag(TEXT("SequencerActor"));
-		// tag this actor so we know it was spawned by sequencer
-		SpawnedActor->Tags.AddUnique(SequencerActorTag);
-
-#if WITH_EDITOR
-		if (GIsEditor)
-		{
-			// Explicitly set RF_Transactional on spawned actors so we can undo/redo properties on them.
-			// This particular UObject will be marked RF_Transactional by the caller, but we need to set it on the components.
-
-			for (UActorComponent* Component : SpawnedActor->GetComponents())
-			{
-				if (Component)
-				{
-					Component->SetFlags(RF_Transactional);
-				}
-			}
-		}
-#endif
-
-#if WITH_EDITOR
-		// Don't set the actor label in PIE as this requires flushing async loading.
-		if (WorldContext->WorldType == EWorldType::Editor)
-		{
-			FString BindingName = GetDesiredBindingName();
-			FString ActorLabel = !BindingName.IsEmpty() ? BindingName : SpawnName.ToString();
-
-			if (FMovieScenePossessable* Possessable = MovieScene.FindPossessable(BindingId))
-			{
-				ActorLabel = Possessable->GetName();
-
-				if (UMovieSceneSequence* Sequence = MovieScene.GetTypedOuter<UMovieSceneSequence>())
-				{
-					if (const FMovieSceneBindingReferences* BindingReferences = Sequence->GetBindingReferences())
-					{
-						if (BindingReferences->GetReferences(BindingId).Num() > 1)
-						{
-							// If there are multiple bound objects, use the Object Template actor label instead of the possessable name
-							if (AActor* ActorTemplate = Cast<AActor>(GetObjectTemplate()))
-							{
-								ActorLabel = ActorTemplate->GetActorLabel();
-							}
-						}
-					}
-				}
-			}
-			SpawnedActor->SetActorLabel(ActorLabel);
-		}
-#endif
-	}
-
 	// Allows derived classes to perform post-spawn logic such as mesh setup on actors.
 	PostSpawnObject(SpawnedObject, WorldContext, BindingId, BindingIndex, MovieScene, TemplateID, SharedPlaybackState);
 
@@ -136,19 +77,6 @@ void UMovieSceneSpawnableBindingBase::DestroySpawnedObject(UObject* Object)
 		// Explicitly remove RF_Transactional on spawned objects since we don't want to transact spawn/destroy events
 		Object->ClearFlags(RF_Transactional);
 	}
-	if (AActor* Actor = Cast<AActor>(Object))
-	{
-		// Explicitly remove RF_Transactional on spawned actors since we don't want to trasact spawn/destroy events
-		// This particular UObject will have RF_Transactional cleared by the caller, but we need to cleared it on the components.
-		for (UActorComponent* Component : Actor->GetComponents())
-		{
-			if (Component)
-			{
-				Component->ClearFlags(RF_Transactional);
-			}
-		}
-	}
-
 #endif
 
 	DestroySpawnedObjectInternal(Object);
@@ -203,7 +131,7 @@ const FSlateBrush* UMovieSceneSpawnableBindingBase::GetBindingTrackCustomIconOve
 
 FText UMovieSceneSpawnableBindingBase::GetBindingTrackIconTooltip() const
 {
-	return LOCTEXT("CustomSpawnableTooltip", "This item is spawned by sequencer by a custom spawnable binding according to this object's spawn track.");
+	return LOCTEXT("CustomSpawnableTooltip", "This item is spawned by sequencer by a custom spawnable binding according to this object's binding lifetime track.");
 }
 
 #endif
