Sequencer: Multiple objects bound to the same track are now required to be the same binding type in the Binding Properties UI. This simplifies the UX and removes a case that is likely not a real use case.

diff --git a/Engine/Plugins/MovieScene/LevelSequenceEditor/Source/LevelSequenceEditor/Private/LevelSequenceEditorSubsystem.cpp b/Engine/Plugins/MovieScene/LevelSequenceEditor/Source/LevelSequenceEditor/Private/LevelSequenceEditorSubsystem.cpp
--- a/Engine/Plugins/MovieScene/LevelSequenceEditor/Source/LevelSequenceEditor/Private/LevelSequenceEditorSubsystem.cpp
+++ b/Engine/Plugins/MovieScene/LevelSequenceEditor/Source/LevelSequenceEditor/Private/LevelSequenceEditorSubsystem.cpp
@@ -148,6 +148,18 @@ void FMovieSceneBindingPropertyInfoDetailCustomization::CustomizeChildren(TShare
 			}
 
 			BindingIndex = InStructPropertyHandle->GetArrayIndex();
+
+			uint32 NumBindings = 0;
+			if (TSharedPtr<IPropertyHandle> ParentHandle = InStructPropertyHandle->GetParentHandle())
+			{
+				ParentHandle->GetNumChildren(NumBindings);
+			}
+
+			if (NumBindings > 1)
+			{
+				bShowConvert = false;
+			}
+
 			int32 InitialIndex = 0;
 			BindingTypeNames.Add(MakeShared<FText>(LOCTEXT("BindingType_Possessable", "Possessable")));
 			TArrayView<const TSubclassOf<UMovieSceneCustomBinding>> SupportedBindingTypes = Sequencer->GetSupportedCustomBindingTypes();
@@ -2253,6 +2265,17 @@ void ULevelSequenceEditorSubsystem::OnFinishedChangingLocators(const FPropertyCh
 	}
 	if (FMovieSceneBindingReferences* BindingReferences = Sequence->GetBindingReferences())
 	{
+		// A bit hacky, but saves a complicated detail customization. If the change we've just made is to add a new entry, ensure the new entry is initialized
+		// to the same binding type as previous entries.
+		TArrayView<const FMovieSceneBindingReference> PreviousReferences = BindingReferences->GetReferences(ObjectBindingID);
+		if (PreviousReferences.Num() > 0 && PreviousReferences.Num() == BindingPropertyInfoList->Bindings.Num() - 1)
+		{
+			if (UMovieSceneCustomBinding* PreviousCustomBinding = BindingPropertyInfoList->Bindings[BindingPropertyInfoList->Bindings.Num() - 2].CustomBinding)
+			{
+				BindingPropertyInfoList->Bindings[BindingPropertyInfoList->Bindings.Num() - 1].CustomBinding = NewObject<UMovieSceneCustomBinding>(MovieScene, PreviousCustomBinding->GetClass());
+			}
+		}
+
 		MovieScene->Modify();
 		Sequence->Modify();
 		// Clear the previous binding
diff --git a/Engine/Plugins/MovieScene/LevelSequenceEditor/Source/LevelSequenceEditor/Private/Misc/LevelSequenceCustomization.cpp b/Engine/Plugins/MovieScene/LevelSequenceEditor/Source/LevelSequenceEditor/Private/Misc/LevelSequenceCustomization.cpp
--- a/Engine/Plugins/MovieScene/LevelSequenceEditor/Source/LevelSequenceEditor/Private/Misc/LevelSequenceCustomization.cpp
+++ b/Engine/Plugins/MovieScene/LevelSequenceEditor/Source/LevelSequenceEditor/Private/Misc/LevelSequenceCustomization.cpp
@@ -147,10 +147,6 @@ void FLevelSequenceCustomization::ExtendObjectBindingContextMenu(FMenuBuilder& M
 			bMultipleBindings = BindingReferences->GetReferences(ObjectBindingID).Num() > 1;
 			UE::UniversalObjectLocator::FResolveParams LocatorResolveParams(ResolutionContext);
 			FMovieSceneBindingResolveParams BindingResolveParams{ Sequencer->GetFocusedMovieSceneSequence(), ObjectBindingID, Sequencer->GetFocusedTemplateID(), ResolutionContext };
-			if (!BindingReferences->ResolveSingleBinding(BindingResolveParams, 0, LocatorResolveParams, Sequencer->GetSharedPlaybackState()))
-			{
-				bShowConvert = false;
-			}
 
 			// Can convert to possessable
 			int32 BindingIndex = 0;
diff --git a/Engine/Source/Editor/Sequencer/Private/Sequencer.cpp b/Engine/Source/Editor/Sequencer/Private/Sequencer.cpp
--- a/Engine/Source/Editor/Sequencer/Private/Sequencer.cpp
+++ b/Engine/Source/Editor/Sequencer/Private/Sequencer.cpp
@@ -6091,46 +6091,12 @@ FGuid FSequencer::AddEmptyBinding()
 		PossessableNames.Add(*MovieScene->GetPossessable(i).GetName());
 	}
 	FName PossessableName = FSequencerUtilities::GetUniqueName(TEXT("Empty Binding"), PossessableNames);
-	PossessableGuid = MovieScene->AddPossessable(PossessableName.ToString(), UObject::StaticClass());
 
-	if (PossessableGuid.IsValid())
-	{
-		Sequence->GetBindingReferences()->AddBinding(PossessableGuid, FUniversalObjectLocator());
-		OnAddBinding(PossessableGuid, MovieScene);
-
-		// Check if a folder is selected so we can add the actors to the selected folder.
-		TArray<UMovieSceneFolder*> SelectedParentFolders;
-		FString NewNodePath;
-		if (ViewModel->GetSelection()->Outliner.Num() > 0)
-		{
-			for (FViewModelPtr CurrentItem : ViewModel->GetSelection()->Outliner)
-			{
-				if (TSharedPtr<FFolderModel> Folder = CurrentItem->FindAncestorOfType<FFolderModel>(true))
-				{
-					SelectedParentFolders.Add(Folder->GetFolder());
-
-					// The first valid folder we find will be used to put the new binding into, so it's the node that we
-					// want to know the path from.
-					if (NewNodePath.Len() == 0)
-					{
-						// Add an extra delimiter (".") as we know that the new objects will be appended onto the end of this.
-						NewNodePath = FString::Printf(TEXT("%s."), *IOutlinerExtension::GetPathName(*Folder));
-
-						// Make sure the folder is expanded too so that adding objects to hidden folders become visible.
-						Folder->SetExpansion(true);
-					}
-				}
-			}
-		}
-
-		// Add the possessable as child of the first selected folder
-		if (SelectedParentFolders.Num() > 0)
-		{
-			SelectedParentFolders[0]->Modify();
-			SelectedParentFolders[0]->AddChildObjectBinding(PossessableGuid);
-		}
-	}
-	
+	UE::Sequencer::FCreateBindingParams CreateBindingParams;
+	CreateBindingParams.bAllowCustomBinding = false;
+	CreateBindingParams.BindingNameOverride = PossessableName.ToString();
+	CreateBindingParams.bAllowEmptyBinding = true;
+	PossessableGuid = FSequencerUtilities::CreateOrReplaceBinding(SharedThis(this), nullptr, CreateBindingParams);
 	RefreshTree();
 
 	return PossessableGuid;
diff --git a/Engine/Source/Editor/Sequencer/Private/SequencerUtilities.cpp b/Engine/Source/Editor/Sequencer/Private/SequencerUtilities.cpp
--- a/Engine/Source/Editor/Sequencer/Private/SequencerUtilities.cpp
+++ b/Engine/Source/Editor/Sequencer/Private/SequencerUtilities.cpp
@@ -1120,11 +1120,6 @@ FMovieScenePossessable* FSequencerUtilities::ConvertToPossessable(TSharedRef<ISe
 		ObjectToConvert = MovieSceneHelpers::GetObjectTemplate(Sequence, BindingGuid, Sequencer->GetSharedPlaybackState(), BindingIndex);
 	}
 
-	if (!ObjectToConvert)
-	{
-		return CreatedPossessable;
-	}
-
 	AActor* SpawnableActorTemplate = Cast<AActor>(ObjectToConvert);
 
 	TMap<TWeakObjectPtr<AActor>, FTransform> AttachedChildTransforms;
@@ -1203,11 +1198,6 @@ FMovieScenePossessable* FSequencerUtilities::ConvertToPossessable(TSharedRef<ISe
 		PossessedObject = PossessedActor;
 	}
 
-	if (!PossessedObject)
-	{
-		return CreatedPossessable;
-	}
-
 	Sequence->Modify();
 	MovieScene->Modify();
 
@@ -1215,9 +1205,10 @@ FMovieScenePossessable* FSequencerUtilities::ConvertToPossessable(TSharedRef<ISe
 	CreateBindingParams.ReplacementGuid = BindingGuid;
 	CreateBindingParams.BindingIndex = BindingIndex;
 	CreateBindingParams.bAllowCustomBinding = false;
+	CreateBindingParams.bAllowEmptyBinding = PossessedObject == nullptr;
 
 	// Create or replace the binding
-	FGuid NewPossessableGuid = CreateBinding(Sequencer, *PossessedObject, CreateBindingParams);
+	FGuid NewPossessableGuid = CreateOrReplaceBinding(Sequencer, Sequence, PossessedObject, CreateBindingParams);
 
 	TArrayView<const FMovieSceneBindingReference> BindingReferences = Sequence->GetBindingReferences()->GetReferences(BindingGuid);
 	bool bAnySpawnablesLeft = Algo::AnyOf(BindingReferences, [](const FMovieSceneBindingReference& BindingReference) { return BindingReference.CustomBinding && BindingReference.CustomBinding->IsA<UMovieSceneSpawnableBindingBase>(); });
@@ -1354,12 +1345,6 @@ FMovieScenePossessable* FSequencerUtilities::ConvertToCustomBinding(TSharedRef<I
 
 	bool bConvertFromPossessable = !bConvertFromSpawnable && !BindingReferences->GetCustomBinding(BindingGuid, BindingIndex);
 
-	if (!ObjectToConvert)
-	{
-		return CreatedPossessable;
-	}
-
-	
 	UMovieSceneCustomBinding* NewCustomBinding = nullptr;
 	if (PreviousBindingReference)
 	{
@@ -1385,9 +1370,10 @@ FMovieScenePossessable* FSequencerUtilities::ConvertToCustomBinding(TSharedRef<I
 	CreateBindingParams.BindingNameOverride = NewCustomBinding->GetDesiredBindingName();
 	CreateBindingParams.CustomBinding = NewCustomBinding;
 	CreateBindingParams.bSetupDefaults = false;
+	CreateBindingParams.bAllowEmptyBinding = !ObjectToConvert;
 
 	// Create or replace the binding
-	FGuid NewPossessableGuid = CreateBinding(Sequencer, *ObjectToConvert, CreateBindingParams);
+	FGuid NewPossessableGuid = CreateOrReplaceBinding(Sequencer, ObjectToConvert, CreateBindingParams);
 
 	TArrayView<const FMovieSceneBindingReference> BindingReferencesForGuid = BindingReferences->GetReferences(BindingGuid);
 	bool bAnySpawnablesLeft = Algo::AnyOf(BindingReferencesForGuid, [](const FMovieSceneBindingReference& BindingReference) { return BindingReference.CustomBinding && BindingReference.CustomBinding->IsA<UMovieSceneSpawnableBindingBase>(); });
@@ -2643,14 +2629,14 @@ FGuid CreateGenericBinding(TSharedPtr<ISequencer> Sequencer, UMovieSceneSequence
 		}
 	}
 
-	if (!InObject)
+	if (!InObject && !InParams.bAllowEmptyBinding)
 	{
 		return FGuid();
 	}
 
 
 	// If no custom bindings support this object type, but InParams.bSpawnable is true, attempt to make an old-style spawnable.
-	if (bSpawnable && Sequencer)
+	if (InObject && bSpawnable && Sequencer)
 	{
 		NewBindingID = FSequencerUtilities::MakeNewSpawnable(Sequencer.ToSharedRef(), *InObject, InParams.ActorFactory);
 		if (NewBindingID.IsValid())
@@ -2664,6 +2650,11 @@ FGuid CreateGenericBinding(TSharedPtr<ISequencer> Sequencer, UMovieSceneSequence
 	TArray<TPair<UObject*, FString>> ObjectsToPossess;
 
 	// Build up the list of child->parent bindings required for this object
+	if (!InObject)
+	{
+		ObjectsToPossess.Add(MakeTuple(InObject, InParams.BindingNameOverride.IsEmpty() ? TEXT("Empty Binding") : InParams.BindingNameOverride));
+	}
+	else
 	{
 		UObject* CurrentObject = InObject;
 		while (CurrentObject)
@@ -2738,15 +2729,15 @@ FGuid CreateGenericBinding(TSharedPtr<ISequencer> Sequencer, UMovieSceneSequence
 		}
 		if (!NewPossessable)
 		{
-			NewID = OwnerMovieScene->AddPossessable(CurrentName, CurrentObject->GetClass());
+			NewID = OwnerMovieScene->AddPossessable(CurrentName, CurrentObject ? CurrentObject->GetClass() : UObject::StaticClass());
 			NewPossessable = OwnerMovieScene->FindPossessable(NewID);
 		}
 
 		// If we're not trying to replace a binding, and the object is a spawnable, try and bind to that first
-		if (InParams.ReplacementGuid.IsValid() || (Sequencer && !NewPossessable->BindSpawnableObject(Sequencer->GetFocusedTemplateID(), CurrentObject, Sequencer->GetSharedPlaybackState())))
+		if (InParams.ReplacementGuid.IsValid() || (Sequencer && (!CurrentObject || !NewPossessable->BindSpawnableObject(Sequencer->GetFocusedTemplateID(), CurrentObject, Sequencer->GetSharedPlaybackState()))))
 		{
 			FUniversalObjectLocator Locator;
-			if (!OwnerSequence->MakeLocatorForObject(CurrentObject, Context, Locator) || Locator.IsEmpty())
+			if (CurrentObject && (!OwnerSequence->MakeLocatorForObject(CurrentObject, Context, Locator) || Locator.IsEmpty()))
 			{
 				// Unable to possess this object
 				return FGuid();
@@ -4233,15 +4224,13 @@ bool FSequencerUtilities::CanConvertToCustomBinding(TSharedRef<ISequencer> Seque
 			{
 				UE::UniversalObjectLocator::FResolveParams LocatorResolveParams(ResolutionContext);
 				FMovieSceneBindingResolveParams BindingResolveParams{ Sequence, BindingGuid, Sequencer->GetFocusedTemplateID(), ResolutionContext };
-				if (UObject* CurrentBoundObject = BindingReferences->ResolveSingleBinding(BindingResolveParams, BindingIndex, LocatorResolveParams, Sequencer->GetSharedPlaybackState()))
+				UObject* CurrentBoundObject = BindingReferences->ResolveSingleBinding(BindingResolveParams, BindingIndex, LocatorResolveParams, Sequencer->GetSharedPlaybackState());
+				if (CustomBindingType
+					&& (!CurrentBindingReference->CustomBinding
+						|| CurrentBindingReference->CustomBinding->GetClass() != CustomBindingType)
+					&& CustomBindingType->GetDefaultObject<UMovieSceneCustomBinding>()->SupportsConversionFromBinding(*CurrentBindingReference, CurrentBoundObject))
 				{
-					if (CustomBindingType
-						&& (!CurrentBindingReference->CustomBinding
-							|| CurrentBindingReference->CustomBinding->GetClass() != CustomBindingType)
-						&& CustomBindingType->GetDefaultObject<UMovieSceneCustomBinding>()->SupportsConversionFromBinding(*CurrentBindingReference, CurrentBoundObject))
-					{
-						return true;
-					}
+					return true;
 				}
 			}
 		}
diff --git a/Engine/Source/Editor/Sequencer/Public/SequencerUtilities.h b/Engine/Source/Editor/Sequencer/Public/SequencerUtilities.h
--- a/Engine/Source/Editor/Sequencer/Public/SequencerUtilities.h
+++ b/Engine/Source/Editor/Sequencer/Public/SequencerUtilities.h
@@ -60,6 +60,8 @@ struct FCreateBindingParams
 	bool bReplaceable = false;
 	/* If true, will allow the creation of custom bindings if they support the object type.*/
 	bool bAllowCustomBinding = true;
+	/* If true, will allow the creation of an empty binding if a UObject* passed in is nullptr*/
+	bool bAllowEmptyBinding = false;
 	/* If set, will attempt to replace any existing possessable binding at the provided guid and binding index */
 	FGuid ReplacementGuid;
 	/* Optional BindingIndex used if bReplace is true to replace a specific possessable binding*/
diff --git a/Engine/Source/Runtime/MovieSceneTracks/Private/Bindings/MovieSceneSpawnableActorBinding.cpp b/Engine/Source/Runtime/MovieSceneTracks/Private/Bindings/MovieSceneSpawnableActorBinding.cpp
--- a/Engine/Source/Runtime/MovieSceneTracks/Private/Bindings/MovieSceneSpawnableActorBinding.cpp
+++ b/Engine/Source/Runtime/MovieSceneTracks/Private/Bindings/MovieSceneSpawnableActorBinding.cpp
@@ -346,9 +346,10 @@ bool UMovieSceneSpawnableActorBinding::SupportsBindingCreationFromObject(const U
 {
 	if (!SourceObject)
 	{
-		return false;
+		// In this case we would just make an empty binding
+		return true;
 	}
-	if (SourceObject->IsA<AActor>())
+	else if (SourceObject->IsA<AActor>())
 	{
 		return true;
 	}
@@ -375,14 +376,9 @@ bool UMovieSceneSpawnableActorBinding::SupportsBindingCreationFromObject(const U
 
 UMovieSceneCustomBinding* UMovieSceneSpawnableActorBinding::CreateNewCustomBinding(UObject* SourceObject, UMovieScene& OwnerMovieScene)
 {
-	if (!SourceObject)
-	{
-		return nullptr;
-	}
-
 	UMovieSceneSpawnableActorBinding* NewCustomBinding = nullptr;
 
-	const FName TemplateName = MakeUniqueObjectName(&OwnerMovieScene, UObject::StaticClass(), SourceObject->GetFName());
+	const FName TemplateName = MakeUniqueObjectName(&OwnerMovieScene, UObject::StaticClass(), SourceObject ? SourceObject->GetFName() : TEXT("Empty Binding"));
 	const FName InstancedBindingName = MakeUniqueObjectName(&OwnerMovieScene, UObject::StaticClass(), *FString(TemplateName.ToString() + TEXT("_CustomBinding")));
 
 	auto CreateBinding = [&]()
@@ -457,7 +453,7 @@ UMovieSceneCustomBinding* UMovieSceneSpawnableActorBinding::CreateNewCustomBindi
 
 	if (!NewCustomBinding)
 	{
-		if (UClass* InClass = Cast<UClass>(SourceObject))
+		if (UClass* InClass = Cast<UClass>(SourceObject ? SourceObject : AActor::StaticClass()))
 		{
 			if (!InClass->IsChildOf(AActor::StaticClass()))
 			{		
