Sequencer: Deprecate FSectionModel::OnUpdated
This event does not fall within the design methodology of the MVVM framework so is being removed.

diff --git a/Engine/Source/Editor/Sequencer/Private/MVVM/ViewModels/SectionModel.cpp b/Engine/Source/Editor/Sequencer/Private/MVVM/ViewModels/SectionModel.cpp
--- a/Engine/Source/Editor/Sequencer/Private/MVVM/ViewModels/SectionModel.cpp
+++ b/Engine/Source/Editor/Sequencer/Private/MVVM/ViewModels/SectionModel.cpp
@@ -48,11 +48,13 @@ FSectionModel::FSectionModel(UMovieSceneSection* InSection, TSharedPtr<ISequence
 	RegisterChildList(&ChannelList);
 }
 
+PRAGMA_DISABLE_DEPRECATION_WARNINGS
 FSectionModel::~FSectionModel()
 {
 	SectionInterface.Reset();
 	WeakSection.Reset();
 }
+PRAGMA_ENABLE_DEPRECATION_WARNINGS
 
 TRange<FFrameNumber> FSectionModel::GetRange() const
 {
@@ -129,7 +131,6 @@ void FSectionModel::OffsetLayerBar(FFrameNumber Amount)
 void FSectionModel::OnModifiedDirectly(UMovieSceneSignedObject*)
 {
 	UpdateCachedData();
-	OnUpdated.Broadcast(this);
 }
 
 void FSectionModel::UpdateCachedData()
diff --git a/Engine/Source/Editor/Sequencer/Public/MVVM/ViewModels/SectionModel.h b/Engine/Source/Editor/Sequencer/Public/MVVM/ViewModels/SectionModel.h
--- a/Engine/Source/Editor/Sequencer/Public/MVVM/ViewModels/SectionModel.h
+++ b/Engine/Source/Editor/Sequencer/Public/MVVM/ViewModels/SectionModel.h
@@ -58,6 +58,7 @@ public:
 
 	DECLARE_MULTICAST_DELEGATE_OneParam(FOnModelUpdated, FSectionModel*)
 
+	UE_DEPRECATED(5.5, "This member is no longer supported, please subscribe to UMovieSceneSignedObject::OnModifiedDirectly.")
 	FOnModelUpdated OnUpdated;
 
 	explicit FSectionModel(UMovieSceneSection* InSection, TSharedPtr<ISequencerSection> InSectionInterface);
