Sequencer: Added header geometry and color tint to SectionPainter

diff --git a/Engine/Source/Editor/Sequencer/Private/SSequencerSection.cpp b/Engine/Source/Editor/Sequencer/Private/SSequencerSection.cpp
--- a/Engine/Source/Editor/Sequencer/Private/SSequencerSection.cpp
+++ b/Engine/Source/Editor/Sequencer/Private/SSequencerSection.cpp
@@ -240,13 +240,11 @@ struct FSequencerSectionPainterImpl : FSequencerSectionPainter
 		}
 		else
 		{
-			TSharedPtr<IGeometryExtension> Geometry = Outliner.ImplicitCast();
-			const float HeaderHeight = Geometry ? Geometry->GetVirtualGeometry().GetHeight() : 10.f;
-
-			FGeometry HeaderGeometry = ExpandedSectionGeometry.MakeChild(
-				FVector2D(ExpandedSectionGeometry.GetLocalSize().X, HeaderHeight),
+			FGeometry LocalHeaderGeometry = ExpandedSectionGeometry.MakeChild(
+				FVector2D(ExpandedSectionGeometry.GetLocalSize().X, HeaderGeometry.GetLocalSize().Y),
 				FSlateLayoutTransform()
 			);
+			const float HeaderHeight = LocalHeaderGeometry.GetLocalSize().Y;
 			FGeometry ContentsGeometry = ExpandedSectionGeometry.MakeChild(
 				FVector2D(ExpandedSectionGeometry.GetLocalSize().X, ExpandedSectionGeometry.GetLocalSize().Y - HeaderHeight),
 				FSlateLayoutTransform(FVector2D(0.f, HeaderHeight)) 
@@ -255,7 +253,7 @@ struct FSequencerSectionPainterImpl : FSequencerSectionPainter
 			FSlateDrawElement::MakeBox(
 				DrawElements,
 				LayerId,
-				HeaderGeometry.ToPaintGeometry(),
+				LocalHeaderGeometry.ToPaintGeometry(),
 				SectionHeaderBackgroundBrush,
 				DrawEffects,
 				BlendedTint
@@ -279,7 +277,7 @@ struct FSequencerSectionPainterImpl : FSequencerSectionPainter
 				FSlateDrawElement::MakeBox(
 					DrawElements,
 					++LayerId,
-					HeaderGeometry.ToPaintGeometry(HeaderGeometry.GetLocalSize() - FVector2f(2.f, 2.f), FSlateLayoutTransform(FVector2f(1.f, 1.f))),
+					LocalHeaderGeometry.ToPaintGeometry(LocalHeaderGeometry.GetLocalSize() - FVector2f(2.f, 2.f), FSlateLayoutTransform(FVector2f(1.f, 1.f))),
 					SectionHeaderSelectedSectionOverlay,
 					DrawEffects,
 					SelectionColor.GetValue().CopyWithNewOpacity(0.8f)
diff --git a/Engine/Source/Editor/Sequencer/Private/SequencerSectionPainter.cpp b/Engine/Source/Editor/Sequencer/Private/SequencerSectionPainter.cpp
--- a/Engine/Source/Editor/Sequencer/Private/SequencerSectionPainter.cpp
+++ b/Engine/Source/Editor/Sequencer/Private/SequencerSectionPainter.cpp
@@ -4,6 +4,8 @@
 #include "MovieSceneSection.h"
 #include "MovieSceneTrack.h"
 #include "MVVM/ViewModels/SectionModel.h"
+#include "MVVM/Extensions/IOutlinerExtension.h"
+#include "MVVM/Extensions/ITrackExtension.h"
 
 FSequencerSectionPainter::FSequencerSectionPainter(FSlateWindowElementList& OutDrawElements, const FGeometry& InSectionGeometry, TSharedPtr<UE::Sequencer::FSectionModel> InSection)
 	: SectionModel(InSection)
@@ -15,6 +17,14 @@ FSequencerSectionPainter::FSequencerSectionPainter(FSlateWindowElementList& OutD
 	, bIsSelected(false)
 	, GhostAlpha(1.f)
 {
+	HeaderGeometry = SectionGeometry.MakeChild(
+		FVector2f(
+			SectionGeometry.GetLocalSize().X,
+			InSection->GetLinkedOutlinerItem()->GetOutlinerSizing().Height
+		),
+		FSlateLayoutTransform()
+	);
+
 }
 
 FSequencerSectionPainter::~FSequencerSectionPainter()
@@ -47,3 +57,19 @@ FLinearColor FSequencerSectionPainter::BlendColor(FLinearColor InColor)
 	
 	return BaseColor * (1.f - Alpha) + InColor * Alpha;
 }
+
+FLinearColor FSequencerSectionPainter::GetSectionColor() const
+{
+	using namespace UE::Sequencer;
+
+	TSharedPtr<ITrackExtension> Track = SectionModel->FindAncestorOfType<ITrackExtension>();
+	UMovieSceneTrack* TrackObject = Track ? Track->GetTrack() : nullptr;
+
+	FLinearColor Tint = FLinearColor::White;
+	if (TrackObject)
+	{
+		return BlendColor(TrackObject->GetColorTint()).CopyWithNewOpacity(1.f);
+	}
+
+	return Tint;
+}
\ No newline at end of file
diff --git a/Engine/Source/Editor/Sequencer/Public/SequencerSectionPainter.h b/Engine/Source/Editor/Sequencer/Public/SequencerSectionPainter.h
--- a/Engine/Source/Editor/Sequencer/Public/SequencerSectionPainter.h
+++ b/Engine/Source/Editor/Sequencer/Public/SequencerSectionPainter.h
@@ -57,6 +57,8 @@ public:
 	/** Get the track that this painter is painting sections for */
 	UMovieSceneTrack* GetTrack() const;
 
+	FLinearColor GetSectionColor() const;
+
 	/** Blend the specified color with the default track color */
 	static FLinearColor BlendColor(FLinearColor InColor);
 
@@ -76,6 +78,9 @@ public:
 	/** The full geometry of the section. This is the width of the track area in the case of infinite sections */
 	FGeometry SectionGeometry;
 
+	/** The header geometry of the section. This is the same as SectionGeometry for collapsed sections, but just the height of the top row for expanded sections. */
+	FGeometry HeaderGeometry;
+
 	/** The full clipping rectangle for the section */
 	FSlateRect SectionClippingRect;
 	
