Sequencer: Fixed typo in parameter name

diff --git a/Engine/Source/Editor/Sequencer/Public/ISequencerSection.h b/Engine/Source/Editor/Sequencer/Public/ISequencerSection.h
--- a/Engine/Source/Editor/Sequencer/Public/ISequencerSection.h
+++ b/Engine/Source/Editor/Sequencer/Public/ISequencerSection.h
@@ -187,7 +187,7 @@ public:
 	/**
 	 * Create a custom channel model
 	 */
-	virtual TSharedPtr<UE::Sequencer::FChannelModel> ConstructChannelModel(FName InCategoryName, const FMovieSceneChannelHandle& InChannelHandle) const { return nullptr; }
+	virtual TSharedPtr<UE::Sequencer::FChannelModel> ConstructChannelModel(FName InChannelName, const FMovieSceneChannelHandle& InChannelHandle) const { return nullptr; }
 
 	/**
 	 * @return The height of the section
