Defer setting the timeline position until the next frame while scrubbing in Sequencer. This improves the frame rate while scrubbing.

Windows often sends multiple mouse move events within a frame. This change saves time by evaluating only the last position set each frame.

diff --git a/Engine/Source/Editor/Sequencer/Private/Sequencer.cpp b/Engine/Source/Editor/Sequencer/Private/Sequencer.cpp
--- a/Engine/Source/Editor/Sequencer/Private/Sequencer.cpp
+++ b/Engine/Source/Editor/Sequencer/Private/Sequencer.cpp
@@ -1001,6 +1001,15 @@ void FSequencer::Tick(float InDeltaTime)
 		}
 	}
 
+	if (PendingScrubPosition.IsSet())
+	{
+		const bool bEvaluate = false;
+		SetLocalTimeDirectly(PendingScrubPosition.GetValue(), bEvaluate);
+		PendingScrubPosition.Reset();
+
+		bNeedsEvaluate = true;
+	}
+
 	UpdateSubSequenceData();
 
 	// Tick all the tools we own as well
@@ -5408,6 +5417,12 @@ void FSequencer::OnScrubPositionChanged( FFrameTime NewScrubPosition, bool bScru
 	{
 		AutoScrubToTime(NewScrubPosition);
 	}
+	else if (bEvaluate)
+	{
+		// Evaluation can be expensive and we may receive multiple OnScrubPositionChanged events in
+		// a frame, so defer the evaluation until the next tick
+		PendingScrubPosition = NewScrubPosition;
+	}
 	else
 	{
 		SetLocalTimeDirectly(NewScrubPosition, bEvaluate);
diff --git a/Engine/Source/Editor/Sequencer/Private/Sequencer.h b/Engine/Source/Editor/Sequencer/Private/Sequencer.h
--- a/Engine/Source/Editor/Sequencer/Private/Sequencer.h
+++ b/Engine/Source/Editor/Sequencer/Private/Sequencer.h
@@ -1265,6 +1265,9 @@ private:
 
 	TOptional<FAutoScrubTarget> AutoScrubTarget;
 
+	/** The latest scrub position request, to be applied in the next tick */
+	TOptional<FFrameTime> PendingScrubPosition;
+
 	/** Zoom smoothing curves */
 	FCurveSequence ZoomAnimation;
 	FCurveHandle ZoomCurve;
