Sequencer: Only snap to marked frames if marked frames are currently visible

diff --git a/Engine/Source/Editor/Sequencer/Private/Tools/SequencerSnapField.cpp b/Engine/Source/Editor/Sequencer/Private/Tools/SequencerSnapField.cpp
--- a/Engine/Source/Editor/Sequencer/Private/Tools/SequencerSnapField.cpp
+++ b/Engine/Source/Editor/Sequencer/Private/Tools/SequencerSnapField.cpp
@@ -7,6 +7,7 @@
 #include "MovieSceneTimeHelpers.h"
 #include "MovieSceneSequence.h"
 #include "ISequencerSection.h"
+#include "SequencerSettings.h"
 #include "IKeyArea.h"
 #include "MVVM/ViewModels/ViewModel.h"
 #include "MVVM/Views/SOutlinerView.h"
@@ -148,16 +149,19 @@ void FSequencerSnapField::Initialize(const FSequencer& InSequencer, UE::Sequence
 		Visitor.Snaps.Emplace(FSnapPoint{ FSnapPoint::InOutRange, UE::MovieScene::DiscreteExclusiveUpper(SelectionRange) - 1});
 	}
 
-	// Add in the marked frames
-	for (const FMovieSceneMarkedFrame& MarkedFrame : InSequencer.GetFocusedMovieSceneSequence()->GetMovieScene()->GetMarkedFrames())
+	if (InSequencer.GetSequencerSettings()->GetShowMarkedFrames())
 	{
-		Visitor.Snaps.Emplace( FSnapPoint{ FSnapPoint::Mark, MarkedFrame.FrameNumber } );
-	}
+		// Add in the marked frames
+		for (const FMovieSceneMarkedFrame& MarkedFrame : InSequencer.GetFocusedMovieSceneSequence()->GetMovieScene()->GetMarkedFrames())
+		{
+			Visitor.Snaps.Emplace( FSnapPoint{ FSnapPoint::Mark, MarkedFrame.FrameNumber } );
+		}
 
-	// Add in the global marked frames
-	for (const FMovieSceneMarkedFrame& MarkedFrame : InSequencer.GetGlobalMarkedFrames())
-	{
-		Visitor.Snaps.Emplace(FSnapPoint{ FSnapPoint::Mark, MarkedFrame.FrameNumber });
+		// Add in the global marked frames
+		for (const FMovieSceneMarkedFrame& MarkedFrame : InSequencer.GetGlobalMarkedFrames())
+		{
+			Visitor.Snaps.Emplace(FSnapPoint{ FSnapPoint::Mark, MarkedFrame.FrameNumber });
+		}
 	}
 
 	// Copy the paged array to our linear array ready for final sorting
