Sequencer: Added time-warp capabilities to sequences, sub-sections and skeletal animation sections

diff --git a/Engine/Source/Editor/Sequencer/Private/Sequencer.cpp b/Engine/Source/Editor/Sequencer/Private/Sequencer.cpp
--- a/Engine/Source/Editor/Sequencer/Private/Sequencer.cpp
+++ b/Engine/Source/Editor/Sequencer/Private/Sequencer.cpp
@@ -29,6 +29,7 @@
 #include "MVVM/TrackModelStorageExtension.h"
 #include "MVVM/FolderModelStorageExtension.h"
 #include "Camera/PlayerCameraManager.h"
+#include "MovieSceneTransformTypes.h"
 #include "Misc/MessageDialog.h"
 #include "Misc/FeedbackContext.h"
 #include "Misc/ScopedSlowTask.h"
@@ -301,6 +302,8 @@ void FSequencer::InitSequencer(const FSequencerInitParams& InitParams, const TSh
 	ScrubStyle = InitParams.ViewParams.ScrubberStyle;
 	HostCapabilities = InitParams.HostCapabilities;
 
+	CurrentTimeBreadcrumbs = FMovieSceneTransformBreadcrumbs(EMovieSceneBreadcrumbMode::Dense);
+
 	SilentModeCount = 0;
 	bReadOnly = InitParams.ViewParams.bReadOnly;
 
@@ -708,9 +711,6 @@ FSequencer::FSequencer()
 	, LastViewRange(0.f, 5.f)
 	, ViewRangeBeforeZoom(TRange<double>::Empty())
 	, PlaybackState( EMovieScenePlayerStatus::Stopped )
-	, LocalLoopIndexOnBeginScrubbing(FMovieSceneTimeWarping::InvalidWarpCount)
-	, LocalLoopIndexOffsetDuringScrubbing(0)
-	, MaxLocalLoopIndex(FMovieSceneTimeWarping::InvalidWarpCount)
 	, bPerspectiveViewportPossessionEnabled( true )
 	, bPerspectiveViewportCameraCutEnabled( false )
 	, bIsEditingWithinLevelEditor( false )
@@ -944,7 +944,7 @@ void FSequencer::Tick(float InDeltaTime)
 	{
 		FQualifiedFrameTime CurrentTime = GetLocalTime();
 		FFrameTime Offset = (AutoscrubOffset.GetValue() * AutoScrollFactor) * CurrentTime.Rate;
-		SetLocalTimeLooped(CurrentTime.Time + Offset, RootToLocalLoopCounter);
+		SetLocalTimeLooped(CurrentTime.Time + Offset, CurrentTimeBreadcrumbs);
 	}
 
 	if (GetSelectionRange().IsEmpty() && GetLoopMode() == SLM_LoopSelectionRange)
@@ -958,10 +958,10 @@ void FSequencer::Tick(float InDeltaTime)
 
 		// Put the time into clamped local space
 		FFrameTime LocalTime;
-		FMovieSceneWarpCounter WarpCounter;
-		CalculateLocalTimeClamped(NewGlobalTime, RootToLocalTransform, LocalTime, WarpCounter);
+		FMovieSceneTransformBreadcrumbs Breadcrumbs;
+		CalculateLocalTimeClamped(NewGlobalTime, RootToLocalTransform, LocalTime, Breadcrumbs);
 
- 		SetLocalTimeLooped(LocalTime, WarpCounter);
+		SetLocalTimeLooped(LocalTime, Breadcrumbs);
 
 		if (GetPlaybackStatus() == EMovieScenePlayerStatus::Playing)
 		{
@@ -1151,8 +1151,17 @@ void FSequencer::ResetToNewRootSequence(UMovieSceneSequence& NewSequence)
 	RootTemplateInstance.Initialize(NewSequence, *this, CompiledDataManager);
 	RootTemplateInstance.EnableGlobalPreAnimatedStateCapture();
 
-	RootToLocalTransform = FMovieSceneSequenceTransform();
-	RootToLocalLoopCounter = FMovieSceneWarpCounter();
+	const FMovieSceneSequenceHierarchy* Hierarchy = CompiledDataManager->FindHierarchy(RootTemplateInstance.GetCompiledDataID());
+	if (Hierarchy)
+	{
+		RootToLocalTransform = Hierarchy->GetRootTransform();
+	}
+	else
+	{
+		RootToLocalTransform = FMovieSceneSequenceTransform();
+	}
+
+	CurrentTimeBreadcrumbs.Reset();
 
 	ResetPerMovieSceneData();
 	SequencerWidget->ResetBreadcrumbs();
@@ -1450,23 +1459,31 @@ void FSequencer::PopToSequenceInstance(FMovieSceneSequenceIDRef SequenceID)
 
 void FSequencer::UpdateSubSequenceData()
 {
+	using namespace UE::MovieScene;
+
 	SubSequenceRange = TRange<FFrameNumber>::Empty();
-	RootToLocalTransform = FMovieSceneSequenceTransform();
-	// else: we're scrubbing, and we don't want to increase/decrease the loop index quite yet,
-	// because that would mess up time transforms. This would be because the mouse would still be
-	// before/after the current loop, and therefore would already add/subtract more than a full
-	// loop's time to the current time, so we don't need the loop counter to change yet.
+
+	const FMovieSceneSequenceHierarchy* Hierarchy = CompiledDataManager->FindHierarchy(RootTemplateInstance.GetCompiledDataID());
+	if (Hierarchy)
+	{
+		RootToLocalTransform = RootTransform = Hierarchy->GetRootTransform();
+	}
+	else
+	{
+		RootToLocalTransform = RootTransform = FMovieSceneSequenceTransform();
+	}
 
 	// Find the parent sub section and set up the sub sequence range, if necessary
 	if (ActiveTemplateIDs.Num() <= 1)
 	{
-		// Ensure in this case we also reset the RootToLocalLoopCounter.
-		RootToLocalLoopCounter = FMovieSceneWarpCounter();
+		// Ensure in this case we also reset the CurrentTimeBreadcrumbs.
+		CurrentTimeBreadcrumbs.Reset();
 		return;
 	}
 
-	const FMovieSceneSequenceHierarchy& Hierarchy       = CompiledDataManager->GetHierarchyChecked(RootTemplateInstance.GetCompiledDataID());
-	const FMovieSceneSubSequenceData*   SubSequenceData = Hierarchy.FindSubData(ActiveTemplateIDs.Top());
+	check(Hierarchy);
+
+	const FMovieSceneSubSequenceData*   SubSequenceData = Hierarchy->FindSubData(ActiveTemplateIDs.Top());
 
 	if (SubSequenceData)
 	{
@@ -1474,133 +1491,73 @@ void FSequencer::UpdateSubSequenceData()
 		RootToLocalTransform = SubSequenceData->RootToSequenceTransform;
 
 		const bool bIsScrubbing = GetPlaybackStatus() == EMovieScenePlayerStatus::Scrubbing;
-		const bool bIsSubSequenceLooping = RootToLocalTransform.NestedTransforms.Num() > 0 && RootToLocalTransform.NestedTransforms.Last().IsLooping();
-		const bool bIsScrubbingLoopingSubSequence = bIsScrubbing && bIsSubSequenceLooping;
 
-		const FQualifiedFrameTime RootTime = GetGlobalTime();
-		int32 CurLoopIndexOffset = 0;
-		if (!bIsScrubbingLoopingSubSequence)
-		{
-			FFrameTime LocalTime;
-			RootToLocalLoopCounter = FMovieSceneWarpCounter();
-			RootToLocalTransform.TransformTime(RootTime.Time, LocalTime, RootToLocalLoopCounter);
-		}
-		else
-		{
-			// If we are scrubbing _and_ the current sequence is warping, we need to do some custom stuff.
-			const FFrameNumber PlayRangeSize = SubSequenceData->PlayRange.Value.Size<FFrameNumber>();
-			const FFrameNumber PlayRangeUpperBound = SubSequenceData->PlayRange.Value.GetUpperBoundValue();
-			const FFrameNumber PlayRangeLowerBound = SubSequenceData->PlayRange.Value.GetLowerBoundValue();
-			
-			ensure(LocalLoopIndexOnBeginScrubbing != FMovieSceneTimeWarping::InvalidWarpCount);
-			ensure(RootToLocalLoopCounter.WarpCounts.Num() > 0);
-
-			// Compute the new local time based on the specific loop that we had when we started scrubbing.
-			FMovieSceneSequenceTransform RootToLocalTransformWithoutLeafLooping = RootToLocalTransform;
-			FMovieSceneNestedSequenceTransform LeafLooping = RootToLocalTransformWithoutLeafLooping.NestedTransforms.Pop();
-			FFrameTime LocalTimeWithLastLoopUnwarped = RootTime.Time * RootToLocalTransformWithoutLeafLooping;
-			LocalTimeWithLastLoopUnwarped = LocalTimeWithLastLoopUnwarped * LeafLooping.LinearTransform;
-			if (LeafLooping.IsLooping())
-			{
-				LeafLooping.Warping.TransformTimeSpecific(
-						LocalTimeWithLastLoopUnwarped, LocalLoopIndexOnBeginScrubbing, LocalTimeWithLastLoopUnwarped);
-			}
+		CurrentTimeBreadcrumbs.Reset();
+		RootToLocalTransform.TransformTime(GetGlobalTime().Time, FTransformTimeParams().HarvestBreadcrumbs(CurrentTimeBreadcrumbs).IgnoreClamps());
 
-			// Now figure out if we're in a next/previous loop because we scrubbed past the lower/upper bound
-			// of the loop. Note, again, that we only compute the new loop index for UI display purposes at this
-			// point (see comment at the beginning of this method). We will commit to the new loop indices
-			// once we're done scrubbing.
-			if (PlayRangeSize > 0)
-			{
-				while (LocalTimeWithLastLoopUnwarped >= PlayRangeUpperBound)
-				{
-					LocalTimeWithLastLoopUnwarped = LocalTimeWithLastLoopUnwarped - PlayRangeSize;
-					++CurLoopIndexOffset;
-				}
-				while (LocalTimeWithLastLoopUnwarped <= PlayRangeLowerBound)
-				{
-					LocalTimeWithLastLoopUnwarped = LocalTimeWithLastLoopUnwarped + PlayRangeSize;
-					--CurLoopIndexOffset;
-				}
-			}
-		}
+		// Inner play range
+		UMovieSceneSequence* SubSequence   = SubSequenceData->GetSequence();
+		TRange<FFrameNumber> PlaybackRange = SubSequence->GetMovieScene()->GetPlaybackRange();
 
-		// If we're looping a sub sequence we may need to clamp to the max number of loops, which will modify SubSequenceRange and RootToLocalLoopCounter
-		if (bIsSubSequenceLooping)
-		{
-			const FMovieSceneSectionParameters SubSectionParameters = SubSequenceData->ToSubSectionParameters();
-			const TRange<FFrameNumber> ChildPlayRange = UMovieSceneSubSection::GetValidatedInnerPlaybackRange(SubSectionParameters, *SubSequenceData->GetSequence()->GetMovieScene());
-			const FFrameNumber ChildLength = UE::MovieScene::DiscreteSize(ChildPlayRange);
-			FMovieSceneSequenceID ParentID = Hierarchy.FindNode(ActiveTemplateIDs.Top())->ParentID;
-			UMovieScene* ParentMovieScene = (ParentID == GetRootTemplateID()) ? GetRootMovieSceneSequence()->GetMovieScene() : Hierarchy.FindSubSequence(ParentID)->GetMovieScene();
+		FMovieSceneInverseSequenceTransform LocalToRootTransform = RootToLocalTransform.Inverse();
+		const int32 PlaybackSize = UE::MovieScene::DiscreteSize(PlaybackRange);
 
-			// Now we need to know how long this child play range is in the parent's time space. This is how we can figure
-			// out how many loops we can fit.
-			const FFrameRate ParentFrameRate = ParentMovieScene->GetTickResolution();
-			const FFrameRate ChildFrameRate = SubSequenceData->GetSequence()->GetMovieScene()->GetTickResolution();
+		TOptional<FFrameTime> PlayStart = LocalToRootTransform.TryTransformTime(PlaybackRange.GetLowerBoundValue(), CurrentTimeBreadcrumbs, EInverseEvaluateFlags::Backwards | EInverseEvaluateFlags::Cycle);
+		TOptional<FFrameTime> PlayEnd   = LocalToRootTransform.TryTransformTime(PlaybackRange.GetUpperBoundValue(), CurrentTimeBreadcrumbs, EInverseEvaluateFlags::Forwards  | EInverseEvaluateFlags::Cycle);
 
-			const float ChildTimeScale = SubSequenceData->OuterToInnerTransform.LinearTransform.TimeScale;
-			const float InvChildTimeScale = FMath::IsNearlyZero(ChildTimeScale) ? 1.0f : 1.0f / ChildTimeScale;
+		FFrameTime CurrentTime = GetGlobalTime().Time;
 
-			const FFrameNumber ChildLengthInParentSpace = (ConvertFrameTime(ChildLength, ChildFrameRate, ParentFrameRate) * InvChildTimeScale).FrameNumber;
-			const FFrameNumber ChildFirstLoopLength = ChildLength - SubSequenceData->ParentFirstLoopStartFrameOffset;
-			const FFrameNumber ChildFirstLoopLengthInParentSpace = (ConvertFrameTime(ChildFirstLoopLength, ChildFrameRate, ParentFrameRate) * InvChildTimeScale).FrameNumber;
-
-			// We can finally start iterating: we iterate for how many times as we can fit the child sequence's length,
-			// modified by the time scale, into the parent play range.
-			const TRange<FFrameNumber> ParentPlayRange = SubSequenceData->ParentPlayRange.Value;
-			const FFrameNumber ParentExclusiveEnd = UE::MovieScene::DiscreteExclusiveUpper(ParentPlayRange);
-			FFrameNumber CurLoopStart = UE::MovieScene::DiscreteInclusiveLower(ParentPlayRange);
-			FFrameNumber CurLoopEnd = CurLoopStart + FMath::Min(
-				FMath::Max(ChildFirstLoopLengthInParentSpace, FFrameNumber(0)),
-				ParentExclusiveEnd);
+		TRange<FFrameTime> ValidRange = TRange<FFrameTime>::All();
+		if (PlayStart)
+		{
+			ValidRange.SetLowerBound(TRangeBound<FFrameTime>::Inclusive(PlayStart.GetValue()));
+		}
+		if (PlayEnd)
+		{
+			ValidRange.SetUpperBound(TRangeBound<FFrameTime>::Inclusive(PlayEnd.GetValue()));
+		}
 
-			int32 LocalLoopIndex = FMath::Max(0, (int32)RootToLocalLoopCounter.WarpCounts.Last() + CurLoopIndexOffset);
-			uint8 CurLoopIndex = 0;
-			TRange<FFrameNumber> NewSubSequenceRange = TRange<FFrameNumber>(CurLoopStart, FMath::Min(CurLoopEnd, ParentExclusiveEnd));
-			while (CurLoopStart < ParentExclusiveEnd)
+		for (int32 Index = 0; Index < ActiveTemplateIDs.Num(); ++Index)
+		{
+			FMovieSceneSequenceID SequenceID = ActiveTemplateIDs[Index];
+			const FMovieSceneSubSequenceData* SubData = Hierarchy->FindSubData(SequenceID);
+			if (SubData)
 			{
-				if (CurLoopIndex <= LocalLoopIndex)
+				// Clamp to the sub section range recursively
+				if (!SubData->ParentPlayRange.Value.Contains(CurrentTime.FrameNumber))
 				{
-					NewSubSequenceRange = TRange<FFrameNumber>(CurLoopStart, FMath::Min(CurLoopEnd, ParentExclusiveEnd));
-					NewSubSequenceRange = SubSequenceData->OuterToInnerTransform.TransformRangeConstrained(NewSubSequenceRange);
-				}
-
-				CurLoopStart = CurLoopEnd;
-				CurLoopEnd += FMath::Max(ChildLengthInParentSpace, FFrameNumber(0));
-				CurLoopIndex++;
-			}
+					FFrameNumber Lower = DiscreteInclusiveLower(SubData->ParentPlayRange.Value);
+					FFrameNumber Upper = DiscreteExclusiveUpper(SubData->ParentPlayRange.Value);
 
-			MaxLocalLoopIndex = CurLoopIndex - 1;
+					if (CurrentTime.FrameNumber < Lower)
+					{
+						CurrentTime = FFrameTime(Lower);
+					}
+					else if (CurrentTime.FrameNumber >= Upper)
+					{
+						CurrentTime = FFrameTime(Upper - 1, FFrameTime::MaxSubframe);
+					}
+				}
 
-			if (bIsScrubbingLoopingSubSequence)
-			{
-				// Clamp LocalLoopIndexOffset between 0 and MaxLocalLoopIndex- as it's a delta we need to subtract off LocalLoopIndexOnBeginScrubbing.
-				CurLoopIndexOffset = FMath::Max(-LocalLoopIndexOnBeginScrubbing, FMath::Min(MaxLocalLoopIndex - LocalLoopIndexOnBeginScrubbing, CurLoopIndexOffset));
+				CurrentTime = SubData->OuterToInnerTransform.TransformTime(CurrentTime);
 
-				if (CurLoopIndexOffset != LocalLoopIndexOffsetDuringScrubbing)
+				// Clamp the range and transform it into inner space
+				ValidRange = TRange<FFrameTime>::Intersection(ValidRange, ConvertToFrameTimeRange(SubData->ParentPlayRange.Value));
+				if (ValidRange.IsEmpty())
 				{
-					LocalLoopIndexOffsetDuringScrubbing = CurLoopIndexOffset;
-					// If we jumped to the previous or next loop, we need to invalidate the global marked frames because
-					// the focused (currently edited) sequence's time transform just changed.
-					InvalidateGlobalMarkedFramesCache();
+					break;
 				}
-			}
-			else
-			{
-				ensure(RootToLocalLoopCounter.WarpCounts.Num() > 0);
-				
-				// Clamp warp count
-				RootToLocalLoopCounter.WarpCounts.Last() = FMath::Clamp(RootToLocalLoopCounter.WarpCounts.Last(), 0, MaxLocalLoopIndex);
-			}
 
-			SubSequenceRange = NewSubSequenceRange;
+				ValidRange = SubData->OuterToInnerTransform.ComputeTraversedHull(ValidRange);
+			}
 		}
+
+		SubSequenceRange = ConvertToDiscreteRange(ValidRange);
 	}
 	else
 	{
-		// Ensure in this case we also reset the RootToLocalLoopCounter.
-		RootToLocalLoopCounter = FMovieSceneWarpCounter();
+		// Ensure in this case we also reset the CurrentTimeBreadcrumbs.
+		CurrentTimeBreadcrumbs.Reset();
 	}
 }
 
@@ -1778,26 +1735,37 @@ void FSequencer::DeleteSelectedKeys()
 	bool bAnythingRemoved = false;
 
 	FSelectedKeysByChannel KeysByChannel(ViewModel->GetSelection()->KeySelection);
-	TSet<UMovieSceneSection*> ModifiedSections;
+	TSet<UObject*> ModifiedObjects;
 
 	for (const FSelectedChannelInfo& ChannelInfo : KeysByChannel.SelectedChannels)
 	{
+		if (ChannelInfo.OwningSection->IsReadOnly())
+		{
+			continue;
+		}
+
 		FMovieSceneChannel* Channel = ChannelInfo.Channel.Get();
 		if (Channel)
 		{
-			bool bModified = ModifiedSections.Contains(ChannelInfo.OwningSection);
-			if (!bModified)
+			UObject* Owner = nullptr;
+			const FMovieSceneChannelMetaData* ChannelMetaData = ChannelInfo.Channel.GetMetaData();
+			if (ChannelMetaData)
 			{
-				bModified = ChannelInfo.OwningSection->TryModify();
+				Owner = ChannelMetaData->WeakOwningObject.Get();
 			}
-
-			if (bModified)
+			if (!Owner)
 			{
-				ModifiedSections.Add(ChannelInfo.OwningSection);
+				Owner = ChannelInfo.OwningSection;
+			}
 
-				Channel->DeleteKeys(ChannelInfo.KeyHandles);
-				bAnythingRemoved = true;
+			if (!ModifiedObjects.Contains(Owner))
+			{
+				Owner->Modify();
+				ModifiedObjects.Add(ChannelInfo.OwningSection);
 			}
+
+			Channel->DeleteKeys(ChannelInfo.KeyHandles);
+			bAnythingRemoved = true;
 		}
 	}
 
@@ -3194,29 +3162,25 @@ bool FSequencer::GetAutoSetTrackDefaults() const
 
 FQualifiedFrameTime FSequencer::GetLocalTime() const
 {
+	using namespace UE::MovieScene;
+
 	const FFrameRate FocusedResolution = GetFocusedTickResolution();
 	const FFrameTime CurrentPosition   = PlayPosition.GetCurrentPosition();
 
-	const FFrameTime RootTime = ConvertFrameTime(CurrentPosition, PlayPosition.GetInputRate(), PlayPosition.GetOutputRate());
+	FFrameTime RootTime  = ConvertFrameTime(CurrentPosition, PlayPosition.GetInputRate(), PlayPosition.GetOutputRate());
 
-	FFrameTime LocalTime;
-	FMovieSceneWarpCounter LoopCounter;
-	CalculateLocalTimeClamped(RootTime, RootToLocalTransform, LocalTime, LoopCounter);
+	const FFrameTime LocalTime = RootToLocalTransform.TransformTime(RootTime, FTransformTimeParams().IgnoreClamps());
 	return FQualifiedFrameTime(LocalTime, FocusedResolution);
 }
 
 
-uint32 FSequencer::GetLocalLoopIndex() const
+TOptional<int32> FSequencer::GetLocalLoopIndex() const
 {
-	if (RootToLocalLoopCounter.WarpCounts.Num() == 0)
-	{
-		return FMovieSceneTimeWarping::InvalidWarpCount;
-	}
-	else
-	{
-		const bool bIsScrubbing = GetPlaybackStatus() == EMovieScenePlayerStatus::Scrubbing;
-		return RootToLocalLoopCounter.WarpCounts.Last() + (bIsScrubbing ? LocalLoopIndexOffsetDuringScrubbing : 0);
-	}
+	using namespace UE::MovieScene;
+
+	TOptional<int32> LoopIndex;
+	RootToLocalTransform.TransformTime(GetGlobalTime().Time, FTransformTimeParams().TrackCycleCounts(&LoopIndex));
+	return LoopIndex;
 }
 
 
@@ -3276,19 +3240,41 @@ void FSequencer::SetLocalTime( FFrameTime NewTime, ESnapTimeMode SnapTimeMode, b
 
 void FSequencer::SetLocalTimeDirectly(FFrameTime NewTime, bool bEvaluate)
 {
-	// Special-case. If the RootToLocalTransform contains a zero timescale, then scrubbing in the subscene should not change the current time
-	// Inverting the zero time-scale time is technically non-deterministic, and so doing this scrub would just result in losing the current global time
-	// while maintaining the same sub-sequence time (as we're holding a single frame there). So the best case here is just preventing the scrub.
-	if (!FMath::IsNearlyZero(RootToLocalTransform.GetTimeScale()))
+	using namespace UE::MovieScene;
+
+	if (ActiveTemplateIDs.Num() <= 1)
+	{
+		if (!RootToLocalTransform.IsLinear())
+		{
+			NewTime = RootToLocalTransform.Inverse().TryTransformTime(NewTime).Get(NewTime);
+		}
+		SetGlobalTime(NewTime, bEvaluate);
+		return;
+	}
+
+
+	FMovieSceneInverseSequenceTransform LocalToRootTransform = RootToLocalTransform.Inverse();
+
+	const FMovieSceneTransformBreadcrumbs& Breadcrumbs = GetPlaybackStatus() == EMovieScenePlayerStatus::Scrubbing
+		? ScrubStartBreadcrumbs
+		: CurrentTimeBreadcrumbs;
+
+	// Transform the time to the root time-space
+	TOptional<FFrameTime> NewGlobalTime = LocalToRootTransform.TryTransformTime(NewTime + ScrubLinearOffset, Breadcrumbs,
+		EInverseEvaluateFlags::AnyDirection | EInverseEvaluateFlags::Cycle | EInverseEvaluateFlags::IgnoreClamps);
+
+	// If we still didn't find a time there's nothing we can do
+	if (NewGlobalTime.IsSet())
 	{
-		// Transform the time to the root time-space
-		SetGlobalTime(NewTime * RootToLocalTransform.InverseFromLoop(RootToLocalLoopCounter), bEvaluate);
+		SetGlobalTime(NewGlobalTime.GetValue(), bEvaluate);
 	}
 }
 
 
 void FSequencer::SetGlobalTime(FFrameTime NewTime, bool bEvaluate)
 {
+	using namespace UE::MovieScene;
+
 	NewTime = ConvertFrameTime(NewTime, GetRootTickResolution(), PlayPosition.GetInputRate());
 	if (PlayPosition.GetEvaluationType() == EMovieSceneEvaluationType::FrameLocked)
 	{
@@ -3297,9 +3283,11 @@ void FSequencer::SetGlobalTime(FFrameTime NewTime, bool bEvaluate)
 
 	// Don't update the sequence if the time hasn't changed as this will cause duplicate events and the like to fire.
 	// If we need to reevaluate the sequence at the same time for whetever reason, we should call ForceEvaluate()
-	TOptional<FFrameTime> CurrentPosition = PlayPosition.GetCurrentPosition();
 	if (PlayPosition.GetCurrentPosition() != NewTime)
 	{
+		// Make sure breadcrumbs are up to date
+		RootToLocalTransform.TransformTime(PlayPosition.GetCurrentPosition(), FTransformTimeParams().HarvestBreadcrumbs(CurrentTimeBreadcrumbs).IgnoreClamps());
+
 		FMovieSceneEvaluationRange EvalRange = PlayPosition.JumpTo(NewTime);
 		if (bEvaluate)
 		{
@@ -3392,6 +3380,10 @@ void FSequencer::ForceEvaluate()
 
 void FSequencer::EvaluateInternal(FMovieSceneEvaluationRange InRange, bool bHasJumped)
 {
+	using namespace UE::MovieScene;
+
+	// Ensure breadcrumbs are up to date
+	RootToLocalTransform.TransformTime(InRange.GetTime(), FTransformTimeParams().HarvestBreadcrumbs(CurrentTimeBreadcrumbs).IgnoreClamps());
 
 	LastEvaluatedLocalTime = GetLocalTime().Time;
 
@@ -3574,11 +3566,11 @@ void FSequencer::UpdateAutoScroll(double NewTime, float ThresholdPercentage)
 	{
 		if (AutoscrollOffset.GetValue() < 0 && LocalTime.AsSeconds() > ViewRange.GetLowerBoundValue() + Threshold)
 		{
-			SetLocalTimeLooped( (ViewRange.GetLowerBoundValue() + Threshold) * LocalTime.Rate, RootToLocalLoopCounter);
+			SetLocalTimeLooped( (ViewRange.GetLowerBoundValue() + Threshold) * LocalTime.Rate, CurrentTimeBreadcrumbs);
 		}
 		else if (AutoscrollOffset.GetValue() > 0 && LocalTime.AsSeconds() < ViewRange.GetUpperBoundValue() - Threshold)
 		{
-			SetLocalTimeLooped( (ViewRange.GetUpperBoundValue() - Threshold) * LocalTime.Rate, RootToLocalLoopCounter);
+			SetLocalTimeLooped( (ViewRange.GetUpperBoundValue() - Threshold) * LocalTime.Rate, CurrentTimeBreadcrumbs);
 		}
 	}
 
@@ -3675,6 +3667,8 @@ void FSequencer::RenderMovie(const TArray<UMovieSceneCinematicShotSection*>& InS
 
 void FSequencer::RenderMovieInternal(TRange<FFrameNumber> Range, bool bSetFrameOverrides) const
 {
+	using namespace UE::MovieScene;
+
 	ISequencerModule& SequencerModule = FModuleManager::LoadModuleChecked<ISequencerModule>("Sequencer");
 	if (IMovieRendererInterface* MovieRenderer = SequencerModule.GetMovieRenderer(GetMovieRendererName()))
 	{
@@ -3700,7 +3694,12 @@ void FSequencer::RenderMovieInternal(TRange<FFrameNumber> Range, bool bSetFrameO
 
 		if (const FMovieSceneSubSequenceData* SubSequenceData = RootTemplateInstance.FindSubData(GetFocusedTemplateID()))
 		{
-			Range = SubSequenceData->RootToSequenceTransform.InverseNoLooping().TransformRangeUnwarped(Range);
+			auto Visit = [&Range](TRange<FFrameTime> RootRange)
+			{
+				Range = ConvertToDiscreteRange(RootRange);
+				return false;
+			};
+			SubSequenceData->RootToSequenceTransform.Inverse().TransformFiniteRangeWithinRange(ConvertToFrameTimeRange(Range), Visit, CurrentTimeBreadcrumbs, CurrentTimeBreadcrumbs);
 		}
 	}
 
@@ -4717,15 +4716,14 @@ FReply FSequencer::OnPlay(bool bTogglePlay)
 	}
 	else
 	{
-		TRange<FFrameNumber> TimeBounds = GetRootTimeBounds();
+		TRange<FFrameNumber> RootTimeBounds = GetRootTimeBounds();
 
-		FFrameNumber MinInclusiveTime = UE::MovieScene::DiscreteInclusiveLower(TimeBounds);
-		FFrameNumber MaxInclusiveTime = UE::MovieScene::DiscreteExclusiveUpper(TimeBounds) - 1;
+		FFrameNumber MinInclusiveTime = UE::MovieScene::DiscreteInclusiveLower(RootTimeBounds);
+		FFrameNumber MaxInclusiveTime = UE::MovieScene::DiscreteExclusiveUpper(RootTimeBounds) - 1;
 
-		if (GetLocalTime().Time <= MinInclusiveTime || GetLocalTime().Time >= MaxInclusiveTime)
+		if (GetGlobalTime().Time <= MinInclusiveTime || GetGlobalTime().Time >= MaxInclusiveTime)
 		{
-			FFrameTime NewGlobalTime = (PlaybackSpeed > 0 ? MinInclusiveTime : MaxInclusiveTime) * RootToLocalTransform.InverseFromLoop(RootToLocalLoopCounter);
-			SetGlobalTime(NewGlobalTime);
+			SetGlobalTime(PlaybackSpeed > 0 ? MinInclusiveTime : MaxInclusiveTime);
 		}
 
 		SpeedIndexBeforePlay = CurrentSpeedIndex;
@@ -4952,18 +4950,18 @@ ESequencerLoopMode FSequencer::GetLoopMode() const
 	return Settings->GetLoopMode();
 }
 
-void FSequencer::SetLocalTimeLooped(FFrameTime NewLocalTime, FMovieSceneWarpCounter WarpCounter)
+void FSequencer::SetLocalTimeLooped(FFrameTime NewLocalTime, const FMovieSceneTransformBreadcrumbs& Breadcrumbs)
 {
+	using namespace UE::MovieScene;
+
 	TOptional<EMovieScenePlayerStatus::Type> NewPlaybackStatus;
-	// Default to the RootToLocalLoopCounter
-	if (WarpCounter.NumWarpCounts() == 0)
-	{
-		WarpCounter = RootToLocalLoopCounter;
-	}
 
-	const FMovieSceneSequenceTransform LocalToRootTransform = RootToLocalTransform.InverseFromLoop(WarpCounter);
+	FMovieSceneInverseSequenceTransform LocalToRootTransform = ActiveTemplateIDs.Num() > 1 ? RootToLocalTransform.Inverse() : FMovieSceneInverseSequenceTransform();
 
-	FFrameTime NewGlobalTime = NewLocalTime * LocalToRootTransform;
+	// Default to the CurrentTimeBreadcrumbs
+	const FMovieSceneTransformBreadcrumbs& BreadcrumbsToUse = Breadcrumbs.Num() == 0 ? CurrentTimeBreadcrumbs : Breadcrumbs;
+
+	TOptional<FFrameTime> NewGlobalTime;
 
 	bool         bResetPosition       = false;
 
@@ -4972,10 +4970,14 @@ void FSequencer::SetLocalTimeLooped(FFrameTime NewLocalTime, FMovieSceneWarpCoun
 
 	if (PauseOnFrame.IsSet() && ((PlaybackSpeed > 0 && NewLocalTime > PauseOnFrame.GetValue()) || (PlaybackSpeed < 0 && NewLocalTime < PauseOnFrame.GetValue())))
 	{
-		NewGlobalTime = PauseOnFrame.GetValue() * LocalToRootTransform;
-		PauseOnFrame.Reset();
-		bResetPosition = true;
-		NewPlaybackStatus = EMovieScenePlayerStatus::Stopped;
+		NewGlobalTime = LocalToRootTransform.TryTransformTime(PauseOnFrame.GetValue(), BreadcrumbsToUse);
+
+		if (NewGlobalTime)
+		{
+			PauseOnFrame.Reset();
+			bResetPosition = true;
+			NewPlaybackStatus = EMovieScenePlayerStatus::Stopped;
+		}
 	}
 	else if (GetLoopMode() == ESequencerLoopMode::SLM_Loop || GetLoopMode() == ESequencerLoopMode::SLM_LoopSelectionRange)
 	{
@@ -4983,58 +4985,91 @@ void FSequencer::SetLocalTimeLooped(FFrameTime NewLocalTime, FMovieSceneWarpCoun
 		FFrameNumber MinInclusiveTime = UE::MovieScene::DiscreteInclusiveLower(TimeBounds);
 		FFrameNumber MaxInclusiveTime = UE::MovieScene::DiscreteExclusiveUpper(TimeBounds) - 1;
 
-		const UMovieSceneSequence* FocusedSequence = GetFocusedMovieSceneSequence();
-		if (FocusedSequence)
+		if (NewLocalTime < MinInclusiveTime || NewLocalTime > MaxInclusiveTime)
 		{
-			if (NewLocalTime < MinInclusiveTime || NewLocalTime > MaxInclusiveTime)
+			NewGlobalTime = LocalToRootTransform.TryTransformTime((PlaybackSpeed > 0 ? MinInclusiveTime : MaxInclusiveTime), BreadcrumbsToUse);
+			if (NewGlobalTime)
 			{
-				NewGlobalTime = (PlaybackSpeed > 0 ? MinInclusiveTime : MaxInclusiveTime) * LocalToRootTransform;
-
 				bResetPosition = true;
 				bHasJumped = true;
 			}
 		}
-	}
-	else
-	{
-		TRange<FFrameNumber> RootTimeBounds = GetRootTimeBounds();
-		FFrameNumber MinInclusiveTime = UE::MovieScene::DiscreteInclusiveLower(RootTimeBounds);
-		FFrameNumber MaxInclusiveTime = UE::MovieScene::DiscreteExclusiveUpper(RootTimeBounds) - 1;
-		bool bReachedEnd = false;
-		if (PlaybackSpeed > 0)
-		{
-			bReachedEnd = GetLocalTime().Time <= MaxInclusiveTime && NewLocalTime >= MaxInclusiveTime;
-		}
 		else
 		{
-			bReachedEnd = GetLocalTime().Time >= MinInclusiveTime && NewLocalTime <= MinInclusiveTime;
+			NewGlobalTime = LocalToRootTransform.TryTransformTime(NewLocalTime, BreadcrumbsToUse);
 		}
+	}
+	else
+	{
+		NewGlobalTime = LocalToRootTransform.TryTransformTime(NewLocalTime, BreadcrumbsToUse);
 
-		// Stop if we hit the playback range end
-		if (bReachedEnd)
+		if (NewGlobalTime)
 		{
-			NewGlobalTime = (PlaybackSpeed > 0 ? MaxInclusiveTime : MinInclusiveTime) * LocalToRootTransform;
-			NewPlaybackStatus = EMovieScenePlayerStatus::Stopped;
+			TRange<FFrameNumber> RootTimeBounds = GetRootTimeBounds();
+			FFrameNumber MinInclusiveTime = UE::MovieScene::DiscreteInclusiveLower(RootTimeBounds);
+			FFrameNumber MaxInclusiveTime = UE::MovieScene::DiscreteExclusiveUpper(RootTimeBounds) - 1;
+			bool bReachedEnd = false;
+			if (PlaybackSpeed > 0)
+			{
+				bReachedEnd = GetGlobalTime().Time <= MaxInclusiveTime && NewGlobalTime.GetValue() >= MaxInclusiveTime;
+			}
+			else
+			{
+				bReachedEnd = GetGlobalTime().Time >= MinInclusiveTime && NewGlobalTime.GetValue() <= MinInclusiveTime;
+			}
+
+			// Stop if we hit the playback range end
+			if (bReachedEnd)
+			{
+				NewGlobalTime = PlaybackSpeed > 0 ? MaxInclusiveTime : MinInclusiveTime;
+				NewPlaybackStatus = EMovieScenePlayerStatus::Stopped;
+			}
 		}
 	}
 
+	if (!NewGlobalTime)
+	{
+		NewGlobalTime = LocalToRootTransform.TryTransformTime(NewLocalTime, BreadcrumbsToUse);
+	}
+
+	if (!NewGlobalTime && RootToLocalTransform.NestedTransforms.Num() > 0)
+	{
+		// If we couldn't transform the time, we try one last-ditch attempt to guess a cycling or looping
+		//     transformation based on the linear transformation between each of the breadcrumbs.
+		FMovieSceneTransformBreadcrumbs DenseBreadcrumbs;
+		FFrameTime CurrentLocalTime = RootToLocalTransform.TransformTime(GetGlobalTime().Time, FTransformTimeParams().HarvestBreadcrumbs(DenseBreadcrumbs));
+
+		FFrameTime LocalDeltaGuess   = NewLocalTime - CurrentLocalTime;
+		FFrameTime NewLocalTimeGuess = RootToLocalTransform.NestedTransforms.Last().TransformTime(LocalDeltaGuess);
+
+		NewGlobalTime = LocalToRootTransform.TryTransformTime(NewLocalTimeGuess, DenseBreadcrumbs);
+	}
+
+	if (!NewGlobalTime)
+	{
+		return;
+	}
+
 	FFrameRate   RootTickResolution = GetRootTickResolution();
 	// Ensure the time is in the current view - must occur before the time cursor changes
 	UMovieScene* MovieScene = GetFocusedMovieSceneSequence()->GetMovieScene();
 	if (IsAutoScrollEnabled())
 	{
-		ScrollIntoView((NewGlobalTime * RootToLocalTransform) / RootTickResolution);
+		ScrollIntoView((NewGlobalTime.GetValue() * RootToLocalTransform) / RootTickResolution);
 	}
 
-	FFrameTime NewPlayPosition = ConvertFrameTime(NewGlobalTime, RootTickResolution, PlayPosition.GetInputRate());
+	FFrameTime NewPlayPosition = ConvertFrameTime(NewGlobalTime.GetValue(), RootTickResolution, PlayPosition.GetInputRate());
 
 	// Reset the play cursor if we're looping or have otherwise jumpted to a new position in the sequence
 	if (bResetPosition)
 	{
 		PlayPosition.Reset(NewPlayPosition);
-		TimeController->Reset(FQualifiedFrameTime(NewGlobalTime, RootTickResolution));
+		TimeController->Reset(FQualifiedFrameTime(NewGlobalTime.GetValue(), RootTickResolution));
 	}
 
+	// Ensure breadcrumbs are up to date
+	RootToLocalTransform.TransformTime(NewPlayPosition, FTransformTimeParams().HarvestBreadcrumbs(CurrentTimeBreadcrumbs).IgnoreClamps());
+
 	// Evaluate the sequence
 	FMovieSceneEvaluationRange EvalRange = PlayPosition.PlayTo(NewPlayPosition);
 	EvaluateInternal(EvalRange, bHasJumped);
@@ -5071,25 +5106,17 @@ void FSequencer::RestorePlaybackSpeedAfterPlay()
 	PlaybackSpeed = PlaybackSpeedBeforePlay;
 }
 
-void FSequencer::CalculateLocalTimeClamped(FFrameTime RootTime, const FMovieSceneSequenceTransform& RootToParentChainTransform, FFrameTime& OutLocalTime, FMovieSceneWarpCounter& OutLoopCounter) const
+void FSequencer::CalculateLocalTimeClamped(FFrameTime RootTime, const FMovieSceneSequenceTransform& RootToParentChainTransform, FFrameTime& OutLocalTime, FMovieSceneTransformBreadcrumbs& OutBreadcrumbs) const
 {
-	RootToParentChainTransform.TransformTime(RootTime, OutLocalTime, OutLoopCounter);
-	if (RootToParentChainTransform.NestedTransforms.Num() > 0 && RootToParentChainTransform.NestedTransforms.Last().IsLooping())
-	{
-		ensure(OutLoopCounter.WarpCounts.Num() > 0);
-		// Clamp warp count
-		OutLoopCounter.WarpCounts.Last() = FMath::Clamp(OutLoopCounter.WarpCounts.Last(), 0, MaxLocalLoopIndex);
+	using namespace UE::MovieScene;
 
-		// Update OutLocalTime based on the current clamped warp count]
-		FMovieSceneSequenceTransform RootToLocalTransformWithoutLeafLooping = RootToParentChainTransform;
-		FMovieSceneNestedSequenceTransform LeafLooping = RootToLocalTransformWithoutLeafLooping.NestedTransforms.Pop();
-		FFrameTime LocalTimeWithLastLoopUnwarped = RootTime * RootToLocalTransformWithoutLeafLooping;
-		LocalTimeWithLastLoopUnwarped = LocalTimeWithLastLoopUnwarped * LeafLooping.LinearTransform;
-		if (LeafLooping.IsLooping())
-		{
-			LeafLooping.Warping.TransformTimeSpecific(
-				LocalTimeWithLastLoopUnwarped, OutLoopCounter.WarpCounts.Last(), OutLocalTime);
-		}
+	if (ActiveTemplateIDs.Num() > 1)
+	{
+		OutLocalTime = RootToLocalTransform.TransformTime(RootTime, FTransformTimeParams().HarvestBreadcrumbs(OutBreadcrumbs).IgnoreClamps());
+	}
+	else
+	{
+		OutLocalTime = RootTime;
 	}
 }
 
@@ -5160,7 +5187,12 @@ TRange<FFrameNumber> FSequencer::GetTimeBounds() const
 		}
 	}
 
-	if (Settings->ShouldEvaluateSubSequencesInIsolation() || ActiveTemplateIDs.Num() == 1)
+	if (ActiveTemplateIDs.Num() == 1)
+	{
+		return GetRootTimeBounds();
+	}
+
+	if (Settings->ShouldEvaluateSubSequencesInIsolation())
 	{
 		return FocusedSequence->GetMovieScene()->GetPlaybackRange();
 	}
@@ -5176,27 +5208,17 @@ TRange<FFrameNumber> FSequencer::GetRootTimeBounds() const
 		return TRange<FFrameNumber>(-100000, 100000);
 	}
 
-	TRange<FFrameNumber> RootPlayRangeRootSpace = RootMovieSceneSequence->GetMovieScene()->GetPlaybackRange();
-	if (RootMovieSceneSequence == GetFocusedMovieSceneSequence())
-	{
-		return RootPlayRangeRootSpace;
-	}
-	
-	// Edge case for looping- we need to calculate clamped local time
-	if (RootToLocalTransform.NestedTransforms.Num() > 0 && RootToLocalTransform.NestedTransforms.Last().IsLooping())
-	{
-		FFrameTime RootPlayRangeLocalSpaceLowerBound;
-		FFrameTime RootPlayRangeLocalSpaceUpperBound;
-		FMovieSceneWarpCounter WarpCounter;
-		CalculateLocalTimeClamped(RootPlayRangeRootSpace.GetLowerBoundValue(), RootToLocalTransform, RootPlayRangeLocalSpaceLowerBound, WarpCounter);
-		CalculateLocalTimeClamped(RootPlayRangeRootSpace.GetUpperBoundValue(), RootToLocalTransform, RootPlayRangeLocalSpaceUpperBound, WarpCounter);
+	TRange<FFrameNumber> RootTimeBounds = RootMovieSceneSequence->GetMovieScene()->GetPlaybackRange();
 
-		return TRange<FFrameNumber>(RootPlayRangeLocalSpaceLowerBound.FloorToFrame(), RootPlayRangeLocalSpaceUpperBound.CeilToFrame());
-	}
-	else
+	TOptional<FFrameTime> Lower = RootTransform.Inverse().TryTransformTime(RootTimeBounds.GetLowerBoundValue());
+	TOptional<FFrameTime> Upper = RootTransform.Inverse().TryTransformTime(RootTimeBounds.GetUpperBoundValue());
+
+	if (Lower && Upper)
 	{
-		return RootToLocalTransform.TransformRangeConstrained(RootPlayRangeRootSpace);
+		return TRange<FFrameNumber>(Lower->FrameNumber, Upper->FrameNumber);
 	}
+
+	return RootTimeBounds;
 }
 
 
@@ -5401,6 +5423,7 @@ void FSequencer::OnScrubPositionChanged( FFrameTime NewScrubPosition, bool bScru
 	{
 		if (!bScrubbing)
 		{
+			NewScrubPosition += ScrubLinearOffset;
 			OnEndScrubbing();
 		}
 		else
@@ -5437,8 +5460,8 @@ void FSequencer::OnBeginScrubbing()
 	SetPlaybackStatus(EMovieScenePlayerStatus::Scrubbing);
 	SequencerWidget->RegisterActiveTimerForPlayback();
 
-	LocalLoopIndexOnBeginScrubbing = GetLocalLoopIndex();
-	LocalLoopIndexOffsetDuringScrubbing = 0;
+	ScrubLinearOffset = FFrameTime(0);
+	ScrubStartBreadcrumbs = CurrentTimeBreadcrumbs;
 
 	OnBeginScrubbingDelegate.Broadcast();
 }
@@ -5450,8 +5473,8 @@ void FSequencer::OnEndScrubbing()
 	AutoscrubOffset.Reset();
 	StopAutoscroll();
 
-	LocalLoopIndexOnBeginScrubbing = FMovieSceneTimeWarping::InvalidWarpCount;
-	LocalLoopIndexOffsetDuringScrubbing = 0;
+	ScrubStartBreadcrumbs = FMovieSceneTransformBreadcrumbs();
+	ScrubLinearOffset = FFrameTime(0);
 
 	OnEndScrubbingDelegate.Broadcast();
 }
@@ -5529,7 +5552,7 @@ FString FSequencer::GetFrameTimeText() const
 	const FFrameTime RootTime = ConvertFrameTime(CurrentPosition, PlayPosition.GetInputRate(), PlayPosition.GetOutputRate());
 	
 	FFrameTime LocalTime;
-	FMovieSceneWarpCounter LoopCounter;
+	FMovieSceneTransformBreadcrumbs LoopCounter;
 
 	CalculateLocalTimeClamped(RootTime, RootToParentChainTransform, LocalTime, LoopCounter);
 
@@ -9264,11 +9287,15 @@ void FSequencer::Pause()
 		FFrameRate          FocusedDisplayRate = GetFocusedDisplayRate();
 
 		// Snap to the focused play rate
-		FFrameTime RootPosition  = FFrameRate::Snap(LocalTime.Time, LocalTime.Rate, FocusedDisplayRate) * RootToLocalTransform.InverseFromLoop(RootToLocalLoopCounter);
+		TOptional<FFrameTime> RootPosition  = RootToLocalTransform.Inverse().TryTransformTime(
+			FFrameRate::Snap(LocalTime.Time, LocalTime.Rate, FocusedDisplayRate), CurrentTimeBreadcrumbs);
 
-		// Convert the root position from tick resolution time base (the output rate), to the play position input rate
-		FFrameTime InputPosition = ConvertFrameTime(RootPosition, PlayPosition.GetOutputRate(), PlayPosition.GetInputRate());
-		EvaluateInternal(PlayPosition.PlayTo(InputPosition));
+		if (RootPosition)
+		{
+			// Convert the root position from tick resolution time base (the output rate), to the play position input rate
+			FFrameTime InputPosition = ConvertFrameTime(RootPosition.GetValue(), PlayPosition.GetOutputRate(), PlayPosition.GetInputRate());
+			EvaluateInternal(PlayPosition.PlayTo(InputPosition));
+		}
 	}
 	else
 	{
@@ -9345,8 +9372,13 @@ void FSequencer::StepToNextShot()
 		return;
 	}
 
-	FFrameTime CurrentTime = SubSequenceRange.GetLowerBoundValue() * SubData->OuterToInnerTransform.InverseFromLoop(RootToLocalLoopCounter);
-	UMovieSceneSubSection* NextShot = Cast<UMovieSceneSubSection>(MovieSceneHelpers::FindNextSection(CinematicShotTrack->GetAllSections(), CurrentTime.FloorToFrame()));
+	TOptional<FFrameTime> CurrentTime = SubData->OuterToInnerTransform.Inverse().TryTransformTime(SubSequenceRange.GetLowerBoundValue(), CurrentTimeBreadcrumbs);
+	if (!CurrentTime)
+	{
+		return;
+	}
+
+	UMovieSceneSubSection* NextShot = Cast<UMovieSceneSubSection>(MovieSceneHelpers::FindNextSection(CinematicShotTrack->GetAllSections(), CurrentTime->FloorToFrame()));
 	if (!NextShot)
 	{
 		return;
@@ -9403,8 +9435,13 @@ void FSequencer::StepToPreviousShot()
 		return;
 	}
 
-	FFrameTime CurrentTime = SubSequenceRange.GetLowerBoundValue() * SubData->OuterToInnerTransform.InverseFromLoop(RootToLocalLoopCounter);
-	UMovieSceneSubSection* PreviousShot = Cast<UMovieSceneSubSection>(MovieSceneHelpers::FindPreviousSection(CinematicShotTrack->GetAllSections(), CurrentTime.FloorToFrame()));
+	TOptional<FFrameTime> CurrentTime = SubData->OuterToInnerTransform.Inverse().TryTransformTime(SubSequenceRange.GetLowerBoundValue(), CurrentTimeBreadcrumbs);
+	if (!CurrentTime)
+	{
+		return;
+	}
+
+	UMovieSceneSubSection* PreviousShot = Cast<UMovieSceneSubSection>(MovieSceneHelpers::FindPreviousSection(CinematicShotTrack->GetAllSections(), CurrentTime->FloorToFrame()));
 	if (!PreviousShot)
 	{
 		return;
@@ -9777,14 +9814,7 @@ TArray<FMovieSceneMarkedFrame> FSequencer::GetGlobalMarkedFrames() const
 void FSequencer::UpdateGlobalMarkedFramesCache()
 {
 	GlobalMarkedFramesCache.Empty();
-
-	TArray<uint32> LoopCounts = RootToLocalLoopCounter.WarpCounts;
-	if (LoopCounts.Num() > 0)
-	{
-		LoopCounts.Last() += LocalLoopIndexOffsetDuringScrubbing;
-	}
-	FSequencerMarkedFrameHelper::FindGlobalMarkedFrames(*this, LoopCounts, GlobalMarkedFramesCache);
-	
+	FSequencerMarkedFrameHelper::FindGlobalMarkedFrames(*this, GlobalMarkedFramesCache);
 	bGlobalMarkedFramesCached = true;
 }
 
@@ -9984,7 +10014,7 @@ void FSequencer::StepToNextMark()
 		return;
 	}
 
-	FFrameNumber CurrentTime = GetLocalTime().Time.FloorToFrame();
+	FFrameNumber CurrentTime = GetLocalTime().Time.RoundToFrame();
 	TOptional<FFrameNumber> NearestTime;
 
 	const bool bForwards = true;
diff --git a/Engine/Source/Editor/Sequencer/Private/Sequencer.h b/Engine/Source/Editor/Sequencer/Private/Sequencer.h
--- a/Engine/Source/Editor/Sequencer/Private/Sequencer.h
+++ b/Engine/Source/Editor/Sequencer/Private/Sequencer.h
@@ -501,7 +501,7 @@ public:
 	bool IsPilotCamera() const;
 
 	/** Sets the new global time calculated from local time and the given warp counter, accounting for looping options */
-	void SetLocalTimeLooped(FFrameTime InTime, FMovieSceneWarpCounter WarpCounter=FMovieSceneWarpCounter());
+	void SetLocalTimeLooped(FFrameTime InTime, const FMovieSceneTransformBreadcrumbs& Breadcrumbs=FMovieSceneTransformBreadcrumbs());
 
 	ESequencerLoopMode GetLoopMode() const;
 
@@ -718,7 +718,7 @@ public:
 	virtual bool GetAutoSetTrackDefaults() const override;
 	virtual FQualifiedFrameTime GetLocalTime() const override;
 	virtual FQualifiedFrameTime GetGlobalTime() const override;
-	virtual uint32 GetLocalLoopIndex() const override;
+	virtual TOptional<int32> GetLocalLoopIndex() const override;
 	virtual void SetLocalTime(FFrameTime Time, ESnapTimeMode SnapTimeMode = ESnapTimeMode::STM_None, bool bEvaluate = true) override;
 	virtual void SetLocalTimeDirectly(FFrameTime NewTime, bool bEvaluate = true) override;
 	virtual FFrameTime GetLastEvaluatedLocalTime() const override;
@@ -844,7 +844,7 @@ public:
 
 
 	/**
-	 * Gets the time boundaries of the root movie scene in local space. If this is a looping subsequence, this will include all loops.
+	 * Gets the time boundaries of the root movie scene.
 	 */
 	TRange<FFrameNumber> GetRootTimeBounds() const;
 
@@ -1142,7 +1142,7 @@ private:
 
 	FGuid FindUnspawnedObjectGuid(UObject& InObject);
 	// Given the root sequence time, returns the local time and loop counter clamped to the maximum number of loops
-	void CalculateLocalTimeClamped(FFrameTime RootTime, const FMovieSceneSequenceTransform& RootToParentChainTransform, FFrameTime& OutTime, FMovieSceneWarpCounter& OutLoopCounter) const;
+	void CalculateLocalTimeClamped(FFrameTime RootTime, const FMovieSceneSequenceTransform& RootToParentChainTransform, FFrameTime& OutTime, FMovieSceneTransformBreadcrumbs& OutBreadcrumbs) const;
 
 public:
 
@@ -1243,9 +1243,14 @@ private:
 
 	/** Time transformation from the root sequence to the currently edited sequence. */
 	FMovieSceneSequenceTransform RootToLocalTransform;
+	FMovieSceneSequenceTransform RootTransform;
 
-	/** Current loop of the current sub-sequence, if we are in a looping sub-sequence. */
-	FMovieSceneWarpCounter RootToLocalLoopCounter;
+	/** Breadcrumbs to the current local time from the root sequence. */
+	FMovieSceneTransformBreadcrumbs CurrentTimeBreadcrumbs;
+	/** Breadcrumbs to the local time that scrubbing was started at. */
+	FMovieSceneTransformBreadcrumbs ScrubStartBreadcrumbs;
+	/** A linear offset to apply to time operations when scrubbing. Used for traversing loop boundaries etc. */
+	FFrameTime ScrubLinearOffset;
 
 	/** The time range target to be viewed */
 	TRange<double> TargetViewRange;
@@ -1300,15 +1305,6 @@ private:
 	/** Current play position */
 	FMovieScenePlaybackPosition PlayPosition;
 
-	/** Local loop index at the time we began scrubbing */
-	int32 LocalLoopIndexOnBeginScrubbing;
-
-	/** Local loop index to add for the purposes of displaying it in the UI */
-	int32 LocalLoopIndexOffsetDuringScrubbing;
-
-	/** MaxLocalLoopIndex as calculated in UpdateSubSequenceData. Used to ensure LocalTime is also clamped to the correct number of loops. */
-	int32 MaxLocalLoopIndex;
-
 	/** The playback speed */
 	float PlaybackSpeed;
 
diff --git a/Engine/Source/Editor/Sequencer/Private/SequencerChannelTraits.cpp b/Engine/Source/Editor/Sequencer/Private/SequencerChannelTraits.cpp
--- a/Engine/Source/Editor/Sequencer/Private/SequencerChannelTraits.cpp
+++ b/Engine/Source/Editor/Sequencer/Private/SequencerChannelTraits.cpp
@@ -39,7 +39,7 @@ bool SupportsCurveEditorModels(const FMovieSceneChannelHandle& ChannelHandle)
 	return false;
 }
 
-TUniquePtr<FCurveModel> CreateCurveEditorModel(const FMovieSceneChannelHandle& ChannelHandle, UMovieSceneSection* OwningSection, TSharedRef<ISequencer> InSequencer)
+TUniquePtr<FCurveModel> CreateCurveEditorModel(const FMovieSceneChannelHandle& ChannelHandle, const UE::Sequencer::FCreateCurveEditorModelParams& Params)
 {
 	return nullptr;
 }
diff --git a/Engine/Source/Editor/Sequencer/Private/SequencerKeyParams.cpp b/Engine/Source/Editor/Sequencer/Private/SequencerKeyParams.cpp
--- a/Engine/Source/Editor/Sequencer/Private/SequencerKeyParams.cpp
+++ b/Engine/Source/Editor/Sequencer/Private/SequencerKeyParams.cpp
@@ -53,6 +53,14 @@ void FKeyOperation::InitializeOperation(FFrameNumber InKeyTime)
 		{
 			UMovieSceneSection* SectionObject = Operation.Section->GetSectionObject();
 			SectionObject->Modify();
+			for (TSharedPtr<IKeyArea> KeyArea : Operation.KeyAreas)
+			{
+				UObject* OwningObject = KeyArea->GetOwningObject();
+				if (OwningObject != SectionObject)
+				{
+					OwningObject->Modify();
+				}
+			}
 			SectionObject->ExpandToFrame(InKeyTime);
 		}
 	}
diff --git a/Engine/Source/Editor/Sequencer/Private/SequencerMarkedFrameHelper.cpp b/Engine/Source/Editor/Sequencer/Private/SequencerMarkedFrameHelper.cpp
--- a/Engine/Source/Editor/Sequencer/Private/SequencerMarkedFrameHelper.cpp
+++ b/Engine/Source/Editor/Sequencer/Private/SequencerMarkedFrameHelper.cpp
@@ -36,6 +36,9 @@ static void FindGlobalMarkedFrames(
 		const ISequencer& Sequencer, const FMovieSceneSequenceHierarchy* SequenceHierarchy,
 		FMovieSceneSequenceIDRef FocusedSequenceID, FMovieSceneSequenceIDRef SequenceID, 
 		TRange<FFrameNumber> GatherRange,
+		const FMovieSceneInverseSequenceTransform& SequenceToRootTransform,
+		const FMovieSceneTransformBreadcrumbs& StartBreadcrumbs,
+		const FMovieSceneTransformBreadcrumbs& EndBreadcrumbs,
 		TMovieSceneTimeArray<FMovieSceneMarkedFrame>& OutTimestampedGlobalMarkedFrames)
 {
 	// Find the current sequence in the hierarchy.
@@ -55,7 +58,14 @@ static void FindGlobalMarkedFrames(
 		{
 			if (GatherRange.Contains(MarkedFrame.FrameNumber))
 			{
-				OutTimestampedGlobalMarkedFrames.Add(MarkedFrame.FrameNumber, MarkedFrame);
+				// Iterate all mappings in the root
+				auto VisitMarkedFrame = [&OutTimestampedGlobalMarkedFrames, &MarkedFrame](FFrameTime RootTime)
+				{
+					OutTimestampedGlobalMarkedFrames.Add(RootTime, MarkedFrame);
+					return true;
+				};
+
+				SequenceToRootTransform.TransformTimeWithinRange(MarkedFrame.FrameNumber, VisitMarkedFrame, StartBreadcrumbs, EndBreadcrumbs);
 			}
 		}
 	}
@@ -79,91 +89,19 @@ static void FindGlobalMarkedFrames(
 				continue;
 			}
 
-			if (!ChildSubData->bCanLoop || FMath::IsNearlyZero(ChildSubData->OuterToInnerTransform.GetTimeScale()))
-			{
-				// This child doesn't loop (or it does, but we have zero timescale, so it's irrelevant). We just need to take into account its offset and time scale, and how much of it we "see" through the
-				// parent section.
-				ensure(!ChildSubData->OuterToInnerTransform.IsLooping() || FMath::IsNearlyZero(ChildSubData->OuterToInnerTransform.GetTimeScale()));
-
-				const FMovieSceneTimeTransform& OuterToInnerTransform(ChildSubData->OuterToInnerTransform.LinearTransform);
-				OutTimestampedGlobalMarkedFrames.PushTransform(OuterToInnerTransform);
-
-				// We may not be able to loop, but we may have multiple nested transforms if we have a zero timescale
-				const int NumNestedTransforms = ChildSubData->OuterToInnerTransform.NestedTransforms.Num();
-				for (int i = 0; i < NumNestedTransforms; ++i)
-				{
-					const FMovieSceneTimeTransform& SubOuterToInnerTransform(ChildSubData->OuterToInnerTransform.NestedTransforms[i].LinearTransform);
-					OutTimestampedGlobalMarkedFrames.PushTransform(SubOuterToInnerTransform);
-				}
-
-				// Compute the "window" of the child that we can see through the parent section. Note that the parent section could extend past
-				// the end of the child's playback... we don't trim that, so that any marked frame added after the playback end in the child sequence
-				// is still visible in the parent sequence as long as the parent section is long enough. This lets artists see what they're "missing".
-				const TRange<FFrameNumber> ParentPlayRange = ChildSubData->ParentPlayRange.Value;
-				const TRange<FFrameNumber> ChildGatherRange = ChildSubData->OuterToInnerTransform.TransformRangeConstrained(ParentPlayRange);
-
-				// Gather marked frames in this "window".
-				FindGlobalMarkedFrames(Sequencer, SequenceHierarchy, FocusedSequenceID, ChildID, ChildGatherRange, OutTimestampedGlobalMarkedFrames);
-				
-				// Pop the nested transforms off, including the base linear transform
-				for (int i = 0; i < NumNestedTransforms + 1; ++i)
-				{
-					OutTimestampedGlobalMarkedFrames.PopTransform();
-				}
-			}
-			else if (ensure(ChildSubData->OuterToInnerTransform.NestedTransforms.Num() > 0))
-			{
-				// This child is looping. Things are... more complicated.
-				//
-				// First, we push the time transform for this child sequence. It should be found in the first NestedTransforms
-				// since that's how sub-datas are computed.
-				ensure(ChildSubData->OuterToInnerTransform.LinearTransform.IsIdentity());
-				const FMovieSceneNestedSequenceTransform& OuterToInnerTransform(ChildSubData->OuterToInnerTransform.NestedTransforms[0]);
-				OutTimestampedGlobalMarkedFrames.PushTransform(OuterToInnerTransform.LinearTransform, OuterToInnerTransform.Warping);
-
-				// Next, we'll need to gather the marked frames of the child sequence repeatedly, once for each loop.
-				// To know how many loops we have, we need to look at the play range of the parent sub-section, so let's grab
-				// that first.
-				const FMovieSceneSectionParameters SubSectionParameters = ChildSubData->ToSubSectionParameters();
-				const TRange<FFrameNumber> ChildPlayRange = UMovieSceneSubSection::GetValidatedInnerPlaybackRange(SubSectionParameters, *ChildMovieScene);
-				const FFrameNumber ChildLength = UE::MovieScene::DiscreteSize(ChildPlayRange);
-
-				// Now we need to know how long this child play range is in the parent's time space. This is how we can figure
-				// out how many loops we can fit.
-				const FFrameRate ParentFrameRate = MovieScene->GetTickResolution();
-				const FFrameRate ChildFrameRate   = ChildMovieScene->GetTickResolution();
-
-				const float ChildTimeScale = OuterToInnerTransform.LinearTransform.TimeScale;
-				const float InvChildTimeScale = FMath::IsNearlyZero(ChildTimeScale) ? 1.0f : 1.0f / ChildTimeScale;
-
-				const FFrameNumber ChildLengthInParentSpace = (ConvertFrameTime(ChildLength, ChildFrameRate, ParentFrameRate) * InvChildTimeScale).FrameNumber;
-				const FFrameNumber ChildFirstLoopLength = ChildLength - ChildSubData->ParentFirstLoopStartFrameOffset;
-				const FFrameNumber ChildFirstLoopLengthInParentSpace = (ConvertFrameTime(ChildFirstLoopLength, ChildFrameRate, ParentFrameRate) * InvChildTimeScale).FrameNumber;
-
-				// We can finally start iterating: we iterate for how many times as we can fit the child sequence's length,
-				// modified by the time scale, into the parent play range.
-				const TRange<FFrameNumber> ParentPlayRange = ChildSubData->ParentPlayRange.Value;
-				const FFrameNumber ParentExclusiveEnd = UE::MovieScene::DiscreteExclusiveUpper(ParentPlayRange);
-				FFrameNumber CurLoopStart = UE::MovieScene::DiscreteInclusiveLower(ParentPlayRange);
-				FFrameNumber CurLoopEnd = CurLoopStart + FMath::Min(
-						FMath::Max(ChildFirstLoopLengthInParentSpace, FFrameNumber(0)),
-						ParentExclusiveEnd);
-
-				while (CurLoopStart < ParentExclusiveEnd)
-				{
-					TRange<FFrameNumber> CurLoopChildGatherRange(CurLoopStart, CurLoopEnd);
-					CurLoopChildGatherRange = ChildSubData->OuterToInnerTransform.TransformRangeConstrained(CurLoopChildGatherRange);
-
-					FindGlobalMarkedFrames(Sequencer, SequenceHierarchy, FocusedSequenceID, ChildID, CurLoopChildGatherRange, OutTimestampedGlobalMarkedFrames);
-
-					CurLoopStart = CurLoopEnd;
-					CurLoopEnd += FMath::Max(ChildLengthInParentSpace, FFrameNumber(0));
-					
-					OutTimestampedGlobalMarkedFrames.IncrementWarpCounter();
-				}
+			// Gather marked frames in this "window".
+			FindGlobalMarkedFrames(
+				Sequencer,
+				SequenceHierarchy,
+				FocusedSequenceID,
+				ChildID,
+				ChildSubData->PlayRange.Value,
+				ChildSubData->RootToSequenceTransform.Inverse(),
+				ChildSubData->StartTimeBreadcrumbs,
+				ChildSubData->EndTimeBreadcrumbs,
+				OutTimestampedGlobalMarkedFrames
+			);
 
-				OutTimestampedGlobalMarkedFrames.PopTransform();
-			}
 		}
 	}
 }
@@ -171,7 +109,7 @@ static void FindGlobalMarkedFrames(
 } // namespace MovieScene
 } // namespace UE
 
-void FSequencerMarkedFrameHelper::FindGlobalMarkedFrames(ISequencer& Sequencer, TArray<uint32> LoopCounter, TArray<FMovieSceneMarkedFrame>& OutGlobalMarkedFrames)
+void FSequencerMarkedFrameHelper::FindGlobalMarkedFrames(ISequencer& Sequencer, TArray<FMovieSceneMarkedFrame>& OutGlobalMarkedFrames)
 {
 	// Get the focused sequence info. We want to gather all the marked frames that are in the subset of the sequence hierarchy
 	// that hangs below this focused sequence.
@@ -198,45 +136,30 @@ void FSequencerMarkedFrameHelper::FindGlobalMarkedFrames(ISequencer& Sequencer,
 	// transform stack if the focused sequence isn't the root sequence (which has no time transform).
 
 	TMovieSceneTimeArray<FMovieSceneMarkedFrame> TimestampedGlobalMarkedFrames;
-	const FMovieSceneSubSequenceData* FocusedMovieSequenceSubData = SequenceHierarchy->FindSubData(FocusedMovieSequenceID);
-	if (FocusedMovieSequenceSubData)
-	{
-		// If we have a warping, but not looping transform, we need to do the inverse in pieces rather than just a single transform, 
-		// because if we have zero-timescale anywhere in the mix, then we lose information if we just take a simple FMovieSceneTimeTransform as an inverse.
-		if (FocusedMovieSequenceSubData->RootToSequenceTransform.IsWarping() && !FocusedMovieSequenceSubData->RootToSequenceTransform.IsLooping())
-		{
-			const int NumNestedTransforms = FocusedMovieSequenceSubData->RootToSequenceTransform.NestedTransforms.Num();
-			for (int i = NumNestedTransforms - 1; i >= 0; --i)
-			{
-				const FMovieSceneTimeTransform& SubOuterToInnerTransform(FocusedMovieSequenceSubData->RootToSequenceTransform.NestedTransforms[i].LinearTransform.Inverse());
-				TimestampedGlobalMarkedFrames.PushTransform(SubOuterToInnerTransform);
-			}
-			TimestampedGlobalMarkedFrames.PushTransform(FocusedMovieSequenceSubData->RootToSequenceTransform.LinearTransform.Inverse());
-		}
-		else
-		{
-			const int32 FocusedMovieSequenceDepth = FMath::Min(
-				FocusedMovieSequenceSubData->RootToSequenceTransform.NestedTransforms.Num(),
-				LoopCounter.Num());
-			TArrayView<uint32> LoopCounterForFocusedMovieSequence = MakeArrayView(LoopCounter.GetData(), FocusedMovieSequenceDepth);
-			FMovieSceneSequenceTransform InverseLoopTransform = FocusedMovieSequenceSubData->RootToSequenceTransform.InverseFromLoop(LoopCounterForFocusedMovieSequence);
-
-			TimestampedGlobalMarkedFrames.PushTransform(InverseLoopTransform.LinearTransform);
-			for (int i = 0; i < InverseLoopTransform.NestedTransforms.Num(); ++i)
-			{
-				TimestampedGlobalMarkedFrames.PushTransform(InverseLoopTransform.NestedTransforms[i].LinearTransform);
-			}
-		}
-	}
 
 	// Grab the marked frames from the root sequence, and recursively across the whole hierarchy.
-	UE::MovieScene::FindGlobalMarkedFrames(Sequencer, SequenceHierarchy, FocusedMovieSequenceID, MovieSceneSequenceID::Root, TRange<FFrameNumber>::All(), TimestampedGlobalMarkedFrames);
+	UE::MovieScene::FindGlobalMarkedFrames(
+		Sequencer,
+		SequenceHierarchy,
+		FocusedMovieSequenceID,
+		MovieSceneSequenceID::Root,
+		TRange<FFrameNumber>::All(),
+		FMovieSceneInverseSequenceTransform(),
+		FMovieSceneTransformBreadcrumbs(),
+		FMovieSceneTransformBreadcrumbs(),
+		TimestampedGlobalMarkedFrames);
+
+	FMovieSceneSequenceTransform RootToFocusedTransform;
+	if (const FMovieSceneSubSequenceData* SubData = SequenceHierarchy->FindSubData(FocusedMovieSequenceID))
+	{
+		RootToFocusedTransform = SubData->RootToSequenceTransform;
+	}
 
 	// Export the modified timestamped entries.
 	for (const TMovieSceneTimeArrayEntry<FMovieSceneMarkedFrame>& Entry : TimestampedGlobalMarkedFrames.GetEntries())
 	{
 		FMovieSceneMarkedFrame MarkedFrame = Entry.Datum;
-		MarkedFrame.FrameNumber = Entry.RootTime.FrameNumber;
+		MarkedFrame.FrameNumber = RootToFocusedTransform.TransformTime(Entry.RootTime).FrameNumber;
 		OutGlobalMarkedFrames.Add(MarkedFrame);
 	}
 }
diff --git a/Engine/Source/Editor/Sequencer/Private/SequencerMarkedFrameHelper.h b/Engine/Source/Editor/Sequencer/Private/SequencerMarkedFrameHelper.h
--- a/Engine/Source/Editor/Sequencer/Private/SequencerMarkedFrameHelper.h
+++ b/Engine/Source/Editor/Sequencer/Private/SequencerMarkedFrameHelper.h
@@ -16,7 +16,7 @@ class FSequencerMarkedFrameHelper
 {
 public:
 	/** Find all marked frames in a sequence hierarchy, relative to the Sequencer's currently focused sequence */
-	static void FindGlobalMarkedFrames(ISequencer& Sequencer, TArray<uint32> LoopCounter, TArray<FMovieSceneMarkedFrame>& OutGlobalMarkedFrames);
+	static void FindGlobalMarkedFrames(ISequencer& Sequencer, TArray<FMovieSceneMarkedFrame>& OutGlobalMarkedFrames);
 
 	/** Clear the setting to show marked frames globally across an entire sequence hierarchy */
 	static void ClearGlobalMarkedFrames(ISequencer& Sequencer);
diff --git a/Engine/Source/Editor/Sequencer/Private/SequencerSelectedKey.cpp b/Engine/Source/Editor/Sequencer/Private/SequencerSelectedKey.cpp
--- a/Engine/Source/Editor/Sequencer/Private/SequencerSelectedKey.cpp
+++ b/Engine/Source/Editor/Sequencer/Private/SequencerSelectedKey.cpp
@@ -91,8 +91,10 @@ void GetKeyTimes(TArrayView<const FSequencerSelectedKey> InSelectedKeys, TArrayV
 
 	for (const FSelectedChannelInfo& ChannelInfo : KeysByChannel.SelectedChannels)
 	{
+		const FMovieSceneChannelMetaData* MetaData = ChannelInfo.Channel.GetMetaData();
 		FMovieSceneChannel* Channel = ChannelInfo.Channel.Get();
-		if (Channel)
+		UMovieSceneSection* Section = ChannelInfo.OwningSection;
+		if (Channel && MetaData && Section)
 		{
 			// Resize the scratch buffer to the correct size
 			const int32 NumKeys = ChannelInfo.KeyHandles.Num();
@@ -102,10 +104,11 @@ void GetKeyTimes(TArrayView<const FSequencerSelectedKey> InSelectedKeys, TArrayV
 			// Populating the key times scratch buffer with the times for these handles
 			Channel->GetKeyTimes(ChannelInfo.KeyHandles, KeyTimesScratch);
 
+			FFrameNumber Offset = MetaData->GetOffsetTime(Section);
 			for(int32 Index = 0; Index < KeyTimesScratch.Num(); ++Index)
 			{
 				int32 OriginalIndex = ChannelInfo.OriginalIndices[Index];
-				OutTimes[OriginalIndex] = KeyTimesScratch[Index];
+				OutTimes[OriginalIndex] = KeyTimesScratch[Index] + Offset;
 			}
 		}
 	}
@@ -121,20 +124,30 @@ void SetKeyTimes(TArrayView<const FSequencerSelectedKey> InSelectedKeys, TArrayV
 
 	for (const FSelectedChannelInfo& ChannelInfo : KeysByChannel.SelectedChannels)
 	{
+		const FMovieSceneChannelMetaData* MetaData = ChannelInfo.Channel.GetMetaData();
 		FMovieSceneChannel* Channel = ChannelInfo.Channel.Get();
-		if (Channel)
+		UMovieSceneSection* Section = ChannelInfo.OwningSection;
+		if (Channel && MetaData && Section)
 		{
+			// Expand to frame first, then offset
 			KeyTimesScratch.Reset(ChannelInfo.OriginalIndices.Num());
 			for (int32 Index : ChannelInfo.OriginalIndices)
 			{
 				KeyTimesScratch.Add(InTimes[Index]);
 
-				if (UMovieSceneSection* Section = ChannelInfo.OwningSection)
+				if (!Section->GetRange().Contains(InTimes[Index]))
+				{
+					Section->ExpandToFrame(InTimes[Index]);
+				}
+			}
+
+			FFrameNumber Offset = MetaData->GetOffsetTime(Section);
+
+			if (Offset != 0)
+			{
+				for (FFrameNumber& Frame : KeyTimesScratch)
 				{
-					if (!Section->GetRange().Contains(InTimes[Index]))
-					{
-						Section->ExpandToFrame(InTimes[Index]);
-					}
+					Frame -= Offset;
 				}
 			}
 
diff --git a/Engine/Source/Editor/Sequencer/Private/Tools/EditToolDragOperations.cpp b/Engine/Source/Editor/Sequencer/Private/Tools/EditToolDragOperations.cpp
--- a/Engine/Source/Editor/Sequencer/Private/Tools/EditToolDragOperations.cpp
+++ b/Engine/Source/Editor/Sequencer/Private/Tools/EditToolDragOperations.cpp
@@ -773,6 +773,8 @@ void FMoveKeysAndSections::AddModel(TSharedPtr<UE::Sequencer::FViewModel> Model)
 
 void FMoveKeysAndSections::OnBeginDrag(const FPointerEvent& MouseEvent, FVector2D LocalMousePos, const UE::Sequencer::FVirtualTrackArea& VirtualTrackArea)
 {
+	using namespace UE::Sequencer;
+
 	// Early out if we've somehow started a drag operation without any sections or keys. This prevents an empty Undo/Redo Transaction from being created.
 	if (!Sections.Num() && !Keys.Num())
 	{
@@ -945,7 +947,7 @@ void FMoveKeysAndSections::OnEndDrag(const FPointerEvent& MouseEvent, FVector2D
 	}
 
 	InitialSectionRowIndicies.Empty();
-	ModifiedNonSelectedSections.Empty();
+	ModifiedNonSelectedObjects.Empty();
 
 	// Tracks can tell us if the row indexes for any sections were changed during our drag/drop operation.
 	TSet<UMovieSceneTrack*> Tracks;
@@ -992,19 +994,29 @@ void FMoveKeysAndSections::OnEndDrag(const FPointerEvent& MouseEvent, FVector2D
 
 void FMoveKeysAndSections::ModifyNonSelectedSections()
 {
+	using namespace UE::Sequencer;
+
 	for (const FSequencerSelectedKey& Key : Keys)
 	{
-		UMovieSceneSection* OwningSection = Key.Section;
-		const bool bHasBeenModified = ModifiedNonSelectedSections.Contains(OwningSection);
-		const bool bIsAlreadySelected = Sections.Contains(OwningSection);
-		if (!bHasBeenModified && !bIsAlreadySelected)
+		TSharedPtr<FChannelModel> Channel = Key.WeakChannel.Pin();
+		UObject* OwningObject = Channel ? Channel->GetOwningObject() : nullptr;
+
+		if (!OwningObject || ModifiedNonSelectedObjects.Contains(OwningObject))
 		{
-			OwningSection->SetFlags(RF_Transactional);
-			if (OwningSection->TryModify())
+			continue;
+		}
+
+		if (UMovieSceneSection* OwningSection = Cast<UMovieSceneSection>(OwningObject))
+		{
+			if (Sections.Contains(OwningSection) || OwningSection->IsReadOnly())
 			{
-				ModifiedNonSelectedSections.Add(OwningSection);
+				continue;
 			}
 		}
+
+		OwningObject->SetFlags(RF_Transactional);
+		OwningObject->Modify();
+		ModifiedNonSelectedObjects.Add(OwningObject);
 	}
 }
 
@@ -1466,7 +1478,7 @@ void FMoveKeysAndSections::HandleKeyMovement(TOptional<FFrameNumber> MaxDeltaX,
 		FSequencerSelectedKey SelectedKey = KeysAsArray[Index];
 
 		UMovieSceneSection* Section = SelectedKey.Section;
-		if (ModifiedNonSelectedSections.Contains(Section))
+		if (ModifiedNonSelectedObjects.Contains(Section))
 		{
 			// If the key moves outside of the section resize the section to fit the key
 			// @todo Sequencer - Doesn't account for hitting other sections 
@@ -1498,12 +1510,12 @@ void FMoveKeysAndSections::HandleKeyMovement(TOptional<FFrameNumber> MaxDeltaX,
 	}
 
 	// Explicitly mark everything as changed to ensure that the UI is responsive during a drag
-	for (TWeakObjectPtr<UMovieSceneSection> Section : ModifiedNonSelectedSections)
+	for (UObject* OwningObject : ModifiedNonSelectedObjects)
 	{
-		if (Section.Get())
+		if (UMovieSceneSignedObject* SignedObject = Cast<UMovieSceneSignedObject>(OwningObject))
 		{
-			Section.Get()->MarkAsChanged();
-			Section.Get()->BroadcastChanged();
+			SignedObject->MarkAsChanged();
+			SignedObject->BroadcastChanged();
 		}
 	}
 }
diff --git a/Engine/Source/Editor/Sequencer/Private/Tools/EditToolDragOperations.h b/Engine/Source/Editor/Sequencer/Private/Tools/EditToolDragOperations.h
--- a/Engine/Source/Editor/Sequencer/Private/Tools/EditToolDragOperations.h
+++ b/Engine/Source/Editor/Sequencer/Private/Tools/EditToolDragOperations.h
@@ -224,8 +224,8 @@ protected:
 	/** Store the row each section starts on when we start dragging. */
 	TArray<FInitialRowIndex> InitialSectionRowIndicies;
 
-	/** Array of sections that we called Modify on because we're editing keys that belong to these sections, but not actually moving these sections. */
-	TArray<TWeakObjectPtr<UMovieSceneSection> > ModifiedNonSelectedSections;
+	/** Array of sections and owners that we called Modify on because we're editing keys that belong to these sections, but not actually moving these sections. */
+	TSet<UObject*> ModifiedNonSelectedObjects;
 
 	/** Optional snap field to use when dragging */
 	TOptional<FSequencerSnapField> SnapField;
diff --git a/Engine/Source/Editor/Sequencer/Public/IKeyArea.h b/Engine/Source/Editor/Sequencer/Public/IKeyArea.h
--- a/Engine/Source/Editor/Sequencer/Public/IKeyArea.h
+++ b/Engine/Source/Editor/Sequencer/Public/IKeyArea.h
@@ -126,6 +126,15 @@ public:
 	 */
 	UMovieSceneSection* GetOwningSection() const;
 
+
+	/**
+	 * Access object that owns the channel this key area represents.
+	 * @note: this may return the same as GetOwningSection (the common case)
+	 *
+	 * @return The owning object, or nullptr if it has been destroyed
+	 */
+	UObject* GetOwningObject() const;
+
 	/**
 	 * Access section editor interface
 	 *
diff --git a/Engine/Source/Editor/Sequencer/Public/ISequencer.h b/Engine/Source/Editor/Sequencer/Public/ISequencer.h
--- a/Engine/Source/Editor/Sequencer/Public/ISequencer.h
+++ b/Engine/Source/Editor/Sequencer/Public/ISequencer.h
@@ -432,7 +432,7 @@ public:
 	 */
 	virtual FQualifiedFrameTime GetGlobalTime() const = 0;
 
-	virtual uint32 GetLocalLoopIndex() const = 0;
+	virtual TOptional<int32> GetLocalLoopIndex() const = 0;
 
 	/**
 	 * Sets the cursor position relative to the currently focused sequence
diff --git a/Engine/Source/Editor/Sequencer/Public/ISequencerChannelInterface.h b/Engine/Source/Editor/Sequencer/Public/ISequencerChannelInterface.h
--- a/Engine/Source/Editor/Sequencer/Public/ISequencerChannelInterface.h
+++ b/Engine/Source/Editor/Sequencer/Public/ISequencerChannelInterface.h
@@ -33,7 +33,24 @@ namespace UE::Sequencer
 	class FChannelModel;
 	class STrackAreaLaneView;
 	struct FCreateTrackLaneViewParams;
-}
+
+	struct FCreateCurveEditorModelParams
+	{
+		UMovieSceneSection* OwningSection;
+		UObject* OwningObject;
+		TSharedRef<ISequencer> Sequencer;
+	};
+
+	struct FCreateKeyEditorParams
+	{
+		UMovieSceneSection* OwningSection;
+		UObject* OwningObject;
+		TSharedRef<ISequencer> Sequencer;
+		FGuid ObjectBindingID;
+		TWeakPtr<FTrackInstancePropertyBindings> PropertyBindings;
+	};
+
+} // namespace UE::Sequencer
 
 /** Utility struct representing a number of selected keys on a single channel */
 struct FExtendKeyMenuParams
@@ -115,13 +132,10 @@ struct ISequencerChannelInterface
 	 * Create an editor on the sequencer node tree
 	 *
 	 * @param Channel               The channel handle to create a key editor for
-	 * @param Section               The section that owns this channel
-	 * @param InObjectBindingID     The ID of the object this key area's track is bound to
-	 * @param PropertyBindings      (Optional) Property bindings where this channel exists on a property track
-	 * @param Sequencer             The currently active sequencer
+	 * @param Params                Creation parameters containing all the necessary structures for creating the key editor
 	 * @return The editor widget to display on the node tree
 	 */
-	virtual TSharedRef<SWidget> CreateKeyEditor_Raw(const FMovieSceneChannelHandle& Channel, UMovieSceneSection* Section, const FGuid& InObjectBindingID, TWeakPtr<FTrackInstancePropertyBindings> PropertyBindings, TWeakPtr<ISequencer> Sequencer) const = 0;
+	virtual TSharedRef<SWidget> CreateKeyEditor_Raw(const FMovieSceneChannelHandle& Channel, const UE::Sequencer::FCreateKeyEditorParams& Params) const = 0;
 
 	/**
 	 * Extend the key context menu
@@ -171,7 +185,7 @@ struct ISequencerChannelInterface
 	 *
 	 * @return (Optional) A new model to be added to a curve editor
 	 */
-	virtual TUniquePtr<FCurveModel> CreateCurveEditorModel_Raw(const FMovieSceneChannelHandle& Channel, UMovieSceneSection* OwningSection, TSharedRef<ISequencer> InSequencer) const = 0;
+	virtual TUniquePtr<FCurveModel> CreateCurveEditorModel_Raw(const FMovieSceneChannelHandle& Channel, const UE::Sequencer::FCreateCurveEditorModelParams& Params) const = 0;
 
 	/**
 	 * Create a new channel model for this type of channel
diff --git a/Engine/Source/Editor/Sequencer/Public/MVVM/ViewModels/ChannelModel.h b/Engine/Source/Editor/Sequencer/Public/MVVM/ViewModels/ChannelModel.h
--- a/Engine/Source/Editor/Sequencer/Public/MVVM/ViewModels/ChannelModel.h
+++ b/Engine/Source/Editor/Sequencer/Public/MVVM/ViewModels/ChannelModel.h
@@ -72,6 +72,9 @@ public:
 	/** Returns the section object that owns the associated channel */
 	UMovieSceneSection* GetSection() const;
 
+	/** Returns the object that owns the associated channel. May return the same as GetSection(). */
+	UObject* GetOwningObject() const;
+
 	/** Returns the associated channel object */
 	FMovieSceneChannel* GetChannel() const;
 
@@ -84,9 +87,6 @@ public:
 	/** Returns the key area for the channel */
 	TSharedPtr<IKeyArea> GetKeyArea() const { return KeyArea; }
 
-	/** Create the curve editor model for the associated channel */
-	void CreateCurveModels(TArray<TUniquePtr<FCurveModel>>& OutCurveModels);
-
 	/** Returns the desired sizing for the track area row */
 	FOutlinerSizing GetDesiredSizing() const;
 
diff --git a/Engine/Source/Editor/Sequencer/Public/SequencerChannelInterface.h b/Engine/Source/Editor/Sequencer/Public/SequencerChannelInterface.h
--- a/Engine/Source/Editor/Sequencer/Public/SequencerChannelInterface.h
+++ b/Engine/Source/Editor/Sequencer/Public/SequencerChannelInterface.h
@@ -8,6 +8,44 @@
 #include "IKeyArea.h"
 #include "CurveModel.h"
 
+namespace Sequencer
+{
+
+/** Concept for detecting deprecated CreateCurveEditorModel signatures that will no longer be called in later versions */
+struct CLegacyCurveModelCreatable
+{
+	template<typename T>
+	auto Requires()
+		-> decltype(CreateCurveEditorModel(*((TMovieSceneChannelHandle<T>*)0), (UMovieSceneSection*)nullptr, TSharedPtr<ISequencer>().ToSharedRef()));
+};
+
+struct CLegacyKeyEditorCreatable
+{
+	template<typename T>
+	auto Requires()
+		-> decltype(CreateKeyEditor(
+			*((TMovieSceneChannelHandle<T>*)nullptr),
+			(UMovieSceneSection*)nullptr,
+			FGuid(),
+			TWeakPtr<FTrackInstancePropertyBindings>(),
+			TWeakPtr<ISequencer>()
+		));
+};
+
+template<typename T>
+UE_DEPRECATED(5.5, "CreateCurveEditorModel(const TMovieSceneChannelHandle<T>&, UMovieSceneSection*, TSharedRef<ISequencer>) has been deprecated. Please update your signature to use FCreateCurveEditorModelParams")
+void CreateCurveEditorModelDeprecatedSignature()
+{
+}
+
+template<typename T>
+UE_DEPRECATED(5.5, "CreateKeyEditor(const TMovieSceneChannelHandle<T>&, UMovieSceneSection*, const FGuid&, TWeakPtr<FTrackInstancePropertyBindings>, TWeakPtr<ISequencer>); has been deprecated. Please update your signature to use FCreateKeyEditorParams")
+void CreateKeyEditorDeprecatedSignature()
+{
+}
+
+} // namespace Sequencer
+
 /**
  * Templated channel interface that calls overloaded functions matching the necessary channel types.
  * Designed this way to allow for specific customization of key-channel behavior without having to reimplement swathes of boilerplate.
@@ -183,10 +221,22 @@ struct TSequencerChannelInterfaceCommon : ISequencerChannelInterface
 	 *
 	 * @return (Optional) A new model to be added to a curve editor
 	 */
-	virtual TUniquePtr<FCurveModel> CreateCurveEditorModel_Raw(const FMovieSceneChannelHandle& InChannel, UMovieSceneSection* OwningSection, TSharedRef<ISequencer> InSequencer) const override
+	virtual TUniquePtr<FCurveModel> CreateCurveEditorModel_Raw(const FMovieSceneChannelHandle& InChannel, const UE::Sequencer::FCreateCurveEditorModelParams& Params) const override
 	{
 		using namespace Sequencer;
-		return CreateCurveEditorModel(InChannel.Cast<ChannelType>(), OwningSection, InSequencer);
+
+		if constexpr (TModels_V<CLegacyCurveModelCreatable, ChannelType>)
+		{
+			// Emit deprecation warning
+			CreateCurveEditorModelDeprecatedSignature<ChannelType>();
+			TUniquePtr<FCurveModel> Result = CreateCurveEditorModel(InChannel.Cast<ChannelType>(), Params.OwningSection, Params.Sequencer);
+			if (Result)
+			{
+				return Result;
+			}
+		}
+
+		return CreateCurveEditorModel(InChannel.Cast<ChannelType>(), Params);
 	}
 
 	/**
@@ -220,16 +270,21 @@ struct TSequencerChannelInterfaceCommon : ISequencerChannelInterface
 	 * Create an editor on the sequencer node tree
 	 *
 	 * @param Channel               The channel to check
-	 * @param Section               The section that owns this channel
-	 * @param InObjectBindingID     The ID of the object this key area's track is bound to
-	 * @param PropertyBindings      (Optional) Property bindings where this channel exists on a property track
-	 * @param Sequencer             The currently active sequencer
+	 * @param Params                Creation parameters containing all the necessary structures for creating the key editor
 	 * @return The editor widget to display on the node tree
 	 */
-	virtual TSharedRef<SWidget> CreateKeyEditor_Raw(const FMovieSceneChannelHandle& Channel, UMovieSceneSection* Section, const FGuid& InObjectBindingID, TWeakPtr<FTrackInstancePropertyBindings> PropertyBindings, TWeakPtr<ISequencer> Sequencer) const override
+	virtual TSharedRef<SWidget> CreateKeyEditor_Raw(const FMovieSceneChannelHandle& Channel, const UE::Sequencer::FCreateKeyEditorParams& Params) const override
 	{
 		using namespace Sequencer;
-		return CreateKeyEditor(Channel.Cast<ChannelType>(), Section, InObjectBindingID, PropertyBindings, Sequencer);
+
+		if constexpr (TModels_V<CLegacyKeyEditorCreatable, ChannelType>)
+		{
+			// Emit deprecation warning
+			CreateKeyEditorDeprecatedSignature<ChannelType>();
+			return CreateKeyEditor(Channel.Cast<ChannelType>(), Params.OwningSection, Params.ObjectBindingID, Params.PropertyBindings, Params.Sequencer);
+		}
+
+		return CreateKeyEditor(Channel.Cast<ChannelType>(), Params);
 	}
 };
 
diff --git a/Engine/Source/Editor/Sequencer/Public/SequencerChannelTraits.h b/Engine/Source/Editor/Sequencer/Public/SequencerChannelTraits.h
--- a/Engine/Source/Editor/Sequencer/Public/SequencerChannelTraits.h
+++ b/Engine/Source/Editor/Sequencer/Public/SequencerChannelTraits.h
@@ -150,25 +150,14 @@ namespace Sequencer
 	 * Create a key editor widget for the specified channel with the channel's specialized editor data. Such widgets are placed on the sequencer node tree for a given key area node.
 	 *
 	 * @param InChannel          The channel to create a key editor for
-	 * @param InOwningSection    The section that owns the channel
-	 * @param InObjectBindingID  The object binding ID that this section's track is bound to
-	 * @param InPropertyBindings Optionally supplied helper for accessing an object's property pertaining to this channel
-	 * @param InSequencer        The sequencer currently active
+	 * @param Params             Creation parameters containing all the necessary structures for creating the key editor
 	 * @return The key editor widget
 	 */
-	inline TSharedRef<SWidget> CreateKeyEditor(
-		const FMovieSceneChannelHandle&          InChannel,
-		UMovieSceneSection*                      InOwningSection,
-		const FGuid&                             InObjectBindingID,
-		TWeakPtr<FTrackInstancePropertyBindings> InPropertyBindings,
-		TWeakPtr<ISequencer>                     InSequencer
-		)
+	inline TSharedRef<SWidget> CreateKeyEditor(const FMovieSceneChannelHandle& InChannel,const UE::Sequencer::FCreateKeyEditorParams& Params)
 	{
 		return SNullWidget::NullWidget;
 	}
 
-
-
 	/**
 	 * Add a key at the specified time (or update an existing key) with the channel's current value at that time
 	 *
@@ -375,11 +364,6 @@ namespace Sequencer
 	template<typename ChannelType>
 	void PasteKeys(ChannelType* InChannel, UMovieSceneSection* Section, const FMovieSceneClipboardKeyTrack& KeyTrack, const FMovieSceneClipboardEnvironment& SrcEnvironment, const FSequencerPasteEnvironment& DstEnvironment, TArray<FKeyHandle>& OutPastedKeys)
 	{
-		if (!Section || !Section->TryModify())
-		{
-			return;
-		}
-
 		FFrameTime PasteAt = DstEnvironment.CardinalTime;
 
 		auto ChannelData = InChannel->GetData();
@@ -414,7 +398,7 @@ namespace Sequencer
 	 *
 	 * @return (Optional) A new model to be added to a curve editor
 	 */
-	SEQUENCER_API TUniquePtr<FCurveModel> CreateCurveEditorModel(const FMovieSceneChannelHandle& ChannelHandle, UMovieSceneSection* OwningSection, TSharedRef<ISequencer> InSequencer);
+	SEQUENCER_API TUniquePtr<FCurveModel> CreateCurveEditorModel(const FMovieSceneChannelHandle& ChannelHandle, const UE::Sequencer::FCreateCurveEditorModelParams& Params);
 
 	/**
 	 * Create a new channel model for this type of channel
diff --git a/Engine/Source/Editor/UnrealEd/Private/Fbx/FbxMainExport.cpp b/Engine/Source/Editor/UnrealEd/Private/Fbx/FbxMainExport.cpp
--- a/Engine/Source/Editor/UnrealEd/Private/Fbx/FbxMainExport.cpp
+++ b/Engine/Source/Editor/UnrealEd/Private/Fbx/FbxMainExport.cpp
@@ -1867,7 +1867,12 @@ int32 FLevelSequenceAnimTrackAdapter::GetStartFrame() const
 	TRange<FFrameNumber> Range = GetSequenceRange();
 	FFrameRate TickResolution = MovieScene->GetTickResolution();
 	FFrameRate DisplayRate = MovieScene->GetDisplayRate();
-	return FFrameRate::TransformTime(FFrameTime(UE::MovieScene::DiscreteInclusiveLower(Range) * RootToLocalTransform.InverseNoLooping()), TickResolution, DisplayRate).RoundToFrame().Value;
+	FFrameNumber StartFrame = UE::MovieScene::DiscreteInclusiveLower(Range);
+
+	FMovieSceneInverseSequenceTransform Inverse = RootToLocalTransform.Inverse();
+
+	TOptional<FFrameTime> RootTime = Inverse.TryTransformTime(StartFrame);
+	return FFrameRate::TransformTime(RootTime.Get(StartFrame), TickResolution, DisplayRate).RoundToFrame().Value;
 }
 
 int32 FLevelSequenceAnimTrackAdapter::GetLength() const
@@ -1875,6 +1880,7 @@ int32 FLevelSequenceAnimTrackAdapter::GetLength() const
 	TRange<FFrameNumber> Range = GetSequenceRange();
 	FFrameRate TickResolution = MovieScene->GetTickResolution();
 	FFrameRate DisplayRate = MovieScene->GetDisplayRate();
+
 	return FFrameRate::TransformTime(FFrameTime(UE::MovieScene::DiscreteSize(Range)), TickResolution, DisplayRate).RoundToFrame().Value;
 }
 
@@ -1884,11 +1890,13 @@ void FLevelSequenceAnimTrackAdapter::UpdateAnimation( int32 LocalFrame )
 	FFrameRate DisplayRate = MovieScene->GetDisplayRate();
 
 	FFrameTime LocalTime = FFrameRate::TransformTime(FFrameTime(LocalFrame), DisplayRate, TickResolution);
-	FFrameTime GlobalTime = LocalTime * RootToLocalTransform.InverseNoLooping();
+	TOptional<FFrameTime> GlobalTime = RootToLocalTransform.Inverse().TryTransformTime(LocalTime);
 
-	FMovieSceneContext Context = FMovieSceneContext(FMovieSceneEvaluationRange(GlobalTime, TickResolution), MovieScenePlayer->GetPlaybackStatus()).SetHasJumped(true);
-
-	MovieScenePlayer->GetEvaluationTemplate().EvaluateSynchronousBlocking( Context );
+	if (GlobalTime)
+	{
+		FMovieSceneContext Context = FMovieSceneContext(FMovieSceneEvaluationRange(GlobalTime.GetValue(), TickResolution), MovieScenePlayer->GetPlaybackStatus()).SetHasJumped(true);
+		MovieScenePlayer->GetEvaluationTemplate().EvaluateSynchronousBlocking( Context );
+	}
 }
 
 double FLevelSequenceAnimTrackAdapter::GetFrameRate() const
@@ -2265,11 +2273,14 @@ void FFbxExporter::AddTimecodeAttributesAndSetKey(const UMovieSceneSection* InSe
 			if (FbxAnimCurve* AnimCurve = Property.GetCurve(BaseLayer, nullptr, true))
 			{
 				const float KeyValue = (float)TimecodeValues[i];
-					
+
 				AnimCurve->KeyModifyBegin();
-					
+
 				FbxTime FbxTime;
-				const double KeyTimeSeconds = GetExportOptions()->bExportLocalTime ? KeyTime / TickResolution : (KeyTime * RootToLocalTransform.InverseNoLooping()) / TickResolution;
+
+				TOptional<FFrameTime> GlobalKeyTime = RootToLocalTransform.Inverse().TryTransformTime(KeyTime);
+
+				const double KeyTimeSeconds = GetExportOptions()->bExportLocalTime ? KeyTime / TickResolution : GlobalKeyTime.Get(KeyTime) / TickResolution;
 
 				FbxTime.SetSecondDouble(KeyTimeSeconds);
 
@@ -2566,9 +2577,10 @@ void FFbxExporter::ExportTransformChannelsToFbxCurve(FbxNode* InFbxNode, TPair<F
 	FFrameRate TickResolution = Track->GetTypedOuter<UMovieScene>()->GetTickResolution();
 	FFrameRate DisplayRate = Track->GetTypedOuter<UMovieScene>()->GetDisplayRate();
 	TRange<FFrameNumber> PlaybackRange = Track->GetTypedOuter<UMovieScene>()->GetPlaybackRange();
-	
+
+	FMovieSceneInverseSequenceTransform SequenceToRootTransform = RootToLocalTransform.Inverse();
+
 	int32 LocalStartFrame = FFrameRate::TransformTime(FFrameTime(UE::MovieScene::DiscreteInclusiveLower(PlaybackRange)), TickResolution, DisplayRate).RoundToFrame().Value;
-	int32 StartFrame = FFrameRate::TransformTime(FFrameTime(UE::MovieScene::DiscreteInclusiveLower(PlaybackRange) * RootToLocalTransform.InverseNoLooping()), TickResolution, DisplayRate).RoundToFrame().Value;
 	int32 AnimationLength = FFrameRate::TransformTime(FFrameTime(FFrameNumber(UE::MovieScene::DiscreteSize(PlaybackRange))), TickResolution, DisplayRate).RoundToFrame().Value;
 
 	for (int32 FrameCount = 0; FrameCount <= AnimationLength; ++FrameCount)
@@ -2619,7 +2631,23 @@ void FFbxExporter::ExportTransformChannelsToFbxCurve(FbxNode* InFbxNode, TPair<F
 		}
 
 		FbxTime FbxTime;
-		FbxTime.SetSecondDouble(GetExportOptions()->bExportLocalTime ? DisplayRate.AsSeconds(LocalFrame) : DisplayRate.AsSeconds(StartFrame + FrameCount));
+
+		if (GetExportOptions()->bExportLocalTime)
+		{
+			FbxTime.SetSecondDouble(DisplayRate.AsSeconds(LocalFrame));
+		}
+		// @todo: This code does not handle the root sequence having a different tick resolution than the local space, but we do
+		//        not have that information here so we have to just assume they match. This should be improved so that we have
+		//        all the information we need to do the right thing
+		else if (TOptional<FFrameTime> GlobalKeyTime = SequenceToRootTransform.TryTransformTime(LocalTime))
+		{
+			FbxTime.SetSecondDouble(TickResolution.AsSeconds(GlobalKeyTime.GetValue()));
+		}
+		else
+		{
+			// Doesn't map to a root time
+			continue;
+		}
 
 		FbxCurveX->KeySet(FbxCurveX->KeyAdd(FbxTime), FbxTime, KeyVec[0]);
 		FbxCurveY->KeySet(FbxCurveY->KeyAdd(FbxTime), FbxTime, KeyVec[1]);
@@ -2971,11 +2999,12 @@ void FFbxExporter::ExportBezierChannelToFbxCurveBaked(FbxAnimCurve& InFbxCurve,
 
 	InFbxCurve.KeyModifyBegin();
 
+	FMovieSceneInverseSequenceTransform SequenceToRootTransform = RootToLocalTransform.Inverse();
+
 	FFrameRate DisplayRate = Track->GetTypedOuter<UMovieScene>()->GetDisplayRate();
 	TRange<FFrameNumber> PlaybackRange = Track->GetTypedOuter<UMovieScene>()->GetPlaybackRange();
 
 	int32 LocalStartFrame = FFrameRate::TransformTime(FFrameTime(UE::MovieScene::DiscreteInclusiveLower(PlaybackRange)), TickResolution, DisplayRate).RoundToFrame().Value;
-	int32 StartFrame = FFrameRate::TransformTime(FFrameTime(UE::MovieScene::DiscreteInclusiveLower(PlaybackRange) * RootToLocalTransform.InverseNoLooping()), TickResolution, DisplayRate).RoundToFrame().Value;
 	int32 AnimationLength = FFrameRate::TransformTime(FFrameTime(FFrameNumber(UE::MovieScene::DiscreteSize(PlaybackRange))), TickResolution, DisplayRate).RoundToFrame().Value;
 
 	for (int32 FrameCount = 0; FrameCount <= AnimationLength; ++FrameCount)
@@ -2988,7 +3017,23 @@ void FFbxExporter::ExportBezierChannelToFbxCurveBaked(FbxAnimCurve& InFbxCurve,
 		InChannel.Evaluate(LocalTime, Value);
 
 		FbxTime FbxTime;
-		FbxTime.SetSecondDouble(GetExportOptions()->bExportLocalTime ? DisplayRate.AsSeconds(LocalFrame) : DisplayRate.AsSeconds(StartFrame + FrameCount));
+
+		if (GetExportOptions()->bExportLocalTime)
+		{
+			FbxTime.SetSecondDouble(DisplayRate.AsSeconds(LocalFrame));
+		}
+		// @todo: This code does not handle the root sequence having a different tick resolution than the local space, but we do
+		//        not have that information here so we have to just assume they match. This should be improved so that we have
+		//        all the information we need to do the right thing
+		else if (TOptional<FFrameTime> GlobalKeyTime = SequenceToRootTransform.TryTransformTime(LocalTime))
+		{
+			FbxTime.SetSecondDouble(TickResolution.AsSeconds(GlobalKeyTime.GetValue()));
+		}
+		else
+		{
+			// Doesn't map to a root time
+			continue;
+		}
 
 		InFbxCurve.KeySet(InFbxCurve.KeyAdd(FbxTime), FbxTime, Value, FbxAnimCurveDef::eInterpolationLinear);
 	}
@@ -3005,6 +3050,8 @@ void FFbxExporter::ExportBezierChannelToFbxCurve(FbxAnimCurve& InFbxCurve, const
 	const float kOneThird = 1.f / 3.f;
 	InFbxCurve.KeyModifyBegin();
 
+	FMovieSceneInverseSequenceTransform SequenceToRootTransform = RootToLocalTransform.Inverse();
+
 	TArrayView<const FFrameNumber>          Times  = InChannel.GetTimes();
 	TArrayView<const ChannelValueType> Values = InChannel.GetValues();
 
@@ -3017,9 +3064,23 @@ void FFbxExporter::ExportBezierChannelToFbxCurve(FbxAnimCurve& InFbxCurve, const
 
 		FbxTime FbxTime;
 		FbxAnimCurveKey FbxKey;
-		const double KeyTimeSeconds = GetExportOptions()->bExportLocalTime ? KeyTime / TickResolution : (KeyTime * RootToLocalTransform.InverseNoLooping()) / TickResolution;
 
-		FbxTime.SetSecondDouble(KeyTimeSeconds);
+		if (GetExportOptions()->bExportLocalTime)
+		{
+			FbxTime.SetSecondDouble(TickResolution.AsSeconds(KeyTime));
+		}
+		// @todo: This code does not handle the root sequence having a different tick resolution than the local space, but we do
+		//        not have that information here so we have to just assume they match. This should be improved so that we have
+		//        all the information we need to do the right thing
+		else if (TOptional<FFrameTime> GlobalKeyTime = SequenceToRootTransform.TryTransformTime(KeyTime))
+		{
+			FbxTime.SetSecondDouble(TickResolution.AsSeconds(GlobalKeyTime.GetValue()));
+		}
+		else
+		{
+			// Doesn't map to a root time
+			continue;
+		}
 
 		const int FbxKeyIndex = InFbxCurve.KeyAdd(FbxTime);
 
@@ -3035,7 +3096,7 @@ void FFbxExporter::ExportBezierChannelToFbxCurve(FbxAnimCurve& InFbxCurve, const
 				float LeaveTangentWeight = kOneThird;
 				float NextArriveTangentWeight = kOneThird;
 				const double  NextTime = Times[Index + 1] / TickResolution;
-				const float TimeDiff = static_cast<float>(NextTime - KeyTimeSeconds);
+				const float TimeDiff = static_cast<float>(NextTime - FbxTime.GetSecondDouble());
 				
 				float LeaveTangent = KeyValue.Tangent.LeaveTangent * TickResolution.AsDecimal();
 				float NextArriveTangent = Values[Index + 1].Tangent.ArriveTangent * TickResolution.AsDecimal();
@@ -3107,15 +3168,31 @@ void FFbxExporter::ExportConstantChannelToFbxCurve(FbxAnimCurve& InFbxCurve, con
 	const TArrayView<const FFrameNumber> Times  = InChannel.GetTimes();
 	const TArrayView<const T> Values = InChannel.GetValues();
 
+	FMovieSceneInverseSequenceTransform SequenceToRootTransform = RootToLocalTransform.Inverse();
+
 	for (int32 Index = 0; Index < Times.Num(); ++Index)
 	{
 		const FFrameNumber KeyTime = Times[Index];
 		const T KeyValue = Values[Index];
 
 		FbxTime FbxTime;
-		const double KeyTimeSeconds = GetExportOptions()->bExportLocalTime ? KeyTime / TickResolution : (KeyTime * RootToLocalTransform.InverseNoLooping()) / TickResolution;
 
-		FbxTime.SetSecondDouble(KeyTimeSeconds);
+		if (GetExportOptions()->bExportLocalTime)
+		{
+			FbxTime.SetSecondDouble(TickResolution.AsSeconds(KeyTime));
+		}
+		// @todo: This code does not handle the root sequence having a different tick resolution than the local space, but we do
+		//        not have that information here so we have to just assume they match. This should be improved so that we have
+		//        all the information we need to do the right thing
+		else if (TOptional<FFrameTime> GlobalKeyTime = SequenceToRootTransform.TryTransformTime(KeyTime))
+		{
+			FbxTime.SetSecondDouble(TickResolution.AsSeconds(GlobalKeyTime.GetValue()));
+		}
+		else
+		{
+			// Doesn't map to a root time
+			continue;
+		}
 
 		const int FbxKeyIndex = InFbxCurve.KeyAdd(FbxTime);
 
@@ -3235,6 +3312,8 @@ void FFbxExporter::ExportLevelSequence3DTransformTrack(FbxNode* FbxNode, IMovieS
 		return;
 	}
 
+	FMovieSceneInverseSequenceTransform SequenceToRootTransform = RootToLocalTransform.Inverse();
+
 	FbxAnimLayer* BaseLayer = AnimStack->GetMember<FbxAnimLayer>(0);
 
 	AActor* BoundActor = Cast<AActor>(BoundObject);
@@ -3349,7 +3428,6 @@ void FFbxExporter::ExportLevelSequence3DTransformTrack(FbxNode* FbxNode, IMovieS
 		FbxCurveRotZ->KeyModifyBegin();
 
 		int32 LocalStartFrame = FFrameRate::TransformTime(FFrameTime(UE::MovieScene::DiscreteInclusiveLower(InPlaybackRange)), TickResolution, DisplayRate).RoundToFrame().Value;
-		int32 StartFrame = FFrameRate::TransformTime(FFrameTime(UE::MovieScene::DiscreteInclusiveLower(InPlaybackRange) * RootToLocalTransform.InverseNoLooping()), TickResolution, DisplayRate).RoundToFrame().Value;
 		int32 AnimationLength = FFrameRate::TransformTime(FFrameTime(FFrameNumber(UE::MovieScene::DiscreteSize(InPlaybackRange))), TickResolution, DisplayRate).RoundToFrame().Value;
 
 		for (int32 FrameCount = 0; FrameCount <= AnimationLength; ++FrameCount)
@@ -3412,7 +3490,22 @@ void FFbxExporter::ExportLevelSequence3DTransformTrack(FbxNode* FbxNode, IMovieS
 			FbxVector4 KeyScale = Converter.ConvertToFbxScale(RelativeTransform.GetScale3D());
 
 			FbxTime FbxTime;
-			FbxTime.SetSecondDouble(GetExportOptions()->bExportLocalTime ? DisplayRate.AsSeconds(LocalFrame) : DisplayRate.AsSeconds(StartFrame + FrameCount));
+			if (GetExportOptions()->bExportLocalTime)
+			{
+				FbxTime.SetSecondDouble(DisplayRate.AsSeconds(LocalFrame));
+			}
+			// @todo: This code does not handle the root sequence having a different tick resolution than the local space, but we do
+			//        not have that information here so we have to just assume they match. This should be improved so that we have
+			//        all the information we need to do the right thing
+			else if (TOptional<FFrameTime> GlobalKeyTime = SequenceToRootTransform.TryTransformTime(LocalTime))
+			{
+				FbxTime.SetSecondDouble(TickResolution.AsSeconds(GlobalKeyTime.GetValue()));
+			}
+			else
+			{
+				// Doesn't map to a root time
+				continue;
+			}
 
 			FbxCurveRotX->KeySet(FbxCurveRotX->KeyAdd(FbxTime), FbxTime, KeyRot[0]);
 			FbxCurveRotY->KeySet(FbxCurveRotY->KeyAdd(FbxTime), FbxTime, KeyRot[1]);
@@ -3534,7 +3627,7 @@ void FFbxExporter::ExportLevelSequenceBaked3DTransformTrack(IAnimTrackAdapter& A
 		FbxCurveScaleZ->KeyModifyBegin();
 	}
 
-	FMovieSceneSequenceTransform LocalToRootTransform = RootToLocalTransform.InverseNoLooping();
+	FMovieSceneInverseSequenceTransform LocalToRootTransform = RootToLocalTransform.Inverse();
 
 	TArray<FTransform> RelativeTransforms;
 	int32 LocalStartFrame = FFrameRate::TransformTime(FFrameTime(DiscreteInclusiveLower(InPlaybackRange)), TickResolution, DisplayRate).RoundToFrame().Value;
@@ -3599,15 +3692,25 @@ void FFbxExporter::ExportLevelSequenceBaked3DTransformTrack(IAnimTrackAdapter& A
 		FbxVector4 KeyScale = Converter.ConvertToFbxScale(RelativeTransform.GetScale3D());
 
 		const int32 CurrentFrame = LocalStartFrame + TransformIndex;
+
+		const FFrameTime LocalTime = FFrameRate::TransformTime(CurrentFrame, DisplayRate, TickResolution);
+
 		FbxTime FbxTime;
 		if (GetExportOptions()->bExportLocalTime)
 		{
 			FbxTime.SetSecondDouble(DisplayRate.AsSeconds(CurrentFrame));
 		}
+		// @todo: This code does not handle the root sequence having a different tick resolution than the local space, but we do
+		//        not have that information here so we have to just assume they match. This should be improved so that we have
+		//        all the information we need to do the right thing
+		else if (TOptional<FFrameTime> GlobalKeyTime = LocalToRootTransform.TryTransformTime(LocalTime))
+		{
+			FbxTime.SetSecondDouble(TickResolution.AsSeconds(GlobalKeyTime.GetValue()));
+		}
 		else
 		{
-			FFrameTime CurrentTime = FFrameRate::TransformTime(CurrentFrame, DisplayRate, TickResolution) * LocalToRootTransform;
-			FbxTime.SetSecondDouble(DisplayRate.AsSeconds(FFrameRate::TransformTime(CurrentTime, TickResolution, DisplayRate)));
+			// Doesn't map to a root time
+			continue;
 		}
 
 		FbxCurveTransX->KeySet(FbxCurveTransX->KeyAdd(FbxTime), FbxTime, KeyTrans[0]);
diff --git a/Engine/Source/Runtime/MovieScene/Private/Channels/MovieSceneChannelEditorData.cpp b/Engine/Source/Runtime/MovieScene/Private/Channels/MovieSceneChannelEditorData.cpp
--- a/Engine/Source/Runtime/MovieScene/Private/Channels/MovieSceneChannelEditorData.cpp
+++ b/Engine/Source/Runtime/MovieScene/Private/Channels/MovieSceneChannelEditorData.cpp
@@ -1,6 +1,7 @@
 // Copyright Epic Games, Inc. All Rights Reserved.
 
 #include "Channels/MovieSceneChannelEditorData.h"
+#include "MovieSceneSection.h"
 
 #if WITH_EDITOR
 
@@ -24,17 +25,21 @@ const FName FCommonChannelData::GroupDisplayName = TEXT("GroupDisplayName");
 FMovieSceneChannelMetaData::FMovieSceneChannelMetaData()
 	: bEnabled(true)
 	, bCanCollapseToTrack(true)
+	, bRelativeToSection(false)
 	, SortOrder(0)
 	, Name(NAME_None)
+	, KeyOffset(0)
 {}
 
 FMovieSceneChannelMetaData::FMovieSceneChannelMetaData(FName InName, FText InDisplayText, FText InGroup, bool bInEnabled)
 	: bEnabled(bInEnabled)
 	, bCanCollapseToTrack(true)
+	, bRelativeToSection(false)
 	, SortOrder(0)
 	, Name(InName)
 	, DisplayText(InDisplayText)
 	, Group(InGroup)
+	, KeyOffset(0)
 {}
 
 void FMovieSceneChannelMetaData::SetIdentifiers(FName InName, FText InDisplayText, FText InGroup)
@@ -49,4 +54,12 @@ FString FMovieSceneChannelMetaData::GetPropertyMetaData(const FName& InKey) cons
 	return PropertyMetaData.FindRef(InKey);
 }
 
+FFrameNumber FMovieSceneChannelMetaData::GetOffsetTime(const UMovieSceneSection* InSection) const
+{
+	FFrameNumber OffsetTime = (bRelativeToSection && InSection->HasStartFrame())
+		? InSection->GetInclusiveStartFrame()
+		: 0;
+	return OffsetTime + KeyOffset.Get();
+}
+
 #endif	// WITH_EDITOR
\ No newline at end of file
