[Sequencer] Fix gizmo hover cursor icon when in SequencerEdMode.

diff --git a/Engine/Source/Editor/Sequencer/Private/SequencerEdMode.cpp b/Engine/Source/Editor/Sequencer/Private/SequencerEdMode.cpp
--- a/Engine/Source/Editor/Sequencer/Private/SequencerEdMode.cpp
+++ b/Engine/Source/Editor/Sequencer/Private/SequencerEdMode.cpp
@@ -489,9 +489,7 @@ bool FSequencerEdMode::GetCursor(EMouseCursor::Type& OutCursor) const
 		return true;
 	}
 
-	OutCursor = EMouseCursor::Default;
-
-	return true;
+	return FEdMode::GetCursor(OutCursor);
 }
 
 bool FSequencerEdMode::MouseMove(FEditorViewportClient* InViewportClient, FViewport* InViewport, int32 InX, int32 InY)
