Sequencer: Fix spawnable creation issue that was causing actor factories to be ignored under certain circumstances. This was causing issues with volume initialization.

diff --git a/Engine/Source/Editor/Sequencer/Private/SequencerUtilities.cpp b/Engine/Source/Editor/Sequencer/Private/SequencerUtilities.cpp
--- a/Engine/Source/Editor/Sequencer/Private/SequencerUtilities.cpp
+++ b/Engine/Source/Editor/Sequencer/Private/SequencerUtilities.cpp
@@ -2436,6 +2436,14 @@ void ImportObjectBindingsFromText(ISequencer& InSequencer, const FString& TextTo
 FGuid TryCreateCustomBinding(TSharedPtr<ISequencer> Sequencer, UObject* CustomBindingObject, AActor* FactoryCreatedActor, FMovieSceneBindingReferences* BindingReferences, const UE::Sequencer::FCreateBindingParams& InParams, UMovieScene* OwnerMovieScene, bool bSpawnable, bool bReplaceable)
 {
 	UMovieSceneCustomBinding * NewCustomBinding = nullptr;
+
+	// If the passed in object is a UClass, and we have an actor factory created instance, prioritize that, otherwise let the binding choose
+	if (CustomBindingObject && FactoryCreatedActor && CustomBindingObject->IsA<UClass>())
+	{
+		CustomBindingObject = FactoryCreatedActor;
+		FactoryCreatedActor = nullptr;
+	}
+
 	if (InParams.CustomBinding)
 	{
 		const FMovieSceneBindingReference* PreviousBindingReference = BindingReferences->GetReference(InParams.ReplacementGuid, InParams.BindingIndex);
