Sequencer: Fix regression that caused OnActorAddedToSequencer to be called when converting between Possessable and Spawnable. This was causing some tracks (e.g. Media Tracks) to be doubly added- once upon creation of a Possessable, and once upon conversion from Possessable to Spawnable.

This also fixes a related bug where spawn tracks were not automatically being created when adding spawnable bindings from a non-Sequencer context such as MHA.

diff --git a/Engine/Source/Editor/Sequencer/Private/SequencerUtilities.cpp b/Engine/Source/Editor/Sequencer/Private/SequencerUtilities.cpp
--- a/Engine/Source/Editor/Sequencer/Private/SequencerUtilities.cpp
+++ b/Engine/Source/Editor/Sequencer/Private/SequencerUtilities.cpp
@@ -1448,14 +1448,15 @@ FMovieScenePossessable* FSequencerUtilities::ConvertToCustomBinding(TSharedRef<I
 	// If this is a new spawnable or replaceable binding, we need to set up some defaults
 	if (NewCustomBinding->WillSpawnObject(Sequencer->GetSharedPlaybackState()))
 	{
-		// Spawn the object so we can position it correctly, it's going to get spawned anyway since things default to spawned.
-		UObject* SpawnedObject = Sequencer->GetSpawnRegister().SpawnObject(NewPossessableGuid, *MovieScene, Sequencer->GetFocusedTemplateID(), Sequencer->GetSharedPlaybackState(), BindingIndex);
+		// We purposefully pass in nullptr to SetupDefaultsForSpawnable below. 
+		// This will prevent a section of code in it from calling OnActorAddedToSequencer, which should not be called in the case of binding conversion,
+		// as it may cause some default tracks to get added for a second time.
 
 		// Allow the binding to set up any necessary defaults
-		NewCustomBinding->SetupDefaults(SpawnedObject, NewPossessableGuid, *MovieScene, Sequencer->GetSharedPlaybackState());
+		NewCustomBinding->SetupDefaults(nullptr, NewPossessableGuid, *MovieScene);
 
 		FTransformData TransformData;
-		Sequencer->GetSpawnRegister().SetupDefaultsForSpawnable(SpawnedObject, NewPossessableGuid, TransformData, Sequencer, Sequencer->GetSequencerSettings());
+		Sequencer->GetSpawnRegister().SetupDefaultsForSpawnable(nullptr, NewPossessableGuid, TransformData, Sequencer, Sequencer->GetSequencerSettings());
 	}
 
 	//Sequencer->State.Invalidate(NewPossessableGuid, Sequencer->GetFocusedTemplateID());
@@ -2534,6 +2535,7 @@ FGuid TryCreateCustomBinding(TSharedPtr<ISequencer> Sequencer, UObject* CustomBi
 		BindingReferences->AddOrReplaceBinding(NewID, NewCustomBinding, InParams.BindingIndex);
 
 		UObject* SpawnedObject = nullptr;
+
 		// If this is a spawnable or replaceable binding, we need to set up some defaults
 		if (Sequencer)
 		{
@@ -2541,17 +2543,19 @@ FGuid TryCreateCustomBinding(TSharedPtr<ISequencer> Sequencer, UObject* CustomBi
 			{
 				// Spawn the object so we can position it correctly, it's going to get spawned anyway since things default to spawned.
 				SpawnedObject = Sequencer->GetSpawnRegister().SpawnObject(NewID, *OwnerMovieScene, Sequencer->GetFocusedTemplateID(), Sequencer->GetSharedPlaybackState(), 0);
+			}
+		}
 
-				if (InParams.bSetupDefaults)
-				{
-					// Allow the binding to set up any necessary defaults
-					NewCustomBinding->SetupDefaults(SpawnedObject, NewID, *OwnerMovieScene, Sequencer->GetSharedPlaybackState());
+		// Allow the binding to set up any necessary defaults
+		NewCustomBinding->SetupDefaults(SpawnedObject, NewID, *OwnerMovieScene);
 
-					FTransformData TransformData;
-					Sequencer->GetSpawnRegister().SetupDefaultsForSpawnable(SpawnedObject, NewID, TransformData, Sequencer.ToSharedRef(), Sequencer->GetSequencerSettings());
-				}
+		if (Sequencer)
+		{
+			if (InParams.bSetupDefaults)
+			{
+				FTransformData TransformData;
+				Sequencer->GetSpawnRegister().SetupDefaultsForSpawnable(SpawnedObject, NewID, TransformData, Sequencer.ToSharedRef(), Sequencer->GetSequencerSettings());
 			}
-
 			Sequencer->State.Invalidate(NewID, Sequencer->GetFocusedTemplateID());
 			Sequencer->ForceEvaluate();
 
diff --git a/Engine/Source/Runtime/MovieScene/Private/Bindings/MovieSceneReplaceableBinding.cpp b/Engine/Source/Runtime/MovieScene/Private/Bindings/MovieSceneReplaceableBinding.cpp
--- a/Engine/Source/Runtime/MovieScene/Private/Bindings/MovieSceneReplaceableBinding.cpp
+++ b/Engine/Source/Runtime/MovieScene/Private/Bindings/MovieSceneReplaceableBinding.cpp
@@ -22,9 +22,9 @@ namespace UE::MovieScene::ReplaceableBinding
 } // namespace UE::MovieScene::ReplaceableBinding
 
 #if WITH_EDITOR
-void UMovieSceneReplaceableBindingBase::SetupDefaults(UObject* SpawnedObject, FGuid ObjectBindingId, UMovieScene& OwnerMovieScene, TSharedRef<const UE::MovieScene::FSharedPlaybackState> SharedPlaybackState)
+void UMovieSceneReplaceableBindingBase::SetupDefaults(UObject* SpawnedObject, FGuid ObjectBindingId, UMovieScene& OwnerMovieScene)
 {
-	Super::SetupDefaults(SpawnedObject, ObjectBindingId, OwnerMovieScene, SharedPlaybackState);
+	Super::SetupDefaults(SpawnedObject, ObjectBindingId, OwnerMovieScene);
 	// Ensure it has a binding lifetime track which it will need in editor
 	UMovieSceneBindingLifetimeTrack* BindingLifetimeTrack = Cast<UMovieSceneBindingLifetimeTrack>(OwnerMovieScene.FindTrack(UMovieSceneBindingLifetimeTrack::StaticClass(), ObjectBindingId, NAME_None));
 	if (!BindingLifetimeTrack)
diff --git a/Engine/Source/Runtime/MovieScene/Private/Bindings/MovieSceneSpawnableBinding.cpp b/Engine/Source/Runtime/MovieScene/Private/Bindings/MovieSceneSpawnableBinding.cpp
--- a/Engine/Source/Runtime/MovieScene/Private/Bindings/MovieSceneSpawnableBinding.cpp
+++ b/Engine/Source/Runtime/MovieScene/Private/Bindings/MovieSceneSpawnableBinding.cpp
@@ -155,9 +155,9 @@ void UMovieSceneSpawnableBindingBase::DestroySpawnedObject(UObject* Object)
 }
 
 #if WITH_EDITOR
-void UMovieSceneSpawnableBindingBase::SetupDefaults(UObject* SpawnedObject, FGuid ObjectBindingId, UMovieScene& OwnerMovieScene, TSharedRef<const UE::MovieScene::FSharedPlaybackState> SharedPlaybackState)
+void UMovieSceneSpawnableBindingBase::SetupDefaults(UObject* SpawnedObject, FGuid ObjectBindingId, UMovieScene& OwnerMovieScene)
 {
-	Super::SetupDefaults(SpawnedObject, ObjectBindingId, OwnerMovieScene, SharedPlaybackState);
+	Super::SetupDefaults(SpawnedObject, ObjectBindingId, OwnerMovieScene);
 
 	// TODO: For now we are not using binding lifetime track for this, though it will support it. We continue to use spawn track until we improve UX of splitting sections
 	// // 
diff --git a/Engine/Source/Runtime/MovieScene/Public/Bindings/MovieSceneCustomBinding.h b/Engine/Source/Runtime/MovieScene/Public/Bindings/MovieSceneCustomBinding.h
--- a/Engine/Source/Runtime/MovieScene/Public/Bindings/MovieSceneCustomBinding.h
+++ b/Engine/Source/Runtime/MovieScene/Public/Bindings/MovieSceneCustomBinding.h
@@ -100,7 +100,7 @@ public:
 	/*
 	* Called by Sequencer upon creating a new custom binding or converting a binding to use this type. Can be used by custom binding types to add required track types, etc.
 	*/
-	virtual void SetupDefaults(UObject* SpawnedObject, FGuid ObjectBindingId, UMovieScene& OwnerMovieScene, TSharedRef<const UE::MovieScene::FSharedPlaybackState> SharedPlaybackState) {}
+	virtual void SetupDefaults(UObject* SpawnedObject, FGuid ObjectBindingId, UMovieScene& OwnerMovieScene) {}
 	
 	/*
 	* Allows the custom binding to optionally provide a custom icon overlay for the object binding track.
diff --git a/Engine/Source/Runtime/MovieScene/Public/Bindings/MovieSceneReplaceableBinding.h b/Engine/Source/Runtime/MovieScene/Public/Bindings/MovieSceneReplaceableBinding.h
--- a/Engine/Source/Runtime/MovieScene/Public/Bindings/MovieSceneReplaceableBinding.h
+++ b/Engine/Source/Runtime/MovieScene/Public/Bindings/MovieSceneReplaceableBinding.h
@@ -41,7 +41,7 @@ public:
 	virtual bool SupportsBindingCreationFromObject(const UObject* SourceObject) const override;
 	virtual UClass* GetBoundObjectClass() const override;
 #if WITH_EDITOR
-	virtual	void SetupDefaults(UObject* SpawnedObject, FGuid ObjectBindingId, UMovieScene& OwnerMovieScene, TSharedRef<const UE::MovieScene::FSharedPlaybackState> SharedPlaybackState) override;
+	virtual	void SetupDefaults(UObject* SpawnedObject, FGuid ObjectBindingId, UMovieScene& OwnerMovieScene) override;
 	virtual const FSlateBrush* GetBindingTrackCustomIconOverlay() const override;
 	virtual FText GetBindingTrackIconTooltip() const override;
 
diff --git a/Engine/Source/Runtime/MovieScene/Public/Bindings/MovieSceneSpawnableBinding.h b/Engine/Source/Runtime/MovieScene/Public/Bindings/MovieSceneSpawnableBinding.h
--- a/Engine/Source/Runtime/MovieScene/Public/Bindings/MovieSceneSpawnableBinding.h
+++ b/Engine/Source/Runtime/MovieScene/Public/Bindings/MovieSceneSpawnableBinding.h
@@ -78,7 +78,7 @@ public:
 
 #if WITH_EDITOR
 	/* UMovieSceneCustomBinding overrides */
-	virtual	void SetupDefaults(UObject* SpawnedObject, FGuid ObjectBindingId, UMovieScene& OwnerMovieScene, TSharedRef<const UE::MovieScene::FSharedPlaybackState> SharedPlaybackState) override;
+	virtual	void SetupDefaults(UObject* SpawnedObject, FGuid ObjectBindingId, UMovieScene& OwnerMovieScene) override;
 	virtual const FSlateBrush* GetBindingTrackCustomIconOverlay() const override;
 	virtual FText GetBindingTrackIconTooltip() const override;
 #endif
