Sequencer: Fixed keybar dragging not working on timewarp channels

diff --git a/Engine/Source/Editor/Sequencer/Private/SequencerHotspots.cpp b/Engine/Source/Editor/Sequencer/Private/SequencerHotspots.cpp
--- a/Engine/Source/Editor/Sequencer/Private/SequencerHotspots.cpp
+++ b/Engine/Source/Editor/Sequencer/Private/SequencerHotspots.cpp
@@ -311,7 +311,7 @@ TSharedPtr<ISequencerEditToolDragOperation> FKeyBarHotspot::InitiateDrag(const F
 
 		TSharedPtr<FKeyBarHotspot> Hotspot;
 		TUniquePtr<FScopedTransaction> Transaction;
-		TSet<UMovieSceneSection*> ModifiedSections;
+		TSet<UObject*> ModifiedObjects;
 		FSequencerSnapField SnapField;
 
 		FKeyBarDrag(TSharedPtr<FKeyBarHotspot> InHotspot, TSharedPtr<FSequencer> InSequencer)
@@ -347,9 +347,12 @@ TSharedPtr<ISequencerEditToolDragOperation> FKeyBarHotspot::InitiateDrag(const F
 
 			for (const FSequencerSelectedKey& Key : AllLinearKeys)
 			{
-				if (!ModifiedSections.Contains(Key.Section))
+				TSharedPtr<FChannelModel> Channel = Key.WeakChannel.Pin();
+				UObject* OwningObject = Channel ? Channel->GetOwningObject() : nullptr;
+
+				if (OwningObject && !ModifiedObjects.Contains(OwningObject))
 				{
-					ModifiedSections.Add(Key.Section);
+					ModifiedObjects.Add(OwningObject);
 				}
 			}
 		}
@@ -359,9 +362,9 @@ TSharedPtr<ISequencerEditToolDragOperation> FKeyBarHotspot::InitiateDrag(const F
 
 			const FFrameTime NewTime = VirtualTrackArea.PixelToFrame(LocalMousePos.X);
 
-			for (UMovieSceneSection* Section : ModifiedSections)
+			for (UObject* Object : ModifiedObjects)
 			{
-				Section->Modify();
+				Object->Modify();
 			}
 
 			// Set the position of leading and trailing keys
