Sequencer: Fixed layer bar selection not correctly honoring SHIFT and CTRL

diff --git a/Engine/Source/Editor/Sequencer/Private/SequencerHotspots.cpp b/Engine/Source/Editor/Sequencer/Private/SequencerHotspots.cpp
--- a/Engine/Source/Editor/Sequencer/Private/SequencerHotspots.cpp
+++ b/Engine/Source/Editor/Sequencer/Private/SequencerHotspots.cpp
@@ -141,6 +141,21 @@ void FHotspotSelectionManager::SelectKeysExclusive(TArrayView<const FSequencerSe
 		}
 	}
 }
+void FHotspotSelectionManager::DefaultModelSelection(TSharedPtr<FViewModel> InModel)
+{
+	if (bForceSelect)
+	{
+		SelectModelExclusive(InModel);
+	}
+	else if (MouseEvent->IsControlDown())
+	{
+		ToggleModel(InModel);
+	}
+	else if (MouseEvent->IsShiftDown())
+	{
+		Selection->TrackArea.Select(InModel);
+	}
+}
 
 void FHotspotSelectionManager::SelectModelExclusive(TSharedPtr<FViewModel> InModel)
 {
diff --git a/Engine/Source/Editor/Sequencer/Private/SequencerHotspots.h b/Engine/Source/Editor/Sequencer/Private/SequencerHotspots.h
--- a/Engine/Source/Editor/Sequencer/Private/SequencerHotspots.h
+++ b/Engine/Source/Editor/Sequencer/Private/SequencerHotspots.h
@@ -41,6 +41,7 @@ struct FHotspotSelectionManager
 
 	void ToggleKeys(TArrayView<const FSequencerSelectedKey> InKeys);
 	void ToggleModel(TSharedPtr<FViewModel> InModel);
+	void DefaultModelSelection(TSharedPtr<FViewModel> InModel);
 
 	void SelectKeysExclusive(TArrayView<const FSequencerSelectedKey> InKeys);
 	void SelectModelExclusive(TSharedPtr<FViewModel> InModel);
diff --git a/Engine/Source/Editor/Sequencer/Private/Widgets/SSequencerLayerBar.cpp b/Engine/Source/Editor/Sequencer/Private/Widgets/SSequencerLayerBar.cpp
--- a/Engine/Source/Editor/Sequencer/Private/Widgets/SSequencerLayerBar.cpp
+++ b/Engine/Source/Editor/Sequencer/Private/Widgets/SSequencerLayerBar.cpp
@@ -82,7 +82,7 @@ struct FLayerBarHotspot
 	}
 	void HandleMouseSelection(FHotspotSelectionManager& SelectionManager) override
 	{
-		SelectionManager.SelectModelExclusive(LayerBar);
+		SelectionManager.DefaultModelSelection(LayerBar);
 	}
 
 	TSharedPtr<FLayerBarModel> LayerBar;
