Sequencer: Change 'Set Selection Start/End' to use the upper/lower playback range by default

diff --git a/Engine/Source/Editor/Sequencer/Private/Sequencer.cpp b/Engine/Source/Editor/Sequencer/Private/Sequencer.cpp
--- a/Engine/Source/Editor/Sequencer/Private/Sequencer.cpp
+++ b/Engine/Source/Editor/Sequencer/Private/Sequencer.cpp
@@ -2675,16 +2675,22 @@ void FSequencer::SetSelectionRange(TRange<FFrameNumber> Range)
 
 void FSequencer::SetSelectionRangeEnd(FFrameTime EndFrame)
 {
+	using namespace UE::MovieScene;
+
 	const bool bInitiallyEmpty = GetSelectionRange().IsEmpty();
 	const FFrameNumber LocalTime = EndFrame.FrameNumber;
 
-	if (GetSelectionRange().GetLowerBoundValue() >= LocalTime)
+	const FFrameNumber StartFrame = (bInitiallyEmpty || GetSelectionRange().GetLowerBoundValue() >= LocalTime)
+		? DiscreteInclusiveLower(GetFocusedMovieSceneSequence()->GetMovieScene()->GetPlaybackRange())
+		: GetSelectionRange().GetLowerBoundValue();
+
+	if (StartFrame >= LocalTime)
 	{
 		SetSelectionRange(TRange<FFrameNumber>(LocalTime - 1, LocalTime));
 	}
 	else
 	{
-		SetSelectionRange(TRange<FFrameNumber>(GetSelectionRange().GetLowerBound(), LocalTime));
+		SetSelectionRange(TRange<FFrameNumber>(StartFrame, LocalTime));
 	}
 
 	if (bInitiallyEmpty && GetLoopMode() != ESequencerLoopMode::SLM_LoopSelectionRange)
@@ -2696,16 +2702,22 @@ void FSequencer::SetSelectionRangeEnd(FFrameTime EndFrame)
 
 void FSequencer::SetSelectionRangeStart(FFrameTime StartFrame)
 {
+	using namespace UE::MovieScene;
+
 	const bool bInitiallyEmpty = GetSelectionRange().IsEmpty();
 	const FFrameNumber LocalTime = StartFrame.FrameNumber;
 
-	if (GetSelectionRange().GetUpperBoundValue() <= LocalTime)
+	const FFrameNumber EndFrame = (bInitiallyEmpty || GetSelectionRange().GetUpperBoundValue() <= LocalTime)
+		? DiscreteExclusiveUpper(GetFocusedMovieSceneSequence()->GetMovieScene()->GetPlaybackRange())
+		: GetSelectionRange().GetUpperBoundValue();
+
+	if (EndFrame <= LocalTime)
 	{
 		SetSelectionRange(TRange<FFrameNumber>(LocalTime, LocalTime + 1));
 	}
 	else
 	{
-		SetSelectionRange(TRange<FFrameNumber>(LocalTime, GetSelectionRange().GetUpperBound()));
+		SetSelectionRange(TRange<FFrameNumber>(LocalTime, EndFrame));
 	}
 
 	if (bInitiallyEmpty && GetLoopMode() != ESequencerLoopMode::SLM_LoopSelectionRange)
