Sequencer- Fix bug where converting from possessable to spawnable doesn't keep the previous position of the possessable.

diff --git a/Engine/Source/Editor/Sequencer/Private/SequencerUtilities.cpp b/Engine/Source/Editor/Sequencer/Private/SequencerUtilities.cpp
--- a/Engine/Source/Editor/Sequencer/Private/SequencerUtilities.cpp
+++ b/Engine/Source/Editor/Sequencer/Private/SequencerUtilities.cpp
@@ -1434,6 +1434,7 @@ FMovieScenePossessable* FSequencerUtilities::ConvertToCustomBinding(TSharedRef<I
 		Sequencer->ForceEvaluate();
 	}
 
+	TOptional<FTransformData> TransformData;
 	if (bConvertFromSpawnable)
 	{
 		Sequencer->GetSpawnRegister().DestroySpawnedObject(BindingGuid, Sequencer->GetFocusedTemplateID(), Sequencer->GetSharedPlaybackState(), BindingIndex);
@@ -1441,7 +1442,6 @@ FMovieScenePossessable* FSequencerUtilities::ConvertToCustomBinding(TSharedRef<I
 	else if (bConvertFromPossessable)
 	{
 		// We have an old possessable to destroy
-		TOptional<FTransformData> TransformData;
 		Sequencer->GetSpawnRegister().HandleConvertPossessableToSpawnable(ObjectToConvert, *Sequencer, TransformData);
 	}
 
@@ -1455,7 +1455,6 @@ FMovieScenePossessable* FSequencerUtilities::ConvertToCustomBinding(TSharedRef<I
 		// Allow the binding to set up any necessary defaults
 		NewCustomBinding->SetupDefaults(nullptr, NewPossessableGuid, *MovieScene);
 
-		FTransformData TransformData;
 		Sequencer->GetSpawnRegister().SetupDefaultsForSpawnable(nullptr, NewPossessableGuid, TransformData, Sequencer, Sequencer->GetSequencerSettings());
 	}
 
