[Sequencer] Sidebar Details: Motion Design and Sequencer implementations

diff --git a/Engine/Plugins/Experimental/Avalanche/Source/AvalancheSequencer/Private/AvaSequencer.cpp b/Engine/Plugins/Experimental/Avalanche/Source/AvalancheSequencer/Private/AvaSequencer.cpp
--- a/Engine/Plugins/Experimental/Avalanche/Source/AvalancheSequencer/Private/AvaSequencer.cpp
+++ b/Engine/Plugins/Experimental/Avalanche/Source/AvalancheSequencer/Private/AvaSequencer.cpp
@@ -5,7 +5,6 @@
 #include "AvaSequencePlaybackObject.h"
 #include "AvaSequencePlayer.h"
 #include "AvaSequencerArgs.h"
-#include "AvaSequencerModule.h"
 #include "AvaSequencerUtils.h"
 #include "Clipboard/AvaSequenceExporter.h"
 #include "Clipboard/AvaSequenceImporter.h"
@@ -13,16 +12,16 @@
 #include "Commands/AvaSequencerCommands.h"
 #include "Commands/Stagger/AvaSequencerStagger.h"
 #include "CoreGlobals.h"
-#include "DetailsView/SAvaSequenceDetails.h"
+#include "DetailsView/SAvaMarkDetails.h"
 #include "DetailsView/Section/AvaSequencePlaybackDetails.h"
-#include "DetailsView/Section/AvaSequenceSelectionDetails.h"
+#include "DetailsView/Section/AvaSequencerEaseCurveToolSection.h"
 #include "DetailsView/Section/AvaSequenceSettingsDetails.h"
+#include "DetailsView/Section/AvaSequenceTreeDetails.h"
 #include "EaseCurveTool/AvaEaseCurveTool.h"
 #include "EaseCurveTool/AvaEaseCurveToolCommands.h"
 #include "Editor/Sequencer/Private/Sequencer.h"
-#include "EngineUtils.h"
-#include "EntitySystem/MovieSceneEntitySystemRunner.h"
 #include "Framework/Commands/GenericCommands.h"
+#include "Framework/Commands/UIAction.h"
 #include "Framework/MultiBox/MultiBoxBuilder.h"
 #include "IAvaSequenceProvider.h"
 #include "ISequencer.h"
@@ -32,6 +31,7 @@
 #include "MVVM/Selection/Selection.h"
 #include "MVVM/ViewModels/SequencerEditorViewModel.h"
 #include "MVVM/Views/SOutlinerView.h"
+#include "Misc/NotifyHook.h"
 #include "Misc/TextFilter.h"
 #include "MovieScene.h"
 #include "Playback/AvaSequencerCleanView.h"
@@ -39,21 +39,22 @@
 #include "PropertyHandle.h"
 #include "ScopedTransaction.h"
 #include "Selection/AvaEditorSelection.h"
+#include "SequencerCommands.h"
+#include "SequencerSettings.h"
+#include "SequencerUtilities.h"
 #include "SequenceTree/AvaSequenceItem.h"
 #include "SequenceTree/Columns/AvaSequenceNameColumn.h"
 #include "SequenceTree/Columns/AvaSequenceStatusColumn.h"
 #include "SequenceTree/Columns/IAvaSequenceColumn.h"
 #include "SequenceTree/Widgets/SAvaSequenceTree.h"
-#include "SequencerCommands.h"
-#include "SequencerSettings.h"
-#include "SequencerUtilities.h"
 #include "Settings/AvaSequencerSettings.h"
+#include "Sidebar/SidebarDrawerConfig.h"
+#include "Sidebar/SSidebar.h"
 #include "ToolMenu.h"
 #include "ToolMenuEntry.h"
 #include "ToolMenuSection.h"
 #include "ToolMenus.h"
 #include "Toolkits/AssetEditorToolkit.h"
-#include "Tracks/MovieScene3DTransformTrack.h"
 #include "Widgets/Views/STreeView.h"
 
 DEFINE_LOG_CATEGORY_STATIC(LogAvaSequencer, Log, All);
@@ -113,6 +114,8 @@ namespace UE::AvaSequencer::Private
 	};
 }
 
+const FName FAvaSequencer::SidebarDrawerId = TEXT("Sequence");
+
 FAvaSequencer::FAvaSequencer(IAvaSequencerProvider& InProvider, FAvaSequencerArgs&& InArgs)
 	: Provider(InProvider)
 	, CommandList(MakeShared<FUICommandList>())
@@ -140,8 +143,24 @@ FAvaSequencer::FAvaSequencer(IAvaSequencerProvider& InProvider, FAvaSequencerArg
 	{
 		const int32 NewIndex = SequencerModule.GetAddTrackMenuExtensibilityManager()->GetExtenderDelegates().Add(
 			FAssetEditorExtender::CreateRaw(this, &FAvaSequencer::GetAddTrackSequencerExtender));
-		
+
 		SequencerAddTrackExtenderHandle = SequencerModule.GetAddTrackMenuExtensibilityManager()->GetExtenderDelegates()[NewIndex].GetHandle();
+
+		SidebarExtender = MakeShared<FExtender>();
+
+		SidebarExtender->AddMenuExtension(
+			TEXT("KeyEdit"),
+			EExtensionHook::First,
+			CommandList,
+			FMenuExtensionDelegate::CreateRaw(this, &FAvaSequencer::ExtendSidebarSelectionMenu));
+
+		SidebarExtender->AddMenuExtension(
+			TEXT("MarkedFrames"),
+			EExtensionHook::After,
+			CommandList,
+			FMenuExtensionDelegate::CreateRaw(this, &FAvaSequencer::ExtendSidebarMarkedFramesMenu));
+
+		SequencerModule.GetSidebarExtensibilityManager()->AddExtender(SidebarExtender);
 	}
 
 	// Register to update when an undo/redo operation has been called to update our list of items
@@ -155,7 +174,7 @@ FAvaSequencer::~FAvaSequencer()
 	OnSequenceStartedHandle.Reset();
 	OnSequenceFinishedHandle.Reset();
 
-	if(GEngine)
+	if (GEngine)
 	{
 		GEditor->UnregisterForUndo(this);
 	}
@@ -169,6 +188,26 @@ FAvaSequencer::~FAvaSequencer()
 				return SequencerAddTrackExtenderHandle == Extender.GetHandle();
 			});
 	}
+
+	if (const TSharedPtr<ISequencer> Sequencer = SequencerWeak.Pin())
+	{
+		Sequencer->UnregisterDrawerSection(SidebarDrawerId, FAvaSequencePlaybackDetails::UniqueId);
+		Sequencer->UnregisterDrawerSection(SidebarDrawerId, FAvaSequenceSettingsDetails::UniqueId);
+		Sequencer->UnregisterDrawerSection(SidebarDrawerId, FAvaSequencerEaseCurveToolSection::UniqueId);
+
+		Sequencer->UnregisterDrawer(SidebarDrawerId);
+
+		Sequencer->UnregisterDrawerSection(FSequencer::SelectionDrawerId, FAvaSequencerEaseCurveToolSection::UniqueId);
+
+		if (SidebarExtender.IsValid())
+		{
+			ISequencerModule& SequencerModule = FAvaSequencerUtils::GetSequencerModule();
+
+			SequencerModule.GetSidebarExtensibilityManager()->RemoveExtender(SidebarExtender);
+
+			SidebarExtender.Reset();
+		}
+	}
 }
 
 void FAvaSequencer::BindCommands()
@@ -1449,39 +1488,66 @@ TSharedRef<SWidget> FAvaSequencer::CreateSequenceWidget()
 
 	TSharedRef<ISequencer> Sequencer = GetSequencer();
 
-	TArray<TSharedRef<IAvaSequenceSectionDetails>> Sections =
-		{
-			MakeShared<FAvaSequenceSelectionDetails>(),
-			MakeShared<FAvaSequenceSettingsDetails>(),
-			MakeShared<FAvaSequencePlaybackDetails>(),
-		};
+	const TSharedRef<SOverlay> SidebarOverlay = SNew(SOverlay);
+	SequenceTreeDockLocation = SNew(SBox);
 
-	return SNew(SSplitter)
-		+ SSplitter::Slot()
-		.Value(0.15f)
-		[
-			GetSequenceTreeWidget()
-		]
-		+ SSplitter::Slot()
-		.Value(0.65f)
+	TSharedRef<SOverlay> OutWidget = SNew(SOverlay)
+		+ SOverlay::Slot()
 		[
-			SNew(SOverlay)
-			.AddMetaData<FTagMetaData>(FTagMetaData(TEXT("Sequencer")))
-			+ SOverlay::Slot()
+			SNew(SHorizontalBox)
+			+ SHorizontalBox::Slot()
+			.AutoWidth()
 			[
-				Sequencer->GetSequencerWidget()
+				SAssignNew(LeftSidebar, SSidebar, SidebarOverlay, SequenceTreeDockLocation.ToSharedRef())
+				.TabLocation(ESidebarTabLocation::Left)
+				.OnDockStateChanged(this, &FAvaSequencer::OnSidebarDockStateChanged)
+			]
+			+ SHorizontalBox::Slot()
+			[
+				SNew(SSplitter)
+				+ SSplitter::Slot()
+				.Resizable(false)
+				.SizeRule(SSplitter::ESizeRule::SizeToContent)
+				.Value(FSidebarDrawerConfig::DefaultSizeCoefficient)
+				.OnSlotResized(this, &FAvaSequencer::OnSidebarSlotResized)
+				.Expose(SidebarSlot)
+				[
+					SequenceTreeDockLocation.ToSharedRef()
+				]
+				+ SSplitter::Slot()
+				.Value(0.8f)
+				[
+					SNew(SOverlay)
+					.AddMetaData<FTagMetaData>(FTagMetaData(TEXT("Sequencer")))
+					+ SOverlay::Slot()
+					[
+						Sequencer->GetSequencerWidget()
+					]
+				]
 			]
 		]
-		+ SSplitter::Slot()
-		.Value(0.2f)
+		+ SOverlay::Slot()
 		[
-			SNew(SAvaSequenceDetails, SharedThis(this), MoveTemp(Sections))
-			.InitiallySelectedSections(SelectedSections)
-			.OnSelectedSectionsChanged_Lambda([this](const TSet<FName>& InNewSelectedSections)
-				{
-					SelectedSections = InNewSelectedSections;
-				})
+			SidebarOverlay
 		];
+
+	FSidebarDrawerConfig SequenceTreeDrawerConfig;
+	SequenceTreeDrawerConfig.UniqueId = SidebarDrawerId;
+	SequenceTreeDrawerConfig.ButtonText = LOCTEXT("SequenceLabel", "Sequence");
+	SequenceTreeDrawerConfig.ToolTipText = LOCTEXT("SequenceTooltip", "Open the Sequence options panel");
+	SequenceTreeDrawerConfig.Icon = FAppStyle::GetBrush(TEXT("Persona.EditInSequencer"));
+	SequenceTreeDrawerConfig.bInitiallyDocked = true;
+	LeftSidebar->RegisterDrawer(MoveTemp(SequenceTreeDrawerConfig));
+	
+	// Make sure the sequence tree widget is created
+	GetSequenceTreeWidget();
+
+	const TSharedRef<FAvaSequencer> ThisSequencerRef = SharedThis(this);
+	LeftSidebar->RegisterDrawerSection(SidebarDrawerId, MakeShared<FAvaSequenceTreeDetails>(ThisSequencerRef));
+	LeftSidebar->RegisterDrawerSection(SidebarDrawerId, MakeShared<FAvaSequencePlaybackDetails>(ThisSequencerRef));
+	LeftSidebar->RegisterDrawerSection(SidebarDrawerId, MakeShared<FAvaSequenceSettingsDetails>(ThisSequencerRef));
+
+	return OutWidget;
 }
 
 void FAvaSequencer::OnActorsCopied(FString& InOutCopiedData, TConstArrayView<AActor*> InCopiedActors)
@@ -1738,4 +1804,55 @@ TSharedRef<FAvaEaseCurveTool> FAvaSequencer::GetEaseCurveTool() const
 	return EaseCurveTool.ToSharedRef(); 
 }
 
+void FAvaSequencer::OnSidebarDockStateChanged(const FName InDrawerId)
+{
+	LeftSidebar->UpdateDockedSplitterSlot(SidebarDrawerId, SidebarSlot);
+}
+
+void FAvaSequencer::OnSidebarSlotResized(const float InFillCoefficient)
+{
+	if (InFillCoefficient == 0.f && LeftSidebar->HasDrawerDocked())
+	{
+		LeftSidebar->UndockAllDrawers();
+	}
+}
+
+void FAvaSequencer::ExtendSidebarSelectionMenu(FMenuBuilder& OutMenuBuilder)
+{
+	const TSharedPtr<FAvaSequencerEaseCurveToolSection> Section = MakeShared<FAvaSequencerEaseCurveToolSection>(SharedThis(this));
+	OutMenuBuilder.AddWidget(Section->CreateContentWidget(), FText::GetEmpty(), /*bInNoIndent=*/true);
+}
+
+void FAvaSequencer::ExtendSidebarMarkedFramesMenu(FMenuBuilder& OutMenuBuilder)
+{
+	UAvaSequence* const Sequence = GetViewedSequence();
+	if (!IsValid(Sequence))
+	{
+		return;
+	}
+
+	UMovieScene* const MovieScene = Sequence->GetMovieScene();
+	if (!IsValid(MovieScene))
+	{
+		return;
+	}
+
+	const TSharedPtr<UE::Sequencer::FSequencerSelection> SequencerSelection = GetSequencer()->GetViewModel()->GetSelection();
+	if (!SequencerSelection.IsValid())
+	{
+		return;
+	}
+
+	const TArray<FMovieSceneMarkedFrame>& MarkedFrames = MovieScene->GetMarkedFrames();
+
+	for (const int32 MarkIndex : SequencerSelection->MarkedFrames)
+	{
+		if (MarkedFrames.IsValidIndex(MarkIndex))
+		{
+			const TSharedRef<SWidget> DetailsWidget = SNew(SAvaMarkDetails, Sequence, MarkedFrames[MarkIndex]);
+			OutMenuBuilder.AddWidget(DetailsWidget, FText::GetEmpty(), /*bInNoIndent=*/true);
+		}
+	}
+}
+
 #undef LOCTEXT_NAMESPACE
diff --git a/Engine/Plugins/Experimental/Avalanche/Source/AvalancheSequencer/Private/AvaSequencer.h b/Engine/Plugins/Experimental/Avalanche/Source/AvalancheSequencer/Private/AvaSequencer.h
--- a/Engine/Plugins/Experimental/Avalanche/Source/AvalancheSequencer/Private/AvaSequencer.h
+++ b/Engine/Plugins/Experimental/Avalanche/Source/AvalancheSequencer/Private/AvaSequencer.h
@@ -8,6 +8,7 @@
 #include "IAvaSequencer.h"
 #include "IAvaSequencerProvider.h"
 #include "MovieSceneSequenceID.h"
+#include "Widgets/Layout/SSplitter.h"
 
 class AActor;
 class FAvaEaseCurveTool;
@@ -25,7 +26,9 @@ class IPropertyHandle;
 class ISequencer;
 class ISequencerObjectChangeListener;
 class SAvaSequenceTree;
+class SBox;
 class SHeaderRow;
+class SSidebar;
 class SWidget;
 class UAvaSequence;
 class UAvaSequencerSettings;
@@ -51,6 +54,8 @@ namespace UE::Sequencer
 class FAvaSequencer : public IAvaSequencer, public FEditorUndoClient, public TSharedFromThis<FAvaSequencer>
 {
 public:
+	static const FName SidebarDrawerId;
+
 	explicit FAvaSequencer(IAvaSequencerProvider& InProvider, FAvaSequencerArgs&& InArgs);
 
 	virtual ~FAvaSequencer() override;
@@ -216,6 +221,12 @@ private:
 
 	void OnUpdateCameraCut(UObject* InCameraObject, bool bInJumpCut);
 
+	void OnSidebarDockStateChanged(const FName InDrawerId);
+	void OnSidebarSlotResized(const float InFillCoefficient);
+
+	void ExtendSidebarSelectionMenu(FMenuBuilder& OutMenuBuilder);
+	void ExtendSidebarMarkedFramesMenu(FMenuBuilder& OutMenuBuilder);
+
 	FOnViewedSequenceChanged OnViewedSequenceChanged;
 
 	IAvaSequencerProvider& Provider;
@@ -265,8 +276,12 @@ private:
 	/** Whether FAvaSequencer is allowed to select to/from the ISequencer instance */
 	const bool bCanProcessSequencerSelections;
 
-	/** Selected sequence details sections to restore when a new sequence is selected. */
-	TSet<FName> SelectedSections;
-
 	TSharedPtr<FAvaEaseCurveTool> EaseCurveTool;
+
+	TSharedPtr<SSidebar> LeftSidebar;
+	TSharedPtr<SBox> SequenceTreeDockLocation;
+	SSplitter::FSlot* SidebarSlot = nullptr;
+
+	TSharedPtr<FExtender> SidebarExtender;
+	FDelegateHandle SidebarSelectionExtenderHandle;
 };
diff --git a/Engine/Plugins/Experimental/Avalanche/Source/AvalancheSequencer/Private/DetailsView/SAvaMarkDetails.cpp b/Engine/Plugins/Experimental/Avalanche/Source/AvalancheSequencer/Private/DetailsView/SAvaMarkDetails.cpp
new file mode 100644
--- /dev/null
+++ b/Engine/Plugins/Experimental/Avalanche/Source/AvalancheSequencer/Private/DetailsView/SAvaMarkDetails.cpp
@@ -0,0 +1,48 @@
+﻿// Copyright Epic Games, Inc. All Rights Reserved.
+
+#include "SAvaMarkDetails.h"
+#include "AvaSequence.h"
+#include "CustomDetailsViewArgs.h"
+#include "CustomDetailsViewModule.h"
+#include "ICustomDetailsView.h"
+#include "Items/CustomDetailsViewItemId.h"
+#include "Marks/AvaMark.h"
+#include "MovieSceneMarkedFrame.h"
+#include "UObject/StructOnScope.h"
+#include "UObject/UnrealType.h"
+
+void SAvaMarkDetails::Construct(const FArguments& InArgs, UAvaSequence* InSequence, const FMovieSceneMarkedFrame& InMarkedFrame)
+{
+	SequenceToModify = InSequence;
+	check(IsValid(SequenceToModify));
+
+	FCustomDetailsViewArgs CustomDetailsViewArgs;
+	CustomDetailsViewArgs.IndentAmount = 0.f;
+	CustomDetailsViewArgs.bShowCategories = true;
+	CustomDetailsViewArgs.bAllowGlobalExtensions = true;
+	CustomDetailsViewArgs.CategoryAllowList.Allow(TEXT("Marks"));
+	CustomDetailsViewArgs.ItemAllowList.Disallow(FCustomDetailsViewItemId::MakePropertyId(TEXT("Label")));
+	CustomDetailsViewArgs.ExpansionState.Add(FCustomDetailsViewItemId::MakePropertyId<UAvaSequence>(TEXT("MotionDesign")), true);
+	CustomDetailsViewArgs.ExpansionState.Add(FCustomDetailsViewItemId::MakePropertyId<UAvaSequence>(TEXT("Marks")), true);
+
+	TSharedRef<ICustomDetailsView> SettingsDetailsView = ICustomDetailsViewModule::Get().CreateCustomDetailsView(CustomDetailsViewArgs);
+
+	FAvaMark& AvaMark = SequenceToModify->FindOrAddMark(InMarkedFrame.Label);
+	AvaMarkStruct = MakeShared<FStructOnScope>(FAvaMark::StaticStruct(), (uint8*)&AvaMark);
+	SettingsDetailsView->SetStruct(AvaMarkStruct);
+
+	ChildSlot
+		[
+			SettingsDetailsView
+		];
+}
+
+void SAvaMarkDetails::NotifyPreChange(FProperty* InPropertyAboutToChange)
+{
+	SequenceToModify->Modify();
+}
+
+void SAvaMarkDetails::NotifyPreChange(FEditPropertyChain* InPropertyAboutToChange)
+{
+	SequenceToModify->Modify();
+}
diff --git a/Engine/Plugins/Experimental/Avalanche/Source/AvalancheSequencer/Private/DetailsView/SAvaMarkDetails.h b/Engine/Plugins/Experimental/Avalanche/Source/AvalancheSequencer/Private/DetailsView/SAvaMarkDetails.h
new file mode 100644
--- /dev/null
+++ b/Engine/Plugins/Experimental/Avalanche/Source/AvalancheSequencer/Private/DetailsView/SAvaMarkDetails.h
@@ -0,0 +1,33 @@
+﻿// Copyright Epic Games, Inc. All Rights Reserved.
+
+#pragma once
+
+#include "Misc/NotifyHook.h"
+#include "Templates/SharedPointerFwd.h"
+#include "UObject/ObjectPtr.h"
+#include "Widgets/DeclarativeSyntaxSupport.h"
+#include "Widgets/SCompoundWidget.h"
+
+class FEditPropertyChain;
+class FProperty;
+class FStructOnScope;
+class UAvaSequence;
+struct FMovieSceneMarkedFrame;
+
+class SAvaMarkDetails : public SCompoundWidget, public FNotifyHook
+{
+public:
+	SLATE_BEGIN_ARGS(SAvaMarkDetails) {}
+	SLATE_END_ARGS()
+
+	void Construct(const FArguments& InArgs, UAvaSequence* InSequence, const FMovieSceneMarkedFrame& InMarkedFrame);
+
+protected:
+	//~ Begin FNotifyHook
+	virtual void NotifyPreChange(FProperty* InPropertyAboutToChange) override;
+	virtual void NotifyPreChange(FEditPropertyChain* InPropertyAboutToChange) override;
+	//~ End FNotifyHook
+
+	TObjectPtr<UAvaSequence> SequenceToModify;
+	TSharedPtr<FStructOnScope> AvaMarkStruct;
+};
diff --git a/Engine/Plugins/Experimental/Avalanche/Source/AvalancheSequencer/Private/DetailsView/SAvaSequenceDetails.cpp b/Engine/Plugins/Experimental/Avalanche/Source/AvalancheSequencer/Private/DetailsView/SAvaSequenceDetails.cpp
deleted file mode 100644
diff --git a/Engine/Plugins/Experimental/Avalanche/Source/AvalancheSequencer/Private/DetailsView/SAvaSequenceDetails.h b/Engine/Plugins/Experimental/Avalanche/Source/AvalancheSequencer/Private/DetailsView/SAvaSequenceDetails.h
deleted file mode 100644
diff --git a/Engine/Plugins/Experimental/Avalanche/Source/AvalancheSequencer/Private/DetailsView/Section/AvaSequencePlaybackDetails.cpp b/Engine/Plugins/Experimental/Avalanche/Source/AvalancheSequencer/Private/DetailsView/Section/AvaSequencePlaybackDetails.cpp
--- a/Engine/Plugins/Experimental/Avalanche/Source/AvalancheSequencer/Private/DetailsView/Section/AvaSequencePlaybackDetails.cpp
+++ b/Engine/Plugins/Experimental/Avalanche/Source/AvalancheSequencer/Private/DetailsView/Section/AvaSequencePlaybackDetails.cpp
@@ -2,50 +2,66 @@
 
 #include "AvaSequencePlaybackDetails.h"
 #include "AvaSequencePlaybackObject.h"
-#include "AvaSequence.h"
 #include "AvaSequencer.h"
 #include "CustomDetailsViewArgs.h"
 #include "CustomDetailsViewModule.h"
 
 #define LOCTEXT_NAMESPACE "AvaSequencePlaybackDetails"
 
-FName FAvaSequencePlaybackDetails::GetSectionName() const
+const FName FAvaSequencePlaybackDetails::UniqueId = TEXT("AvaSequencePlaybackDetails");
+
+FAvaSequencePlaybackDetails::FAvaSequencePlaybackDetails(const TSharedRef<FAvaSequencer>& InAvaSequencer)
+	: AvaSequencerWeak(InAvaSequencer)
+{
+}
+
+FName FAvaSequencePlaybackDetails::GetUniqueId() const
+{
+	return UniqueId;
+}
+
+FName FAvaSequencePlaybackDetails::GetSectionId() const
 {
 	return TEXT("Playback");
 }
 
-FText FAvaSequencePlaybackDetails::GetSectionDisplayName() const
+FText FAvaSequencePlaybackDetails::GetSectionDisplayText() const
 {
 	return LOCTEXT("PlaybackLabel", "Playback");
 }
 
-TSharedRef<SWidget> FAvaSequencePlaybackDetails::CreateContentWidget(const TSharedRef<FAvaSequencer>& InAvaSequencer)
+bool FAvaSequencePlaybackDetails::ShouldShowSection() const
 {
-	AvaSequencerWeak = InAvaSequencer;
+	return AvaSequencerWeak.IsValid();
+}
 
+int32 FAvaSequencePlaybackDetails::GetSortOrder() const
+{
+	return 2;
+}
+
+TSharedRef<SWidget> FAvaSequencePlaybackDetails::CreateContentWidget()
+{
 	FCustomDetailsViewArgs CustomDetailsViewArgs;
-	CustomDetailsViewArgs.IndentAmount           = 0.f;
-	CustomDetailsViewArgs.ValueColumnWidth       = 0.5f;
-	CustomDetailsViewArgs.bShowCategories        = true;
+	CustomDetailsViewArgs.IndentAmount = 0.f;
+	CustomDetailsViewArgs.ValueColumnWidth  = 0.5f;
+	CustomDetailsViewArgs.bShowCategories = true;
 	CustomDetailsViewArgs.bAllowGlobalExtensions = true;
 	CustomDetailsViewArgs.CategoryAllowList.Allow(TEXT("Scheduled Playback"));
 	CustomDetailsViewArgs.ExpansionState.Add(FCustomDetailsViewItemId::MakeCategoryId("Scheduled Playback"), true);
 
-	TSharedRef<ICustomDetailsView> PlaybackDetailsView = ICustomDetailsViewModule::Get().CreateCustomDetailsView(CustomDetailsViewArgs);
+	const TSharedRef<ICustomDetailsView> PlaybackDetailsView = ICustomDetailsViewModule::Get().CreateCustomDetailsView(CustomDetailsViewArgs);
 
-	IAvaSequencePlaybackObject* PlaybackObject = InAvaSequencer->GetProvider().GetPlaybackObject();
-	if (ensureAlways(PlaybackObject))
+	if (const TSharedPtr<FAvaSequencer> AvaSequencer = AvaSequencerWeak.Pin())
 	{
-		PlaybackDetailsView->SetObject(PlaybackObject->ToUObject());
+		IAvaSequencePlaybackObject* const PlaybackObject = AvaSequencer->GetProvider().GetPlaybackObject();
+		if (ensureAlways(PlaybackObject))
+		{
+			PlaybackDetailsView->SetObject(PlaybackObject->ToUObject());
+		}
 	}
 
 	return PlaybackDetailsView;
 }
 
-bool FAvaSequencePlaybackDetails::ShouldShowSection() const
-{
-	TSharedPtr<FAvaSequencer> AvaSequencer = AvaSequencerWeak.Pin();
-	return AvaSequencer.IsValid() && IsValid(AvaSequencer->GetViewedSequence());
-}
-
 #undef LOCTEXT_NAMESPACE
diff --git a/Engine/Plugins/Experimental/Avalanche/Source/AvalancheSequencer/Private/DetailsView/Section/AvaSequencePlaybackDetails.h b/Engine/Plugins/Experimental/Avalanche/Source/AvalancheSequencer/Private/DetailsView/Section/AvaSequencePlaybackDetails.h
--- a/Engine/Plugins/Experimental/Avalanche/Source/AvalancheSequencer/Private/DetailsView/Section/AvaSequencePlaybackDetails.h
+++ b/Engine/Plugins/Experimental/Avalanche/Source/AvalancheSequencer/Private/DetailsView/Section/AvaSequencePlaybackDetails.h
@@ -2,16 +2,30 @@
 
 #pragma once
 
-#include "IAvaSequenceSectionDetails.h"
+#include "Sidebar/ISidebarDrawerContent.h"
+#include "Templates/SharedPointer.h"
 
-class FAvaSequencePlaybackDetails : public IAvaSequenceSectionDetails
+class FAvaSequencer;
+class FName;
+class FText;
+class SWidget;
+
+class FAvaSequencePlaybackDetails : public ISidebarDrawerContent
 {
-	//~ Begin IAvaSequenceSectionDetails
-	virtual FName GetSectionName() const override;
-	virtual FText GetSectionDisplayName() const override;
-	virtual TSharedRef<SWidget> CreateContentWidget(const TSharedRef<FAvaSequencer>& InAvaSequencer) override;
-	virtual bool ShouldShowSection() const;
-	//~ End IAvaSequenceSectionDetails
+public:
+	static const FName UniqueId;
+
+	FAvaSequencePlaybackDetails(const TSharedRef<FAvaSequencer>& InAvaSequencer);
+
+	//~ Begin ISidebarDrawerContent
+	virtual FName GetUniqueId() const override;
+	virtual FName GetSectionId() const override;
+	virtual FText GetSectionDisplayText() const override;
+	virtual bool ShouldShowSection() const override;
+	virtual int32 GetSortOrder() const override;
+	virtual TSharedRef<SWidget> CreateContentWidget() override;
+	//~ End ISidebarDrawerContent
 
+protected:
 	TWeakPtr<FAvaSequencer> AvaSequencerWeak;
 };
diff --git a/Engine/Plugins/Experimental/Avalanche/Source/AvalancheSequencer/Private/DetailsView/Section/AvaSequenceSelectionDetails.cpp b/Engine/Plugins/Experimental/Avalanche/Source/AvalancheSequencer/Private/DetailsView/Section/AvaSequenceSelectionDetails.cpp
deleted file mode 100644
diff --git a/Engine/Plugins/Experimental/Avalanche/Source/AvalancheSequencer/Private/DetailsView/Section/AvaSequenceSelectionDetails.h b/Engine/Plugins/Experimental/Avalanche/Source/AvalancheSequencer/Private/DetailsView/Section/AvaSequenceSelectionDetails.h
deleted file mode 100644
diff --git a/Engine/Plugins/Experimental/Avalanche/Source/AvalancheSequencer/Private/DetailsView/Section/AvaSequenceSettingsDetails.cpp b/Engine/Plugins/Experimental/Avalanche/Source/AvalancheSequencer/Private/DetailsView/Section/AvaSequenceSettingsDetails.cpp
--- a/Engine/Plugins/Experimental/Avalanche/Source/AvalancheSequencer/Private/DetailsView/Section/AvaSequenceSettingsDetails.cpp
+++ b/Engine/Plugins/Experimental/Avalanche/Source/AvalancheSequencer/Private/DetailsView/Section/AvaSequenceSettingsDetails.cpp
@@ -8,47 +8,65 @@
 
 #define LOCTEXT_NAMESPACE "AvaSequenceSettingsDetails"
 
-FName FAvaSequenceSettingsDetails::GetSectionName() const
+const FName FAvaSequenceSettingsDetails::UniqueId = TEXT("AvaSequenceSettingsDetails");
+
+FAvaSequenceSettingsDetails::FAvaSequenceSettingsDetails(const TSharedRef<FAvaSequencer>& InAvaSequencer)
+	: AvaSequencerWeak(InAvaSequencer)
+{
+}
+
+FName FAvaSequenceSettingsDetails::GetUniqueId() const
+{
+	return UniqueId;
+}
+
+FName FAvaSequenceSettingsDetails::GetSectionId() const
 {
 	return TEXT("Settings");
 }
 
-FText FAvaSequenceSettingsDetails::GetSectionDisplayName() const
+FText FAvaSequenceSettingsDetails::GetSectionDisplayText() const
 {
-	return LOCTEXT("SettingsLabel", "Sequence");
+	return LOCTEXT("SettingsLabel", "Settings");
 }
 
-TSharedRef<SWidget> FAvaSequenceSettingsDetails::CreateContentWidget(const TSharedRef<FAvaSequencer>& InAvaSequencer)
+bool FAvaSequenceSettingsDetails::ShouldShowSection() const
 {
-	AvaSequencerWeak = InAvaSequencer;
+	return AvaSequencerWeak.IsValid();
+}
 
+int32 FAvaSequenceSettingsDetails::GetSortOrder() const
+{
+	return 3;
+}
+
+TSharedRef<SWidget> FAvaSequenceSettingsDetails::CreateContentWidget()
+{
 	FCustomDetailsViewArgs CustomDetailsViewArgs;
 	CustomDetailsViewArgs.IndentAmount = 0.f;
 	CustomDetailsViewArgs.bShowCategories = true;
 	CustomDetailsViewArgs.bAllowGlobalExtensions = true;
 	CustomDetailsViewArgs.CategoryAllowList.Allow(TEXT("Sequence Settings"));
-	CustomDetailsViewArgs.ExpansionState.Add(FCustomDetailsViewItemId::MakeCategoryId("Sequence Settings"), true);
+	CustomDetailsViewArgs.ExpansionState.Add(FCustomDetailsViewItemId::MakeCategoryId(TEXT("Sequence Settings")), true);
 	CustomDetailsViewArgs.ExpansionState.Add(FCustomDetailsViewItemId::MakePropertyId<UAvaSequence>(TEXT("Marks")), true);
 
 	SettingsDetailsView = ICustomDetailsViewModule::Get().CreateCustomDetailsView(CustomDetailsViewArgs);
 
-	InAvaSequencer->GetOnViewedSequenceChanged().AddSP(this, &FAvaSequenceSettingsDetails::OnViewedSequenceChanged);
-
-	if (UAvaSequence* const ViewedSequence = InAvaSequencer->GetViewedSequence())
+	if (const TSharedPtr<FAvaSequencer> AvaSequencer = AvaSequencerWeak.Pin())
 	{
-		OnViewedSequenceChanged(ViewedSequence);	
+		AvaSequencer->GetOnViewedSequenceChanged().AddSP(this, &FAvaSequenceSettingsDetails::OnViewedSequenceChanged);
+
+		UAvaSequence* const ViewedSequence = AvaSequencer->GetViewedSequence();
+		if (IsValid(ViewedSequence))
+		{
+			OnViewedSequenceChanged(ViewedSequence);
+		}
 	}
 
 	return SettingsDetailsView.ToSharedRef();
 }
 
-bool FAvaSequenceSettingsDetails::ShouldShowSection() const
-{
-	TSharedPtr<FAvaSequencer> AvaSequencer = AvaSequencerWeak.Pin();
-	return AvaSequencer.IsValid() && IsValid(AvaSequencer->GetViewedSequence());
-}
-
-void FAvaSequenceSettingsDetails::OnViewedSequenceChanged(UAvaSequence* InSequence)
+void FAvaSequenceSettingsDetails::OnViewedSequenceChanged(UAvaSequence* const InSequence)
 {
 	SettingsDetailsView->SetObject(InSequence);
 }
diff --git a/Engine/Plugins/Experimental/Avalanche/Source/AvalancheSequencer/Private/DetailsView/Section/AvaSequenceSettingsDetails.h b/Engine/Plugins/Experimental/Avalanche/Source/AvalancheSequencer/Private/DetailsView/Section/AvaSequenceSettingsDetails.h
--- a/Engine/Plugins/Experimental/Avalanche/Source/AvalancheSequencer/Private/DetailsView/Section/AvaSequenceSettingsDetails.h
+++ b/Engine/Plugins/Experimental/Avalanche/Source/AvalancheSequencer/Private/DetailsView/Section/AvaSequenceSettingsDetails.h
@@ -2,21 +2,31 @@
 
 #pragma once
 
-#include "IAvaSequenceSectionDetails.h"
+#include "Sidebar/ISidebarDrawerContent.h"
+#include "Templates/SharedPointer.h"
 
-class UAvaSequence;
+class FAvaSequencer;
 class ICustomDetailsView;
+class UAvaSequence;
 
-class FAvaSequenceSettingsDetails : public IAvaSequenceSectionDetails
+class FAvaSequenceSettingsDetails : public ISidebarDrawerContent
 {
-	//~ Begin IAvaSequenceSectionDetails
-	virtual FName GetSectionName() const override;
-	virtual FText GetSectionDisplayName() const override;
-	virtual TSharedRef<SWidget> CreateContentWidget(const TSharedRef<FAvaSequencer>& InAvaSequencer) override;
+public:
+	static const FName UniqueId;
+
+	FAvaSequenceSettingsDetails(const TSharedRef<FAvaSequencer>& InAvaSequencer);
+
+	//~ Begin ISidebarDrawerContent
+	virtual FName GetUniqueId() const override;
+	virtual FName GetSectionId() const override;
+	virtual FText GetSectionDisplayText() const override;
 	virtual bool ShouldShowSection() const override;
-	//~ End IAvaSequenceSectionDetails
+	virtual int32 GetSortOrder() const override;
+	virtual TSharedRef<SWidget> CreateContentWidget() override;
+	//~ End ISidebarDrawerContent
 
-	void OnViewedSequenceChanged(UAvaSequence* InSequence);
+protected:
+	void OnViewedSequenceChanged(UAvaSequence* const InSequence);
 
 	TWeakPtr<FAvaSequencer> AvaSequencerWeak;
 
diff --git a/Engine/Plugins/Experimental/Avalanche/Source/AvalancheSequencer/Private/DetailsView/Section/AvaSequenceTreeDetails.cpp b/Engine/Plugins/Experimental/Avalanche/Source/AvalancheSequencer/Private/DetailsView/Section/AvaSequenceTreeDetails.cpp
new file mode 100644
--- /dev/null
+++ b/Engine/Plugins/Experimental/Avalanche/Source/AvalancheSequencer/Private/DetailsView/Section/AvaSequenceTreeDetails.cpp
@@ -0,0 +1,49 @@
+// Copyright Epic Games, Inc. All Rights Reserved.
+
+#include "AvaSequenceTreeDetails.h"
+#include "AvaSequencer.h"
+
+#define LOCTEXT_NAMESPACE "AvaSequenceTreeDetails"
+
+const FName FAvaSequenceTreeDetails::UniqueId = TEXT("AvaSequenceTreeDetails");
+
+FAvaSequenceTreeDetails::FAvaSequenceTreeDetails(const TSharedRef<FAvaSequencer>& InAvaSequencer)
+	: AvaSequencerWeak(InAvaSequencer)
+{
+}
+
+FName FAvaSequenceTreeDetails::GetUniqueId() const
+{
+	return UniqueId;
+}
+
+FName FAvaSequenceTreeDetails::GetSectionId() const
+{
+	return TEXT("SequenceTree");
+}
+
+FText FAvaSequenceTreeDetails::GetSectionDisplayText() const
+{
+	return LOCTEXT("SequenceTreeLabel", "Tree");
+}
+
+bool FAvaSequenceTreeDetails::ShouldShowSection() const
+{
+	return AvaSequencerWeak.IsValid();
+}
+
+int32 FAvaSequenceTreeDetails::GetSortOrder() const
+{
+	return 1;
+}
+
+TSharedRef<SWidget> FAvaSequenceTreeDetails::CreateContentWidget()
+{
+	if (const TSharedPtr<FAvaSequencer> AvaSequencer = AvaSequencerWeak.Pin())
+	{
+		return AvaSequencer->GetSequenceTreeWidget();
+	}
+	return SNullWidget::NullWidget;
+}
+
+#undef LOCTEXT_NAMESPACE
diff --git a/Engine/Plugins/Experimental/Avalanche/Source/AvalancheSequencer/Private/DetailsView/Section/AvaSequenceTreeDetails.h b/Engine/Plugins/Experimental/Avalanche/Source/AvalancheSequencer/Private/DetailsView/Section/AvaSequenceTreeDetails.h
new file mode 100644
--- /dev/null
+++ b/Engine/Plugins/Experimental/Avalanche/Source/AvalancheSequencer/Private/DetailsView/Section/AvaSequenceTreeDetails.h
@@ -0,0 +1,31 @@
+// Copyright Epic Games, Inc. All Rights Reserved.
+
+#pragma once
+
+#include "Sidebar/ISidebarDrawerContent.h"
+#include "Templates/SharedPointer.h"
+
+class FAvaSequencer;
+class FName;
+class FText;
+class SWidget;
+
+class FAvaSequenceTreeDetails : public ISidebarDrawerContent
+{
+public:
+	static const FName UniqueId;
+
+	FAvaSequenceTreeDetails(const TSharedRef<FAvaSequencer>& InAvaSequencer);
+
+	//~ Begin ISidebarDrawerContent
+	virtual FName GetUniqueId() const override;
+	virtual FName GetSectionId() const override;
+	virtual FText GetSectionDisplayText() const override;
+	virtual bool ShouldShowSection() const override;
+	virtual int32 GetSortOrder() const override;
+	virtual TSharedRef<SWidget> CreateContentWidget() override;
+	//~ End ISidebarDrawerContent
+
+protected:
+	TWeakPtr<FAvaSequencer> AvaSequencerWeak;
+};
diff --git a/Engine/Plugins/Experimental/Avalanche/Source/AvalancheSequencer/Private/DetailsView/Section/AvaSequencerEaseCurveToolSection.cpp b/Engine/Plugins/Experimental/Avalanche/Source/AvalancheSequencer/Private/DetailsView/Section/AvaSequencerEaseCurveToolSection.cpp
new file mode 100644
--- /dev/null
+++ b/Engine/Plugins/Experimental/Avalanche/Source/AvalancheSequencer/Private/DetailsView/Section/AvaSequencerEaseCurveToolSection.cpp
@@ -0,0 +1,53 @@
+// Copyright Epic Games, Inc. All Rights Reserved.
+
+#include "AvaSequencerEaseCurveToolSection.h"
+#include "AvaSequence.h"
+#include "AvaSequencer.h"
+#include "EaseCurveTool/AvaEaseCurveTool.h"
+#include "Widgets/SCompoundWidget.h"
+
+#define LOCTEXT_NAMESPACE "AvaSequencerEaseCurveToolSection"
+
+const FName FAvaSequencerEaseCurveToolSection::UniqueId = TEXT("AvaSequencerEaseCurveToolSection");
+
+FAvaSequencerEaseCurveToolSection::FAvaSequencerEaseCurveToolSection(const TSharedRef<FAvaSequencer>& InAvaSequencer)
+	: AvaSequencerWeak(InAvaSequencer)
+{
+}
+
+FName FAvaSequencerEaseCurveToolSection::GetUniqueId() const
+{
+	return UniqueId;
+}
+
+FName FAvaSequencerEaseCurveToolSection::GetSectionId() const
+{
+	return TEXT("Selection");
+}
+
+FText FAvaSequencerEaseCurveToolSection::GetSectionDisplayText() const
+{
+	return LOCTEXT("EaseCurveToolLabel", "Ease Curve Tool");
+}
+
+bool FAvaSequencerEaseCurveToolSection::ShouldShowSection() const
+{
+	const TSharedPtr<FAvaSequencer> AvaSequencer = AvaSequencerWeak.Pin();
+	return AvaSequencer.IsValid() && IsValid(AvaSequencer->GetViewedSequence());
+}
+
+int32 FAvaSequencerEaseCurveToolSection::GetSortOrder() const
+{
+	return -1;
+}
+
+TSharedRef<SWidget> FAvaSequencerEaseCurveToolSection::CreateContentWidget()
+{
+	if (const TSharedPtr<FAvaSequencer> AvaSequencer = AvaSequencerWeak.Pin())
+	{
+		return AvaSequencer->GetEaseCurveTool()->GenerateWidget();
+	}
+	return SNullWidget::NullWidget;
+}
+
+#undef LOCTEXT_NAMESPACE
diff --git a/Engine/Plugins/Experimental/Avalanche/Source/AvalancheSequencer/Private/DetailsView/Section/AvaSequencerEaseCurveToolSection.h b/Engine/Plugins/Experimental/Avalanche/Source/AvalancheSequencer/Private/DetailsView/Section/AvaSequencerEaseCurveToolSection.h
new file mode 100644
--- /dev/null
+++ b/Engine/Plugins/Experimental/Avalanche/Source/AvalancheSequencer/Private/DetailsView/Section/AvaSequencerEaseCurveToolSection.h
@@ -0,0 +1,32 @@
+// Copyright Epic Games, Inc. All Rights Reserved.
+
+#pragma once
+
+#include "Sidebar/ISidebarDrawerContent.h"
+#include "Templates/SharedPointer.h"
+
+class FAvaSequencer;
+class FName;
+class FText;
+class ISequencer;
+class SWidget;
+
+class FAvaSequencerEaseCurveToolSection : public ISidebarDrawerContent
+{
+public:
+	static const FName UniqueId;
+
+	FAvaSequencerEaseCurveToolSection(const TSharedRef<FAvaSequencer>& InAvaSequencer);
+
+	//~ Begin ISidebarDrawerContent
+	virtual FName GetUniqueId() const override;
+	virtual FName GetSectionId() const override;
+	virtual FText GetSectionDisplayText() const override;
+	virtual bool ShouldShowSection() const override;
+	virtual int32 GetSortOrder() const override;
+	virtual TSharedRef<SWidget> CreateContentWidget() override;
+	//~ End ISidebarDrawerContent
+
+protected:
+	TWeakPtr<FAvaSequencer> AvaSequencerWeak;
+};
diff --git a/Engine/Plugins/Experimental/Avalanche/Source/AvalancheSequencer/Private/DetailsView/Section/IAvaSequenceSectionDetails.h b/Engine/Plugins/Experimental/Avalanche/Source/AvalancheSequencer/Private/DetailsView/Section/IAvaSequenceSectionDetails.h
deleted file mode 100644
diff --git a/Engine/Plugins/Experimental/PropertyAnimator/Source/PropertyAnimatorEditor/Private/Sequencer/PropertyAnimatorEditorCurveChannelInterface.h b/Engine/Plugins/Experimental/PropertyAnimator/Source/PropertyAnimatorEditor/Private/Sequencer/PropertyAnimatorEditorCurveChannelInterface.h
--- a/Engine/Plugins/Experimental/PropertyAnimator/Source/PropertyAnimatorEditor/Private/Sequencer/PropertyAnimatorEditorCurveChannelInterface.h
+++ b/Engine/Plugins/Experimental/PropertyAnimator/Source/PropertyAnimatorEditor/Private/Sequencer/PropertyAnimatorEditorCurveChannelInterface.h
@@ -35,6 +35,7 @@ struct TPropertyAnimatorEditorCurveChannelInterface : ISequencerChannelInterface
 	virtual TSharedPtr<UE::Sequencer::FChannelModel> CreateChannelModel_Raw(const FMovieSceneChannelHandle&, FName) const override { return nullptr; }
 	virtual TSharedPtr<UE::Sequencer::STrackAreaLaneView> CreateChannelView_Raw(const FMovieSceneChannelHandle&, TWeakPtr<UE::Sequencer::FChannelModel>, const UE::Sequencer::FCreateTrackLaneViewParams&) const override { return nullptr; }
 	virtual void ExtendSectionMenu_Raw(FMenuBuilder& InMenuBuilder, TSharedPtr<FExtender> InMenuExtender, TConstArrayView<FMovieSceneChannelHandle> InChannels, TConstArrayView<UMovieSceneSection*> InSections, TWeakPtr<ISequencer> InSequencer) const override;
+	virtual void ExtendSidebarMenu_Raw(FMenuBuilder& InMenuBuilder, TSharedPtr<FExtender> InMenuExtender, TConstArrayView<FMovieSceneChannelHandle> InChannels, TConstArrayView<UMovieSceneSection*> InSections, TWeakPtr<ISequencer> InSequencer) const override;
 	virtual int32 DrawExtra_Raw(FMovieSceneChannel* InChannel, const UMovieSceneSection* InOwner, const FSequencerChannelPaintArgs& InPaintArgs, int32 InLayerId) const override;
 	virtual TSharedRef<SWidget> CreateKeyEditor_Raw(const FMovieSceneChannelHandle& InChannel, const UE::Sequencer::FCreateKeyEditorParams& Params) const override;
 	//~ End ISequencerChannelInterface
@@ -55,6 +56,17 @@ void TPropertyAnimatorEditorCurveChannelInterface<InChannelType, InMenuExtension
 		}));
 }
 
+template<typename InChannelType, typename InMenuExtensionType>
+void TPropertyAnimatorEditorCurveChannelInterface<InChannelType, InMenuExtensionType>::ExtendSidebarMenu_Raw(FMenuBuilder& InMenuBuilder
+	, TSharedPtr<FExtender> InMenuExtender
+	, TConstArrayView<FMovieSceneChannelHandle> InChannels
+	, TConstArrayView<UMovieSceneSection*> InSections
+	, TWeakPtr<ISequencer> InSequencer) const
+{
+	TSharedRef<FPropertyAnimatorEditorCurveSectionMenuExtension> Extension = MakeShared<FMenuExtensionType>(InChannels, InSections);
+	Extension->ExtendMenu(InMenuBuilder, false);
+}
+
 template<typename InChannelType, typename InMenuExtensionType>
 int32 TPropertyAnimatorEditorCurveChannelInterface<InChannelType, InMenuExtensionType>::DrawExtra_Raw(FMovieSceneChannel* InChannel
 	, const UMovieSceneSection* InOwner
diff --git a/Engine/Plugins/Experimental/PropertyAnimator/Source/PropertyAnimatorEditor/Private/Sequencer/PropertyAnimatorEditorCurveSectionMenuExtension.cpp b/Engine/Plugins/Experimental/PropertyAnimator/Source/PropertyAnimatorEditor/Private/Sequencer/PropertyAnimatorEditorCurveSectionMenuExtension.cpp
--- a/Engine/Plugins/Experimental/PropertyAnimator/Source/PropertyAnimatorEditor/Private/Sequencer/PropertyAnimatorEditorCurveSectionMenuExtension.cpp
+++ b/Engine/Plugins/Experimental/PropertyAnimator/Source/PropertyAnimatorEditor/Private/Sequencer/PropertyAnimatorEditorCurveSectionMenuExtension.cpp
@@ -45,7 +45,7 @@ FPropertyAnimatorEditorCurveSectionMenuExtension::FPropertyAnimatorEditorCurveSe
 	Initialize();
 }
 
-void FPropertyAnimatorEditorCurveSectionMenuExtension::ExtendMenu(FMenuBuilder& InMenuBuilder)
+void FPropertyAnimatorEditorCurveSectionMenuExtension::ExtendMenu(FMenuBuilder& InMenuBuilder, const bool bSubMenu)
 {
 	TSharedRef<FPropertyAnimatorEditorCurveSectionMenuExtension> This = SharedThis(this);
 
@@ -54,17 +54,31 @@ void FPropertyAnimatorEditorCurveSectionMenuExtension::ExtendMenu(FMenuBuilder&
 
 	if (ChannelHandles.Num() > 1)
 	{
-		InMenuBuilder.AddSubMenu(MenuTitle, MenuTooltip, FNewMenuDelegate::CreateLambda([This](FMenuBuilder& InInnerMenuBuilder)
-			{
-				This->BuildChannelsMenu(InInnerMenuBuilder);
-			}));
+		if (bSubMenu)
+		{
+			InMenuBuilder.AddSubMenu(MenuTitle, MenuTooltip, FNewMenuDelegate::CreateLambda([This](FMenuBuilder& InInnerMenuBuilder)
+				{
+					This->BuildChannelsMenu(InInnerMenuBuilder);
+				}));
+		}
+		else
+		{
+			This->BuildChannelsMenu(InMenuBuilder);
+		}
 	}
 	else if (ChannelHandles.Num() == 1)
 	{
-		InMenuBuilder.AddSubMenu(MenuTitle, MenuTooltip, FNewMenuDelegate::CreateLambda([This](FMenuBuilder& InInnerMenuBuilder)
-			{
-				This->BuildParametersMenu(InInnerMenuBuilder, 0);
-			}));
+		if (bSubMenu)
+		{
+			InMenuBuilder.AddSubMenu(MenuTitle, MenuTooltip, FNewMenuDelegate::CreateLambda([This](FMenuBuilder& InInnerMenuBuilder)
+				{
+					This->BuildParametersMenu(InInnerMenuBuilder, 0);
+				}));
+		}
+		else
+		{
+			This->BuildParametersMenu(InMenuBuilder, 0);
+		}
 	}
 }
 
diff --git a/Engine/Plugins/Experimental/PropertyAnimator/Source/PropertyAnimatorEditor/Private/Sequencer/PropertyAnimatorEditorCurveSectionMenuExtension.h b/Engine/Plugins/Experimental/PropertyAnimator/Source/PropertyAnimatorEditor/Private/Sequencer/PropertyAnimatorEditorCurveSectionMenuExtension.h
--- a/Engine/Plugins/Experimental/PropertyAnimator/Source/PropertyAnimatorEditor/Private/Sequencer/PropertyAnimatorEditorCurveSectionMenuExtension.h
+++ b/Engine/Plugins/Experimental/PropertyAnimator/Source/PropertyAnimatorEditor/Private/Sequencer/PropertyAnimatorEditorCurveSectionMenuExtension.h
@@ -36,7 +36,7 @@ public:
 
 	virtual bool GetParameterStructData(FMovieSceneChannelHandle InChannelHandle, UStruct*& OutStruct, uint8*& OutData) const = 0;
 
-	void ExtendMenu(FMenuBuilder& InMenuBuilder);
+	void ExtendMenu(FMenuBuilder& InMenuBuilder, const bool bSubMenu = true);
 
 private:
 	void Initialize();
diff --git a/Engine/Plugins/MovieScene/LevelSequenceEditor/Source/LevelSequenceEditor/Private/LevelSequenceEditorSubsystem.cpp b/Engine/Plugins/MovieScene/LevelSequenceEditor/Source/LevelSequenceEditor/Private/LevelSequenceEditorSubsystem.cpp
--- a/Engine/Plugins/MovieScene/LevelSequenceEditor/Source/LevelSequenceEditor/Private/LevelSequenceEditorSubsystem.cpp
+++ b/Engine/Plugins/MovieScene/LevelSequenceEditor/Source/LevelSequenceEditor/Private/LevelSequenceEditorSubsystem.cpp
@@ -294,13 +294,13 @@ FMovieSceneBindingPropertyInfoDetailCustomization::FMovieSceneBindingPropertyInf
 
 void FMovieSceneBindingPropertyInfoDetailCustomization::CustomizeHeader(TSharedRef<IPropertyHandle> InStructPropertyHandle, FDetailWidgetRow& HeaderRow, IPropertyTypeCustomizationUtils& CustomizationUtils)
 {
-	TSharedPtr<IPropertyUtilities> PropertyUtils = CustomizationUtils.GetPropertyUtilities();
+	HeaderRow.ShouldAutoExpand(true);
 
 	HeaderRow.NameContent()
 		[
 			InStructPropertyHandle->CreatePropertyNameWidget()
 		]
-	.ValueContent()
+		.ValueContent()
 		[
 			InStructPropertyHandle->CreatePropertyValueWidget()
 		];
@@ -558,6 +558,33 @@ void ULevelSequenceEditorSubsystem::Initialize(FSubsystemCollectionBase& Collect
 		}));
 
 	SequencerModule.GetObjectBindingContextMenuExtensibilityManager()->AddExtender(RebindComponentMenuExtender);
+
+	SidebarMenuExtender = MakeShared<FExtender>();
+
+	SidebarMenuExtender->AddMenuExtension(TEXT("Possessable"), EExtensionHook::First, CommandList,
+		FMenuExtensionDelegate::CreateLambda([this](FMenuBuilder& MenuBuilder)
+		{
+			// Only add menu entries where the focused sequence is a ULevelSequence
+			if (GetActiveSequencer())
+			{
+				AddBindingPropertiesSidebar(MenuBuilder);
+			}
+
+			TArray<FName> ComponentNames;
+			GetRebindComponentNames(ComponentNames);
+			if (ComponentNames.Num() > 0)
+			{
+				RebindComponentMenu(MenuBuilder);
+			}
+		}));
+
+	SidebarMenuExtender->AddMenuExtension(TEXT("CustomBinding"), EExtensionHook::First, CommandList,
+		FMenuExtensionDelegate::CreateLambda([this](FMenuBuilder& MenuBuilder)
+		{
+			AddBindingPropertiesMenu(MenuBuilder);
+		}));
+
+	SequencerModule.GetSidebarExtensibilityManager()->AddExtender(SidebarMenuExtender);
 }
 
 void ULevelSequenceEditorSubsystem::Deinitialize()
@@ -2435,7 +2462,10 @@ void ULevelSequenceEditorSubsystem::AddBindingPropertiesMenu(FMenuBuilder& MenuB
 	}
 }
 
-
+void ULevelSequenceEditorSubsystem::AddBindingPropertiesSidebar(FMenuBuilder& MenuBuilder)
+{
+	AddBindingPropertiesMenu(MenuBuilder);
+}
 
 void ULevelSequenceEditorSubsystem::FBindingPropertiesNotifyHook::NotifyPreChange(FProperty* PropertyAboutToChange)
 {
diff --git a/Engine/Plugins/MovieScene/LevelSequenceEditor/Source/LevelSequenceEditor/Private/Misc/LevelSequenceCustomization.cpp b/Engine/Plugins/MovieScene/LevelSequenceEditor/Source/LevelSequenceEditor/Private/Misc/LevelSequenceCustomization.cpp
--- a/Engine/Plugins/MovieScene/LevelSequenceEditor/Source/LevelSequenceEditor/Private/Misc/LevelSequenceCustomization.cpp
+++ b/Engine/Plugins/MovieScene/LevelSequenceEditor/Source/LevelSequenceEditor/Private/Misc/LevelSequenceCustomization.cpp
@@ -58,6 +58,7 @@ void FLevelSequenceCustomization::RegisterSequencerCustomization(FSequencerCusto
 	// Add a customization callback for the object binding context menu.
 	FSequencerCustomizationInfo Customization;
 	Customization.OnBuildObjectBindingContextMenu = FOnGetSequencerMenuExtender::CreateRaw(this, &FLevelSequenceCustomization::CreateObjectBindingContextMenuExtender);
+	Customization.OnBuildSidebarMenu = FOnGetSequencerMenuExtender::CreateRaw(this, &FLevelSequenceCustomization::CreateSidebarMenuExtender);
 	Builder.AddCustomization(Customization);
 }
 
@@ -260,7 +261,26 @@ void FLevelSequenceCustomization::ExtendObjectBindingContextMenu(FMenuBuilder& M
 	MenuBuilder.EndSection();
 }
 
+TSharedPtr<FExtender> FLevelSequenceCustomization::CreateSidebarMenuExtender(FViewModelPtr InViewModel)
+{
+	TSharedRef<FExtender> Extender = MakeShared<FExtender>();
+	
+	TSharedPtr<FObjectBindingModel> ObjectBindingModel = InViewModel->CastThisShared<FObjectBindingModel>();
+	
+	Extender->AddMenuExtension(TEXT("ObjectBindingActions"), EExtensionHook::Before, nullptr,
+		FMenuExtensionDelegate::CreateRaw(this, &FLevelSequenceCustomization::ExtendSidebarMenu, ObjectBindingModel));
+	
+	return Extender.ToSharedPtr();
+}
+
+void FLevelSequenceCustomization::ExtendSidebarMenu(FMenuBuilder& MenuBuilder, TSharedPtr<FObjectBindingModel> ObjectBindingModel)
+{
+	if (ObjectBindingModel.IsValid())
+	{
+		ObjectBindingModel->BuildContextMenu(MenuBuilder);
+	}
+}
+
 } // namespace UE::Sequencer
 
 #undef LOCTEXT_NAMESPACE
-
diff --git a/Engine/Plugins/MovieScene/LevelSequenceEditor/Source/LevelSequenceEditor/Private/Misc/LevelSequenceCustomization.h b/Engine/Plugins/MovieScene/LevelSequenceEditor/Source/LevelSequenceEditor/Private/Misc/LevelSequenceCustomization.h
--- a/Engine/Plugins/MovieScene/LevelSequenceEditor/Source/LevelSequenceEditor/Private/Misc/LevelSequenceCustomization.h
+++ b/Engine/Plugins/MovieScene/LevelSequenceEditor/Source/LevelSequenceEditor/Private/Misc/LevelSequenceCustomization.h
@@ -35,6 +35,9 @@ private:
 	TSharedPtr<FExtender> CreateObjectBindingContextMenuExtender(FViewModelPtr InViewModel);
 	void ExtendObjectBindingContextMenu(FMenuBuilder& MenuBuilder, TSharedPtr<FObjectBindingModel> ObjectBindingModel);
 
+	TSharedPtr<FExtender> CreateSidebarMenuExtender(FViewModelPtr InViewModel);
+	void ExtendSidebarMenu(FMenuBuilder& MenuBuilder, TSharedPtr<FObjectBindingModel> ObjectBindingModel);
+
 private:
 
 	TWeakPtr<ISequencer> WeakSequencer;
diff --git a/Engine/Plugins/MovieScene/LevelSequenceEditor/Source/LevelSequenceEditor/Public/LevelSequenceEditorSubsystem.h b/Engine/Plugins/MovieScene/LevelSequenceEditor/Source/LevelSequenceEditor/Public/LevelSequenceEditorSubsystem.h
--- a/Engine/Plugins/MovieScene/LevelSequenceEditor/Source/LevelSequenceEditor/Public/LevelSequenceEditorSubsystem.h
+++ b/Engine/Plugins/MovieScene/LevelSequenceEditor/Source/LevelSequenceEditor/Public/LevelSequenceEditorSubsystem.h
@@ -304,6 +304,9 @@ private:
 
 	void AddAssignActorMenu(FMenuBuilder& MenuBuilder);
 	void AddBindingPropertiesMenu(FMenuBuilder& MenuBuilder);
+public:
+	void AddBindingPropertiesSidebar(FMenuBuilder& MenuBuilder);
+private:
 	void OnFinishedChangingLocators(const FPropertyChangedEvent& PropertyChangedEvent, TSharedRef<IDetailsView> DetailsView, FGuid ObjectBindingID);
 
 	void GetRebindComponentNames(TArray<FName>& OutComponentNames);
@@ -330,6 +333,7 @@ private:
 	TSharedPtr<FExtender> AssignActorMenuExtender;
 	TSharedPtr<FExtender> BindingPropertiesMenuExtender;
 	TSharedPtr<FExtender> RebindComponentMenuExtender;
+	TSharedPtr<FExtender> SidebarMenuExtender;
 };
 
 #if UE_ENABLE_INCLUDE_ORDER_DEPRECATED_IN_5_2
diff --git a/Engine/Source/Editor/MovieSceneTools/Private/Channels/BuiltInChannelEditors.cpp b/Engine/Source/Editor/MovieSceneTools/Private/Channels/BuiltInChannelEditors.cpp
--- a/Engine/Source/Editor/MovieSceneTools/Private/Channels/BuiltInChannelEditors.cpp
+++ b/Engine/Source/Editor/MovieSceneTools/Private/Channels/BuiltInChannelEditors.cpp
@@ -1915,6 +1915,38 @@ void ExtendSectionMenu(FMenuBuilder& OuterMenuBuilder, TSharedPtr<FExtender> Men
 	MenuExtender->AddMenuExtension("SequencerChannels", EExtensionHook::First, nullptr, FMenuExtensionDelegate::CreateLambda([Extension](FMenuBuilder& MenuBuilder) { Extension->ExtendMenu(MenuBuilder); }));
 }
 
+void ExtendSidebarMenu(FMenuBuilder& OuterMenuBuilder, TSharedPtr<FExtender> MenuExtender, TArray<TMovieSceneChannelHandle<FMovieSceneFloatChannel>>&& Channels, TArrayView<UMovieSceneSection* const> Sections, TWeakPtr<ISequencer> InSequencer)
+{
+	TSharedRef<FCurveChannelSectionMenuExtension> Extension = FCurveChannelSectionMenuExtension::GetOrCreate(InSequencer);
+	Extension->AddSections(Sections);
+
+	MenuExtender->AddMenuExtension("SequencerChannels", EExtensionHook::First, nullptr, FMenuExtensionDelegate::CreateLambda([Extension](FMenuBuilder& MenuBuilder) { Extension->ExtendMenu(MenuBuilder); }));
+}
+
+void ExtendSidebarMenu(FMenuBuilder& OuterMenuBuilder, TSharedPtr<FExtender> MenuExtender, TArray<TMovieSceneChannelHandle<FMovieSceneDoubleChannel>>&& Channels, TArrayView<UMovieSceneSection* const> Sections, TWeakPtr<ISequencer> InSequencer)
+{
+	TSharedRef<FCurveChannelSectionMenuExtension> Extension = FCurveChannelSectionMenuExtension::GetOrCreate(InSequencer);
+	Extension->AddSections(Sections);
+
+	MenuExtender->AddMenuExtension("SequencerChannels", EExtensionHook::First, nullptr, FMenuExtensionDelegate::CreateLambda([Extension](FMenuBuilder& MenuBuilder) { Extension->ExtendMenu(MenuBuilder); }));
+}
+
+void ExtendSidebarMenu(FMenuBuilder& OuterMenuBuilder, TSharedPtr<FExtender> MenuExtender, TArray<TMovieSceneChannelHandle<FMovieSceneIntegerChannel>>&& Channels, TArrayView<UMovieSceneSection* const> Sections, TWeakPtr<ISequencer> InSequencer)
+{
+	TSharedRef<FCurveChannelSectionMenuExtension> Extension = FCurveChannelSectionMenuExtension::GetOrCreate(InSequencer);
+	Extension->AddSections(Sections);
+
+	MenuExtender->AddMenuExtension("SequencerChannels", EExtensionHook::First, nullptr, FMenuExtensionDelegate::CreateLambda([Extension](FMenuBuilder& MenuBuilder) { Extension->ExtendMenu(MenuBuilder); }));
+}
+
+void ExtendSidebarMenu(FMenuBuilder& OuterMenuBuilder, TSharedPtr<FExtender> MenuExtender, TArray<TMovieSceneChannelHandle<FMovieSceneBoolChannel>>&& Channels, TArrayView<UMovieSceneSection* const> Sections, TWeakPtr<ISequencer> InSequencer)
+{
+	TSharedRef<FCurveChannelSectionMenuExtension> Extension = FCurveChannelSectionMenuExtension::GetOrCreate(InSequencer);
+	Extension->AddSections(Sections);
+
+	MenuExtender->AddMenuExtension("SequencerChannels", EExtensionHook::First, nullptr, FMenuExtensionDelegate::CreateLambda([Extension](FMenuBuilder& MenuBuilder) { Extension->ExtendMenu(MenuBuilder); }));
+}
+
 void ExtendKeyMenu(FMenuBuilder& OuterMenuBuilder, TSharedPtr<FExtender> MenuExtender, TArray<TExtendKeyMenuParams<FMovieSceneFloatChannel>>&& Channels, TWeakPtr<ISequencer> InSequencer)
 {
 	using ExtensionType = TCurveChannelKeyMenuExtension<FMovieSceneFloatChannel>;
diff --git a/Engine/Source/Editor/MovieSceneTools/Private/Channels/BuiltInChannelEditors.h b/Engine/Source/Editor/MovieSceneTools/Private/Channels/BuiltInChannelEditors.h
--- a/Engine/Source/Editor/MovieSceneTools/Private/Channels/BuiltInChannelEditors.h
+++ b/Engine/Source/Editor/MovieSceneTools/Private/Channels/BuiltInChannelEditors.h
@@ -79,6 +79,12 @@ void ExtendSectionMenu(FMenuBuilder& OuterMenuBuilder, TSharedPtr<FExtender> Men
 void ExtendSectionMenu(FMenuBuilder& OuterMenuBuilder, TSharedPtr<FExtender> MenuExtender, TArray<TMovieSceneChannelHandle<FMovieSceneFloatChannel>>&& Channels, TArrayView<UMovieSceneSection* const> Sections, TWeakPtr<ISequencer> InSequencer);
 void ExtendSectionMenu(FMenuBuilder& OuterMenuBuilder, TSharedPtr<FExtender> MenuExtender, TArray<TMovieSceneChannelHandle<FMovieSceneIntegerChannel>>&& Channels, TArrayView<UMovieSceneSection* const> Sections, TWeakPtr<ISequencer> InSequencer);
 void ExtendSectionMenu(FMenuBuilder& OuterMenuBuilder, TSharedPtr<FExtender> MenuExtender, TArray<TMovieSceneChannelHandle<FMovieSceneBoolChannel>>&& Channels, TArrayView<UMovieSceneSection* const> Sections, TWeakPtr<ISequencer> InSequencer);
+
+void ExtendSidebarMenu(FMenuBuilder& OuterMenuBuilder, TSharedPtr<FExtender> MenuExtender, TArray<TMovieSceneChannelHandle<FMovieSceneDoubleChannel>>&& Channels, TArrayView<UMovieSceneSection* const> Sections, TWeakPtr<ISequencer> InSequencer);
+void ExtendSidebarMenu(FMenuBuilder& OuterMenuBuilder, TSharedPtr<FExtender> MenuExtender, TArray<TMovieSceneChannelHandle<FMovieSceneFloatChannel>>&& Channels, TArrayView<UMovieSceneSection* const> Sections, TWeakPtr<ISequencer> InSequencer);
+void ExtendSidebarMenu(FMenuBuilder& OuterMenuBuilder, TSharedPtr<FExtender> MenuExtender, TArray<TMovieSceneChannelHandle<FMovieSceneIntegerChannel>>&& Channels, TArrayView<UMovieSceneSection* const> Sections, TWeakPtr<ISequencer> InSequencer);
+void ExtendSidebarMenu(FMenuBuilder& OuterMenuBuilder, TSharedPtr<FExtender> MenuExtender, TArray<TMovieSceneChannelHandle<FMovieSceneBoolChannel>>&& Channels, TArrayView<UMovieSceneSection* const> Sections, TWeakPtr<ISequencer> InSequencer);
+
 void ExtendKeyMenu(FMenuBuilder& OuterMenuBuilder, TSharedPtr<FExtender> MenuExtender, TArray<TExtendKeyMenuParams<FMovieSceneDoubleChannel>>&& Channels, TWeakPtr<ISequencer> InSequencer);
 void ExtendKeyMenu(FMenuBuilder& OuterMenuBuilder, TSharedPtr<FExtender> MenuExtender, TArray<TExtendKeyMenuParams<FMovieSceneFloatChannel>>&& Channels, TWeakPtr<ISequencer> InSequencer);
 void ExtendKeyMenu(FMenuBuilder& OuterMenuBuilder, TSharedPtr<FExtender> MenuExtender, TArray<TExtendKeyMenuParams<FMovieSceneTimeWarpChannel>>&& Channels, TWeakPtr<ISequencer> InSequencer);
diff --git a/Engine/Source/Editor/MovieSceneTools/Private/Channels/PerlinNoiseChannelInterface.cpp b/Engine/Source/Editor/MovieSceneTools/Private/Channels/PerlinNoiseChannelInterface.cpp
--- a/Engine/Source/Editor/MovieSceneTools/Private/Channels/PerlinNoiseChannelInterface.cpp
+++ b/Engine/Source/Editor/MovieSceneTools/Private/Channels/PerlinNoiseChannelInterface.cpp
@@ -22,25 +22,39 @@ FPerlinNoiseChannelSectionMenuExtension::FPerlinNoiseChannelSectionMenuExtension
 	Initialize();
 }
 
-void FPerlinNoiseChannelSectionMenuExtension::ExtendMenu(FMenuBuilder& MenuBuilder)
+void FPerlinNoiseChannelSectionMenuExtension::ExtendMenu(FMenuBuilder& MenuBuilder, const bool bSubMenu)
 {
 	TSharedRef<FPerlinNoiseChannelSectionMenuExtension> SharedThis = this->AsShared();
 
 	if (ChannelHandles.Num() > 1)
 	{
-		MenuBuilder.AddSubMenu(
-			LOCTEXT("PerlinNoiseChannelsMenu", "Perlin Noise Channels"),
-			LOCTEXT("PerlinNoiseChannelsMenuToolTip", "Edit parameters for Perlin Noise channels"),
-			FNewMenuDelegate::CreateLambda([SharedThis](FMenuBuilder& InnerMenuBuilder) { SharedThis->BuildChannelsMenu(InnerMenuBuilder); })
-		);
+		if (bSubMenu)
+		{
+			MenuBuilder.AddSubMenu(
+				LOCTEXT("PerlinNoiseChannelsMenu", "Perlin Noise Channels"),
+				LOCTEXT("PerlinNoiseChannelsMenuToolTip", "Edit parameters for Perlin Noise channels"),
+				FNewMenuDelegate::CreateLambda([SharedThis](FMenuBuilder& InnerMenuBuilder) { SharedThis->BuildChannelsMenu(InnerMenuBuilder); })
+			);
+		}
+		else
+		{
+			SharedThis->BuildChannelsMenu(MenuBuilder);
+		}
 	}
 	else if (ChannelHandles.Num() == 1)
 	{
-		MenuBuilder.AddSubMenu(
-			LOCTEXT("PerlinNoiseChannelsMenu", "Perlin Noise Channels"),
-			LOCTEXT("PerlinNoiseChannelsMenuToolTip", "Edit parameters for Perlin Noise channels"),
-			FNewMenuDelegate::CreateLambda([SharedThis](FMenuBuilder& InnerMenuBuilder) { SharedThis->BuildParametersMenu(InnerMenuBuilder, 0); })
-		);
+		if (bSubMenu)
+		{
+			MenuBuilder.AddSubMenu(
+				LOCTEXT("PerlinNoiseChannelsMenu", "Perlin Noise Channels"),
+				LOCTEXT("PerlinNoiseChannelsMenuToolTip", "Edit parameters for Perlin Noise channels"),
+				FNewMenuDelegate::CreateLambda([SharedThis](FMenuBuilder& InnerMenuBuilder) { SharedThis->BuildParametersMenu(InnerMenuBuilder, 0); })
+			);
+		}
+		else
+		{
+			SharedThis->BuildParametersMenu(MenuBuilder, 0);
+		}
 	}
 }
 
diff --git a/Engine/Source/Editor/MovieSceneTools/Private/Sections/CinematicShotSection.cpp b/Engine/Source/Editor/MovieSceneTools/Private/Sections/CinematicShotSection.cpp
--- a/Engine/Source/Editor/MovieSceneTools/Private/Sections/CinematicShotSection.cpp
+++ b/Engine/Source/Editor/MovieSceneTools/Private/Sections/CinematicShotSection.cpp
@@ -267,6 +267,16 @@ void FCinematicShotSection::BuildSectionContextMenu(FMenuBuilder& MenuBuilder, c
 {
 	FViewportThumbnailSection::BuildSectionContextMenu(MenuBuilder, ObjectBinding);
 
+	AddShotMenuSection(MenuBuilder, ObjectBinding);
+}
+
+void FCinematicShotSection::BuildSectionSidebarMenu(FMenuBuilder& MenuBuilder, const FGuid& ObjectBinding)
+{
+	AddShotMenuSection(MenuBuilder, ObjectBinding);
+}
+
+void FCinematicShotSection::AddShotMenuSection(FMenuBuilder& MenuBuilder, const FGuid& ObjectBinding)
+{
 	UMovieSceneCinematicShotSection& SectionObject = GetSectionObjectAs<UMovieSceneCinematicShotSection>();
 
 	MenuBuilder.BeginSection(NAME_None, LOCTEXT("ShotMenuText", "Shot"));
diff --git a/Engine/Source/Editor/MovieSceneTools/Private/Sections/CinematicShotSection.h b/Engine/Source/Editor/MovieSceneTools/Private/Sections/CinematicShotSection.h
--- a/Engine/Source/Editor/MovieSceneTools/Private/Sections/CinematicShotSection.h
+++ b/Engine/Source/Editor/MovieSceneTools/Private/Sections/CinematicShotSection.h
@@ -35,6 +35,7 @@ public:
 	virtual void Tick(const FGeometry& AllottedGeometry, const FGeometry& ClippedGeometry, const double InCurrentTime, const float InDeltaTime) override;
 	virtual int32 OnPaintSection( FSequencerSectionPainter& Painter ) const override;
 	virtual void BuildSectionContextMenu(FMenuBuilder& MenuBuilder, const FGuid& ObjectBinding) override;
+	virtual void BuildSectionSidebarMenu(FMenuBuilder& MenuBuilder, const FGuid& ObjectBinding) override;
 	virtual FText GetSectionTitle() const override;
 	virtual float GetSectionHeight(const UE::Sequencer::FViewDensityInfo& ViewDensity) const override;
 	virtual FMargin GetContentPadding() const override;
@@ -48,6 +49,8 @@ public:
 
 private:
 
+	void AddShotMenuSection(FMenuBuilder& MenuBuilder, const FGuid& ObjectBinding);
+
 	/** The cinematic shot track editor that contains this section */
 	TWeakPtr<FCinematicShotTrackEditor> CinematicShotTrackEditor;
 
diff --git a/Engine/Source/Editor/MovieSceneTools/Private/Sections/ThumbnailSection.cpp b/Engine/Source/Editor/MovieSceneTools/Private/Sections/ThumbnailSection.cpp
--- a/Engine/Source/Editor/MovieSceneTools/Private/Sections/ThumbnailSection.cpp
+++ b/Engine/Source/Editor/MovieSceneTools/Private/Sections/ThumbnailSection.cpp
@@ -118,63 +118,80 @@ void FThumbnailSection::BuildSectionContextMenu(FMenuBuilder& MenuBuilder, const
 		MenuBuilder.AddSubMenu(
 			LOCTEXT("ThumbnailsMenu", "Thumbnails"),
 			FText(),
-			FNewMenuDelegate::CreateLambda([this](FMenuBuilder& InMenuBuilder){
+			FNewMenuDelegate::CreateLambda([this](FMenuBuilder& InMenuBuilder)
+			{
+				BuildThumbnailsMenu(InMenuBuilder);
+			})
+		);
+	}
+	MenuBuilder.EndSection();
+}
 
-				TSharedPtr<ISequencer> Sequencer = SequencerPtr.Pin();
+void FThumbnailSection::BuildThumbnailsMenu(FMenuBuilder& InMenuBuilder)
+{
+	const TSharedPtr<ISequencer> Sequencer = SequencerPtr.Pin();
+	if (!Sequencer.IsValid())
+	{
+		return;
+	}
 
-				FText CurrentTime = FText::FromString(Sequencer->GetNumericTypeInterface()->ToString(Sequencer->GetLocalTime().Time.GetFrame().Value));
+	const FText CurrentTime = FText::FromString(Sequencer->GetNumericTypeInterface()->ToString(Sequencer->GetLocalTime().Time.GetFrame().Value));
 
-				InMenuBuilder.BeginSection(NAME_None, LOCTEXT("ThisSectionText", "This Section"));
+	InMenuBuilder.BeginSection(TEXT("Thumbnails"), LOCTEXT("ThumbnailsMenuSection", "Thumbnails"));
+	{
+		InMenuBuilder.AddMenuEntry(
+			LOCTEXT("RefreshText", "Refresh"),
+			LOCTEXT("RefreshTooltip", "Refresh this section's thumbnails"),
+			FSlateIcon(),
+			FUIAction(FExecuteAction::CreateRaw(this, &FThumbnailSection::RedrawThumbnails))
+		);
+		InMenuBuilder.AddMenuEntry(
+			FText::Format(LOCTEXT("SetSingleTime", "Set Thumbnail Time To {0}"), CurrentTime),
+			LOCTEXT("SetSingleTimeTooltip", "Defines the time at which this section should draw its single thumbnail to the current cursor position"),
+			FSlateIcon(),
+			FUIAction(
+			FExecuteAction::CreateLambda([this]
 				{
-					InMenuBuilder.AddMenuEntry(
-						LOCTEXT("RefreshText", "Refresh"),
-						LOCTEXT("RefreshTooltip", "Refresh this section's thumbnails"),
-						FSlateIcon(),
-						FUIAction(FExecuteAction::CreateRaw(this, &FThumbnailSection::RedrawThumbnails))
-					);
-					InMenuBuilder.AddMenuEntry(
-						FText::Format(LOCTEXT("SetSingleTime", "Set Thumbnail Time To {0}"), CurrentTime),
-						LOCTEXT("SetSingleTimeTooltip", "Defines the time at which this section should draw its single thumbnail to the current cursor position"),
-						FSlateIcon(),
-						FUIAction(
-						FExecuteAction::CreateLambda([this, Sequencer]{
-								SetSingleTime(Sequencer->GetLocalTime().AsSeconds());
-								GetMutableDefault<UMovieSceneUserThumbnailSettings>()->bDrawSingleThumbnails = true;
-								GetMutableDefault<UMovieSceneUserThumbnailSettings>()->SaveConfig();
-							})
-						)
-					);
-				}
-				InMenuBuilder.EndSection();
-
-				InMenuBuilder.BeginSection(NAME_None, LOCTEXT("GlobalSettingsText", "Global Settings"));
+					const TSharedPtr<ISequencer> Sequencer = SequencerPtr.Pin();
+					if (!Sequencer.IsValid())
+					{
+						return;
+					}
+
+					SetSingleTime(Sequencer->GetLocalTime().AsSeconds());
+					GetMutableDefault<UMovieSceneUserThumbnailSettings>()->bDrawSingleThumbnails = true;
+					GetMutableDefault<UMovieSceneUserThumbnailSettings>()->SaveConfig();
+				})
+			)
+		);
+
+		InMenuBuilder.AddMenuEntry(
+			LOCTEXT("RefreshAllText", "Refresh All"),
+			LOCTEXT("RefreshAllTooltip", "Refresh all sections' thumbnails"),
+			FSlateIcon(),
+			FUIAction(FExecuteAction::CreateLambda([]
 				{
-					InMenuBuilder.AddMenuEntry(
-						LOCTEXT("RefreshAllText", "Refresh All"),
-						LOCTEXT("RefreshAllTooltip", "Refresh all sections' thumbnails"),
-						FSlateIcon(),
-						FUIAction(FExecuteAction::CreateLambda([]{
-							GetDefault<UMovieSceneUserThumbnailSettings>()->BroadcastRedrawThumbnails();
-						}))
-					);
-
-					FPropertyEditorModule& PropertyModule = FModuleManager::LoadModuleChecked<FPropertyEditorModule>("PropertyEditor");
-
-					FDetailsViewArgs Args;
-					Args.bAllowSearch = false;
-					Args.NameAreaSettings = FDetailsViewArgs::HideNameArea;
-
-					TSharedRef<IDetailsView> DetailView = PropertyModule.CreateDetailView(Args);
-					DetailView->SetObject(GetMutableDefault<UMovieSceneUserThumbnailSettings>());
-					InMenuBuilder.AddWidget(DetailView, FText(), true);
-				}
-				InMenuBuilder.EndSection();
-			})
+					GetDefault<UMovieSceneUserThumbnailSettings>()->BroadcastRedrawThumbnails();
+				}))
 		);
+
+		FPropertyEditorModule& PropertyModule = FModuleManager::LoadModuleChecked<FPropertyEditorModule>("PropertyEditor");
+
+		FDetailsViewArgs Args;
+		Args.bAllowSearch = false;
+		Args.NameAreaSettings = FDetailsViewArgs::HideNameArea;
+
+		TSharedRef<IDetailsView> DetailView = PropertyModule.CreateDetailView(Args);
+		DetailView->SetObject(GetMutableDefault<UMovieSceneUserThumbnailSettings>());
+		InMenuBuilder.AddWidget(DetailView, FText(), true);
 	}
-	MenuBuilder.EndSection();
+	InMenuBuilder.EndSection();
 }
 
+void FThumbnailSection::BuildSectionSidebarMenu(FMenuBuilder& MenuBuilder, const FGuid& ObjectBinding)
+{
+	BuildThumbnailsMenu(MenuBuilder);
+}
 
 float FThumbnailSection::GetSectionGripSize() const
 {
diff --git a/Engine/Source/Editor/MovieSceneTools/Private/TrackEditors/CinematicShotTrackEditor.cpp b/Engine/Source/Editor/MovieSceneTools/Private/TrackEditors/CinematicShotTrackEditor.cpp
--- a/Engine/Source/Editor/MovieSceneTools/Private/TrackEditors/CinematicShotTrackEditor.cpp
+++ b/Engine/Source/Editor/MovieSceneTools/Private/TrackEditors/CinematicShotTrackEditor.cpp
@@ -179,7 +179,7 @@ void FCinematicShotTrackEditor::Tick(float DeltaTime)
 	}
 }
 
-void FCinematicShotTrackEditor::BuildTrackContextMenu( FMenuBuilder& MenuBuilder, UMovieSceneTrack* Track )
+void FCinematicShotTrackEditor::BuildTrackContextMenu(FMenuBuilder& MenuBuilder, UMovieSceneTrack* Track)
 {
 	MenuBuilder.BeginSection("Import/Export", NSLOCTEXT("Sequencer", "ImportExportMenuSectionName", "Import/Export"));
 
@@ -214,6 +214,11 @@ void FCinematicShotTrackEditor::BuildTrackContextMenu( FMenuBuilder& MenuBuilder
 	MenuBuilder.EndSection();
 }
 
+void FCinematicShotTrackEditor::BuildTrackSidebarMenu(FMenuBuilder& MenuBuilder, UMovieSceneTrack* Track)
+{
+	BuildTrackContextMenu(MenuBuilder, Track);
+}
+
 void FCinematicShotTrackEditor::InsertShot()
 {
 	UMovieScene* FocusedMovieScene = GetFocusedMovieScene();
diff --git a/Engine/Source/Editor/MovieSceneTools/Private/TrackEditors/SkeletalAnimationTrackEditor.cpp b/Engine/Source/Editor/MovieSceneTools/Private/TrackEditors/SkeletalAnimationTrackEditor.cpp
--- a/Engine/Source/Editor/MovieSceneTools/Private/TrackEditors/SkeletalAnimationTrackEditor.cpp
+++ b/Engine/Source/Editor/MovieSceneTools/Private/TrackEditors/SkeletalAnimationTrackEditor.cpp
@@ -448,7 +448,7 @@ public:
 	{
 		if (Params.ParentObjectBindingGuid.IsValid())
 		{
-			if (USkeletalMeshComponent* SkelMeshComp = AcquireSkeletalMeshFromObjectGuid(Params.ParentObjectBindingGuid, Params.Sequencer))
+			if (USkeletalMeshComponent* SkelMeshComp = AcquireSkeletalMeshFromObjectGuid(Params.ParentObjectBindingGuid, Params.SequencerWeak.Pin()))
 			{
 				TScriptInterface<ISequencerAnimationOverride> SequencerAnimOverride = ISequencerAnimationOverride::GetSequencerAnimOverride(SkelMeshComp);
 				if (SequencerAnimOverride.GetObject())
@@ -491,7 +491,7 @@ public:
 				{
 					// Store the compatible skeleton's name, and create a property widget with a filter that will check
 					// for animations that match that skeleton.
-					Skeleton = AcquireSkeletonFromObjectGuid(Params.ParentObjectBindingGuid, Params.Sequencer);
+					Skeleton = AcquireSkeletonFromObjectGuid(Params.ParentObjectBindingGuid, Params.SequencerWeak.Pin());
 					SkeletonName = FAssetData(Skeleton).GetExportTextName();
 
 					TSharedPtr<IPropertyUtilities> PropertyUtilities = CustomizationUtils.GetPropertyUtilities();
diff --git a/Engine/Source/Editor/MovieSceneTools/Public/Channels/PerlinNoiseChannelInterface.h b/Engine/Source/Editor/MovieSceneTools/Public/Channels/PerlinNoiseChannelInterface.h
--- a/Engine/Source/Editor/MovieSceneTools/Public/Channels/PerlinNoiseChannelInterface.h
+++ b/Engine/Source/Editor/MovieSceneTools/Public/Channels/PerlinNoiseChannelInterface.h
@@ -48,7 +48,7 @@ struct FPerlinNoiseChannelSectionMenuExtension : TSharedFromThis<FPerlinNoiseCha
 {
 	FPerlinNoiseChannelSectionMenuExtension(TArrayView<const FMovieSceneChannelHandle> InChannelHandles, TArrayView<UMovieSceneSection* const> InSections);
 
-	void ExtendMenu(FMenuBuilder& MenuBuilder);
+	void ExtendMenu(FMenuBuilder& MenuBuilder, const bool bSubMenu = true);
 
 private:
 
@@ -131,6 +131,12 @@ struct TPerlinNoiseChannelInterface : ISequencerChannelInterface
 		MenuExtender->AddMenuExtension("SequencerChannels", EExtensionHook::First, nullptr, FMenuExtensionDelegate::CreateLambda([Extension](FMenuBuilder& MenuBuilder) { Extension->ExtendMenu(MenuBuilder); }));
 	}
 
+	virtual void ExtendSidebarMenu_Raw(FMenuBuilder& MenuBuilder, TSharedPtr<FExtender> MenuExtender, TArrayView<const FMovieSceneChannelHandle> Channels, TArrayView<UMovieSceneSection* const> Sections, TWeakPtr<ISequencer> InSequencer) const override
+	{
+		TSharedRef<FPerlinNoiseChannelSectionMenuExtension> Extension = MakeShared<FPerlinNoiseChannelSectionMenuExtension>(Channels, Sections);
+		Extension->ExtendMenu(MenuBuilder, false);
+	}
+
 	virtual void DrawKeys_Raw(FMovieSceneChannel* Channel, TArrayView<const FKeyHandle> InKeyHandles, const UMovieSceneSection* InOwner, TArrayView<FKeyDrawParams> OutKeyDrawParams) const override
 	{
 	}
diff --git a/Engine/Source/Editor/MovieSceneTools/Public/Sections/ThumbnailSection.h b/Engine/Source/Editor/MovieSceneTools/Public/Sections/ThumbnailSection.h
--- a/Engine/Source/Editor/MovieSceneTools/Public/Sections/ThumbnailSection.h
+++ b/Engine/Source/Editor/MovieSceneTools/Public/Sections/ThumbnailSection.h
@@ -69,6 +69,7 @@ public:
 	//~ ISequencerSection interface
 
 	virtual void BuildSectionContextMenu(FMenuBuilder& MenuBuilder, const FGuid& ObjectBinding) override;
+	virtual void BuildSectionSidebarMenu(FMenuBuilder& MenuBuilder, const FGuid& ObjectBinding) override;
 	virtual TSharedRef<SWidget> GenerateSectionWidget() override;
 	virtual float GetSectionGripSize() const override;
 	virtual float GetSectionHeight(const UE::Sequencer::FViewDensityInfo& ViewDensity) const override;
@@ -91,6 +92,8 @@ protected:
 	/** Get rename visibility */
 	EVisibility GetRenameVisibility() const;
 
+	void BuildThumbnailsMenu(FMenuBuilder& InMenuBuilder);
+
 protected:
 
 	/** The section we are visualizing. */
diff --git a/Engine/Source/Editor/MovieSceneTools/Public/TrackEditors/CinematicShotTrackEditor.h b/Engine/Source/Editor/MovieSceneTools/Public/TrackEditors/CinematicShotTrackEditor.h
--- a/Engine/Source/Editor/MovieSceneTools/Public/TrackEditors/CinematicShotTrackEditor.h
+++ b/Engine/Source/Editor/MovieSceneTools/Public/TrackEditors/CinematicShotTrackEditor.h
@@ -55,7 +55,8 @@ public:
 	virtual TSharedRef<ISequencerSection> MakeSectionInterface(UMovieSceneSection& SectionObject, UMovieSceneTrack& Track, FGuid ObjectBinding) override;
 	virtual bool SupportsSequence(UMovieSceneSequence* InSequence) const override;
 	virtual void Tick(float DeltaTime) override;
-	virtual void BuildTrackContextMenu( FMenuBuilder& MenuBuilder, UMovieSceneTrack* Track ) override;
+	virtual void BuildTrackContextMenu(FMenuBuilder& MenuBuilder, UMovieSceneTrack* Track) override;
+	virtual void BuildTrackSidebarMenu(FMenuBuilder& MenuBuilder, UMovieSceneTrack* Track) override;
 
 	UE_DEPRECATED(5.3, "InsertShot has been deprecated in favor of FSubTrackEditor::InsertSection(UMovieSceneTrack*)")
 	void InsertShot();
diff --git a/Engine/Source/Editor/Sequencer/Private/MVVM/ViewModels/FolderModel.cpp b/Engine/Source/Editor/Sequencer/Private/MVVM/ViewModels/FolderModel.cpp
--- a/Engine/Source/Editor/Sequencer/Private/MVVM/ViewModels/FolderModel.cpp
+++ b/Engine/Source/Editor/Sequencer/Private/MVVM/ViewModels/FolderModel.cpp
@@ -322,6 +322,21 @@ void FFolderModel::BuildContextMenu(FMenuBuilder& MenuBuilder)
 	MenuBuilder.EndSection();
 }
 
+void FFolderModel::BuildSidebarMenu(FMenuBuilder& MenuBuilder)
+{
+	FOutlinerItemModel::BuildSidebarMenu(MenuBuilder);
+
+	MenuBuilder.BeginSection(TEXT("Folder"), LOCTEXT("FolderContextMenuSectionName", "Folder"));
+	{
+		MenuBuilder.AddMenuEntry(
+			LOCTEXT("SetColor", "Set Color"),
+			LOCTEXT("SetColorTooltip", "Set the color for the selected folders"),
+			FSlateIcon(),
+			FUIAction(FExecuteAction::CreateSP(this, &FFolderModel::SetFolderColor))
+		);
+	}
+	MenuBuilder.EndSection();
+}
 
 TSharedPtr<SWidget> FFolderModel::CreateOutlinerViewForColumn(const FCreateOutlinerViewParams& InParams, const FName& InColumnName)
 {
diff --git a/Engine/Source/Editor/Sequencer/Private/MVVM/ViewModels/ObjectBindingModel.cpp b/Engine/Source/Editor/Sequencer/Private/MVVM/ViewModels/ObjectBindingModel.cpp
--- a/Engine/Source/Editor/Sequencer/Private/MVVM/ViewModels/ObjectBindingModel.cpp
+++ b/Engine/Source/Editor/Sequencer/Private/MVVM/ViewModels/ObjectBindingModel.cpp
@@ -989,8 +989,8 @@ void FObjectBindingModel::BuildContextMenu(FMenuBuilder& MenuBuilder)
 	const UClass* ObjectClass = FindObjectClass();
 	
 	TSharedPtr<FExtender> Extender = EditorViewModel->GetSequencerMenuExtender(
-			SequencerModule.GetObjectBindingContextMenuExtensibilityManager(), TArrayBuilder<UObject*>().Add(BoundObject),
-			&FSequencerCustomizationInfo::OnBuildObjectBindingContextMenu, SharedThis(this));
+		SequencerModule.GetObjectBindingContextMenuExtensibilityManager(), TArrayBuilder<UObject*>().Add(BoundObject),
+		&FSequencerCustomizationInfo::OnBuildObjectBindingContextMenu, SharedThis(this));
 	if (Extender.IsValid())
 	{
 		MenuBuilder.PushExtender(Extender.ToSharedRef());
@@ -1026,6 +1026,29 @@ void FObjectBindingModel::BuildOrganizeContextMenu(FMenuBuilder& MenuBuilder)
 	FOutlinerItemModel::BuildOrganizeContextMenu(MenuBuilder);
 }
 
+void FObjectBindingModel::BuildSidebarMenu(FMenuBuilder& MenuBuilder)
+{
+	const TSharedPtr<FSequencerEditorViewModel> EditorViewModel = GetEditor();
+	if (!EditorViewModel.IsValid())
+	{
+		return;
+	}
+
+	UObject* const BoundObject = EditorViewModel->GetSequencerImpl()->FindSpawnedObjectOrTemplate(ObjectBindingID);
+	
+	ISequencerModule& SequencerModule = FModuleManager::GetModuleChecked<ISequencerModule>(TEXT("Sequencer"));
+	
+	const TSharedPtr<FExtender> Extender = EditorViewModel->GetSequencerMenuExtender(SequencerModule.GetSidebarExtensibilityManager()
+		, TArrayBuilder<UObject*>().Add(BoundObject), &FSequencerCustomizationInfo::OnBuildSidebarMenu, SharedThis(this));
+	if (Extender.IsValid())
+	{
+		MenuBuilder.PushExtender(Extender.ToSharedRef());
+	}
+	
+	MenuBuilder.BeginSection(TEXT("ObjectBindingActions"), LOCTEXT("ObjectBindingsMenuSection", "Object Bindings"));
+	MenuBuilder.EndSection();
+}
+
 void FObjectBindingModel::AddTagMenu(FMenuBuilder& MenuBuilder)
 {
 	MenuBuilder.AddMenuEntry(FSequencerCommands::Get().OpenTaggedBindingManager);
diff --git a/Engine/Source/Editor/Sequencer/Private/MVVM/ViewModels/OutlinerItemModel.cpp b/Engine/Source/Editor/Sequencer/Private/MVVM/ViewModels/OutlinerItemModel.cpp
--- a/Engine/Source/Editor/Sequencer/Private/MVVM/ViewModels/OutlinerItemModel.cpp
+++ b/Engine/Source/Editor/Sequencer/Private/MVVM/ViewModels/OutlinerItemModel.cpp
@@ -63,68 +63,101 @@ static bool NodeMatchesTextFilterTerm(TSharedPtr<const UE::Sequencer::IOutlinerE
 	return bMatched;
 }
 
-void AddEvalOptionsPropertyMenuItem(FMenuBuilder& MenuBuilder, FCanExecuteAction InCanExecute, const TArray<UMovieSceneTrack*>& AllTracks, const FBoolProperty* Property, TFunction<bool(UMovieSceneTrack*)> Validator = nullptr)
+void FOutlinerItemModelMixin::AddEvalOptionsPropertyMenuItem(FMenuBuilder& InMenuBuilder, const FBoolProperty* InProperty, TFunction<bool(UMovieSceneTrack*)> InValidator)
 {
-	bool bIsChecked = AllTracks.ContainsByPredicate(
-		[=](UMovieSceneTrack* InTrack)
-		{
-			return (!Validator || Validator(InTrack)) && Property->GetPropertyValue(Property->ContainerPtrToValuePtr<void>(&InTrack->EvalOptions));
-		});
+	auto IsChecked = [InProperty, InValidator](const TArray<UMovieSceneTrack*>& InTracks) -> bool
+	{
+		return InTracks.ContainsByPredicate(
+			[InValidator, InProperty](UMovieSceneTrack* InTrack)
+			{
+				return (!InValidator || InValidator(InTrack)) && InProperty->GetPropertyValue(InProperty->ContainerPtrToValuePtr<void>(&InTrack->EvalOptions));
+			});
+	};
 
-	MenuBuilder.AddMenuEntry(
-		Property->GetDisplayNameText(),
-		Property->GetToolTipText(),
+	InMenuBuilder.AddMenuEntry(
+		InProperty->GetDisplayNameText(),
+		InProperty->GetToolTipText(),
 		FSlateIcon(),
 		FUIAction(
-			FExecuteAction::CreateLambda([AllTracks, Property, Validator, bIsChecked]{
-				FScopedTransaction Transaction(FText::Format(NSLOCTEXT("Sequencer", "TrackNodeSetRoundEvaluation", "Set '{0}'"), Property->GetDisplayNameText()));
+			FExecuteAction::CreateLambda([this, InProperty, InValidator, IsChecked]
+			{
+				FScopedTransaction Transaction(FText::Format(NSLOCTEXT("Sequencer", "TrackNodeSetRoundEvaluation", "Set '{0}'"), InProperty->GetDisplayNameText()));
+				const TArray<UMovieSceneTrack*> AllTracks = GetSelectedTracks();
 				for (UMovieSceneTrack* Track : AllTracks)
 				{
-					if (Validator && !Validator(Track))
+					if (InValidator && !InValidator(Track))
 					{
 						continue;
 					}
-					void* PropertyContainer = Property->ContainerPtrToValuePtr<void>(&Track->EvalOptions);
+					void* PropertyContainer = InProperty->ContainerPtrToValuePtr<void>(&Track->EvalOptions);
 					Track->Modify();
-					Property->SetPropertyValue(PropertyContainer, !bIsChecked);
+					InProperty->SetPropertyValue(PropertyContainer, !IsChecked(AllTracks));
+				}
+			}),
+			FCanExecuteAction::CreateLambda([this]
+			{
+				const TSharedPtr<FSequencer> Sequencer = GetEditor()->GetSequencerImpl();
+				if (Sequencer)
+				{
+					return !Sequencer->IsReadOnly();
 				}
+				return false;
 			}),
-			InCanExecute,
-			FIsActionChecked::CreateLambda([=]{ return bIsChecked; })
+			FIsActionChecked::CreateLambda([this, IsChecked]
+			{
+				const TArray<UMovieSceneTrack*> AllTracks = GetSelectedTracks();
+				return IsChecked(AllTracks);
+			})
 		),
 		NAME_None,
 		EUserInterfaceActionType::Check
 	);
 }
 
-void AddDisplayOptionsPropertyMenuItem(FMenuBuilder& MenuBuilder, FCanExecuteAction InCanExecute, const TArray<UMovieSceneTrack*>& AllTracks, const FBoolProperty* Property, TFunction<bool(UMovieSceneTrack*)> Validator = nullptr)
+void FOutlinerItemModelMixin::AddDisplayOptionsPropertyMenuItem(FMenuBuilder& InMenuBuilder, const FBoolProperty* InProperty, TFunction<bool(UMovieSceneTrack*)> InValidator)
 {
-	bool bIsChecked = AllTracks.ContainsByPredicate(
-		[=](UMovieSceneTrack* InTrack)
+	auto IsChecked = [InProperty, InValidator](const TArray<UMovieSceneTrack*>& InTracks) -> bool
 	{
-		return (!Validator || Validator(InTrack)) && Property->GetPropertyValue(Property->ContainerPtrToValuePtr<void>(&InTrack->DisplayOptions));
-	});
+		return InTracks.ContainsByPredicate(
+			[InValidator, InProperty](UMovieSceneTrack* InTrack)
+			{
+				return (!InValidator || InValidator(InTrack)) && InProperty->GetPropertyValue(InProperty->ContainerPtrToValuePtr<void>(&InTrack->DisplayOptions));
+			});
+	};
 
-	MenuBuilder.AddMenuEntry(
-		Property->GetDisplayNameText(),
-		Property->GetToolTipText(),
+	InMenuBuilder.AddMenuEntry(
+		InProperty->GetDisplayNameText(),
+		InProperty->GetToolTipText(),
 		FSlateIcon(),
 		FUIAction(
-			FExecuteAction::CreateLambda([AllTracks, Property, Validator, bIsChecked] {
-				FScopedTransaction Transaction(FText::Format(NSLOCTEXT("Sequencer", "TrackNodeSetDisplayOption", "Set '{0}'"), Property->GetDisplayNameText()));
+			FExecuteAction::CreateLambda([this, InProperty, InValidator, IsChecked] {
+				FScopedTransaction Transaction(FText::Format(NSLOCTEXT("Sequencer", "TrackNodeSetDisplayOption", "Set '{0}'"), InProperty->GetDisplayNameText()));
+				const TArray<UMovieSceneTrack*> AllTracks = GetSelectedTracks();
 				for (UMovieSceneTrack* Track : AllTracks)
 				{
-					if (Validator && !Validator(Track))
+					if (InValidator && !InValidator(Track))
 					{
 						continue;
 					}
-					void* PropertyContainer = Property->ContainerPtrToValuePtr<void>(&Track->DisplayOptions);
+					void* PropertyContainer = InProperty->ContainerPtrToValuePtr<void>(&Track->DisplayOptions);
 					Track->Modify();
-					Property->SetPropertyValue(PropertyContainer, !bIsChecked);
+					InProperty->SetPropertyValue(PropertyContainer, !IsChecked(AllTracks));
+				}
+			}),
+			FCanExecuteAction::CreateLambda([this]
+			{
+				const TSharedPtr<FSequencer> Sequencer = GetEditor()->GetSequencerImpl();
+				if (Sequencer)
+				{
+					return !Sequencer->IsReadOnly();
 				}
+				return false;
 			}),
-			InCanExecute,
-			FIsActionChecked::CreateLambda([=] { return bIsChecked; })
+			FIsActionChecked::CreateLambda([this, IsChecked]
+			{
+				const TArray<UMovieSceneTrack*> AllTracks = GetSelectedTracks();
+				return IsChecked(AllTracks);
+			})
 		),
 		NAME_None,
 		EUserInterfaceActionType::Check
@@ -551,48 +584,8 @@ void FOutlinerItemModelMixin::BuildContextMenu(FMenuBuilder& MenuBuilder)
 
 	if (AllTracks.Num())
 	{
-		MenuBuilder.BeginSection("GeneralTrackOptions", NSLOCTEXT("Sequencer", "TrackNodeGeneralOptions", "Track Options"));
-		{
-			UStruct* EvalOptionsStruct = FMovieSceneTrackEvalOptions::StaticStruct();
-
-			const FBoolProperty* NearestSectionProperty = CastField<FBoolProperty>(EvalOptionsStruct->FindPropertyByName(GET_MEMBER_NAME_CHECKED(FMovieSceneTrackEvalOptions, bEvalNearestSection)));
-			auto CanEvaluateNearest = [](UMovieSceneTrack* InTrack) { return InTrack->EvalOptions.bCanEvaluateNearestSection != 0; };
-			if (NearestSectionProperty && AllTracks.ContainsByPredicate(CanEvaluateNearest))
-			{
-				TFunction<bool(UMovieSceneTrack*)> Validator = CanEvaluateNearest;
-				AddEvalOptionsPropertyMenuItem(MenuBuilder, CanExecute, AllTracks, NearestSectionProperty, Validator);
-			}
-
-			const FBoolProperty* PrerollProperty = CastField<FBoolProperty>(EvalOptionsStruct->FindPropertyByName(GET_MEMBER_NAME_CHECKED(FMovieSceneTrackEvalOptions, bEvaluateInPreroll)));
-			if (PrerollProperty)
-			{
-				AddEvalOptionsPropertyMenuItem(MenuBuilder, CanExecute, AllTracks, PrerollProperty);
-			}
-
-			const FBoolProperty* PostrollProperty = CastField<FBoolProperty>(EvalOptionsStruct->FindPropertyByName(GET_MEMBER_NAME_CHECKED(FMovieSceneTrackEvalOptions, bEvaluateInPostroll)));
-			if (PostrollProperty)
-			{
-				AddEvalOptionsPropertyMenuItem(MenuBuilder, CanExecute, AllTracks, PostrollProperty);
-			}
-		}
-		MenuBuilder.EndSection();
-
-		MenuBuilder.BeginSection("TrackDisplayOptions", NSLOCTEXT("Sequencer", "TrackNodeDisplayOptions", "Display Options"));
-		{
-			MenuBuilder.AddSubMenu(
-				LOCTEXT("SetColorTint", "Set Color Tint"),
-				LOCTEXT("SetColorTintTooltip", "Set color tint from the preferences for the selected sections or the track's sections"),
-				FNewMenuDelegate::CreateSP(SharedThis, &FOutlinerItemModelMixin::BuildSectionColorTintsContextMenu));
-
-			UStruct* DisplayOptionsStruct = FMovieSceneTrackDisplayOptions::StaticStruct();
-
-			const FBoolProperty* ShowVerticalFramesProperty = CastField<FBoolProperty>(DisplayOptionsStruct->FindPropertyByName(GET_MEMBER_NAME_CHECKED(FMovieSceneTrackDisplayOptions, bShowVerticalFrames)));
-			if (ShowVerticalFramesProperty)
-			{
-				AddDisplayOptionsPropertyMenuItem(MenuBuilder, CanExecute, AllTracks, ShowVerticalFramesProperty);
-			}
-		}
-		MenuBuilder.EndSection();
+		BuildTrackOptionsMenu(MenuBuilder);
+		BuildDisplayOptionsMenu(MenuBuilder);
 	}
 }
 
@@ -652,15 +645,116 @@ void FOutlinerItemModelMixin::BuildOrganizeContextMenu(FMenuBuilder& MenuBuilder
 	}
 }
 
-void FOutlinerItemModelMixin::BuildSectionColorTintsContextMenu(FMenuBuilder& MenuBuilder)
+void FOutlinerItemModelMixin::BuildDisplayOptionsMenu(FMenuBuilder& MenuBuilder)
 {
-	TSharedPtr<FSequencer> Sequencer = GetEditor()->GetSequencerImpl();
-	TSharedPtr<FSequencerSelection> Selection = Sequencer->GetViewModel()->GetSelection();
+	const TSharedPtr<FSequencer> Sequencer = GetEditor()->GetSequencerImpl();
+	if (!Sequencer.IsValid())
+	{
+		return;
+	}
+	
+	const TSharedRef<FOutlinerItemModelMixin> SharedThis(AsViewModel()->AsShared(), this);
+
+	const bool bIsReadOnly = Sequencer->IsReadOnly();
+	const FCanExecuteAction CanExecute = FCanExecuteAction::CreateLambda([bIsReadOnly]{ return !bIsReadOnly; });
+
+	TArray<UMovieSceneTrack*> AllTracks;
+	for (TViewModelPtr<ITrackExtension> TrackExtension : Sequencer->GetViewModel()->GetSelection()->Outliner.Filter<ITrackExtension>())
+	{
+		UMovieSceneTrack* const Track = TrackExtension->GetTrack();
+		if (IsValid(Track))
+		{
+			AllTracks.Add(Track);
+		}
+	}
+	if (AllTracks.IsEmpty())
+	{
+		return;
+	}
+
+	MenuBuilder.BeginSection(TEXT("TrackDisplayOptions"), LOCTEXT("TrackNodeDisplayOptions", "Display Options"));
+	{
+		MenuBuilder.AddSubMenu(
+			LOCTEXT("SetColorTint", "Set Color Tint"),
+			LOCTEXT("SetColorTintTooltip", "Set color tint from the preferences for the selected sections or the track's sections"),
+			FNewMenuDelegate::CreateSP(SharedThis, &FOutlinerItemModelMixin::BuildSectionColorTintsContextMenu));
+
+		UStruct* const DisplayOptionsStruct = FMovieSceneTrackDisplayOptions::StaticStruct();
 
+		const FBoolProperty* const ShowVerticalFramesProperty = CastField<FBoolProperty>(DisplayOptionsStruct->FindPropertyByName(GET_MEMBER_NAME_CHECKED(FMovieSceneTrackDisplayOptions, bShowVerticalFrames)));
+		if (ShowVerticalFramesProperty)
+		{
+			AddDisplayOptionsPropertyMenuItem(MenuBuilder, ShowVerticalFramesProperty);
+		}
+	}
+	MenuBuilder.EndSection();
+}
+
+void FOutlinerItemModelMixin::BuildTrackOptionsMenu(FMenuBuilder& MenuBuilder)
+{
+	const TArray<UMovieSceneTrack*> AllTracks = GetSelectedTracks();
+	if (AllTracks.IsEmpty())
+	{
+		return;
+	}
+	
+	MenuBuilder.BeginSection(TEXT("GeneralTrackOptions"), LOCTEXT("TrackNodeGeneralOptions", "Track Options"));
+	{
+		UStruct* const EvalOptionsStruct = FMovieSceneTrackEvalOptions::StaticStruct();
+
+		const FBoolProperty* const NearestSectionProperty = CastField<FBoolProperty>(EvalOptionsStruct->FindPropertyByName(GET_MEMBER_NAME_CHECKED(FMovieSceneTrackEvalOptions, bEvalNearestSection)));
+		auto CanEvaluateNearest = [](const UMovieSceneTrack* const InTrack) { return InTrack->EvalOptions.bCanEvaluateNearestSection != 0; };
+		if (NearestSectionProperty && AllTracks.ContainsByPredicate(CanEvaluateNearest))
+		{
+			TFunction<bool(UMovieSceneTrack*)> Validator = CanEvaluateNearest;
+			AddEvalOptionsPropertyMenuItem(MenuBuilder, NearestSectionProperty, Validator);
+		}
+
+		const FBoolProperty* const PrerollProperty = CastField<FBoolProperty>(EvalOptionsStruct->FindPropertyByName(GET_MEMBER_NAME_CHECKED(FMovieSceneTrackEvalOptions, bEvaluateInPreroll)));
+		if (PrerollProperty)
+		{
+			AddEvalOptionsPropertyMenuItem(MenuBuilder, PrerollProperty);
+		}
+
+		const FBoolProperty* const PostrollProperty = CastField<FBoolProperty>(EvalOptionsStruct->FindPropertyByName(GET_MEMBER_NAME_CHECKED(FMovieSceneTrackEvalOptions, bEvaluateInPostroll)));
+		if (PostrollProperty)
+		{
+			AddEvalOptionsPropertyMenuItem(MenuBuilder, PostrollProperty);
+		}
+	}
+	MenuBuilder.EndSection();
+}
+
+void FOutlinerItemModelMixin::BuildSidebarMenu(FMenuBuilder& MenuBuilder)
+{
+	/*BuildTrackOptionsMenu(MenuBuilder);
+	
+	const TArray<UMovieSceneSection*> Sections = GetSelectedSections();
+	if (!Sections.IsEmpty())
+	{
+		BuildDisplayOptionsMenu(MenuBuilder);
+	}*/
+}
+
+TArray<UMovieSceneSection*> FOutlinerItemModelMixin::GetSelectedSections() const
+{
 	TArray<UMovieSceneSection*> Sections;
-	for (TViewModelPtr<FSectionModel> SectionModel : Selection->Outliner.Filter<FSectionModel>())
+
+	const TSharedPtr<ISequencer> Sequencer = GetEditor()->GetSequencer();
+	if (!Sequencer.IsValid())
+	{
+		return Sections;
+	}
+
+	const TSharedPtr<FSequencerSelection> Selection = Sequencer->GetViewModel()->GetSelection();
+	if (!Selection.IsValid())
 	{
-		if (UMovieSceneSection* Section = SectionModel->GetSection())
+		return Sections;
+	}
+
+	for (const TViewModelPtr<FSectionModel> SectionModel : Selection->Outliner.Filter<FSectionModel>())
+	{
+		if (UMovieSceneSection* const Section = SectionModel->GetSection())
 		{
 			Sections.Add(Section);
 		}
@@ -668,15 +762,52 @@ void FOutlinerItemModelMixin::BuildSectionColorTintsContextMenu(FMenuBuilder& Me
 
 	if (!Sections.Num())
 	{
-		for (TViewModelPtr<ITrackExtension> TrackExtension : Selection->Outliner.Filter<ITrackExtension>())
+		for (const TViewModelPtr<ITrackExtension> TrackExtension : Selection->Outliner.Filter<ITrackExtension>())
 		{
-			for (UMovieSceneSection* Section : TrackExtension->GetSections())
+			for (UMovieSceneSection* const Section : TrackExtension->GetSections())
 			{
 				Sections.Add(Section);
 			}
 		}
 	}
 
+	return Sections;
+}
+
+TArray<UMovieSceneTrack*> FOutlinerItemModelMixin::GetSelectedTracks() const
+{
+	TArray<UMovieSceneTrack*> AllTracks;
+
+	const TSharedPtr<FSequencerEditorViewModel> EditorViewModel = GetEditor();
+	if (!EditorViewModel.IsValid())
+	{
+		return AllTracks;
+	}
+
+	const TSharedPtr<FSequencerSelection> Selection = EditorViewModel->GetSelection();
+	if (!Selection.IsValid())
+	{
+		return AllTracks;
+	}
+
+	for (const TViewModelPtr<ITrackExtension> TrackExtension : Selection->Outliner.Filter<ITrackExtension>())
+	{
+		UMovieSceneTrack* const Track = TrackExtension->GetTrack();
+		if (IsValid(Track))
+		{
+			AllTracks.Add(Track);
+		}
+	}
+
+	return AllTracks;
+}
+
+void FOutlinerItemModelMixin::BuildSectionColorTintsContextMenu(FMenuBuilder& MenuBuilder)
+{
+	TSharedPtr<FSequencer> Sequencer = GetEditor()->GetSequencerImpl();
+	TSharedPtr<FSequencerSelection> Selection = Sequencer->GetViewModel()->GetSelection();
+
+	TArray<UMovieSceneSection*> Sections = GetSelectedSections();
 	if (!Sections.Num())
 	{
 		return;
diff --git a/Engine/Source/Editor/Sequencer/Private/MVVM/ViewModels/TrackModel.cpp b/Engine/Source/Editor/Sequencer/Private/MVVM/ViewModels/TrackModel.cpp
--- a/Engine/Source/Editor/Sequencer/Private/MVVM/ViewModels/TrackModel.cpp
+++ b/Engine/Source/Editor/Sequencer/Private/MVVM/ViewModels/TrackModel.cpp
@@ -678,23 +678,7 @@ void FTrackModel::BuildContextMenu(FMenuBuilder& MenuBuilder)
 	}
 
 	// Add menu items for selecting a blender
-	IMovieSceneBlenderSystemSupport* BlenderSystemSupport = Cast<IMovieSceneBlenderSystemSupport>(Track);
-	if (BlenderSystemSupport)
-	{
-		TArray<TSubclassOf<UMovieSceneBlenderSystem>> BlenderTypes;
-		BlenderSystemSupport->GetSupportedBlenderSystems(BlenderTypes);
-
-		if (BlenderTypes.Num() > 1)
-		{
-			MenuBuilder.AddSubMenu(
-				LOCTEXT("BlendingAlgorithmSubMenu", "Blending Algorithm"),
-				FText(),
-				FNewMenuDelegate::CreateLambda([=](FMenuBuilder& SubMenuBuilder){
-					FSequencerUtilities::PopulateMenu_BlenderSubMenu(SubMenuBuilder, Track, WeakSequencer);
-				})
-			);
-		}
-	}
+	BuildBlendingMenu(MenuBuilder);
 
 	// Find sections in the track to add batch properties for
 	TArray<TWeakObjectPtr<UObject>> TrackSections;
@@ -740,6 +724,83 @@ void FTrackModel::BuildContextMenu(FMenuBuilder& MenuBuilder)
 	FOutlinerItemModel::BuildContextMenu(MenuBuilder);
 }
 
+void FTrackModel::BuildSidebarMenu(FMenuBuilder& MenuBuilder)
+{
+	UMovieSceneTrack* const Track = GetTrack();
+	if (!IsValid(Track))
+	{
+		return;
+	}
+
+	const TSharedPtr<FSequencerEditorViewModel> EditorViewModel = GetEditor();
+	if (!EditorViewModel.IsValid())
+	{
+		return;
+	}
+
+	const TSharedPtr<ISequencer> Sequencer = EditorViewModel->GetSequencer();
+	if (!Sequencer.IsValid())
+	{
+		return;
+	}
+
+	const IMovieSceneBlenderSystemSupport* const BlenderSystemSupport = Cast<IMovieSceneBlenderSystemSupport>(Track);
+	if (BlenderSystemSupport)
+	{
+		FSequencerUtilities::PopulateMenu_BlenderSubMenu(MenuBuilder, Track, Sequencer);
+	}
+
+	// Find sections in the track to add batch properties for
+	const TArray<TWeakObjectPtr<UObject>> TrackSections = GetSelectedTrackSections();
+	if (TrackSections.Num())
+	{
+		MenuBuilder.BeginSection(TEXT("TrackSection"), LOCTEXT("TrackSectionMenuSection", "Track Section"));
+		
+		FSequencer* const SequencerRaw = static_cast<FSequencer*>(Sequencer.Get());
+		SequencerHelpers::AddPropertiesMenu(*SequencerRaw, MenuBuilder, TrackSections);
+
+		MenuBuilder.EndSection();
+	}
+
+	FOutlinerItemModelMixin::BuildSidebarMenu(MenuBuilder);
+}
+
+void FTrackModel::BuildBlendingMenu(FMenuBuilder& MenuBuilder)
+{
+	UMovieSceneTrack* const Track = GetTrack();
+	if (!IsValid(Track))
+	{
+		return;
+	}
+
+	const IMovieSceneBlenderSystemSupport* const BlenderSystemSupport = Cast<IMovieSceneBlenderSystemSupport>(Track);
+	if (!BlenderSystemSupport)
+	{
+		return;
+	}
+
+	TArray<TSubclassOf<UMovieSceneBlenderSystem>> BlenderTypes;
+	BlenderSystemSupport->GetSupportedBlenderSystems(BlenderTypes);
+	if (BlenderTypes.Num() == 0)
+	{
+		return;
+	}
+	
+	const TSharedPtr<ISequencer> Sequencer = GetEditor()->GetSequencer();
+	if (!Sequencer)
+	{
+		return;
+	}
+	
+	MenuBuilder.AddSubMenu(
+		LOCTEXT("BlendingAlgorithmSubMenu", "Blending Algorithm"),
+		FText(),
+		FNewMenuDelegate::CreateLambda([=](FMenuBuilder& SubMenuBuilder)
+			{
+				FSequencerUtilities::PopulateMenu_BlenderSubMenu(SubMenuBuilder, Track, Sequencer);
+			}));
+}
+
 bool FTrackModel::CanDelete(FText* OutErrorMessage) const
 {
 	return true;
@@ -813,6 +874,57 @@ bool FTrackModel::FindBoundObjects(TArray<UObject*>& OutBoundObjects) const
 	return true;
 }
 
+TArray<TWeakObjectPtr<UObject>> FTrackModel::GetSelectedTrackSections() const
+{
+	TArray<TWeakObjectPtr<UObject>> TrackSections;
+
+	const TSharedPtr<FSequencerEditorViewModel> EditorViewModel = GetEditor();
+	if (!EditorViewModel.IsValid())
+	{
+		return TrackSections;
+	}
+
+	const TSharedPtr<ISequencer> Sequencer = EditorViewModel->GetSequencer();
+	if (!Sequencer.IsValid())
+	{
+		return TrackSections;
+	}
+
+	const TSharedPtr<UE::Sequencer::FSequencerEditorViewModel> ViewModel = Sequencer->GetViewModel();
+	if (!ViewModel.IsValid())
+	{
+		return TrackSections;
+	}
+
+	const TSharedPtr<FSequencerSelection> Selection = ViewModel->GetSelection();
+	if (!Selection.IsValid())
+	{
+		return TrackSections;
+	}
+	
+	for (const TViewModelPtr<ITrackExtension> TrackExtension : Selection->Outliner.Filter<ITrackExtension>())
+	{
+		for (UMovieSceneSection* const Section : TrackExtension->GetSections())
+		{
+			TrackSections.Add(Section);
+		}
+	}
+
+	for (const TSharedPtr<FViewModel> TrackAreaModel : GetTrackAreaModelList())
+	{
+		constexpr bool bIncludeThis = true;
+		for (const TSharedPtr<FSectionModel> Section : TParentFirstChildIterator<FSectionModel>(TrackAreaModel, bIncludeThis))
+		{
+			if (UMovieSceneSection* const SectionObject = Section->GetSection())
+			{
+				TrackSections.AddUnique(SectionObject);
+			}
+		}
+	}
+
+	return TrackSections;
+}
+
 } // namespace UE::Sequencer
 
 #undef LOCTEXT_NAMESPACE
diff --git a/Engine/Source/Editor/Sequencer/Private/MVVM/ViewModels/TrackRowModel.cpp b/Engine/Source/Editor/Sequencer/Private/MVVM/ViewModels/TrackRowModel.cpp
--- a/Engine/Source/Editor/Sequencer/Private/MVVM/ViewModels/TrackRowModel.cpp
+++ b/Engine/Source/Editor/Sequencer/Private/MVVM/ViewModels/TrackRowModel.cpp
@@ -391,6 +391,47 @@ void FTrackRowModel::BuildContextMenu(FMenuBuilder& MenuBuilder)
 	FOutlinerItemModel::BuildContextMenu(MenuBuilder);
 }
 
+void FTrackRowModel::BuildSidebarMenu(FMenuBuilder& MenuBuilder)
+{
+	UMovieSceneTrack* const Track = GetTrack();
+	if (!IsValid(Track))
+	{
+		return;
+	}
+
+	const TSharedPtr<FSequencerEditorViewModel> EditorViewModel = GetEditor();
+	if (!EditorViewModel.IsValid())
+	{
+		return;
+	}
+
+	const TSharedPtr<ISequencer> Sequencer = EditorViewModel->GetSequencer();
+	if (!Sequencer.IsValid())
+	{
+		return;
+	}
+
+	IMovieSceneBlenderSystemSupport* const BlenderSystemSupport = Cast<IMovieSceneBlenderSystemSupport>(Track);
+	if (BlenderSystemSupport)
+	{
+		FSequencerUtilities::PopulateMenu_BlenderSubMenu(MenuBuilder, Track, Sequencer);
+	}
+
+	// Find sections in the track to add batch properties for
+	const TArray<TWeakObjectPtr<UObject>> TrackSections = GetSelectedTrackSections();
+	if (TrackSections.Num())
+	{
+		MenuBuilder.BeginSection(TEXT("TrackSection"), LOCTEXT("TrackSectionMenuSection", "Track Section"));
+		
+		FSequencer* const SequencerRaw = static_cast<FSequencer*>(GetEditor()->GetSequencer().Get());
+		SequencerHelpers::AddPropertiesMenu(*SequencerRaw, MenuBuilder, TrackSections);
+
+		MenuBuilder.EndSection();
+	}
+
+	FOutlinerItemModel::BuildSidebarMenu(MenuBuilder);
+}
+
 bool FTrackRowModel::CanDelete(FText* OutErrorMessage) const
 {
 	return true;
@@ -428,6 +469,57 @@ void FTrackRowModel::Delete()
 	Track->FixRowIndices();
 }
 
+TArray<TWeakObjectPtr<UObject>> FTrackRowModel::GetSelectedTrackSections() const
+{
+	TArray<TWeakObjectPtr<UObject>> TrackSections;
+
+	const TSharedPtr<FSequencerEditorViewModel> EditorViewModel = GetEditor();
+	if (!EditorViewModel.IsValid())
+	{
+		return TrackSections;
+	}
+
+	const TSharedPtr<ISequencer> Sequencer = EditorViewModel->GetSequencer();
+	if (!Sequencer.IsValid())
+	{
+		return TrackSections;
+	}
+
+	const TSharedPtr<UE::Sequencer::FSequencerEditorViewModel> ViewModel = Sequencer->GetViewModel();
+	if (!ViewModel.IsValid())
+	{
+		return TrackSections;
+	}
+
+	const TSharedPtr<FSequencerSelection> Selection = ViewModel->GetSelection();
+	if (!Selection.IsValid())
+	{
+		return TrackSections;
+	}
+
+	for (const TViewModelPtr<ITrackExtension> TrackExtension : Selection->Outliner.Filter<ITrackExtension>())
+	{
+		for (UMovieSceneSection* const Section : TrackExtension->GetSections())
+		{
+			TrackSections.Add(Section);
+		}
+	}
+
+	for (const TSharedPtr<FViewModel> TrackAreaModel : GetTrackAreaModelList())
+	{
+		constexpr bool bIncludeThis = true;
+		for (const TSharedPtr<FSectionModel> Section : TParentFirstChildIterator<FSectionModel>(TrackAreaModel, bIncludeThis))
+		{
+			if (UMovieSceneSection* const SectionObject = Section->GetSection())
+			{
+				TrackSections.AddUnique(SectionObject);
+			}
+		}
+	}
+
+	return TrackSections;
+}
+
 } // namespace UE::Sequencer
 
 #undef LOCTEXT_NAMESPACE
diff --git a/Engine/Source/Editor/Sequencer/Private/Menus/CurveChannelSectionSidebarExtension.cpp b/Engine/Source/Editor/Sequencer/Private/Menus/CurveChannelSectionSidebarExtension.cpp
new file mode 100644
--- /dev/null
+++ b/Engine/Source/Editor/Sequencer/Private/Menus/CurveChannelSectionSidebarExtension.cpp
@@ -0,0 +1,664 @@
+// Copyright Epic Games, Inc. All Rights Reserved.
+
+#include "CurveChannelSectionSidebarExtension.h"
+#include "Channels/MovieSceneChannelHandle.h"
+#include "Curves/RealCurve.h"
+#include "IKeyArea.h"
+#include "ISequencer.h"
+#include "ScopedTransaction.h"
+#include "SequencerSettings.h"
+#include "SSequencer.h"
+#include "GameFramework/Actor.h"
+#include "Styling/AppStyle.h"
+#include "Styling/CoreStyle.h"
+#include "UObject/StructOnScope.h"
+#include "Framework/MultiBox/MultiBoxBuilder.h"
+#include "Channels/MovieSceneChannelProxy.h"
+#include "Channels/MovieSceneFloatChannel.h"
+#include "Channels/MovieSceneDoubleChannel.h"
+#include "Channels/MovieSceneBoolChannel.h"
+#include "Channels/MovieSceneIntegerChannel.h"
+#include "Curves/RealCurve.h"
+#include "Modules/ModuleManager.h"
+#include "Framework/Application/MenuStack.h"
+#include "Framework/Application/SlateApplication.h"
+#include "Framework/Commands/UIAction.h"
+#include "Widgets/Input/SComboButton.h"
+#include "Widgets/Input/SSpinBox.h"
+#include "Widgets/Layout/SSpacer.h"
+#include "Widgets/SNullWidget.h"
+
+#define LOCTEXT_NAMESPACE "CurveChannelSectionSidebarExtension"
+
+FCurveChannelSectionSidebarExtension::FCurveChannelSectionSidebarExtension(const TWeakPtr<ISequencer>& InSequencerWeak)
+	: SequencerWeak(InSequencerWeak)
+{
+}
+
+void FCurveChannelSectionSidebarExtension::AddSections(TArrayView<UMovieSceneSection* const> InSections)
+{
+	for (UMovieSceneSection* const Section : InSections)
+	{
+		Sections.Add(Section);
+	}
+}
+
+void FCurveChannelSectionSidebarExtension::ExtendMenu(FMenuBuilder& MenuBuilder)
+{
+	AddExtrapolationMenu(MenuBuilder, true);
+	AddExtrapolationMenu(MenuBuilder, false);
+	AddDisplayOptionsMenu(MenuBuilder);
+}
+
+void FCurveChannelSectionSidebarExtension::AddDisplayOptionsMenu(FMenuBuilder& MenuBuilder)
+{
+	const TSharedPtr<ISequencer> Sequencer = SequencerWeak.Pin();
+	if (!Sequencer)
+	{
+		return;
+	}
+	
+	MenuBuilder.BeginSection(TEXT("DisplayOptions"), LOCTEXT("DisplayOptionsTooltip", "Display Options"));
+
+	MenuBuilder.AddMenuEntry(
+		LOCTEXT("ToggleShowCurve", "Show Curve"),
+		LOCTEXT("ToggleShowCurveTooltip", "Toggle showing the curve in the track area"),
+		FSlateIcon(),
+		FUIAction(
+			FExecuteAction::CreateSP(this, &FCurveChannelSectionSidebarExtension::ToggleShowCurve),
+			FCanExecuteAction(),
+			FGetActionCheckState::CreateSP(this, &FCurveChannelSectionSidebarExtension::IsShowCurve)
+		),
+		NAME_None,
+		EUserInterfaceActionType::ToggleButton
+	);	
+	
+	FString KeyAreaName;
+	TArray<const IKeyArea*> SelectedKeyAreas;
+	Sequencer->GetSelectedKeyAreas(SelectedKeyAreas);
+	for (const IKeyArea* KeyArea : SelectedKeyAreas)
+	{			
+		if (KeyArea)
+		{
+			KeyAreaName = KeyArea->GetName().ToString();
+			break;
+		}
+	}
+
+	MenuBuilder.AddMenuEntry(
+		LOCTEXT("ToggleKeyAreaCurveNormalized", "Key Area Curve Normalized"),
+		LOCTEXT("ToggleKeyAreaCurveNormalizedTooltip", "Toggle showing the curve in the track area as normalized"),
+		FSlateIcon(),
+		FUIAction(
+			FExecuteAction::CreateSP(this, &FCurveChannelSectionSidebarExtension::OnKeyAreaCurveNormalized, KeyAreaName),
+			FCanExecuteAction::CreateSP(this, &FCurveChannelSectionSidebarExtension::IsAnyShowCurve),
+			FIsActionChecked::CreateSP(this, &FCurveChannelSectionSidebarExtension::GetKeyAreaCurveNormalized, KeyAreaName)
+		),
+		NAME_None,
+		EUserInterfaceActionType::ToggleButton);	
+
+	MenuBuilder.AddWidget(
+		SNew(SHorizontalBox)
+		+ SHorizontalBox::Slot()
+		[
+			SNew(SSpacer)
+		]
+		+ SHorizontalBox::Slot()
+		.AutoWidth()
+		[
+			SNew(SBox)
+			.WidthOverride(50.f)
+			.IsEnabled_Lambda([this, KeyAreaName]()
+				{
+					if (const TSharedPtr<ISequencer> Sequencer = SequencerWeak.Pin())
+					{
+						return IsAnyShowCurve() && Sequencer->GetSequencerSettings()->HasKeyAreaCurveExtents(KeyAreaName);
+					}
+					return false;
+				})
+			[
+				SNew(SSpinBox<double>)
+				.Style(&FAppStyle::GetWidgetStyle<FSpinBoxStyle>(TEXT("Sequencer.HyperlinkSpinBox")))
+				.Value(this, &FCurveChannelSectionSidebarExtension::GetKeyAreaCurveMin, KeyAreaName)
+				.OnValueChanged(this, &FCurveChannelSectionSidebarExtension::OnKeyAreaCurveMinChanged, KeyAreaName)
+				.OnValueCommitted_Lambda([this, KeyAreaName](const double InNewValue, const ETextCommit::Type InCommitType)
+					{
+						OnKeyAreaCurveMinChanged(InNewValue, KeyAreaName);
+					})
+			]
+		]
+		+ SHorizontalBox::Slot()
+		.AutoWidth()
+		[
+			SNew(SBox)
+			.WidthOverride(50.f)
+			.IsEnabled_Lambda([this, KeyAreaName]()
+				{
+					if (const TSharedPtr<ISequencer> Sequencer = SequencerWeak.Pin())
+					{
+						return IsAnyShowCurve() && Sequencer->GetSequencerSettings()->HasKeyAreaCurveExtents(KeyAreaName);
+					}
+					return false;
+				})
+			[
+				SNew(SSpinBox<double>)
+				.Style(&FAppStyle::GetWidgetStyle<FSpinBoxStyle>(TEXT("Sequencer.HyperlinkSpinBox")))
+				.Value(this, &FCurveChannelSectionSidebarExtension::GetKeyAreaCurveMax, KeyAreaName)
+				.OnValueChanged(this, &FCurveChannelSectionSidebarExtension::OnKeyAreaCurveMaxChanged, KeyAreaName)
+				.OnValueCommitted_Lambda([this, KeyAreaName](const double InNewValue, const ETextCommit::Type InCommitType)
+					{
+						OnKeyAreaCurveMaxChanged(InNewValue, KeyAreaName);
+					})
+			]
+		],
+		LOCTEXT("KeyAreaCurveRangeText", "Key Area Curve Range")
+	);		
+	
+	MenuBuilder.AddWidget(
+		SNew(SHorizontalBox)
+		+ SHorizontalBox::Slot()
+		[
+			SNew(SSpacer)
+		]
+		+ SHorizontalBox::Slot()
+		.AutoWidth()
+		[
+			SNew(SBox)
+			.WidthOverride(50.f)
+			[
+				SNew(SSpinBox<int32>)
+				.Style(&FAppStyle::GetWidgetStyle<FSpinBoxStyle>(TEXT("Sequencer.HyperlinkSpinBox")))
+				.MinValue(15)
+				.MaxValue(300)
+				.Value(this, &FCurveChannelSectionSidebarExtension::GetKeyAreaHeight)
+				.OnValueChanged(this, &FCurveChannelSectionSidebarExtension::OnKeyAreaHeightChanged)
+				.OnValueCommitted_Lambda([this](const int32 InValue, const ETextCommit::Type InCommitType)
+					{
+						OnKeyAreaHeightChanged(InValue);
+					})
+			]
+		],
+		LOCTEXT("KeyAreaHeightText", "Key Area Height")
+	);
+
+	MenuBuilder.EndSection();
+}
+
+void FCurveChannelSectionSidebarExtension::AddExtrapolationMenu(FMenuBuilder& MenuBuilder, const bool bInPreInfinity)
+{
+	auto CreateUIAction = [this, bInPreInfinity](const ERichCurveExtrapolation InExtrapolation)
+	{
+		return FUIAction(
+				FExecuteAction::CreateSP(this, &FCurveChannelSectionSidebarExtension::SetExtrapolationMode, InExtrapolation, bInPreInfinity),
+				FCanExecuteAction(),
+				FIsActionChecked::CreateSP(this, &FCurveChannelSectionSidebarExtension::IsExtrapolationModeSelected, InExtrapolation, bInPreInfinity)
+			);
+	};
+
+	if (bInPreInfinity)
+	{
+		MenuBuilder.BeginSection(TEXT("PreInfinityExtrapolation"), LOCTEXT("SetPreInfinityExtrapolation", "Pre-Infinity"));
+	}
+	else
+	{
+		MenuBuilder.BeginSection(TEXT("PostInfinityExtrapolation"), LOCTEXT("SetPostInfinityExtrapolation", "Post-Infinity"));
+	}
+
+	MenuBuilder.AddMenuEntry(
+		LOCTEXT("SetExtrapConstant", "Constant"),
+		LOCTEXT("SetExtrapConstantTooltip", "Set extrapolation constant"),
+		FSlateIcon(FAppStyle::GetAppStyleSetName(), TEXT("GenericCurveEditor.SetPreInfinityExtrapConstant")),
+		CreateUIAction(RCCE_Constant),
+		NAME_None,
+		EUserInterfaceActionType::RadioButton);
+
+	MenuBuilder.AddMenuEntry(
+		LOCTEXT("SetExtrapCycle", "Cycle"),
+		LOCTEXT("SetExtrapCycleTooltip", "Set extrapolation cycle"),
+		FSlateIcon(FAppStyle::GetAppStyleSetName(), TEXT("GenericCurveEditor.SetPreInfinityExtrapCycle")),
+		CreateUIAction(RCCE_Cycle),
+		NAME_None,
+		EUserInterfaceActionType::RadioButton);
+
+	MenuBuilder.AddMenuEntry(
+		LOCTEXT("SetExtrapCycleWithOffset", "Cycle with Offset"),
+		LOCTEXT("SetExtrapCycleWithOffsetTooltip", "Set extrapolation cycle with offset"),
+		FSlateIcon(FAppStyle::GetAppStyleSetName(), TEXT("GenericCurveEditor.SetPreInfinityExtrapCycleWithOffset")),
+		CreateUIAction(RCCE_CycleWithOffset),
+		NAME_None,
+		EUserInterfaceActionType::RadioButton);
+
+	MenuBuilder.AddMenuEntry(
+		LOCTEXT("SetExtrapLinear", "Linear"),
+		LOCTEXT("SetExtrapLinearTooltip", "Set extrapolation linear"),
+		FSlateIcon(FAppStyle::GetAppStyleSetName(), TEXT("GenericCurveEditor.SetPreInfinityExtrapLinear")),
+		CreateUIAction(RCCE_Linear),
+		NAME_None,
+		EUserInterfaceActionType::RadioButton);
+
+	MenuBuilder.AddMenuEntry(
+		LOCTEXT("SetExtrapOscillate", "Oscillate"),
+		LOCTEXT("SetExtrapOscillateTooltip", "Set extrapolation oscillate"),
+		FSlateIcon(FAppStyle::GetAppStyleSetName(), TEXT("GenericCurveEditor.SetPreInfinityExtrapOscillate")),
+		CreateUIAction(RCCE_Oscillate),
+		NAME_None,
+		EUserInterfaceActionType::RadioButton);
+
+	MenuBuilder.EndSection();
+}
+
+void FCurveChannelSectionSidebarExtension::GetChannels(TArray<FMovieSceneFloatChannel*>& FloatChannels, TArray<FMovieSceneDoubleChannel*>& DoubleChannels,
+	TArray<FMovieSceneIntegerChannel*>& IntegerChannels, TArray<FMovieSceneBoolChannel*>& BoolChannels) const
+{
+	ISequencer* const Sequencer = SequencerWeak.Pin().Get();
+	if (!Sequencer)
+	{
+		return;
+	}
+
+	// Get selected channels
+	TArray<const IKeyArea*> KeyAreas;
+	Sequencer->GetSelectedKeyAreas(KeyAreas);
+	for (const IKeyArea* const KeyArea : KeyAreas)
+	{
+		FMovieSceneChannelHandle Handle = KeyArea->GetChannel();
+		if (Handle.GetChannelTypeName() == FMovieSceneFloatChannel::StaticStruct()->GetFName())
+		{
+			FMovieSceneFloatChannel* const Channel = static_cast<FMovieSceneFloatChannel*>(Handle.Get());
+			FloatChannels.Add(Channel);
+		}
+		else if (Handle.GetChannelTypeName() == FMovieSceneDoubleChannel::StaticStruct()->GetFName())
+		{
+			FMovieSceneDoubleChannel* const Channel = static_cast<FMovieSceneDoubleChannel*>(Handle.Get());
+			DoubleChannels.Add(Channel);
+		}
+		else if (Handle.GetChannelTypeName() == FMovieSceneIntegerChannel::StaticStruct()->GetFName())
+		{
+			FMovieSceneIntegerChannel* const Channel = static_cast<FMovieSceneIntegerChannel*>(Handle.Get());
+			IntegerChannels.Add(Channel);
+		}
+		else if (Handle.GetChannelTypeName() == FMovieSceneBoolChannel::StaticStruct()->GetFName())
+		{
+			FMovieSceneBoolChannel* const Channel = static_cast<FMovieSceneBoolChannel*>(Handle.Get());
+			BoolChannels.Add(Channel);
+		}
+	}
+
+	// Otherwise, the channels of all the sections
+	if (FloatChannels.Num() + DoubleChannels.Num() + IntegerChannels.Num() + BoolChannels.Num() == 0)
+	{
+		for (const TWeakObjectPtr<UMovieSceneSection>& WeakSection : Sections)
+		{
+			if (const UMovieSceneSection* const Section = WeakSection.Get())
+			{
+				FMovieSceneChannelProxy& ChannelProxy = Section->GetChannelProxy();
+				for (FMovieSceneFloatChannel* const Channel : ChannelProxy.GetChannels<FMovieSceneFloatChannel>())
+				{
+					FloatChannels.Add(Channel);
+				}
+				for (FMovieSceneDoubleChannel* const Channel : ChannelProxy.GetChannels<FMovieSceneDoubleChannel>())
+				{
+					DoubleChannels.Add(Channel);
+				}
+				for (FMovieSceneIntegerChannel* const Channel : ChannelProxy.GetChannels<FMovieSceneIntegerChannel>())
+				{
+					IntegerChannels.Add(Channel);
+				}
+				for (FMovieSceneBoolChannel* const Channel : ChannelProxy.GetChannels<FMovieSceneBoolChannel>())
+				{
+					BoolChannels.Add(Channel);
+				}
+			}
+		}
+	}
+}
+
+void FCurveChannelSectionSidebarExtension::SetExtrapolationMode(const ERichCurveExtrapolation InExtrapolation, const bool bInPreInfinity)
+{
+	TArray<FMovieSceneFloatChannel*> FloatChannels;
+	TArray<FMovieSceneDoubleChannel*> DoubleChannels;
+	TArray<FMovieSceneIntegerChannel*> IntegerChannels;
+	TArray<FMovieSceneBoolChannel*> BoolChannels;
+	
+	GetChannels(FloatChannels, DoubleChannels,IntegerChannels,BoolChannels);
+
+	if (FloatChannels.Num() + DoubleChannels.Num() + IntegerChannels.Num() + BoolChannels.Num() == 0)
+	{
+		return;
+	}
+
+	FScopedTransaction Transaction(LOCTEXT("SetExtrapolationMode_Transaction", "Set Extrapolation Mode"));
+
+	bool bAnythingChanged = false;
+
+	// Modify all sections
+	for (TWeakObjectPtr<UMovieSceneSection> WeakSection : Sections)
+	{
+		UMovieSceneSection* const Section = WeakSection.Get();
+		if (IsValid(Section))
+		{
+			Section->Modify();
+		}
+	}
+
+	// Apply to all channels
+	for (FMovieSceneFloatChannel* const Channel : FloatChannels)
+	{
+		TEnumAsByte<ERichCurveExtrapolation>& DestExtrap = bInPreInfinity ? Channel->PreInfinityExtrap : Channel->PostInfinityExtrap;
+		DestExtrap = InExtrapolation;
+		bAnythingChanged = true;
+	}
+	for (FMovieSceneDoubleChannel* const Channel : DoubleChannels)
+	{
+		TEnumAsByte<ERichCurveExtrapolation>& DestExtrap = bInPreInfinity ? Channel->PreInfinityExtrap : Channel->PostInfinityExtrap;
+		DestExtrap = InExtrapolation;
+		bAnythingChanged = true;
+	}
+	for (FMovieSceneIntegerChannel* const Channel : IntegerChannels)
+	{
+		TEnumAsByte<ERichCurveExtrapolation>& DestExtrap = bInPreInfinity ? Channel->PreInfinityExtrap : Channel->PostInfinityExtrap;
+		DestExtrap = InExtrapolation;
+		bAnythingChanged = true;
+	}
+	for (FMovieSceneBoolChannel* const Channel : BoolChannels)
+	{
+		TEnumAsByte<ERichCurveExtrapolation>& DestExtrap = bInPreInfinity ? Channel->PreInfinityExtrap : Channel->PostInfinityExtrap;
+		DestExtrap = InExtrapolation;
+		bAnythingChanged = true;
+	}
+
+	if (bAnythingChanged)
+	{
+		if (ISequencer* const Sequencer = SequencerWeak.Pin().Get())
+		{
+			Sequencer->NotifyMovieSceneDataChanged(EMovieSceneDataChangeType::TrackValueChanged);
+		}
+	}
+	else
+	{
+		Transaction.Cancel();
+	}
+}
+
+bool FCurveChannelSectionSidebarExtension::IsExtrapolationModeSelected(const ERichCurveExtrapolation InExtrapolation, const bool bInPreInfinity) const
+{
+	TArray<FMovieSceneFloatChannel*> FloatChannels;
+	TArray<FMovieSceneDoubleChannel*> DoubleChannels;
+	TArray<FMovieSceneIntegerChannel*> IntegerChannels;
+	TArray<FMovieSceneBoolChannel*> BoolChannels;
+	
+	GetChannels(FloatChannels, DoubleChannels, IntegerChannels, BoolChannels);
+
+	for (FMovieSceneFloatChannel* const Channel : FloatChannels)
+	{
+		const ERichCurveExtrapolation SourceExtrapolation = bInPreInfinity ? Channel->PreInfinityExtrap : Channel->PostInfinityExtrap;
+		if (SourceExtrapolation != InExtrapolation)
+		{
+			return false;
+		}
+	}
+	for (FMovieSceneDoubleChannel* const Channel : DoubleChannels)
+	{
+		const ERichCurveExtrapolation SourceExtrapolation = bInPreInfinity ? Channel->PreInfinityExtrap : Channel->PostInfinityExtrap;
+		if (SourceExtrapolation != InExtrapolation)
+		{
+			return false;
+		}
+	}
+	for (FMovieSceneIntegerChannel* const Channel : IntegerChannels)
+	{
+		const ERichCurveExtrapolation SourceExtrapolation = bInPreInfinity ? Channel->PreInfinityExtrap : Channel->PostInfinityExtrap;
+		if (SourceExtrapolation != InExtrapolation)
+		{
+			return false;
+		}
+	}
+	for (FMovieSceneBoolChannel* const Channel : BoolChannels)
+	{
+		const ERichCurveExtrapolation SourceExtrapolation = bInPreInfinity ? Channel->PreInfinityExtrap : Channel->PostInfinityExtrap;
+		if (SourceExtrapolation != InExtrapolation)
+		{
+			return false;
+		}
+	}
+
+	return true;
+}
+
+void FCurveChannelSectionSidebarExtension::ToggleShowCurve()
+{
+	const ECheckBoxState CurrentState = IsShowCurve();
+	const bool bShowCurve = (CurrentState != ECheckBoxState::Checked); // If unchecked or mixed, check it
+
+	FScopedTransaction Transaction(LOCTEXT("ToggleShowCurve_Transaction", "Toggle Show Curve"));
+
+	bool bAnythingChanged = false;
+
+	// Modify all sections
+	for (const TWeakObjectPtr<UMovieSceneSection>& WeakSection : Sections)
+	{
+		UMovieSceneSection* const Section = WeakSection.Get();
+		if (IsValid(Section))
+		{
+			Section->Modify();
+		}
+	}
+
+	// Apply to all channels
+	for (const TWeakObjectPtr<UMovieSceneSection>& WeakSection : Sections)
+	{
+		if (const UMovieSceneSection* const Section = WeakSection.Get())
+		{
+			const FMovieSceneChannelProxy& ChannelProxy = Section->GetChannelProxy();
+
+			for (FMovieSceneFloatChannel* const Channel : ChannelProxy.GetChannels<FMovieSceneFloatChannel>())
+			{
+				if (Channel)
+				{
+					Channel->SetShowCurve(bShowCurve);
+					bAnythingChanged = true;
+				}
+			}
+
+			for (FMovieSceneDoubleChannel* const Channel : ChannelProxy.GetChannels<FMovieSceneDoubleChannel>())
+			{
+				if (Channel)
+				{
+					Channel->SetShowCurve(bShowCurve);
+					bAnythingChanged = true;
+				}
+			}
+		}
+	}
+
+	if (!bAnythingChanged)
+	{
+		Transaction.Cancel();
+	}
+}
+
+ECheckBoxState FCurveChannelSectionSidebarExtension::IsShowCurve() const
+{
+	int32 NumShowedAndHidden[2] = { 0, 0 };
+	for (const TWeakObjectPtr<UMovieSceneSection>& WeakSection : Sections)
+	{
+		if (const UMovieSceneSection* const Section = WeakSection.Get())
+		{
+			const FMovieSceneChannelProxy& ChannelProxy = Section->GetChannelProxy();
+
+			for (FMovieSceneFloatChannel* const Channel : ChannelProxy.GetChannels<FMovieSceneFloatChannel>())
+			{
+				if (Channel)
+				{
+					NumShowedAndHidden[Channel->GetShowCurve() ? 0 : 1]++;
+				}
+			}
+
+			for (FMovieSceneDoubleChannel* const Channel : ChannelProxy.GetChannels<FMovieSceneDoubleChannel>())
+			{
+				if (Channel)
+				{
+					NumShowedAndHidden[Channel->GetShowCurve() ? 0 : 1]++;
+				}
+			}
+		}
+	}
+
+	if (NumShowedAndHidden[0] == 0 && NumShowedAndHidden[1] > 0)  // No curve showed, some hidden
+	{
+		return ECheckBoxState::Unchecked;
+	}
+	else if (NumShowedAndHidden[0] > 0 && NumShowedAndHidden[1] == 0) // Some curves showed, none hidden
+	{
+		return ECheckBoxState::Checked;
+	}
+	return ECheckBoxState::Undetermined;  // Mixed states, or no curves
+}
+
+bool FCurveChannelSectionSidebarExtension::IsAnyShowCurve() const
+{
+	for (const TWeakObjectPtr<UMovieSceneSection>& WeakSection : Sections)
+	{
+		const UMovieSceneSection* const Section = WeakSection.Get();
+		if (IsValid(Section))
+		{
+			const FMovieSceneChannelProxy& ChannelProxy = Section->GetChannelProxy();
+
+			for (const FMovieSceneFloatChannel* const Channel : ChannelProxy.GetChannels<FMovieSceneFloatChannel>())
+			{
+				if (Channel && Channel->GetShowCurve())
+				{
+					return true;
+				}
+			}
+
+			for (const FMovieSceneDoubleChannel* const Channel : ChannelProxy.GetChannels<FMovieSceneDoubleChannel>())
+			{
+				if (Channel && Channel->GetShowCurve())
+				{
+					return true;
+				}
+			}
+		}
+	}
+	
+	return false;
+}
+
+int32 FCurveChannelSectionSidebarExtension::GetKeyAreaHeight() const
+{
+	if (const TSharedPtr<ISequencer> Sequencer = SequencerWeak.Pin())
+	{
+		return (int32)Sequencer->GetSequencerSettings()->GetKeyAreaHeightWithCurves();
+	}
+	return 0;
+}
+
+void FCurveChannelSectionSidebarExtension::OnKeyAreaHeightChanged(const int32 InNewValue)
+{
+	USequencerSettings* const SequencerSettings = GetSequencerSettings();
+	if (IsValid(SequencerSettings))
+	{
+		SequencerSettings->SetKeyAreaHeightWithCurves((float)InNewValue);
+	}
+}
+
+bool FCurveChannelSectionSidebarExtension::GetKeyAreaCurveNormalized(const FString InKeyAreaName) const
+{
+	USequencerSettings* const SequencerSettings = GetSequencerSettings();
+	if (!IsValid(SequencerSettings))
+	{
+		return !SequencerSettings->HasKeyAreaCurveExtents(InKeyAreaName);
+	}
+	return false;
+}
+
+void FCurveChannelSectionSidebarExtension::OnKeyAreaCurveNormalized(const FString InKeyAreaName) 
+{
+	USequencerSettings* const SequencerSettings = GetSequencerSettings();
+	if (!IsValid(SequencerSettings))
+	{
+		return;
+	}
+
+	if (SequencerSettings->HasKeyAreaCurveExtents(InKeyAreaName))
+	{
+		SequencerSettings->RemoveKeyAreaCurveExtents(InKeyAreaName);
+	}
+	else
+	{
+		// Initialize to some arbitrary value
+		SequencerSettings->SetKeyAreaCurveExtents(InKeyAreaName, 0.f, 6.f); 
+	}
+}
+
+double FCurveChannelSectionSidebarExtension::GetKeyAreaCurveMin(const FString InKeyAreaName) const
+{
+	USequencerSettings* const SequencerSettings = GetSequencerSettings();
+	if (!IsValid(SequencerSettings))
+	{
+		return 0.0;
+	}
+
+	double CurveMin = 0.f;
+	double CurveMax = 0.f;
+	SequencerSettings->GetKeyAreaCurveExtents(InKeyAreaName, CurveMin, CurveMax);
+
+	return CurveMin;
+}
+
+void FCurveChannelSectionSidebarExtension::OnKeyAreaCurveMinChanged(const double InNewValue, const FString InKeyAreaName)
+{
+	USequencerSettings* const SequencerSettings = GetSequencerSettings();
+	if (!IsValid(SequencerSettings))
+	{
+		return;
+	}
+
+	double CurveMin = 0.f;
+	double CurveMax = 0.f;
+	SequencerSettings->GetKeyAreaCurveExtents(InKeyAreaName, CurveMin, CurveMax);
+
+	SequencerSettings->SetKeyAreaCurveExtents(InKeyAreaName, InNewValue, CurveMax);
+}
+	
+double FCurveChannelSectionSidebarExtension::GetKeyAreaCurveMax(const FString InKeyAreaName) const
+{
+	USequencerSettings* const SequencerSettings = GetSequencerSettings();
+	if (!IsValid(SequencerSettings))
+	{
+		return 0.0;
+	}
+
+	double CurveMin = 0.f;
+	double CurveMax = 0.f;
+	SequencerSettings->GetKeyAreaCurveExtents(InKeyAreaName, CurveMin, CurveMax);
+
+	return CurveMax;
+}
+
+void FCurveChannelSectionSidebarExtension::OnKeyAreaCurveMaxChanged(const double InNewValue, const FString InKeyAreaName)
+{
+	USequencerSettings* const SequencerSettings = GetSequencerSettings();
+	if (!IsValid(SequencerSettings))
+	{
+		return;
+	}
+
+	double CurveMin = 0.f;
+	double CurveMax = 0.f;
+	SequencerSettings->GetKeyAreaCurveExtents(InKeyAreaName, CurveMin, CurveMax);
+
+	SequencerSettings->SetKeyAreaCurveExtents(InKeyAreaName, CurveMin, InNewValue);
+}
+
+USequencerSettings* FCurveChannelSectionSidebarExtension::GetSequencerSettings() const
+{
+	if (const TSharedPtr<ISequencer> Sequencer = SequencerWeak.Pin())
+	{
+		return Sequencer->GetSequencerSettings();
+	}
+	return nullptr;
+}
+
+#undef LOCTEXT_NAMESPACE
diff --git a/Engine/Source/Editor/Sequencer/Private/Menus/CurveChannelSectionSidebarExtension.h b/Engine/Source/Editor/Sequencer/Private/Menus/CurveChannelSectionSidebarExtension.h
new file mode 100644
--- /dev/null
+++ b/Engine/Source/Editor/Sequencer/Private/Menus/CurveChannelSectionSidebarExtension.h
@@ -0,0 +1,54 @@
+// Copyright Epic Games, Inc. All Rights Reserved.
+
+#pragma once
+
+#include "SequencerChannelTraits.h"
+#include "MovieSceneSection.h"
+
+struct FMovieSceneBoolChannel;
+struct FMovieSceneDoubleChannel;
+struct FMovieSceneFloatChannel;
+struct FMovieSceneIntegerChannel;
+enum class ECheckBoxState : uint8;
+enum ERichCurveExtrapolation : int;
+
+struct FCurveChannelSectionSidebarExtension : TSharedFromThis<FCurveChannelSectionSidebarExtension>
+{
+	FCurveChannelSectionSidebarExtension(const TWeakPtr<ISequencer>& InSequencerWeak);
+
+	void AddSections(TArrayView<UMovieSceneSection* const> InSections);
+
+	void ExtendMenu(FMenuBuilder& MenuBuilder);
+
+private:
+	void AddDisplayOptionsMenu(FMenuBuilder& MenuBuilder);
+	void AddExtrapolationMenu(FMenuBuilder& MenuBuilder, const bool bInPreInfinity);
+
+	void GetChannels(TArray<FMovieSceneFloatChannel*>& FloatChannels, TArray<FMovieSceneDoubleChannel*>& DoubleChannels,
+		TArray<FMovieSceneIntegerChannel*>& IntegerChannels, TArray<FMovieSceneBoolChannel*>& BoolChannels) const;
+
+	void SetExtrapolationMode(const ERichCurveExtrapolation InExtrapolation, const bool bInPreInfinity);
+	bool IsExtrapolationModeSelected(const ERichCurveExtrapolation InExtrapolation, const bool bInPreInfinity) const;
+
+	void ToggleShowCurve();
+	ECheckBoxState IsShowCurve() const;
+	bool IsAnyShowCurve() const;
+
+	int32 GetKeyAreaHeight() const;
+	void OnKeyAreaHeightChanged(const int32 InNewValue);
+	
+	bool GetKeyAreaCurveNormalized(const FString InKeyAreaName) const;
+	void OnKeyAreaCurveNormalized(const FString InKeyAreaName);
+
+	double GetKeyAreaCurveMin(const FString InKeyAreaName) const;
+	void OnKeyAreaCurveMinChanged(const double InNewValue, const FString InKeyAreaName);
+
+	double GetKeyAreaCurveMax(const FString InKeyAreaName) const;
+	void OnKeyAreaCurveMaxChanged(const double InNewValue, const FString InKeyAreaName);
+
+	USequencerSettings* GetSequencerSettings() const;
+
+	TWeakPtr<ISequencer> SequencerWeak;
+
+	TSet<TWeakObjectPtr<UMovieSceneSection>> Sections;
+};
diff --git a/Engine/Source/Editor/Sequencer/Private/MovieSceneTrackEditor.cpp b/Engine/Source/Editor/Sequencer/Private/MovieSceneTrackEditor.cpp
--- a/Engine/Source/Editor/Sequencer/Private/MovieSceneTrackEditor.cpp
+++ b/Engine/Source/Editor/Sequencer/Private/MovieSceneTrackEditor.cpp
@@ -327,7 +327,11 @@ TSharedPtr<SWidget> FMovieSceneTrackEditor::BuildOutlinerEditWidget(const FGuid&
 	}
 }
 
-void FMovieSceneTrackEditor::BuildTrackContextMenu( FMenuBuilder& MenuBuilder, UMovieSceneTrack* Track ) 
+void FMovieSceneTrackEditor::BuildTrackContextMenu(FMenuBuilder& MenuBuilder, UMovieSceneTrack* Track) 
+{ 
+}
+
+void FMovieSceneTrackEditor::BuildTrackSidebarMenu(FMenuBuilder& MenuBuilder, UMovieSceneTrack* Track) 
 { 
 }
 
diff --git a/Engine/Source/Editor/Sequencer/Private/SSequencer.cpp b/Engine/Source/Editor/Sequencer/Private/SSequencer.cpp
--- a/Engine/Source/Editor/Sequencer/Private/SSequencer.cpp
+++ b/Engine/Source/Editor/Sequencer/Private/SSequencer.cpp
@@ -42,7 +42,6 @@
 #include "Widgets/Layout/SScrollBorder.h"
 #include "Widgets/Input/SComboButton.h"
 #include "Widgets/Layout/SSeparator.h"
-#include "Widgets/Layout/SSplitter.h"
 #include "Widgets/Input/SSpinBox.h"
 #include "Widgets/Input/SCheckBox.h"
 #include "Widgets/Docking/SDockTab.h"
@@ -126,6 +125,9 @@
 #include "UnrealEdGlobals.h"
 #include "UniversalObjectLocators/ActorLocatorFragment.h"
 #include "SequencerUtilities.h"
+#include "Sidebar/SidebarDrawerConfig.h"
+#include "Sidebar/SSidebar.h"
+#include "Widgets/SOverlay.h"
 
 #define LOCTEXT_NAMESPACE "Sequencer"
 
@@ -287,8 +289,14 @@ void SSequencer::Construct(const FArguments& InArgs, TSharedRef<FSequencer> InSe
 
 	if (GetSequencerSettings())
 	{
-		ColumnFillCoefficients[0] = GetSequencerSettings()->GetTreeViewWidth();
-		ColumnFillCoefficients[1] = 1.f - GetSequencerSettings()->GetTreeViewWidth();
+		const float TreeViewWidth = GetSequencerSettings()->GetTreeViewWidth();
+		const float TimelineWidth = 1.f - TreeViewWidth;
+
+		if (TreeViewWidth > 0.f && TimelineWidth > 0.f)
+		{
+			ColumnFillCoefficients[0] = TreeViewWidth;
+			ColumnFillCoefficients[1] = TimelineWidth;
+		}
 	}
 
 	TAttribute<float> FillCoefficient_0, FillCoefficient_1;
@@ -341,6 +349,10 @@ void SSequencer::Construct(const FArguments& InArgs, TSharedRef<FSequencer> InSe
 
 	PlayTimeDisplay = StaticCastSharedRef<STemporarilyFocusedSpinBox<double>>(SequencerPtr.Pin()->MakePlayTimeDisplay(GetNumericTypeInterface()));
 
+	const TSharedRef<SOverlay> SidebarOverlay = SNew(SOverlay);
+
+	DetailsDockLocation = SNew(SBox);
+
 	TAttribute<FAnimatedRange> ViewRangeAttribute = InArgs._ViewRange;
 
 	const int32 Column0 = 0, Column1 = 1;
@@ -517,228 +529,212 @@ void SSequencer::Construct(const FArguments& InArgs, TSharedRef<FSequencer> InSe
 
 		+ SVerticalBox::Slot()
 		[
-			SNew(SSplitter)
-			.Orientation(Orient_Horizontal)
-			
-			+ SSplitter::Slot()
-			.Value(0.9f)
+			// Overlay for sequencer details sliding drawers
+			SNew(SOverlay)
+
+			+ SOverlay::Slot()
 			[
-				SNew(SOverlay)
+				SNew(SHorizontalBox)
 
-				+ SOverlay::Slot()
+				+ SHorizontalBox::Slot()
 				[
-					// track area grid panel
-					SNew( SGridPanel )
-					.FillRow( 2, 1.f )
-					.FillColumn( 0, FillCoefficient_0 )
-					.FillColumn( 1, FillCoefficient_1 )
+					SNew(SSplitter)
+					.Orientation(Orient_Horizontal)
 
-					+ SGridPanel::Slot( Column0, Row1 )
+					+ SSplitter::Slot()
+					.Value(1 - FSidebarDrawerConfig::DefaultSizeCoefficient)
 					[
-						SNew(SBorder)
-						.BorderImage(FAppStyle::GetBrush("ToolPanel.GroupBorder"))
-						[
-							SNew(SSpacer)
-						]
-					]
+						SNew(SOverlay)
 
-					// outliner search box
-					+ SGridPanel::Slot( Column0, Row1, SGridPanel::Layer(10) )
-					[
-						SNew(SBorder)
-						.BorderImage(FAppStyle::GetBrush("ToolPanel.GroupBorder"))
-						.Padding(FMargin(CommonPadding*2, CommonPadding))
-						.Clipping(EWidgetClipping::ClipToBounds)
+						+ SOverlay::Slot()
 						[
-							SNew(SHorizontalBox)
-							+ SHorizontalBox::Slot()
-							.AutoWidth()
-							.VAlign(VAlign_Center)
-							.Padding(FMargin(0.f, 0.f, CommonPadding, 0.f))
-							[
-								MakeAddButton()
-							]
-						
-							+ SHorizontalBox::Slot()
-							.VAlign(VAlign_Center)
-							[
-								SAssignNew(SearchBox, SSearchBox)
-								.HintText(LOCTEXT("SearchNodesHint", "Search Tracks"))
-								.OnTextChanged( this, &SSequencer::OnOutlinerSearchChanged )
-							]
+							// track area grid panel
+							SNew( SGridPanel )
+							.FillRow( 2, 1.f )
+							.FillColumn( 0, FillCoefficient_0 )
+							.FillColumn( 1, FillCoefficient_1 )
 
-							+ SHorizontalBox::Slot()
-							.AutoWidth()
-							.VAlign(VAlign_Center)
-							.Padding(FMargin(0.f, 0.f, CommonPadding, 0.f))
+							+ SGridPanel::Slot( Column0, Row1 )
 							[
-								MakeFilterButton()
+								SNew(SBorder)
+								.BorderImage(FAppStyle::GetBrush("ToolPanel.GroupBorder"))
+								[
+									SNew(SSpacer)
+								]
 							]
-						]
-					]
-
-					// main sequencer area
-					+ SGridPanel::Slot( Column0, Row2, SGridPanel::Layer(10) )
-					.ColumnSpan(2)
-					[
-						SAssignNew(MainSequencerArea, SVerticalBox)
 
-						+SVerticalBox::Slot()
-						.AutoHeight()
-						.MaxHeight(TAttribute<float>::Create(TAttribute<float>::FGetter::CreateSP(this, &SSequencer::GetPinnedAreaMaxHeight)))
-						.Padding(FMargin(0.f, 0.f, 0.f, CommonPadding))
-						[
-							SNew(SOverlay)
-							.Visibility(this, &SSequencer::GetPinnedAreaVisibility)
-
-							+ SOverlay::Slot()
+							// outliner search box
+							+ SGridPanel::Slot( Column0, Row1, SGridPanel::Layer(10) )
 							[
-								SNew(SScrollBorder, PinnedTreeView.ToSharedRef())
+								SNew(SBorder)
+								.BorderImage(FAppStyle::GetBrush("ToolPanel.GroupBorder"))
+								.Padding(FMargin(CommonPadding*2, CommonPadding))
+								.Clipping(EWidgetClipping::ClipToBounds)
 								[
 									SNew(SHorizontalBox)
-
-									// outliner tree
 									+ SHorizontalBox::Slot()
-									.FillWidth( FillCoefficient_0 )
+									.AutoWidth()
+									.VAlign(VAlign_Center)
+									.Padding(FMargin(0.f, 0.f, CommonPadding, 0.f))
+									[
+										MakeAddButton()
+									]
+						
+									+ SHorizontalBox::Slot()
+									.VAlign(VAlign_Center)
 									[
-										PinnedTreeView.ToSharedRef()
+										SAssignNew(SearchBox, SSearchBox)
+										.HintText(LOCTEXT("SearchNodesHint", "Search Tracks"))
+										.OnTextChanged( this, &SSequencer::OnOutlinerSearchChanged )
 									]
 
-									// track area
 									+ SHorizontalBox::Slot()
-									.FillWidth( FillCoefficient_1 )
+									.AutoWidth()
+									.VAlign(VAlign_Center)
+									.Padding(FMargin(0.f, 0.f, CommonPadding, 0.f))
 									[
-										SNew(SBox)
-										.Padding(ResizeBarPadding)
-										.Clipping(EWidgetClipping::ClipToBounds)
-										[
-											PinnedTrackArea.ToSharedRef()
-										]
+										MakeFilterButton()
 									]
 								]
 							]
 
-							+ SOverlay::Slot()
-							.HAlign(HAlign_Right)
+							// main sequencer area
+							+ SGridPanel::Slot( Column0, Row2, SGridPanel::Layer(10) )
+							.ColumnSpan(2)
 							[
-								PinnedAreaScrollBar
-							]
-						]
-
-						+SVerticalBox::Slot()
-						[
-							SNew(SOverlay)
+								SAssignNew(MainSequencerArea, SVerticalBox)
 
-							+ SOverlay::Slot()
-							[
-								SNew(SScrollBorder, TreeView.ToSharedRef())
+								+SVerticalBox::Slot()
+								.AutoHeight()
+								.MaxHeight(TAttribute<float>::Create(TAttribute<float>::FGetter::CreateSP(this, &SSequencer::GetPinnedAreaMaxHeight)))
+								.Padding(FMargin(0.f, 0.f, 0.f, CommonPadding))
 								[
-									SNew(SHorizontalBox)
+									SNew(SOverlay)
+									.Visibility(this, &SSequencer::GetPinnedAreaVisibility)
 
-									// outliner tree
-									+ SHorizontalBox::Slot()
-									.FillWidth( FillCoefficient_0 )
+									+ SOverlay::Slot()
 									[
-										TreeView.ToSharedRef()
+										SNew(SScrollBorder, PinnedTreeView.ToSharedRef())
+										[
+											SNew(SHorizontalBox)
+
+											// outliner tree
+											+ SHorizontalBox::Slot()
+											.FillWidth( FillCoefficient_0 )
+											[
+												PinnedTreeView.ToSharedRef()
+											]
+
+											// track area
+											+ SHorizontalBox::Slot()
+											.FillWidth( FillCoefficient_1 )
+											[
+												SNew(SBox)
+												.Padding(ResizeBarPadding)
+												.Clipping(EWidgetClipping::ClipToBounds)
+												[
+													PinnedTrackArea.ToSharedRef()
+												]
+											]
+										]
 									]
 
-									// track area
-									+ SHorizontalBox::Slot()
-									.FillWidth( FillCoefficient_1 )
+									+ SOverlay::Slot()
+									.HAlign(HAlign_Right)
 									[
-										SNew(SBox)
-										.Padding(ResizeBarPadding)
-										.Clipping(EWidgetClipping::ClipToBounds)
-										[
-											TrackArea.ToSharedRef()
-										]
+										PinnedAreaScrollBar
 									]
 								]
-							]
 
-							+ SOverlay::Slot()
-							.HAlign(HAlign_Right)
-							[
-								ScrollBar
-							]
+								+SVerticalBox::Slot()
+								[
+									SNew(SOverlay)
 
-							+ SOverlay::Slot()
-							.VAlign(VAlign_Bottom)
-							[
-								SNew(SHorizontalBox)
+									+ SOverlay::Slot()
+									[
+										SNew(SScrollBorder, TreeView.ToSharedRef())
+										[
+											SNew(SHorizontalBox)
+
+											// outliner tree
+											+ SHorizontalBox::Slot()
+											.FillWidth( FillCoefficient_0 )
+											[
+												TreeView.ToSharedRef()
+											]
+
+											// track area
+											+ SHorizontalBox::Slot()
+											.FillWidth( FillCoefficient_1 )
+											[
+												SNew(SBox)
+												.Padding(ResizeBarPadding)
+												.Clipping(EWidgetClipping::ClipToBounds)
+												[
+													TrackArea.ToSharedRef()
+												]
+											]
+										]
+									]
 
-								+ SHorizontalBox::Slot()
-								.AutoWidth()
-								[
-									SAssignNew(SequencerTreeFilterStatusBar, SSequencerTreeFilterStatusBar, InSequencer)
-									.Visibility(EVisibility::Hidden) // Initially hidden, visible on hover of the info button
-								]
-							]
-						]
-					]
+									+ SOverlay::Slot()
+									.HAlign(HAlign_Right)
+									[
+										ScrollBar
+									]
 
-					// Info Button, Transport Controls and Current Frame
-					+ SGridPanel::Slot( Column0, Row4, SGridPanel::Layer(10) )
-					[
-						SNew(SBorder)
-						.BorderImage(FAppStyle::GetBrush("ToolPanel.GroupBorder"))
-						.Clipping(EWidgetClipping::ClipToBounds)
-						[
-							SNew(SHorizontalBox)
-						
-							+ SHorizontalBox::Slot()
-							.AutoWidth()
-							.VAlign(VAlign_Center)
-							.HAlign(HAlign_Left)
-							[
-								SNew(SButton)
-								.VAlign(EVerticalAlignment::VAlign_Center)
-								.ButtonStyle(FAppStyle::Get(), "SimpleButton")
-								.ToolTipText_Lambda([this] { return LOCTEXT("ShowStatus", "Show Status"); })
-								.ContentPadding(FMargin(1, 0))
-								.Visibility(this, &SSequencer::GetInfoButtonVisibility)
-								.OnHovered_Lambda([this] { SequencerTreeFilterStatusBar->ShowStatusBar(); })
-								.OnUnhovered_Lambda([this] { SequencerTreeFilterStatusBar->FadeOutStatusBar(); })
-								.OnClicked_Lambda([this] { SequencerTreeFilterStatusBar->HideStatusBar(); return FReply::Handled(); })
-								[
-									SNew(SImage)
-									.ColorAndOpacity(FSlateColor::UseForeground())
-									.Image(FAppStyle::Get().GetBrush("Icons.Info.Small"))
+									+ SOverlay::Slot()
+									.VAlign(VAlign_Bottom)
+									[
+										SNew(SHorizontalBox)
+
+										+ SHorizontalBox::Slot()
+										.AutoWidth()
+										[
+											SAssignNew(SequencerTreeFilterStatusBar, SSequencerTreeFilterStatusBar, InSequencer)
+											.Visibility(EVisibility::Hidden) // Initially hidden, visible on hover of the info button
+										]
+									]
 								]
 							]
 
-							+ SHorizontalBox::Slot()
+							// Info Button, Transport Controls and Current Frame
+							+ SGridPanel::Slot( Column0, Row4, SGridPanel::Layer(10) )
 							[
 								SNew(SBorder)
 								.BorderImage(FAppStyle::GetBrush("ToolPanel.GroupBorder"))
-								.HAlign(HAlign_Center)
-								[
-									SequencerPtr.Pin()->MakeTransportControls(true)
-								]
-							]
-
-							+ SHorizontalBox::Slot()
-							.AutoWidth()
-							.VAlign(VAlign_Center)
-							.HAlign(HAlign_Right)
-							[
-								SNew(SButton)
-								.VAlign(EVerticalAlignment::VAlign_Center)
-								.ButtonStyle(FAppStyle::Get(), "NoBorder")
-								.ContentPadding(FMargin(1, 0))
+								.Clipping(EWidgetClipping::ClipToBounds)
 								[
 									SNew(SHorizontalBox)
-
+						
 									+ SHorizontalBox::Slot()
 									.AutoWidth()
 									.VAlign(VAlign_Center)
-									.HAlign(HAlign_Right)
-									.Padding(FMargin(CommonPadding, 0.f, 0.f, 0.f))
+									.HAlign(HAlign_Left)
+									[
+										SNew(SButton)
+										.VAlign(EVerticalAlignment::VAlign_Center)
+										.ButtonStyle(FAppStyle::Get(), "SimpleButton")
+										.ToolTipText_Lambda([this] { return LOCTEXT("ShowStatus", "Show Status"); })
+										.ContentPadding(FMargin(1, 0))
+										.Visibility(this, &SSequencer::GetInfoButtonVisibility)
+										.OnHovered_Lambda([this] { SequencerTreeFilterStatusBar->ShowStatusBar(); })
+										.OnUnhovered_Lambda([this] { SequencerTreeFilterStatusBar->FadeOutStatusBar(); })
+										.OnClicked_Lambda([this] { SequencerTreeFilterStatusBar->HideStatusBar(); return FReply::Handled(); })
+										[
+											SNew(SImage)
+											.ColorAndOpacity(FSlateColor::UseForeground())
+											.Image(FAppStyle::Get().GetBrush("Icons.Info.Small"))
+										]
+									]
+
+									+ SHorizontalBox::Slot()
 									[
 										SNew(SBorder)
-										.BorderImage(nullptr)
+										.BorderImage(FAppStyle::GetBrush("ToolPanel.GroupBorder"))
+										.HAlign(HAlign_Center)
 										[
-											PlayTimeDisplay.ToSharedRef()
+											SequencerPtr.Pin()->MakeTransportControls(true)
 										]
 									]
 
@@ -746,153 +742,205 @@ void SSequencer::Construct(const FArguments& InArgs, TSharedRef<FSequencer> InSe
 									.AutoWidth()
 									.VAlign(VAlign_Center)
 									.HAlign(HAlign_Right)
-									.Padding(FMargin(CommonPadding, 0.f, 0.f, 0.f))
 									[
-										SNew(SBorder)
-										.BorderImage(nullptr)
+										SNew(SButton)
+										.VAlign(EVerticalAlignment::VAlign_Center)
+										.ButtonStyle(FAppStyle::Get(), "NoBorder")
+										.ContentPadding(FMargin(1, 0))
 										[
-											// Current loop index, if any
-											SAssignNew(LoopIndexDisplay, STextBlock)
-											.Font(FCoreStyle::GetDefaultFontStyle("Regular", 8))
-											.Text_Lambda([this]() -> FText {
-												TOptional<int32> LoopIndex = SequencerPtr.Pin()->GetLocalLoopIndex();
-												return LoopIndex ? FText::AsNumber(LoopIndex.GetValue()) : FText();
-											})
+											SNew(SHorizontalBox)
+
+											+ SHorizontalBox::Slot()
+											.AutoWidth()
+											.VAlign(VAlign_Center)
+											.HAlign(HAlign_Right)
+											.Padding(FMargin(CommonPadding, 0.f, 0.f, 0.f))
+											[
+												SNew(SBorder)
+												.BorderImage(nullptr)
+												[
+													PlayTimeDisplay.ToSharedRef()
+												]
+											]
+
+											+ SHorizontalBox::Slot()
+											.AutoWidth()
+											.VAlign(VAlign_Center)
+											.HAlign(HAlign_Right)
+											.Padding(FMargin(CommonPadding, 0.f, 0.f, 0.f))
+											[
+												SNew(SBorder)
+												.BorderImage(nullptr)
+												[
+													// Current loop index, if any
+													SAssignNew(LoopIndexDisplay, STextBlock)
+													.Font(FCoreStyle::GetDefaultFontStyle("Regular", 8))
+													.Text_Lambda([this]() -> FText {
+														TOptional<int32> LoopIndex = SequencerPtr.Pin()->GetLocalLoopIndex();
+														return LoopIndex ? FText::AsNumber(LoopIndex.GetValue()) : FText();
+													})
+												]
+											]
 										]
 									]
 								]
 							]
-						]
-					]
 
-					// Second column
+							// Second column
 
-					+ SGridPanel::Slot( Column1, Row1 )
-					.Padding(ResizeBarPadding)
-					.RowSpan(3)
-					[
-						SNew(SBorder)
-						.BorderImage(FAppStyle::GetBrush("ToolPanel.GroupBorder"))
-						[
-							SNew(SSpacer)
-						]
-					]
+							+ SGridPanel::Slot( Column1, Row1 )
+							.Padding(ResizeBarPadding)
+							.RowSpan(3)
+							[
+								SNew(SBorder)
+								.BorderImage(FAppStyle::GetBrush("ToolPanel.GroupBorder"))
+								[
+									SNew(SSpacer)
+								]
+							]
 
-					+ SGridPanel::Slot( Column1, Row1, SGridPanel::Layer(10) )
-					.Padding(ResizeBarPadding)
-					[
-						SNew( SBorder )
-						.BorderImage( FAppStyle::GetBrush("ToolPanel.GroupBorder") )
-						.BorderBackgroundColor( FLinearColor(.50f, .50f, .50f, 1.0f ) )
-						.Padding(0)
-						.Clipping(EWidgetClipping::ClipToBounds)
-						[
-							TopTimeSlider.ToSharedRef()
-						]
-					]
+							+ SGridPanel::Slot( Column1, Row1, SGridPanel::Layer(10) )
+							.Padding(ResizeBarPadding)
+							[
+								SNew( SBorder )
+								.BorderImage( FAppStyle::GetBrush("ToolPanel.GroupBorder") )
+								.BorderBackgroundColor( FLinearColor(.50f, .50f, .50f, 1.0f ) )
+								.Padding(0)
+								.Clipping(EWidgetClipping::ClipToBounds)
+								[
+									TopTimeSlider.ToSharedRef()
+								]
+							]
 
-					// Overlay that draws the tick lines
-					+ SGridPanel::Slot( Column1, Row2, SGridPanel::Layer(10) )
-					.Padding(ResizeBarPadding)
-					[
-						SNew( SSequencerSectionOverlay, TimeSliderControllerRef )
-						.Visibility(this, &SSequencer::GetShowTickLines)
-						.DisplayScrubPosition( false )
-						.DisplayTickLines( true )
-						.Clipping(EWidgetClipping::ClipToBounds)
-					]
+							// Overlay that draws the tick lines
+							+ SGridPanel::Slot( Column1, Row2, SGridPanel::Layer(10) )
+							.Padding(ResizeBarPadding)
+							[
+								SNew( SSequencerSectionOverlay, TimeSliderControllerRef )
+								.Visibility(this, &SSequencer::GetShowTickLines)
+								.DisplayScrubPosition( false )
+								.DisplayTickLines( true )
+								.Clipping(EWidgetClipping::ClipToBounds)
+							]
 
-					// Overlay that draws the scrub position
-					+ SGridPanel::Slot( Column1, Row2, SGridPanel::Layer(20) )
-					.Padding(ResizeBarPadding)
-					[
-						SNew( SSequencerSectionOverlay, TimeSliderControllerRef )
-						.Visibility( EVisibility::HitTestInvisible )
-						.DisplayScrubPosition( true )
-						.DisplayTickLines( false )
-						.DisplayMarkedFrames( true )
-						.PaintPlaybackRangeArgs(this, &SSequencer::GetSectionPlaybackRangeArgs)
-						.Clipping(EWidgetClipping::ClipToBounds)
-					]
+							// Overlay that draws the scrub position
+							+ SGridPanel::Slot( Column1, Row2, SGridPanel::Layer(20) )
+							.Padding(ResizeBarPadding)
+							[
+								SNew( SSequencerSectionOverlay, TimeSliderControllerRef )
+								.Visibility( EVisibility::HitTestInvisible )
+								.DisplayScrubPosition( true )
+								.DisplayTickLines( false )
+								.DisplayMarkedFrames( true )
+								.PaintPlaybackRangeArgs(this, &SSequencer::GetSectionPlaybackRangeArgs)
+								.Clipping(EWidgetClipping::ClipToBounds)
+							]
 
-					+ SGridPanel::Slot(Column1, Row2, SGridPanel::Layer(30))
-						.Padding(ResizeBarPadding)
-						.HAlign(HAlign_Left)
-						.VAlign(VAlign_Top)
-					[
-						// Transform box
-						SAssignNew(TransformBox, SSequencerTransformBox, SequencerPtr.Pin().ToSharedRef(), *GetSequencerSettings(), NumericTypeInterface.ToSharedRef())
-					]
+							+ SGridPanel::Slot(Column1, Row2, SGridPanel::Layer(30))
+								.Padding(ResizeBarPadding)
+								.HAlign(HAlign_Left)
+								.VAlign(VAlign_Top)
+							[
+								// Transform box
+								SAssignNew(TransformBox, SSequencerTransformBox, SequencerPtr.Pin().ToSharedRef(), *GetSequencerSettings(), NumericTypeInterface.ToSharedRef())
+							]
 
-					+ SGridPanel::Slot(Column1, Row2, SGridPanel::Layer(40))
-						.Padding(ResizeBarPadding)
-						.HAlign(HAlign_Left)
-						.VAlign(VAlign_Top)
-					[
-						// Stretch box
-						SAssignNew(StretchBox, SSequencerStretchBox, SequencerPtr.Pin().ToSharedRef(), *GetSequencerSettings(), NumericTypeInterface.ToSharedRef())
-					]
+							+ SGridPanel::Slot(Column1, Row2, SGridPanel::Layer(40))
+								.Padding(ResizeBarPadding)
+								.HAlign(HAlign_Left)
+								.VAlign(VAlign_Top)
+							[
+								// Stretch box
+								SAssignNew(StretchBox, SSequencerStretchBox, SequencerPtr.Pin().ToSharedRef(), *GetSequencerSettings(), NumericTypeInterface.ToSharedRef())
+							]
 
-					// debug vis
-					+ SGridPanel::Slot( Column1, Row3, SGridPanel::Layer(10) )
-					.Padding(ResizeBarPadding)
-					[
-						SNew(SSequencerDebugVisualizer, InSequencer)
-						.ViewRange(FAnimatedRange::WrapAttribute(InArgs._ViewRange))
-						.Visibility(this, &SSequencer::GetDebugVisualizerVisibility)
-					]
+							// debug vis
+							+ SGridPanel::Slot( Column1, Row3, SGridPanel::Layer(10) )
+							.Padding(ResizeBarPadding)
+							[
+								SNew(SSequencerDebugVisualizer, InSequencer)
+								.ViewRange(FAnimatedRange::WrapAttribute(InArgs._ViewRange))
+								.Visibility(this, &SSequencer::GetDebugVisualizerVisibility)
+							]
 
-					// play range sliders
-					+ SGridPanel::Slot( Column1, Row4, SGridPanel::Layer(10) )
-					.Padding(ResizeBarPadding)
-					[
-						SNew( SBorder )
-						.BorderImage( FAppStyle::GetBrush("ToolPanel.GroupBorder") )
-						.BorderBackgroundColor( FLinearColor(.50f, .50f, .50f, 1.0f ) )
-						.Clipping(EWidgetClipping::ClipToBounds)
-						.Padding(0)
-						[
-							SNew( SOverlay )
+							// play range sliders
+							+ SGridPanel::Slot( Column1, Row4, SGridPanel::Layer(10) )
+							.Padding(ResizeBarPadding)
+							[
+								SNew( SBorder )
+								.BorderImage( FAppStyle::GetBrush("ToolPanel.GroupBorder") )
+								.BorderBackgroundColor( FLinearColor(.50f, .50f, .50f, 1.0f ) )
+								.Clipping(EWidgetClipping::ClipToBounds)
+								.Padding(0)
+								[
+									SNew( SOverlay )
+
+									+ SOverlay::Slot()
+									[
+										BottomTimeSlider
+									]
 
-							+ SOverlay::Slot()
+									+ SOverlay::Slot()
+									[
+										BottomTimeRange
+									]
+								]
+							]
+						]
+
+						+ SOverlay::Slot()
+						[
+							// track area virtual splitter overlay
+							SAssignNew(TreeViewSplitter, SSequencerSplitterOverlay)
+							.Style(FAppStyle::Get(), "Sequencer.AnimationOutliner.Splitter")
+							.Visibility(EVisibility::SelfHitTestInvisible)
+							.OnSplitterFinishedResizing(this, &SSequencer::OnSplitterFinishedResizing)
+
+							+ SSplitter::Slot()
+							.Value(FillCoefficient_0)
+							// Can't use a minsize here because the grid panel that is actually being used to
+							//   lay out the widgets only supports fill coefficients and this leads to a disparity between the two
+							// .MinSize(200)
+							.OnSlotResized(SSplitter::FOnSlotResized::CreateSP(this, &SSequencer::OnColumnFillCoefficientChanged, 0))
 							[
-								BottomTimeSlider
+								SNew(SSpacer)
 							]
 
-							+ SOverlay::Slot()
+							+ SSplitter::Slot()
+							.Value(FillCoefficient_1)
+							.OnSlotResized(SSplitter::FOnSlotResized::CreateSP(this, &SSequencer::OnColumnFillCoefficientChanged, 1))
 							[
-								BottomTimeRange
+								SNew(SSpacer)
 							]
 						]
 					]
-				]
-
-				+ SOverlay::Slot()
-				[
-					// track area virtual splitter overlay
-					SAssignNew(TreeViewSplitter, SSequencerSplitterOverlay)
-					.Style(FAppStyle::Get(), "Sequencer.AnimationOutliner.Splitter")
-					.Visibility(EVisibility::SelfHitTestInvisible)
-					.OnSplitterFinishedResizing(this, &SSequencer::OnSplitterFinishedResizing)
 
+					// Sequencer details sliding drawer content dock location
 					+ SSplitter::Slot()
-					.Value(FillCoefficient_0)
-					// Can't use a minsize here because the grid panel that is actually being used to
-					//   lay out the widgets only supports fill coefficients and this leads to a disparity between the two
-					// .MinSize(200)
-					.OnSlotResized(SSplitter::FOnSlotResized::CreateSP(this, &SSequencer::OnColumnFillCoefficientChanged, 0))
+					.Resizable(false)
+					.SizeRule(SSplitter::ESizeRule::SizeToContent)
+					.Value(FSidebarDrawerConfig::DefaultSizeCoefficient)
+					.OnSlotResized(this, &SSequencer::OnSidebarSlotResized)
+					.Expose(SidebarSlot)
 					[
-						SNew(SSpacer)
+						DetailsDockLocation.ToSharedRef()
 					]
+				]
 
-					+ SSplitter::Slot()
-					.Value(FillCoefficient_1)
-					.OnSlotResized(SSplitter::FOnSlotResized::CreateSP(this, &SSequencer::OnColumnFillCoefficientChanged, 1))
-					[
-						SNew(SSpacer)
-					]
+				// Sequencer details sliding drawer sidebar
+				+ SHorizontalBox::Slot()
+				.AutoWidth()
+				[
+					SAssignNew(DetailsSidebar, SSidebar, SidebarOverlay, DetailsDockLocation.ToSharedRef())
+					.OnDockStateChanged(this, &SSequencer::OnSidebarDockStateChanged)
+					.Visibility_Lambda([this]() { return IsSidebarVisible() ? EVisibility::Visible : EVisibility::Collapsed; })
 				]
 			]
+
+			+ SOverlay::Slot()
+			[
+				SidebarOverlay
+			]
 		]
 	];
 
@@ -4178,5 +4226,88 @@ USequencerSettings* SSequencer::GetSequencerSettings() const
 	return nullptr;
 }
 
-#undef LOCTEXT_NAMESPACE
+bool SSequencer::RegisterDrawer(FSidebarDrawerConfig&& InDrawerConfig)
+{
+	if (DetailsSidebar.IsValid())
+	{
+		return DetailsSidebar->RegisterDrawer(MoveTemp(InDrawerConfig));
+	}
+	return false;
+}
+
+bool SSequencer::UnregisterDrawer(const FName InDrawerId)
+{
+	if (DetailsSidebar.IsValid())
+	{
+		return DetailsSidebar->UnregisterDrawer(InDrawerId);
+	}
+	return false;
+}
+
+bool SSequencer::RegisterDrawerSection(const FName InDrawerId, const TSharedPtr<ISidebarDrawerContent>& InSection)
+{
+	if (DetailsSidebar.IsValid())
+	{
+		return DetailsSidebar->RegisterDrawerSection(InDrawerId, InSection);
+	}
+	return false;
+}
+
+bool SSequencer::UnregisterDrawerSection(const FName InDrawerId, const FName InSectionId)
+{
+	if (DetailsSidebar.IsValid())
+	{
+		return DetailsSidebar->UnregisterDrawerSection(InDrawerId, InSectionId);
+	}
+	return false;
+}
+
+void SSequencer::OnSidebarDockStateChanged(const FName InDrawerId)
+{
+	DetailsSidebar->UpdateDockedSplitterSlot(InDrawerId, SidebarSlot);
+}
+
+void SSequencer::OnSidebarSlotResized(const float InFillCoefficient)
+{
+	if (InFillCoefficient == 0.f && DetailsSidebar->HasDrawerDocked())
+	{
+		DetailsSidebar->UndockAllDrawers();
+	}
+}
+
+bool SSequencer::IsSidebarVisible() const
+{
+	return DetailsSidebar.IsValid() && DetailsSidebar->GetDrawerCount() > 0;
+}
 
+void SSequencer::ToggleSidebarSelectionDrawerOpen()
+{
+	if (DetailsSidebar->IsDrawerDocked(FSequencer::SelectionDrawerId))
+	{
+		DetailsSidebar->UndockAllDrawers();
+	}
+	else if (DetailsSidebar->IsDrawerOpened(FSequencer::SelectionDrawerId))
+	{
+		DetailsSidebar->CloseAllDrawers();
+	}
+	else
+	{
+		DetailsSidebar->TryOpenDrawer(FSequencer::SelectionDrawerId);
+	}
+}
+
+void SSequencer::ToggleSidebarDrawerDock()
+{
+	const FName OpenedDrawerId = DetailsSidebar->GetOpenedDrawerId();
+	if (!OpenedDrawerId.IsNone())
+	{
+		DetailsSidebar->SetDrawerDocked(OpenedDrawerId, true);
+		DetailsSidebar->CloseAllDrawers();
+	}
+	else if (DetailsSidebar->HasDrawerDocked())
+	{
+		DetailsSidebar->UndockAllDrawers();
+	}
+}
+
+#undef LOCTEXT_NAMESPACE
diff --git a/Engine/Source/Editor/Sequencer/Private/SSequencer.h b/Engine/Source/Editor/Sequencer/Private/SSequencer.h
--- a/Engine/Source/Editor/Sequencer/Private/SSequencer.h
+++ b/Engine/Source/Editor/Sequencer/Private/SSequencer.h
@@ -17,7 +17,7 @@
 #include "ToolMenu.h"
 #include "Framework/Commands/UICommandList.h"
 #include "Widgets/Input/NumericTypeInterface.h"
-#include "Widgets/Input/SSpinBox.h"
+#include "Widgets/Layout/SSplitter.h"
 #include "Widgets/Text/STextBlock.h"
 #include "Sequencer.h"
 #include "SequencerWidgetsDelegates.h"
@@ -29,11 +29,14 @@ class FAssetDragDropOp;
 class FClassDragDropOp;
 class FMovieSceneClipboard;
 class FSequencerTimeSliderController;
+class ISidebarDrawerContent;
 class SCurveEditorTree;
 class SSequencerTransformBox;
 class SSequencerStretchBox;
 class SCurveEditorPanel;
+class SBox;
 class SDockTab;
+class SSidebar;
 class SWindow;
 class USequencerSettings;
 class FSequencerTrackFilter;
@@ -41,6 +44,7 @@ class SSequencerGroupManager;
 class SSequencerTreeFilterStatusBar;
 struct FPaintPlaybackRangeArgs;
 struct FSequencerCustomizationInfo;
+struct FSidebarDrawerConfig;
 
 namespace UE
 {
@@ -649,11 +653,60 @@ public:
 	/** Applies dynamic sequencer customizations to this editor. */
 	void ApplySequencerCustomizations(const TArrayView<const FSequencerCustomizationInfo> Customizations);
 
+	/**
+	 * Registers and displays a new drawer in the sidebar.
+	 * 
+	 * @param InDrawerConfig Configuration info for the new drawer
+	 * 
+	 * @return True if the new drawer registration was successful.
+	 */
+	bool RegisterDrawer(FSidebarDrawerConfig&& InDrawerConfig);
+
+	/**
+	 * Unregisters and removes a drawer from the sidebar.
+	 *
+	 * @param InDrawerId Unique drawer Id to unregister
+	 * 
+	 * @return True if the drawer removal was successful.
+	 */
+	bool UnregisterDrawer(const FName InDrawerId);
+
+	/**
+	 * Registers and displays a new drawer section in the sidebar.
+	 * 
+	 * @param InDrawerId Unique drawer Id to register
+	 * @param InSection Drawer content interface for the section
+	 * 
+	 * @return True if the new drawer section registration was successful.
+	 */
+	bool RegisterDrawerSection(const FName InDrawerId, const TSharedPtr<ISidebarDrawerContent>& InSection);
+
+	/**
+	 * Unregisters and removes a drawer section from the sidebar.
+	 * 
+	 * @param InDrawerId Unique drawer Id that contains the section to unregister
+	 * @param InSectionId Unique drawer section Id to unregister
+	 * 
+	 * @return True if the drawer removal was successful.
+	 */
+	bool UnregisterDrawerSection(const FName InDrawerId, const FName InSectionId);
+
+	/** @return True if the sidebar is being displayed. */
+	bool IsSidebarVisible() const;
+
+	/** Toggles the sidebar "Selection" drawer open or closed. */
+	void ToggleSidebarSelectionDrawerOpen();
+
+	/** Undocks the docked sidebar drawer if docked or docks the sidebar drawer if there is one open and no currently docked drawer. */
+	void ToggleSidebarDrawerDock();
+
 private:
 	/** Applies a single customization. */
 	void ApplySequencerCustomization(const FSequencerCustomizationInfo& Customization);
 
-private:
+	void OnSidebarDockStateChanged(const FName InDrawerId);
+
+	void OnSidebarSlotResized(const float InFillCoefficient);
 
 	/** Transform box widget. */
 	TSharedPtr<SSequencerTransformBox> TransformBox;
@@ -803,4 +856,9 @@ private:
 	TWeakPtr<SWindow> WeakNodeGroupWindow;
 
 	TSharedPtr<SSequencerGroupManager> NodeGroupManager;
+
+	TSharedPtr<SBox> DetailsDockLocation;
+
+	TSharedPtr<SSidebar> DetailsSidebar;
+	SSplitter::FSlot* SidebarSlot = nullptr;
 };
diff --git a/Engine/Source/Editor/Sequencer/Private/Sequencer.cpp b/Engine/Source/Editor/Sequencer/Private/Sequencer.cpp
--- a/Engine/Source/Editor/Sequencer/Private/Sequencer.cpp
+++ b/Engine/Source/Editor/Sequencer/Private/Sequencer.cpp
@@ -162,6 +162,8 @@
 #include "Systems/MovieSceneMotionVectorSimulationSystem.h"
 #include "IKeyArea.h"
 #include "Editor/TransBuffer.h"
+#include "Sidebar/SidebarDrawerConfig.h"
+#include "Widgets/Sidebar/SequencerSelectionDrawer.h"
 
 #include "EngineModule.h"
 #include "IViewportSelectableObject.h"
@@ -292,6 +294,7 @@ namespace UE
 } // namespace UE
 
 bool FSequencer::bSelectionLimited = false;
+const FName FSequencer::SelectionDrawerId = TEXT("SelectionDetails");
 
 void FSequencer::InitSequencer(const FSequencerInitParams& InitParams, const TSharedRef<ISequencerObjectChangeListener>& InObjectChangeListener, const TArray<FOnCreateTrackEditor>& TrackEditorDelegates, const TArray<FOnCreateEditorObjectBinding>& EditorObjectBindingDelegates, const TArray<FOnCreateOutlinerColumn>& OutlinerColumnDelegates)
 {
@@ -665,6 +668,20 @@ void FSequencer::InitSequencer(const FSequencerInitParams& InitParams, const TSh
 	AddNodeGroupsCollectionChangedDelegate();
 
 	OnActivateSequenceEvent.Broadcast(ActiveTemplateIDs[0]);
+
+	{ // Selection Details Drawer
+		FSequencerSelectionDrawer SelectionDrawer(SharedThis(this));
+
+		FSidebarDrawerConfig DetailsDrawer;
+		DetailsDrawer.UniqueId = SelectionDrawerId;
+		DetailsDrawer.ButtonText = LOCTEXT("SelectionDetailsPanelLabel", "Selection");
+		DetailsDrawer.ToolTipText = TAttribute<FText>::CreateSP(this, &FSequencer::GetSidebarSelectionDrawerToolTipText);
+		DetailsDrawer.Icon = FAppStyle::GetBrush(TEXT("EditorPreferences.TabIcon"));
+		DetailsDrawer.bInitiallyDocked = true;
+		RegisterDrawer(MoveTemp(DetailsDrawer));
+
+		RegisterDrawerSection(SelectionDrawerId, MakeShared<FSequencerSelectionDrawer>(SharedThis(this)));
+	}
 }
 
 void FSequencer::SetSequencerSettings(USequencerSettings* InSettings)
@@ -11553,6 +11570,14 @@ void FSequencer::BindCommands()
 		FCanExecuteAction(),
 		FIsActionChecked::CreateSP(this, &FSequencer::IsViewportSelectionLimited));
 
+	SequencerCommandBindings->MapAction(
+		Commands.ToggleSidebarSelectionDrawerOpen,
+		FExecuteAction::CreateSP( this, &FSequencer::ShowHideSidebarSelectionDrawer));
+
+	SequencerCommandBindings->MapAction(
+		Commands.ToggleSidebarDrawerDock,
+		FExecuteAction::CreateSP( this, &FSequencer::ToggleSidebarDrawerDocked));
+
 	// If this sequencer supports a curve editor, let's add bindings for it.
 	FCurveEditorExtension* CurveEditorExtension = ViewModel->CastDynamic<FCurveEditorExtension>();
 	if (CurveEditorExtension && ensure(CurveEditorExtension->GetCurveEditor()))
@@ -12094,4 +12119,87 @@ ISequencer::FOnViewportSelectionLimitedChanged& FSequencer::OnViewportSelectionL
 	return OnSelectionLimitedChangedDelegate;
 }
 
+bool FSequencer::RegisterDrawer(FSidebarDrawerConfig&& InDrawerConfig)
+{
+	if (SequencerWidget.IsValid())
+	{
+		return SequencerWidget->RegisterDrawer(MoveTemp(InDrawerConfig));
+	}
+	return false;
+}
+
+bool FSequencer::UnregisterDrawer(const FName InDrawerId)
+{
+	if (SequencerWidget.IsValid())
+	{
+		return SequencerWidget->UnregisterDrawer(InDrawerId);
+	}
+	return false;
+}
+
+bool FSequencer::RegisterDrawerSection(const FName InDrawerId, const TSharedPtr<ISidebarDrawerContent>& InSection)
+{
+	if (SequencerWidget.IsValid())
+	{
+		return SequencerWidget->RegisterDrawerSection(InDrawerId, InSection);
+	}
+	return false;
+}
+
+bool FSequencer::UnregisterDrawerSection(const FName InDrawerId, const FName InSectionId)
+{
+	if (SequencerWidget.IsValid())
+	{
+		return SequencerWidget->UnregisterDrawerSection(InDrawerId, InSectionId);
+	}
+	return false;
+}
+
+void FSequencer::ShowHideSidebarSelectionDrawer()
+{
+	if (SequencerWidget.IsValid())
+	{
+		SequencerWidget->ToggleSidebarSelectionDrawerOpen();
+	}
+}
+
+void FSequencer::ToggleSidebarDrawerDocked()
+{
+	if (SequencerWidget.IsValid())
+	{
+		SequencerWidget->ToggleSidebarDrawerDock();
+	}
+}
+
+FText FSequencer::GetSidebarSelectionDrawerToolTipText() const
+{
+	const TSharedRef<const FInputChord> DrawerOpenActiveChord = FSequencerCommands::Get().ToggleSidebarSelectionDrawerOpen->GetFirstValidChord();
+	const TSharedRef<const FInputChord> DrawerDockActiveChord = FSequencerCommands::Get().ToggleSidebarDrawerDock->GetFirstValidChord();
+
+	FText ToolTipText = LOCTEXT("SelectionDetailsPanelTooltip", "Open Sequencer selection details panel.");
+
+	if (DrawerOpenActiveChord->IsValidChord() || DrawerDockActiveChord->IsValidChord())
+	{
+		ToolTipText = FText::Format(LOCTEXT("ExtendedSelectionDetailsPanelTooltip", "{0}\n"), ToolTipText);
+	}
+
+	if (DrawerOpenActiveChord->IsValidChord())
+	{
+		ToolTipText = FText::Format(LOCTEXT("ExtendedSelectionDrawerOpenDetailsPanelTooltip", "{0}\n"
+			"{1} to toggle the drawer open or closed")
+			, ToolTipText
+			, DrawerOpenActiveChord->GetInputText(true));
+	}
+
+	if (DrawerDockActiveChord->IsValidChord())
+	{
+		ToolTipText = FText::Format(LOCTEXT("ExtendedSelectionDrawerDockDetailsPanelTooltip", "{0}\n"
+			"{1} to toggle the drawer docked or undocked")
+			, ToolTipText
+			, DrawerDockActiveChord->GetInputText(true));
+	}
+
+	return ToolTipText;
+}
+
 #undef LOCTEXT_NAMESPACE
diff --git a/Engine/Source/Editor/Sequencer/Private/Sequencer.h b/Engine/Source/Editor/Sequencer/Private/Sequencer.h
--- a/Engine/Source/Editor/Sequencer/Private/Sequencer.h
+++ b/Engine/Source/Editor/Sequencer/Private/Sequencer.h
@@ -59,6 +59,7 @@ class FObjectBindingTagCache;
 class ISequencerTrackEditor;
 class ISequencerEditorObjectBinding;
 class SSequencer;
+class ISidebarDrawerContent;
 class ULevel;
 class UMovieSceneSequence;
 class UMovieSceneSubSection;
@@ -78,6 +79,7 @@ struct FKeyAttributes;
 struct FNotificationInfo;
 struct FEditorViewportViewModifierParams;
 struct FMovieSceneMarkedFrame;
+struct FSidebarDrawerConfig;
 
 enum class EMapChangeType : uint8;
 enum class ENearestKeyOption : uint8;
@@ -117,6 +119,7 @@ class FSequencer final
 	using FViewModel = UE::Sequencer::FViewModel;
 
 public:
+	SEQUENCER_API static const FName SelectionDrawerId;
 
 	/** Constructor */
 	FSequencer();
@@ -1171,6 +1174,50 @@ public:
 
 	ISequencer::FOnViewportSelectionLimitedChanged& OnViewportSelectionLimitedChanged() override;
 
+	/**
+	 * Registers and displays a new drawer in the sidebar.
+	 * 
+	 * @param InDrawerConfig	Configuration info for the new drawer
+	 * 
+	 * @return True if the new drawer registration was successful.
+	 */
+	virtual bool RegisterDrawer(FSidebarDrawerConfig&& InDrawerConfig) override;
+
+	/**
+	 * Unregisters and removes a drawer from the sidebar.
+	 *
+	 * @param InDrawerId	Unique drawer Id to unregister
+	 * 
+	 * @return True if the drawer removal was successful.
+	 */
+	virtual bool UnregisterDrawer(const FName InDrawerId) override;
+	
+	/**
+	 * Registers and displays a new drawer section in the sidebar.
+	 * 
+	 * @param InDrawerId	Unique drawer Id to register
+	 * @param InSection		Drawer content interface for the section
+	 * 
+	 * @return True if the new drawer section registration was successful.
+	 */
+	virtual bool RegisterDrawerSection(const FName InDrawerId, const TSharedPtr<ISidebarDrawerContent>& InSection) override;
+	
+	/**
+	 * Unregisters and removes a drawer section from the sidebar.
+	 * 
+	 * @param InDrawerId	Unique drawer Id that contains the section to unregister
+	 * @param InSectionId	Unique drawer section Id to unregister
+	 * 
+	 * @return True if the drawer removal was successful.
+	 */
+	virtual bool UnregisterDrawerSection(const FName InDrawerId, const FName InSectionId) override;
+
+	/** Toggles the sidebar "Selection" drawer open or closed. */
+	void ShowHideSidebarSelectionDrawer();
+
+	/** Undocks the docked sidebar drawer if docked or docks the sidebar drawer if there is one open and no currently docked drawer. */
+	void ToggleSidebarDrawerDocked();
+
 private:
 
 	/** Update the time bases for the current movie scene */
@@ -1181,6 +1228,8 @@ private:
 
 	void ForEachSubSequenceRecursively(UMovieSceneSequence* const InSequence, const TFunctionRef<bool(UMovieSceneSequence* const InCurrentSequence)>& InFunction);
 
+	FText GetSidebarSelectionDrawerToolTipText() const;
+
 	/** User-supplied settings object for this sequencer */
 	TObjectPtr<USequencerSettings> Settings;
 
diff --git a/Engine/Source/Editor/Sequencer/Private/SequencerCommands.cpp b/Engine/Source/Editor/Sequencer/Private/SequencerCommands.cpp
--- a/Engine/Source/Editor/Sequencer/Private/SequencerCommands.cpp
+++ b/Engine/Source/Editor/Sequencer/Private/SequencerCommands.cpp
@@ -197,6 +197,9 @@ void FSequencerCommands::RegisterCommands()
 	UI_COMMAND(RefreshUI, "Refresh UI", "Forcibly refresh the UI from source data.", EUserInterfaceActionType::Button, FInputChord(EKeys::F5));
 
 	UI_COMMAND(ToggleLimitViewportSelection, "Toggle Limit Viewport Selection", "Toggles viewport selectability between Sequencer only objects and all objects.", EUserInterfaceActionType::ToggleButton, FInputChord());
+
+	UI_COMMAND(ToggleSidebarSelectionDrawerOpen, "Toggle Sidebar Selection Drawer Open", "Toggles the visible state of the sidebar selection drawer. Undocks and closes the sidebar drawer if docked.", EUserInterfaceActionType::Button, FInputChord(EModifierKey::Control, EKeys::E));
+	UI_COMMAND(ToggleSidebarDrawerDock, "Toggle Sidebar Drawer Docked", "Toggles the docked state of the sidebar current drawer. Undocks the docked sidebar drawer if docked or docks the sidebar drawer if there is one open and no currently docked drawer.", EUserInterfaceActionType::Button, FInputChord(EModifierKey::Alt, EKeys::D));
 }
 
 #undef LOCTEXT_NAMESPACE
diff --git a/Engine/Source/Editor/Sequencer/Private/SequencerCommonHelpers.cpp b/Engine/Source/Editor/Sequencer/Private/SequencerCommonHelpers.cpp
--- a/Engine/Source/Editor/Sequencer/Private/SequencerCommonHelpers.cpp
+++ b/Engine/Source/Editor/Sequencer/Private/SequencerCommonHelpers.cpp
@@ -285,23 +285,26 @@ public:
 
 	void Construct(FArguments InArgs) { }
 
-	void SetDetailsAndSequencer(TSharedRef<SWidget> InDetailsPanel, TSharedRef<ISequencer> InSequencer)
+	void SetDetailsAndSequencer(TSharedRef<SWidget> InDetailsPanel, TWeakPtr<ISequencer> InSequencerWeak)
 	{
 		ChildSlot
 		[
 			InDetailsPanel
 		];
-		Sequencer = InSequencer;
+		SequencerWeak = InSequencerWeak;
 	}
 
 	//~ FNotifyHook interface
 	virtual void NotifyPostChange(const FPropertyChangedEvent& PropertyChangedEvent, FProperty* PropertyThatChanged) override
 	{
-		Sequencer->NotifyMovieSceneDataChanged(EMovieSceneDataChangeType::TrackValueChanged);
+		if (const TSharedPtr<ISequencer> Sequencer = SequencerWeak.Pin())
+		{
+			Sequencer->NotifyMovieSceneDataChanged(EMovieSceneDataChangeType::TrackValueChanged);
+		}
 	}
 
 private:
-	TSharedPtr<ISequencer> Sequencer;
+	TWeakPtr<ISequencer> SequencerWeak;
 };
 
 
@@ -353,7 +356,7 @@ void SequencerHelpers::AddPropertiesMenu(FSequencer& Sequencer, FMenuBuilder& Me
 			{
 				TSharedPtr<ISequencerSection> SectionInterface = SectionHandle->GetSectionInterface();
 				FSequencerSectionPropertyDetailsViewCustomizationParams CustomizationDetails(
-					SectionInterface.ToSharedRef(), Sequencer.AsShared(), *SectionHandle->GetParentTrackExtension()->GetTrackEditor().Get());
+					SectionInterface.ToSharedRef(), Sequencer.AsWeak(), *SectionHandle->GetParentTrackExtension()->GetTrackEditor().Get());
 				TSharedPtr<FObjectBindingModel> ParentObjectBindingNode = SectionHandle->FindAncestorOfType<FObjectBindingModel>();
 				if (ParentObjectBindingNode.IsValid())
 				{
@@ -367,7 +370,7 @@ void SequencerHelpers::AddPropertiesMenu(FSequencer& Sequencer, FMenuBuilder& Me
 	Sequencer.OnInitializeDetailsPanel().Broadcast(DetailsView, Sequencer.AsShared());
 	DetailsView->SetObjects(Sections);
 
-	DetailsNotifyWrapper->SetDetailsAndSequencer(DetailsView, Sequencer.AsShared());
+	DetailsNotifyWrapper->SetDetailsAndSequencer(DetailsView, Sequencer.AsWeak());
 	DetailsNotifyWrapper->SetEnabled(!Sequencer.IsReadOnly());
 	MenuBuilder.AddWidget(DetailsNotifyWrapper, FText::GetEmpty(), true);
 }
diff --git a/Engine/Source/Editor/Sequencer/Private/SequencerModule.cpp b/Engine/Source/Editor/Sequencer/Private/SequencerModule.cpp
--- a/Engine/Source/Editor/Sequencer/Private/SequencerModule.cpp
+++ b/Engine/Source/Editor/Sequencer/Private/SequencerModule.cpp
@@ -427,6 +427,7 @@ public:
 		AddTrackMenuExtensibilityManager = MakeShareable( new FExtensibilityManager );
 		ToolBarExtensibilityManager = MakeShareable(new FExtensibilityManager);
 		ActionsMenuExtensibilityManager = MakeShareable(new FExtensibilityManager);
+		SidebarExtensibilityManager = MakeShareable(new FExtensibilityManager);
 
 		SequencerCustomizationManager = MakeShareable(new FSequencerCustomizationManager);
 	}
@@ -514,6 +515,7 @@ public:
 	virtual TSharedPtr<FExtensibilityManager> GetAddTrackMenuExtensibilityManager() const override { return AddTrackMenuExtensibilityManager; }
 	virtual TSharedPtr<FExtensibilityManager> GetToolBarExtensibilityManager() const override { return ToolBarExtensibilityManager; }
 	virtual TSharedPtr<FExtensibilityManager> GetActionsMenuExtensibilityManager() const override { return ActionsMenuExtensibilityManager; }
+	virtual TSharedPtr<FExtensibilityManager> GetSidebarExtensibilityManager() const override { return SidebarExtensibilityManager; }
 
 	virtual TSharedPtr<FSequencerCustomizationManager> GetSequencerCustomizationManager() const override { return SequencerCustomizationManager; }
 
@@ -626,6 +628,7 @@ private:
 	TSharedPtr<FExtensibilityManager> AddTrackMenuExtensibilityManager;
 	TSharedPtr<FExtensibilityManager> ToolBarExtensibilityManager;
 	TSharedPtr<FExtensibilityManager> ActionsMenuExtensibilityManager;
+	TSharedPtr<FExtensibilityManager> SidebarExtensibilityManager;
 
 	TSharedPtr<FSequencerCustomizationManager> SequencerCustomizationManager;
 
diff --git a/Engine/Source/Editor/Sequencer/Private/SequencerUtilities.cpp b/Engine/Source/Editor/Sequencer/Private/SequencerUtilities.cpp
--- a/Engine/Source/Editor/Sequencer/Private/SequencerUtilities.cpp
+++ b/Engine/Source/Editor/Sequencer/Private/SequencerUtilities.cpp
@@ -260,6 +260,8 @@ void FSequencerUtilities::PopulateMenu_BlenderSubMenu(FMenuBuilder& MenuBuilder,
 		}
 	);
 
+	MenuBuilder.BeginSection(TEXT("Blending"), LOCTEXT("BlendingMenuSection", "Blending"));
+
 	for (TSubclassOf<UMovieSceneBlenderSystem> SystemClass : BlenderTypes)
 	{
 		MenuBuilder.AddMenuEntry(
@@ -282,6 +284,8 @@ void FSequencerUtilities::PopulateMenu_BlenderSubMenu(FMenuBuilder& MenuBuilder,
 			EUserInterfaceActionType::RadioButton
 		);
 	}
+	
+	MenuBuilder.EndSection();
 }
 
 void FSequencerUtilities::PopulateMenu_SetBlendType(FMenuBuilder& MenuBuilder, UMovieSceneSection* Section, TWeakPtr<ISequencer> InSequencer)
diff --git a/Engine/Source/Editor/Sequencer/Private/Widgets/Sidebar/SequenceSelectionDrawer.cpp b/Engine/Source/Editor/Sequencer/Private/Widgets/Sidebar/SequenceSelectionDrawer.cpp
new file mode 100644
--- /dev/null
+++ b/Engine/Source/Editor/Sequencer/Private/Widgets/Sidebar/SequenceSelectionDrawer.cpp
@@ -0,0 +1,604 @@
+// Copyright Epic Games, Inc. All Rights Reserved.
+
+#include "SequencerSelectionDrawer.h"
+#include "DetailsViewArgs.h"
+#include "FrameNumberDetailsCustomization.h"
+#include "Framework/MultiBox/MultiBoxBuilder.h"
+#include "IKeyArea.h"
+#include "ISequencerSection.h"
+#include "ISequencerTrackEditor.h"
+#include "IStructureDetailsView.h"
+#include "Menus/CurveChannelSectionSidebarExtension.h"
+#include "Modules/ModuleManager.h"
+#include "MovieScene.h"
+#include "MovieSceneMarkedFrame.h"
+#include "MovieSceneSequence.h"
+#include "MVVM/Extensions/ITrackExtension.h"
+#include "MVVM/Selection/Selection.h"
+#include "MVVM/ViewModelPtr.h"
+#include "MVVM/ViewModels/CategoryModel.h"
+#include "MVVM/ViewModels/LayerBarModel.h"
+#include "MVVM/ViewModels/PossessableModel.h"
+#include "MVVM/ViewModels/SectionModel.h"
+#include "MVVM/ViewModels/SequencerEditorViewModel.h"
+#include "MVVM/ViewModels/TrackModel.h"
+#include "PropertyEditorModule.h"
+#include "Sequencer.h"
+#include "SequencerCommonHelpers.h"
+#include "SKeyEditInterface.h"
+#include "Styling/AppStyle.h"
+#include "Templates/SharedPointer.h"
+#include "Widgets/Layout/SBorder.h"
+
+#define LOCTEXT_NAMESPACE "SequencerSelectionDrawer"
+
+namespace UE::Sequencer::Private
+{
+	FKeyEditData GetKeyEditData(const UE::Sequencer::FKeySelection& InKeySelection)
+	{
+		if (InKeySelection.Num() == 1)
+		{
+			for (const FKeyHandle Key : InKeySelection)
+			{
+				if (const TSharedPtr<UE::Sequencer::FChannelModel> Channel = InKeySelection.GetModelForKey(Key))
+				{
+					FKeyEditData KeyEditData;
+					KeyEditData.KeyStruct     = Channel->GetKeyArea()->GetKeyStruct(Key);
+					KeyEditData.OwningSection = Channel->GetSection();
+					return KeyEditData;
+				}
+			}
+		}
+		else
+		{
+			TArray<FKeyHandle> KeyHandles;
+			UMovieSceneSection* CommonSection = nullptr;
+			for (FKeyHandle Key : InKeySelection)
+			{
+				TSharedPtr<UE::Sequencer::FChannelModel> Channel = InKeySelection.GetModelForKey(Key);
+				if (Channel.IsValid())
+				{
+					KeyHandles.Add(Key);
+					if (!CommonSection)
+					{
+						CommonSection = Channel->GetSection();
+					}
+					else if (CommonSection != Channel->GetSection())
+					{
+						CommonSection = nullptr;
+						break;
+					}
+				}
+			}
+
+			if (CommonSection)
+			{
+				FKeyEditData KeyEditData;
+				KeyEditData.KeyStruct     = CommonSection->GetKeyStruct(KeyHandles);
+				KeyEditData.OwningSection = CommonSection;
+				return KeyEditData;
+			}
+		}
+
+		return FKeyEditData();
+	}
+
+	TSharedPtr<UE::Sequencer::FSequencerSelection> GetSelection(const ISequencer& InSequencer)
+	{
+		const TSharedPtr<UE::Sequencer::FSequencerEditorViewModel> ViewModel = InSequencer.GetViewModel();
+		if (!ViewModel.IsValid())
+		{
+			return nullptr;
+		}
+
+		return ViewModel->GetSelection();
+	}
+}
+
+using namespace UE::Sequencer;
+
+const FName FSequencerSelectionDrawer::UniqueId = TEXT("SequencerSelectionDrawer");
+
+FSequencerSelectionDrawer::FSequencerSelectionDrawer(const TWeakPtr<FSequencer>& InSequencerWeak)
+	: SequencerWeak(InSequencerWeak)
+{
+}
+
+FName FSequencerSelectionDrawer::GetUniqueId() const
+{
+	return UniqueId;
+}
+
+FName FSequencerSelectionDrawer::GetSectionId() const
+{
+	return TEXT("Selection");
+}
+
+FText FSequencerSelectionDrawer::GetSectionDisplayText() const
+{
+	return LOCTEXT("SelectionDisplayText", "Selection");
+}
+
+TSharedRef<SWidget> FSequencerSelectionDrawer::CreateContentWidget()
+{
+	if (const TSharedPtr<ISequencer> Sequencer = SequencerWeak.Pin())
+	{
+		Sequencer->OnActorAddedToSequencer().AddLambda([this](AActor* InActor, const FGuid InGuid)
+			{
+				OnSequencerSelectionChanged();
+			});
+
+		if (const TSharedPtr<UE::Sequencer::FSequencerSelection> SequencerSelection = UE::Sequencer::Private::GetSelection(*Sequencer.Get()))
+		{
+			SequencerSelection->OnChanged.AddSP(this, &FSequencerSelectionDrawer::OnSequencerSelectionChanged);
+			
+			OnSequencerSelectionChanged();
+		}
+	}
+
+	return SNew(SBorder)
+		.HAlign(HAlign_Fill)
+		.BorderImage(FAppStyle::GetBrush(TEXT("NoBorder")))
+		[
+			SAssignNew(ContentBox, SVerticalBox)
+			+ SVerticalBox::Slot()
+			[
+				CreateNoSelectionHintText()
+			]
+		];
+}
+
+void FSequencerSelectionDrawer::OnSequencerSelectionChanged()
+{
+	if (!ContentBox.IsValid())
+	{
+		return;
+	}
+
+	ContentBox->ClearChildren();
+
+	CurveChannelExtension.Reset();
+
+	const TSharedPtr<FSequencer> Sequencer = SequencerWeak.Pin();
+	if (!SequencerWeak.IsValid())
+	{
+		return;
+	}
+
+	const TSharedPtr<FSequencerSelection> SequencerSelection = Private::GetSelection(*Sequencer.Get());
+	if (!SequencerSelection.IsValid())
+	{
+		return;
+	}
+
+	const TSharedRef<FSequencerSelection> SelectionRef = SequencerSelection.ToSharedRef();
+
+	auto AddToContent = [this](const TSharedRef<SWidget>& InWidget)
+		{
+			ContentBox->AddSlot()
+				.AutoHeight()
+				[
+					InWidget
+				];
+		};
+
+	const ISequencerModule& SequencerModule = FModuleManager::Get().LoadModuleChecked<ISequencerModule>(TEXT("Sequencer"));
+	const TSharedPtr<FExtensibilityManager> SidebarExtensibilityManager = SequencerModule.GetSidebarExtensibilityManager();
+	const TSharedPtr<FExtender> Extender = SidebarExtensibilityManager->GetAllExtenders();
+
+	FMenuBuilder MenuBuilder(/*bInShouldCloseWindowAfterMenuSelection=*/false
+		, Sequencer->GetCommandBindings(), Extender
+		, /*bInCloseSelfOnly=*/true, &FCoreStyle::Get(), /*bInSearchable=*/true, TEXT("Sequencer.Sidebar"));
+
+	/**
+	 * Selection details display order preference:
+	 *  1) Key items
+	 *  2) Binding properties
+	 *  2) Track area items (if no key selected)
+	 *  3) Outliner items (if no key or track area selected)
+	 *  5) Marked frames
+	 */
+
+	// 1) Key items
+	BuildKeySelectionDetails(SelectionRef, MenuBuilder);
+
+	// 2) Binding properties
+	BuildBindingPropertiesDetails(*Sequencer, SelectionRef, MenuBuilder);
+
+	// Early out for key selections
+	const bool bIsKeySelected = SequencerSelection->KeySelection.Num() > 0;
+	if (bIsKeySelected)
+	{
+		AddToContent(MenuBuilder.MakeWidget());
+		return;
+	}
+
+	// 3) Track area items
+	BuildTrackAreaDetails(*Sequencer, SelectionRef, MenuBuilder);
+
+	// 4) Outliner items
+	const bool bIsTrackAreaSelected = SequencerSelection->TrackArea.Num() > 0;
+	if (!bIsTrackAreaSelected)
+	{
+		BuildOutlinerDetails(*Sequencer, SelectionRef, MenuBuilder);
+	}
+
+	// 5) Marked frames
+	BuildMarkedFrameDetails(SelectionRef, MenuBuilder);
+
+	AddToContent(MenuBuilder.MakeWidget());
+}
+
+void FSequencerSelectionDrawer::BuildKeySelectionDetails(const TSharedRef<UE::Sequencer::FSequencerSelection>& InSelection, FMenuBuilder& MenuBuilder)
+{
+	if (InSelection->KeySelection.Num() == 0)
+	{
+		return;
+	}
+
+	MenuBuilder.BeginSection(TEXT("KeyEdit"), LOCTEXT("KeyEditMenuSection", "Key Edit"));
+
+	MenuBuilder.AddWidget(CreateKeyFrameDetails(InSelection).ToSharedRef(), FText::GetEmpty(), /*bInNoIndent=*/true);
+
+	MenuBuilder.EndSection();
+
+	// Show the section for the keys if they are all part of the same section
+	TArray<TViewModelPtr<FChannelModel>> Channels;
+	for (const FKeyHandle KeyHandle : InSelection->KeySelection)
+	{
+		const TViewModelPtr<FChannelModel> Channel = InSelection->KeySelection.GetModelForKey(KeyHandle);
+		Channels.Add(Channel);
+	}
+	if (Channels.Num() == 1)
+	{
+		// NOTE: Can't wrap with a section because the BuildSidebarMenu functions below add sections.
+			
+		if (const TSharedPtr<IKeyArea> KeyArea = Channels[0]->GetKeyArea())
+		{
+			if (const TViewModelPtr<IOutlinerExtension> LinkedOutlinerItem = Channels[0]->GetLinkedOutlinerItem())
+			{
+				if (const TViewModelPtr<FOutlinerItemModel> OutlinerItemModel = LinkedOutlinerItem.ImplicitCast())
+				{
+					OutlinerItemModel->BuildTrackOptionsMenu(MenuBuilder);
+					OutlinerItemModel->BuildDisplayOptionsMenu(MenuBuilder);
+				}
+			}
+		}
+			
+		if (const TViewModelPtr<FTrackModel>& TrackModel = Channels[0].AsModel()->FindAncestorOfType<FTrackModel>())
+		{
+			TrackModel->BuildSidebarMenu(MenuBuilder);
+		}
+	}
+}
+
+void FSequencerSelectionDrawer::BuildBindingPropertiesDetails(FSequencer& InSequencer, const TSharedRef<UE::Sequencer::FSequencerSelection>& InSelection, FMenuBuilder& MenuBuilder)
+{
+	TArray<FGuid> ObjectBindings;
+	InSequencer.GetSelectedObjects(ObjectBindings);
+	if (ObjectBindings.Num() == 0)
+	{
+		return;
+	}
+
+	UMovieSceneSequence* const FocusedMovieSceneSequence = InSequencer.GetFocusedMovieSceneSequence();
+	if (!IsValid(FocusedMovieSceneSequence) || !FocusedMovieSceneSequence->AllowsSpawnableObjects())
+	{
+		return;
+	}
+
+	if (InSelection->Outliner.Num() == 1)
+	{
+		for (const FViewModelPtr Possessable : InSelection->Outliner.Filter<FPossessableModel>())
+		{
+			MenuBuilder.BeginSection(TEXT("Possessable"));
+			MenuBuilder.EndSection();
+
+			// Adding this causes two sets of details for the property bindings to show since
+			// ULevelSequenceEditorSubsystem::Initialize adds extension hooks "Possessable" and "CustomBinding"
+			// that call the exact same function. Will add this if there is a reason for this or remove this
+			// if not needed.
+			//MenuBuilder.BeginSection(TEXT("CustomBinding"));
+			//MenuBuilder.EndSection();
+		}
+	}
+}
+
+void FSequencerSelectionDrawer::BuildTrackAreaDetails(FSequencer& InSequencer, const TSharedRef<UE::Sequencer::FSequencerSelection>& InSelection, FMenuBuilder& MenuBuilder)
+{
+	if (InSelection->TrackArea.Num() == 0)
+	{
+		return;
+	}
+
+	TArray<TWeakObjectPtr<>> AllSectionObjects;
+	TArray<TViewModelPtr<FSectionModel>> AllSectionModels;
+	
+	for (const FViewModelPtr TrackAreaItem : InSelection->TrackArea)
+	{
+		if (const TViewModelPtr<FLayerBarModel> LayerBarModel = TrackAreaItem.ImplicitCast())
+		{
+			const TViewModelPtr<IOutlinerExtension> LinkedOutlinerItem = LayerBarModel->GetLinkedOutlinerItem();
+
+			if (InSelection->TrackArea.Num() == 1)
+			{
+				if (const TViewModelPtr<FOutlinerItemModel> OutlinerItemModel = LinkedOutlinerItem.ImplicitCast())
+				{
+					OutlinerItemModel->BuildTrackOptionsMenu(MenuBuilder);
+					OutlinerItemModel->BuildDisplayOptionsMenu(MenuBuilder);
+				}
+			}
+		}
+		else if (const TViewModelPtr<FSectionModel> SectionModel = TrackAreaItem.ImplicitCast())
+		{
+			// ISequencerSection Details (Shot Takes, Etc.)
+			if (const TSharedPtr<ISequencerSection> SectionInterface = SectionModel->GetSectionInterface())
+			{
+				AllSectionModels.Add(SectionModel);
+			}
+			
+			// Gather track section to use to build a single details for all selected sections
+			AllSectionObjects.Add(SectionModel->GetSection());
+		}
+	}
+
+	if (AllSectionModels.Num() == 1)
+	{
+		if (const TSharedPtr<ISequencerSection> SectionInterface = AllSectionModels[0]->GetSectionInterface())
+		{
+			const TViewModelPtr<IObjectBindingExtension> ObjectBindingModel = AllSectionModels[0]->FindAncestorOfType<IObjectBindingExtension>();
+			const FGuid ObjectGuid = ObjectBindingModel ? ObjectBindingModel->GetObjectGuid() : FGuid();
+			SectionInterface->BuildSectionSidebarMenu(MenuBuilder, ObjectGuid);
+		}
+	}
+	
+	SequencerHelpers::AddPropertiesMenu(InSequencer, MenuBuilder, AllSectionObjects);
+}
+
+void FSequencerSelectionDrawer::BuildOutlinerDetails(FSequencer& InSequencer, const TSharedRef<UE::Sequencer::FSequencerSelection>& InSelection, FMenuBuilder& MenuBuilder)
+{
+	if (InSelection->Outliner.Num() == 0)
+	{
+		return;
+	}
+
+	TSet<TViewModelPtr<FChannelGroupOutlinerModel>> ChannelGroups;
+
+	for (const FViewModelPtr OutlinerItem : InSelection->Outliner)
+	{
+		if (const TViewModelPtr<FTrackModel> TrackModel = OutlinerItem.ImplicitCast())
+		{
+			if (const TSharedPtr<ISequencerTrackEditor> TrackEditor = TrackModel->GetTrackEditor())
+			{
+				TrackEditor->BuildTrackSidebarMenu(MenuBuilder, TrackModel->GetTrack());
+			}
+		}
+
+		if (InSelection->Outliner.Num() == 1)
+		{
+			if (const TViewModelPtr<FOutlinerItemModel> OutlinerItemModel = OutlinerItem.ImplicitCast())
+			{
+				OutlinerItemModel->BuildTrackOptionsMenu(MenuBuilder);
+				OutlinerItemModel->BuildDisplayOptionsMenu(MenuBuilder);
+			}
+		}
+
+		// Ex. "Location.X"
+		if (const TViewModelPtr<FChannelGroupOutlinerModel> ChannelGroupOutlinerModel = OutlinerItem.ImplicitCast())
+		{
+			ChannelGroups.Add(ChannelGroupOutlinerModel);
+		}
+	}
+
+	if (!ChannelGroups.IsEmpty())
+	{
+		const ISequencerModule& SequencerModule = FModuleManager::Get().LoadModuleChecked<ISequencerModule>(TEXT("Sequencer"));
+		const TSharedPtr<FExtensibilityManager> SidebarExtensibilityManager = SequencerModule.GetSidebarExtensibilityManager();
+		const TSharedPtr<FExtender> Extender = SidebarExtensibilityManager->GetAllExtenders();
+
+		TArray<ISequencerChannelInterface*> ChannelInterfaces;
+		TArray<FMovieSceneChannelHandle> ChannelHandles;
+		TArray<UMovieSceneSection*> SceneSections;
+
+		for (const TViewModelPtr<FChannelGroupOutlinerModel>& ChannelModel : ChannelGroups)
+		{
+			for (const TSharedRef<IKeyArea>& KeyArea : ChannelModel->GetAllKeyAreas())
+			{
+				if (ISequencerChannelInterface* const SequencerChannelIterface = KeyArea->FindChannelEditorInterface())
+				{
+					ChannelInterfaces.Add(SequencerChannelIterface);
+					ChannelHandles.Add(KeyArea->GetChannel());
+					SceneSections.Add(KeyArea->GetOwningSection());
+				}
+			}
+		}
+
+		// Channel Interface Extensions (Perlin Noise, Easing, Wave)
+		for (const ISequencerChannelInterface* const ChannelInterface : ChannelInterfaces)
+		{
+			ChannelInterface->ExtendSidebarMenu_Raw(MenuBuilder, Extender, ChannelHandles, SceneSections, SequencerWeak);
+		}
+
+		// Curve Channel Options (Pre-Finity, Post-Finity, etc.)
+		CurveChannelExtension = MakeShared<FCurveChannelSectionSidebarExtension>(SequencerWeak);
+		CurveChannelExtension->AddSections(SceneSections);
+		CurveChannelExtension->ExtendMenu(MenuBuilder);
+	}
+}
+
+void FSequencerSelectionDrawer::BuildMarkedFrameDetails(const TSharedRef<UE::Sequencer::FSequencerSelection>& InSelection, FMenuBuilder& MenuBuilder)
+{
+	if (InSelection->MarkedFrames.Num() == 0)
+	{
+		return;
+	}
+
+	MenuBuilder.BeginSection(TEXT("MarkedFrames"), LOCTEXT("MarkedFramesMenuSection", "Marked Frames"));
+	
+	for (const int32 MarkIndex : InSelection->MarkedFrames)
+	{
+		MenuBuilder.AddWidget(CreateMarkedFrameDetails(MarkIndex).ToSharedRef(), FText::GetEmpty(), /*bInNoIndent=*/true);
+	}
+
+	MenuBuilder.EndSection();
+}
+
+TSharedRef<SWidget> FSequencerSelectionDrawer::CreateHintText(const FText& InMessage)
+{
+	return SNew(SBox)
+		.HAlign(HAlign_Center)
+		.Padding(2.f, 12.f, 2.f, 12.f)
+		[
+			SNew(STextBlock)
+			.Text(InMessage)
+			.TextStyle(FAppStyle::Get(), "HintText")
+		];
+}
+
+TSharedRef<SWidget> FSequencerSelectionDrawer::CreateNoSelectionHintText()
+{
+	return CreateHintText(LOCTEXT("NoSelection", "Select an object to view details."));
+}
+
+FKeyEditData FSequencerSelectionDrawer::GetKeyEditData() const
+{
+	const TSharedPtr<ISequencer> Sequencer = SequencerWeak.Pin();
+	if (!Sequencer.IsValid())
+	{
+		return FKeyEditData();
+	}
+
+	const TSharedPtr<UE::Sequencer::FSequencerSelection> SequencerSelection = UE::Sequencer::Private::GetSelection(*Sequencer.Get());
+	if (!SequencerSelection.IsValid())
+	{
+		return FKeyEditData();
+	}
+
+	return UE::Sequencer::Private::GetKeyEditData(SequencerSelection->KeySelection);
+}
+
+TSharedPtr<SWidget> FSequencerSelectionDrawer::CreateKeyFrameDetails(const TSharedRef<UE::Sequencer::FSequencerSelection>& InSequencerSelection)
+{
+	const TSharedPtr<ISequencer> Sequencer = SequencerWeak.Pin();
+	if (!Sequencer.IsValid())
+	{
+		return nullptr;
+	}
+
+	const FKeyEditData KeyEditData = UE::Sequencer::Private::GetKeyEditData(InSequencerSelection->KeySelection);
+	if (KeyEditData.KeyStruct.IsValid())
+	{
+		return SNew(SKeyEditInterface, Sequencer.ToSharedRef())
+			.EditData(this, &FSequencerSelectionDrawer::GetKeyEditData);
+	}
+
+	return CreateHintText(LOCTEXT("InvalidKeyCombination", "Selected keys must belong to the same section."));
+}
+
+TSharedPtr<SWidget> FSequencerSelectionDrawer::CreateMarkedFrameDetails(const int32 InMarkedFrameIndex)
+{
+	const TSharedPtr<FSequencer> Sequencer = SequencerWeak.Pin();
+	if (!Sequencer.IsValid())
+	{
+		return nullptr;
+	}
+
+	const UMovieSceneSequence* const FocusedMovieSceneSequence = Sequencer->GetFocusedMovieSceneSequence();
+	if (!IsValid(FocusedMovieSceneSequence))
+	{
+		return nullptr;
+	}
+
+	UMovieScene* const FocusedMovieScene = FocusedMovieSceneSequence->GetMovieScene();
+	if (!IsValid(FocusedMovieScene))
+	{
+		return nullptr;
+	}
+
+	if (FocusedMovieScene->GetMarkedFrames().Num() == 0)
+	{
+		return nullptr;
+	}
+
+	class SMarkedFramePropertyWidget : public SCompoundWidget, public FNotifyHook
+	{
+	public:
+		SLATE_BEGIN_ARGS(SMarkedFramePropertyWidget) {}
+		SLATE_END_ARGS()
+
+		void Construct(const FArguments& InArgs, UMovieScene* const InMovieScene, const int32 InMarkedFrameIndex, const TWeakPtr<FSequencer>& InWeakSequencer)
+		{
+			MovieSceneToModify = InMovieScene;
+			WeakSequencer = InWeakSequencer;
+
+			FDetailsViewArgs DetailsViewArgs;
+			DetailsViewArgs.bAllowSearch = false;
+			DetailsViewArgs.bShowScrollBar = false;
+			DetailsViewArgs.NameAreaSettings = FDetailsViewArgs::HideNameArea;
+			DetailsViewArgs.NotifyHook = this;
+
+			FStructureDetailsViewArgs StructureDetailsViewArgs;
+			StructureDetailsViewArgs.bShowObjects = true;
+			StructureDetailsViewArgs.bShowAssets = true;
+			StructureDetailsViewArgs.bShowClasses = true;
+			StructureDetailsViewArgs.bShowInterfaces = true;
+			
+			const TSharedPtr<FStructOnScope> StructOnScope = MakeShared<FStructOnScope>(FMovieSceneMarkedFrame::StaticStruct(), (uint8*)&InMovieScene->GetMarkedFrames()[InMarkedFrameIndex]);
+
+			FPropertyEditorModule& PropertyEditorModule = FModuleManager::GetModuleChecked<FPropertyEditorModule>(TEXT("PropertyEditor"));
+
+			DetailsView = PropertyEditorModule.CreateStructureDetailView(DetailsViewArgs, StructureDetailsViewArgs, nullptr);
+			DetailsView->GetDetailsView()->RegisterInstancedCustomPropertyTypeLayout(TEXT("FrameNumber"), FOnGetPropertyTypeCustomizationInstance::CreateLambda([this]() {
+				return MakeShared<FFrameNumberDetailsCustomization>(WeakSequencer.Pin()->GetNumericTypeInterface()); }));
+			DetailsView->SetStructureData(StructOnScope);
+			
+			ChildSlot
+			[
+				DetailsView->GetWidget().ToSharedRef()
+			];
+		}
+
+		virtual void NotifyPreChange(FProperty* InPropertyAboutToChange) override
+		{
+			MovieSceneToModify->Modify();
+		}
+
+		virtual void NotifyPreChange(FEditPropertyChain* InPropertyAboutToChange) override
+		{
+			MovieSceneToModify->Modify();
+		}
+
+	private:
+		TObjectPtr<UMovieScene> MovieSceneToModify;
+		TSharedPtr<IStructureDetailsView> DetailsView;
+		TWeakPtr<FSequencer> WeakSequencer;
+	};
+
+	auto AreMarkedFramesLocked = [&Sequencer]() -> bool
+	{
+		if (Sequencer->IsReadOnly())
+		{
+			return true;
+		}
+
+		const UMovieSceneSequence* const FocusedMovieSceneSequence = Sequencer->GetFocusedMovieSceneSequence();
+		if (FocusedMovieSceneSequence != nullptr)
+		{
+			const UMovieScene* const MovieScene = FocusedMovieSceneSequence->GetMovieScene();
+			if (MovieScene->IsReadOnly())
+			{
+				return true;
+			}
+			return MovieScene->AreMarkedFramesLocked();
+		}
+
+		return false;
+	};
+
+	const TSharedRef<SMarkedFramePropertyWidget> Widget = SNew(SMarkedFramePropertyWidget, FocusedMovieScene, InMarkedFrameIndex, SequencerWeak);
+	Widget->SetEnabled(!AreMarkedFramesLocked());
+
+	return Widget;
+}
+
+#undef LOCTEXT_NAMESPACE
diff --git a/Engine/Source/Editor/Sequencer/Private/Widgets/Sidebar/SequencerSelectionDrawer.h b/Engine/Source/Editor/Sequencer/Private/Widgets/Sidebar/SequencerSelectionDrawer.h
new file mode 100644
--- /dev/null
+++ b/Engine/Source/Editor/Sequencer/Private/Widgets/Sidebar/SequencerSelectionDrawer.h
@@ -0,0 +1,62 @@
+// Copyright Epic Games, Inc. All Rights Reserved.
+
+#pragma once
+
+#include "Sidebar/ISidebarDrawerContent.h"
+#include "Templates/SharedPointerFwd.h"
+
+class FMenuBuilder;
+class FName;
+class FSequencer;
+class FText;
+class IDetailsView;
+class SVerticalBox;
+class SWidget;
+class UMovieSceneFolder;
+class UMovieSceneSequence;
+struct FCurveChannelSectionSidebarExtension;
+struct FKeyEditData;
+struct FMovieSceneMarkedFrame;
+
+namespace UE::Sequencer
+{
+	class FSequencerSelection;
+}
+
+class FSequencerSelectionDrawer : public ISidebarDrawerContent
+{
+public:
+	static const FName UniqueId;
+
+	FSequencerSelectionDrawer(const TWeakPtr<FSequencer>& InSequencerWeak);
+
+	//~ Begin ISidebarDrawerContent
+	virtual FName GetUniqueId() const override;
+	virtual FName GetSectionId() const override;
+	virtual FText GetSectionDisplayText() const override;
+	virtual TSharedRef<SWidget> CreateContentWidget() override;
+	//~ End ISidebarDrawerContent
+
+protected:
+	void OnSequencerSelectionChanged();
+
+	void BuildKeySelectionDetails(const TSharedRef<UE::Sequencer::FSequencerSelection>& InSelection, FMenuBuilder& MenuBuilder);
+	void BuildBindingPropertiesDetails(FSequencer& InSequencer, const TSharedRef<UE::Sequencer::FSequencerSelection>& InSelection, FMenuBuilder& MenuBuilder);
+	void BuildTrackAreaDetails(FSequencer& InSequencer, const TSharedRef<UE::Sequencer::FSequencerSelection>& InSelection, FMenuBuilder& MenuBuilder);
+	void BuildOutlinerDetails(FSequencer& InSequencer, const TSharedRef<UE::Sequencer::FSequencerSelection>& InSelection, FMenuBuilder& MenuBuilder);
+	void BuildMarkedFrameDetails(const TSharedRef<UE::Sequencer::FSequencerSelection>& InSelection, FMenuBuilder& MenuBuilder);
+
+	TSharedRef<SWidget> CreateHintText(const FText& InMessage);
+	TSharedRef<SWidget> CreateNoSelectionHintText();
+
+	FKeyEditData GetKeyEditData() const;
+	
+	TSharedPtr<SWidget> CreateKeyFrameDetails(const TSharedRef<UE::Sequencer::FSequencerSelection>& InSequencerSelection);
+	TSharedPtr<SWidget> CreateMarkedFrameDetails(const int32 InMarkedFrameIndex);
+
+	TWeakPtr<FSequencer> SequencerWeak;
+
+	TSharedPtr<SVerticalBox> ContentBox;
+
+	TSharedPtr<FCurveChannelSectionSidebarExtension> CurveChannelExtension;
+};
diff --git a/Engine/Source/Editor/Sequencer/Public/ISequencer.h b/Engine/Source/Editor/Sequencer/Public/ISequencer.h
--- a/Engine/Source/Editor/Sequencer/Public/ISequencer.h
+++ b/Engine/Source/Editor/Sequencer/Public/ISequencer.h
@@ -46,6 +46,7 @@ class FSequencerSelectionPreview;
 class FUICommandList;
 class IDetailsView;
 class IKeyArea;
+class ISidebarDrawerContent;
 class ISequencerTrackEditor;
 class SWidget;
 class UActorFactory;
@@ -65,6 +66,7 @@ struct FMovieSceneBinding;
 struct FMovieSceneChannelHandle;
 struct FMovieSceneMarkedFrame;
 struct FQualifiedFrameTime;
+struct FSidebarDrawerConfig;
 template <typename NumericType> struct INumericTypeInterface;
 
 enum class EMapChangeType : uint8;
@@ -888,6 +890,44 @@ public:
 	/** @return Delegate executed when Sequencer selection limiting is enabled or disabled. */
 	virtual FOnViewportSelectionLimitedChanged& OnViewportSelectionLimitedChanged() = 0;
 
+	/**
+	 * Registers and displays a new drawer in the sidebar.
+	 * 
+	 * @param InDrawerConfig Configuration info for the new drawer
+	 * 
+	 * @return True if the new drawer registration was successful.
+	 */
+	virtual bool RegisterDrawer(FSidebarDrawerConfig&& InDrawerConfig) = 0;
+
+	/**
+	 * Unregisters and removes a drawer from the sidebar.
+	 *
+	 * @param InDrawerId Unique drawer Id to unregister
+	 * 
+	 * @return True if the drawer removal was successful.
+	 */
+	virtual bool UnregisterDrawer(const FName InDrawerId) = 0;
+
+	/**
+	 * Registers and displays a new drawer section in the sidebar.
+	 * 
+	 * @param InDrawerId Unique drawer Id to register
+	 * @param InSection Drawer content interface for the section
+	 * 
+	 * @return True if the new drawer section registration was successful.
+	 */
+	virtual bool RegisterDrawerSection(const FName InDrawerId, const TSharedPtr<ISidebarDrawerContent>& InSection) = 0;
+
+	/**
+	 * Unregisters and removes a drawer section from the sidebar.
+	 * 
+	 * @param InDrawerId Unique drawer Id that contains the section to unregister
+	 * @param InSectionId Unique drawer section Id to unregister
+	 * 
+	 * @return True if the drawer removal was successful.
+	 */
+	virtual bool UnregisterDrawerSection(const FName InDrawerId, const FName InSectionId) = 0;
+
 protected:
 	FOnInitializeDetailsPanel InitializeDetailsPanelEvent;
 	FOnCameraAddedToSequencer CameraAddedToSequencer;
diff --git a/Engine/Source/Editor/Sequencer/Public/ISequencerChannelInterface.h b/Engine/Source/Editor/Sequencer/Public/ISequencerChannelInterface.h
--- a/Engine/Source/Editor/Sequencer/Public/ISequencerChannelInterface.h
+++ b/Engine/Source/Editor/Sequencer/Public/ISequencerChannelInterface.h
@@ -150,12 +150,24 @@ struct ISequencerChannelInterface
 	 * Extend the section context menu
 	 *
 	 * @param MenuBuilder           The menu builder used to create this context menu
+	 * @param MenuExtender          The menu extender to use
 	 * @param Channels              Array of type specific channels that exist in the selected sections
 	 * @param Sections              Array of sections being shown on the context menu
 	 * @param InSequencer           The currently active sequencer
 	 */
 	virtual void ExtendSectionMenu_Raw(FMenuBuilder& MenuBuilder, TSharedPtr<FExtender> MenuExtender, TArrayView<const FMovieSceneChannelHandle> Channels, TArrayView<UMovieSceneSection* const> Sections, TWeakPtr<ISequencer> InSequencer) const = 0;
 
+	/**
+	 * Extend the section sidebar menu
+	 *
+	 * @param MenuBuilder           The menu builder used to create this context menu
+	 * @param MenuExtender          The menu extender to use
+	 * @param Channels              Array of type specific channels that exist in the selected sections
+	 * @param Sections              Array of sections being shown on the context menu
+	 * @param InSequencer           The currently active sequencer
+	 */
+	virtual void ExtendSidebarMenu_Raw(FMenuBuilder& MenuBuilder, TSharedPtr<FExtender> MenuExtender, TArrayView<const FMovieSceneChannelHandle> Channels, TArrayView<UMovieSceneSection* const> Sections, TWeakPtr<ISequencer> InSequencer) const = 0;
+
 	/**
 	 * Gather information on how to draw the specified keys
 	 *
diff --git a/Engine/Source/Editor/Sequencer/Public/ISequencerModule.h b/Engine/Source/Editor/Sequencer/Public/ISequencerModule.h
--- a/Engine/Source/Editor/Sequencer/Public/ISequencerModule.h
+++ b/Engine/Source/Editor/Sequencer/Public/ISequencerModule.h
@@ -347,6 +347,13 @@ public:
 	 */
 	virtual TSharedPtr<FExtensibilityManager> GetActionsMenuExtensibilityManager() const = 0;
 
+	/**
+	 * Get the extensibility manager for the sidebar.
+	 *
+	 * @return Sidebar extensibility manager.
+	 */
+	virtual TSharedPtr<FExtensibilityManager> GetSidebarExtensibilityManager() const = 0;
+
 	/**
 	 * Get the sequencer customization manager, which handles editor customizations applied based on
 	 * the currently focused sequence type and other dynamic criteria.
diff --git a/Engine/Source/Editor/Sequencer/Public/ISequencerSection.h b/Engine/Source/Editor/Sequencer/Public/ISequencerSection.h
--- a/Engine/Source/Editor/Sequencer/Public/ISequencerSection.h
+++ b/Engine/Source/Editor/Sequencer/Public/ISequencerSection.h
@@ -75,15 +75,15 @@ namespace SequencerSectionConstants
  */
 struct FSequencerSectionPropertyDetailsViewCustomizationParams
 {
-	FSequencerSectionPropertyDetailsViewCustomizationParams(TSharedRef<ISequencerSection> InSectionInterface, TSharedRef<ISequencer> InSequencer, ISequencerTrackEditor& InTrackEditor)
+	FSequencerSectionPropertyDetailsViewCustomizationParams(TSharedRef<ISequencerSection> InSectionInterface, TWeakPtr<ISequencer> InSequencerWeak, ISequencerTrackEditor& InTrackEditor)
 		: SectionInterface(InSectionInterface)
-		, Sequencer(InSequencer)
+		, SequencerWeak(InSequencerWeak)
 		, TrackEditor(InTrackEditor)
 	{}
 
 	FGuid ParentObjectBindingGuid;
 	TSharedRef<ISequencerSection> SectionInterface;
-	TSharedRef<ISequencer> Sequencer;
+	TWeakPtr<ISequencer> SequencerWeak;
 	ISequencerTrackEditor& TrackEditor;
 };
 
@@ -235,6 +235,14 @@ public:
 	 */
 	virtual void BuildSectionContextMenu(FMenuBuilder& MenuBuilder, const FGuid& ObjectBinding) {}
 
+	/**
+	 * Builds up the section sidebar menu for the outliner
+	 *
+	 * @param MenuBuilder	The menu builder to change
+	 * @param ObjectBinding The object guid bound to this section
+	 */
+	virtual void BuildSectionSidebarMenu(FMenuBuilder& MenuBuilder, const FGuid& ObjectBinding) {}
+
 	/**
 	 * Called when the user requests that a category from this section be deleted. 
 	 *
diff --git a/Engine/Source/Editor/Sequencer/Public/ISequencerTrackEditor.h b/Engine/Source/Editor/Sequencer/Public/ISequencerTrackEditor.h
--- a/Engine/Source/Editor/Sequencer/Public/ISequencerTrackEditor.h
+++ b/Engine/Source/Editor/Sequencer/Public/ISequencerTrackEditor.h
@@ -196,7 +196,13 @@ public:
 	 * Builds the context menu for the track.
 	 * @param MenuBuilder The menu builder to use to build the track menu. 
 	 */
-	virtual void BuildTrackContextMenu( FMenuBuilder& MenuBuilder, UMovieSceneTrack* Track ) = 0;
+	virtual void BuildTrackContextMenu(FMenuBuilder& MenuBuilder, UMovieSceneTrack* Track) = 0;
+
+	/**
+	 * Builds the sidebar menu for the track.
+	 * @param MenuBuilder The menu builder to use to build the track menu. 
+	 */
+	virtual void BuildTrackSidebarMenu(FMenuBuilder& MenuBuilder, UMovieSceneTrack* Track) = 0;
 
 	/**
 	 * Called when an asset is dropped into Sequencer. Can potentially consume the asset
diff --git a/Engine/Source/Editor/Sequencer/Public/MVVM/ViewModels/FolderModel.h b/Engine/Source/Editor/Sequencer/Public/MVVM/ViewModels/FolderModel.h
--- a/Engine/Source/Editor/Sequencer/Public/MVVM/ViewModels/FolderModel.h
+++ b/Engine/Source/Editor/Sequencer/Public/MVVM/ViewModels/FolderModel.h
@@ -60,6 +60,7 @@ public:
 
 	/*~ FOutlinerItemModel */
 	void BuildContextMenu(FMenuBuilder& MenuBuilder) override;
+	void BuildSidebarMenu(FMenuBuilder& MenuBuilder) override;
 
 	/*~ IFolderEventHandler */
 	void OnTrackAdded(UMovieSceneTrack* Track) override;
diff --git a/Engine/Source/Editor/Sequencer/Public/MVVM/ViewModels/ObjectBindingModel.h b/Engine/Source/Editor/Sequencer/Public/MVVM/ViewModels/ObjectBindingModel.h
--- a/Engine/Source/Editor/Sequencer/Public/MVVM/ViewModels/ObjectBindingModel.h
+++ b/Engine/Source/Editor/Sequencer/Public/MVVM/ViewModels/ObjectBindingModel.h
@@ -102,6 +102,7 @@ public:
 	void BuildContextMenu(FMenuBuilder& MenuBuilder) override;
 	void BuildOrganizeContextMenu(FMenuBuilder& MenuBuilder) override;
 	bool GetDefaultExpansionState() const override;
+	void BuildSidebarMenu(FMenuBuilder& MenuBuilder) override;
 
 	/*~ IDraggableOutlinerExtension */
 	bool CanDrag() const override;
diff --git a/Engine/Source/Editor/Sequencer/Public/MVVM/ViewModels/OutlinerItemModel.h b/Engine/Source/Editor/Sequencer/Public/MVVM/ViewModels/OutlinerItemModel.h
--- a/Engine/Source/Editor/Sequencer/Public/MVVM/ViewModels/OutlinerItemModel.h
+++ b/Engine/Source/Editor/Sequencer/Public/MVVM/ViewModels/OutlinerItemModel.h
@@ -16,7 +16,10 @@
 #include "CurveEditorTypes.h"
 #include "Tree/ICurveEditorTreeItem.h"
 
+class UMovieSceneSection;
 class UMovieSceneSequence;
+class UMovieSceneTrack;
+class FBoolProperty;
 class FSequencer;
 
 namespace UE
@@ -67,12 +70,16 @@ public:
 	/*~ IDimmableExtension */
 	bool IsDimmed() const override;
 
-protected:
+	virtual void BuildSidebarMenu(FMenuBuilder& MenuBuilder);
+
+//protected:
 
 	/** Get context menu contents. */
 	virtual void BuildContextMenu(FMenuBuilder& MenuBuilder);
 	virtual void BuildOrganizeContextMenu(FMenuBuilder& MenuBuilder);
-
+	virtual void BuildDisplayOptionsMenu(FMenuBuilder& MenuBuilder);
+	virtual void BuildTrackOptionsMenu(FMenuBuilder& MenuBuilder);
+protected:
 	/** Set identifier for computing node paths */
 	void SetIdentifier(FName InNewIdentifier);
 
@@ -103,6 +110,12 @@ private:
 	ECheckBoxState SelectedModelsMuteState() const;
 	void ToggleSelectedModelsMuted();
 
+	TArray<UMovieSceneSection*> GetSelectedSections() const;
+	TArray<UMovieSceneTrack*> GetSelectedTracks() const;
+
+	void AddEvalOptionsPropertyMenuItem(FMenuBuilder& InMenuBuilder, const FBoolProperty* InProperty, TFunction<bool(UMovieSceneTrack*)> InValidator = nullptr);
+	void AddDisplayOptionsPropertyMenuItem(FMenuBuilder& InMenuBuilder, const FBoolProperty* InProperty, TFunction<bool(UMovieSceneTrack*)> InValidator = nullptr);
+
 private:
 
 	ICastable* CastableThis;
diff --git a/Engine/Source/Editor/Sequencer/Public/MVVM/ViewModels/TrackModel.h b/Engine/Source/Editor/Sequencer/Public/MVVM/ViewModels/TrackModel.h
--- a/Engine/Source/Editor/Sequencer/Public/MVVM/ViewModels/TrackModel.h
+++ b/Engine/Source/Editor/Sequencer/Public/MVVM/ViewModels/TrackModel.h
@@ -135,6 +135,7 @@ public:
 	bool HasCurves() const override;
 	void BuildContextMenu(FMenuBuilder& MenuBuilder) override;
 	bool GetDefaultExpansionState() const override;
+	void BuildSidebarMenu(FMenuBuilder& MenuBuilder) override;
 
 	/*~ IDeletableExtension */
 	bool CanDelete(FText* OutErrorMessage) const override;
@@ -143,12 +144,16 @@ public:
 	/*~ FViewModel interface */
 	virtual void OnConstruct() override;
 
+	virtual void BuildBlendingMenu(FMenuBuilder& MenuBuilder);
+
 private:
 
 	void ForceUpdate();
 
 	bool FindBoundObjects(TArray<UObject*>& OutBoundObjects) const;
 
+	TArray<TWeakObjectPtr<UObject>> GetSelectedTrackSections() const;
+
 	/** A second children list for the sections inside this track */
 	FViewModelListHead SectionList;
 	FViewModelListHead TopLevelChannelList;
diff --git a/Engine/Source/Editor/Sequencer/Public/MVVM/ViewModels/TrackRowModel.h b/Engine/Source/Editor/Sequencer/Public/MVVM/ViewModels/TrackRowModel.h
--- a/Engine/Source/Editor/Sequencer/Public/MVVM/ViewModels/TrackRowModel.h
+++ b/Engine/Source/Editor/Sequencer/Public/MVVM/ViewModels/TrackRowModel.h
@@ -48,6 +48,7 @@ public:
 	/*~ FOutlinerItemModel */
 	void BuildContextMenu(FMenuBuilder& MenuBuilder) override;
 	void CreateCurveModels(TArray<TUniquePtr<FCurveModel>>& OutCurveModels) override;
+	void BuildSidebarMenu(FMenuBuilder& MenuBuilder) override;
 
 	/*~ IOutlinerExtension */
 	FOutlinerSizing GetOutlinerSizing() const override;
@@ -90,6 +91,8 @@ public:
 
 private:
 
+	TArray<TWeakObjectPtr<UObject>> GetSelectedTrackSections() const;
+	
 	FViewModelListHead SectionList;
 	FViewModelListHead TopLevelChannelList;
 	TWeakObjectPtr<UMovieSceneTrack> WeakTrack;
diff --git a/Engine/Source/Editor/Sequencer/Public/MovieSceneTrackEditor.h b/Engine/Source/Editor/Sequencer/Public/MovieSceneTrackEditor.h
--- a/Engine/Source/Editor/Sequencer/Public/MovieSceneTrackEditor.h
+++ b/Engine/Source/Editor/Sequencer/Public/MovieSceneTrackEditor.h
@@ -181,7 +181,8 @@ public:
 	virtual void BuildObjectBindingTrackMenu(FMenuBuilder& MenuBuilder, const TArray<FGuid>& ObjectBindings, const UClass* ObjectClass) override;
 	virtual TSharedPtr<SWidget> BuildOutlinerEditWidget(const FGuid& ObjectBinding, UMovieSceneTrack* Track, const FBuildEditWidgetParams& Params) override;
 	virtual TSharedPtr<SWidget> BuildOutlinerColumnWidget(const FBuildColumnWidgetParams& Params, const FName& ColumnName) override;
-	virtual void BuildTrackContextMenu( FMenuBuilder& MenuBuilder, UMovieSceneTrack* Track ) override;
+	virtual void BuildTrackContextMenu(FMenuBuilder& MenuBuilder, UMovieSceneTrack* Track) override;
+	virtual void BuildTrackSidebarMenu(FMenuBuilder& MenuBuilder, UMovieSceneTrack* Track) override;
 	virtual bool HandleAssetAdded(UObject* Asset, const FGuid& TargetObjectGuid) override;
 	virtual bool OnAllowDrop(const FDragDropEvent& DragDropEvent, FSequencerDragDropParams& DragDropParams) override;
 	virtual FReply OnDrop(const FDragDropEvent& DragDropEvent, const FSequencerDragDropParams& DragDropParams) override;
diff --git a/Engine/Source/Editor/Sequencer/Public/SequencerChannelInterface.h b/Engine/Source/Editor/Sequencer/Public/SequencerChannelInterface.h
--- a/Engine/Source/Editor/Sequencer/Public/SequencerChannelInterface.h
+++ b/Engine/Source/Editor/Sequencer/Public/SequencerChannelInterface.h
@@ -145,6 +145,7 @@ struct TSequencerChannelInterfaceCommon : ISequencerChannelInterface
 	 * Extend the section context menu
 	 *
 	 * @param MenuBuilder           The menu builder used to create this context menu
+	 * @param MenuExtender          The menu extender to use
 	 * @param Channels              Array of type specific channels that exist in the selected sections
 	 * @param Sections              Array of sections being shown on the context menu
 	 * @param InSequencer           The currently active sequencer
@@ -162,6 +163,28 @@ struct TSequencerChannelInterfaceCommon : ISequencerChannelInterface
 		ExtendSectionMenu(MenuBuilder, MenuExtender, MoveTemp(TypedChannels), Sections, InSequencer);
 	}
 
+	/**
+	 * Extend the section sidebar menu
+	 *
+	 * @param MenuBuilder           The menu builder used to create this context menu
+	 * @param MenuExtender          The menu extender to use
+	 * @param Channels              Array of type specific channels that exist in the selected sections
+	 * @param Sections              Array of sections being shown on the context menu
+	 * @param InSequencer           The currently active sequencer
+	 */
+	virtual void ExtendSidebarMenu_Raw(FMenuBuilder& MenuBuilder, TSharedPtr<FExtender> MenuExtender, TArrayView<const FMovieSceneChannelHandle> Channels, TArrayView<UMovieSceneSection* const> Sections, TWeakPtr<ISequencer> InSequencer) const override
+	{
+		using namespace Sequencer;
+		TArray<TMovieSceneChannelHandle<ChannelType>> TypedChannels;
+
+		for (const FMovieSceneChannelHandle& RawHandle : Channels)
+		{
+			TypedChannels.Add(RawHandle.Cast<ChannelType>());
+		}
+
+		ExtendSidebarMenu(MenuBuilder, MenuExtender, MoveTemp(TypedChannels), Sections, InSequencer);
+	}
+
 	/**
 	 * Gather information on how to draw the specified keys
 	 *
diff --git a/Engine/Source/Editor/Sequencer/Public/SequencerChannelTraits.h b/Engine/Source/Editor/Sequencer/Public/SequencerChannelTraits.h
--- a/Engine/Source/Editor/Sequencer/Public/SequencerChannelTraits.h
+++ b/Engine/Source/Editor/Sequencer/Public/SequencerChannelTraits.h
@@ -106,6 +106,17 @@ namespace Sequencer
 	void ExtendSectionMenu(FMenuBuilder& MenuBuilder, TSharedPtr<FExtender> MenuExtender, TArray<TMovieSceneChannelHandle<ChannelType>>&& Channels, TArrayView<UMovieSceneSection* const> Sections, TWeakPtr<ISequencer> InSequencer)
 	{}
 
+	/**
+	 * Extend the specified selected sidebar menu
+	 *
+	 * @param MenuBuilder    The menu builder that will construct the section context menu
+	 * @param Channels       An array of all channels that are currently selected, in no particular order
+	 * @param Sections       An array of all sections that the selected channels reside in
+	 * @param InSequencer    The sequencer that is currently active
+	 */
+	template<typename ChannelType>
+	void ExtendSidebarMenu(FMenuBuilder& MenuBuilder, TSharedPtr<FExtender> MenuExtender, TArray<TMovieSceneChannelHandle<ChannelType>>&& Channels, TArrayView<UMovieSceneSection* const> Sections, TWeakPtr<ISequencer> InSequencer)
+	{}
 
 	/**
 	 * Extend the specified selected key context menu
diff --git a/Engine/Source/Editor/Sequencer/Public/SequencerCommands.h b/Engine/Source/Editor/Sequencer/Public/SequencerCommands.h
--- a/Engine/Source/Editor/Sequencer/Public/SequencerCommands.h
+++ b/Engine/Source/Editor/Sequencer/Public/SequencerCommands.h
@@ -478,6 +478,12 @@ public:
 	/** Toggles selectability of viewport objects to Sequencer items only */
 	TSharedPtr< FUICommandInfo > ToggleLimitViewportSelection;
 
+	/** Toggles the sidebar "Selection" drawer open or closed */
+	TSharedPtr< FUICommandInfo > ToggleSidebarSelectionDrawerOpen;
+
+	/** Toggles the docked state of the current sidebar drawer. Undocks the docked sidebar drawer if docked or docks the sidebar drawer if there is one open and no currently docked drawer. */
+	TSharedPtr< FUICommandInfo > ToggleSidebarDrawerDock;
+
 	/**
 	 * Initialize commands
 	 */
diff --git a/Engine/Source/Editor/Sequencer/Public/SequencerCustomizationManager.h b/Engine/Source/Editor/Sequencer/Public/SequencerCustomizationManager.h
--- a/Engine/Source/Editor/Sequencer/Public/SequencerCustomizationManager.h
+++ b/Engine/Source/Editor/Sequencer/Public/SequencerCustomizationManager.h
@@ -62,9 +62,15 @@ struct FSequencerCustomizationInfo
 	/** Extender for the sequencer toolbar. */
 	TSharedPtr<FExtender> ToolbarExtender;
 
+	/** Extender for the sequencer sidebar. */
+	TSharedPtr<FExtender> SidebarExtender;
+
 	/** Extender for the object binding. */
 	FOnGetSequencerMenuExtender OnBuildObjectBindingContextMenu;
 
+	/** Called when the sequencer sidebar is being built. */
+	FOnGetSequencerMenuExtender OnBuildSidebarMenu;
+
 	/** Called when something is dragged over the sequencer. */
 	FOptionalOnDragDrop OnReceivedDragOver;
 
