Sequencer- Converting a spawnable to a possessable doesn't correctly copy the actor transform to the new possessable, causing a snap to origin when closing the sequence.

diff --git a/Engine/Source/Editor/Sequencer/Private/SequencerUtilities.cpp b/Engine/Source/Editor/Sequencer/Private/SequencerUtilities.cpp
--- a/Engine/Source/Editor/Sequencer/Private/SequencerUtilities.cpp
+++ b/Engine/Source/Editor/Sequencer/Private/SequencerUtilities.cpp
@@ -1115,7 +1115,8 @@ FMovieScenePossessable* FSequencerUtilities::ConvertToPossessable(TSharedRef<ISe
 		}
 	}
 
-	UObject* ObjectToConvert = MovieSceneHelpers::GetSingleBoundObject(Sequence, BindingGuid, Sequencer->GetSharedPlaybackState(), BindingIndex);
+	UObject* BoundObject = MovieSceneHelpers::GetSingleBoundObject(Sequence, BindingGuid, Sequencer->GetSharedPlaybackState(), BindingIndex);
+	UObject* ObjectToConvert = BoundObject;
 
 	// If we have an old-style spawnable, use the template as the object to convert instead.
 	bool bConvertFromSpawnable = MovieSceneHelpers::IsBoundToSpawnable(Sequence, BindingGuid, Sequencer->GetSharedPlaybackState(), BindingIndex);
@@ -1130,7 +1131,7 @@ FMovieScenePossessable* FSequencerUtilities::ConvertToPossessable(TSharedRef<ISe
 	FTransform DefaultTransform = SpawnableActorTemplate ? SpawnableActorTemplate->GetActorTransform() : FTransform();
 	// Prefer the transform at the current time over the spawnable actor template's transform because that's most likely 0. 
 	// This makes it so that the object will return to the current position on restore state.
-	AActor* Actor = Cast<AActor>(ObjectToConvert);
+	AActor* Actor = Cast<AActor>(BoundObject);
 	if (Actor)
 	{
 		if (Actor->GetRootComponent())
