Sequencer: Fix copy folders copying folders multiple times

diff --git a/Engine/Source/Editor/Sequencer/Private/SequencerUtilities.cpp b/Engine/Source/Editor/Sequencer/Private/SequencerUtilities.cpp
--- a/Engine/Source/Editor/Sequencer/Private/SequencerUtilities.cpp
+++ b/Engine/Source/Editor/Sequencer/Private/SequencerUtilities.cpp
@@ -1513,7 +1513,7 @@ void GatherChildFolders(UMovieSceneFolder* ParentFolder, TArray<UObject*>& Objec
 	{
 		if (ChildFolder)
 		{
-			Objects.Add(ChildFolder);
+			Objects.AddUnique(ChildFolder);
 
 			GatherChildFolders(ChildFolder, Objects);
 		}
@@ -1525,7 +1525,7 @@ void FSequencerUtilities::CopyFolders(const TArray<UMovieSceneFolder*>& Folders,
 	TArray<UObject*> Objects;
 	for (UMovieSceneFolder* Folder : Folders)
 	{
-		Objects.Add(Folder);
+		Objects.AddUnique(Folder);
 
 		GatherChildFolders(Folder, Objects);
 	}
