Sequencer: Remove call to UpdateBindingIDs when pasting bindings.

The original intention of calling UpdateBindingIDs is that if bindings are duplicated, existing references would be remapped to the duplicated bindings. For example, if Cube is attached to Sphere and both Cube and Sphere are duplicated, Cube2 should attached to Sphere2. However, this doesn't take into account if only one of the objects is duplicated. In that case, existing references shouldn't be remapped.

There isn't really a good way to determine whether sections should be remapped and it's causing more confusion at the moment than it's worth, so we should just remove it for now.

diff --git a/Engine/Source/Editor/Sequencer/Private/SequencerUtilities.cpp b/Engine/Source/Editor/Sequencer/Private/SequencerUtilities.cpp
--- a/Engine/Source/Editor/Sequencer/Private/SequencerUtilities.cpp
+++ b/Engine/Source/Editor/Sequencer/Private/SequencerUtilities.cpp
@@ -1985,7 +1985,6 @@ bool FSequencerUtilities::PasteSections(const FString& TextToImport, FMovieScene
 	for (int32 Index = 0; Index < PasteSectionsParams.Tracks.Num(); ++Index)
 	{
 		UMovieSceneTrack* Track = PasteSectionsParams.Tracks[Index];
-		int32 RowIndex = Index < PasteSectionsParams.TrackRowIndices.Num() ? PasteSectionsParams.TrackRowIndices[Index] : 0;
 
 		for (int32 SectionIndex = 0; SectionIndex < ImportedSections.Num(); ++SectionIndex)
 		{
@@ -1995,6 +1994,8 @@ bool FSequencerUtilities::PasteSections(const FString& TextToImport, FMovieScene
 				continue;
 			}
 
+			int32 RowIndex = SectionIndex < PasteSectionsParams.TrackRowIndices.Num() ? PasteSectionsParams.TrackRowIndices[SectionIndex] : Section->GetRowIndex();
+
 			SectionIndicesImported.AddUnique(SectionIndex);
 
 			Track->Modify();
@@ -3297,12 +3298,6 @@ bool FSequencerUtilities::PasteBindings(const FString& TextToImport, TSharedRef<
 		}
 	}
 
-	// Remap bindings in sections (ie. attach tracks)
-	for (TPair<FGuid, FGuid> GuidPair : OldToNewGuidMap)
-	{
-		UpdateBindingIDs(Sequencer, GuidPair.Key, GuidPair.Value);
-	}
-
 	for (auto BindingPasted : BindingsPasted)
 	{
 		OutBindings.Add(FMovieSceneBindingProxy(BindingPasted.GetObjectGuid(), Sequence));
