Sequencer: Prevent re-evaluation of entire ranges when stopping naturally

diff --git a/Engine/Source/Editor/Sequencer/Private/Sequencer.cpp b/Engine/Source/Editor/Sequencer/Private/Sequencer.cpp
--- a/Engine/Source/Editor/Sequencer/Private/Sequencer.cpp
+++ b/Engine/Source/Editor/Sequencer/Private/Sequencer.cpp
@@ -5108,7 +5108,7 @@ void FSequencer::SetLocalTimeLooped(FFrameTime NewLocalTime, const FMovieSceneTr
 	{
 		SetPlaybackStatus(NewPlaybackStatus.GetValue());
 		// Evaluate the sequence with the new status
-		EvaluateInternal(EvalRange);
+		EvaluateInternal(PlayPosition.GetCurrentPositionAsRange());
 
 		if (NewPlaybackStatus.GetValue() == EMovieScenePlayerStatus::Stopped)
 		{
