[Sequencer] Disable channel interface extensions for the sidebar until the control rig shape actor selection crash can be properly addressed.

diff --git a/Engine/Source/Editor/Sequencer/Private/Widgets/Sidebar/SequenceSelectionDrawer.cpp b/Engine/Source/Editor/Sequencer/Private/Widgets/Sidebar/SequenceSelectionDrawer.cpp
--- a/Engine/Source/Editor/Sequencer/Private/Widgets/Sidebar/SequenceSelectionDrawer.cpp
+++ b/Engine/Source/Editor/Sequencer/Private/Widgets/Sidebar/SequenceSelectionDrawer.cpp
@@ -414,10 +414,13 @@ void FSequencerSelectionDrawer::BuildOutlinerDetails(FSequencer& InSequencer, co
 		}
 
 		// Channel Interface Extensions (Perlin Noise, Easing, Wave)
-		for (const ISequencerChannelInterface* const ChannelInterface : ChannelInterfaces)
+		// Disabling this completely for now until a proper fix can be submitted
+		// Crashes when selecting the control rig ankle shape control (maybe some others, but tends to only happen for the ankle)
+		// for the default mannequinn full body control rig.
+		/*for (const ISequencerChannelInterface* const ChannelInterface : ChannelInterfaces)
 		{
 			ChannelInterface->ExtendSidebarMenu_Raw(MenuBuilder, Extender, ChannelHandles, SceneSections, SequencerWeak);
-		}
+		}*/
 
 		// Curve Channel Options (Pre-Finity, Post-Finity, etc.)
 		CurveChannelExtension = MakeShared<FCurveChannelSectionSidebarExtension>(SequencerWeak);
