Sequencer- fix UpdateVirtualGeometry being called with view model list in invalid state during undo.

diff --git a/Engine/Source/Editor/Sequencer/Private/SequencerNodeTree.cpp b/Engine/Source/Editor/Sequencer/Private/SequencerNodeTree.cpp
--- a/Engine/Source/Editor/Sequencer/Private/SequencerNodeTree.cpp
+++ b/Engine/Source/Editor/Sequencer/Private/SequencerNodeTree.cpp
@@ -147,9 +147,14 @@ void FSequencerNodeTree::Update()
 		SortableChild->SortChildren();
 	}
 
-	// Update all virtual geometries
-	// This must happen after the sorting
-	IGeometryExtension::UpdateVirtualGeometry(0.f, RootNode);
+	// Avoid updating geometry during an undo/redo, as we may have changed the nodes and they won't get updated until next frame.
+	// Any deleted nodes will be present in the hierarchy but garbage.
+	if (!GIsTransacting)
+	{
+		// Update all virtual geometries
+		// This must happen after the sorting
+		IGeometryExtension::UpdateVirtualGeometry(0.f, RootNode);
+	}
 
 	// Cache pinned state of nodes, needs to happen after OnTreeRefreshed
 	FPinnableExtensionShim::UpdateCachedPinnedState(RootNode);
