Sequencer: Round to frame number when adding keys from channel navigation buttons

diff --git a/Engine/Source/Editor/Sequencer/Private/MVVM/Views/SSequencerKeyNavigationButtons.h b/Engine/Source/Editor/Sequencer/Private/MVVM/Views/SSequencerKeyNavigationButtons.h
--- a/Engine/Source/Editor/Sequencer/Private/MVVM/Views/SSequencerKeyNavigationButtons.h
+++ b/Engine/Source/Editor/Sequencer/Private/MVVM/Views/SSequencerKeyNavigationButtons.h
@@ -96,7 +96,7 @@ public:
 		if (InModel && Sequencer)
 		{
 			FScopedTransaction Transaction(LOCTEXT("AddKeys", "Add Keys at Current Time"));
-			FAddKeyOperation::FromNode(InModel).Commit(Time.FrameNumber, *Sequencer);
+			FAddKeyOperation::FromNode(InModel).Commit(Time.RoundToFrame(), *Sequencer);
 		}
 	}
 
