Sequencer: Add the ability to change the spawnable class for multiple spawnables at once

Renamed FSequencerConvertBindingInfo to FSequencerChangeBindingInfo to use it for converting bindings and changing spawnable class. No need for backwards compatibility since it's not exposed.

diff --git a/Engine/Plugins/MovieScene/LevelSequenceEditor/Source/LevelSequenceEditor/Private/LevelSequenceEditorSubsystem.cpp b/Engine/Plugins/MovieScene/LevelSequenceEditor/Source/LevelSequenceEditor/Private/LevelSequenceEditorSubsystem.cpp
--- a/Engine/Plugins/MovieScene/LevelSequenceEditor/Source/LevelSequenceEditor/Private/LevelSequenceEditorSubsystem.cpp
+++ b/Engine/Plugins/MovieScene/LevelSequenceEditor/Source/LevelSequenceEditor/Private/LevelSequenceEditorSubsystem.cpp
@@ -242,7 +242,7 @@ TSharedRef<SWidget> FMovieSceneBindingPropertyInfoListCustomization::OnGetConver
 
 	if (Sequencer.IsValid() && BindingList)
 	{
-		TArray<FSequencerConvertBindingInfo> Bindings;
+		TArray<FSequencerChangeBindingInfo> Bindings;
 		for (int32 BindingIndex = 0; BindingIndex < BindingList->Bindings.Num(); ++BindingIndex)
 		{
 			Bindings.Add({ BindingGuid, BindingIndex });
@@ -391,7 +391,10 @@ TSharedRef<SWidget> FMovieSceneBindingPropertyInfoDetailCustomization::OnGetChan
 	TSharedPtr<ISequencer> Sequencer = SequencerPtr.Pin();
 	if (Sequencer.IsValid() && StructPropertyHandle->IsValidHandle())
 	{
-		FSequencerUtilities::AddChangeClassMenu(MenuBuilder, Sequencer.ToSharedRef(), BindingGuid, BindingIndex, [this, &StructBuilder, &CustomizationUtils]()
+		TArray<FSequencerChangeBindingInfo> Bindings;
+		Bindings.Add(FSequencerChangeBindingInfo(BindingGuid, BindingIndex));
+
+		FSequencerUtilities::AddChangeClassMenu(MenuBuilder, Sequencer.ToSharedRef(), Bindings, [this, &StructBuilder, &CustomizationUtils]()
 			{
 				if (IDetailsView* DetailsView = StructBuilder.GetParentCategory().GetParentLayout().GetDetailsView())
 				{
@@ -971,7 +974,10 @@ bool ULevelSequenceEditorSubsystem::ChangeActorTemplateClass(const FMovieSceneBi
 
 	bool bSuccess = false;
 
-	FSequencerUtilities::HandleTemplateActorClassPicked(ActorClass, Sequencer.ToSharedRef(), ObjectBinding.BindingID, 0, [&bSuccess](){bSuccess=true;});
+	TArray<FSequencerChangeBindingInfo> Bindings;
+	Bindings.Add(FSequencerChangeBindingInfo(ObjectBinding.BindingID, 0));
+
+	FSequencerUtilities::HandleTemplateActorClassPicked(ActorClass, Sequencer.ToSharedRef(), Bindings, [&bSuccess](){bSuccess=true;});
 
 	return bSuccess;
 }
diff --git a/Engine/Plugins/MovieScene/LevelSequenceEditor/Source/LevelSequenceEditor/Private/Misc/LevelSequenceCustomization.cpp b/Engine/Plugins/MovieScene/LevelSequenceEditor/Source/LevelSequenceEditor/Private/Misc/LevelSequenceCustomization.cpp
--- a/Engine/Plugins/MovieScene/LevelSequenceEditor/Source/LevelSequenceEditor/Private/Misc/LevelSequenceCustomization.cpp
+++ b/Engine/Plugins/MovieScene/LevelSequenceEditor/Source/LevelSequenceEditor/Private/Misc/LevelSequenceCustomization.cpp
@@ -200,10 +200,21 @@ void FLevelSequenceCustomization::ExtendObjectBindingContextMenu(FMenuBuilder& M
 
 				if (!bMultipleBindings)
 				{
+					TArray<FSequencerChangeBindingInfo> Bindings;
+					const FMovieSceneBindingReferences* BindingReferences = Sequence->GetBindingReferences();
+					for (TViewModelPtr<IObjectBindingExtension> ObjectBindingNode : Sequencer->GetViewModel()->GetSelection()->Outliner.Filter<IObjectBindingExtension>())
+					{
+						int32 BindingIndex = 0;
+						for (const FMovieSceneBindingReference& Reference : BindingReferences->GetReferences(ObjectBindingNode->GetObjectGuid()))
+						{
+							Bindings.Add({ Reference.ID, BindingIndex++ });
+						}
+					}
+
 					MenuBuilder.AddSubMenu(
 						LOCTEXT("ChangeClassLabel", "Change Class"),
 						LOCTEXT("ChangeClassTooltip", "Change the class (object template) that this spawns from"),
-						FNewMenuDelegate::CreateLambda([=](FMenuBuilder& MenuBuilder) { FSequencerUtilities::AddChangeClassMenu(MenuBuilder, Sequencer.ToSharedRef(), ObjectBindingID, 0, TFunction<void()>()); }));
+						FNewMenuDelegate::CreateLambda([=](FMenuBuilder& MenuBuilder) { FSequencerUtilities::AddChangeClassMenu(MenuBuilder, Sequencer.ToSharedRef(), Bindings, TFunction<void()>()); }));
 				}
 			}
 
@@ -219,7 +230,7 @@ void FLevelSequenceCustomization::ExtendObjectBindingContextMenu(FMenuBuilder& M
 			LOCTEXT("ConvertBindingLabelTooltip", "Convert selected bindings into another binding type"),
 			FNewMenuDelegate::CreateLambda([Sequencer, Sequence](FMenuBuilder& MenuBuilder) 
 			{ 
-					TArray<FSequencerConvertBindingInfo> Bindings;
+					TArray<FSequencerChangeBindingInfo> Bindings;
 					const FMovieSceneBindingReferences* BindingReferences = Sequence->GetBindingReferences();
 					for (TViewModelPtr<IObjectBindingExtension> ObjectBindingNode : Sequencer->GetViewModel()->GetSelection()->Outliner.Filter<IObjectBindingExtension>())
 					{
diff --git a/Engine/Source/Editor/Sequencer/Private/SequencerUtilities.cpp b/Engine/Source/Editor/Sequencer/Private/SequencerUtilities.cpp
--- a/Engine/Source/Editor/Sequencer/Private/SequencerUtilities.cpp
+++ b/Engine/Source/Editor/Sequencer/Private/SequencerUtilities.cpp
@@ -4043,7 +4043,7 @@ TRange<FFrameNumber> FSequencerUtilities::GetTimeBounds(TSharedRef<ISequencer> S
 	return StaticCastSharedPtr<FSequencer>(Sequencer.ToSharedPtr())->GetTimeBounds();
 }
 
-void FSequencerUtilities::AddChangeClassMenu(FMenuBuilder& MenuBuilder, TSharedRef<ISequencer> Sequencer, FGuid ObjectBindingID, int32 BindingIndex, TFunction<void()> OnBindingChanged)
+void FSequencerUtilities::AddChangeClassMenu(FMenuBuilder& MenuBuilder, TSharedRef<ISequencer> Sequencer, const TArray<FSequencerChangeBindingInfo>& Bindings, TFunction<void()> OnBindingChanged)
 {
 	UMovieSceneSequence* Sequence = Sequencer->GetFocusedMovieSceneSequence();
 	if (!Sequence)
@@ -4056,70 +4056,72 @@ void FSequencerUtilities::AddChangeClassMenu(FMenuBuilder& MenuBuilder, TSharedR
 	Options.Mode = EClassViewerMode::ClassPicker;
 	Options.bIsPlaceableOnly = true;
 
-	FMovieSceneSpawnable* Spawnable = MovieScene->FindSpawnable(ObjectBindingID);
-	if (Spawnable)
+	for (const FSequencerChangeBindingInfo& Binding : Bindings)
 	{
-		Options.bIsActorsOnly = true;
-	}
-	else if (const FMovieSceneBindingReferences* BindingReferences = Sequence->GetBindingReferences())
-	{
-		TArrayView<const FMovieSceneBindingReference> BindingReferencesList = BindingReferences->GetReferences(ObjectBindingID);
-		if (BindingReferencesList.IsValidIndex(BindingIndex) && BindingReferencesList[BindingIndex].CustomBinding && BindingReferencesList[BindingIndex].CustomBinding->WillSpawnObject(Sequencer->GetSharedPlaybackState()))
+		FMovieSceneSpawnable* Spawnable = MovieScene->FindSpawnable(Binding.BindingID);
+		if (Spawnable)
 		{
-			// Class filter for the custom binding type
-			class FCustomBindingClassFilter : public IClassViewerFilter
+			Options.bIsActorsOnly = true;
+		}
+		else if (const FMovieSceneBindingReferences* BindingReferences = Sequence->GetBindingReferences())
+		{
+			TArrayView<const FMovieSceneBindingReference> BindingReferencesList = BindingReferences->GetReferences(Binding.BindingID);
+			if (BindingReferencesList.IsValidIndex(Binding.BindingIndex) && BindingReferencesList[Binding.BindingIndex].CustomBinding && BindingReferencesList[Binding.BindingIndex].CustomBinding->WillSpawnObject(Sequencer->GetSharedPlaybackState()))
 			{
-			public:
-				bool IsClassAllowed(const FClassViewerInitializationOptions& InInitOptions, const UClass* InClass, TSharedRef< FClassViewerFilterFuncs > InFilterFuncs) override
+				// Class filter for the custom binding type
+				class FCustomBindingClassFilter : public IClassViewerFilter
 				{
-					return CustomBinding && InClass && CustomBinding->SupportsBindingCreationFromObject(InClass->GetDefaultObject());
-				}
+				public:
+					bool IsClassAllowed(const FClassViewerInitializationOptions& InInitOptions, const UClass* InClass, TSharedRef< FClassViewerFilterFuncs > InFilterFuncs) override
+					{
+						return CustomBinding && InClass && CustomBinding->SupportsBindingCreationFromObject(InClass->GetDefaultObject());
+					}
 
-				virtual bool IsUnloadedClassAllowed(const FClassViewerInitializationOptions& InInitOptions, const TSharedRef< const IUnloadedBlueprintData > InClass, TSharedRef< FClassViewerFilterFuncs > InFilterFuncs) override
-				{
-					if (const UClass* ClassWithin = InClass->GetClassWithin())
+					virtual bool IsUnloadedClassAllowed(const FClassViewerInitializationOptions& InInitOptions, const TSharedRef< const IUnloadedBlueprintData > InClass, TSharedRef< FClassViewerFilterFuncs > InFilterFuncs) override
 					{
-						return IsClassAllowed(InInitOptions, ClassWithin, InFilterFuncs);
+						if (const UClass* ClassWithin = InClass->GetClassWithin())
+						{
+							return IsClassAllowed(InInitOptions, ClassWithin, InFilterFuncs);
+						}
+						return false;
 					}
-					return false;
-				}
 
-				TObjectPtr<UMovieSceneCustomBinding> CustomBinding;
-			};
+					TObjectPtr<UMovieSceneCustomBinding> CustomBinding;
+				};
 
-			TSharedRef<FCustomBindingClassFilter> ClassFilter = MakeShared<FCustomBindingClassFilter>();
-			ClassFilter->CustomBinding = BindingReferencesList[0].CustomBinding;
-			Options.ClassFilters.Add(ClassFilter);
+				TSharedRef<FCustomBindingClassFilter> ClassFilter = MakeShared<FCustomBindingClassFilter>();
+				ClassFilter->CustomBinding = BindingReferencesList[0].CustomBinding;
+				Options.ClassFilters.Add(ClassFilter);
+			}
+			else
+			{
+				return;
+			}
 		}
 		else
 		{
 			return;
 		}
-	}
-	else
-	{
-		return;
+
+		const UClass* ClassForObjectBinding = MovieSceneHelpers::GetBoundObjectClass(Sequence, Binding.BindingID, Binding.BindingIndex);
+		if (ClassForObjectBinding)
+		{
+			Options.ViewerTitleString = FText::FromString(TEXT("Change from: ") + ClassForObjectBinding->GetFName().ToString());
+		}
+		else
+		{
+			Options.ViewerTitleString = FText::FromString(TEXT("Change from: (empty)"));
+		}
 	}
 
 	FClassViewerModule& ClassViewerModule = FModuleManager::LoadModuleChecked<FClassViewerModule>("ClassViewer");
 
-	
-	const UClass* ClassForObjectBinding = MovieSceneHelpers::GetBoundObjectClass(Sequence, ObjectBindingID, BindingIndex);
-	if (ClassForObjectBinding)
-	{
-		Options.ViewerTitleString = FText::FromString(TEXT("Change from: ") + ClassForObjectBinding->GetFName().ToString());
-	}
-	else
-	{
-		Options.ViewerTitleString = FText::FromString(TEXT("Change from: (empty)"));
-	}
-
 	MenuBuilder.AddWidget(
 		SNew(SBox)
 		.MinDesiredWidth(300.0f)
 		.MaxDesiredHeight(400.0f)
 		[
-			ClassViewerModule.CreateClassViewer(Options, FOnClassPicked::CreateLambda([Sequencer, ObjectBindingID, BindingIndex, OnBindingChanged](UClass* Class) { FSequencerUtilities::HandleTemplateActorClassPicked(Class, Sequencer, ObjectBindingID, BindingIndex, OnBindingChanged); }))
+			ClassViewerModule.CreateClassViewer(Options, FOnClassPicked::CreateLambda([Sequencer, Bindings, OnBindingChanged](UClass* Class) { FSequencerUtilities::HandleTemplateActorClassPicked(Class, Sequencer, Bindings, OnBindingChanged); }))
 		],
 		FText(), true, false
 	);
@@ -4157,7 +4159,7 @@ void UpdatePossessedClasses(UMovieScene* MovieScene, FMovieSceneSequenceIDRef Se
 	}
 }
 
-void FSequencerUtilities::HandleTemplateActorClassPicked(UClass* ChosenClass, TSharedRef<ISequencer> Sequencer, FGuid ObjectBindingID, int32 BindingIndex, TFunction<void()> OnBindingChanged)
+void FSequencerUtilities::HandleTemplateActorClassPicked(UClass* ChosenClass, TSharedRef<ISequencer> Sequencer, const TArray<FSequencerChangeBindingInfo>& Bindings, TFunction<void()> OnBindingChanged)
 {
 	UMovieScene* MovieScene = Sequencer->GetFocusedMovieSceneSequence()->GetMovieScene();
 
@@ -4170,11 +4172,15 @@ void FSequencerUtilities::HandleTemplateActorClassPicked(UClass* ChosenClass, TS
 	{
 		FMovieSceneRootEvaluationTemplateInstance& RootInstance = Sequencer->GetEvaluationTemplate();
 		const FMovieSceneSequenceHierarchy* Hierarchy = RootInstance.GetCompiledDataManager()->FindHierarchy(RootInstance.GetCompiledDataID());
-		UpdatePossessedClasses(Sequencer->GetRootMovieSceneSequence()->GetMovieScene(), MovieSceneSequenceID::Root, Hierarchy, ObjectBindingID, ChosenClass);
 
-		MovieSceneHelpers::SetObjectTemplate(Sequencer->GetFocusedMovieSceneSequence(), ObjectBindingID, Result.GetValue().ObjectTemplate, Sequencer->GetSharedPlaybackState(), BindingIndex);
+		for (const FSequencerChangeBindingInfo& Binding : Bindings)
+		{
+			UpdatePossessedClasses(Sequencer->GetRootMovieSceneSequence()->GetMovieScene(), MovieSceneSequenceID::Root, Hierarchy, Binding.BindingID, ChosenClass);
+
+			MovieSceneHelpers::SetObjectTemplate(Sequencer->GetFocusedMovieSceneSequence(), Binding.BindingID, Result.GetValue().ObjectTemplate, Sequencer->GetSharedPlaybackState(), Binding.BindingIndex);
 
-		Sequencer->GetSpawnRegister().DestroySpawnedObject(ObjectBindingID, Sequencer->GetFocusedTemplateID(), Sequencer->GetSharedPlaybackState(), BindingIndex);
+			Sequencer->GetSpawnRegister().DestroySpawnedObject(Binding.BindingID, Sequencer->GetFocusedTemplateID(), Sequencer->GetSharedPlaybackState(), Binding.BindingIndex);
+		}
 		Sequencer->ForceEvaluate();
 	}
 
@@ -4250,7 +4256,7 @@ bool FSequencerUtilities::CanConvertToCustomBinding(TSharedRef<ISequencer> Seque
 }
 
 
-void FSequencerUtilities::AddConvertBindingMenu(FMenuBuilder& MenuBuilder, TSharedRef<ISequencer> Sequencer, const TArray<FSequencerConvertBindingInfo>& BindingsToConvert, TFunction<void()> OnBindingChanged)
+void FSequencerUtilities::AddConvertBindingMenu(FMenuBuilder& MenuBuilder, TSharedRef<ISequencer> Sequencer, const TArray<FSequencerChangeBindingInfo>& BindingsToConvert, TFunction<void()> OnBindingChanged)
 {
 	UMovieSceneSequence* Sequence = Sequencer->GetFocusedMovieSceneSequence();
 	UMovieScene* MovieScene = Sequence ? Sequence->GetMovieScene() : nullptr;
@@ -4278,7 +4284,7 @@ void FSequencerUtilities::AddConvertBindingMenu(FMenuBuilder& MenuBuilder, TShar
 			SlowTask.MakeDialog(true);
 
 			TArray<AActor*> PossessedActors;
-			for (const FSequencerConvertBindingInfo& BindingInfo : BindingsToConvert)
+			for (const FSequencerChangeBindingInfo& BindingInfo : BindingsToConvert)
 			{
 				SlowTask.EnterProgressFrame();
 
@@ -4328,7 +4334,7 @@ void FSequencerUtilities::AddConvertBindingMenu(FMenuBuilder& MenuBuilder, TShar
 	};
 
 	// Can convert to possessable
-	if (Algo::AllOf(BindingsToConvert, [Sequencer, &Sequence, &MovieScene](const FSequencerConvertBindingInfo& BindingInfo) {
+	if (Algo::AllOf(BindingsToConvert, [Sequencer, &Sequence, &MovieScene](const FSequencerChangeBindingInfo& BindingInfo) {
 		return CanConvertToPossessable(Sequencer, BindingInfo.BindingID, BindingInfo.BindingIndex );
 	}))
 	{
@@ -4347,7 +4353,7 @@ void FSequencerUtilities::AddConvertBindingMenu(FMenuBuilder& MenuBuilder, TShar
 	TArrayView<const TSubclassOf<UMovieSceneCustomBinding>> PrioritySortedCustomBindingTypes = Sequencer->GetSupportedCustomBindingTypes();
 	for (const TSubclassOf<UMovieSceneCustomBinding>& CustomBindingType : PrioritySortedCustomBindingTypes)
 	{
-		if (Algo::AllOf(BindingsToConvert, [&Sequence, &MovieScene, &CustomBindingType, Sequencer](const FSequencerConvertBindingInfo& BindingInfo) 
+		if (Algo::AllOf(BindingsToConvert, [&Sequence, &MovieScene, &CustomBindingType, Sequencer](const FSequencerChangeBindingInfo& BindingInfo) 
 		{
 			return CanConvertToCustomBinding(Sequencer, BindingInfo.BindingID, CustomBindingType, BindingInfo.BindingIndex);
 		}))
diff --git a/Engine/Source/Editor/Sequencer/Public/SequencerUtilities.h b/Engine/Source/Editor/Sequencer/Public/SequencerUtilities.h
--- a/Engine/Source/Editor/Sequencer/Public/SequencerUtilities.h
+++ b/Engine/Source/Editor/Sequencer/Public/SequencerUtilities.h
@@ -171,7 +171,7 @@ struct FMovieScenePasteBindingsParams
 };
 
 USTRUCT()
-struct FSequencerConvertBindingInfo
+struct FSequencerChangeBindingInfo
 {
 	GENERATED_BODY()
 
@@ -286,7 +286,7 @@ struct SEQUENCER_API FSequencerUtilities
 	static TRange<FFrameNumber> GetTimeBounds(TSharedRef<ISequencer> Sequencer);
 
 	// Functions allowing menus to be built for modifying bindings
-	static void AddChangeClassMenu(FMenuBuilder& MenuBuilder, TSharedRef<ISequencer> Sequencer, FGuid Guid, int32 BindingIndex, TFunction<void()> OnBindingChanged);
-	static void HandleTemplateActorClassPicked(UClass* ChosenClass, TSharedRef<ISequencer> Sequencer, FGuid Guid, int32 BindingIndex, TFunction<void()> OnBindingChanged);
-	static void AddConvertBindingMenu(FMenuBuilder& MenuBuilder, TSharedRef<ISequencer> Sequencer, const TArray<FSequencerConvertBindingInfo>& BindingsToConvert, TFunction<void()> OnBindingChanged);
+	static void AddChangeClassMenu(FMenuBuilder& MenuBuilder, TSharedRef<ISequencer> Sequencer, const TArray<FSequencerChangeBindingInfo>& BindingsToConvert, TFunction<void()> OnBindingChanged);
+	static void HandleTemplateActorClassPicked(UClass* ChosenClass, TSharedRef<ISequencer> Sequencer, const TArray<FSequencerChangeBindingInfo>& BindingsToConvert, TFunction<void()> OnBindingChanged);
+	static void AddConvertBindingMenu(FMenuBuilder& MenuBuilder, TSharedRef<ISequencer> Sequencer, const TArray<FSequencerChangeBindingInfo>& BindingsToConvert, TFunction<void()> OnBindingChanged);
 };
