Sequencer: Fix CIS

diff --git a/Engine/Source/Editor/Sequencer/Public/SequencerUtilities.h b/Engine/Source/Editor/Sequencer/Public/SequencerUtilities.h
--- a/Engine/Source/Editor/Sequencer/Public/SequencerUtilities.h
+++ b/Engine/Source/Editor/Sequencer/Public/SequencerUtilities.h
@@ -175,6 +175,11 @@ struct FSequencerChangeBindingInfo
 {
 	GENERATED_BODY()
 
+	FSequencerChangeBindingInfo() {}
+	FSequencerChangeBindingInfo(FGuid InBindingID, int32 InBindingIndex)
+		: BindingID(InBindingID)
+		, BindingIndex(InBindingIndex) {}
+
 	UPROPERTY()
 	FGuid BindingID;
 
