Sequencer: Fix CIS

#rnx

diff --git a/Engine/Source/Editor/Sequencer/Private/Sequencer.cpp b/Engine/Source/Editor/Sequencer/Private/Sequencer.cpp
--- a/Engine/Source/Editor/Sequencer/Private/Sequencer.cpp
+++ b/Engine/Source/Editor/Sequencer/Private/Sequencer.cpp
@@ -5860,32 +5860,34 @@ void FSequencer::OnNewActorsDropped(const TArray<UObject*>& DroppedObjects, cons
 					TSubclassOf<UMovieSceneCustomBinding> CustomBindingClass = bAddSpawnable ? UMovieSceneSpawnableActorBinding::StaticClass() : UMovieSceneReplaceableActorBinding::StaticClass();
 
 					const FMovieSceneBindingReferences* BindingReferences = Sequence->GetBindingReferences();
-					TArrayView<const FMovieSceneBindingReference> References = BindingReferences ? BindingReferences->GetReferences(PossessableGuid) : TArray<const FMovieSceneBindingReference>();
 
-					for (const FMovieSceneBindingReference& Reference : References)
+					if (BindingReferences)
 					{
-						for (const TSubclassOf<UMovieSceneCustomBinding>& SupportedCustomBindingType : SupportedCustomBindingTypes)
+						for (const FMovieSceneBindingReference& Reference : BindingReferences->GetReferences(PossessableGuid))
 						{
-							if (SupportedCustomBindingType && SupportedCustomBindingType->IsChildOf(CustomBindingClass) && 
-								SupportedCustomBindingType->GetDefaultObject<UMovieSceneCustomBinding>()->SupportsConversionFromBinding(Reference, Actor))
+							for (const TSubclassOf<UMovieSceneCustomBinding>& SupportedCustomBindingType : SupportedCustomBindingTypes)
 							{
-								FMovieScenePossessable* NewPossessable = FSequencerUtilities::ConvertToCustomBinding(AsShared(), NewGuid, CustomBindingClass);
-
-								if (NewPossessable)
+								if (SupportedCustomBindingType && SupportedCustomBindingType->IsChildOf(CustomBindingClass) && 
+									SupportedCustomBindingType->GetDefaultObject<UMovieSceneCustomBinding>()->SupportsConversionFromBinding(Reference, Actor))
 								{
-									for (TWeakObjectPtr<> WeakObject : FindBoundObjects(NewPossessable->GetGuid(), ActiveTemplateIDs.Top()))
+									FMovieScenePossessable* NewPossessable = FSequencerUtilities::ConvertToCustomBinding(AsShared(), NewGuid, CustomBindingClass);
+
+									if (NewPossessable)
 									{
-										AActor* SpawnedActor = Cast<AActor>(WeakObject.Get());
-										if (SpawnedActor)
+										for (TWeakObjectPtr<> WeakObject : FindBoundObjects(NewPossessable->GetGuid(), ActiveTemplateIDs.Top()))
 										{
-											SpawnedActors.Add(SpawnedActor);
-											NewActor = SpawnedActor;
+											AActor* SpawnedActor = Cast<AActor>(WeakObject.Get());
+											if (SpawnedActor)
+											{
+												SpawnedActors.Add(SpawnedActor);
+												NewActor = SpawnedActor;
+											}
 										}
-									}
 
-									NewGuid = NewPossessable->GetGuid();
+										NewGuid = NewPossessable->GetGuid();
+									}
+									break;
 								}
-								break;
 							}
 						}
 					}
