Sequencer: Add Selection from Content Browser

Added a menu item to get the selection from the content browser and add it as if it were drag dropped from the content browser (ie. if you select an audio asset, it will add it as an audio track, if you drag in a static mesh, it will add it as a spawnable).

diff --git a/Engine/Source/Editor/Sequencer/Private/SSequencer.cpp b/Engine/Source/Editor/Sequencer/Private/SSequencer.cpp
--- a/Engine/Source/Editor/Sequencer/Private/SSequencer.cpp
+++ b/Engine/Source/Editor/Sequencer/Private/SSequencer.cpp
@@ -120,6 +120,7 @@
 #include "EditorActorFolders.h"
 #include "Tracks/MovieSceneEventTrack.h"
 #include "ToolMenus.h"
+#include "MovieSceneFolder.h"
 #include "MovieSceneToolHelpers.h"
 #include "Editor/UnrealEdEngine.h"
 #include "UnrealEdGlobals.h"
@@ -128,6 +129,8 @@
 #include "Sidebar/SidebarDrawerConfig.h"
 #include "Sidebar/SSidebar.h"
 #include "Widgets/SOverlay.h"
+#include "IContentBrowserSingleton.h"
+#include "ContentBrowserModule.h"
 
 #define LOCTEXT_NAMESPACE "Sequencer"
 
@@ -1683,6 +1686,15 @@ TSharedRef<SWidget> SSequencer::MakeAddMenu()
 
 	TSharedPtr<FExtender> Extender = FExtender::Combine(AddMenuExtenders);
 	FMenuBuilder MenuBuilder(true, nullptr, Extender);
+
+	MenuBuilder.AddMenuEntry(
+		LOCTEXT("SelectedFromContentBrowser", "Selected from Content Browser"),
+		LOCTEXT("SelectedFromContentBrowserToolTip", "Add selected content from the content browser"),
+		FSlateIcon(FAppStyle::GetAppStyleSetName(), "Icons.Use"),
+		FUIAction(
+			FExecuteAction::CreateRaw(this, &SSequencer::AddFromContentBrowser),
+			FCanExecuteAction::CreateRaw(this, &SSequencer::CanAddFromContentBrowser)));
+
 	{
 		TSharedPtr<FSequencer> Sequencer = SequencerPtr.Pin();
 		Sequencer->GetViewModel()->GetOutliner()->CastThisChecked<FSequencerOutlinerViewModel>()->BuildContextMenu(MenuBuilder);
@@ -3311,30 +3323,30 @@ FReply SSequencer::OnDrop( const FGeometry& MyGeometry, const FDragDropEvent& Dr
 	{
 		if ( Operation->IsOfType<FAssetDragDropOp>() )
 		{
-			const auto& DragDropOp = StaticCastSharedPtr<FAssetDragDropOp>(Operation);
+			TSharedPtr<FAssetDragDropOp> DragDropOp = StaticCastSharedPtr<FAssetDragDropOp>(Operation);
 
-			OnAssetsDropped( *DragDropOp );
+			OnAssetsDropped( DragDropOp );
 			bWasDropHandled = true;
 		}
 		else if( Operation->IsOfType<FClassDragDropOp>() )
 		{
-			const auto& DragDropOp = StaticCastSharedPtr<FClassDragDropOp>(Operation);
+			TSharedPtr<FClassDragDropOp> DragDropOp = StaticCastSharedPtr<FClassDragDropOp>(Operation);
 
-			OnClassesDropped( *DragDropOp );
+			OnClassesDropped( DragDropOp );
 			bWasDropHandled = true;
 		}
 		else if( Operation->IsOfType<FActorDragDropOp>() )
 		{
-			const auto& DragDropOp = StaticCastSharedPtr<FActorDragDropOp>(Operation);
+			TSharedPtr<FActorDragDropOp> DragDropOp = StaticCastSharedPtr<FActorDragDropOp>(Operation);
 
-			OnActorsDropped( *DragDropOp );
+			OnActorsDropped( DragDropOp );
 			bWasDropHandled = true;
 		}
 		else if (Operation->IsOfType<FFolderDragDropOp>())
 		{
-			const auto& DragDropOp = StaticCastSharedPtr<FFolderDragDropOp>(Operation);
+			TSharedPtr<FFolderDragDropOp> DragDropOp = StaticCastSharedPtr<FFolderDragDropOp>(Operation);
 
-			OnFolderDropped(*DragDropOp);
+			OnFolderDropped(DragDropOp);
 			bWasDropHandled = true;
 		}
 		else if (Operation->IsOfType<FCompositeDragDropOp>())
@@ -3342,15 +3354,14 @@ FReply SSequencer::OnDrop( const FGeometry& MyGeometry, const FDragDropEvent& Dr
 			const TSharedPtr<FCompositeDragDropOp> CompositeOp = StaticCastSharedPtr<FCompositeDragDropOp>(Operation);
 			if (const TSharedPtr<FActorDragDropOp> ActorDragDropOp = CompositeOp->GetSubOp<FActorDragDropOp>())
 			{
-				OnActorsDropped(*ActorDragDropOp);
+				OnActorsDropped(ActorDragDropOp);
 				bWasDropHandled = true;
 			}
 			if (const TSharedPtr<FFolderDragDropOp> FolderDragDropOp = CompositeOp->GetSubOp<FFolderDragDropOp>())
 			{
-				OnFolderDropped(*FolderDragDropOp);
+				OnFolderDropped(FolderDragDropOp);
 				bWasDropHandled = true;
 			}
-			
 		}
 	}
 
@@ -3390,7 +3401,33 @@ void SSequencer::OnMouseLeave(const FPointerEvent& MouseEvent)
 	PendingFocus.ResetPendingFocus();
 }
 
-void SSequencer::OnAssetsDropped( const FAssetDragDropOp& DragDropOp )
+void SSequencer::AddFromContentBrowser()
+{
+	FContentBrowserModule& ContentBrowserModule = FModuleManager::LoadModuleChecked<FContentBrowserModule>("ContentBrowser");
+	TArray<FAssetData> SelectedAssets;
+	ContentBrowserModule.Get().GetSelectedAssets(SelectedAssets);
+
+	TSharedRef<FAssetDragDropOp> DragDropOp = FAssetDragDropOp::New(SelectedAssets);
+	OnAssetsDropped(DragDropOp.ToSharedPtr());
+}
+
+bool SSequencer::CanAddFromContentBrowser() const
+{
+	FSequencer& SequencerRef = *SequencerPtr.Pin();
+	FContentBrowserModule& ContentBrowserModule = FModuleManager::LoadModuleChecked<FContentBrowserModule>("ContentBrowser");
+	TArray<FAssetData> SelectedAssets;
+	ContentBrowserModule.Get().GetSelectedAssets(SelectedAssets);
+	for (const FAssetData& AssetData : SelectedAssets)
+	{
+		if (MovieSceneToolHelpers::IsValidAsset(SequencerRef.GetFocusedMovieSceneSequence(), AssetData))
+		{
+			return true;
+		}
+	}
+	return false;
+}
+
+void SSequencer::OnAssetsDropped(TSharedPtr<FAssetDragDropOp> DragDropOp)
 {
 	using namespace UE::Sequencer;
 
@@ -3400,7 +3437,7 @@ void SSequencer::OnAssetsDropped( const FAssetDragDropOp& DragDropOp )
 	bool bAllAssetsWereLoaded = true;
 	bool bNeedsLoad = false;
 
-	for (const FAssetData& AssetData : DragDropOp.GetAssets())
+	for (const FAssetData& AssetData : DragDropOp->GetAssets())
 	{
 		if (!AssetData.IsAssetLoaded())
 		{
@@ -3414,7 +3451,7 @@ void SSequencer::OnAssetsDropped( const FAssetDragDropOp& DragDropOp )
 		GWarn->BeginSlowTask(LOCTEXT("OnDrop_FullyLoadPackage", "Fully Loading Package For Drop"), true, false);
 	}
 
-	for (const FAssetData& AssetData : DragDropOp.GetAssets())
+	for (const FAssetData& AssetData : DragDropOp->GetAssets())
 	{
 		if (!MovieSceneToolHelpers::IsValidAsset(SequencerRef.GetFocusedMovieSceneSequence(), AssetData))
 		{
@@ -3457,7 +3494,7 @@ void SSequencer::OnAssetsDropped( const FAssetDragDropOp& DragDropOp )
 	{
 		if (Delegate.IsBound())
 		{
-			DropResult = Delegate.Execute(DroppedObjects, DragDropOp);
+			DropResult = Delegate.Execute(DroppedObjects, *DragDropOp);
 			if (DropResult != ESequencerDropResult::Unhandled)
 			{
 				break;
@@ -3470,6 +3507,9 @@ void SSequencer::OnAssetsDropped( const FAssetDragDropOp& DragDropOp )
 	{
 		FMovieSceneTrackEditor::BeginKeying(SequencerPtr.Pin()->GetLocalTime().Time.FrameNumber);
 
+		TArray<UMovieSceneFolder*> Folders;
+		SequencerRef.GetSelectedFolders(Folders);
+
 		for (TArray<UObject*>::TConstIterator CurObjectIter = DroppedObjects.CreateConstIterator(); CurObjectIter; ++CurObjectIter)
 		{
 			UObject* CurObject = *CurObjectIter;
@@ -3487,11 +3527,17 @@ void SSequencer::OnAssetsDropped( const FAssetDragDropOp& DragDropOp )
 				Params.BindingNameOverride = CurObject->GetName();
 				Params.bSpawnable = true;
 				Params.bReplaceable = true;
-				Params.ActorFactory = DragDropOp.GetActorFactory();
+				Params.ActorFactory = DragDropOp->GetActorFactory();
 
-				if (SequencerRef.CreateBinding(*CurObject, Params).IsValid())
+				FGuid Guid = SequencerRef.CreateBinding(*CurObject, Params);
+				if (Guid.IsValid())
 				{
 					DropResult = ESequencerDropResult::DropHandled;
+
+					if (Folders.Num() > 0)
+					{
+						Folders[0]->AddChildObjectBinding(Guid);
+					}
 				}
 			}
 		}
@@ -3512,7 +3558,7 @@ void SSequencer::OnAssetsDropped( const FAssetDragDropOp& DragDropOp )
 }
 
 
-void SSequencer::OnClassesDropped( const FClassDragDropOp& DragDropOp )
+void SSequencer::OnClassesDropped(TSharedPtr<FClassDragDropOp> DragDropOp )
 {
 	const FScopedTransaction Transaction(LOCTEXT("DropClasses", "Drop Classes"));
 
@@ -3522,7 +3568,7 @@ void SSequencer::OnClassesDropped( const FClassDragDropOp& DragDropOp )
 	{
 		if (Delegate.IsBound())
 		{
-			DropResult = Delegate.Execute(DragDropOp.ClassesToDrop, DragDropOp);
+			DropResult = Delegate.Execute(DragDropOp->ClassesToDrop, *DragDropOp);
 			if (DropResult != ESequencerDropResult::Unhandled)
 			{
 				break;
@@ -3534,7 +3580,7 @@ void SSequencer::OnClassesDropped( const FClassDragDropOp& DragDropOp )
 	{
 		FSequencer& SequencerRef = *SequencerPtr.Pin();
 
-		for (auto ClassIter = DragDropOp.ClassesToDrop.CreateConstIterator(); ClassIter; ++ClassIter)
+		for (auto ClassIter = DragDropOp->ClassesToDrop.CreateConstIterator(); ClassIter; ++ClassIter)
 		{
 			UClass* Class = (*ClassIter).Get();
 			if (Class != nullptr)
@@ -3550,7 +3596,7 @@ void SSequencer::OnClassesDropped( const FClassDragDropOp& DragDropOp )
 	}
 }
 
-void SSequencer::OnActorsDropped( FActorDragDropOp& DragDropOp )
+void SSequencer::OnActorsDropped(TSharedPtr<FActorDragDropOp> DragDropOp )
 {
 	const FScopedTransaction Transaction(LOCTEXT("DropActors", "Drop Actors"));
 
@@ -3560,7 +3606,7 @@ void SSequencer::OnActorsDropped( FActorDragDropOp& DragDropOp )
 	{
 		if (Delegate.IsBound())
 		{
-			DropResult = Delegate.Execute(DragDropOp.Actors, DragDropOp);
+			DropResult = Delegate.Execute(DragDropOp->Actors, *DragDropOp);
 			if (DropResult != ESequencerDropResult::Unhandled)
 			{
 				break;
@@ -3570,14 +3616,14 @@ void SSequencer::OnActorsDropped( FActorDragDropOp& DragDropOp )
 
 	if (DropResult == ESequencerDropResult::Unhandled)
 	{
-		SequencerPtr.Pin()->OnActorsDropped(DragDropOp.Actors);
+		SequencerPtr.Pin()->OnActorsDropped(DragDropOp->Actors);
 	}
 }
 
-void SSequencer::OnFolderDropped( FFolderDragDropOp& DragDropOp )
+void SSequencer::OnFolderDropped(TSharedPtr<FFolderDragDropOp> DragDropOp )
 {
 	// Sequencer doesn't support dragging folder with a root object
-	if (!FFolder::IsRootObjectPersistentLevel(DragDropOp.RootObject))
+	if (!FFolder::IsRootObjectPersistentLevel(DragDropOp->RootObject))
 	{
 		return;
 	}
@@ -3595,14 +3641,14 @@ void SSequencer::OnFolderDropped( FFolderDragDropOp& DragDropOp )
 	UWorld* World = PlaybackContext ? PlaybackContext->GetWorld() : nullptr;
 	if (World)
 	{
-		FActorFolders::GetWeakActorsFromFolders(*World, DragDropOp.Folders, DraggedActors);
+		FActorFolders::GetWeakActorsFromFolders(*World, DragDropOp->Folders, DraggedActors);
 	}
 	
 	for (FOnFoldersDrop Delegate : OnFoldersDrop)
 	{
 		if (Delegate.IsBound())
 		{
-			DropResult = Delegate.Execute(DragDropOp.Folders, DragDropOp);
+			DropResult = Delegate.Execute(DragDropOp->Folders, *DragDropOp);
 			if (DropResult != ESequencerDropResult::Unhandled)
 			{
 				break;
diff --git a/Engine/Source/Editor/Sequencer/Private/SSequencer.h b/Engine/Source/Editor/Sequencer/Private/SSequencer.h
--- a/Engine/Source/Editor/Sequencer/Private/SSequencer.h
+++ b/Engine/Source/Editor/Sequencer/Private/SSequencer.h
@@ -538,33 +538,36 @@ private:
 		return ColumnFillCoefficients[ColumnIndex];
 	}
 
+	void AddFromContentBrowser();
+	bool CanAddFromContentBrowser() const;
+
 	/**
 	 * Called when one or more assets are dropped into the widget
 	 *
 	 * @param	DragDropOp	Information about the asset(s) that were dropped
 	 */
-	void OnAssetsDropped(const FAssetDragDropOp& DragDropOp);
+	void OnAssetsDropped(TSharedPtr<FAssetDragDropOp> DragDropOp);
 	
 	/**
 	 * Called when one or more classes are dropped into the widget
 	 *
 	 * @param	DragDropOp	Information about the class(es) that were dropped
 	 */
-	void OnClassesDropped(const FClassDragDropOp& DragDropOp);
+	void OnClassesDropped(TSharedPtr<FClassDragDropOp> DragDropOp);
 		
 	/**
 	 * Called when one or more actors are dropped into the widget
 	 *
 	 * @param	DragDropOp	Information about the actor(s) that was dropped
 	 */
-	void OnActorsDropped(FActorDragDropOp& DragDropOp); 
+	void OnActorsDropped(TSharedPtr<FActorDragDropOp> DragDropOp);
 
 	/**
 	 * Called when one or more folders are dropped into the widget
 	 *
 	 * @param	DragDropOp	Information about the objects(s) that was dropped
 	 */
-	void OnFolderDropped(FFolderDragDropOp& DragDropOp); 
+	void OnFolderDropped(TSharedPtr<FFolderDragDropOp> DragDropOp);
 
 	/** Called when a breadcrumb is clicked on in the sequencer */
 	void OnCrumbClicked(const FSequencerBreadcrumb& Item);
diff --git a/Engine/Source/Editor/Sequencer/Private/Sequencer.cpp b/Engine/Source/Editor/Sequencer/Private/Sequencer.cpp
--- a/Engine/Source/Editor/Sequencer/Private/Sequencer.cpp
+++ b/Engine/Source/Editor/Sequencer/Private/Sequencer.cpp
@@ -5846,6 +5846,11 @@ void FSequencer::OnNewActorsDropped(const TArray<UObject*>& DroppedObjects, cons
 
 		Sequence->Modify();
 
+		TArray<UMovieSceneFolder*> SelectedParentFolders;
+		FString NewNodePath;
+		CalculateSelectedFolderAndPath(SelectedParentFolders, NewNodePath);
+		UMovieSceneFolder* ParentFolder = SelectedParentFolders.Num() > 0 ? SelectedParentFolders[0] : nullptr;
+
 		for ( AActor* Actor : DroppedActors )
 		{
 			if (AActor* NewActor = Actor)
@@ -5898,6 +5903,11 @@ void FSequencer::OnNewActorsDropped(const TArray<UObject*>& DroppedObjects, cons
 					ACineCameraActor* OutActor;
 					FSequencerUtilities::CreateCameraWithRig(AsShared(), NewActor, bAddSpawnable, OutActor);
 				}
+
+				if (ParentFolder)
+				{
+					ParentFolder->AddChildObjectBinding(NewGuid);
+				}
 			}
 		}
 
