[Sequencer] Sidebar Details:
 - Show remaining context menus in sidebar
 - Fix lost lambda references in context menus
 - Fix FSequencer shared pointer reference counts > 1
 - Restructure how some menus are built to avoid code duplication
 - Code cleanup and readability

diff --git a/Engine/Plugins/Animation/ControlRig/Source/ControlRigEditor/Private/Sequencer/ControlRigParameterTrackEditor.cpp b/Engine/Plugins/Animation/ControlRig/Source/ControlRigEditor/Private/Sequencer/ControlRigParameterTrackEditor.cpp
--- a/Engine/Plugins/Animation/ControlRig/Source/ControlRigEditor/Private/Sequencer/ControlRigParameterTrackEditor.cpp
+++ b/Engine/Plugins/Animation/ControlRig/Source/ControlRigEditor/Private/Sequencer/ControlRigParameterTrackEditor.cpp
@@ -5108,268 +5108,302 @@ void FControlRigParameterSection::KeyWeightValue(float Val)
 
 void FControlRigParameterSection::BuildSectionContextMenu(FMenuBuilder& MenuBuilder, const FGuid& InObjectBinding)
 {
-	UMovieSceneControlRigParameterSection* ParameterSection = CastChecked<UMovieSceneControlRigParameterSection>(WeakSection.Get());
-	TSharedPtr<ISequencer> SequencerPtr = WeakSequencer.Pin();
+	UMovieSceneControlRigParameterSection* const ParameterSection = CastChecked<UMovieSceneControlRigParameterSection>(WeakSection.Get());
+	if (!IsValid(ParameterSection))
+	{
+		return;
+	}
 
-	UControlRig* ControlRig = ParameterSection->GetControlRig();
+	UControlRig* const ControlRig = ParameterSection->GetControlRig();
+	if (!IsValid(ControlRig))
+	{
+		return;
+	}
 
-	if (ControlRig)
+	UFKControlRig* AutoRig = Cast<UFKControlRig>(ControlRig);
+	if (AutoRig || ControlRig->SupportsEvent(FRigUnit_InverseExecution::EventName))
 	{
+		UObject* BoundObject = nullptr;
+		USkeleton* Skeleton = AcquireSkeletonFromObjectGuid(InObjectBinding, &BoundObject, WeakSequencer.Pin());
 
-		UFKControlRig* AutoRig = Cast<UFKControlRig>(ControlRig);
-		if (AutoRig || ControlRig->SupportsEvent(FRigUnit_InverseExecution::EventName))
+		if (Skeleton)
 		{
-			UObject* BoundObject = nullptr;
-			USkeleton* Skeleton = AcquireSkeletonFromObjectGuid(InObjectBinding, &BoundObject, WeakSequencer.Pin());
+			// Load the asset registry module
+			FAssetRegistryModule& AssetRegistryModule = FModuleManager::LoadModuleChecked<FAssetRegistryModule>(TEXT("AssetRegistry"));
 
-			if (Skeleton)
-			{
-				// Load the asset registry module
-				FAssetRegistryModule& AssetRegistryModule = FModuleManager::LoadModuleChecked<FAssetRegistryModule>(TEXT("AssetRegistry"));
-
-				// Collect a full list of assets with the specified class
-				TArray<FAssetData> AssetDataList;
-				AssetRegistryModule.Get().GetAssetsByClass(UAnimSequenceBase::StaticClass()->GetClassPathName(), AssetDataList, true);
+			// Collect a full list of assets with the specified class
+			TArray<FAssetData> AssetDataList;
+			AssetRegistryModule.Get().GetAssetsByClass(UAnimSequenceBase::StaticClass()->GetClassPathName(), AssetDataList, true);
 
-				if (AssetDataList.Num())
-				{
-					MenuBuilder.AddSubMenu(
-						LOCTEXT("ImportAnimSequenceIntoThisSection", "Import Anim Sequence Into This Section"), NSLOCTEXT("Sequencer", "ImportAnimSequenceIntoThisSectionTP", "Import Anim Sequence Into This Section"),
-						FNewMenuDelegate::CreateRaw(this, &FControlRigParameterSection::AddAnimationSubMenuForFK, InObjectBinding, Skeleton, ParameterSection)
-					);
-				}
+			if (AssetDataList.Num())
+			{
+				MenuBuilder.AddSubMenu(
+					LOCTEXT("ImportAnimSequenceIntoThisSection", "Import Anim Sequence Into This Section"), NSLOCTEXT("Sequencer", "ImportAnimSequenceIntoThisSectionTP", "Import Anim Sequence Into This Section"),
+					FNewMenuDelegate::CreateRaw(this, &FControlRigParameterSection::AddAnimationSubMenuForFK, InObjectBinding, Skeleton, ParameterSection)
+				);
 			}
 		}
-		TArray<FRigControlElement*> Controls;
-		ControlRig->GetControlsInOrder(Controls);
+	}
+	TArray<FRigControlElement*> Controls;
+	ControlRig->GetControlsInOrder(Controls);
 
-		auto MakeUIAction = [=](EMovieSceneTransformChannel ChannelsToToggle)
-		{
-			return FUIAction(
-				FExecuteAction::CreateLambda([=]
+	auto MakeUIAction = [this, InObjectBinding](EMovieSceneTransformChannel ChannelsToToggle)
+	{
+		return FUIAction(
+			FExecuteAction::CreateLambda([this, InObjectBinding, ChannelsToToggle]
+				{
+					const TSharedPtr<ISequencer> SequencerPtr = WeakSequencer.Pin();
+					if (!SequencerPtr)
 					{
-						FScopedTransaction Transaction(LOCTEXT("SetActiveChannelsTransaction", "Set Active Channels"));
-						ParameterSection->Modify();
-						EMovieSceneTransformChannel Channels = ParameterSection->GetTransformMask().GetChannels();
+						return;
+					}
 
-						if (EnumHasAllFlags(Channels, ChannelsToToggle) || (Channels & ChannelsToToggle) == EMovieSceneTransformChannel::None)
-						{
-							ParameterSection->SetTransformMask(ParameterSection->GetTransformMask().GetChannels() ^ ChannelsToToggle);
-						}
-						else
-						{
-							ParameterSection->SetTransformMask(ParameterSection->GetTransformMask().GetChannels() | ChannelsToToggle);
-						}
+					UMovieSceneControlRigParameterSection* const ParameterSection = CastChecked<UMovieSceneControlRigParameterSection>(WeakSection.Get());
+					if (!IsValid(ParameterSection))
+					{
+						return;
+					}
+
+					FScopedTransaction Transaction(LOCTEXT("SetActiveChannelsTransaction", "Set Active Channels"));
+					ParameterSection->Modify();
+					EMovieSceneTransformChannel Channels = ParameterSection->GetTransformMask().GetChannels();
 
-						// Restore pre-animated state for the bound objects so that inactive channels will return to their default values.
-						for (TWeakObjectPtr<> WeakObject : SequencerPtr->FindBoundObjects(InObjectBinding, SequencerPtr->GetFocusedTemplateID()))
+					if (EnumHasAllFlags(Channels, ChannelsToToggle) || (Channels & ChannelsToToggle) == EMovieSceneTransformChannel::None)
+					{
+						ParameterSection->SetTransformMask(ParameterSection->GetTransformMask().GetChannels() ^ ChannelsToToggle);
+					}
+					else
+					{
+						ParameterSection->SetTransformMask(ParameterSection->GetTransformMask().GetChannels() | ChannelsToToggle);
+					}
+
+					// Restore pre-animated state for the bound objects so that inactive channels will return to their default values.
+					for (TWeakObjectPtr<> WeakObject : SequencerPtr->FindBoundObjects(InObjectBinding, SequencerPtr->GetFocusedTemplateID()))
+					{
+						if (UObject* Object = WeakObject.Get())
 						{
-							if (UObject* Object = WeakObject.Get())
-							{
-								SequencerPtr->RestorePreAnimatedState();
-							}
+							SequencerPtr->RestorePreAnimatedState();
 						}
+					}
 
-						SequencerPtr->NotifyMovieSceneDataChanged(EMovieSceneDataChangeType::MovieSceneStructureItemsChanged);
+					SequencerPtr->NotifyMovieSceneDataChanged(EMovieSceneDataChangeType::MovieSceneStructureItemsChanged);
+				}
+			),
+			FCanExecuteAction(),
+			FGetActionCheckState::CreateLambda([this, ChannelsToToggle]
+				{
+					const UMovieSceneControlRigParameterSection* const ParameterSection = CastChecked<UMovieSceneControlRigParameterSection>(WeakSection.Get());
+					if (!IsValid(ParameterSection))
+					{
+						return ECheckBoxState::Unchecked;
 					}
-				),
-				FCanExecuteAction(),
-						FGetActionCheckState::CreateLambda([=]
-							{
-								EMovieSceneTransformChannel Channels = ParameterSection->GetTransformMask().GetChannels();
-								if (EnumHasAllFlags(Channels, ChannelsToToggle))
-								{
-									return ECheckBoxState::Checked;
-								}
-								else if (EnumHasAnyFlags(Channels, ChannelsToToggle))
-								{
-									return ECheckBoxState::Undetermined;
-								}
-								return ECheckBoxState::Unchecked;
-							})
-						);
-		};
-		auto ToggleControls = [=](int32 Index)
-		{
-			return FUIAction(
-				FExecuteAction::CreateLambda([=]
+
+					const EMovieSceneTransformChannel Channels = ParameterSection->GetTransformMask().GetChannels();
+					if (EnumHasAllFlags(Channels, ChannelsToToggle))
+					{
+						return ECheckBoxState::Checked;
+					}
+					else if (EnumHasAnyFlags(Channels, ChannelsToToggle))
+					{
+						return ECheckBoxState::Undetermined;
+					}
+					return ECheckBoxState::Unchecked;
+				})
+			);
+	};
+	auto ToggleControls = [this](int32 Index)
+	{
+		return FUIAction(
+			FExecuteAction::CreateLambda([this, Index]
+				{
+					const TSharedPtr<ISequencer> SequencerPtr = WeakSequencer.Pin();
+					if (!SequencerPtr)
+					{
+						return;
+					}
+
+					UMovieSceneControlRigParameterSection* const ParameterSection = CastChecked<UMovieSceneControlRigParameterSection>(WeakSection.Get());
+					if (!IsValid(ParameterSection))
+					{
+						return;
+					}
+
+					FScopedTransaction Transaction(LOCTEXT("ToggleRigControlFiltersTransaction", "Toggle Rig Control Filters"));
+					ParameterSection->Modify();
+					if (Index >= 0)
 					{
-						FScopedTransaction Transaction(LOCTEXT("ToggleRigControlFiltersTransaction", "Toggle Rig Control Filters"));
-						ParameterSection->Modify();
-						if (Index >= 0)
+						ParameterSection->SetControlsMask(Index, !ParameterSection->GetControlsMask(Index));
+					}
+					else
+					{
+						ParameterSection->FillControlsMask(!ParameterSection->GetControlsMask(0));
+					}
+					SequencerPtr->NotifyMovieSceneDataChanged(EMovieSceneDataChangeType::MovieSceneStructureItemsChanged);
+				}
+			),
+			FCanExecuteAction(),
+			FGetActionCheckState::CreateLambda([this, Index]
+				{
+					UMovieSceneControlRigParameterSection* const ParameterSection = CastChecked<UMovieSceneControlRigParameterSection>(WeakSection.Get());
+					if (!IsValid(ParameterSection))
+					{
+						return ECheckBoxState::Unchecked;
+					}
+
+					TArray<bool> ControlBool = ParameterSection->GetControlsMask();
+					if (Index >= 0)
+					{
+						if (ControlBool[Index])
 						{
-							ParameterSection->SetControlsMask(Index, !ParameterSection->GetControlsMask(Index));
+							return ECheckBoxState::Checked;
 						}
 						else
 						{
-							ParameterSection->FillControlsMask(!ParameterSection->GetControlsMask(0));
+							return ECheckBoxState::Unchecked;
 						}
-						SequencerPtr->NotifyMovieSceneDataChanged(EMovieSceneDataChangeType::MovieSceneStructureItemsChanged);
-
 					}
-				),
-				FCanExecuteAction(),
-						FGetActionCheckState::CreateLambda([=]
+					else
+					{
+						TOptional<bool> FirstVal;
+						for (bool Val : ControlBool)
+						{
+							if (FirstVal.IsSet())
 							{
-								TArray<bool> ControlBool = ParameterSection->GetControlsMask();
-								if (Index >= 0)
-								{
-									if (ControlBool[Index])
-									{
-										return ECheckBoxState::Checked;
-
-									}
-									else
-									{
-										return ECheckBoxState::Unchecked;
-									}
-								}
-								else
+								if (Val != FirstVal)
 								{
-									TOptional<bool> FirstVal;
-									for (bool Val : ControlBool)
-									{
-										if (FirstVal.IsSet())
-										{
-											if (Val != FirstVal)
-											{
-												return ECheckBoxState::Undetermined;
-											}
-										}
-										else
-										{
-											FirstVal = Val;
-										}
-
-									}
-									return (FirstVal.IsSet() && FirstVal.GetValue()) ? ECheckBoxState::Checked : ECheckBoxState::Unchecked;
+									return ECheckBoxState::Undetermined;
 								}
-							})
-						);
-		};
-		UMovieSceneControlRigParameterTrack* Track = ParameterSection->GetTypedOuter<UMovieSceneControlRigParameterTrack>();
-		if (Track)
+							}
+							else
+							{
+								FirstVal = Val;
+							}
+						}
+						return (FirstVal.IsSet() && FirstVal.GetValue()) ? ECheckBoxState::Checked : ECheckBoxState::Unchecked;
+					}
+				})
+			);
+	};
+	UMovieSceneControlRigParameterTrack* Track = ParameterSection->GetTypedOuter<UMovieSceneControlRigParameterTrack>();
+	if (Track)
+	{
+		TArray<UMovieSceneSection*> Sections = Track->GetAllSections();
+		//If Base Absolute section 
+		if (ParameterSection->GetBlendType().Get() == EMovieSceneBlendType::Absolute && Sections[0] == ParameterSection)
 		{
-			TArray<UMovieSceneSection*> Sections = Track->GetAllSections();
-			//If Base Absolute section 
-			if (ParameterSection->GetBlendType().Get() == EMovieSceneBlendType::Absolute && Sections[0] == ParameterSection)
+			MenuBuilder.BeginSection(NAME_None, LOCTEXT("AnimationLayers", "Animation Layers"));
 			{
-				MenuBuilder.BeginSection(NAME_None, LOCTEXT("AnimationLayers", "Animation Layers"));
-				{
-					MenuBuilder.AddMenuEntry(
-						LOCTEXT("CollapseAllSections", "Collapse All Sections"),
-						LOCTEXT("CollapseAllSections_ToolTip", "Collapse all sections onto this section"),
-						FSlateIcon(),
-						FUIAction(FExecuteAction::CreateLambda([this] { CollapseAllLayers(); }))
-					);
-				}
-			}
-			if (ParameterSection->GetBlendType().Get() == EMovieSceneBlendType::Additive)
-			{
-
-				MenuBuilder.BeginSection(NAME_None, LOCTEXT("AnimationLayers", "Animation Layers"));
-				{
-					MenuBuilder.AddMenuEntry(
-						LOCTEXT("KeyZeroValue", "Key Zero Value"),
-						LOCTEXT("KeyZeroValue_Tooltip", "Set zero key on all controls in this section"),
-						FSlateIcon(),
-						FUIAction(FExecuteAction::CreateLambda([this] { KeyZeroValue(); }))
-					);
-				}
-				
 				MenuBuilder.AddMenuEntry(
-					LOCTEXT("KeyWeightZero", "Key Weight Zero"),
-					LOCTEXT("KeyWeightZero_Tooltip", "Key a zero value on the Weight channel"),
+					LOCTEXT("CollapseAllSections", "Collapse All Sections"),
+					LOCTEXT("CollapseAllSections_ToolTip", "Collapse all sections onto this section"),
 					FSlateIcon(),
-					FUIAction(FExecuteAction::CreateLambda([this] { KeyWeightValue(0.0f); }))
+					FUIAction(FExecuteAction::CreateLambda([this] { CollapseAllLayers(); }))
 				);
-				
+			}
+		}
+		if (ParameterSection->GetBlendType().Get() == EMovieSceneBlendType::Additive)
+		{
+			MenuBuilder.BeginSection(NAME_None, LOCTEXT("AnimationLayers", "Animation Layers"));
+			{
 				MenuBuilder.AddMenuEntry(
-					LOCTEXT("KeyWeightOne", "Key Weight One"),
-					LOCTEXT("KeyWeightOne_Tooltip", "Key a one value on the Weight channel"),
+					LOCTEXT("KeyZeroValue", "Key Zero Value"),
+					LOCTEXT("KeyZeroValue_Tooltip", "Set zero key on all controls in this section"),
 					FSlateIcon(),
-					FUIAction(FExecuteAction::CreateLambda([this] { KeyWeightValue(1.0f); }))
+					FUIAction(FExecuteAction::CreateLambda([this] { KeyZeroValue(); }))
 				);
-				
 			}
-		}
-		MenuBuilder.BeginSection(NAME_None, LOCTEXT("RigSectionActiveChannels", "Active Channels"));
-		{
+
 			MenuBuilder.AddMenuEntry(
-				LOCTEXT("SetFromSelectedControls", "Set From Selected Controls"),
-				LOCTEXT("SetFromSelectedControls_ToolTip", "Set active channels from the current control selection"),
+				LOCTEXT("KeyWeightZero", "Key Weight Zero"),
+				LOCTEXT("KeyWeightZero_Tooltip", "Key a zero value on the Weight channel"),
 				FSlateIcon(),
-				FUIAction(
-					FExecuteAction::CreateLambda([this] { ShowSelectedControlsChannels(); }),
-					FCanExecuteAction::CreateLambda([ControlRig] { return ControlRig->CurrentControlSelection().Num() > 0; } )
-				)
+				FUIAction(FExecuteAction::CreateLambda([this] { KeyWeightValue(0.0f); }))
 			);
 
 			MenuBuilder.AddMenuEntry(
-				LOCTEXT("ShowAllControls", "Show All Controls"),
-				LOCTEXT("ShowAllControls_ToolTip", "Set active channels from all controls"),
+				LOCTEXT("KeyWeightOne", "Key Weight One"),
+				LOCTEXT("KeyWeightOne_Tooltip", "Key a one value on the Weight channel"),
 				FSlateIcon(),
-				FUIAction(FExecuteAction::CreateLambda([this] { return ShowAllControlsChannels(); }))
+				FUIAction(FExecuteAction::CreateLambda([this] { KeyWeightValue(1.0f); }))
 			);
-
-			MenuBuilder.AddSubMenu(
-				LOCTEXT("AllTranslation", "Translation"), LOCTEXT("AllTranslation_ToolTip", "Causes this section to affect the translation of rig control transforms"),
-				FNewMenuDelegate::CreateLambda([=](FMenuBuilder& SubMenuBuilder) {
-					SubMenuBuilder.AddMenuEntry(
-						LOCTEXT("TranslationX", "X"), LOCTEXT("TranslationX_ToolTip", "Causes this section to affect the X channel of the transform's translation"),
-						FSlateIcon(), MakeUIAction(EMovieSceneTransformChannel::TranslationX), NAME_None, EUserInterfaceActionType::ToggleButton);
-					SubMenuBuilder.AddMenuEntry(
-						LOCTEXT("TranslationY", "Y"), LOCTEXT("TranslationY_ToolTip", "Causes this section to affect the Y channel of the transform's translation"),
-						FSlateIcon(), MakeUIAction(EMovieSceneTransformChannel::TranslationY), NAME_None, EUserInterfaceActionType::ToggleButton);
-					SubMenuBuilder.AddMenuEntry(
-						LOCTEXT("TranslationZ", "Z"), LOCTEXT("TranslationZ_ToolTip", "Causes this section to affect the Z channel of the transform's translation"),
-						FSlateIcon(), MakeUIAction(EMovieSceneTransformChannel::TranslationZ), NAME_None, EUserInterfaceActionType::ToggleButton);
-					}),
-				MakeUIAction(EMovieSceneTransformChannel::Translation),
-						NAME_None,
-						EUserInterfaceActionType::ToggleButton);
-
-			MenuBuilder.AddSubMenu(
-				LOCTEXT("AllRotation", "Rotation"), LOCTEXT("AllRotation_ToolTip", "Causes this section to affect the rotation of the rig control transform"),
-				FNewMenuDelegate::CreateLambda([=](FMenuBuilder& SubMenuBuilder) {
-					SubMenuBuilder.AddMenuEntry(
-						LOCTEXT("RotationX", "Roll (X)"), LOCTEXT("RotationX_ToolTip", "Causes this section to affect the roll (X) channel the transform's rotation"),
-						FSlateIcon(), MakeUIAction(EMovieSceneTransformChannel::RotationX), NAME_None, EUserInterfaceActionType::ToggleButton);
-					SubMenuBuilder.AddMenuEntry(
-						LOCTEXT("RotationY", "Pitch (Y)"), LOCTEXT("RotationY_ToolTip", "Causes this section to affect the pitch (Y) channel the transform's rotation"),
-						FSlateIcon(), MakeUIAction(EMovieSceneTransformChannel::RotationY), NAME_None, EUserInterfaceActionType::ToggleButton);
-					SubMenuBuilder.AddMenuEntry(
-						LOCTEXT("RotationZ", "Yaw (Z)"), LOCTEXT("RotationZ_ToolTip", "Causes this section to affect the yaw (Z) channel the transform's rotation"),
-						FSlateIcon(), MakeUIAction(EMovieSceneTransformChannel::RotationZ), NAME_None, EUserInterfaceActionType::ToggleButton);
-					}),
-				MakeUIAction(EMovieSceneTransformChannel::Rotation),
-						NAME_None,
-						EUserInterfaceActionType::ToggleButton);
-
-			MenuBuilder.AddSubMenu(
-				LOCTEXT("AllScale", "Scale"), LOCTEXT("AllScale_ToolTip", "Causes this section to affect the scale of the rig control transform"),
-				FNewMenuDelegate::CreateLambda([=](FMenuBuilder& SubMenuBuilder) {
-					SubMenuBuilder.AddMenuEntry(
-						LOCTEXT("ScaleX", "X"), LOCTEXT("ScaleX_ToolTip", "Causes this section to affect the X channel of the transform's scale"),
-						FSlateIcon(), MakeUIAction(EMovieSceneTransformChannel::ScaleX), NAME_None, EUserInterfaceActionType::ToggleButton);
-					SubMenuBuilder.AddMenuEntry(
-						LOCTEXT("ScaleY", "Y"), LOCTEXT("ScaleY_ToolTip", "Causes this section to affect the Y channel of the transform's scale"),
-						FSlateIcon(), MakeUIAction(EMovieSceneTransformChannel::ScaleY), NAME_None, EUserInterfaceActionType::ToggleButton);
-					SubMenuBuilder.AddMenuEntry(
-						LOCTEXT("ScaleZ", "Z"), LOCTEXT("ScaleZ_ToolTip", "Causes this section to affect the Z channel of the transform's scale"),
-						FSlateIcon(), MakeUIAction(EMovieSceneTransformChannel::ScaleZ), NAME_None, EUserInterfaceActionType::ToggleButton);
-					}),
-				MakeUIAction(EMovieSceneTransformChannel::Scale),
-						NAME_None,
-						EUserInterfaceActionType::ToggleButton);
-
-			//mz todo h
-			MenuBuilder.AddMenuEntry(
-				LOCTEXT("Weight", "Weight"), LOCTEXT("Weight_ToolTip", "Causes this section to be applied with a user-specified weight curve"),
-				FSlateIcon(), MakeUIAction(EMovieSceneTransformChannel::Weight), NAME_None, EUserInterfaceActionType::ToggleButton);
 		}
-		MenuBuilder.EndSection();
 	}
+	MenuBuilder.BeginSection(NAME_None, LOCTEXT("RigSectionActiveChannels", "Active Channels"));
+	{
+		MenuBuilder.AddMenuEntry(
+			LOCTEXT("SetFromSelectedControls", "Set From Selected Controls"),
+			LOCTEXT("SetFromSelectedControls_ToolTip", "Set active channels from the current control selection"),
+			FSlateIcon(),
+			FUIAction(
+				FExecuteAction::CreateLambda([this] { ShowSelectedControlsChannels(); }),
+				FCanExecuteAction::CreateLambda([ControlRig] { return ControlRig->CurrentControlSelection().Num() > 0; } )
+			)
+		);
+
+		MenuBuilder.AddMenuEntry(
+			LOCTEXT("ShowAllControls", "Show All Controls"),
+			LOCTEXT("ShowAllControls_ToolTip", "Set active channels from all controls"),
+			FSlateIcon(),
+			FUIAction(FExecuteAction::CreateLambda([this] { return ShowAllControlsChannels(); }))
+		);
+
+		MenuBuilder.AddSubMenu(
+			LOCTEXT("AllTranslation", "Translation"), LOCTEXT("AllTranslation_ToolTip", "Causes this section to affect the translation of rig control transforms"),
+			FNewMenuDelegate::CreateLambda([=](FMenuBuilder& SubMenuBuilder) {
+				SubMenuBuilder.AddMenuEntry(
+					LOCTEXT("TranslationX", "X"), LOCTEXT("TranslationX_ToolTip", "Causes this section to affect the X channel of the transform's translation"),
+					FSlateIcon(), MakeUIAction(EMovieSceneTransformChannel::TranslationX), NAME_None, EUserInterfaceActionType::ToggleButton);
+				SubMenuBuilder.AddMenuEntry(
+					LOCTEXT("TranslationY", "Y"), LOCTEXT("TranslationY_ToolTip", "Causes this section to affect the Y channel of the transform's translation"),
+					FSlateIcon(), MakeUIAction(EMovieSceneTransformChannel::TranslationY), NAME_None, EUserInterfaceActionType::ToggleButton);
+				SubMenuBuilder.AddMenuEntry(
+					LOCTEXT("TranslationZ", "Z"), LOCTEXT("TranslationZ_ToolTip", "Causes this section to affect the Z channel of the transform's translation"),
+					FSlateIcon(), MakeUIAction(EMovieSceneTransformChannel::TranslationZ), NAME_None, EUserInterfaceActionType::ToggleButton);
+				}),
+			MakeUIAction(EMovieSceneTransformChannel::Translation),
+					NAME_None,
+					EUserInterfaceActionType::ToggleButton);
+
+		MenuBuilder.AddSubMenu(
+			LOCTEXT("AllRotation", "Rotation"), LOCTEXT("AllRotation_ToolTip", "Causes this section to affect the rotation of the rig control transform"),
+			FNewMenuDelegate::CreateLambda([=](FMenuBuilder& SubMenuBuilder) {
+				SubMenuBuilder.AddMenuEntry(
+					LOCTEXT("RotationX", "Roll (X)"), LOCTEXT("RotationX_ToolTip", "Causes this section to affect the roll (X) channel the transform's rotation"),
+					FSlateIcon(), MakeUIAction(EMovieSceneTransformChannel::RotationX), NAME_None, EUserInterfaceActionType::ToggleButton);
+				SubMenuBuilder.AddMenuEntry(
+					LOCTEXT("RotationY", "Pitch (Y)"), LOCTEXT("RotationY_ToolTip", "Causes this section to affect the pitch (Y) channel the transform's rotation"),
+					FSlateIcon(), MakeUIAction(EMovieSceneTransformChannel::RotationY), NAME_None, EUserInterfaceActionType::ToggleButton);
+				SubMenuBuilder.AddMenuEntry(
+					LOCTEXT("RotationZ", "Yaw (Z)"), LOCTEXT("RotationZ_ToolTip", "Causes this section to affect the yaw (Z) channel the transform's rotation"),
+					FSlateIcon(), MakeUIAction(EMovieSceneTransformChannel::RotationZ), NAME_None, EUserInterfaceActionType::ToggleButton);
+				}),
+			MakeUIAction(EMovieSceneTransformChannel::Rotation),
+					NAME_None,
+					EUserInterfaceActionType::ToggleButton);
+
+		MenuBuilder.AddSubMenu(
+			LOCTEXT("AllScale", "Scale"), LOCTEXT("AllScale_ToolTip", "Causes this section to affect the scale of the rig control transform"),
+			FNewMenuDelegate::CreateLambda([=](FMenuBuilder& SubMenuBuilder) {
+				SubMenuBuilder.AddMenuEntry(
+					LOCTEXT("ScaleX", "X"), LOCTEXT("ScaleX_ToolTip", "Causes this section to affect the X channel of the transform's scale"),
+					FSlateIcon(), MakeUIAction(EMovieSceneTransformChannel::ScaleX), NAME_None, EUserInterfaceActionType::ToggleButton);
+				SubMenuBuilder.AddMenuEntry(
+					LOCTEXT("ScaleY", "Y"), LOCTEXT("ScaleY_ToolTip", "Causes this section to affect the Y channel of the transform's scale"),
+					FSlateIcon(), MakeUIAction(EMovieSceneTransformChannel::ScaleY), NAME_None, EUserInterfaceActionType::ToggleButton);
+				SubMenuBuilder.AddMenuEntry(
+					LOCTEXT("ScaleZ", "Z"), LOCTEXT("ScaleZ_ToolTip", "Causes this section to affect the Z channel of the transform's scale"),
+					FSlateIcon(), MakeUIAction(EMovieSceneTransformChannel::ScaleZ), NAME_None, EUserInterfaceActionType::ToggleButton);
+				}),
+			MakeUIAction(EMovieSceneTransformChannel::Scale),
+					NAME_None,
+					EUserInterfaceActionType::ToggleButton);
+
+		//mz todo h
+		MenuBuilder.AddMenuEntry(
+			LOCTEXT("Weight", "Weight"), LOCTEXT("Weight_ToolTip", "Causes this section to be applied with a user-specified weight curve"),
+			FSlateIcon(), MakeUIAction(EMovieSceneTransformChannel::Weight), NAME_None, EUserInterfaceActionType::ToggleButton);
+	}
+	MenuBuilder.EndSection();
 }
 
 void FControlRigParameterSection::ShowSelectedControlsChannels()
diff --git a/Engine/Plugins/Experimental/PropertyAnimator/Source/PropertyAnimatorEditor/Private/Sequencer/PropertyAnimatorEditorCurveChannelInterface.h b/Engine/Plugins/Experimental/PropertyAnimator/Source/PropertyAnimatorEditor/Private/Sequencer/PropertyAnimatorEditorCurveChannelInterface.h
--- a/Engine/Plugins/Experimental/PropertyAnimator/Source/PropertyAnimatorEditor/Private/Sequencer/PropertyAnimatorEditorCurveChannelInterface.h
+++ b/Engine/Plugins/Experimental/PropertyAnimator/Source/PropertyAnimatorEditor/Private/Sequencer/PropertyAnimatorEditorCurveChannelInterface.h
@@ -49,11 +49,12 @@ void TPropertyAnimatorEditorCurveChannelInterface<InChannelType, InMenuExtension
 	, TWeakPtr<ISequencer> InSequencer) const
 {
 	TSharedRef<FPropertyAnimatorEditorCurveSectionMenuExtension> Extension = MakeShared<FMenuExtensionType>(InChannels, InSections);
-	InMenuExtender->AddMenuExtension("SequencerChannels", EExtensionHook::First, nullptr, FMenuExtensionDelegate::CreateLambda(
-		[Extension](FMenuBuilder& InInnerMenuBuilder)
-		{
-			Extension->ExtendMenu(InInnerMenuBuilder);
-		}));
+
+	InMenuExtender->AddMenuExtension(TEXT("SequencerChannels"), EExtensionHook::First, nullptr
+		, FMenuExtensionDelegate::CreateLambda([Extension](FMenuBuilder& InInnerMenuBuilder)
+			{
+				Extension->ExtendMenu(InInnerMenuBuilder);
+			}));
 }
 
 template<typename InChannelType, typename InMenuExtensionType>
@@ -64,7 +65,12 @@ void TPropertyAnimatorEditorCurveChannelInterface<InChannelType, InMenuExtension
 	, TWeakPtr<ISequencer> InSequencer) const
 {
 	TSharedRef<FPropertyAnimatorEditorCurveSectionMenuExtension> Extension = MakeShared<FMenuExtensionType>(InChannels, InSections);
-	Extension->ExtendMenu(InMenuBuilder, false);
+
+	InMenuExtender->AddMenuExtension(TEXT("SequencerChannels"), EExtensionHook::First, nullptr
+		, FMenuExtensionDelegate::CreateLambda([Extension](FMenuBuilder& InInnerMenuBuilder)
+			{
+				Extension->ExtendMenu(InInnerMenuBuilder);
+			}));
 }
 
 template<typename InChannelType, typename InMenuExtensionType>
diff --git a/Engine/Plugins/MovieScene/LevelSequenceEditor/Source/LevelSequenceEditor/Private/LevelSequenceEditorSubsystem.cpp b/Engine/Plugins/MovieScene/LevelSequenceEditor/Source/LevelSequenceEditor/Private/LevelSequenceEditorSubsystem.cpp
--- a/Engine/Plugins/MovieScene/LevelSequenceEditor/Source/LevelSequenceEditor/Private/LevelSequenceEditorSubsystem.cpp
+++ b/Engine/Plugins/MovieScene/LevelSequenceEditor/Source/LevelSequenceEditor/Private/LevelSequenceEditorSubsystem.cpp
@@ -568,17 +568,18 @@ void ULevelSequenceEditorSubsystem::Initialize(FSubsystemCollectionBase& Collect
 		FMenuExtensionDelegate::CreateLambda([this](FMenuBuilder& MenuBuilder)
 			{
 				// Only add menu entries where the focused sequence is a ULevelSequence
-				if (GetActiveSequencer())
+				if (!GetActiveSequencer())
 				{
-					AddBindingPropertiesSidebar(MenuBuilder);
+					return;
 				}
 
-				TArray<FName> ComponentNames;
-				GetRebindComponentNames(ComponentNames);
-				if (ComponentNames.Num() > 0)
-				{
-					RebindComponentMenu(MenuBuilder);
-				}
+				AddBindingPropertiesSidebar(MenuBuilder);
+
+				FFormatNamedArguments Args;
+				MenuBuilder.AddSubMenu(
+					FText::Format(LOCTEXT("RebindComponent", "Rebind Component"), Args),
+					FText::Format(LOCTEXT("RebindComponentTooltip", "Rebind component by moving the tracks from one component to another component."), Args),
+					FNewMenuDelegate::CreateUObject(this, &ULevelSequenceEditorSubsystem::RebindComponentMenu));
 			}));
 
 	SidebarMenuExtender->AddMenuExtension(TEXT("CustomBinding"), EExtensionHook::First, CommandList,
diff --git a/Engine/Plugins/MovieScene/LevelSequenceEditor/Source/LevelSequenceEditor/Private/Misc/LevelSequenceCustomization.cpp b/Engine/Plugins/MovieScene/LevelSequenceEditor/Source/LevelSequenceEditor/Private/Misc/LevelSequenceCustomization.cpp
--- a/Engine/Plugins/MovieScene/LevelSequenceEditor/Source/LevelSequenceEditor/Private/Misc/LevelSequenceCustomization.cpp
+++ b/Engine/Plugins/MovieScene/LevelSequenceEditor/Source/LevelSequenceEditor/Private/Misc/LevelSequenceCustomization.cpp
@@ -58,7 +58,7 @@ void FLevelSequenceCustomization::RegisterSequencerCustomization(FSequencerCusto
 	// Add a customization callback for the object binding context menu.
 	FSequencerCustomizationInfo Customization;
 	Customization.OnBuildObjectBindingContextMenu = FOnGetSequencerMenuExtender::CreateRaw(this, &FLevelSequenceCustomization::CreateObjectBindingContextMenuExtender);
-	Customization.OnBuildSidebarMenu = FOnGetSequencerMenuExtender::CreateRaw(this, &FLevelSequenceCustomization::CreateSidebarMenuExtender);
+	Customization.OnBuildSidebarMenu = FOnGetSequencerMenuExtender::CreateRaw(this, &FLevelSequenceCustomization::CreateObjectBindingSidebarMenuExtender);
 	Builder.AddCustomization(Customization);
 }
 
@@ -142,17 +142,17 @@ void FLevelSequenceCustomization::ExtendObjectBindingContextMenu(FMenuBuilder& M
 		bool bMultipleBindings = false;
 		UObject* ResolutionContext = MovieSceneHelpers::GetResolutionContext(Sequence, ObjectBindingID, Sequencer->GetFocusedTemplateID(), Sequencer->GetSharedPlaybackState());
 
-		if (const FMovieSceneBindingReferences* BindingReferences = Sequencer->GetFocusedMovieSceneSequence()->GetBindingReferences())
+		if (const FMovieSceneBindingReferences* BindingReferences = Sequence->GetBindingReferences())
 		{
 			bCustomBinding = Algo::AnyOf(BindingReferences->GetReferences(ObjectBindingID), [](const FMovieSceneBindingReference& Reference) { return Reference.CustomBinding; });
 			bMultipleBindings = BindingReferences->GetReferences(ObjectBindingID).Num() > 1;
 			UE::UniversalObjectLocator::FResolveParams LocatorResolveParams(ResolutionContext);
-			FMovieSceneBindingResolveParams BindingResolveParams{ Sequencer->GetFocusedMovieSceneSequence(), ObjectBindingID, Sequencer->GetFocusedTemplateID(), ResolutionContext };
+			FMovieSceneBindingResolveParams BindingResolveParams{ Sequence, ObjectBindingID, Sequencer->GetFocusedTemplateID(), ResolutionContext };
 
 			// Can convert to possessable
 			int32 BindingIndex = 0;
 			bool bAnyValidConversions = false;
-			if (Algo::AnyOf(BindingReferences->GetReferences(ObjectBindingID), [&BindingIndex, Sequencer, &Sequence, &MovieScene](const FMovieSceneBindingReference& BindingReference) {
+			if (Algo::AnyOf(BindingReferences->GetReferences(ObjectBindingID), [&BindingIndex, Sequencer](const FMovieSceneBindingReference& BindingReference) {
 				return FSequencerUtilities::CanConvertToPossessable(Sequencer.ToSharedRef(), BindingReference.ID, BindingIndex++);
 				}))
 			{
@@ -164,7 +164,7 @@ void FLevelSequenceCustomization::ExtendObjectBindingContextMenu(FMenuBuilder& M
 				for (const TSubclassOf<UMovieSceneCustomBinding>& CustomBindingType : PrioritySortedCustomBindingTypes)
 				{
 					BindingIndex = 0;
-					if (Algo::AllOf(BindingReferences->GetReferences(ObjectBindingID), [&BindingIndex, &Sequence, &MovieScene, &CustomBindingType, Sequencer](const FMovieSceneBindingReference& BindingReference)
+					if (Algo::AllOf(BindingReferences->GetReferences(ObjectBindingID), [&BindingIndex, &CustomBindingType, Sequencer](const FMovieSceneBindingReference& BindingReference)
 						{
 							return FSequencerUtilities::CanConvertToCustomBinding(Sequencer.ToSharedRef(), BindingReference.ID, CustomBindingType, BindingIndex++);
 						}))
@@ -191,7 +191,7 @@ void FLevelSequenceCustomization::ExtendObjectBindingContextMenu(FMenuBuilder& M
 		else
 		{
 			MenuBuilder.BeginSection("CustomBinding");
-			bool bCustomSpawnable = MovieSceneHelpers::SupportsObjectTemplate(Sequencer->GetFocusedMovieSceneSequence(), ObjectBindingID, Sequencer->GetSharedPlaybackState());
+			bool bCustomSpawnable = MovieSceneHelpers::SupportsObjectTemplate(Sequence, ObjectBindingID, Sequencer->GetSharedPlaybackState());
 			// Check for custom binding types
 
 			if (bCustomSpawnable)
@@ -200,21 +200,32 @@ void FLevelSequenceCustomization::ExtendObjectBindingContextMenu(FMenuBuilder& M
 
 				if (!bMultipleBindings)
 				{
-					TArray<FSequencerChangeBindingInfo> Bindings;
-					const FMovieSceneBindingReferences* BindingReferences = Sequence->GetBindingReferences();
-					for (TViewModelPtr<IObjectBindingExtension> ObjectBindingNode : Sequencer->GetViewModel()->GetSelection()->Outliner.Filter<IObjectBindingExtension>())
-					{
-						int32 BindingIndex = 0;
-						for (const FMovieSceneBindingReference& Reference : BindingReferences->GetReferences(ObjectBindingNode->GetObjectGuid()))
-						{
-							Bindings.Add({ Reference.ID, BindingIndex++ });
-						}
-					}
-
 					MenuBuilder.AddSubMenu(
 						LOCTEXT("ChangeClassLabel", "Change Class"),
 						LOCTEXT("ChangeClassTooltip", "Change the class (object template) that this spawns from"),
-						FNewMenuDelegate::CreateLambda([=](FMenuBuilder& MenuBuilder) { FSequencerUtilities::AddChangeClassMenu(MenuBuilder, Sequencer.ToSharedRef(), Bindings, TFunction<void()>()); }));
+						FNewMenuDelegate::CreateLambda([this](FMenuBuilder& MenuBuilder)
+						{
+							const TSharedPtr<ISequencer> Sequencer = WeakSequencer.Pin();
+							if (!Sequencer.IsValid())
+							{
+								return;
+							}
+
+							UMovieSceneSequence* Sequence = Sequencer->GetFocusedMovieSceneSequence();
+
+							TArray<FSequencerChangeBindingInfo> Bindings;
+							const FMovieSceneBindingReferences* BindingReferences = Sequence->GetBindingReferences();
+							for (TViewModelPtr<IObjectBindingExtension> ObjectBindingNode : Sequencer->GetViewModel()->GetSelection()->Outliner.Filter<IObjectBindingExtension>())
+							{
+								int32 BindingIndex = 0;
+								for (const FMovieSceneBindingReference& Reference : BindingReferences->GetReferences(ObjectBindingNode->GetObjectGuid()))
+								{
+									Bindings.Add({ Reference.ID, BindingIndex++ });
+								}
+							}
+
+							FSequencerUtilities::AddChangeClassMenu(MenuBuilder, Sequencer.ToSharedRef(), Bindings, TFunction<void()>());
+						}));
 				}
 			}
 
@@ -228,18 +239,30 @@ void FLevelSequenceCustomization::ExtendObjectBindingContextMenu(FMenuBuilder& M
 		MenuBuilder.AddSubMenu(
 			LOCTEXT("ConvertBindingLabel", "Convert Selected Binding(s) To..."),
 			LOCTEXT("ConvertBindingLabelTooltip", "Convert selected bindings into another binding type"),
-			FNewMenuDelegate::CreateLambda([Sequencer, Sequence](FMenuBuilder& MenuBuilder) 
-			{ 
-					TArray<FSequencerChangeBindingInfo> Bindings;
-					const FMovieSceneBindingReferences* BindingReferences = Sequence->GetBindingReferences();
-					for (TViewModelPtr<IObjectBindingExtension> ObjectBindingNode : Sequencer->GetViewModel()->GetSelection()->Outliner.Filter<IObjectBindingExtension>())
+			FNewMenuDelegate::CreateLambda([this](FMenuBuilder& MenuBuilder) 
+			{
+				const TSharedPtr<ISequencer> Sequencer = WeakSequencer.Pin();
+				if (!Sequencer.IsValid())
+				{
+					return;
+				}
+
+				UMovieSceneSequence* const Sequence = Sequencer->GetFocusedMovieSceneSequence();
+				if (!IsValid(Sequence))
+				{
+					return;
+				}
+
+				TArray<FSequencerChangeBindingInfo> Bindings;
+				const FMovieSceneBindingReferences* BindingReferences = Sequence->GetBindingReferences();
+				for (TViewModelPtr<IObjectBindingExtension> ObjectBindingNode : Sequencer->GetViewModel()->GetSelection()->Outliner.Filter<IObjectBindingExtension>())
+				{
+					int32 BindingIndex = 0;
+					for (const FMovieSceneBindingReference& Reference : BindingReferences->GetReferences(ObjectBindingNode->GetObjectGuid()))
 					{
-						int32 BindingIndex = 0;
-						for (const FMovieSceneBindingReference& Reference : BindingReferences->GetReferences(ObjectBindingNode->GetObjectGuid()))
-						{
-							Bindings.Add({ Reference.ID, BindingIndex++ });
-						}
+						Bindings.Add({ Reference.ID, BindingIndex++ });
 					}
+				}
 
 				FSequencerUtilities::AddConvertBindingMenu(MenuBuilder, Sequencer.ToSharedRef(), Bindings, TFunction<void()>()); 
 			}));
@@ -272,24 +295,21 @@ void FLevelSequenceCustomization::ExtendObjectBindingContextMenu(FMenuBuilder& M
 	MenuBuilder.EndSection();
 }
 
-TSharedPtr<FExtender> FLevelSequenceCustomization::CreateSidebarMenuExtender(FViewModelPtr InViewModel)
+TSharedPtr<FExtender> FLevelSequenceCustomization::CreateObjectBindingSidebarMenuExtender(FViewModelPtr InViewModel)
 {
 	TSharedRef<FExtender> Extender = MakeShared<FExtender>();
 	
 	TSharedPtr<FObjectBindingModel> ObjectBindingModel = InViewModel->CastThisShared<FObjectBindingModel>();
 	
 	Extender->AddMenuExtension(TEXT("ObjectBindingActions"), EExtensionHook::Before, nullptr,
-		FMenuExtensionDelegate::CreateRaw(this, &FLevelSequenceCustomization::ExtendSidebarMenu, ObjectBindingModel));
+		FMenuExtensionDelegate::CreateRaw(this, &FLevelSequenceCustomization::ExtendObjectBindingSidebarMenu, ObjectBindingModel));
 	
 	return Extender.ToSharedPtr();
 }
 
-void FLevelSequenceCustomization::ExtendSidebarMenu(FMenuBuilder& MenuBuilder, TSharedPtr<FObjectBindingModel> ObjectBindingModel)
+void FLevelSequenceCustomization::ExtendObjectBindingSidebarMenu(FMenuBuilder& MenuBuilder, TSharedPtr<FObjectBindingModel> ObjectBindingModel)
 {
-	if (ObjectBindingModel.IsValid())
-	{
-		ObjectBindingModel->BuildContextMenu(MenuBuilder);
-	}
+	ExtendObjectBindingContextMenu(MenuBuilder, ObjectBindingModel);
 }
 
 } // namespace UE::Sequencer
diff --git a/Engine/Plugins/MovieScene/LevelSequenceEditor/Source/LevelSequenceEditor/Private/Misc/LevelSequenceCustomization.h b/Engine/Plugins/MovieScene/LevelSequenceEditor/Source/LevelSequenceEditor/Private/Misc/LevelSequenceCustomization.h
--- a/Engine/Plugins/MovieScene/LevelSequenceEditor/Source/LevelSequenceEditor/Private/Misc/LevelSequenceCustomization.h
+++ b/Engine/Plugins/MovieScene/LevelSequenceEditor/Source/LevelSequenceEditor/Private/Misc/LevelSequenceCustomization.h
@@ -35,8 +35,9 @@ private:
 	TSharedPtr<FExtender> CreateObjectBindingContextMenuExtender(FViewModelPtr InViewModel);
 	void ExtendObjectBindingContextMenu(FMenuBuilder& MenuBuilder, TSharedPtr<FObjectBindingModel> ObjectBindingModel);
 
-	TSharedPtr<FExtender> CreateSidebarMenuExtender(FViewModelPtr InViewModel);
-	void ExtendSidebarMenu(FMenuBuilder& MenuBuilder, TSharedPtr<FObjectBindingModel> ObjectBindingModel);
+	// Object binding sidebar menu extensions
+	TSharedPtr<FExtender> CreateObjectBindingSidebarMenuExtender(FViewModelPtr InViewModel);
+	void ExtendObjectBindingSidebarMenu(FMenuBuilder& MenuBuilder, TSharedPtr<FObjectBindingModel> ObjectBindingModel);
 
 private:
 
diff --git a/Engine/Source/Editor/MovieSceneTools/Private/Sections/TransformPropertySection.cpp b/Engine/Source/Editor/MovieSceneTools/Private/Sections/TransformPropertySection.cpp
--- a/Engine/Source/Editor/MovieSceneTools/Private/Sections/TransformPropertySection.cpp
+++ b/Engine/Source/Editor/MovieSceneTools/Private/Sections/TransformPropertySection.cpp
@@ -28,14 +28,14 @@ class UObject;
 
 void FTransformSection::BuildSectionContextMenu(FMenuBuilder& MenuBuilder, const FGuid& InObjectBinding)
 {
-	UMovieScene3DTransformSection* TransformSection = CastChecked<UMovieScene3DTransformSection>(WeakSection.Get());
-	TSharedPtr<ISequencer> SequencerPtr = WeakSequencer.Pin();
-
-	auto MakeUIAction = [=](EMovieSceneTransformChannel ChannelsToToggle)
+	auto MakeUIAction = [this, InObjectBinding](EMovieSceneTransformChannel ChannelsToToggle)
 	{
 		return FUIAction(
-			FExecuteAction::CreateLambda([=]
+			FExecuteAction::CreateLambda([this, InObjectBinding, ChannelsToToggle]
 				{
+					UMovieScene3DTransformSection* TransformSection = CastChecked<UMovieScene3DTransformSection>(WeakSection.Get());
+					TSharedPtr<ISequencer> SequencerPtr = WeakSequencer.Pin();
+
 					FScopedTransaction Transaction(LOCTEXT("SetActiveChannelsTransaction", "Set Active Channels"));
 					TransformSection->Modify();
 					EMovieSceneTransformChannel Channels = TransformSection->GetMask().GetChannels();
@@ -62,9 +62,15 @@ void FTransformSection::BuildSectionContextMenu(FMenuBuilder& MenuBuilder, const
 				}
 			),
 			FCanExecuteAction(),
-			FGetActionCheckState::CreateLambda([=]
+			FGetActionCheckState::CreateLambda([this, ChannelsToToggle]
 			{
-				EMovieSceneTransformChannel Channels = TransformSection->GetMask().GetChannels();
+				const UMovieScene3DTransformSection* const TransformSection = CastChecked<UMovieScene3DTransformSection>(WeakSection.Get());
+				if (!IsValid(TransformSection))
+				{
+					return ECheckBoxState::Unchecked;
+				}
+
+				const EMovieSceneTransformChannel Channels = TransformSection->GetMask().GetChannels();
 				if (EnumHasAllFlags(Channels, ChannelsToToggle))
 				{
 					return ECheckBoxState::Checked;
diff --git a/Engine/Source/Editor/MovieSceneTools/Private/TrackEditors/CinematicShotTrackEditor.cpp b/Engine/Source/Editor/MovieSceneTools/Private/TrackEditors/CinematicShotTrackEditor.cpp
--- a/Engine/Source/Editor/MovieSceneTools/Private/TrackEditors/CinematicShotTrackEditor.cpp
+++ b/Engine/Source/Editor/MovieSceneTools/Private/TrackEditors/CinematicShotTrackEditor.cpp
@@ -214,11 +214,6 @@ void FCinematicShotTrackEditor::BuildTrackContextMenu(FMenuBuilder& MenuBuilder,
 	MenuBuilder.EndSection();
 }
 
-void FCinematicShotTrackEditor::BuildTrackSidebarMenu(FMenuBuilder& MenuBuilder, UMovieSceneTrack* Track)
-{
-	BuildTrackContextMenu(MenuBuilder, Track);
-}
-
 void FCinematicShotTrackEditor::InsertShot()
 {
 	UMovieScene* FocusedMovieScene = GetFocusedMovieScene();
diff --git a/Engine/Source/Editor/MovieSceneTools/Private/TrackEditors/SkeletalAnimationTrackEditor.cpp b/Engine/Source/Editor/MovieSceneTools/Private/TrackEditors/SkeletalAnimationTrackEditor.cpp
--- a/Engine/Source/Editor/MovieSceneTools/Private/TrackEditors/SkeletalAnimationTrackEditor.cpp
+++ b/Engine/Source/Editor/MovieSceneTools/Private/TrackEditors/SkeletalAnimationTrackEditor.cpp
@@ -2038,13 +2038,10 @@ FKeyPropertyResult FSkeletalAnimationTrackEditor::AddKeyInternal( FFrameNumber K
 	return KeyPropertyResult;
 }
 
-void FSkeletalAnimationTrackEditor::BuildTrackContextMenu( FMenuBuilder& MenuBuilder, UMovieSceneTrack* Track )
+void FSkeletalAnimationTrackEditor::BuildTrackContextMenu_Internal(FMenuBuilder& MenuBuilder, UMovieSceneTrack* Track, const bool bAddSeparatorAtEnd)
 {
 	//there's a bug with a section being open already, so we end it.
-
-	TSharedPtr<ISequencer> SequencerPtr = GetSequencer();
-
-	UMovieSceneSkeletalAnimationTrack* SkeletalAnimationTrack = Cast<UMovieSceneSkeletalAnimationTrack>( Track );
+	UMovieSceneSkeletalAnimationTrack* SkeletalAnimationTrack = Cast<UMovieSceneSkeletalAnimationTrack>(Track);
 	/** Put this back when and if it works
 	MenuBuilder.BeginSection(NAME_None, LOCTEXT("MotionBlendingOptions", "Motion Blending Options"));
 	{
@@ -2074,34 +2071,57 @@ void FSkeletalAnimationTrackEditor::BuildTrackContextMenu( FMenuBuilder& MenuBui
 			NSLOCTEXT("Sequencer", "BlendFirstChildOfRootTooltip", "If True, do not blend and match the root bones but instead the first child bone of the root. Toggle this on when the matched sequences in the track have no motion on the root."),
 			FSlateIcon(),
 			FUIAction(
-				FExecuteAction::CreateLambda([SequencerPtr,SkeletalAnimationTrack]()->void {
+				FExecuteAction::CreateLambda([this, SkeletalAnimationTrack]() -> void
+				{
 					SkeletalAnimationTrack->bBlendFirstChildOfRoot = SkeletalAnimationTrack->bBlendFirstChildOfRoot ? false : true;
-		SkeletalAnimationTrack->SetRootMotionsDirty();
-		SequencerPtr->NotifyMovieSceneDataChanged(EMovieSceneDataChangeType::TrackValueChanged);
+					SkeletalAnimationTrack->SetRootMotionsDirty();
 
-					}),
-				FCanExecuteAction::CreateLambda([SequencerPtr, SkeletalAnimationTrack]()->bool { return SequencerPtr && SkeletalAnimationTrack != nullptr; }),
-						FIsActionChecked::CreateLambda([SequencerPtr, SkeletalAnimationTrack]()->bool { return SkeletalAnimationTrack != nullptr && SkeletalAnimationTrack->bBlendFirstChildOfRoot; })),
-			NAME_None,
-						EUserInterfaceActionType::ToggleButton
-						);
+					if (const TSharedPtr<ISequencer> SequencerPtr = GetSequencer())
+					{
+						SequencerPtr->NotifyMovieSceneDataChanged(EMovieSceneDataChangeType::TrackValueChanged);
+					}
+				}),
+				FCanExecuteAction::CreateLambda([this, SkeletalAnimationTrack]() -> bool
+				{
+					return SkeletalAnimationTrack != nullptr;
+				}),
+				FIsActionChecked::CreateLambda([this, SkeletalAnimationTrack]() -> bool
+				{
+					return SkeletalAnimationTrack != nullptr && SkeletalAnimationTrack->bBlendFirstChildOfRoot;
+				})),
+				NAME_None,
+				EUserInterfaceActionType::ToggleButton
+			);
 
 		MenuBuilder.AddMenuEntry(
 			NSLOCTEXT("Sequencer", "ShowRootMotionTrails", "Show Root Motion Trail"),
 			NSLOCTEXT("Sequencer", "ShowRootMotionTrailsTooltip", "Show the Root Motion Trail for all Animation Clips."),
 			FSlateIcon(),
 			FUIAction(
-				FExecuteAction::CreateLambda([SequencerPtr, SkeletalAnimationTrack]()->void {
+				FExecuteAction::CreateLambda([this, SkeletalAnimationTrack]() -> void
+				{
 					SkeletalAnimationTrack->ToggleShowRootMotionTrail();
-		SequencerPtr->NotifyMovieSceneDataChanged(EMovieSceneDataChangeType::TrackValueChanged);
 
-					}),
-				FCanExecuteAction::CreateLambda([SequencerPtr, SkeletalAnimationTrack]()->bool { return SequencerPtr && SkeletalAnimationTrack != nullptr; }),
-						FIsActionChecked::CreateLambda([SequencerPtr, SkeletalAnimationTrack]()->bool { return SkeletalAnimationTrack != nullptr && SkeletalAnimationTrack->bShowRootMotionTrail; })),
+					if (const TSharedPtr<ISequencer> SequencerPtr = GetSequencer())
+					{
+						SequencerPtr->NotifyMovieSceneDataChanged(EMovieSceneDataChangeType::TrackValueChanged);
+					}
+				}),
+				FCanExecuteAction::CreateLambda([this, SkeletalAnimationTrack]()->bool
+				{
+					if (const TSharedPtr<ISequencer> SequencerPtr = GetSequencer())
+					{
+						return SkeletalAnimationTrack != nullptr;
+					}
+					return false;
+				}),
+				FIsActionChecked::CreateLambda([this, SkeletalAnimationTrack]()->bool
+				{
+					return SkeletalAnimationTrack != nullptr && SkeletalAnimationTrack->bShowRootMotionTrail;
+				})),
 			NAME_None,
-						EUserInterfaceActionType::ToggleButton
-						);
-
+				EUserInterfaceActionType::ToggleButton
+			);
 
 		MenuBuilder.AddMenuEntry(
 			NSLOCTEXT("Sequencer", "SwapRootBoneNone", "Swap Root Bone None"),
@@ -2145,7 +2165,21 @@ void FSkeletalAnimationTrackEditor::BuildTrackContextMenu( FMenuBuilder& MenuBui
 			);
 	}
 	MenuBuilder.EndSection();
-	MenuBuilder.AddSeparator();
+
+	if (bAddSeparatorAtEnd)
+	{
+		MenuBuilder.AddSeparator();
+	}
+}
+
+void FSkeletalAnimationTrackEditor::BuildTrackContextMenu(FMenuBuilder& MenuBuilder, UMovieSceneTrack* Track)
+{
+	BuildTrackContextMenu_Internal(MenuBuilder, Track, true);
+}
+
+void FSkeletalAnimationTrackEditor::BuildTrackSidebarMenu(FMenuBuilder& MenuBuilder, UMovieSceneTrack* Track)
+{
+	BuildTrackContextMenu_Internal(MenuBuilder, Track, false);
 }
 
 TSharedPtr<SWidget> FSkeletalAnimationTrackEditor::BuildOutlinerEditWidget(const FGuid& ObjectBinding, UMovieSceneTrack* Track, const FBuildEditWidgetParams& Params)
diff --git a/Engine/Source/Editor/MovieSceneTools/Private/TrackEditors/SkeletalAnimationTrackEditor.h b/Engine/Source/Editor/MovieSceneTools/Private/TrackEditors/SkeletalAnimationTrackEditor.h
--- a/Engine/Source/Editor/MovieSceneTools/Private/TrackEditors/SkeletalAnimationTrackEditor.h
+++ b/Engine/Source/Editor/MovieSceneTools/Private/TrackEditors/SkeletalAnimationTrackEditor.h
@@ -69,7 +69,8 @@ public:
 	virtual TSharedRef<ISequencerSection> MakeSectionInterface( UMovieSceneSection& SectionObject, UMovieSceneTrack& Track, FGuid ObjectBinding ) override;
 	virtual bool SupportsSequence(UMovieSceneSequence* InSequence) const override;
 	virtual bool SupportsType( TSubclassOf<UMovieSceneTrack> Type ) const override;
-	virtual void BuildTrackContextMenu( FMenuBuilder& MenuBuilder, UMovieSceneTrack* Track ) override;
+	virtual void BuildTrackContextMenu(FMenuBuilder& MenuBuilder, UMovieSceneTrack* Track) override;
+	virtual void BuildTrackSidebarMenu(FMenuBuilder& MenuBuilder, UMovieSceneTrack* Track) override;
 	virtual TSharedPtr<SWidget> BuildOutlinerEditWidget(const FGuid& ObjectBinding, UMovieSceneTrack* Track, const FBuildEditWidgetParams& Params) override;
 	virtual bool OnAllowDrop(const FDragDropEvent& DragDropEvent, FSequencerDragDropParams& DragDropParams) override;
 	virtual FReply OnDrop(const FDragDropEvent& DragDropEvent, const FSequencerDragDropParams& DragDropParams) override;
@@ -112,6 +113,9 @@ private:
 	/** Can Open the linked Anim Sequence*/
 	bool CanOpenLinkedAnimSequence(FGuid Binding);
 
+	/** Common function used to build context and sidebar menus */
+	void BuildTrackContextMenu_Internal(FMenuBuilder& MenuBuilder, UMovieSceneTrack* const InTrack, const bool bAddSeparatorAtEnd);
+
 	friend class FMovieSceneSkeletalAnimationParamsDetailCustomization;
 
 private:
diff --git a/Engine/Source/Editor/MovieSceneTools/Private/TrackEditors/TransformTrackEditor.cpp b/Engine/Source/Editor/MovieSceneTools/Private/TrackEditors/TransformTrackEditor.cpp
--- a/Engine/Source/Editor/MovieSceneTools/Private/TrackEditors/TransformTrackEditor.cpp
+++ b/Engine/Source/Editor/MovieSceneTools/Private/TrackEditors/TransformTrackEditor.cpp
@@ -174,41 +174,58 @@ void F3DTransformTrackEditor::BuildTrackContextMenu( FMenuBuilder& MenuBuilder,
 {
 	UMovieScene3DTransformTrack* TransformTrack = Cast<UMovieScene3DTransformTrack>( Track );
 
-	auto AnimSubMenuDelegate = [](FMenuBuilder& InMenuBuilder, TSharedRef<ISequencer> InSequencer, UMovieScene3DTransformTrack* InTransformTrack)
-	{
-		UMovieSceneSequence* Sequence = InSequencer->GetFocusedMovieSceneSequence();
-
-		FAssetPickerConfig AssetPickerConfig;
-		AssetPickerConfig.bAddFilterUI = true;
-		AssetPickerConfig.SelectionMode = ESelectionMode::Single;
-		AssetPickerConfig.Filter.ClassPaths.Add(UAnimSequence::StaticClass()->GetClassPathName());
-		AssetPickerConfig.OnAssetSelected = FOnAssetSelected::CreateStatic(&F3DTransformTrackEditor::ImportAnimSequenceTransforms, InSequencer, InTransformTrack);
-		AssetPickerConfig.OnAssetEnterPressed = FOnAssetEnterPressed::CreateStatic(&F3DTransformTrackEditor::ImportAnimSequenceTransformsEnterPressed, InSequencer, InTransformTrack);
-		AssetPickerConfig.SaveSettingsName = TEXT("SequencerAssetPicker");
-		AssetPickerConfig.AdditionalReferencingAssets.Add(FAssetData(Sequence));
-
-		FContentBrowserModule& ContentBrowserModule = FModuleManager::LoadModuleChecked<FContentBrowserModule>(TEXT("ContentBrowser"));
-
-		InMenuBuilder.AddWidget(
-			SNew(SBox)
-			.WidthOverride(200.0f)
-			.HeightOverride(400.0f)
-			[
-				ContentBrowserModule.Get().CreateAssetPicker(AssetPickerConfig)
-			], 
-			FText(), true, false);
-	};
-
 	MenuBuilder.AddSubMenu(
 		NSLOCTEXT( "Sequencer", "ImportTransforms", "Import From Animation Root" ),
 		NSLOCTEXT( "Sequencer", "ImportTransformsTooltip", "Import transform keys from an animation sequence's root motion." ),
-		FNewMenuDelegate::CreateLambda(AnimSubMenuDelegate, GetSequencer().ToSharedRef(), TransformTrack)
+		FNewMenuDelegate::CreateSP(this, &F3DTransformTrackEditor::BuildAssetPickerSubMenu, TransformTrack)
 	);
 
 	MenuBuilder.AddMenuSeparator();
+
 	FKeyframeTrackEditor::BuildTrackContextMenu(MenuBuilder, Track);
 }
 
+void F3DTransformTrackEditor::BuildAssetPickerSubMenu(FMenuBuilder& InMenuBuilder, UMovieScene3DTransformTrack* InTransformTrack)
+{
+	const TSharedPtr<ISequencer> TrackSequencer = GetSequencer();
+	if (!TrackSequencer.IsValid())
+	{
+		return;
+	}
+
+	TWeakObjectPtr<UMovieScene3DTransformTrack> TransformTrackWeak = InTransformTrack;
+
+	UMovieSceneSequence* Sequence = TrackSequencer->GetFocusedMovieSceneSequence();
+
+	FAssetPickerConfig AssetPickerConfig;
+	AssetPickerConfig.bAddFilterUI = true;
+	AssetPickerConfig.SelectionMode = ESelectionMode::Single;
+	AssetPickerConfig.Filter.ClassPaths.Add(UAnimSequence::StaticClass()->GetClassPathName());
+	AssetPickerConfig.OnAssetSelected = FOnAssetSelected::CreateLambda(
+		[this, TransformTrackWeak](const FAssetData& InAssetData)
+		{
+			ImportAnimSequenceTransforms(InAssetData, TransformTrackWeak.Get());
+		});
+	AssetPickerConfig.OnAssetEnterPressed = FOnAssetEnterPressed::CreateLambda(
+		[this, TransformTrackWeak](const TArray<FAssetData>& InSelectedAssets)
+		{
+			ImportAnimSequenceTransformsEnterPressed(InSelectedAssets, TransformTrackWeak.Get());
+		});
+	AssetPickerConfig.SaveSettingsName = TEXT("SequencerAssetPicker");
+	AssetPickerConfig.AdditionalReferencingAssets.Add(FAssetData(Sequence));
+
+	FContentBrowserModule& ContentBrowserModule = FModuleManager::LoadModuleChecked<FContentBrowserModule>(TEXT("ContentBrowser"));
+
+	InMenuBuilder.AddWidget(
+		SNew(SBox)
+		.WidthOverride(200.0f)
+		.HeightOverride(400.0f)
+		[
+			ContentBrowserModule.Get().CreateAssetPicker(AssetPickerConfig)
+		], 
+		FText(), true, false);
+}
+
 
 TSharedRef<ISequencerSection> F3DTransformTrackEditor::MakeSectionInterface(UMovieSceneSection& SectionObject, UMovieSceneTrack& Track, FGuid ObjectBinding)
 {
@@ -1533,17 +1550,28 @@ void AddUnwoundKey(FMovieSceneDoubleChannel& Channel, FFrameNumber Time, double
 }
 
 
-void F3DTransformTrackEditor::ImportAnimSequenceTransforms(const FAssetData& Asset, TSharedRef<ISequencer> Sequencer, UMovieScene3DTransformTrack* TransformTrack)
+void F3DTransformTrackEditor::ImportAnimSequenceTransforms(const FAssetData& Asset, UMovieScene3DTransformTrack* TransformTrack)
 {
+	if (!IsValid(TransformTrack))
+	{
+		return;
+	}
+
+	const TSharedPtr<ISequencer> SequencerPtr = GetSequencer();
+	if (!SequencerPtr.IsValid())
+	{
+		return;
+	}
+
 	FSlateApplication::Get().DismissAllMenus();
 
-	FQualifiedFrameTime CurrentTime = Sequencer->GetLocalTime();
+	FQualifiedFrameTime CurrentTime = SequencerPtr->GetLocalTime();
 
 	UAnimSequence* AnimSequence = Cast<UAnimSequence>(Asset.GetAsset());
 
 	// find object binding to recover any component transforms we need to incorporate (for characters)
 	FTransform InvComponentTransform;
-	UMovieSceneSequence* MovieSceneSequence = Sequencer->GetFocusedMovieSceneSequence();
+	UMovieSceneSequence* MovieSceneSequence = SequencerPtr->GetFocusedMovieSceneSequence();
 	if(MovieSceneSequence)
 	{
 		UMovieScene* MovieScene = MovieSceneSequence->GetMovieScene();
@@ -1685,16 +1713,16 @@ void F3DTransformTrackEditor::ImportAnimSequenceTransforms(const FAssetData& Ass
 			Section->SetRange(Range);
 			Section->SetRowIndex(MovieSceneToolHelpers::FindAvailableRowIndex(TransformTrack, Section));
 
-			Sequencer->NotifyMovieSceneDataChanged( EMovieSceneDataChangeType::MovieSceneStructureItemAdded );
+			SequencerPtr->NotifyMovieSceneDataChanged( EMovieSceneDataChangeType::MovieSceneStructureItemAdded );
 		}
 	}
 }
 
-void F3DTransformTrackEditor::ImportAnimSequenceTransformsEnterPressed(const TArray<FAssetData>& Asset, TSharedRef<ISequencer> Sequencer, UMovieScene3DTransformTrack* TransformTrack)
+void F3DTransformTrackEditor::ImportAnimSequenceTransformsEnterPressed(const TArray<FAssetData>& Asset, UMovieScene3DTransformTrack* TransformTrack)
 {
 	if (Asset.Num() > 0)
 	{
-		ImportAnimSequenceTransforms(Asset[0].GetAsset(), Sequencer, TransformTrack);
+		ImportAnimSequenceTransforms(Asset[0].GetAsset(), TransformTrack);
 	}
 }
 
diff --git a/Engine/Source/Editor/MovieSceneTools/Private/TrackEditors/TransformTrackEditor.h b/Engine/Source/Editor/MovieSceneTools/Private/TrackEditors/TransformTrackEditor.h
--- a/Engine/Source/Editor/MovieSceneTools/Private/TrackEditors/TransformTrackEditor.h
+++ b/Engine/Source/Editor/MovieSceneTools/Private/TrackEditors/TransformTrackEditor.h
@@ -178,10 +178,10 @@ private:
 private:
 
 	/** Import an animation sequence's root transforms into a transform section */
-	static void ImportAnimSequenceTransforms(const FAssetData& Asset, TSharedRef<class ISequencer> Sequencer, UMovieScene3DTransformTrack* TransformTrack);
+	void ImportAnimSequenceTransforms(const FAssetData& Asset, UMovieScene3DTransformTrack* TransformTrack);
 
 	/** Import an animation sequence's root transforms into a transform section */
-	static void ImportAnimSequenceTransformsEnterPressed(const TArray<FAssetData>& Asset, TSharedRef<class ISequencer> Sequencer, UMovieScene3DTransformTrack* TransformTrack);
+	void ImportAnimSequenceTransformsEnterPressed(const TArray<FAssetData>& Asset, UMovieScene3DTransformTrack* TransformTrack);
 
 	/** ConstraintChannel Delegates*/
 	FDelegateHandle OnSceneComponentConstrainedHandle;
@@ -194,6 +194,9 @@ private:
 	void HandleConstraintPropertyChanged(UTickableTransformConstraint* InConstraint, const FPropertyChangedEvent& InPropertyChangedEvent) const;
 	
 	void ClearOutConstraintDelegates();
+
+	void BuildAssetPickerSubMenu(FMenuBuilder& InMenuBuilder, UMovieScene3DTransformTrack* InTransformTrack);
+
 private:
 
 	static FName TransformPropertyName;
diff --git a/Engine/Source/Editor/MovieSceneTools/Public/Channels/PerlinNoiseChannelInterface.h b/Engine/Source/Editor/MovieSceneTools/Public/Channels/PerlinNoiseChannelInterface.h
--- a/Engine/Source/Editor/MovieSceneTools/Public/Channels/PerlinNoiseChannelInterface.h
+++ b/Engine/Source/Editor/MovieSceneTools/Public/Channels/PerlinNoiseChannelInterface.h
@@ -128,13 +128,22 @@ struct TPerlinNoiseChannelInterface : ISequencerChannelInterface
 	{
 		TSharedRef<FPerlinNoiseChannelSectionMenuExtension> Extension = MakeShared<FPerlinNoiseChannelSectionMenuExtension>(Channels, Sections);
 
-		MenuExtender->AddMenuExtension("SequencerChannels", EExtensionHook::First, nullptr, FMenuExtensionDelegate::CreateLambda([Extension](FMenuBuilder& MenuBuilder) { Extension->ExtendMenu(MenuBuilder); }));
+		MenuExtender->AddMenuExtension(TEXT("SequencerChannels"), EExtensionHook::First, nullptr
+			, FMenuExtensionDelegate::CreateLambda([Extension](FMenuBuilder& MenuBuilder)
+				{
+					Extension->ExtendMenu(MenuBuilder);
+				}));
 	}
 
 	virtual void ExtendSidebarMenu_Raw(FMenuBuilder& MenuBuilder, TSharedPtr<FExtender> MenuExtender, TArrayView<const FMovieSceneChannelHandle> Channels, TArrayView<UMovieSceneSection* const> Sections, TWeakPtr<ISequencer> InSequencer) const override
 	{
 		TSharedRef<FPerlinNoiseChannelSectionMenuExtension> Extension = MakeShared<FPerlinNoiseChannelSectionMenuExtension>(Channels, Sections);
-		Extension->ExtendMenu(MenuBuilder, false);
+
+		MenuExtender->AddMenuExtension(TEXT("SequencerChannels"), EExtensionHook::First, nullptr
+			, FMenuExtensionDelegate::CreateLambda([Extension](FMenuBuilder& MenuBuilder)
+				{
+					Extension->ExtendMenu(MenuBuilder);
+				}));
 	}
 
 	virtual void DrawKeys_Raw(FMovieSceneChannel* Channel, TArrayView<const FKeyHandle> InKeyHandles, const UMovieSceneSection* InOwner, TArrayView<FKeyDrawParams> OutKeyDrawParams) const override
diff --git a/Engine/Source/Editor/MovieSceneTools/Public/TrackEditors/CinematicShotTrackEditor.h b/Engine/Source/Editor/MovieSceneTools/Public/TrackEditors/CinematicShotTrackEditor.h
--- a/Engine/Source/Editor/MovieSceneTools/Public/TrackEditors/CinematicShotTrackEditor.h
+++ b/Engine/Source/Editor/MovieSceneTools/Public/TrackEditors/CinematicShotTrackEditor.h
@@ -56,7 +56,6 @@ public:
 	virtual bool SupportsSequence(UMovieSceneSequence* InSequence) const override;
 	virtual void Tick(float DeltaTime) override;
 	virtual void BuildTrackContextMenu(FMenuBuilder& MenuBuilder, UMovieSceneTrack* Track) override;
-	virtual void BuildTrackSidebarMenu(FMenuBuilder& MenuBuilder, UMovieSceneTrack* Track) override;
 
 	UE_DEPRECATED(5.3, "InsertShot has been deprecated in favor of FSubTrackEditor::InsertSection(UMovieSceneTrack*)")
 	void InsertShot();
diff --git a/Engine/Source/Editor/Sequencer/Private/MVVM/ViewModels/CategoryModel.cpp b/Engine/Source/Editor/Sequencer/Private/MVVM/ViewModels/CategoryModel.cpp
--- a/Engine/Source/Editor/Sequencer/Private/MVVM/ViewModels/CategoryModel.cpp
+++ b/Engine/Source/Editor/Sequencer/Private/MVVM/ViewModels/CategoryModel.cpp
@@ -30,6 +30,7 @@
 #include "Widgets/SBoxPanel.h"
 #include "MVVM/Views/STrackLane.h"
 #include "MVVM/ViewModels/SequencerEditorViewModel.h"
+#include "MVVM/ViewModels/TrackModel.h"
 
 class SWidget;
 
@@ -294,6 +295,33 @@ void FCategoryGroupModel::OnRecycle()
 	Categories.Empty();
 }
 
+void FCategoryGroupModel::BuildSidebarMenu(FMenuBuilder& MenuBuilder)
+{
+	const TSharedPtr<FSequencerEditorViewModel> EditorViewModel = GetEditor();
+	if (!EditorViewModel.IsValid())
+	{
+		return;
+	}
+
+	const TSharedPtr<FSequencer> Sequencer = EditorViewModel->GetSequencerImpl();
+	if (!Sequencer.IsValid())
+	{
+		return;
+	}
+
+	if (const TViewModelPtr<FTrackModel> ParentTrack = FindAncestorOfType<FTrackModel>())
+	{
+		TArray<TWeakObjectPtr<>> WeakSectionObjects;
+		Algo::Transform(ParentTrack->GetSections(), WeakSectionObjects, [](UMovieSceneSection* const InSection)
+			{
+				return InSection;
+			});
+		SequencerHelpers::BuildEditSectionMenu(*Sequencer.Get(), WeakSectionObjects, MenuBuilder, false);
+	}
+
+	FOutlinerItemModel::BuildSidebarMenu(MenuBuilder);
+}
+	
 } // namespace UE::Sequencer
 
 #undef LOCTEXT_NAMESPACE
diff --git a/Engine/Source/Editor/Sequencer/Private/MVVM/ViewModels/ObjectBindingModel.cpp b/Engine/Source/Editor/Sequencer/Private/MVVM/ViewModels/ObjectBindingModel.cpp
--- a/Engine/Source/Editor/Sequencer/Private/MVVM/ViewModels/ObjectBindingModel.cpp
+++ b/Engine/Source/Editor/Sequencer/Private/MVVM/ViewModels/ObjectBindingModel.cpp
@@ -1034,19 +1034,40 @@ void FObjectBindingModel::BuildSidebarMenu(FMenuBuilder& MenuBuilder)
 		return;
 	}
 
-	UObject* const BoundObject = EditorViewModel->GetSequencerImpl()->FindSpawnedObjectOrTemplate(ObjectBindingID);
-	
+	FSequencer* const Sequencer = EditorViewModel->GetSequencerImpl().Get();
+	if (!Sequencer)
+	{
+		return;
+	} 
+
 	ISequencerModule& SequencerModule = FModuleManager::GetModuleChecked<ISequencerModule>(TEXT("Sequencer"));
-	
+
+	UObject* const BoundObject = Sequencer->FindSpawnedObjectOrTemplate(ObjectBindingID);
+
 	const TSharedPtr<FExtender> Extender = EditorViewModel->GetSequencerMenuExtender(SequencerModule.GetSidebarExtensibilityManager()
 		, TArrayBuilder<UObject*>().Add(BoundObject), &FSequencerCustomizationInfo::OnBuildSidebarMenu, SharedThis(this));
 	if (Extender.IsValid())
 	{
 		MenuBuilder.PushExtender(Extender.ToSharedRef());
 	}
-	
+
 	MenuBuilder.BeginSection(TEXT("ObjectBindingActions"), LOCTEXT("ObjectBindingsMenuSection", "Object Bindings"));
 	MenuBuilder.EndSection();
+
+	// External extension.
+	Sequencer->BuildCustomContextMenuForGuid(MenuBuilder, ObjectBindingID);
+
+	// Track editor extension.
+	TArray<FGuid> ObjectBindings;
+	ObjectBindings.Add(ObjectBindingID);
+
+	const UClass* const ObjectClass = FindObjectClass();
+	for (const TSharedPtr<ISequencerTrackEditor>& TrackEditor : Sequencer->GetTrackEditors())
+	{
+		TrackEditor->BuildObjectBindingContextMenu(MenuBuilder, ObjectBindings, ObjectClass);
+	}
+
+	FOutlinerItemModel::BuildSidebarMenu(MenuBuilder);
 }
 
 void FObjectBindingModel::AddTagMenu(FMenuBuilder& MenuBuilder)
diff --git a/Engine/Source/Editor/Sequencer/Private/MVVM/ViewModels/OutlinerItemModel.cpp b/Engine/Source/Editor/Sequencer/Private/MVVM/ViewModels/OutlinerItemModel.cpp
--- a/Engine/Source/Editor/Sequencer/Private/MVVM/ViewModels/OutlinerItemModel.cpp
+++ b/Engine/Source/Editor/Sequencer/Private/MVVM/ViewModels/OutlinerItemModel.cpp
@@ -568,38 +568,41 @@ void FOutlinerItemModelMixin::BuildContextMenu(FMenuBuilder& MenuBuilder)
 	}
 	MenuBuilder.EndSection();
 
-	TArray<UMovieSceneTrack*> AllTracks;
-	for (TViewModelPtr<ITrackExtension> TrackExtension : Sequencer->GetViewModel()->GetSelection()->Outliner.Filter<ITrackExtension>())
-	{
-		UMovieSceneTrack* Track = TrackExtension->GetTrack();
-		if (Track)
-		{
-			AllTracks.Add(Track);
-		}
-	}
-
 	MenuBuilder.BeginSection("Organize", LOCTEXT("OrganizeContextMenuSectionName", "Organize"));
 	BuildOrganizeContextMenu(MenuBuilder);
 	MenuBuilder.EndSection();
 
+	const TArray<UMovieSceneTrack*> AllTracks = GetSelectedTracks();
 	if (AllTracks.Num())
 	{
-		BuildTrackOptionsMenu(MenuBuilder);
+		BuildTrackOptionsMenu(MenuBuilder, AllTracks);
 		BuildDisplayOptionsMenu(MenuBuilder);
 	}
 }
 
 void FOutlinerItemModelMixin::BuildOrganizeContextMenu(FMenuBuilder& MenuBuilder)
 {
+	const TSharedPtr<FSequencerEditorViewModel> EditorViewModel = GetEditor();
+	if (!EditorViewModel.IsValid())
+	{
+		return;
+	}
+
+	const TSharedPtr<FSequencer> Sequencer = EditorViewModel->GetSequencerImpl();
+	if (!Sequencer.IsValid())
+	{
+		return;
+	}
+
+	FSequencer* const SequencerRaw = Sequencer.Get();
 	TSharedRef<FViewModel> ThisNode = AsViewModel()->AsShared();
-	TSharedPtr<FSequencer> Sequencer = GetEditor()->GetSequencerImpl();
 
 	const bool bFilterableNode = (ThisNode->IsA<ITrackExtension>() || ThisNode->IsA<IObjectBindingExtension>() || ThisNode->IsA<FFolderModel>());
 	const bool bIsReadOnly = Sequencer->IsReadOnly();
 	
 	TArray<UMovieSceneTrack*> AllTracks;
 	TArray<TSharedPtr<FViewModel> > DraggableNodes;
-	for (FViewModelPtr Node : Sequencer->GetViewModel()->GetSelection()->Outliner)
+	for (const FViewModelPtr Node : EditorViewModel->GetSelection()->Outliner)
 	{
 		if (ITrackExtension* TrackExtension = Node->CastThis<ITrackExtension>())
 		{
@@ -624,8 +627,7 @@ void FOutlinerItemModelMixin::BuildOrganizeContextMenu(FMenuBuilder& MenuBuilder
 		MenuBuilder.AddSubMenu(
 			LOCTEXT("AddNodesToNodeGroup", "Add to Group"),
 			LOCTEXT("AddNodesToNodeGroupTooltip", "Add selected nodes to a group"),
-			FNewMenuDelegate::CreateSP(Sequencer.Get(), &FSequencer::BuildAddSelectedToNodeGroupMenu));
-
+			FNewMenuDelegate::CreateSP(SequencerRaw, &FSequencer::BuildAddSelectedToNodeGroupMenu));
 	}
 
 	if (DraggableNodes.Num() && !bIsReadOnly)
@@ -633,40 +635,38 @@ void FOutlinerItemModelMixin::BuildOrganizeContextMenu(FMenuBuilder& MenuBuilder
 		MenuBuilder.AddSubMenu(
 			LOCTEXT("MoveToFolder", "Move to Folder"),
 			LOCTEXT("MoveToFolderTooltip", "Move the selected nodes to a folder"),
-			FNewMenuDelegate::CreateSP(Sequencer.Get(), &FSequencer::BuildAddSelectedToFolderMenu));
+			FNewMenuDelegate::CreateSP(SequencerRaw, &FSequencer::BuildAddSelectedToFolderMenu));
 
 		MenuBuilder.AddMenuEntry(
 			LOCTEXT("RemoveFromFolder", "Remove from Folder"),
 			LOCTEXT("RemoveFromFolderTooltip", "Remove selected nodes from their folders"),
 			FSlateIcon(),
 			FUIAction(
-				FExecuteAction::CreateSP(Sequencer.Get(), &FSequencer::RemoveSelectedNodesFromFolders),
-				FCanExecuteAction::CreateLambda( [Sequencer] { return Sequencer->GetSelectedNodesInFolders().Num() > 0; } )));
+				FExecuteAction::CreateSP(SequencerRaw, &FSequencer::RemoveSelectedNodesFromFolders),
+				FCanExecuteAction::CreateLambda([SequencerRaw] { return SequencerRaw->GetSelectedNodesInFolders().Num() > 0; } )));
 	}
 }
 
 void FOutlinerItemModelMixin::BuildDisplayOptionsMenu(FMenuBuilder& MenuBuilder)
 {
-	const TSharedPtr<FSequencer> Sequencer = GetEditor()->GetSequencerImpl();
+	const TSharedPtr<FSequencerEditorViewModel> EditorViewModel = GetEditor();
+	if (!EditorViewModel.IsValid())
+	{
+		return;
+	}
+
+	const TSharedPtr<ISequencer> Sequencer = EditorViewModel->GetSequencer();
 	if (!Sequencer.IsValid())
 	{
 		return;
 	}
-	
+
 	const TSharedRef<FOutlinerItemModelMixin> SharedThis(AsViewModel()->AsShared(), this);
 
 	const bool bIsReadOnly = Sequencer->IsReadOnly();
 	const FCanExecuteAction CanExecute = FCanExecuteAction::CreateLambda([bIsReadOnly]{ return !bIsReadOnly; });
 
-	TArray<UMovieSceneTrack*> AllTracks;
-	for (TViewModelPtr<ITrackExtension> TrackExtension : Sequencer->GetViewModel()->GetSelection()->Outliner.Filter<ITrackExtension>())
-	{
-		UMovieSceneTrack* const Track = TrackExtension->GetTrack();
-		if (IsValid(Track))
-		{
-			AllTracks.Add(Track);
-		}
-	}
+	TArray<UMovieSceneTrack*> AllTracks = GetSelectedTracks();
 	if (AllTracks.IsEmpty())
 	{
 		return;
@@ -677,7 +677,7 @@ void FOutlinerItemModelMixin::BuildDisplayOptionsMenu(FMenuBuilder& MenuBuilder)
 		MenuBuilder.AddSubMenu(
 			LOCTEXT("SetColorTint", "Set Color Tint"),
 			LOCTEXT("SetColorTintTooltip", "Set color tint from the preferences for the selected sections or the track's sections"),
-			FNewMenuDelegate::CreateSP(SharedThis, &FOutlinerItemModelMixin::BuildSectionColorTintsContextMenu));
+			FNewMenuDelegate::CreateSP(SharedThis, &FOutlinerItemModelMixin::BuildSectionColorTintsMenu));
 
 		UStruct* const DisplayOptionsStruct = FMovieSceneTrackDisplayOptions::StaticStruct();
 
@@ -690,10 +690,9 @@ void FOutlinerItemModelMixin::BuildDisplayOptionsMenu(FMenuBuilder& MenuBuilder)
 	MenuBuilder.EndSection();
 }
 
-void FOutlinerItemModelMixin::BuildTrackOptionsMenu(FMenuBuilder& MenuBuilder)
+void FOutlinerItemModelMixin::BuildTrackOptionsMenu(FMenuBuilder& MenuBuilder, const TArray<UMovieSceneTrack*>& InTracks)
 {
-	const TArray<UMovieSceneTrack*> AllTracks = GetSelectedTracks();
-	if (AllTracks.IsEmpty())
+	if (InTracks.IsEmpty())
 	{
 		return;
 	}
@@ -704,7 +703,7 @@ void FOutlinerItemModelMixin::BuildTrackOptionsMenu(FMenuBuilder& MenuBuilder)
 
 		const FBoolProperty* const NearestSectionProperty = CastField<FBoolProperty>(EvalOptionsStruct->FindPropertyByName(GET_MEMBER_NAME_CHECKED(FMovieSceneTrackEvalOptions, bEvalNearestSection)));
 		auto CanEvaluateNearest = [](const UMovieSceneTrack* const InTrack) { return InTrack->EvalOptions.bCanEvaluateNearestSection != 0; };
-		if (NearestSectionProperty && AllTracks.ContainsByPredicate(CanEvaluateNearest))
+		if (NearestSectionProperty && InTracks.ContainsByPredicate(CanEvaluateNearest))
 		{
 			TFunction<bool(UMovieSceneTrack*)> Validator = CanEvaluateNearest;
 			AddEvalOptionsPropertyMenuItem(MenuBuilder, NearestSectionProperty, Validator);
@@ -727,26 +726,31 @@ void FOutlinerItemModelMixin::BuildTrackOptionsMenu(FMenuBuilder& MenuBuilder)
 
 void FOutlinerItemModelMixin::BuildSidebarMenu(FMenuBuilder& MenuBuilder)
 {
-	/*BuildTrackOptionsMenu(MenuBuilder);
-	
-	const TArray<UMovieSceneSection*> Sections = GetSelectedSections();
-	if (!Sections.IsEmpty())
-	{
-		BuildDisplayOptionsMenu(MenuBuilder);
-	}*/
+	MenuBuilder.BeginSection(TEXT("Organize"), LOCTEXT("OrganizeContextMenuSectionName", "Organize"));
+	BuildOrganizeContextMenu(MenuBuilder);
+	MenuBuilder.EndSection();
+
+	BuildTrackOptionsMenu(MenuBuilder, GetSelectedTracks());
+	BuildDisplayOptionsMenu(MenuBuilder);
 }
 
 TArray<UMovieSceneSection*> FOutlinerItemModelMixin::GetSelectedSections() const
 {
 	TArray<UMovieSceneSection*> Sections;
 
-	const TSharedPtr<ISequencer> Sequencer = GetEditor()->GetSequencer();
+	const TSharedPtr<FSequencerEditorViewModel> EditorViewModel = GetEditor();
+	if (!EditorViewModel.IsValid())
+	{
+		return Sections;
+	}
+
+	const TSharedPtr<ISequencer> Sequencer = EditorViewModel->GetSequencer();
 	if (!Sequencer.IsValid())
 	{
 		return Sections;
 	}
 
-	const TSharedPtr<FSequencerSelection> Selection = Sequencer->GetViewModel()->GetSelection();
+	const TSharedPtr<FSequencerSelection> Selection = EditorViewModel->GetSelection();
 	if (!Selection.IsValid())
 	{
 		return Sections;
@@ -760,7 +764,7 @@ TArray<UMovieSceneSection*> FOutlinerItemModelMixin::GetSelectedSections() const
 		}
 	}
 
-	if (!Sections.Num())
+	if (Sections.IsEmpty())
 	{
 		for (const TViewModelPtr<ITrackExtension> TrackExtension : Selection->Outliner.Filter<ITrackExtension>())
 		{
@@ -784,55 +788,80 @@ TArray<UMovieSceneTrack*> FOutlinerItemModelMixin::GetSelectedTracks() const
 		return AllTracks;
 	}
 
-	const TSharedPtr<FSequencerSelection> Selection = EditorViewModel->GetSelection();
-	if (!Selection.IsValid())
+	const TSharedPtr<ISequencer> Sequencer = EditorViewModel->GetSequencer();
+	if (!Sequencer.IsValid())
 	{
 		return AllTracks;
 	}
 
-	for (const TViewModelPtr<ITrackExtension> TrackExtension : Selection->Outliner.Filter<ITrackExtension>())
+	const TSharedPtr<FSequencerSelection> Selection = EditorViewModel->GetSelection();
+	if (!Selection.IsValid())
 	{
-		UMovieSceneTrack* const Track = TrackExtension->GetTrack();
-		if (IsValid(Track))
-		{
-			AllTracks.Add(Track);
-		}
+		return AllTracks;
 	}
 
-	return AllTracks;
+	return Selection->GetSelectedTracks().Array();
 }
 
-void FOutlinerItemModelMixin::BuildSectionColorTintsContextMenu(FMenuBuilder& MenuBuilder)
+void FOutlinerItemModelMixin::BuildSectionColorTintsMenu(FMenuBuilder& MenuBuilder)
 {
-	TSharedPtr<FSequencer> Sequencer = GetEditor()->GetSequencerImpl();
-	TSharedPtr<FSequencerSelection> Selection = Sequencer->GetViewModel()->GetSelection();
+	const TSharedPtr<FSequencerEditorViewModel> EditorViewModel = GetEditor();
+	if (!EditorViewModel.IsValid())
+	{
+		return;
+	}
 
-	TArray<UMovieSceneSection*> Sections = GetSelectedSections();
-	if (!Sections.Num())
+	const TSharedPtr<FSequencer> Sequencer = EditorViewModel->GetSequencerImpl();
+	if (!Sequencer.IsValid())
 	{
 		return;
 	}
 
-	const bool bIsReadOnly = Sequencer->IsReadOnly();
-	FCanExecuteAction CanExecute = FCanExecuteAction::CreateLambda([bIsReadOnly] { return !bIsReadOnly; });
+	const TArray<UMovieSceneSection*> Sections = GetSelectedSections();
+	if (Sections.IsEmpty())
+	{
+		return;
+	}
+
+	const TWeakPtr<FSequencer> WeakSequencer = Sequencer;
+
+	FCanExecuteAction CanExecuteAction = FCanExecuteAction::CreateLambda([WeakSequencer]
+		{
+			return WeakSequencer.IsValid() ? !WeakSequencer.Pin()->IsReadOnly() : false;
+		});
 
-	TArray<FColor> SectionColorTints = Sequencer->GetSequencerSettings()->GetSectionColorTints();
+	const TArray<FColor> SectionColorTints = Sequencer->GetSequencerSettings()->GetSectionColorTints();
 
 	for (const FColor& SectionColorTint : SectionColorTints)
 	{
 		TSharedPtr<SBox> ColorWidget = SNew(SBox)
 			.WidthOverride(70.f)
 			.HeightOverride(20.f)
-		[
-			SNew(SBorder)
-				.BorderImage(FAppStyle::GetBrush("WhiteBrush"))
+			[
+				SNew(SBorder)
+				.BorderImage(FAppStyle::GetBrush(TEXT("WhiteBrush")))
 				.BorderBackgroundColor(FLinearColor::FromSRGBColor(SectionColorTint))
-		];
+			];
 
 		MenuBuilder.AddMenuEntry(
 			FUIAction(
-				FExecuteAction::CreateSP(Sequencer.Get(), &FSequencer::SetSectionColorTint, Sections, SectionColorTint),
-				CanExecute),
+				FExecuteAction::CreateLambda([this, WeakSequencer, SectionColorTint]
+					{
+						const TSharedPtr<FSequencer> Sequencer = WeakSequencer.Pin();
+						if (!Sequencer)
+						{
+							return;
+						}
+
+						const TArray<UMovieSceneSection*> Sections = GetSelectedSections();
+						if (Sections.IsEmpty())
+						{
+							return;
+						}
+
+						Sequencer->SetSectionColorTint(Sections, SectionColorTint);
+					}),
+				CanExecuteAction),
 			ColorWidget.ToSharedRef());
 	}
 
@@ -844,8 +873,23 @@ void FOutlinerItemModelMixin::BuildSectionColorTintsContextMenu(FMenuBuilder& Me
 		LOCTEXT("ClearColorTintTooltip", "Clear any assigned color tints"),
 		FSlateIcon(),
 		FUIAction(
-			FExecuteAction::CreateSP(Sequencer.Get(), &FSequencer::SetSectionColorTint, Sections, FColor(0, 0, 0, 0)),
-			CanExecute));
+			FExecuteAction::CreateLambda([this, WeakSequencer]
+				{
+					const TSharedPtr<FSequencer> Sequencer = WeakSequencer.Pin();
+					if (!Sequencer)
+					{
+						return;
+					}
+
+					const TArray<UMovieSceneSection*> Sections = GetSelectedSections();
+					if (Sections.IsEmpty())
+					{
+						return;
+					}
+
+					Sequencer->SetSectionColorTint(Sections, FColor(0, 0, 0, 0));
+				}),
+			CanExecuteAction));
 
 	// Pop up preferences to edit custom color tints
 	MenuBuilder.AddMenuEntry(
@@ -853,9 +897,22 @@ void FOutlinerItemModelMixin::BuildSectionColorTintsContextMenu(FMenuBuilder& Me
 		LOCTEXT("EditColorTintTooltip", "Edit the custom color tints"),
 		FSlateIcon(),
 		FUIAction(
-			FExecuteAction::CreateLambda([=] { 
-				FString SettingsName = Sequencer->GetSequencerSettings()->GetName();
-				return FModuleManager::LoadModuleChecked<ISettingsModule>("Settings").ShowViewer("Editor", "ContentEditors", *SettingsName); 
+			FExecuteAction::CreateLambda([WeakSequencer]
+			{
+				const TSharedPtr<FSequencer> Sequencer = WeakSequencer.Pin();
+				if (!Sequencer)
+				{
+					return;
+				}
+
+				const USequencerSettings* SequencerSettings = Sequencer->GetSequencerSettings();
+				if (!IsValid(SequencerSettings))
+				{
+					return;
+				}
+
+				ISettingsModule& SettingsModule = FModuleManager::LoadModuleChecked<ISettingsModule>(TEXT("Settings"));
+				SettingsModule.ShowViewer("Editor", "ContentEditors", *SequencerSettings->GetName()); 
 			})
 		));
 }
diff --git a/Engine/Source/Editor/Sequencer/Private/MVVM/ViewModels/TrackModel.cpp b/Engine/Source/Editor/Sequencer/Private/MVVM/ViewModels/TrackModel.cpp
--- a/Engine/Source/Editor/Sequencer/Private/MVVM/ViewModels/TrackModel.cpp
+++ b/Engine/Source/Editor/Sequencer/Private/MVVM/ViewModels/TrackModel.cpp
@@ -651,72 +651,42 @@ bool FTrackModel::CanDrag() const
 
 void FTrackModel::BuildContextMenu(FMenuBuilder& MenuBuilder)
 {
-	UMovieSceneTrack* Track = GetTrack();
-	if (!Track)
+	const TSharedPtr<FSequencerEditorViewModel> EditorViewModel = GetEditor();
+	if (!EditorViewModel.IsValid())
 	{
 		return;
 	}
 
-	TWeakPtr<ISequencer> WeakSequencer = GetEditor()->GetSequencer();
-
-	const int32 TrackRowIndex = GetRowIndex();
-
-	if (TrackEditor)
+	const TSharedPtr<FSequencer> Sequencer = EditorViewModel->GetSequencerImpl();
+	if (!Sequencer.IsValid())
 	{
-		TrackEditor->BuildTrackContextMenu(MenuBuilder, Track);
+		return;
 	}
 
-	if (Track && Track->GetSupportedBlendTypes().Num() > 0)
+	UMovieSceneTrack* const Track = GetTrack();
+	if (!IsValid(Track))
 	{
-		MenuBuilder.AddSubMenu(
-			LOCTEXT("AddSection", "Add Section"),
-			FText(),
-			FNewMenuDelegate::CreateLambda([=](FMenuBuilder& SubMenuBuilder){
-				FSequencerUtilities::PopulateMenu_CreateNewSection(SubMenuBuilder, TrackRowIndex + 1, Track, WeakSequencer);
-			})
-		);	
+		return;
 	}
 
-	// Add menu items for selecting a blender
-	BuildBlendingMenu(MenuBuilder);
-
-	// Find sections in the track to add batch properties for
-	TArray<TWeakObjectPtr<UObject>> TrackSections;
+	FSequencer& SequencerRef = *Sequencer.Get();
 
-	for (TViewModelPtr<ITrackExtension> TrackExtension : WeakSequencer.Pin()->GetViewModel()->GetSelection()->Outliner.Filter<ITrackExtension>())
+	if (TrackEditor)
 	{
-		for (UMovieSceneSection* Section : TrackExtension->GetSections())
-		{
-			TrackSections.Add(Section);
-		}
+		TrackEditor->BuildTrackContextMenu(MenuBuilder, Track);
 	}
 
-	for (TSharedPtr<FViewModel> TrackAreaModel : GetTrackAreaModelList())
+	if (Track->GetSupportedBlendTypes().Num() > 0)
 	{
-		constexpr bool bIncludeThis = true;
-		for (TSharedPtr<FSectionModel> Section : TParentFirstChildIterator<FSectionModel>(TrackAreaModel, bIncludeThis))
-		{
-			if (UMovieSceneSection* SectionObject = Section->GetSection())
-			{
-				TrackSections.AddUnique(SectionObject);
-			}
-		}
-	}
-		
-	if (TrackSections.Num())
-	{
-		MenuBuilder.AddSubMenu(
-			TrackSections.Num() > 1 ? LOCTEXT("BatchEditSections", "Batch Edit Sections") : LOCTEXT("EditSection", "Edit Section"),
-			FText(),
-			FNewMenuDelegate::CreateLambda([this, TrackSections](FMenuBuilder& SubMenuBuilder){
-				FSequencer* Sequencer = static_cast<FSequencer*>(GetEditor()->GetSequencer().Get());
-				SequencerHelpers::AddPropertiesMenu(*Sequencer, SubMenuBuilder, TrackSections);
-			})
-		);
+		SequencerHelpers::BuildNewSectionMenu(SequencerRef, GetRowIndex() + 1, GetTrack(), MenuBuilder);
 	}
 
-	TViewModelPtr<FChannelGroupModel> ChannelGroup = TopLevelChannelList.GetHead().ImplicitCast();
-	if (ChannelGroup)
+	SequencerHelpers::BuildBlendingMenu(SequencerRef, Track, MenuBuilder);
+
+	const TArray<TWeakObjectPtr<>> TrackAreaModels = SequencerHelpers::GetSectionObjectsFromTrackAreaModels(GetTrackAreaModelList());
+	SequencerHelpers::BuildEditSectionMenu(SequencerRef, TrackAreaModels, MenuBuilder, true);
+
+	if (const TViewModelPtr<FChannelGroupModel> ChannelGroup = TopLevelChannelList.GetHead().ImplicitCast())
 	{
 		ChannelGroup->BuildChannelOverrideMenu(MenuBuilder);
 	}
@@ -726,79 +696,47 @@ void FTrackModel::BuildContextMenu(FMenuBuilder& MenuBuilder)
 
 void FTrackModel::BuildSidebarMenu(FMenuBuilder& MenuBuilder)
 {
-	UMovieSceneTrack* const Track = GetTrack();
-	if (!IsValid(Track))
-	{
-		return;
-	}
-
 	const TSharedPtr<FSequencerEditorViewModel> EditorViewModel = GetEditor();
 	if (!EditorViewModel.IsValid())
 	{
 		return;
 	}
 
-	const TSharedPtr<ISequencer> Sequencer = EditorViewModel->GetSequencer();
+	const TSharedPtr<FSequencer> Sequencer = EditorViewModel->GetSequencerImpl();
 	if (!Sequencer.IsValid())
 	{
 		return;
 	}
 
-	const IMovieSceneBlenderSystemSupport* const BlenderSystemSupport = Cast<IMovieSceneBlenderSystemSupport>(Track);
-	if (BlenderSystemSupport)
-	{
-		FSequencerUtilities::PopulateMenu_BlenderSubMenu(MenuBuilder, Track, Sequencer);
-	}
-
-	// Find sections in the track to add batch properties for
-	const TArray<TWeakObjectPtr<UObject>> TrackSections = GetSelectedTrackSections();
-	if (TrackSections.Num())
-	{
-		MenuBuilder.BeginSection(TEXT("TrackSection"), LOCTEXT("TrackSectionMenuSection", "Track Section"));
-		
-		FSequencer* const SequencerRaw = static_cast<FSequencer*>(Sequencer.Get());
-		SequencerHelpers::AddPropertiesMenu(*SequencerRaw, MenuBuilder, TrackSections);
-
-		MenuBuilder.EndSection();
-	}
-
-	FOutlinerItemModelMixin::BuildSidebarMenu(MenuBuilder);
-}
-
-void FTrackModel::BuildBlendingMenu(FMenuBuilder& MenuBuilder)
-{
 	UMovieSceneTrack* const Track = GetTrack();
 	if (!IsValid(Track))
 	{
 		return;
 	}
 
-	const IMovieSceneBlenderSystemSupport* const BlenderSystemSupport = Cast<IMovieSceneBlenderSystemSupport>(Track);
-	if (!BlenderSystemSupport)
+	FSequencer& SequencerRef = *Sequencer.Get();
+
+	if (TrackEditor)
 	{
-		return;
+		TrackEditor->BuildTrackSidebarMenu(MenuBuilder, Track);
 	}
 
-	TArray<TSubclassOf<UMovieSceneBlenderSystem>> BlenderTypes;
-	BlenderSystemSupport->GetSupportedBlenderSystems(BlenderTypes);
-	if (BlenderTypes.Num() == 0)
+	if (Track->GetSupportedBlendTypes().Num() > 0)
 	{
-		return;
+		SequencerHelpers::BuildNewSectionMenu(SequencerRef, GetRowIndex() + 1, GetTrack(), MenuBuilder);
 	}
-	
-	const TSharedPtr<ISequencer> Sequencer = GetEditor()->GetSequencer();
-	if (!Sequencer)
+
+	SequencerHelpers::BuildBlendingMenu(SequencerRef, Track, MenuBuilder);
+
+	const TArray<TWeakObjectPtr<>> TrackAreaModels = SequencerHelpers::GetSectionObjectsFromTrackAreaModels(GetTrackAreaModelList());
+	SequencerHelpers::BuildEditSectionMenu(SequencerRef, TrackAreaModels, MenuBuilder, false);
+
+	if (const TViewModelPtr<FChannelGroupModel> ChannelGroup = TopLevelChannelList.GetHead().ImplicitCast())
 	{
-		return;
+		ChannelGroup->BuildChannelOverrideMenu(MenuBuilder);
 	}
-	
-	MenuBuilder.AddSubMenu(
-		LOCTEXT("BlendingAlgorithmSubMenu", "Blending Algorithm"),
-		FText(),
-		FNewMenuDelegate::CreateLambda([=](FMenuBuilder& SubMenuBuilder)
-			{
-				FSequencerUtilities::PopulateMenu_BlenderSubMenu(SubMenuBuilder, Track, Sequencer);
-			}));
+
+	FOutlinerItemModel::BuildSidebarMenu(MenuBuilder);
 }
 
 bool FTrackModel::CanDelete(FText* OutErrorMessage) const
@@ -874,57 +812,6 @@ bool FTrackModel::FindBoundObjects(TArray<UObject*>& OutBoundObjects) const
 	return true;
 }
 
-TArray<TWeakObjectPtr<UObject>> FTrackModel::GetSelectedTrackSections() const
-{
-	TArray<TWeakObjectPtr<UObject>> TrackSections;
-
-	const TSharedPtr<FSequencerEditorViewModel> EditorViewModel = GetEditor();
-	if (!EditorViewModel.IsValid())
-	{
-		return TrackSections;
-	}
-
-	const TSharedPtr<ISequencer> Sequencer = EditorViewModel->GetSequencer();
-	if (!Sequencer.IsValid())
-	{
-		return TrackSections;
-	}
-
-	const TSharedPtr<UE::Sequencer::FSequencerEditorViewModel> ViewModel = Sequencer->GetViewModel();
-	if (!ViewModel.IsValid())
-	{
-		return TrackSections;
-	}
-
-	const TSharedPtr<FSequencerSelection> Selection = ViewModel->GetSelection();
-	if (!Selection.IsValid())
-	{
-		return TrackSections;
-	}
-	
-	for (const TViewModelPtr<ITrackExtension> TrackExtension : Selection->Outliner.Filter<ITrackExtension>())
-	{
-		for (UMovieSceneSection* const Section : TrackExtension->GetSections())
-		{
-			TrackSections.Add(Section);
-		}
-	}
-
-	for (const TSharedPtr<FViewModel> TrackAreaModel : GetTrackAreaModelList())
-	{
-		constexpr bool bIncludeThis = true;
-		for (const TSharedPtr<FSectionModel> Section : TParentFirstChildIterator<FSectionModel>(TrackAreaModel, bIncludeThis))
-		{
-			if (UMovieSceneSection* const SectionObject = Section->GetSection())
-			{
-				TrackSections.AddUnique(SectionObject);
-			}
-		}
-	}
-
-	return TrackSections;
-}
-
 } // namespace UE::Sequencer
 
 #undef LOCTEXT_NAMESPACE
diff --git a/Engine/Source/Editor/Sequencer/Private/MVVM/ViewModels/TrackRowModel.cpp b/Engine/Source/Editor/Sequencer/Private/MVVM/ViewModels/TrackRowModel.cpp
--- a/Engine/Source/Editor/Sequencer/Private/MVVM/ViewModels/TrackRowModel.cpp
+++ b/Engine/Source/Editor/Sequencer/Private/MVVM/ViewModels/TrackRowModel.cpp
@@ -305,85 +305,42 @@ void FTrackRowModel::Resize(float NewSize)
 
 void FTrackRowModel::BuildContextMenu(FMenuBuilder& MenuBuilder)
 {
-	UMovieSceneTrack* Track = GetTrack();
-	if (!Track)
+	const TSharedPtr<FSequencerEditorViewModel> EditorViewModel = GetEditor();
+	if (!EditorViewModel.IsValid())
 	{
 		return;
 	}
 
-	TSharedPtr<ISequencer> Sequencer = GetEditor()->GetSequencer();
-
-	const int32 TrackRowIndex = GetRowIndex();
-
-	TSharedPtr<ISequencerTrackEditor> TrackEditorPtr = GetTrackEditor();
-	TrackEditorPtr->BuildTrackContextMenu(MenuBuilder, Track);
-
-	if (Track && Track->GetSupportedBlendTypes().Num() > 0)
+	const TSharedPtr<FSequencer> Sequencer = EditorViewModel->GetSequencerImpl();
+	if (!Sequencer.IsValid())
 	{
-		MenuBuilder.AddSubMenu(
-			LOCTEXT("AddSection", "Add Section"),
-			FText(),
-			FNewMenuDelegate::CreateLambda([=](FMenuBuilder& SubMenuBuilder){
-				FSequencerUtilities::PopulateMenu_CreateNewSection(SubMenuBuilder, TrackRowIndex + 1, Track, Sequencer);
-			})
-		);	
+		return;
 	}
 
-	// Add menu items for selecting a blender
-	IMovieSceneBlenderSystemSupport* BlenderSystemSupport = Cast<IMovieSceneBlenderSystemSupport>(Track);
-	if (BlenderSystemSupport)
+	UMovieSceneTrack* const Track = GetTrack();
+	if (!IsValid(Track))
 	{
-		TArray<TSubclassOf<UMovieSceneBlenderSystem>> BlenderTypes;
-		BlenderSystemSupport->GetSupportedBlenderSystems(BlenderTypes);
-
-		if (BlenderTypes.Num() > 1)
-		{
-			MenuBuilder.AddSubMenu(
-				LOCTEXT("BlendingAlgorithmSubMenu", "Blending Algorithm"),
-				FText(),
-				FNewMenuDelegate::CreateLambda([=](FMenuBuilder& SubMenuBuilder){
-					FSequencerUtilities::PopulateMenu_BlenderSubMenu(SubMenuBuilder, Track, Sequencer);
-				})
-			);
-		}
+		return;
 	}
 
-	// Find sections in the track to add batch properties for
-	TArray<TWeakObjectPtr<UObject>> TrackSections;
-	
-	for (TViewModelPtr<ITrackExtension> TrackExtension : Sequencer->GetViewModel()->GetSelection()->Outliner.Filter<ITrackExtension>())
-	{
-		for (UMovieSceneSection* Section : TrackExtension->GetSections())
-		{
-			TrackSections.Add(Section);
-		}
-	}
-	
-	for (TSharedPtr<FViewModel> TrackAreaModel : GetTrackAreaModelList())
+	FSequencer& SequencerRef = *Sequencer.Get();
+
+	if (TrackEditor)
 	{
-		constexpr bool bIncludeThis = true;
-		for (TSharedPtr<FSectionModel> Section : TParentFirstChildIterator<FSectionModel>(TrackAreaModel, bIncludeThis))
-		{
-			if (UMovieSceneSection* SectionObject = Section->GetSection())
-			{
-				TrackSections.AddUnique(SectionObject);
-			}
-		}
+		TrackEditor->BuildTrackContextMenu(MenuBuilder, Track);
 	}
-		
-	if (TrackSections.Num())
+
+	if (Track->GetSupportedBlendTypes().Num() > 0)
 	{
-		MenuBuilder.AddSubMenu(
-			TrackSections.Num() > 1 ? LOCTEXT("BatchEditSections", "Batch Edit Sections") : LOCTEXT("EditSection", "Edit Section"),
-			FText(),
-			FNewMenuDelegate::CreateLambda([=](FMenuBuilder& SubMenuBuilder){
-				SequencerHelpers::AddPropertiesMenu(static_cast<FSequencer&>(*Sequencer), SubMenuBuilder, TrackSections);
-			})
-		);
+		SequencerHelpers::BuildNewSectionMenu(SequencerRef, GetRowIndex() + 1, GetTrack(), MenuBuilder);
 	}
 
-	TViewModelPtr<FChannelGroupModel> ChannelGroup = TopLevelChannelList.GetHead().ImplicitCast();
-	if (ChannelGroup)
+	SequencerHelpers::BuildBlendingMenu(SequencerRef, Track, MenuBuilder);
+
+	const TArray<TWeakObjectPtr<>> TrackAreaModels = SequencerHelpers::GetSectionObjectsFromTrackAreaModels(GetTrackAreaModelList());
+	SequencerHelpers::BuildEditSectionMenu(SequencerRef, TrackAreaModels, MenuBuilder, true);
+
+	if (const TViewModelPtr<FChannelGroupModel> ChannelGroup = TopLevelChannelList.GetHead().ImplicitCast())
 	{
 		ChannelGroup->BuildChannelOverrideMenu(MenuBuilder);
 	}
@@ -393,40 +350,44 @@ void FTrackRowModel::BuildContextMenu(FMenuBuilder& MenuBuilder)
 
 void FTrackRowModel::BuildSidebarMenu(FMenuBuilder& MenuBuilder)
 {
-	UMovieSceneTrack* const Track = GetTrack();
-	if (!IsValid(Track))
+	const TSharedPtr<FSequencerEditorViewModel> EditorViewModel = GetEditor();
+	if (!EditorViewModel.IsValid())
 	{
 		return;
 	}
 
-	const TSharedPtr<FSequencerEditorViewModel> EditorViewModel = GetEditor();
-	if (!EditorViewModel.IsValid())
+	const TSharedPtr<FSequencer> Sequencer = EditorViewModel->GetSequencerImpl();
+	if (!Sequencer.IsValid())
 	{
 		return;
 	}
 
-	const TSharedPtr<ISequencer> Sequencer = EditorViewModel->GetSequencer();
-	if (!Sequencer.IsValid())
+	UMovieSceneTrack* const Track = GetTrack();
+	if (!IsValid(Track))
 	{
 		return;
 	}
 
-	IMovieSceneBlenderSystemSupport* const BlenderSystemSupport = Cast<IMovieSceneBlenderSystemSupport>(Track);
-	if (BlenderSystemSupport)
+	FSequencer& SequencerRef = *Sequencer.Get();
+
+	if (TrackEditor)
 	{
-		FSequencerUtilities::PopulateMenu_BlenderSubMenu(MenuBuilder, Track, Sequencer);
+		TrackEditor->BuildTrackSidebarMenu(MenuBuilder, Track);
 	}
 
-	// Find sections in the track to add batch properties for
-	const TArray<TWeakObjectPtr<UObject>> TrackSections = GetSelectedTrackSections();
-	if (TrackSections.Num())
+	if (Track->GetSupportedBlendTypes().Num() > 0)
 	{
-		MenuBuilder.BeginSection(TEXT("TrackSection"), LOCTEXT("TrackSectionMenuSection", "Track Section"));
-		
-		FSequencer* const SequencerRaw = static_cast<FSequencer*>(GetEditor()->GetSequencer().Get());
-		SequencerHelpers::AddPropertiesMenu(*SequencerRaw, MenuBuilder, TrackSections);
+		SequencerHelpers::BuildNewSectionMenu(SequencerRef, GetRowIndex() + 1, GetTrack(), MenuBuilder);
+	}
+
+	SequencerHelpers::BuildBlendingMenu(SequencerRef, Track, MenuBuilder);
+
+	const TArray<TWeakObjectPtr<>> TrackAreaModels = SequencerHelpers::GetSectionObjectsFromTrackAreaModels(GetTrackAreaModelList());
+	SequencerHelpers::BuildEditSectionMenu(SequencerRef, TrackAreaModels, MenuBuilder, false);
 
-		MenuBuilder.EndSection();
+	if (const TViewModelPtr<FChannelGroupModel> ChannelGroup = TopLevelChannelList.GetHead().ImplicitCast())
+	{
+		ChannelGroup->BuildChannelOverrideMenu(MenuBuilder);
 	}
 
 	FOutlinerItemModel::BuildSidebarMenu(MenuBuilder);
@@ -469,57 +430,6 @@ void FTrackRowModel::Delete()
 	Track->FixRowIndices();
 }
 
-TArray<TWeakObjectPtr<UObject>> FTrackRowModel::GetSelectedTrackSections() const
-{
-	TArray<TWeakObjectPtr<UObject>> TrackSections;
-
-	const TSharedPtr<FSequencerEditorViewModel> EditorViewModel = GetEditor();
-	if (!EditorViewModel.IsValid())
-	{
-		return TrackSections;
-	}
-
-	const TSharedPtr<ISequencer> Sequencer = EditorViewModel->GetSequencer();
-	if (!Sequencer.IsValid())
-	{
-		return TrackSections;
-	}
-
-	const TSharedPtr<UE::Sequencer::FSequencerEditorViewModel> ViewModel = Sequencer->GetViewModel();
-	if (!ViewModel.IsValid())
-	{
-		return TrackSections;
-	}
-
-	const TSharedPtr<FSequencerSelection> Selection = ViewModel->GetSelection();
-	if (!Selection.IsValid())
-	{
-		return TrackSections;
-	}
-
-	for (const TViewModelPtr<ITrackExtension> TrackExtension : Selection->Outliner.Filter<ITrackExtension>())
-	{
-		for (UMovieSceneSection* const Section : TrackExtension->GetSections())
-		{
-			TrackSections.Add(Section);
-		}
-	}
-
-	for (const TSharedPtr<FViewModel> TrackAreaModel : GetTrackAreaModelList())
-	{
-		constexpr bool bIncludeThis = true;
-		for (const TSharedPtr<FSectionModel> Section : TParentFirstChildIterator<FSectionModel>(TrackAreaModel, bIncludeThis))
-		{
-			if (UMovieSceneSection* const SectionObject = Section->GetSection())
-			{
-				TrackSections.AddUnique(SectionObject);
-			}
-		}
-	}
-
-	return TrackSections;
-}
-
 } // namespace UE::Sequencer
 
 #undef LOCTEXT_NAMESPACE
diff --git a/Engine/Source/Editor/Sequencer/Private/MovieSceneTrackEditor.cpp b/Engine/Source/Editor/Sequencer/Private/MovieSceneTrackEditor.cpp
--- a/Engine/Source/Editor/Sequencer/Private/MovieSceneTrackEditor.cpp
+++ b/Engine/Source/Editor/Sequencer/Private/MovieSceneTrackEditor.cpp
@@ -328,15 +328,16 @@ TSharedPtr<SWidget> FMovieSceneTrackEditor::BuildOutlinerEditWidget(const FGuid&
 }
 
 void FMovieSceneTrackEditor::BuildTrackContextMenu(FMenuBuilder& MenuBuilder, UMovieSceneTrack* Track) 
-{ 
+{
 }
 
 void FMovieSceneTrackEditor::BuildTrackSidebarMenu(FMenuBuilder& MenuBuilder, UMovieSceneTrack* Track) 
-{ 
+{
+	BuildTrackContextMenu(MenuBuilder, Track);
 }
 
 bool FMovieSceneTrackEditor::HandleAssetAdded(UObject* Asset, const FGuid& TargetObjectGuid) 
-{ 
+{
 	return false; 
 }
 
diff --git a/Engine/Source/Editor/Sequencer/Private/SSequencer.cpp b/Engine/Source/Editor/Sequencer/Private/SSequencer.cpp
--- a/Engine/Source/Editor/Sequencer/Private/SSequencer.cpp
+++ b/Engine/Source/Editor/Sequencer/Private/SSequencer.cpp
@@ -4040,7 +4040,7 @@ bool SSequencer::OpenPasteMenu()
 	if (Sequencer->GetClipboardStack().Num() != 0)
 	{
 		FPasteContextMenuArgs Args = GeneratePasteArgs(Sequencer->GetLocalTime().Time.FrameNumber, Sequencer->GetClipboardStack().Last());
-		ContextMenu = FPasteContextMenu::CreateMenu(*Sequencer, Args);
+		ContextMenu = FPasteContextMenu::CreateMenu(SequencerPtr, Args);
 	}
 
 	if (!ContextMenu.IsValid() || !ContextMenu->IsValidPaste())
@@ -4081,7 +4081,7 @@ void SSequencer::PasteFromHistory()
 	}
 
 	FPasteContextMenuArgs Args = GeneratePasteArgs(Sequencer->GetLocalTime().Time.FrameNumber);
-	TSharedPtr<FPasteFromHistoryContextMenu> ContextMenu = FPasteFromHistoryContextMenu::CreateMenu(*Sequencer, Args);
+	TSharedPtr<FPasteFromHistoryContextMenu> ContextMenu = FPasteFromHistoryContextMenu::CreateMenu(SequencerPtr, Args);
 
 	if (ContextMenu.IsValid())
 	{
diff --git a/Engine/Source/Editor/Sequencer/Private/SequencerCommonHelpers.cpp b/Engine/Source/Editor/Sequencer/Private/SequencerCommonHelpers.cpp
--- a/Engine/Source/Editor/Sequencer/Private/SequencerCommonHelpers.cpp
+++ b/Engine/Source/Editor/Sequencer/Private/SequencerCommonHelpers.cpp
@@ -2,6 +2,8 @@
 
 #include "SequencerCommonHelpers.h"
 
+#include "EntitySystem/IMovieSceneBlenderSystemSupport.h"
+#include "EntitySystem/MovieSceneBlenderSystem.h"
 #include "FrameNumberDetailsCustomization.h"
 #include "IDetailsView.h"
 #include "ISequencerSection.h"
@@ -28,6 +30,7 @@
 #include "Sequencer.h"
 #include "SequencerContextMenus.h"
 #include "SequencerSelectedKey.h"
+#include "SequencerUtilities.h"
 #include "Styling/CoreStyle.h"
 
 void SequencerHelpers::GetAllChannels(TSharedPtr<FViewModel> DataModel, TSet<TSharedPtr<UE::Sequencer::FChannelModel>>& Channels)
@@ -263,7 +266,8 @@ TSharedPtr<SWidget> SequencerHelpers::SummonContextMenu(FSequencer& Sequencer, c
 	}
 	else if (Sequencer.GetClipboardStack().Num() != 0)
 	{
-		TSharedPtr<FPasteContextMenu> PasteMenu = FPasteContextMenu::CreateMenu(Sequencer, SequencerWidget->GeneratePasteArgs(PasteAtTime));
+		const TWeakPtr<FSequencer> WeakSequencer = StaticCastWeakPtr<FSequencer>(Sequencer.AsWeak());
+		TSharedPtr<FPasteContextMenu> PasteMenu = FPasteContextMenu::CreateMenu(WeakSequencer, SequencerWidget->GeneratePasteArgs(PasteAtTime));
 		if (PasteMenu.IsValid() && PasteMenu->IsValidPaste())
 		{
 			PasteMenu->PopulateMenu(MenuBuilder, MenuExtender);
@@ -285,92 +289,209 @@ public:
 
 	void Construct(FArguments InArgs) { }
 
-	void SetDetailsAndSequencer(TSharedRef<SWidget> InDetailsPanel, TWeakPtr<ISequencer> InSequencerWeak)
+	void SetDetailsAndSequencer(TSharedRef<SWidget> InDetailsPanel, TWeakPtr<ISequencer> InWeakSequencer)
 	{
 		ChildSlot
 		[
 			InDetailsPanel
 		];
-		SequencerWeak = InSequencerWeak;
+		WeakSequencer = InWeakSequencer;
 	}
 
 	//~ FNotifyHook interface
 	virtual void NotifyPostChange(const FPropertyChangedEvent& PropertyChangedEvent, FProperty* PropertyThatChanged) override
 	{
-		if (const TSharedPtr<ISequencer> Sequencer = SequencerWeak.Pin())
+		if (const TSharedPtr<ISequencer> Sequencer = WeakSequencer.Pin())
 		{
 			Sequencer->NotifyMovieSceneDataChanged(EMovieSceneDataChangeType::TrackValueChanged);
 		}
 	}
 
 private:
-	TWeakPtr<ISequencer> SequencerWeak;
+	TWeakPtr<ISequencer> WeakSequencer;
 };
 
+void SequencerHelpers::BuildNewSectionMenu(FSequencer& Sequencer
+	, const int32 InRowIndex
+	, const TWeakObjectPtr<UMovieSceneTrack>& InTrackWeak
+	, FMenuBuilder& MenuBuilder)
+{
+	using namespace UE::Sequencer;
+
+	const TWeakPtr<ISequencer> WeakSequencer = Sequencer.AsWeak();
+
+	MenuBuilder.AddSubMenu(
+		NSLOCTEXT("Sequencer", "AddSection", "Add Section"),
+		FText(),
+		FNewMenuDelegate::CreateLambda([WeakSequencer, InRowIndex, InTrackWeak](FMenuBuilder& SubMenuBuilder)
+		{
+			if (const TSharedPtr<ISequencer> Sequencer = WeakSequencer.Pin())
+			{
+				FSequencerUtilities::PopulateMenu_CreateNewSection(SubMenuBuilder, InRowIndex, InTrackWeak.Get(), Sequencer);
+			}
+		}));
+}
 
-void SequencerHelpers::AddPropertiesMenu(FSequencer& Sequencer, FMenuBuilder& MenuBuilder, const TArray<TWeakObjectPtr<UObject>>& Sections)
+void SequencerHelpers::BuildEditSectionMenu(FSequencer& Sequencer
+	, const TArray<TWeakObjectPtr<>>& InWeakSections
+	, FMenuBuilder& MenuBuilder
+	, const bool bInSubMenu)
 {
 	using namespace UE::Sequencer;
 
-	TSharedRef<SSectionDetailsNotifyHookWrapper> DetailsNotifyWrapper = SNew(SSectionDetailsNotifyHookWrapper);
-	FDetailsViewArgs DetailsViewArgs;
+	if (InWeakSections.Num() == 0)
 	{
-		DetailsViewArgs.bAllowSearch = false;
-		DetailsViewArgs.bCustomFilterAreaLocation = true;
-		DetailsViewArgs.bCustomNameAreaLocation = true;
-		DetailsViewArgs.bHideSelectionTip = true;
-		DetailsViewArgs.bLockable = false;
-		DetailsViewArgs.bSearchInitialKeyFocus = true;
-		DetailsViewArgs.bUpdatesFromSelection = false;
-		DetailsViewArgs.bShowOptions = false;
-		DetailsViewArgs.bShowModifiedPropertiesOption = false;
-		DetailsViewArgs.NotifyHook = &DetailsNotifyWrapper.Get();
-		DetailsViewArgs.ColumnWidth = 0.45f;
+		return;
 	}
 
-	// We pass the current scene to the UMovieSceneSection customization so we can get the overall bounds of the section when we change a section from infinite->bounded.
-	UMovieScene* CurrentScene = Sequencer.GetFocusedMovieSceneSequence()->GetMovieScene();
-
-	TSharedRef<INumericTypeInterface<double>> NumericTypeInterface = Sequencer.GetNumericTypeInterface();
+	auto BuildSection = [&Sequencer, InWeakSections](FMenuBuilder& LambdaMenuBuilder)
+	{
+		TSharedRef<SSectionDetailsNotifyHookWrapper> DetailsNotifyWrapper = SNew(SSectionDetailsNotifyHookWrapper);
+		FDetailsViewArgs DetailsViewArgs;
+		{
+			DetailsViewArgs.bAllowSearch = false;
+			DetailsViewArgs.bCustomFilterAreaLocation = true;
+			DetailsViewArgs.bCustomNameAreaLocation = true;
+			DetailsViewArgs.bHideSelectionTip = true;
+			DetailsViewArgs.bLockable = false;
+			DetailsViewArgs.bSearchInitialKeyFocus = true;
+			DetailsViewArgs.bUpdatesFromSelection = false;
+			DetailsViewArgs.bShowOptions = false;
+			DetailsViewArgs.bShowModifiedPropertiesOption = false;
+			DetailsViewArgs.bShowScrollBar = false;
+			DetailsViewArgs.NotifyHook = &DetailsNotifyWrapper.Get();
+			DetailsViewArgs.ColumnWidth = 0.45f;
+		}
 
-	TSharedRef<IDetailsView> DetailsView = FModuleManager::GetModuleChecked<FPropertyEditorModule>("PropertyEditor").CreateDetailView(DetailsViewArgs);
-	DetailsView->RegisterInstancedCustomPropertyTypeLayout("FrameNumber", FOnGetPropertyTypeCustomizationInstance::CreateLambda([=]() {
-		return MakeShared<FFrameNumberDetailsCustomization>(NumericTypeInterface); }));
-	DetailsView->RegisterInstancedCustomPropertyLayout(UMovieSceneSection::StaticClass(), FOnGetDetailCustomizationInstance::CreateLambda([=]() {
-		return MakeShared<FMovieSceneSectionDetailsCustomization>(NumericTypeInterface, CurrentScene); }));
+		TSharedRef<IDetailsView> DetailsView = FModuleManager::GetModuleChecked<FPropertyEditorModule>("PropertyEditor").CreateDetailView(DetailsViewArgs);
+		DetailsView->RegisterInstancedCustomPropertyTypeLayout("FrameNumber",
+			FOnGetPropertyTypeCustomizationInstance::CreateLambda([&Sequencer]()
+			{
+				return MakeShared<FFrameNumberDetailsCustomization>(Sequencer.GetNumericTypeInterface());
+			}));
+		DetailsView->RegisterInstancedCustomPropertyLayout(UMovieSceneSection::StaticClass(),
+			FOnGetDetailCustomizationInstance::CreateLambda([&Sequencer]()
+			{
+				// We pass the current scene to the UMovieSceneSection customization so we can get the overall bounds of the section when we change a section from infinite->bounded.
+				UMovieScene* const CurrentScene = Sequencer.GetFocusedMovieSceneSequence()->GetMovieScene();
+				check(CurrentScene);
+				return MakeShared<FMovieSceneSectionDetailsCustomization>(Sequencer.GetNumericTypeInterface(), CurrentScene);
+			}));
 	
-	DetailsView->SetIsPropertyVisibleDelegate(FIsPropertyVisible::CreateLambda([](const FPropertyAndParent& PropertyAndParent)
-		{			
-			return FPropertyEditorPermissionList::Get().DoesPropertyPassFilter(PropertyAndParent.Property.GetOwnerStruct(), PropertyAndParent.Property.GetFName());
-		})
-	);
-
-	// Let section interfaces further customize the properties details view.
-	TSharedRef<FSequencerNodeTree> SequencerNodeTree = Sequencer.GetNodeTree();
-	for (TWeakObjectPtr<UObject> Section : Sections)
-	{
-		if (Section.IsValid())
+		DetailsView->SetIsPropertyVisibleDelegate(FIsPropertyVisible::CreateLambda([](const FPropertyAndParent& PropertyAndParent)
+			{			
+				return FPropertyEditorPermissionList::Get().DoesPropertyPassFilter(PropertyAndParent.Property.GetOwnerStruct(), PropertyAndParent.Property.GetFName());
+			})
+		);
+
+		// Let section interfaces further customize the properties details view.
+		TSharedRef<FSequencerNodeTree> SequencerNodeTree = Sequencer.GetNodeTree();
+		for (TWeakObjectPtr<> Section : InWeakSections)
 		{
-			TSharedPtr<FSectionModel> SectionHandle = SequencerNodeTree->GetSectionModel(Cast<UMovieSceneSection>(Section));
-			if (SectionHandle)
+			if (Section.IsValid())
 			{
-				TSharedPtr<ISequencerSection> SectionInterface = SectionHandle->GetSectionInterface();
-				FSequencerSectionPropertyDetailsViewCustomizationParams CustomizationDetails(
-					SectionInterface.ToSharedRef(), Sequencer.AsWeak(), *SectionHandle->GetParentTrackExtension()->GetTrackEditor().Get());
-				TSharedPtr<FObjectBindingModel> ParentObjectBindingNode = SectionHandle->FindAncestorOfType<FObjectBindingModel>();
-				if (ParentObjectBindingNode.IsValid())
+				TSharedPtr<FSectionModel> SectionHandle = SequencerNodeTree->GetSectionModel(Cast<UMovieSceneSection>(Section));
+				if (SectionHandle)
 				{
-					CustomizationDetails.ParentObjectBindingGuid = ParentObjectBindingNode->GetObjectGuid();
+					TSharedPtr<ISequencerSection> SectionInterface = SectionHandle->GetSectionInterface();
+					FSequencerSectionPropertyDetailsViewCustomizationParams CustomizationDetails(
+						SectionInterface.ToSharedRef(), Sequencer.AsWeak(), *SectionHandle->GetParentTrackExtension()->GetTrackEditor().Get());
+					TSharedPtr<FObjectBindingModel> ParentObjectBindingNode = SectionHandle->FindAncestorOfType<FObjectBindingModel>();
+					if (ParentObjectBindingNode.IsValid())
+					{
+						CustomizationDetails.ParentObjectBindingGuid = ParentObjectBindingNode->GetObjectGuid();
+					}
+					SectionInterface->CustomizePropertiesDetailsView(DetailsView, CustomizationDetails);
 				}
-				SectionInterface->CustomizePropertiesDetailsView(DetailsView, CustomizationDetails);
 			}
 		}
+
+		Sequencer.OnInitializeDetailsPanel().Broadcast(DetailsView, Sequencer.AsShared());
+		DetailsView->SetObjects(InWeakSections);
+
+		DetailsNotifyWrapper->SetDetailsAndSequencer(DetailsView, Sequencer.AsWeak());
+		DetailsNotifyWrapper->SetEnabled(!Sequencer.IsReadOnly());
+
+		LambdaMenuBuilder.BeginSection(TEXT("TrackSection"), NSLOCTEXT("Sequencer", "TrackSectionMenuSection", "Track Section"));
+        {
+			LambdaMenuBuilder.AddWidget(DetailsNotifyWrapper, FText::GetEmpty(), true);
+        }
+        LambdaMenuBuilder.EndSection();
+	};
+
+	if (bInSubMenu)
+	{
+		const FText MenuLabel = InWeakSections.Num() > 1
+			? NSLOCTEXT("Sequencer", "BatchEditSections", "Batch Edit Sections")
+			: NSLOCTEXT("Sequencer", "EditSection", "Edit Section");
+
+		MenuBuilder.AddSubMenu(
+			MenuLabel,
+			FText(),
+			FNewMenuDelegate::CreateLambda([BuildSection](FMenuBuilder& SubMenuBuilder)
+			{
+				BuildSection(SubMenuBuilder);
+			}));
+	}
+	else
+	{
+		BuildSection(MenuBuilder);
+	}
+}
+
+void SequencerHelpers::BuildBlendingMenu(FSequencer& Sequencer
+	, const TWeakObjectPtr<UMovieSceneTrack>& InTrackWeak
+	, FMenuBuilder& MenuBuilder)
+{
+	if (!InTrackWeak.IsValid())
+	{
+		return;
+	}
+
+	IMovieSceneBlenderSystemSupport* const BlenderSystemSupport = Cast<IMovieSceneBlenderSystemSupport>(InTrackWeak.Get());
+	if (!BlenderSystemSupport)
+	{
+		return;
 	}
 
-	Sequencer.OnInitializeDetailsPanel().Broadcast(DetailsView, Sequencer.AsShared());
-	DetailsView->SetObjects(Sections);
+	TArray<TSubclassOf<UMovieSceneBlenderSystem>> BlenderTypes;
+	BlenderSystemSupport->GetSupportedBlenderSystems(BlenderTypes);
+	if (BlenderTypes.Num() < 2)
+	{
+		return;
+	}
+
+	const TWeakPtr<ISequencer> WeakSequencer = Sequencer.AsWeak();
+
+	MenuBuilder.AddSubMenu(
+		NSLOCTEXT("Sequencer", "BlendingAlgorithmSubMenu", "Blending Algorithm"),
+		FText(),
+		FNewMenuDelegate::CreateLambda([WeakSequencer, InTrackWeak](FMenuBuilder& SubMenuBuilder)
+		{
+			if (const TSharedPtr<ISequencer> Sequencer = WeakSequencer.Pin())
+			{
+				FSequencerUtilities::PopulateMenu_BlenderSubMenu(SubMenuBuilder, InTrackWeak.Get(), Sequencer);
+			}
+		}));
+}
+
+TArray<TWeakObjectPtr<>> SequencerHelpers::GetSectionObjectsFromTrackAreaModels(const UE::Sequencer::FViewModelVariantIterator& InTrackAreaModels)
+{
+	using namespace UE::Sequencer;
+
+	TArray<TWeakObjectPtr<>> OutWeakSectionObjects;
+
+	for (const TViewModelPtr<FViewModel>& TrackAreaModel : InTrackAreaModels)
+	{
+		constexpr bool bIncludeThis = true;
+		for (const TSharedPtr<FSectionModel> SectionModel : TParentFirstChildIterator<FSectionModel>(TrackAreaModel, bIncludeThis))
+		{
+			if (UMovieSceneSection* const SectionObject = SectionModel->GetSection())
+			{
+				OutWeakSectionObjects.AddUnique(SectionObject);
+			}
+		}
+	}
 
-	DetailsNotifyWrapper->SetDetailsAndSequencer(DetailsView, Sequencer.AsWeak());
-	DetailsNotifyWrapper->SetEnabled(!Sequencer.IsReadOnly());
-	MenuBuilder.AddWidget(DetailsNotifyWrapper, FText::GetEmpty(), true);
+	return OutWeakSectionObjects;
 }
diff --git a/Engine/Source/Editor/Sequencer/Private/SequencerCommonHelpers.h b/Engine/Source/Editor/Sequencer/Private/SequencerCommonHelpers.h
--- a/Engine/Source/Editor/Sequencer/Private/SequencerCommonHelpers.h
+++ b/Engine/Source/Editor/Sequencer/Private/SequencerCommonHelpers.h
@@ -14,13 +14,14 @@
 class FSequencer;
 class IKeyArea;
 class UMovieSceneSection;
+class UMovieSceneTrack;
 
 namespace UE
 {
 namespace Sequencer
 {
-
-class FChannelModel;
+	class FChannelModel;
+	struct FViewModelVariantIterator;
 
 } // namespace Sequencer
 } // namespace UE
@@ -68,9 +69,32 @@ public:
 	static TSharedPtr<SWidget> SummonContextMenu(FSequencer& Sequencer, const FGeometry& MyGeometry, const FPointerEvent& MouseEvent);
 
 	/*
-	 * Build a context menu for the sections
+	 * Build a sub menu for adding a new track section
+	 */
+	static void BuildNewSectionMenu(FSequencer& Sequencer
+		, const int32 InRowIndex
+		, const TWeakObjectPtr<UMovieSceneTrack>& InTrackWeak
+		, FMenuBuilder& MenuBuilder);
+
+	/*
+	 * Build an inline menu or sub menu for editing track section(s)
+	 */
+	static void BuildEditSectionMenu(FSequencer& Sequencer
+		, const TArray<TWeakObjectPtr<>>& InWeakSections
+		, FMenuBuilder& MenuBuilder
+		, const bool bInSubMenu);
+
+	/*
+	 * Build a menu for selection the blend algorithm
+	 */
+	static void BuildBlendingMenu(FSequencer& Sequencer
+		, const TWeakObjectPtr<UMovieSceneTrack>& InTrackWeak
+		, FMenuBuilder& MenuBuilder);
+
+	/**
+	 * Gets all section objects from track area models
 	 */
-	static void AddPropertiesMenu(FSequencer& Sequencer, FMenuBuilder& MenuBuilder, const TArray<TWeakObjectPtr<UObject>>& Sections);
+	static TArray<TWeakObjectPtr<>> GetSectionObjectsFromTrackAreaModels(const UE::Sequencer::FViewModelVariantIterator& InTrackAreaModels);
 };
 
 #undef LOCTEXT_NAMESPACE
diff --git a/Engine/Source/Editor/Sequencer/Private/SequencerContextMenus.cpp b/Engine/Source/Editor/Sequencer/Private/SequencerContextMenus.cpp
--- a/Engine/Source/Editor/Sequencer/Private/SequencerContextMenus.cpp
+++ b/Engine/Source/Editor/Sequencer/Private/SequencerContextMenus.cpp
@@ -54,11 +54,17 @@
 
 #define LOCTEXT_NAMESPACE "SequencerContextMenus"
 
-static void CreateKeyStructForSelection(TSharedPtr<ISequencer> InSequencer, TSharedPtr<FStructOnScope>& OutKeyStruct, TWeakObjectPtr<UMovieSceneSection>& OutKeyStructSection)
+static void CreateKeyStructForSelection(TWeakPtr<ISequencer> InWeakSequencer, TSharedPtr<FStructOnScope>& OutKeyStruct, TWeakObjectPtr<UMovieSceneSection>& OutKeyStructSection)
 {
 	using namespace UE::Sequencer;
 
-	const FKeySelection& SelectedKeys = InSequencer->GetViewModel()->GetSelection()->KeySelection;
+	const TSharedPtr<ISequencer> Sequencer = InWeakSequencer.Pin();
+	if (!Sequencer.IsValid())
+	{
+		return;
+	}
+
+	const FKeySelection& SelectedKeys = Sequencer->GetViewModel()->GetSelection()->KeySelection;
 
 	if (SelectedKeys.Num() == 1)
 	{
@@ -104,9 +110,9 @@ static void CreateKeyStructForSelection(TSharedPtr<ISequencer> InSequencer, TSha
 	}
 }
 
-void FKeyContextMenu::BuildMenu(FMenuBuilder& MenuBuilder, TSharedPtr<FExtender> MenuExtender, FSequencer& InSequencer)
+void FKeyContextMenu::BuildMenu(FMenuBuilder& MenuBuilder, TSharedPtr<FExtender> MenuExtender, TWeakPtr<FSequencer> InWeakSequencer)
 {
-	TSharedRef<FKeyContextMenu> Menu = MakeShareable(new FKeyContextMenu(InSequencer));
+	TSharedRef<FKeyContextMenu> Menu = MakeShareable(new FKeyContextMenu(InWeakSequencer));
 	Menu->PopulateMenu(MenuBuilder, MenuExtender);
 }
 
@@ -114,15 +120,20 @@ void FKeyContextMenu::PopulateMenu(FMenuBuilder& MenuBuilder, TSharedPtr<FExtend
 {
 	using namespace UE::Sequencer;
 
-	FSequencer* SequencerPtr = &Sequencer.Get();
+	const TSharedPtr<FSequencer> Sequencer = WeakSequencer.Pin();
+	if (!Sequencer.IsValid())
+	{
+		return;
+	}
+
 	TSharedRef<FKeyContextMenu> Shared = AsShared();
 
-	CreateKeyStructForSelection(Sequencer, KeyStruct, KeyStructSection);
+	CreateKeyStructForSelection(WeakSequencer, KeyStruct, KeyStructSection);
 
 	{
 		ISequencerModule& SequencerModule = FModuleManager::LoadModuleChecked<ISequencerModule>("Sequencer");
 
-		FSelectedKeysByChannel SelectedKeysByChannel(SequencerPtr->GetViewModel()->GetSelection()->KeySelection);
+		FSelectedKeysByChannel SelectedKeysByChannel(Sequencer->GetViewModel()->GetSelection()->KeySelection);
 
 		TMap<FName, TArray<FExtendKeyMenuParams>> ChannelAndHandlesByType;
 		for (FSelectedChannelInfo& ChannelInfo : SelectedKeysByChannel.SelectedChannels)
@@ -145,7 +156,7 @@ void FKeyContextMenu::PopulateMenu(FMenuBuilder& MenuBuilder, TSharedPtr<FExtend
 		}
 	}
 
-	if(KeyStruct.IsValid())
+	if (KeyStruct.IsValid())
 	{
 		MenuBuilder.AddSubMenu(
 			LOCTEXT("KeyProperties", "Properties"),
@@ -163,8 +174,7 @@ void FKeyContextMenu::PopulateMenu(FMenuBuilder& MenuBuilder, TSharedPtr<FExtend
 
 	MenuBuilder.BeginSection("SequencerKeyEdit", LOCTEXT("EditMenu", "Edit"));
 	{
-		TSharedPtr<FSequencerEditorViewModel> SequencerViewModel = SequencerPtr->GetViewModel()->CastThisShared<FSequencerEditorViewModel>();
-		if (HotspotCast<FKeyHotspot>(SequencerViewModel->GetHotspot()))
+		if (HotspotCast<FKeyHotspot>(Sequencer->GetViewModel()->GetHotspot()))
 		{
 			MenuBuilder.AddMenuEntry(FGenericCommands::Get().Cut);
 			MenuBuilder.AddMenuEntry(FGenericCommands::Get().Copy);
@@ -185,24 +195,36 @@ void FKeyContextMenu::PopulateMenu(FMenuBuilder& MenuBuilder, TSharedPtr<FExtend
 
 void FKeyContextMenu::AddPropertiesMenu(FMenuBuilder& MenuBuilder)
 {
+	const TSharedPtr<FSequencer> Sequencer = WeakSequencer.Pin();
+	if (!Sequencer.IsValid())
+	{
+		return;
+	}
+
 	auto UpdateAndRetrieveEditData = [this]
 	{
 		FKeyEditData EditData;
-		CreateKeyStructForSelection(Sequencer, EditData.KeyStruct, EditData.OwningSection);
+		CreateKeyStructForSelection(WeakSequencer, EditData.KeyStruct, EditData.OwningSection);
 		return EditData;
 	};
 
 	MenuBuilder.AddWidget(
-		SNew(SKeyEditInterface, Sequencer)
+		SNew(SKeyEditInterface, Sequencer.ToSharedRef())
 		.EditData_Lambda(UpdateAndRetrieveEditData)
-	, FText::GetEmpty(), true);
+		, FText::GetEmpty(), true);
 }
 
 
-FSectionContextMenu::FSectionContextMenu(FSequencer& InSeqeuncer, FFrameTime InMouseDownTime)
-	: Sequencer(StaticCastSharedRef<FSequencer>(InSeqeuncer.AsShared()))
+FSectionContextMenu::FSectionContextMenu(TWeakPtr<FSequencer> InWeakSequencer, FFrameTime InMouseDownTime)
+	: WeakSequencer(InWeakSequencer)
 	, MouseDownTime(InMouseDownTime)
 {
+	const TSharedPtr<FSequencer> Sequencer = WeakSequencer.Pin();
+	if (!Sequencer.IsValid())
+	{
+		return;
+	}
+
 	for (UMovieSceneSection* Section : Sequencer->GetViewModel()->GetSelection()->GetSelectedSections())
 	{
 		FMovieSceneChannelProxy& ChannelProxy = Section->GetChannelProxy();
@@ -224,17 +246,30 @@ FSectionContextMenu::FSectionContextMenu(FSequencer& InSeqeuncer, FFrameTime InM
 	}
 }
 
-void FSectionContextMenu::BuildMenu(FMenuBuilder& MenuBuilder, TSharedPtr<FExtender> MenuExtender, FSequencer& InSequencer, FFrameTime InMouseDownTime)
+void FSectionContextMenu::BuildMenu(FMenuBuilder& MenuBuilder, TSharedPtr<FExtender> MenuExtender, TWeakPtr<FSequencer> InWeakSequencer, FFrameTime InMouseDownTime)
 {
-	TSharedRef<FSectionContextMenu> Menu = MakeShareable(new FSectionContextMenu(InSequencer, InMouseDownTime));
+	TSharedRef<FSectionContextMenu> Menu = MakeShareable(new FSectionContextMenu(InWeakSequencer, InMouseDownTime));
 	Menu->PopulateMenu(MenuBuilder, MenuExtender);
 }
 
+void FSectionContextMenu::BuildKeyEditMenu(FMenuBuilder& MenuBuilder, TWeakPtr<FSequencer> InWeakSequencer, FFrameTime InMouseDownTime)
+{
+	TSharedRef<FSectionContextMenu> Menu = MakeShareable(new FSectionContextMenu(InWeakSequencer, InMouseDownTime));
+	Menu->AddKeyInterpolationMenu(MenuBuilder);
+	Menu->AddKeyEditMenu(MenuBuilder);
+}
+
 
 void FSectionContextMenu::PopulateMenu(FMenuBuilder& MenuBuilder, TSharedPtr<FExtender> MenuExtender)
 {
 	using namespace UE::Sequencer;
 
+	const TSharedPtr<FSequencer> Sequencer = WeakSequencer.Pin();
+	if (!Sequencer.IsValid())
+	{
+		return;
+	}
+
 	// Copy a reference to the context menu by value into each lambda handler to ensure the type stays alive until the menu is closed
 	TSharedRef<FSectionContextMenu> Shared = AsShared();
 
@@ -254,7 +289,7 @@ void FSectionContextMenu::PopulateMenu(FMenuBuilder& MenuBuilder, TSharedPtr<FEx
 		ISequencerChannelInterface* ChannelInterface = SequencerModule.FindChannelEditorInterface(Pair.Key);
 		if (ChannelInterface)
 		{
-			ChannelInterface->ExtendSectionMenu_Raw(MenuBuilder, MenuExtender, Pair.Value, Sections, Sequencer);
+			ChannelInterface->ExtendSectionMenu_Raw(MenuBuilder, MenuExtender, Pair.Value, Sections, WeakSequencer);
 		}
 	}
 
@@ -263,6 +298,12 @@ void FSectionContextMenu::PopulateMenu(FMenuBuilder& MenuBuilder, TSharedPtr<FEx
 		LOCTEXT("SectionPropertiesTooltip", "Modify the section properties"),
 		FNewMenuDelegate::CreateLambda([this](FMenuBuilder& SubMenuBuilder)
 		{
+			const TSharedPtr<FSequencer> Sequencer = WeakSequencer.Pin();
+			if (!Sequencer.IsValid())
+			{
+				return;
+			}
+
 			TArray<TWeakObjectPtr<UObject>> Sections;
 			for (TViewModelPtr<FSectionModel> SectionModel : Sequencer->GetViewModel()->GetSelection()->TrackArea.Filter<FSectionModel>())
 			{
@@ -272,7 +313,7 @@ void FSectionContextMenu::PopulateMenu(FMenuBuilder& MenuBuilder, TSharedPtr<FEx
 				}
 			}
 
-			SequencerHelpers::AddPropertiesMenu(*Sequencer, SubMenuBuilder, Sections);
+			SequencerHelpers::BuildEditSectionMenu(*Sequencer.Get(), Sections, SubMenuBuilder, /*bInSubMenu=*/true);
 		})
 	);
 
@@ -284,8 +325,8 @@ void FSectionContextMenu::PopulateMenu(FMenuBuilder& MenuBuilder, TSharedPtr<FEx
 		if (Sequencer->GetClipboardStack().Num() != 0)
 		{
 			FPasteContextMenuArgs PasteArgs = FPasteContextMenuArgs::PasteAt(MouseDownTime.FrameNumber);
-			PasteMenu = FPasteContextMenu::CreateMenu(*Sequencer, PasteArgs);
-			PasteFromHistoryMenu = FPasteFromHistoryContextMenu::CreateMenu(*Sequencer, PasteArgs);
+			PasteMenu = FPasteContextMenu::CreateMenu(Sequencer, PasteArgs);
+			PasteFromHistoryMenu = FPasteFromHistoryContextMenu::CreateMenu(Sequencer, PasteArgs);
 		}
 
 		MenuBuilder.AddSubMenu(
@@ -385,7 +426,7 @@ void FSectionContextMenu::PopulateMenu(FMenuBuilder& MenuBuilder, TSharedPtr<FEx
 			LOCTEXT("GroupSectionsTooltip", "Group selected sections together so that when any section is moved, all sections in that group move together."),
 			FSlateIcon(),
 			FUIAction(
-				FExecuteAction::CreateSP(Sequencer, &FSequencer::GroupSelectedSections),
+				FExecuteAction::CreateSP(Sequencer.ToSharedRef(), &FSequencer::GroupSelectedSections),
 				FCanExecuteAction::CreateLambda([bCanGroup] { return bCanGroup; })
 			),
 			NAME_None,
@@ -397,7 +438,7 @@ void FSectionContextMenu::PopulateMenu(FMenuBuilder& MenuBuilder, TSharedPtr<FEx
 			LOCTEXT("UngroupSectionsTooltip", "Ungroup selected sections"),
 			FSlateIcon(),
 			FUIAction(
-				FExecuteAction::CreateSP(Sequencer, &FSequencer::UngroupSelectedSections),
+				FExecuteAction::CreateSP(Sequencer.ToSharedRef(), &FSequencer::UngroupSelectedSections),
 				FCanExecuteAction::CreateLambda([bCanUngroup] { return bCanUngroup; })
 			),
 			NAME_None,
@@ -452,9 +493,25 @@ void FSectionContextMenu::AddEditMenu(FMenuBuilder& MenuBuilder)
 		LOCTEXT("DeleteKeysWhenTrimmingTooltip", "Delete keys outside of the trimmed range"),
 		FSlateIcon(),
 		FUIAction(
-			FExecuteAction::CreateLambda([this] { Sequencer->GetSequencerSettings()->SetDeleteKeysWhenTrimming(!Sequencer->GetSequencerSettings()->GetDeleteKeysWhenTrimming()); }),
+			FExecuteAction::CreateLambda([this]
+			{
+				const TSharedPtr<FSequencer> Sequencer = WeakSequencer.Pin();
+				if (!Sequencer.IsValid())
+				{
+					return;
+				}
+				Sequencer->GetSequencerSettings()->SetDeleteKeysWhenTrimming(!Sequencer->GetSequencerSettings()->GetDeleteKeysWhenTrimming());
+			}),
 			FCanExecuteAction(),
-			FIsActionChecked::CreateLambda([this] { return Sequencer->GetSequencerSettings()->GetDeleteKeysWhenTrimming(); })),
+			FIsActionChecked::CreateLambda([this]
+			{
+				const TSharedPtr<FSequencer> Sequencer = WeakSequencer.Pin();
+				if (!Sequencer.IsValid())
+				{
+					return false;
+				}
+				return Sequencer->GetSequencerSettings()->GetDeleteKeysWhenTrimming();
+			})),
 		NAME_None,
 		EUserInterfaceActionType::ToggleButton
 	);
@@ -472,12 +529,21 @@ void FSectionContextMenu::AddEditMenu(FMenuBuilder& MenuBuilder)
 			FCanExecuteAction::CreateLambda([=]{ return Shared->CanAutoSize(); }))
 	);
 
-	MenuBuilder.BeginSection("SequencerInterpolation", LOCTEXT("KeyInterpolationMenu", "Key Interpolation"));
-	
+	AddKeyInterpolationMenu(MenuBuilder);
+	AddKeyEditMenu(MenuBuilder);
+}
+
+void FSectionContextMenu::AddKeyInterpolationMenu(FMenuBuilder& MenuBuilder)
+{
+	// Copy a reference to the context menu by value into each lambda handler to ensure the type stays alive until the menu is closed
+	TSharedRef<FSectionContextMenu> Shared = AsShared();
+
+	MenuBuilder.BeginSection(TEXT("SequencerInterpolation"), LOCTEXT("KeyInterpolationMenu", "Key Interpolation"));
+
 	MenuBuilder.AddMenuEntry(
 		LOCTEXT("SetKeyInterpolationSmartAuto", "Cubic (Smart Auto)"),
 		LOCTEXT("SetKeyInterpolationSmartAutoTooltip", "Set key interpolation to smart auto"),
-		FSlateIcon(FAppStyle::GetAppStyleSetName(), "Sequencer.IconKeySmartAuto"),
+		FSlateIcon(FAppStyle::GetAppStyleSetName(), TEXT("Sequencer.IconKeySmartAuto")),
 		FUIAction(
 			FExecuteAction::CreateLambda([=] { Shared->SetInterpTangentMode(RCIM_Cubic, RCTM_SmartAuto); }),
 			FCanExecuteAction::CreateLambda([=] { return Shared->CanSetInterpTangentMode(); }))
@@ -486,7 +552,7 @@ void FSectionContextMenu::AddEditMenu(FMenuBuilder& MenuBuilder)
 	MenuBuilder.AddMenuEntry(
 		LOCTEXT("SetKeyInterpolationAuto", "Cubic (Auto)"),
 		LOCTEXT("SetKeyInterpolationAutoTooltip", "Set key interpolation to auto"),
-		FSlateIcon(FAppStyle::GetAppStyleSetName(), "Sequencer.IconKeyAuto"),
+		FSlateIcon(FAppStyle::GetAppStyleSetName(), TEXT("Sequencer.IconKeyAuto")),
 		FUIAction(
 			FExecuteAction::CreateLambda([=]{ Shared->SetInterpTangentMode(RCIM_Cubic, RCTM_Auto); }),
 			FCanExecuteAction::CreateLambda([=]{ return Shared->CanSetInterpTangentMode(); }) )
@@ -495,7 +561,7 @@ void FSectionContextMenu::AddEditMenu(FMenuBuilder& MenuBuilder)
 	MenuBuilder.AddMenuEntry(
 		LOCTEXT("SetKeyInterpolationUser", "Cubic (User)"),
 		LOCTEXT("SetKeyInterpolationUserTooltip", "Set key interpolation to user"),
-		FSlateIcon(FAppStyle::GetAppStyleSetName(), "Sequencer.IconKeyUser"),
+		FSlateIcon(FAppStyle::GetAppStyleSetName(), TEXT("Sequencer.IconKeyUser")),
 		FUIAction(
 			FExecuteAction::CreateLambda([=]{ Shared->SetInterpTangentMode(RCIM_Cubic, RCTM_User); }),
 			FCanExecuteAction::CreateLambda([=]{ return Shared->CanSetInterpTangentMode(); }) )
@@ -504,7 +570,7 @@ void FSectionContextMenu::AddEditMenu(FMenuBuilder& MenuBuilder)
 	MenuBuilder.AddMenuEntry(
 		LOCTEXT("SetKeyInterpolationBreak", "Cubic (Break)"),
 		LOCTEXT("SetKeyInterpolationBreakTooltip", "Set key interpolation to break"),
-		FSlateIcon(FAppStyle::GetAppStyleSetName(), "Sequencer.IconKeyBreak"),
+		FSlateIcon(FAppStyle::GetAppStyleSetName(), TEXT("Sequencer.IconKeyBreak")),
 		FUIAction(
 			FExecuteAction::CreateLambda([=]{ Shared->SetInterpTangentMode(RCIM_Cubic, RCTM_Break); }),
 			FCanExecuteAction::CreateLambda([=]{ return Shared->CanSetInterpTangentMode(); }) )
@@ -513,7 +579,7 @@ void FSectionContextMenu::AddEditMenu(FMenuBuilder& MenuBuilder)
 	MenuBuilder.AddMenuEntry(
 		LOCTEXT("SetKeyInterpolationLinear", "Linear"),
 		LOCTEXT("SetKeyInterpolationLinearTooltip", "Set key interpolation to linear"),
-		FSlateIcon(FAppStyle::GetAppStyleSetName(), "Sequencer.IconKeyLinear"),
+		FSlateIcon(FAppStyle::GetAppStyleSetName(), TEXT("Sequencer.IconKeyLinear")),
 		FUIAction(
 			FExecuteAction::CreateLambda([=]{ Shared->SetInterpTangentMode(RCIM_Linear, RCTM_Auto); }),
 			FCanExecuteAction::CreateLambda([=]{ return Shared->CanSetInterpTangentMode(); }) )
@@ -522,15 +588,21 @@ void FSectionContextMenu::AddEditMenu(FMenuBuilder& MenuBuilder)
 	MenuBuilder.AddMenuEntry(
 		LOCTEXT("SetKeyInterpolationConstant", "Constant"),
 		LOCTEXT("SetKeyInterpolationConstantTooltip", "Set key interpolation to constant"),
-		FSlateIcon(FAppStyle::GetAppStyleSetName(), "Sequencer.IconKeyConstant"),
+		FSlateIcon(FAppStyle::GetAppStyleSetName(), TEXT("Sequencer.IconKeyConstant")),
 		FUIAction(
 			FExecuteAction::CreateLambda([=]{ Shared->SetInterpTangentMode(RCIM_Constant, RCTM_Auto); }),
 			FCanExecuteAction::CreateLambda([=]{ return Shared->CanSetInterpTangentMode(); }) )
 	);
 
-	MenuBuilder.EndSection(); // SequencerInterpolation
+	MenuBuilder.EndSection();
+}
 
-	MenuBuilder.BeginSection("Key Editing", LOCTEXT("KeyEditingSectionMenus", "Key Editing"));
+void FSectionContextMenu::AddKeyEditMenu(FMenuBuilder& MenuBuilder)
+{
+	// Copy a reference to the context menu by value into each lambda handler to ensure the type stays alive until the menu is closed
+	TSharedRef<FSectionContextMenu> Shared = AsShared();
+
+	MenuBuilder.BeginSection(TEXT("Key Editing"), LOCTEXT("KeyEditingSectionMenus", "Key Editing"));
 
 	MenuBuilder.AddMenuEntry(
 		LOCTEXT("ReduceKeysSection", "Reduce Keys"),
@@ -541,30 +613,45 @@ void FSectionContextMenu::AddEditMenu(FMenuBuilder& MenuBuilder)
 			FCanExecuteAction::CreateLambda([=]{ return Shared->CanReduceKeys(); }))
 	);
 
-	auto OnReduceKeysToleranceChanged = [this](float NewValue) {
-		Sequencer->GetSequencerSettings()->SetReduceKeysTolerance(NewValue);
+	auto OnReduceKeysToleranceChanged = [this](float InNewValue)
+	{
+		const TSharedPtr<FSequencer> Sequencer = WeakSequencer.Pin();
+		if (!Sequencer.IsValid())
+		{
+			return;
+		}
+		Sequencer->GetSequencerSettings()->SetReduceKeysTolerance(InNewValue);
 	};
 
 	MenuBuilder.AddWidget(
 		SNew(SHorizontalBox)
 		+ SHorizontalBox::Slot()
-			[
-				SNew(SSpacer)
-			]
+		[
+			SNew(SSpacer)
+		]
 		+ SHorizontalBox::Slot()
-			.AutoWidth()
-			[
-				SNew(SSpinBox<float>)
-				.Style(&FAppStyle::GetWidgetStyle<FSpinBoxStyle>("Sequencer.HyperlinkSpinBox"))
-				.OnValueCommitted_Lambda([=](float Value, ETextCommit::Type) { OnReduceKeysToleranceChanged(Value); })
-				.OnValueChanged_Lambda(OnReduceKeysToleranceChanged)
-				.MinValue(0)
-				.MaxValue(TOptional<float>())
-				.Value_Lambda([this]() -> float {
+		.AutoWidth()
+		[
+			SNew(SSpinBox<float>)
+			.Style(&FAppStyle::GetWidgetStyle<FSpinBoxStyle>(TEXT("Sequencer.HyperlinkSpinBox")))
+			.OnValueCommitted_Lambda([&OnReduceKeysToleranceChanged](float Value, ETextCommit::Type)
+			{
+				OnReduceKeysToleranceChanged(Value);
+			})
+			.OnValueChanged_Lambda(OnReduceKeysToleranceChanged)
+			.MinValue(0)
+			.MaxValue(TOptional<float>())
+			.Value_Lambda([this]() -> float
+			{
+				const TSharedPtr<FSequencer> Sequencer = WeakSequencer.Pin();
+				if (!Sequencer.IsValid())
+				{
+					return 0;
+				}
 				return Sequencer->GetSequencerSettings()->GetReduceKeysTolerance();
-				})
-			],
-		LOCTEXT("ReduceKeysTolerance", "Tolerance"));
+			})
+		],
+		LOCTEXT("ReduceKeysTolerance", "Tolerance"), true);
 
 	MenuBuilder.EndSection();
 }
@@ -573,10 +660,13 @@ FMovieSceneBlendTypeField FSectionContextMenu::GetSupportedBlendTypes() const
 {
 	FMovieSceneBlendTypeField BlendTypes = FMovieSceneBlendTypeField::All();
 
-	for (UMovieSceneSection* Section : Sequencer->GetViewModel()->GetSelection()->GetSelectedSections())
+	if (const TSharedPtr<FSequencer> Sequencer = WeakSequencer.Pin())
 	{
-		// Remove unsupported blend types
-		BlendTypes.Remove(Section->GetSupportedBlendTypes().Invert());
+		for (UMovieSceneSection* Section : Sequencer->GetViewModel()->GetSelection()->GetSelectedSections())
+		{
+			// Remove unsupported blend types
+			BlendTypes.Remove(Section->GetSupportedBlendTypes().Invert());
+		}
 	}
 
 	return BlendTypes;
@@ -604,6 +694,12 @@ void FSectionContextMenu::AddBlendTypeMenu(FMenuBuilder& MenuBuilder)
 {
 	using namespace UE::Sequencer;
 
+	const TSharedPtr<FSequencer> Sequencer = WeakSequencer.Pin();
+	if (!Sequencer.IsValid())
+	{
+		return;
+	}
+
 	TArray<TWeakObjectPtr<UMovieSceneSection>> Sections;
 	for (TViewModelPtr<FSectionModel> SectionModel : Sequencer->GetViewModel()->GetSelection()->TrackArea.Filter<FSectionModel>())
 	{
@@ -613,7 +709,6 @@ void FSectionContextMenu::AddBlendTypeMenu(FMenuBuilder& MenuBuilder)
 		}
 	}
 
-	TWeakPtr<FSequencer> WeakSequencer = Sequencer;
 	FSequencerUtilities::PopulateMenu_SetBlendType(MenuBuilder, Sections, WeakSequencer);
 }
 
@@ -621,6 +716,12 @@ void FSectionContextMenu::SelectAllKeys()
 {
 	using namespace UE::Sequencer;
 
+	const TSharedPtr<FSequencer> Sequencer = WeakSequencer.Pin();
+	if (!Sequencer.IsValid())
+	{
+		return;
+	}
+
 	TArray<FKeyHandle> HandlesScratch;
 
 	FSequencerSelection& Selection = *Sequencer->GetViewModel()->GetSelection();
@@ -649,12 +750,24 @@ void FSectionContextMenu::SelectAllKeys()
 
 void FSectionContextMenu::CopyAllKeys()
 {
+	const TSharedPtr<FSequencer> Sequencer = WeakSequencer.Pin();
+	if (!Sequencer.IsValid())
+	{
+		return;
+	}
+
 	SelectAllKeys();
 	Sequencer->CopySelectedKeys();
 }
 
 void FSectionContextMenu::SetSectionToKey()
 {
+	const TSharedPtr<FSequencer> Sequencer = WeakSequencer.Pin();
+	if (!Sequencer.IsValid())
+	{
+		return;
+	}
+
 	if (Sequencer->GetViewModel()->GetSelection()->GetSelectedSections().Num() != 1)
 	{
 		return;
@@ -677,6 +790,12 @@ void FSectionContextMenu::SetSectionToKey()
 
 bool FSectionContextMenu::IsSectionToKey() const
 {
+	const TSharedPtr<FSequencer> Sequencer = WeakSequencer.Pin();
+	if (!Sequencer.IsValid())
+	{
+		return false;
+	}
+
 	for (UMovieSceneSection* Section : Sequencer->GetViewModel()->GetSelection()->GetSelectedSections())
 	{
 		UMovieSceneTrack* Track = Section->GetTypedOuter<UMovieSceneTrack>();
@@ -685,11 +804,18 @@ bool FSectionContextMenu::IsSectionToKey() const
 			return false;
 		}
 	}
+
 	return true;
 }
 
 bool FSectionContextMenu::CanSetSectionToKey() const
 {
+	const TSharedPtr<FSequencer> Sequencer = WeakSequencer.Pin();
+	if (!Sequencer.IsValid())
+	{
+		return false;
+	}
+
 	if (Sequencer->GetViewModel()->GetSelection()->GetSelectedSections().Num() != 1)
 	{
 		return false;
@@ -726,6 +852,12 @@ bool FSectionContextMenu::CanSelectAllKeys() const
 
 void FSectionContextMenu::AutoSizeSection()
 {
+	const TSharedPtr<FSequencer> Sequencer = WeakSequencer.Pin();
+	if (!Sequencer.IsValid())
+	{
+		return;
+	}
+
 	FScopedTransaction AutoSizeSectionTransaction(LOCTEXT("AutoSizeSection_Transaction", "Auto Size Section"));
 
 	for (UMovieSceneSection* Section : Sequencer->GetViewModel()->GetSelection()->GetSelectedSections())
@@ -749,6 +881,12 @@ void FSectionContextMenu::ReduceKeys()
 {
 	using namespace UE::Sequencer;
 
+	const TSharedPtr<FSequencer> Sequencer = WeakSequencer.Pin();
+	if (!Sequencer.IsValid())
+	{
+		return;
+	}
+
 	FScopedTransaction ReduceKeysTransaction(LOCTEXT("ReduceKeys_Transaction", "Reduce Keys"));
 
 	TSet<TSharedPtr<IKeyArea> > KeyAreas;
@@ -791,6 +929,12 @@ void FSectionContextMenu::ReduceKeys()
 
 bool FSectionContextMenu::CanAutoSize() const
 {
+	const TSharedPtr<FSequencer> Sequencer = WeakSequencer.Pin();
+	if (!Sequencer.IsValid())
+	{
+		return false;
+	}
+
 	for (UMovieSceneSection* Section : Sequencer->GetViewModel()->GetSelection()->GetSelectedSections())
 	{
 		if (Section && Section->GetAutoSizeRange().IsSet())
@@ -798,6 +942,7 @@ bool FSectionContextMenu::CanAutoSize() const
 			return true;
 		}
 	}
+
 	return false;
 }
 
@@ -805,6 +950,12 @@ bool FSectionContextMenu::CanReduceKeys() const
 {
 	using namespace UE::Sequencer;
 
+	const TSharedPtr<FSequencer> Sequencer = WeakSequencer.Pin();
+	if (!Sequencer.IsValid())
+	{
+		return false;
+	}
+
 	TSet<TSharedPtr<IKeyArea> > KeyAreas;
 	for (const TWeakPtr<FViewModel>& WeakItem : Sequencer->GetViewModel()->GetSelection()->Outliner)
 	{
@@ -826,6 +977,12 @@ void FSectionContextMenu::SetInterpTangentMode(ERichCurveInterpMode InterpMode,
 {
 	using namespace UE::Sequencer;
 
+	const TSharedPtr<FSequencer> Sequencer = WeakSequencer.Pin();
+	if (!Sequencer.IsValid())
+	{
+		return;
+	}
+
 	FScopedTransaction SetInterpTangentModeTransaction(LOCTEXT("SetInterpTangentMode_Transaction", "Set Interpolation and Tangent Mode"));
 
 	TSet<TSharedPtr<IKeyArea> > KeyAreas;
@@ -898,6 +1055,12 @@ bool FSectionContextMenu::CanSetInterpTangentMode() const
 {
 	using namespace UE::Sequencer;
 
+	const TSharedPtr<FSequencer> Sequencer = WeakSequencer.Pin();
+	if (!Sequencer.IsValid())
+	{
+		return false;
+	}
+
 	TSet<TSharedPtr<IKeyArea> > KeyAreas;
 	for (const TWeakPtr<FViewModel>& WeakItem : Sequencer->GetViewModel()->GetSelection()->Outliner)
 	{
@@ -928,6 +1091,12 @@ bool FSectionContextMenu::CanSetInterpTangentMode() const
 
 void FSectionContextMenu::ToggleSectionActive()
 {
+	const TSharedPtr<FSequencer> Sequencer = WeakSequencer.Pin();
+	if (!Sequencer.IsValid())
+	{
+		return;
+	}
+
 	FScopedTransaction ToggleSectionActiveTransaction( LOCTEXT("ToggleSectionActive_Transaction", "Toggle Section Active") );
 	bool bIsActive = !IsSectionActive();
 	bool bAnythingChanged = false;
@@ -951,6 +1120,12 @@ void FSectionContextMenu::ToggleSectionActive()
 
 bool FSectionContextMenu::IsSectionActive() const
 {
+	const TSharedPtr<FSequencer> Sequencer = WeakSequencer.Pin();
+	if (!Sequencer.IsValid())
+	{
+		return false;
+	}
+
 	// Active only if all are active
 	for (UMovieSceneSection* Section : Sequencer->GetViewModel()->GetSelection()->GetSelectedSections())
 	{
@@ -966,6 +1141,12 @@ bool FSectionContextMenu::IsSectionActive() const
 
 void FSectionContextMenu::ToggleSectionLocked()
 {
+	const TSharedPtr<FSequencer> Sequencer = WeakSequencer.Pin();
+	if (!Sequencer.IsValid())
+	{
+		return;
+	}
+
 	FScopedTransaction ToggleSectionLockedTransaction( NSLOCTEXT("Sequencer", "ToggleSectionLocked_Transaction", "Toggle Section Locked") );
 	bool bIsLocked = !IsSectionLocked();
 	bool bAnythingChanged = false;
@@ -993,6 +1174,12 @@ void FSectionContextMenu::ToggleSectionLocked()
 
 bool FSectionContextMenu::IsSectionLocked() const
 {
+	const TSharedPtr<FSequencer> Sequencer = WeakSequencer.Pin();
+	if (!Sequencer.IsValid())
+	{
+		return false;
+	}
+
 	// Locked only if all are locked
 	for (UMovieSceneSection* Section : Sequencer->GetViewModel()->GetSelection()->GetSelectedSections())
 	{
@@ -1008,10 +1195,15 @@ bool FSectionContextMenu::IsSectionLocked() const
 
 void FSectionContextMenu::DeleteSection()
 {
+	const TSharedPtr<FSequencer> Sequencer = WeakSequencer.Pin();
+	if (!Sequencer.IsValid())
+	{
+		return;
+	}
+
 	Sequencer->DeleteSections(Sequencer->GetViewModel()->GetSelection()->GetSelectedSections());
 }
 
-
 /** Information pertaining to a specific row in a track, required for z-ordering operations */
 struct FTrackSectionRow
 {
@@ -1091,10 +1283,15 @@ void ModifySections(TMap<UMovieSceneTrack*, TMap<int32, FTrackSectionRow>>& Trac
 	}
 }
 
-
 void FSectionContextMenu::BringToFront()
 {
-	TMap<UMovieSceneTrack*, TMap<int32, FTrackSectionRow>> TrackRows = GenerateTrackRowsFromSelection(*Sequencer);
+	const TSharedPtr<FSequencer> Sequencer = WeakSequencer.Pin();
+	if (!Sequencer.IsValid())
+	{
+		return;
+	}
+
+	TMap<UMovieSceneTrack*, TMap<int32, FTrackSectionRow>> TrackRows = GenerateTrackRowsFromSelection(*Sequencer.Get());
 	if (TrackRows.Num() == 0)
 	{
 		return;
@@ -1139,10 +1336,15 @@ void FSectionContextMenu::BringToFront()
 	Sequencer->SetLocalTimeDirectly(Sequencer->GetLocalTime().Time);
 }
 
-
 void FSectionContextMenu::SendToBack()
 {
-	TMap<UMovieSceneTrack*, TMap<int32, FTrackSectionRow>> TrackRows = GenerateTrackRowsFromSelection(*Sequencer);
+	const TSharedPtr<FSequencer> Sequencer = WeakSequencer.Pin();
+	if (!Sequencer.IsValid())
+	{
+		return;
+	}
+
+	TMap<UMovieSceneTrack*, TMap<int32, FTrackSectionRow>> TrackRows = GenerateTrackRowsFromSelection(*Sequencer.Get());
 	if (TrackRows.Num() == 0)
 	{
 		return;
@@ -1187,10 +1389,15 @@ void FSectionContextMenu::SendToBack()
 	Sequencer->SetLocalTimeDirectly(Sequencer->GetLocalTime().Time);
 }
 
-
 void FSectionContextMenu::BringForward()
 {
-	TMap<UMovieSceneTrack*, TMap<int32, FTrackSectionRow>> TrackRows = GenerateTrackRowsFromSelection(*Sequencer);
+	const TSharedPtr<FSequencer> Sequencer = WeakSequencer.Pin();
+	if (!Sequencer.IsValid())
+	{
+		return;
+	}
+
+	TMap<UMovieSceneTrack*, TMap<int32, FTrackSectionRow>> TrackRows = GenerateTrackRowsFromSelection(*Sequencer.Get());
 	if (TrackRows.Num() == 0)
 	{
 		return;
@@ -1232,10 +1439,15 @@ void FSectionContextMenu::BringForward()
 	Sequencer->SetLocalTimeDirectly(Sequencer->GetLocalTime().Time);
 }
 
-
 void FSectionContextMenu::SendBackward()
 {
-	TMap<UMovieSceneTrack*, TMap<int32, FTrackSectionRow>> TrackRows = GenerateTrackRowsFromSelection(*Sequencer);
+	const TSharedPtr<FSequencer> Sequencer = WeakSequencer.Pin();
+	if (!Sequencer.IsValid())
+	{
+		return;
+	}
+
+	TMap<UMovieSceneTrack*, TMap<int32, FTrackSectionRow>> TrackRows = GenerateTrackRowsFromSelection(*Sequencer.Get());
 	if (TrackRows.Num() == 0)
 	{
 		return;
@@ -1277,9 +1489,9 @@ void FSectionContextMenu::SendBackward()
 	Sequencer->SetLocalTimeDirectly(Sequencer->GetLocalTime().Time);
 }
 
-bool FPasteContextMenu::BuildMenu(FMenuBuilder& MenuBuilder, TSharedPtr<FExtender> MenuExtender, FSequencer& InSequencer, const FPasteContextMenuArgs& Args)
+bool FPasteContextMenu::BuildMenu(FMenuBuilder& MenuBuilder, TSharedPtr<FExtender> MenuExtender, TWeakPtr<FSequencer> InWeakSequencer, const FPasteContextMenuArgs& Args)
 {
-	TSharedRef<FPasteContextMenu> Menu = MakeShareable(new FPasteContextMenu(InSequencer, Args));
+	TSharedRef<FPasteContextMenu> Menu = MakeShareable(new FPasteContextMenu(InWeakSequencer, Args));
 	Menu->Setup();
 	if (!Menu->IsValidPaste())
 	{
@@ -1291,9 +1503,9 @@ bool FPasteContextMenu::BuildMenu(FMenuBuilder& MenuBuilder, TSharedPtr<FExtende
 }
 
 
-TSharedRef<FPasteContextMenu> FPasteContextMenu::CreateMenu(FSequencer& InSequencer, const FPasteContextMenuArgs& Args)
+TSharedRef<FPasteContextMenu> FPasteContextMenu::CreateMenu(TWeakPtr<FSequencer> InWeakSequencer, const FPasteContextMenuArgs& Args)
 {
-	TSharedRef<FPasteContextMenu> Menu = MakeShareable(new FPasteContextMenu(InSequencer, Args));
+	TSharedRef<FPasteContextMenu> Menu = MakeShareable(new FPasteContextMenu(InWeakSequencer, Args));
 	Menu->Setup();
 	return Menu;
 }
@@ -1356,6 +1568,12 @@ void FPasteContextMenu::Setup()
 {
 	using namespace UE::Sequencer;
 
+	const TSharedPtr<FSequencer> Sequencer = WeakSequencer.Pin();
+	if (!Sequencer.IsValid())
+	{
+		return;
+	}
+
 	if (!Args.Clipboard.IsValid())
 	{
 		if (Sequencer->GetClipboardStack().Num() != 0)
@@ -1610,6 +1828,12 @@ void FPasteContextMenu::EndPasteInto(bool bAnythingPasted, const TSet<FSequencer
 
 	UE::Sequencer::SSequencerSection::ThrobKeySelection();
 
+	const TSharedPtr<FSequencer> Sequencer = WeakSequencer.Pin();
+	if (!Sequencer.IsValid())
+	{
+		return;
+	}
+
 	FSequencerSelection& Selection = *Sequencer->GetViewModel()->GetSelection();
 	{
 		FSelectionEventSuppressor EventSuppressor = Selection.SuppressEvents();
@@ -1632,6 +1856,12 @@ void FPasteContextMenu::EndPasteInto(bool bAnythingPasted, const TSet<FSequencer
 
 bool FPasteContextMenu::PasteInto(int32 DestinationIndex, FName KeyAreaName, TSet<FSequencerSelectedKey>& NewSelection)
 {
+	const TSharedPtr<FSequencer> Sequencer = WeakSequencer.Pin();
+	if (!Sequencer.IsValid())
+	{
+		return false;
+	}
+
 	FSequencerClipboardReconciler& Reconciler = PasteDestinations[DestinationIndex].Reconcilers[KeyAreaName];
 
 	FSequencerPasteEnvironment PasteEnvironment;
@@ -1645,32 +1875,50 @@ bool FPasteContextMenu::PasteInto(int32 DestinationIndex, FName KeyAreaName, TSe
 }
 
 
-bool FPasteFromHistoryContextMenu::BuildMenu(FMenuBuilder& MenuBuilder, TSharedPtr<FExtender> MenuExtender, FSequencer& InSequencer, const FPasteContextMenuArgs& Args)
+bool FPasteFromHistoryContextMenu::BuildMenu(FMenuBuilder& MenuBuilder, TSharedPtr<FExtender> MenuExtender, TWeakPtr<FSequencer> InWeakSequencer, const FPasteContextMenuArgs& Args)
 {
-	if (InSequencer.GetClipboardStack().Num() == 0)
+	const TSharedPtr<FSequencer> Sequencer = InWeakSequencer.Pin();
+	if (!Sequencer.IsValid())
+	{
+		return false;
+	}
+
+	if (Sequencer->GetClipboardStack().Num() == 0)
 	{
 		return false;
 	}
 
-	TSharedRef<FPasteFromHistoryContextMenu> Menu = MakeShareable(new FPasteFromHistoryContextMenu(InSequencer, Args));
+	TSharedRef<FPasteFromHistoryContextMenu> Menu = MakeShareable(new FPasteFromHistoryContextMenu(InWeakSequencer, Args));
 	Menu->PopulateMenu(MenuBuilder, MenuExtender);
 	return true;
 }
 
 
-TSharedPtr<FPasteFromHistoryContextMenu> FPasteFromHistoryContextMenu::CreateMenu(FSequencer& InSequencer, const FPasteContextMenuArgs& Args)
+TSharedPtr<FPasteFromHistoryContextMenu> FPasteFromHistoryContextMenu::CreateMenu(TWeakPtr<FSequencer> InWeakSequencer, const FPasteContextMenuArgs& Args)
 {
-	if (InSequencer.GetClipboardStack().Num() == 0)
+	const TSharedPtr<FSequencer> Sequencer = InWeakSequencer.Pin();
+	if (!Sequencer.IsValid())
+	{
+		return nullptr;
+	}
+
+	if (Sequencer->GetClipboardStack().Num() == 0)
 	{
 		return nullptr;
 	}
 
-	return MakeShareable(new FPasteFromHistoryContextMenu(InSequencer, Args));
+	return MakeShareable(new FPasteFromHistoryContextMenu(InWeakSequencer, Args));
 }
 
 
 void FPasteFromHistoryContextMenu::PopulateMenu(FMenuBuilder& MenuBuilder, TSharedPtr<FExtender> MenuExtender)
 {
+	const TSharedPtr<FSequencer> Sequencer = WeakSequencer.Pin();
+	if (!Sequencer.IsValid())
+	{
+		return;
+	}
+
 	// Copy a reference to the context menu by value into each lambda handler to ensure the type stays alive until the menu is closed
 	TSharedRef<FPasteFromHistoryContextMenu> Shared = AsShared();
 
@@ -1681,7 +1929,7 @@ void FPasteFromHistoryContextMenu::PopulateMenu(FMenuBuilder& MenuBuilder, TShar
 		FPasteContextMenuArgs ThisPasteArgs = Args;
 		ThisPasteArgs.Clipboard = Sequencer->GetClipboardStack()[Index];
 
-		TSharedRef<FPasteContextMenu> PasteMenu = FPasteContextMenu::CreateMenu(*Sequencer, ThisPasteArgs);
+		TSharedRef<FPasteContextMenu> PasteMenu = FPasteContextMenu::CreateMenu(Sequencer, ThisPasteArgs);
 
 		MenuBuilder.AddSubMenu(
 			ThisPasteArgs.Clipboard->GetDisplayText(),
@@ -1699,20 +1947,26 @@ void FPasteFromHistoryContextMenu::PopulateMenu(FMenuBuilder& MenuBuilder, TShar
 	MenuBuilder.EndSection();
 }
 
-void FEasingContextMenu::BuildMenu(FMenuBuilder& MenuBuilder, TSharedPtr<FExtender> MenuExtender, const TArray<UE::Sequencer::FEasingAreaHandle>& InEasings, FSequencer& Sequencer, FFrameTime InMouseDownTime)
+void FEasingContextMenu::BuildMenu(FMenuBuilder& MenuBuilder, TSharedPtr<FExtender> MenuExtender, const TArray<UE::Sequencer::FEasingAreaHandle>& InEasings, TWeakPtr<FSequencer> InWeakSequencer, FFrameTime InMouseDownTime)
 {
-	TSharedRef<FEasingContextMenu> EasingMenu = MakeShareable(new FEasingContextMenu(InEasings, Sequencer));
+	TSharedRef<FEasingContextMenu> EasingMenu = MakeShareable(new FEasingContextMenu(InEasings, InWeakSequencer));
 	EasingMenu->PopulateMenu(MenuBuilder, MenuExtender);
 
 	MenuBuilder.AddMenuSeparator();
 
-	FSectionContextMenu::BuildMenu(MenuBuilder, MenuExtender, Sequencer, InMouseDownTime);
+	FSectionContextMenu::BuildMenu(MenuBuilder, MenuExtender, InWeakSequencer, InMouseDownTime);
 }
 
 void FEasingContextMenu::PopulateMenu(FMenuBuilder& MenuBuilder, TSharedPtr<FExtender> MenuExtender)
 {
 	using namespace UE::Sequencer;
 
+	const TSharedPtr<FSequencer> Sequencer = WeakSequencer.Pin();
+    if (!Sequencer.IsValid())
+    {
+    	return;
+    }
+
 	FText SectionText = Easings.Num() == 1 ? LOCTEXT("EasingCurve", "Easing Curve") : FText::Format(LOCTEXT("EasingCurvesFormat", "Easing Curves ({0} curves)"), FText::AsNumber(Easings.Num()));
 	const bool bReadOnly = Algo::AnyOf(Easings, [](const FEasingAreaHandle& Handle) -> bool
 		{
@@ -1766,15 +2020,24 @@ void FEasingContextMenu::PopulateMenu(FMenuBuilder& MenuBuilder, TSharedPtr<FExt
 					.MaxValue(TOptional<double>())
 					.MaxSliderValue(TOptional<double>())
 					.MinSliderValue(0.f)
-					.Delta_Lambda([this]() -> double { return Sequencer->GetDisplayRateDeltaFrameCount(); })
-					.Value_Lambda([=] {
-						TOptional<int32> Current = Shared->GetCurrentLength();
-						if (Current.IsSet())
+					.Delta_Lambda([this]() -> double
+						{
+							const TSharedPtr<FSequencer> Sequencer = WeakSequencer.Pin();
+							if (!Sequencer.IsValid())
+							{
+								return 0;
+							}
+							return Sequencer->GetDisplayRateDeltaFrameCount();
+						})
+					.Value_Lambda([=]
 						{
-							return TOptional<double>(Current.GetValue());
-						}
-						return TOptional<double>();
-					})
+							TOptional<int32> Current = Shared->GetCurrentLength();
+							if (Current.IsSet())
+							{
+								return TOptional<double>(Current.GetValue());
+							}
+							return TOptional<double>();
+						})
 					.OnValueChanged_Lambda([=](double NewLength){ Shared->OnUpdateLength(NewLength); })
 					.OnValueCommitted_Lambda(OnValueCommitted)
 					.OnBeginSliderMovement_Lambda(OnBeginSliderMovement)
diff --git a/Engine/Source/Editor/Sequencer/Private/SequencerContextMenus.h b/Engine/Source/Editor/Sequencer/Private/SequencerContextMenus.h
--- a/Engine/Source/Editor/Sequencer/Private/SequencerContextMenus.h
+++ b/Engine/Source/Editor/Sequencer/Private/SequencerContextMenus.h
@@ -25,14 +25,15 @@ class UMovieSceneSection;
  */
 struct FSectionContextMenu : TSharedFromThis<FSectionContextMenu>
 {
-	static void BuildMenu(FMenuBuilder& MenuBuilder, TSharedPtr<FExtender> MenuExtender, FSequencer& Sequencer, FFrameTime InMouseDownTime);
+	static void BuildMenu(FMenuBuilder& MenuBuilder, TSharedPtr<FExtender> MenuExtender, TWeakPtr<FSequencer> InWeakSequencer, FFrameTime InMouseDownTime);
+	static void BuildKeyEditMenu(FMenuBuilder& MenuBuilder, TWeakPtr<FSequencer> InWeakSequencer, FFrameTime InMouseDownTime);
 
 private:
 
 	/** Hidden AsShared() methods to discourage CreateSP delegate use. */
 	using TSharedFromThis::AsShared;
 
-	FSectionContextMenu(FSequencer& InSeqeuncer, FFrameTime InMouseDownTime);
+	FSectionContextMenu(TWeakPtr<FSequencer> InWeakSequencer, FFrameTime InMouseDownTime);
 
 	void PopulateMenu(FMenuBuilder& MenuBuilder, TSharedPtr<FExtender> MenuExtender);
 
@@ -42,6 +43,8 @@ private:
 	/** Add the Order sub-menu. */
 	void AddOrderMenu(FMenuBuilder& MenuBuilder);
 
+	void AddKeyInterpolationMenu(FMenuBuilder& MenuBuilder);
+	void AddKeyEditMenu(FMenuBuilder& MenuBuilder);
 	void AddBlendTypeMenu(FMenuBuilder& MenuBuilder);
 
 	void SelectAllKeys();
@@ -89,7 +92,7 @@ private:
 	FMovieSceneBlendTypeField GetSupportedBlendTypes() const;
 
 	/** The sequencer */
-	TSharedRef<FSequencer> Sequencer;
+	TWeakPtr<FSequencer> WeakSequencer;
 
 	/** The time that we clicked on to summon this menu */
 	FFrameTime MouseDownTime;
@@ -132,9 +135,9 @@ struct FPasteContextMenuArgs
 
 struct FPasteContextMenu : TSharedFromThis<FPasteContextMenu>
 {
-	static bool BuildMenu(FMenuBuilder& MenuBuilder, TSharedPtr<FExtender> MenuExtender, FSequencer& Sequencer, const FPasteContextMenuArgs& Args);
+	static bool BuildMenu(FMenuBuilder& MenuBuilder, TSharedPtr<FExtender> MenuExtender, TWeakPtr<FSequencer> InWeakSequencer, const FPasteContextMenuArgs& Args);
 
-	static TSharedRef<FPasteContextMenu> CreateMenu(FSequencer& Sequencer, const FPasteContextMenuArgs& Args);
+	static TSharedRef<FPasteContextMenu> CreateMenu(TWeakPtr<FSequencer> InWeakSequencer, const FPasteContextMenuArgs& Args);
 
 	void PopulateMenu(FMenuBuilder& MenuBuilder, TSharedPtr<FExtender> MenuExtender);
 
@@ -146,8 +149,8 @@ private:
 
 	using TSharedFromThis::AsShared;
 
-	FPasteContextMenu(FSequencer& InSequencer, const FPasteContextMenuArgs& InArgs)
-		: Sequencer(StaticCastSharedRef<FSequencer>(InSequencer.AsShared()))
+	FPasteContextMenu(TWeakPtr<FSequencer> InWeakSequencer, const FPasteContextMenuArgs& InArgs)
+		: WeakSequencer(InWeakSequencer)
 		, bPasteFirstOnly(true)
 		, Args(InArgs)
 	{}
@@ -163,7 +166,7 @@ private:
 	void GatherPasteDestinationsForNode(const UE::Sequencer::TViewModelPtr<UE::Sequencer::IOutlinerExtension>& InNode, UMovieSceneSection* InSection, const FName& CurrentScope, TMap<FName, FSequencerClipboardReconciler>& Map);
 
 	/** The sequencer */
-	TSharedRef<FSequencer> Sequencer;
+	TWeakPtr<FSequencer> WeakSequencer;
 
 	/** Paste destinations are organized by track type primarily, then by key area name  */
 	struct FPasteDestination
@@ -181,9 +184,9 @@ private:
 
 struct FPasteFromHistoryContextMenu : TSharedFromThis<FPasteFromHistoryContextMenu>
 {
-	static bool BuildMenu(FMenuBuilder& MenuBuilder, TSharedPtr<FExtender> MenuExtender, FSequencer& Sequencer, const FPasteContextMenuArgs& Args);
+	static bool BuildMenu(FMenuBuilder& MenuBuilder, TSharedPtr<FExtender> MenuExtender, TWeakPtr<FSequencer> InWeakSequencer, const FPasteContextMenuArgs& Args);
 
-	static TSharedPtr<FPasteFromHistoryContextMenu> CreateMenu(FSequencer& Sequencer, const FPasteContextMenuArgs& Args);
+	static TSharedPtr<FPasteFromHistoryContextMenu> CreateMenu(TWeakPtr<FSequencer> InWeakSequencer, const FPasteContextMenuArgs& Args);
 
 	void PopulateMenu(FMenuBuilder& MenuBuilder, TSharedPtr<FExtender> MenuExtender);
 
@@ -191,13 +194,13 @@ private:
 
 	using TSharedFromThis::AsShared;
 
-	FPasteFromHistoryContextMenu(FSequencer& InSequencer, const FPasteContextMenuArgs& InArgs)
-		: Sequencer(StaticCastSharedRef<FSequencer>(InSequencer.AsShared()))
+	FPasteFromHistoryContextMenu(TWeakPtr<FSequencer> InWeakSequencer, const FPasteContextMenuArgs& InArgs)
+		: WeakSequencer(InWeakSequencer)
 		, Args(InArgs)
 	{}
 
 	/** The sequencer */
-	TSharedRef<FSequencer> Sequencer;
+	TWeakPtr<FSequencer> WeakSequencer;
 
 	/** Paste arguments */
 	FPasteContextMenuArgs Args;
@@ -211,12 +214,12 @@ private:
  */
 struct FKeyContextMenu : TSharedFromThis<FKeyContextMenu>
 {
-	static void BuildMenu(FMenuBuilder& MenuBuilder, TSharedPtr<FExtender> MenuExtender, FSequencer& Sequencer);
+	static void BuildMenu(FMenuBuilder& MenuBuilder, TSharedPtr<FExtender> MenuExtender, TWeakPtr<FSequencer> InWeakSequencer);
 
 private:
 
-	FKeyContextMenu(FSequencer& InSequencer)
-		: Sequencer(StaticCastSharedRef<FSequencer>(InSequencer.AsShared()))
+	FKeyContextMenu(TWeakPtr<FSequencer> InWeakSequencer)
+		: WeakSequencer(InWeakSequencer)
 	{}
 
 	/** Hidden AsShared() methods to discourage CreateSP delegate use. */
@@ -228,7 +231,7 @@ private:
 	void PopulateMenu(FMenuBuilder& MenuBuilder, TSharedPtr<FExtender> MenuExtender);
 
 	/** The sequencer */
-	TSharedRef<FSequencer> Sequencer;
+	TWeakPtr<FSequencer> WeakSequencer;
 
 	TSharedPtr<FStructOnScope> KeyStruct;
 	TWeakObjectPtr<UMovieSceneSection> KeyStructSection;
@@ -243,14 +246,13 @@ private:
  */
 struct FEasingContextMenu : TSharedFromThis<FEasingContextMenu>
 {
-	static void BuildMenu(FMenuBuilder& MenuBuilder, TSharedPtr<FExtender> MenuExtender, const TArray<UE::Sequencer::FEasingAreaHandle>& InEasings, FSequencer& Sequencer, FFrameTime InMouseDownTime);
+	static void BuildMenu(FMenuBuilder& MenuBuilder, TSharedPtr<FExtender> MenuExtender, const TArray<UE::Sequencer::FEasingAreaHandle>& InEasings, TWeakPtr<FSequencer> InWeakSequencer, FFrameTime InMouseDownTime);
 
 private:
 
-	FEasingContextMenu(const TArray<UE::Sequencer::FEasingAreaHandle>& InEasings, FSequencer& InSequencer)
+	FEasingContextMenu(const TArray<UE::Sequencer::FEasingAreaHandle>& InEasings, TWeakPtr<FSequencer> InWeakSequencer)
 		: Easings(InEasings)
-		, Sequencer(StaticCastSharedRef<FSequencer>(InSequencer.AsShared()))
-
+		, WeakSequencer(InWeakSequencer)
 	{}
 
 	/** Hidden AsShared() methods to discourage CreateSP delegate use. */
@@ -277,7 +279,7 @@ private:
 	TArray<UE::Sequencer::FEasingAreaHandle> Easings;
 
 	/** The sequencer */
-	TSharedRef<FSequencer> Sequencer;
+	TWeakPtr<FSequencer> WeakSequencer;
 
 	/** A scoped transaction for a current operation */
 	TUniquePtr<FScopedTransaction> ScopedTransaction;
diff --git a/Engine/Source/Editor/Sequencer/Private/SequencerHotspots.cpp b/Engine/Source/Editor/Sequencer/Private/SequencerHotspots.cpp
--- a/Engine/Source/Editor/Sequencer/Private/SequencerHotspots.cpp
+++ b/Engine/Source/Editor/Sequencer/Private/SequencerHotspots.cpp
@@ -234,10 +234,7 @@ TOptional<FFrameNumber> FKeyHotspot::GetTime() const
 
 bool FKeyHotspot::PopulateContextMenu(FMenuBuilder& MenuBuilder, TSharedPtr<FExtender> MenuExtender, FFrameTime MouseDownTime)
 {
-	if (TSharedPtr<FSequencer> Sequencer = WeakSequencer.Pin())
-	{
-		FKeyContextMenu::BuildMenu(MenuBuilder, MenuExtender, *Sequencer);
-	}
+	FKeyContextMenu::BuildMenu(MenuBuilder, MenuExtender, WeakSequencer);
 	return true;
 }
 
@@ -275,7 +272,7 @@ bool FKeyBarHotspot::PopulateContextMenu(FMenuBuilder& MenuBuilder, TSharedPtr<F
 		return false;
 	}
 	
-	FSectionContextMenu::BuildMenu(MenuBuilder, MenuExtender, *Sequencer, MouseDownTime);
+	FSectionContextMenu::BuildMenu(MenuBuilder, MenuExtender, WeakSequencer, MouseDownTime);
 
 	TSharedPtr<IObjectBindingExtension> ObjectBinding = SectionModel->FindAncestorOfType<IObjectBindingExtension>();
 	SectionModel->GetSectionInterface()->BuildSectionContextMenu(MenuBuilder, ObjectBinding ? ObjectBinding->GetObjectGuid() : FGuid());
@@ -514,8 +511,7 @@ bool FSectionHotspotBase::PopulateContextMenu(FMenuBuilder& MenuBuilder, TShared
 	UMovieSceneSection*       ThisSection  = SectionModel ? SectionModel->GetSection() : nullptr;
 	if (ThisSection)
 	{
-		TSharedPtr<FSequencer> Sequencer = WeakSequencer.Pin();
-		FSectionContextMenu::BuildMenu(MenuBuilder, MenuExtender, *Sequencer, MouseDownTime);
+		FSectionContextMenu::BuildMenu(MenuBuilder, MenuExtender, WeakSequencer, MouseDownTime);
 
 		TSharedPtr<IObjectBindingExtension> ObjectBinding = SectionModel->FindAncestorOfType<IObjectBindingExtension>();
 		SectionModel->GetSectionInterface()->BuildSectionContextMenu(MenuBuilder, ObjectBinding ? ObjectBinding->GetObjectGuid() : FGuid());
@@ -605,9 +601,7 @@ void FSectionEasingHandleHotspot::UpdateOnHover(FTrackAreaViewModel& InTrackArea
 
 bool FSectionEasingHandleHotspot::PopulateContextMenu(FMenuBuilder& MenuBuilder, TSharedPtr<FExtender> MenuExtender, FFrameTime MouseDownTime)
 {
-	TSharedPtr<FSequencer> Sequencer = WeakSequencer.Pin();
-
-	FEasingContextMenu::BuildMenu(MenuBuilder, MenuExtender, { FEasingAreaHandle{WeakSectionModel, HandleType} }, *Sequencer, MouseDownTime);
+	FEasingContextMenu::BuildMenu(MenuBuilder, MenuExtender, { FEasingAreaHandle{WeakSectionModel, HandleType} }, WeakSequencer, MouseDownTime);
 	return true;
 }
 
@@ -628,7 +622,7 @@ bool FSectionEasingAreaHotspot::PopulateContextMenu(FMenuBuilder& MenuBuilder, T
 	using namespace UE::Sequencer;
 
 	TSharedPtr<FSequencer> Sequencer = WeakSequencer.Pin();
-	FEasingContextMenu::BuildMenu(MenuBuilder, MenuExtender, Easings, *Sequencer, MouseDownTime);
+	FEasingContextMenu::BuildMenu(MenuBuilder, MenuExtender, Easings, WeakSequencer, MouseDownTime);
 
 	TSharedPtr<FSectionModel> SectionModel = WeakSectionModel.Pin();
 	UMovieSceneSection*       ThisSection  = SectionModel ? SectionModel->GetSection() : nullptr;
diff --git a/Engine/Source/Editor/Sequencer/Private/SequencerUtilities.cpp b/Engine/Source/Editor/Sequencer/Private/SequencerUtilities.cpp
--- a/Engine/Source/Editor/Sequencer/Private/SequencerUtilities.cpp
+++ b/Engine/Source/Editor/Sequencer/Private/SequencerUtilities.cpp
@@ -4269,85 +4269,31 @@ void FSequencerUtilities::AddConvertBindingMenu(FMenuBuilder& MenuBuilder, TShar
 		return;
 	}
 
-	auto ConvertBindings = [Sequencer, MovieScene, BindingsToConvert, OnBindingChanged](TFunction<FMovieScenePossessable*(FGuid,int32)> DoConvert) {
-		using namespace UE::Sequencer;
-
-		if (MovieScene->IsReadOnly())
-		{
-			FSequencerUtilities::ShowReadOnlyError();
-			return;
-		}
-
-		if (BindingsToConvert.Num() > 0)
-		{
-			const FScopedTransaction Transaction(LOCTEXT("ConvertSelectedNodesPossessable", "Convert Selected Nodes to Possessables"));
-			MovieScene->Modify();
-
-			FScopedSlowTask SlowTask(BindingsToConvert.Num(), LOCTEXT("ConvertPossessablesProgress", "Converting Selected Spawnable Nodes to Possessables"));
-			SlowTask.MakeDialog(true);
-
-			TArray<AActor*> PossessedActors;
-			for (const FSequencerChangeBindingInfo& BindingInfo : BindingsToConvert)
-			{
-				SlowTask.EnterProgressFrame();
-
-				if (FMovieScenePossessable* Possessable = DoConvert(BindingInfo.BindingID, BindingInfo.BindingIndex))
-				{
-					Sequencer->ForceEvaluate();
-
-					for (TWeakObjectPtr<> WeakObject : Sequencer->FindBoundObjects(Possessable->GetGuid(), Sequencer->GetFocusedTemplateID()))
-					{
-						if (AActor* PossessedActor = Cast<AActor>(WeakObject.Get()))
-						{
-							PossessedActors.Add(PossessedActor);
-						}
-					}
-
-					if (GWarn->ReceivedUserCancel())
-					{
-						break;
-					}
-				}
-			}
-
-			if (PossessedActors.Num())
-			{
-				const bool bNotifySelectionChanged = true;
-				const bool bDeselectBSP = true;
-				const bool bWarnAboutTooManyActors = false;
-				const bool bSelectEvenIfHidden = false;
-
-				GEditor->GetSelectedActors()->Modify();
-				GEditor->GetSelectedActors()->BeginBatchSelectOperation();
-				GEditor->SelectNone(bNotifySelectionChanged, bDeselectBSP, bWarnAboutTooManyActors);
-				for (auto PossessedActor : PossessedActors)
-				{
-					GEditor->SelectActor(PossessedActor, true, bNotifySelectionChanged, bSelectEvenIfHidden);
-				}
-				GEditor->GetSelectedActors()->EndBatchSelectOperation();
-				GEditor->NoteSelectionChange();
-
-				Sequencer->NotifyMovieSceneDataChanged(EMovieSceneDataChangeType::MovieSceneStructureItemsChanged);
-			}
-		}
-		if (OnBindingChanged != nullptr)
-		{
-			OnBindingChanged();
-		}
-	};
-
 	// Can convert to possessable
 	if (Algo::AllOf(BindingsToConvert, [Sequencer, &Sequence, &MovieScene](const FSequencerChangeBindingInfo& BindingInfo) {
 		return CanConvertToPossessable(Sequencer, BindingInfo.BindingID, BindingInfo.BindingIndex );
 	}))
 	{
+		const TWeakPtr<ISequencer> WeakSequencer = Sequencer;
+
 		MenuBuilder.AddMenuEntry(
 			LOCTEXT("ConvertToPossessable", "Possessable"),
 			LOCTEXT("ConvertToPossessableTooltip", "Convert selected binding(s) to a possessable"),
 			FSlateIcon(),
-			FUIAction(FExecuteAction::CreateLambda([Sequencer, MovieScene, BindingsToConvert, ConvertBindings, OnBindingChanged]()
+			FUIAction(FExecuteAction::CreateLambda([WeakSequencer, BindingsToConvert, OnBindingChanged]()
 			{
-				ConvertBindings([Sequencer](FGuid BindingID, int32 BindingIndex){return FSequencerUtilities::ConvertToPossessable(Sequencer, BindingID, BindingIndex);});
+				const TSharedPtr<ISequencer> Sequencer = WeakSequencer.Pin();
+				if (!Sequencer)
+				{
+					return;
+				}
+
+				const TSharedRef<ISequencer> SequencerRef = Sequencer.ToSharedRef();
+				ConvertBindings(SequencerRef, BindingsToConvert, [&SequencerRef](FGuid BindingID, int32 BindingIndex)
+					{
+						return FSequencerUtilities::ConvertToPossessable(SequencerRef, BindingID, BindingIndex);
+					},
+					OnBindingChanged);
 			})));
 	}
 	
@@ -4361,15 +4307,115 @@ void FSequencerUtilities::AddConvertBindingMenu(FMenuBuilder& MenuBuilder, TShar
 			return CanConvertToCustomBinding(Sequencer, BindingInfo.BindingID, CustomBindingType, BindingInfo.BindingIndex);
 		}))
 		{
+			const TWeakPtr<ISequencer> WeakSequencer = Sequencer;
+
 			MenuBuilder.AddMenuEntry(
 				CustomBindingType->GetDefaultObject<UMovieSceneCustomBinding>()->GetBindingTypePrettyName(),
 				FText::Format(LOCTEXT("ConvertToCustomBindingTooltip", "Convert selected binding to {0}"), CustomBindingType->GetDefaultObject<UMovieSceneCustomBinding>()->GetBindingTypePrettyName()),
 				FSlateIcon(),
-				FUIAction(FExecuteAction::CreateLambda([Sequencer, MovieScene, BindingsToConvert, ConvertBindings, CustomBindingType, OnBindingChanged]()
+				FUIAction(FExecuteAction::CreateLambda([WeakSequencer, CustomBindingType, BindingsToConvert, OnBindingChanged]()
 				{
-					ConvertBindings([Sequencer, CustomBindingType](FGuid BindingID, int32 BindingIndex){return FSequencerUtilities::ConvertToCustomBinding(Sequencer, BindingID, CustomBindingType, BindingIndex);});
+					const TSharedPtr<ISequencer> Sequencer = WeakSequencer.Pin();
+					if (!Sequencer)
+					{
+						return;
+					}
+
+					const TSharedRef<ISequencer> SequencerRef = Sequencer.ToSharedRef();
+					ConvertBindings(SequencerRef, BindingsToConvert, [&SequencerRef, CustomBindingType](FGuid BindingID, int32 BindingIndex)
+						{
+							return FSequencerUtilities::ConvertToCustomBinding(SequencerRef, BindingID, CustomBindingType, BindingIndex);
+						},
+						OnBindingChanged);
 				})));
 		}
 	}
 }
+
+void FSequencerUtilities::ConvertBindings(const TSharedRef<ISequencer>& InSequencer
+	, const TArray<FSequencerChangeBindingInfo>& InBindingsToConvert
+	, TFunction<FMovieScenePossessable*(FGuid, int32)> InDoConvert
+	, TFunction<void()> InOnBindingChanged)
+{
+	using namespace UE::Sequencer;
+
+	UMovieSceneSequence* const Sequence = InSequencer->GetFocusedMovieSceneSequence();
+	if (!IsValid(Sequence))
+	{
+		return;
+	}
+
+	UMovieScene* const MovieScene = Sequence->GetMovieScene();
+	if (!IsValid(MovieScene))
+	{
+		return;
+	}
+
+	if (InBindingsToConvert.Num() == 0)
+	{
+		return;
+	}
+
+	if (MovieScene->IsReadOnly())
+	{
+		FSequencerUtilities::ShowReadOnlyError();
+		return;
+	}
+
+	const FScopedTransaction Transaction(LOCTEXT("ConvertSelectedNodesPossessable", "Convert Selected Nodes to Possessables"));
+	MovieScene->Modify();
+
+	FScopedSlowTask SlowTask(InBindingsToConvert.Num(), LOCTEXT("ConvertPossessablesProgress", "Converting Selected Spawnable Nodes to Possessables"));
+	SlowTask.MakeDialog(true);
+
+	TArray<AActor*> PossessedActors;
+	for (const FSequencerChangeBindingInfo& BindingInfo : InBindingsToConvert)
+	{
+		SlowTask.EnterProgressFrame();
+
+		if (FMovieScenePossessable* Possessable = InDoConvert(BindingInfo.BindingID, BindingInfo.BindingIndex))
+		{
+			InSequencer->ForceEvaluate();
+
+			for (TWeakObjectPtr<> WeakObject : InSequencer->FindBoundObjects(Possessable->GetGuid(), InSequencer->GetFocusedTemplateID()))
+			{
+				if (AActor* PossessedActor = Cast<AActor>(WeakObject.Get()))
+				{
+					PossessedActors.Add(PossessedActor);
+				}
+			}
+
+			if (GWarn->ReceivedUserCancel())
+			{
+				break;
+			}
+		}
+	}
+
+	if (PossessedActors.Num())
+	{
+		const bool bNotifySelectionChanged = true;
+		const bool bDeselectBSP = true;
+		const bool bWarnAboutTooManyActors = false;
+		const bool bSelectEvenIfHidden = false;
+
+		GEditor->GetSelectedActors()->Modify();
+		GEditor->GetSelectedActors()->BeginBatchSelectOperation();
+		GEditor->SelectNone(bNotifySelectionChanged, bDeselectBSP, bWarnAboutTooManyActors);
+		for (auto PossessedActor : PossessedActors)
+		{
+			GEditor->SelectActor(PossessedActor, true, bNotifySelectionChanged, bSelectEvenIfHidden);
+		}
+		GEditor->GetSelectedActors()->EndBatchSelectOperation();
+		GEditor->NoteSelectionChange();
+
+		InSequencer->NotifyMovieSceneDataChanged(EMovieSceneDataChangeType::MovieSceneStructureItemsChanged);
+	}
+
+	if (InOnBindingChanged != nullptr)
+	{
+		InOnBindingChanged();
+	}
+}
+
 #undef LOCTEXT_NAMESPACE
diff --git a/Engine/Source/Editor/Sequencer/Private/Widgets/Sidebar/SequenceSelectionDrawer.cpp b/Engine/Source/Editor/Sequencer/Private/Widgets/Sidebar/SequenceSelectionDrawer.cpp
deleted file mode 100644
diff --git a/Engine/Source/Editor/Sequencer/Private/Widgets/Sidebar/SequencerSelectionDrawer.cpp b/Engine/Source/Editor/Sequencer/Private/Widgets/Sidebar/SequencerSelectionDrawer.cpp
new file mode 100644
--- /dev/null
+++ b/Engine/Source/Editor/Sequencer/Private/Widgets/Sidebar/SequencerSelectionDrawer.cpp
@@ -0,0 +1,563 @@
+// Copyright Epic Games, Inc. All Rights Reserved.
+
+#include "SequencerSelectionDrawer.h"
+#include "DetailsViewArgs.h"
+#include "FrameNumberDetailsCustomization.h"
+#include "Framework/MultiBox/MultiBoxBuilder.h"
+#include "IKeyArea.h"
+#include "ISequencerSection.h"
+#include "IStructureDetailsView.h"
+#include "Menus/CurveChannelSectionSidebarExtension.h"
+#include "Modules/ModuleManager.h"
+#include "MovieScene.h"
+#include "MovieSceneMarkedFrame.h"
+#include "MovieSceneSequence.h"
+#include "MVVM/Extensions/ITrackExtension.h"
+#include "MVVM/Selection/Selection.h"
+#include "MVVM/ViewModelPtr.h"
+#include "MVVM/ViewModels/CategoryModel.h"
+#include "MVVM/ViewModels/LayerBarModel.h"
+#include "MVVM/ViewModels/SectionModel.h"
+#include "MVVM/ViewModels/SequencerEditorViewModel.h"
+#include "PropertyEditorModule.h"
+#include "Sequencer.h"
+#include "SequencerCommonHelpers.h"
+#include "SequencerContextMenus.h"
+#include "SKeyEditInterface.h"
+#include "Styling/AppStyle.h"
+#include "Templates/SharedPointer.h"
+#include "Widgets/Layout/SBorder.h"
+
+#define LOCTEXT_NAMESPACE "SequencerSelectionDrawer"
+
+namespace UE::Sequencer::Private
+{
+	FKeyEditData GetKeyEditData(const FKeySelection& InKeySelection)
+	{
+		if (InKeySelection.Num() == 1)
+		{
+			for (const FKeyHandle Key : InKeySelection)
+			{
+				if (const TSharedPtr<FChannelModel> Channel = InKeySelection.GetModelForKey(Key))
+				{
+					FKeyEditData KeyEditData;
+					KeyEditData.KeyStruct     = Channel->GetKeyArea()->GetKeyStruct(Key);
+					KeyEditData.OwningSection = Channel->GetSection();
+					return KeyEditData;
+				}
+			}
+		}
+		else
+		{
+			TArray<FKeyHandle> KeyHandles;
+			UMovieSceneSection* CommonSection = nullptr;
+			for (FKeyHandle Key : InKeySelection)
+			{
+				TSharedPtr<FChannelModel> Channel = InKeySelection.GetModelForKey(Key);
+				if (Channel.IsValid())
+				{
+					KeyHandles.Add(Key);
+					if (!CommonSection)
+					{
+						CommonSection = Channel->GetSection();
+					}
+					else if (CommonSection != Channel->GetSection())
+					{
+						CommonSection = nullptr;
+						break;
+					}
+				}
+			}
+
+			if (CommonSection)
+			{
+				FKeyEditData KeyEditData;
+				KeyEditData.KeyStruct     = CommonSection->GetKeyStruct(KeyHandles);
+				KeyEditData.OwningSection = CommonSection;
+				return KeyEditData;
+			}
+		}
+
+		return FKeyEditData();
+	}
+
+	TSharedPtr<FSequencerSelection> GetSelection(const ISequencer& InSequencer)
+	{
+		const TSharedPtr<FSequencerEditorViewModel> ViewModel = InSequencer.GetViewModel();
+		if (!ViewModel.IsValid())
+		{
+			return nullptr;
+		}
+
+		return ViewModel->GetSelection();
+	}
+}
+
+using namespace UE::Sequencer;
+
+const FName FSequencerSelectionDrawer::UniqueId = TEXT("SequencerSelectionDrawer");
+
+FSequencerSelectionDrawer::FSequencerSelectionDrawer(const TWeakPtr<FSequencer>& InWeakSequencer)
+	: WeakSequencer(InWeakSequencer)
+{
+}
+
+FName FSequencerSelectionDrawer::GetUniqueId() const
+{
+	return UniqueId;
+}
+
+FName FSequencerSelectionDrawer::GetSectionId() const
+{
+	return TEXT("Selection");
+}
+
+FText FSequencerSelectionDrawer::GetSectionDisplayText() const
+{
+	return LOCTEXT("SelectionDisplayText", "Selection");
+}
+
+TSharedRef<SWidget> FSequencerSelectionDrawer::CreateContentWidget()
+{
+	if (const TSharedPtr<ISequencer> Sequencer = WeakSequencer.Pin())
+	{
+		Sequencer->OnActorAddedToSequencer().AddLambda([this](AActor* InActor, const FGuid InGuid)
+			{
+				OnSequencerSelectionChanged();
+			});
+
+		if (const TSharedPtr<FSequencerSelection> SequencerSelection = Private::GetSelection(*Sequencer.Get()))
+		{
+			SequencerSelection->OnChanged.AddSP(this, &FSequencerSelectionDrawer::OnSequencerSelectionChanged);
+			
+			OnSequencerSelectionChanged();
+		}
+	}
+
+	return SNew(SBorder)
+		.BorderImage(FAppStyle::GetBrush(TEXT("NoBorder")))
+		.Padding(0.f)
+		[
+			SAssignNew(ContentBox, SVerticalBox)
+			+ SVerticalBox::Slot()
+			.FillHeight(1.f)
+			[
+				CreateNoSelectionHintText()
+			]
+		];
+}
+
+void FSequencerSelectionDrawer::OnSequencerSelectionChanged()
+{
+	if (!ContentBox.IsValid())
+	{
+		return;
+	}
+
+	ContentBox->ClearChildren();
+
+	CurveChannelExtension.Reset();
+
+	const TSharedPtr<FSequencer> Sequencer = WeakSequencer.Pin();
+	if (!WeakSequencer.IsValid())
+	{
+		return;
+	}
+
+	const TSharedPtr<FSequencerSelection> SequencerSelection = Private::GetSelection(*Sequencer.Get());
+	if (!SequencerSelection.IsValid())
+	{
+		return;
+	}
+
+	const TSharedRef<FSequencerSelection> SelectionRef = SequencerSelection.ToSharedRef();
+
+	auto AddToContent = [this](const TSharedRef<SWidget>& InWidget)
+		{
+			ContentBox->AddSlot()
+				.AutoHeight()
+				[
+					InWidget
+				];
+		};
+
+	const ISequencerModule& SequencerModule = FModuleManager::Get().LoadModuleChecked<ISequencerModule>(TEXT("Sequencer"));
+	const TSharedPtr<FExtensibilityManager> SidebarExtensibilityManager = SequencerModule.GetSidebarExtensibilityManager();
+
+	FMenuBuilder MenuBuilder(/*bInShouldCloseWindowAfterMenuSelection=*/false
+		, Sequencer->GetCommandBindings()
+		, SidebarExtensibilityManager->GetAllExtenders()
+		, /*bInCloseSelfOnly=*/true, &FCoreStyle::Get(), /*bInSearchable=*/true, TEXT("Sequencer.Sidebar"));
+
+	/**
+	 * Selection details display order preference:
+	 *  1) Key items
+	 *  2) Track area items (if no key selected)
+	 *  3) Outliner items (if no key or track area selected)
+	 *  4) Marked frames
+	 */
+
+	// 1) Key items
+	BuildKeySelectionDetails(SelectionRef, MenuBuilder);
+
+	// Early out for key selections
+	const bool bIsKeySelected = SequencerSelection->KeySelection.Num() > 0;
+	if (bIsKeySelected)
+	{
+		AddToContent(MenuBuilder.MakeWidget());
+		return;
+	}
+
+	// 2) Track area items
+	BuildTrackAreaDetails(*Sequencer, SelectionRef, MenuBuilder);
+
+	// 3) Outliner items
+	const bool bIsTrackAreaSelected = SequencerSelection->TrackArea.Num() > 0;
+	if (!bIsTrackAreaSelected)
+	{
+		BuildOutlinerDetails(*Sequencer, SelectionRef, MenuBuilder);
+	}
+
+	// 4) Marked frames
+	BuildMarkedFrameDetails(SelectionRef, MenuBuilder);
+
+	AddToContent(MenuBuilder.MakeWidget());
+}
+
+void FSequencerSelectionDrawer::BuildKeySelectionDetails(const TSharedRef<FSequencerSelection>& InSelection, FMenuBuilder& MenuBuilder)
+{
+	if (InSelection->KeySelection.Num() == 0)
+	{
+		return;
+	}
+
+	const TSharedPtr<FSequencer> Sequencer = WeakSequencer.Pin();
+	if (!Sequencer.IsValid())
+	{
+		return;
+	}
+
+	const ISequencerModule& SequencerModule = FModuleManager::Get().LoadModuleChecked<ISequencerModule>(TEXT("Sequencer"));
+	const TSharedPtr<FExtensibilityManager> SidebarExtensibilityManager = SequencerModule.GetSidebarExtensibilityManager();
+
+	MenuBuilder.BeginSection(TEXT("KeyEdit"), LOCTEXT("KeyEditMenuSection", "Key Edit"));
+	{
+		MenuBuilder.AddWidget(CreateKeyFrameDetails(InSelection).ToSharedRef(), FText::GetEmpty(), /*bInNoIndent=*/true);
+	}
+	MenuBuilder.EndSection();
+
+	// Show the section for the keys if they are all part of the same section
+	TArray<TViewModelPtr<FChannelModel>> Channels;
+	for (const FKeyHandle KeyHandle : InSelection->KeySelection)
+	{
+		const TViewModelPtr<FChannelModel> Channel = InSelection->KeySelection.GetModelForKey(KeyHandle);
+		Channels.AddUnique(Channel);
+	}
+	if (Channels.Num() == 1)
+	{
+		FSectionContextMenu::BuildKeyEditMenu(MenuBuilder, WeakSequencer, Sequencer->GetLastEvaluatedLocalTime());
+	}
+}
+
+void FSequencerSelectionDrawer::BuildTrackAreaDetails(FSequencer& InSequencer, const TSharedRef<FSequencerSelection>& InSelection, FMenuBuilder& MenuBuilder)
+{
+	TArray<TWeakObjectPtr<>> AllSectionObjects;
+
+	for (const FViewModelPtr TrackAreaItem : InSelection->TrackArea)
+	{
+		if (const TViewModelPtr<FLayerBarModel> LayerBarModel = TrackAreaItem.ImplicitCast())
+		{
+			const TViewModelPtr<IOutlinerExtension> LinkedOutlinerItem = LayerBarModel->GetLinkedOutlinerItem();
+
+			if (const TViewModelPtr<FOutlinerItemModel> OutlinerItemModel = LinkedOutlinerItem.ImplicitCast())
+			{
+				OutlinerItemModel->BuildSidebarMenu(MenuBuilder);
+			}
+		}
+		else if (const TViewModelPtr<FSectionModel> SectionModel = TrackAreaItem.ImplicitCast())
+		{
+			if (InSelection->TrackArea.Num() == 1)
+			{
+				const TViewModelPtr<IOutlinerExtension> LinkedOutlinerItem = SectionModel->GetLinkedOutlinerItem();
+
+				if (const TSharedPtr<ISequencerSection> SectionInterface = SectionModel->GetSectionInterface())
+				{
+					const TSharedPtr<IObjectBindingExtension> ObjectBinding = SectionModel->FindAncestorOfType<IObjectBindingExtension>();
+					SectionInterface->BuildSectionSidebarMenu(MenuBuilder, ObjectBinding.IsValid() ? ObjectBinding->GetObjectGuid() : FGuid());
+				}
+			}
+
+			AllSectionObjects.Add(SectionModel->GetSection());
+		}
+	}
+
+	if (!AllSectionObjects.IsEmpty())
+	{
+		SequencerHelpers::BuildEditSectionMenu(InSequencer, AllSectionObjects, MenuBuilder, false);
+	}
+}
+
+void FSequencerSelectionDrawer::BuildOutlinerDetails(FSequencer& InSequencer, const TSharedRef<FSequencerSelection>& InSelection, FMenuBuilder& MenuBuilder)
+{
+	if (InSelection->Outliner.Num() == 0)
+	{
+		return;
+	}
+
+	TSet<TViewModelPtr<FChannelGroupOutlinerModel>> ChannelGroups;
+
+	for (const FViewModelPtr OutlinerItem : InSelection->Outliner)
+	{
+		if (const TViewModelPtr<FOutlinerItemModel> OutlinerItemModel = OutlinerItem.ImplicitCast())
+		{
+			OutlinerItemModel->BuildSidebarMenu(MenuBuilder);
+		}
+		// Ex. "Location.X", "Rotation.Roll", "Color.R", etc.
+        else if (const TViewModelPtr<FChannelGroupOutlinerModel> ChannelGroupOutlinerModel = OutlinerItem.ImplicitCast())
+        {
+        	ChannelGroupOutlinerModel->BuildSidebarMenu(MenuBuilder);
+        	ChannelGroups.Add(ChannelGroupOutlinerModel);
+        }
+	}
+
+	if (!ChannelGroups.IsEmpty())
+	{
+		const ISequencerModule& SequencerModule = FModuleManager::Get().LoadModuleChecked<ISequencerModule>(TEXT("Sequencer"));
+		const TSharedPtr<FExtensibilityManager> SidebarExtensibilityManager = SequencerModule.GetSidebarExtensibilityManager();
+		const TSharedPtr<FExtender> Extender = SidebarExtensibilityManager->GetAllExtenders();
+
+		TArray<FName> ChannelTypeNames;
+		TArray<ISequencerChannelInterface*> ChannelInterfaces;
+		TArray<FMovieSceneChannelHandle> ChannelHandles;
+		TArray<UMovieSceneSection*> SceneSections;
+
+		for (const TViewModelPtr<FChannelGroupOutlinerModel>& ChannelModel : ChannelGroups)
+		{
+			for (const TSharedRef<IKeyArea>& KeyArea : ChannelModel->GetAllKeyAreas())
+			{
+				if (ISequencerChannelInterface* const SequencerChannelIterface = KeyArea->FindChannelEditorInterface())
+				{
+					const FMovieSceneChannelHandle& Channel = KeyArea->GetChannel();
+
+					ChannelTypeNames.Add(Channel.GetChannelTypeName());
+					ChannelInterfaces.Add(SequencerChannelIterface);
+					ChannelHandles.Add(KeyArea->GetChannel());
+					SceneSections.Add(KeyArea->GetOwningSection());
+				}
+			}
+		}
+
+		// Need to make sure all channels are the same type to allow editing of multiple channels as one
+		bool bAllChannelNamesEqual = true;
+		if (!ChannelTypeNames.IsEmpty())
+		{
+			for (int32 Index = 0; Index < ChannelTypeNames.Num(); ++Index)
+			{
+				if (Index > 0 && ChannelTypeNames[Index] != ChannelTypeNames[0])
+				{
+					bAllChannelNamesEqual = false;
+					break;
+				}
+			}
+		}
+
+		// Channel Interface Extensions (Perlin Noise, Easing, Wave)
+		if (ChannelInterfaces.Num() > 0)
+		{
+			if (bAllChannelNamesEqual)
+			{
+				ChannelInterfaces[0]->ExtendSidebarMenu_Raw(MenuBuilder, Extender, ChannelHandles, SceneSections, WeakSequencer);
+			}
+			else
+			{
+				// Display different channels separately and don't allow to edit "all-in-one"
+				for (int32 Index = 0; Index < ChannelInterfaces.Num(); ++Index)
+				{
+					ChannelInterfaces[Index]->ExtendSidebarMenu_Raw(MenuBuilder, Extender, { ChannelHandles[Index] }, { SceneSections[Index] }, WeakSequencer);
+				}
+			}
+		}
+
+		// Curve Channel Options (Pre-Finity, Post-Finity, etc.)
+		CurveChannelExtension = MakeShared<FCurveChannelSectionSidebarExtension>(WeakSequencer);
+		CurveChannelExtension->AddSections(SceneSections);
+		CurveChannelExtension->ExtendMenu(MenuBuilder);
+	}
+}
+
+void FSequencerSelectionDrawer::BuildMarkedFrameDetails(const TSharedRef<FSequencerSelection>& InSelection, FMenuBuilder& MenuBuilder)
+{
+	if (InSelection->MarkedFrames.Num() == 0)
+	{
+		return;
+	}
+
+	MenuBuilder.BeginSection(TEXT("MarkedFrames"), LOCTEXT("MarkedFramesMenuSection", "Marked Frames"));
+	
+	for (const int32 MarkIndex : InSelection->MarkedFrames)
+	{
+		MenuBuilder.AddWidget(CreateMarkedFrameDetails(MarkIndex).ToSharedRef(), FText::GetEmpty(), /*bInNoIndent=*/true);
+	}
+
+	MenuBuilder.EndSection();
+}
+
+TSharedRef<SWidget> FSequencerSelectionDrawer::CreateHintText(const FText& InMessage)
+{
+	return SNew(SBox)
+		.HAlign(HAlign_Center)
+		.Padding(2.f, 12.f, 2.f, 12.f)
+		[
+			SNew(STextBlock)
+			.Text(InMessage)
+			.TextStyle(FAppStyle::Get(), "HintText")
+		];
+}
+
+TSharedRef<SWidget> FSequencerSelectionDrawer::CreateNoSelectionHintText()
+{
+	return CreateHintText(LOCTEXT("NoSelection", "Select an object to view details."));
+}
+
+FKeyEditData FSequencerSelectionDrawer::GetKeyEditData() const
+{
+	const TSharedPtr<ISequencer> Sequencer = WeakSequencer.Pin();
+	if (!Sequencer.IsValid())
+	{
+		return FKeyEditData();
+	}
+
+	const TSharedPtr<FSequencerSelection> SequencerSelection = Private::GetSelection(*Sequencer.Get());
+	if (!SequencerSelection.IsValid())
+	{
+		return FKeyEditData();
+	}
+
+	return Private::GetKeyEditData(SequencerSelection->KeySelection);
+}
+
+TSharedPtr<SWidget> FSequencerSelectionDrawer::CreateKeyFrameDetails(const TSharedRef<FSequencerSelection>& InSequencerSelection)
+{
+	const TSharedPtr<ISequencer> Sequencer = WeakSequencer.Pin();
+	if (!Sequencer.IsValid())
+	{
+		return nullptr;
+	}
+
+	const FKeyEditData KeyEditData = Private::GetKeyEditData(InSequencerSelection->KeySelection);
+	if (KeyEditData.KeyStruct.IsValid())
+	{
+		return SNew(SKeyEditInterface, Sequencer.ToSharedRef())
+			.EditData(this, &FSequencerSelectionDrawer::GetKeyEditData);
+	}
+
+	return CreateHintText(LOCTEXT("InvalidKeyCombination", "Selected keys must belong to the same section."));
+}
+
+TSharedPtr<SWidget> FSequencerSelectionDrawer::CreateMarkedFrameDetails(const int32 InMarkedFrameIndex)
+{
+	const TSharedPtr<FSequencer> Sequencer = WeakSequencer.Pin();
+	if (!Sequencer.IsValid())
+	{
+		return nullptr;
+	}
+
+	const UMovieSceneSequence* const FocusedMovieSceneSequence = Sequencer->GetFocusedMovieSceneSequence();
+	if (!IsValid(FocusedMovieSceneSequence))
+	{
+		return nullptr;
+	}
+
+	UMovieScene* const FocusedMovieScene = FocusedMovieSceneSequence->GetMovieScene();
+	if (!IsValid(FocusedMovieScene))
+	{
+		return nullptr;
+	}
+
+	if (FocusedMovieScene->GetMarkedFrames().Num() == 0)
+	{
+		return nullptr;
+	}
+
+	class SMarkedFramePropertyWidget : public SCompoundWidget, public FNotifyHook
+	{
+	public:
+		SLATE_BEGIN_ARGS(SMarkedFramePropertyWidget) {}
+		SLATE_END_ARGS()
+
+		void Construct(const FArguments& InArgs, UMovieScene* const InMovieScene, const int32 InMarkedFrameIndex, const TWeakPtr<FSequencer>& InWeakSequencer)
+		{
+			MovieSceneToModify = InMovieScene;
+			WeakSequencer = InWeakSequencer;
+
+			FDetailsViewArgs DetailsViewArgs;
+			DetailsViewArgs.bAllowSearch = false;
+			DetailsViewArgs.bShowScrollBar = false;
+			DetailsViewArgs.NameAreaSettings = FDetailsViewArgs::HideNameArea;
+			DetailsViewArgs.NotifyHook = this;
+
+			FStructureDetailsViewArgs StructureDetailsViewArgs;
+			StructureDetailsViewArgs.bShowObjects = true;
+			StructureDetailsViewArgs.bShowAssets = true;
+			StructureDetailsViewArgs.bShowClasses = true;
+			StructureDetailsViewArgs.bShowInterfaces = true;
+			
+			const TSharedPtr<FStructOnScope> StructOnScope = MakeShared<FStructOnScope>(FMovieSceneMarkedFrame::StaticStruct(), (uint8*)&InMovieScene->GetMarkedFrames()[InMarkedFrameIndex]);
+
+			FPropertyEditorModule& PropertyEditorModule = FModuleManager::GetModuleChecked<FPropertyEditorModule>(TEXT("PropertyEditor"));
+
+			DetailsView = PropertyEditorModule.CreateStructureDetailView(DetailsViewArgs, StructureDetailsViewArgs, nullptr);
+			DetailsView->GetDetailsView()->RegisterInstancedCustomPropertyTypeLayout(TEXT("FrameNumber"), FOnGetPropertyTypeCustomizationInstance::CreateLambda([this]() {
+				return MakeShared<FFrameNumberDetailsCustomization>(WeakSequencer.Pin()->GetNumericTypeInterface()); }));
+			DetailsView->SetStructureData(StructOnScope);
+			
+			ChildSlot
+			[
+				DetailsView->GetWidget().ToSharedRef()
+			];
+		}
+
+		virtual void NotifyPreChange(FProperty* InPropertyAboutToChange) override
+		{
+			MovieSceneToModify->Modify();
+		}
+
+		virtual void NotifyPreChange(FEditPropertyChain* InPropertyAboutToChange) override
+		{
+			MovieSceneToModify->Modify();
+		}
+
+	private:
+		TObjectPtr<UMovieScene> MovieSceneToModify;
+		TSharedPtr<IStructureDetailsView> DetailsView;
+		TWeakPtr<FSequencer> WeakSequencer;
+	};
+
+	auto AreMarkedFramesLocked = [&Sequencer]() -> bool
+	{
+		if (Sequencer->IsReadOnly())
+		{
+			return true;
+		}
+
+		const UMovieSceneSequence* const FocusedMovieSceneSequence = Sequencer->GetFocusedMovieSceneSequence();
+		if (FocusedMovieSceneSequence != nullptr)
+		{
+			const UMovieScene* const MovieScene = FocusedMovieSceneSequence->GetMovieScene();
+			if (MovieScene->IsReadOnly())
+			{
+				return true;
+			}
+			return MovieScene->AreMarkedFramesLocked();
+		}
+
+		return false;
+	};
+
+	const TSharedRef<SMarkedFramePropertyWidget> Widget = SNew(SMarkedFramePropertyWidget, FocusedMovieScene, InMarkedFrameIndex, WeakSequencer);
+	Widget->SetEnabled(!AreMarkedFramesLocked());
+
+	return Widget;
+}
+
+#undef LOCTEXT_NAMESPACE
diff --git a/Engine/Source/Editor/Sequencer/Private/Widgets/Sidebar/SequencerSelectionDrawer.h b/Engine/Source/Editor/Sequencer/Private/Widgets/Sidebar/SequencerSelectionDrawer.h
--- a/Engine/Source/Editor/Sequencer/Private/Widgets/Sidebar/SequencerSelectionDrawer.h
+++ b/Engine/Source/Editor/Sequencer/Private/Widgets/Sidebar/SequencerSelectionDrawer.h
@@ -28,7 +28,7 @@ class FSequencerSelectionDrawer : public ISidebarDrawerContent
 public:
 	static const FName UniqueId;
 
-	FSequencerSelectionDrawer(const TWeakPtr<FSequencer>& InSequencerWeak);
+	FSequencerSelectionDrawer(const TWeakPtr<FSequencer>& InWeakSequencer);
 
 	//~ Begin ISidebarDrawerContent
 	virtual FName GetUniqueId() const override;
@@ -41,7 +41,6 @@ protected:
 	void OnSequencerSelectionChanged();
 
 	void BuildKeySelectionDetails(const TSharedRef<UE::Sequencer::FSequencerSelection>& InSelection, FMenuBuilder& MenuBuilder);
-	void BuildBindingPropertiesDetails(FSequencer& InSequencer, const TSharedRef<UE::Sequencer::FSequencerSelection>& InSelection, FMenuBuilder& MenuBuilder);
 	void BuildTrackAreaDetails(FSequencer& InSequencer, const TSharedRef<UE::Sequencer::FSequencerSelection>& InSelection, FMenuBuilder& MenuBuilder);
 	void BuildOutlinerDetails(FSequencer& InSequencer, const TSharedRef<UE::Sequencer::FSequencerSelection>& InSelection, FMenuBuilder& MenuBuilder);
 	void BuildMarkedFrameDetails(const TSharedRef<UE::Sequencer::FSequencerSelection>& InSelection, FMenuBuilder& MenuBuilder);
@@ -54,9 +53,11 @@ protected:
 	TSharedPtr<SWidget> CreateKeyFrameDetails(const TSharedRef<UE::Sequencer::FSequencerSelection>& InSequencerSelection);
 	TSharedPtr<SWidget> CreateMarkedFrameDetails(const int32 InMarkedFrameIndex);
 
-	TWeakPtr<FSequencer> SequencerWeak;
+	TWeakPtr<FSequencer> WeakSequencer;
 
 	TSharedPtr<SVerticalBox> ContentBox;
 
 	TSharedPtr<FCurveChannelSectionSidebarExtension> CurveChannelExtension;
+
+	TSharedPtr<class FCurveChannelSectionMenuExtension> CurveChannelSectionExtension;
 };
diff --git a/Engine/Source/Editor/Sequencer/Public/ISequencerSection.h b/Engine/Source/Editor/Sequencer/Public/ISequencerSection.h
--- a/Engine/Source/Editor/Sequencer/Public/ISequencerSection.h
+++ b/Engine/Source/Editor/Sequencer/Public/ISequencerSection.h
@@ -241,7 +241,10 @@ public:
 	 * @param MenuBuilder	The menu builder to change
 	 * @param ObjectBinding The object guid bound to this section
 	 */
-	virtual void BuildSectionSidebarMenu(FMenuBuilder& MenuBuilder, const FGuid& ObjectBinding) {}
+	virtual void BuildSectionSidebarMenu(FMenuBuilder& MenuBuilder, const FGuid& ObjectBinding)
+	{
+		BuildSectionContextMenu(MenuBuilder, ObjectBinding);
+	}
 
 	/**
 	 * Called when the user requests that a category from this section be deleted. 
diff --git a/Engine/Source/Editor/Sequencer/Public/MVVM/ViewModels/CategoryModel.h b/Engine/Source/Editor/Sequencer/Public/MVVM/ViewModels/CategoryModel.h
--- a/Engine/Source/Editor/Sequencer/Public/MVVM/ViewModels/CategoryModel.h
+++ b/Engine/Source/Editor/Sequencer/Public/MVVM/ViewModels/CategoryModel.h
@@ -134,6 +134,9 @@ public:
 	/*~ IRecyclableExtension */
 	void OnRecycle() override;
 
+	/*~ FOutlinerItemModelMixin */
+	virtual void BuildSidebarMenu(FMenuBuilder& MenuBuilder) override;
+
 private:
 
 	TArray<TWeakViewModelPtr<FCategoryModel>> Categories;
diff --git a/Engine/Source/Editor/Sequencer/Public/MVVM/ViewModels/OutlinerItemModel.h b/Engine/Source/Editor/Sequencer/Public/MVVM/ViewModels/OutlinerItemModel.h
--- a/Engine/Source/Editor/Sequencer/Public/MVVM/ViewModels/OutlinerItemModel.h
+++ b/Engine/Source/Editor/Sequencer/Public/MVVM/ViewModels/OutlinerItemModel.h
@@ -70,16 +70,13 @@ public:
 	/*~ IDimmableExtension */
 	bool IsDimmed() const override;
 
-	virtual void BuildSidebarMenu(FMenuBuilder& MenuBuilder);
-
-//protected:
-
 	/** Get context menu contents. */
 	virtual void BuildContextMenu(FMenuBuilder& MenuBuilder);
-	virtual void BuildOrganizeContextMenu(FMenuBuilder& MenuBuilder);
-	virtual void BuildDisplayOptionsMenu(FMenuBuilder& MenuBuilder);
-	virtual void BuildTrackOptionsMenu(FMenuBuilder& MenuBuilder);
+
+	virtual void BuildSidebarMenu(FMenuBuilder& MenuBuilder);
+
 protected:
+
 	/** Set identifier for computing node paths */
 	void SetIdentifier(FName InNewIdentifier);
 
@@ -89,7 +86,10 @@ protected:
 	/** Set expansion state without saving it in the movie-scene data */
 	void SetExpansionWithoutSaving(bool bInIsExpanded);
 
-	void BuildSectionColorTintsContextMenu(FMenuBuilder& MenuBuilder);
+	virtual void BuildSectionColorTintsMenu(FMenuBuilder& MenuBuilder);
+	virtual void BuildOrganizeContextMenu(FMenuBuilder& MenuBuilder);
+	virtual void BuildDisplayOptionsMenu(FMenuBuilder& MenuBuilder);
+	virtual void BuildTrackOptionsMenu(FMenuBuilder& MenuBuilder, const TArray<UMovieSceneTrack*>& InTracks);
 
 private:
 
diff --git a/Engine/Source/Editor/Sequencer/Public/MVVM/ViewModels/TrackModel.h b/Engine/Source/Editor/Sequencer/Public/MVVM/ViewModels/TrackModel.h
--- a/Engine/Source/Editor/Sequencer/Public/MVVM/ViewModels/TrackModel.h
+++ b/Engine/Source/Editor/Sequencer/Public/MVVM/ViewModels/TrackModel.h
@@ -144,16 +144,12 @@ public:
 	/*~ FViewModel interface */
 	virtual void OnConstruct() override;
 
-	virtual void BuildBlendingMenu(FMenuBuilder& MenuBuilder);
-
 private:
 
 	void ForceUpdate();
 
 	bool FindBoundObjects(TArray<UObject*>& OutBoundObjects) const;
 
-	TArray<TWeakObjectPtr<UObject>> GetSelectedTrackSections() const;
-
 	/** A second children list for the sections inside this track */
 	FViewModelListHead SectionList;
 	FViewModelListHead TopLevelChannelList;
diff --git a/Engine/Source/Editor/Sequencer/Public/MVVM/ViewModels/TrackRowModel.h b/Engine/Source/Editor/Sequencer/Public/MVVM/ViewModels/TrackRowModel.h
--- a/Engine/Source/Editor/Sequencer/Public/MVVM/ViewModels/TrackRowModel.h
+++ b/Engine/Source/Editor/Sequencer/Public/MVVM/ViewModels/TrackRowModel.h
@@ -47,8 +47,8 @@ public:
 
 	/*~ FOutlinerItemModel */
 	void BuildContextMenu(FMenuBuilder& MenuBuilder) override;
-	void CreateCurveModels(TArray<TUniquePtr<FCurveModel>>& OutCurveModels) override;
 	void BuildSidebarMenu(FMenuBuilder& MenuBuilder) override;
+	void CreateCurveModels(TArray<TUniquePtr<FCurveModel>>& OutCurveModels) override;
 
 	/*~ IOutlinerExtension */
 	FOutlinerSizing GetOutlinerSizing() const override;
@@ -91,8 +91,6 @@ public:
 
 private:
 
-	TArray<TWeakObjectPtr<UObject>> GetSelectedTrackSections() const;
-	
 	FViewModelListHead SectionList;
 	FViewModelListHead TopLevelChannelList;
 	TWeakObjectPtr<UMovieSceneTrack> WeakTrack;
diff --git a/Engine/Source/Editor/Sequencer/Public/SequencerUtilities.h b/Engine/Source/Editor/Sequencer/Public/SequencerUtilities.h
--- a/Engine/Source/Editor/Sequencer/Public/SequencerUtilities.h
+++ b/Engine/Source/Editor/Sequencer/Public/SequencerUtilities.h
@@ -294,4 +294,9 @@ struct SEQUENCER_API FSequencerUtilities
 	static void AddChangeClassMenu(FMenuBuilder& MenuBuilder, TSharedRef<ISequencer> Sequencer, const TArray<FSequencerChangeBindingInfo>& BindingsToConvert, TFunction<void()> OnBindingChanged);
 	static void HandleTemplateActorClassPicked(UClass* ChosenClass, TSharedRef<ISequencer> Sequencer, const TArray<FSequencerChangeBindingInfo>& BindingsToConvert, TFunction<void()> OnBindingChanged);
 	static void AddConvertBindingMenu(FMenuBuilder& MenuBuilder, TSharedRef<ISequencer> Sequencer, const TArray<FSequencerChangeBindingInfo>& BindingsToConvert, TFunction<void()> OnBindingChanged);
+
+	static void ConvertBindings(const TSharedRef<ISequencer>& InSequencer
+		, const TArray<FSequencerChangeBindingInfo>& InBindingsToConvert
+		, TFunction<FMovieScenePossessable*(FGuid, int32)> InDoConvert
+		, TFunction<void()> InOnBindingChanged);
 };
