[Sequencer] Sidebar details:
 - Add back channel interface extensions
 - Small code cleanups

diff --git a/Engine/Plugins/Experimental/PropertyAnimator/Source/PropertyAnimatorEditor/Private/Sequencer/PropertyAnimatorEditorCurveChannelInterface.h b/Engine/Plugins/Experimental/PropertyAnimator/Source/PropertyAnimatorEditor/Private/Sequencer/PropertyAnimatorEditorCurveChannelInterface.h
--- a/Engine/Plugins/Experimental/PropertyAnimator/Source/PropertyAnimatorEditor/Private/Sequencer/PropertyAnimatorEditorCurveChannelInterface.h
+++ b/Engine/Plugins/Experimental/PropertyAnimator/Source/PropertyAnimatorEditor/Private/Sequencer/PropertyAnimatorEditorCurveChannelInterface.h
@@ -42,7 +42,7 @@ struct TPropertyAnimatorEditorCurveChannelInterface : ISequencerChannelInterface
 };
 
 template<typename InChannelType, typename InMenuExtensionType>
-void TPropertyAnimatorEditorCurveChannelInterface<InChannelType, InMenuExtensionType>::ExtendSectionMenu_Raw(FMenuBuilder& InMenuBuilder
+void TPropertyAnimatorEditorCurveChannelInterface<InChannelType, InMenuExtensionType>::ExtendSectionMenu_Raw(FMenuBuilder& MenuBuilder
 	, TSharedPtr<FExtender> InMenuExtender
 	, TConstArrayView<FMovieSceneChannelHandle> InChannels
 	, TConstArrayView<UMovieSceneSection*> InSections
@@ -58,19 +58,14 @@ void TPropertyAnimatorEditorCurveChannelInterface<InChannelType, InMenuExtension
 }
 
 template<typename InChannelType, typename InMenuExtensionType>
-void TPropertyAnimatorEditorCurveChannelInterface<InChannelType, InMenuExtensionType>::ExtendSidebarMenu_Raw(FMenuBuilder& InMenuBuilder
+void TPropertyAnimatorEditorCurveChannelInterface<InChannelType, InMenuExtensionType>::ExtendSidebarMenu_Raw(FMenuBuilder& MenuBuilder
 	, TSharedPtr<FExtender> InMenuExtender
 	, TConstArrayView<FMovieSceneChannelHandle> InChannels
 	, TConstArrayView<UMovieSceneSection*> InSections
 	, TWeakPtr<ISequencer> InSequencer) const
 {
 	TSharedRef<FPropertyAnimatorEditorCurveSectionMenuExtension> Extension = MakeShared<FMenuExtensionType>(InChannels, InSections);
-
-	InMenuExtender->AddMenuExtension(TEXT("SequencerChannels"), EExtensionHook::First, nullptr
-		, FMenuExtensionDelegate::CreateLambda([Extension](FMenuBuilder& InInnerMenuBuilder)
-			{
-				Extension->ExtendMenu(InInnerMenuBuilder);
-			}));
+	Extension->ExtendMenu(MenuBuilder, false);
 }
 
 template<typename InChannelType, typename InMenuExtensionType>
diff --git a/Engine/Source/Editor/MovieSceneTools/Public/Channels/PerlinNoiseChannelInterface.h b/Engine/Source/Editor/MovieSceneTools/Public/Channels/PerlinNoiseChannelInterface.h
--- a/Engine/Source/Editor/MovieSceneTools/Public/Channels/PerlinNoiseChannelInterface.h
+++ b/Engine/Source/Editor/MovieSceneTools/Public/Channels/PerlinNoiseChannelInterface.h
@@ -138,12 +138,7 @@ struct TPerlinNoiseChannelInterface : ISequencerChannelInterface
 	virtual void ExtendSidebarMenu_Raw(FMenuBuilder& MenuBuilder, TSharedPtr<FExtender> MenuExtender, TArrayView<const FMovieSceneChannelHandle> Channels, TArrayView<UMovieSceneSection* const> Sections, TWeakPtr<ISequencer> InSequencer) const override
 	{
 		TSharedRef<FPerlinNoiseChannelSectionMenuExtension> Extension = MakeShared<FPerlinNoiseChannelSectionMenuExtension>(Channels, Sections);
-
-		MenuExtender->AddMenuExtension(TEXT("SequencerChannels"), EExtensionHook::First, nullptr
-			, FMenuExtensionDelegate::CreateLambda([Extension](FMenuBuilder& MenuBuilder)
-				{
-					Extension->ExtendMenu(MenuBuilder);
-				}));
+		Extension->ExtendMenu(MenuBuilder, false);
 	}
 
 	virtual void DrawKeys_Raw(FMovieSceneChannel* Channel, TArrayView<const FKeyHandle> InKeyHandles, const UMovieSceneSection* InOwner, TArrayView<FKeyDrawParams> OutKeyDrawParams) const override
diff --git a/Engine/Source/Editor/Sequencer/Private/SequencerCommonHelpers.cpp b/Engine/Source/Editor/Sequencer/Private/SequencerCommonHelpers.cpp
--- a/Engine/Source/Editor/Sequencer/Private/SequencerCommonHelpers.cpp
+++ b/Engine/Source/Editor/Sequencer/Private/SequencerCommonHelpers.cpp
@@ -344,8 +344,16 @@ void SequencerHelpers::BuildEditSectionMenu(FSequencer& Sequencer
 		return;
 	}
 
-	auto BuildSection = [&Sequencer, InWeakSections](FMenuBuilder& LambdaMenuBuilder)
+	const TWeakPtr<FSequencer> WeakSequencer = StaticCastWeakPtr<FSequencer>(Sequencer.AsWeak());
+
+	auto BuildSection = [WeakSequencer, InWeakSections](FMenuBuilder& LambdaMenuBuilder)
 	{
+		const TSharedPtr<FSequencer> Sequencer = WeakSequencer.Pin();
+		if (!Sequencer)
+		{
+			return;
+		}
+
 		TSharedRef<SSectionDetailsNotifyHookWrapper> DetailsNotifyWrapper = SNew(SSectionDetailsNotifyHookWrapper);
 		FDetailsViewArgs DetailsViewArgs;
 		{
@@ -365,17 +373,21 @@ void SequencerHelpers::BuildEditSectionMenu(FSequencer& Sequencer
 
 		TSharedRef<IDetailsView> DetailsView = FModuleManager::GetModuleChecked<FPropertyEditorModule>("PropertyEditor").CreateDetailView(DetailsViewArgs);
 		DetailsView->RegisterInstancedCustomPropertyTypeLayout("FrameNumber",
-			FOnGetPropertyTypeCustomizationInstance::CreateLambda([&Sequencer]()
+			FOnGetPropertyTypeCustomizationInstance::CreateLambda([WeakSequencer]()
 			{
-				return MakeShared<FFrameNumberDetailsCustomization>(Sequencer.GetNumericTypeInterface());
+				const TSharedPtr<FSequencer> Sequencer = WeakSequencer.Pin();
+				check(Sequencer);
+				return MakeShared<FFrameNumberDetailsCustomization>(Sequencer->GetNumericTypeInterface());
 			}));
 		DetailsView->RegisterInstancedCustomPropertyLayout(UMovieSceneSection::StaticClass(),
-			FOnGetDetailCustomizationInstance::CreateLambda([&Sequencer]()
+			FOnGetDetailCustomizationInstance::CreateLambda([WeakSequencer]()
 			{
+				const TSharedPtr<FSequencer> Sequencer = WeakSequencer.Pin();
+				check(Sequencer);
 				// We pass the current scene to the UMovieSceneSection customization so we can get the overall bounds of the section when we change a section from infinite->bounded.
-				UMovieScene* const CurrentScene = Sequencer.GetFocusedMovieSceneSequence()->GetMovieScene();
+				UMovieScene* const CurrentScene = Sequencer->GetFocusedMovieSceneSequence()->GetMovieScene();
 				check(CurrentScene);
-				return MakeShared<FMovieSceneSectionDetailsCustomization>(Sequencer.GetNumericTypeInterface(), CurrentScene);
+				return MakeShared<FMovieSceneSectionDetailsCustomization>(Sequencer->GetNumericTypeInterface(), CurrentScene);
 			}));
 	
 		DetailsView->SetIsPropertyVisibleDelegate(FIsPropertyVisible::CreateLambda([](const FPropertyAndParent& PropertyAndParent)
@@ -385,7 +397,7 @@ void SequencerHelpers::BuildEditSectionMenu(FSequencer& Sequencer
 		);
 
 		// Let section interfaces further customize the properties details view.
-		TSharedRef<FSequencerNodeTree> SequencerNodeTree = Sequencer.GetNodeTree();
+		TSharedRef<FSequencerNodeTree> SequencerNodeTree = Sequencer->GetNodeTree();
 		for (TWeakObjectPtr<> Section : InWeakSections)
 		{
 			if (Section.IsValid())
@@ -395,7 +407,7 @@ void SequencerHelpers::BuildEditSectionMenu(FSequencer& Sequencer
 				{
 					TSharedPtr<ISequencerSection> SectionInterface = SectionHandle->GetSectionInterface();
 					FSequencerSectionPropertyDetailsViewCustomizationParams CustomizationDetails(
-						SectionInterface.ToSharedRef(), Sequencer.AsWeak(), *SectionHandle->GetParentTrackExtension()->GetTrackEditor().Get());
+						SectionInterface.ToSharedRef(), WeakSequencer, *SectionHandle->GetParentTrackExtension()->GetTrackEditor().Get());
 					TSharedPtr<FObjectBindingModel> ParentObjectBindingNode = SectionHandle->FindAncestorOfType<FObjectBindingModel>();
 					if (ParentObjectBindingNode.IsValid())
 					{
@@ -406,11 +418,11 @@ void SequencerHelpers::BuildEditSectionMenu(FSequencer& Sequencer
 			}
 		}
 
-		Sequencer.OnInitializeDetailsPanel().Broadcast(DetailsView, Sequencer.AsShared());
+		Sequencer->OnInitializeDetailsPanel().Broadcast(DetailsView, Sequencer.ToSharedRef());
 		DetailsView->SetObjects(InWeakSections);
 
-		DetailsNotifyWrapper->SetDetailsAndSequencer(DetailsView, Sequencer.AsWeak());
-		DetailsNotifyWrapper->SetEnabled(!Sequencer.IsReadOnly());
+		DetailsNotifyWrapper->SetDetailsAndSequencer(DetailsView, WeakSequencer);
+		DetailsNotifyWrapper->SetEnabled(!Sequencer->IsReadOnly());
 
 		LambdaMenuBuilder.BeginSection(TEXT("TrackSection"), NSLOCTEXT("Sequencer", "TrackSectionMenuSection", "Track Section"));
         {
diff --git a/Engine/Source/Editor/Sequencer/Private/Widgets/Sidebar/SequencerSelectionDrawer.cpp b/Engine/Source/Editor/Sequencer/Private/Widgets/Sidebar/SequencerSelectionDrawer.cpp
--- a/Engine/Source/Editor/Sequencer/Private/Widgets/Sidebar/SequencerSelectionDrawer.cpp
+++ b/Engine/Source/Editor/Sequencer/Private/Widgets/Sidebar/SequencerSelectionDrawer.cpp
@@ -237,9 +237,6 @@ void FSequencerSelectionDrawer::BuildKeySelectionDetails(const TSharedRef<FSeque
 		return;
 	}
 
-	const ISequencerModule& SequencerModule = FModuleManager::Get().LoadModuleChecked<ISequencerModule>(TEXT("Sequencer"));
-	const TSharedPtr<FExtensibilityManager> SidebarExtensibilityManager = SequencerModule.GetSidebarExtensibilityManager();
-
 	MenuBuilder.BeginSection(TEXT("KeyEdit"), LOCTEXT("KeyEditMenuSection", "Key Edit"));
 	{
 		MenuBuilder.AddWidget(CreateKeyFrameDetails(InSelection).ToSharedRef(), FText::GetEmpty(), /*bInNoIndent=*/true);
@@ -341,7 +338,7 @@ void FSequencerSelectionDrawer::BuildOutlinerDetails(FSequencer& InSequencer, co
 
 					ChannelTypeNames.Add(Channel.GetChannelTypeName());
 					ChannelInterfaces.Add(SequencerChannelIterface);
-					ChannelHandles.Add(KeyArea->GetChannel());
+					ChannelHandles.Add(Channel);
 					SceneSections.Add(KeyArea->GetOwningSection());
 				}
 			}
diff --git a/Engine/Source/Editor/Sequencer/Private/Widgets/Sidebar/SequencerSelectionDrawer.h b/Engine/Source/Editor/Sequencer/Private/Widgets/Sidebar/SequencerSelectionDrawer.h
--- a/Engine/Source/Editor/Sequencer/Private/Widgets/Sidebar/SequencerSelectionDrawer.h
+++ b/Engine/Source/Editor/Sequencer/Private/Widgets/Sidebar/SequencerSelectionDrawer.h
@@ -9,14 +9,10 @@ class FMenuBuilder;
 class FName;
 class FSequencer;
 class FText;
-class IDetailsView;
 class SVerticalBox;
 class SWidget;
-class UMovieSceneFolder;
-class UMovieSceneSequence;
 struct FCurveChannelSectionSidebarExtension;
 struct FKeyEditData;
-struct FMovieSceneMarkedFrame;
 
 namespace UE::Sequencer
 {
@@ -58,6 +54,4 @@ protected:
 	TSharedPtr<SVerticalBox> ContentBox;
 
 	TSharedPtr<FCurveChannelSectionSidebarExtension> CurveChannelExtension;
-
-	TSharedPtr<class FCurveChannelSectionMenuExtension> CurveChannelSectionExtension;
 };
