Sequencer:  Change sidebar sequencer selection handling to only happen once per tick so mutiple selection changes are grouped together in order to avoid huge lags when selecting/deselection control rig controls.  There can be multiple sequencer selection notifications in one tick and the sidebard, in particularly constructing the details panel, is slow handling selections.

diff --git a/Engine/Source/Editor/Sequencer/Private/Widgets/Sidebar/SequencerSelectionDrawer.cpp b/Engine/Source/Editor/Sequencer/Private/Widgets/Sidebar/SequencerSelectionDrawer.cpp
--- a/Engine/Source/Editor/Sequencer/Private/Widgets/Sidebar/SequencerSelectionDrawer.cpp
+++ b/Engine/Source/Editor/Sequencer/Private/Widgets/Sidebar/SequencerSelectionDrawer.cpp
@@ -153,11 +153,6 @@ void FSequencerSelectionDrawer::OnSequencerSelectionChanged()
 	{
 		return;
 	}
-
-	ContentBox->ClearChildren();
-
-	CurveChannelExtension.Reset();
-
 	const TSharedPtr<FSequencer> Sequencer = WeakSequencer.Pin();
 	if (!WeakSequencer.IsValid())
 	{
@@ -169,59 +164,88 @@ void FSequencerSelectionDrawer::OnSequencerSelectionChanged()
 	{
 		return;
 	}
+	static bool bWaitingToHandleSelectionChanged = false;
 
-	const TSharedRef<FSequencerSelection> SelectionRef = SequencerSelection.ToSharedRef();
-
-	auto AddToContent = [this](const TSharedRef<SWidget>& InWidget)
-		{
-			ContentBox->AddSlot()
-				.AutoHeight()
-				[
-					InWidget
-				];
-		};
-
-	const ISequencerModule& SequencerModule = FModuleManager::Get().LoadModuleChecked<ISequencerModule>(TEXT("Sequencer"));
-	const TSharedPtr<FExtensibilityManager> SidebarExtensibilityManager = SequencerModule.GetSidebarExtensibilityManager();
-
-	FMenuBuilder MenuBuilder(/*bInShouldCloseWindowAfterMenuSelection=*/false
-		, Sequencer->GetCommandBindings()
-		, SidebarExtensibilityManager->GetAllExtenders()
-		, /*bInCloseSelfOnly=*/true, &FCoreStyle::Get(), /*bInSearchable=*/true, TEXT("Sequencer.Sidebar"));
-
-	/**
-	 * Selection details display order preference:
-	 *  1) Key items
-	 *  2) Track area items (if no key selected)
-	 *  3) Outliner items (if no key or track area selected)
-	 *  4) Marked frames
-	 */
-
-	// 1) Key items
-	BuildKeySelectionDetails(SelectionRef, MenuBuilder);
-
-	// Early out for key selections
-	const bool bIsKeySelected = SequencerSelection->KeySelection.Num() > 0;
-	if (bIsKeySelected)
+	if (bWaitingToHandleSelectionChanged == false)
 	{
-		AddToContent(MenuBuilder.MakeWidget());
-		return;
-	}
+		bWaitingToHandleSelectionChanged = true;
+		GEditor->GetTimerManager()->SetTimerForNextTick([this]()
+		{
+			bWaitingToHandleSelectionChanged = false;
+			if (!ContentBox.IsValid())
+			{
+				return;
+			}
+			ContentBox->ClearChildren();
 
-	// 2) Track area items
-	BuildTrackAreaDetails(*Sequencer, SelectionRef, MenuBuilder);
+			CurveChannelExtension.Reset();
 
-	// 3) Outliner items
-	const bool bIsTrackAreaSelected = SequencerSelection->TrackArea.Num() > 0;
-	if (!bIsTrackAreaSelected)
-	{
-		BuildOutlinerDetails(*Sequencer, SelectionRef, MenuBuilder);
-	}
+			const TSharedPtr<FSequencer> Sequencer = WeakSequencer.Pin();
+			if (!WeakSequencer.IsValid())
+			{
+				return;
+			}
 
-	// 4) Marked frames
-	BuildMarkedFrameDetails(SelectionRef, MenuBuilder);
+			const TSharedPtr<FSequencerSelection> SequencerSelection = Private::GetSelection(*Sequencer.Get());
+			if (!SequencerSelection.IsValid())
+			{
+				return;
+			}
+
+			const TSharedRef<FSequencerSelection> SelectionRef = SequencerSelection.ToSharedRef();
+
+			auto AddToContent = [this](const TSharedRef<SWidget>& InWidget)
+				{
+					ContentBox->AddSlot()
+						.AutoHeight()
+						[
+							InWidget
+						];
+				};
+
+			const ISequencerModule& SequencerModule = FModuleManager::Get().LoadModuleChecked<ISequencerModule>(TEXT("Sequencer"));
+			const TSharedPtr<FExtensibilityManager> SidebarExtensibilityManager = SequencerModule.GetSidebarExtensibilityManager();
+
+			FMenuBuilder MenuBuilder(/*bInShouldCloseWindowAfterMenuSelection=*/false
+				, Sequencer->GetCommandBindings()
+				, SidebarExtensibilityManager->GetAllExtenders()
+				, /*bInCloseSelfOnly=*/true, &FCoreStyle::Get(), /*bInSearchable=*/true, TEXT("Sequencer.Sidebar"));
+
+			/**
+				* Selection details display order preference:
+				*  1) Key items
+				*  2) Track area items (if no key selected)
+				*  3) Outliner items (if no key or track area selected)
+				*  4) Marked frames
+				*/
+
+				// 1) Key items
+			BuildKeySelectionDetails(SelectionRef, MenuBuilder);
+
+			// Early out for key selections
+			const bool bIsKeySelected = SequencerSelection->KeySelection.Num() > 0;
+			if (bIsKeySelected)
+			{
+				AddToContent(MenuBuilder.MakeWidget());
+				return;
+			}
 
-	AddToContent(MenuBuilder.MakeWidget());
+			// 2) Track area items
+			BuildTrackAreaDetails(*Sequencer, SelectionRef, MenuBuilder);
+
+			// 3) Outliner items
+			const bool bIsTrackAreaSelected = SequencerSelection->TrackArea.Num() > 0;
+			if (!bIsTrackAreaSelected)
+			{
+				BuildOutlinerDetails(*Sequencer, SelectionRef, MenuBuilder);
+			}
+
+			// 4) Marked frames
+			BuildMarkedFrameDetails(SelectionRef, MenuBuilder);
+
+			AddToContent(MenuBuilder.MakeWidget());
+		});
+	}
 }
 
 void FSequencerSelectionDrawer::BuildKeySelectionDetails(const TSharedRef<FSequencerSelection>& InSelection, FMenuBuilder& MenuBuilder)
