[Sequencer] Remove default value to match for virtual function causing Horde warning

diff --git a/Engine/Source/Editor/Sequencer/Public/ISequencerChannelInterface.h b/Engine/Source/Editor/Sequencer/Public/ISequencerChannelInterface.h
--- a/Engine/Source/Editor/Sequencer/Public/ISequencerChannelInterface.h
+++ b/Engine/Source/Editor/Sequencer/Public/ISequencerChannelInterface.h
@@ -67,7 +67,7 @@ struct FExtendKeyMenuParams
 
 class ISidebarChannelExtension
 {
-	virtual TSharedPtr<ISidebarChannelExtension> ExtendMenu(FMenuBuilder& MenuBuilder, const bool bInSubMenu = true) = 0;
+	virtual TSharedPtr<ISidebarChannelExtension> ExtendMenu(FMenuBuilder& MenuBuilder, const bool bInSubMenu = false) = 0;
 };
 
 /**
