[ToolWidgets][Sequencer] Sidebar Details:
 - Add SSidebarContainer and move related code from SSidebar
 - Save sidebar details drawer state to config to restore when Sequencer is initialized
 - Option/action/menu item to completely disable sidebar
 - Save state per sequencer type
 - Add bSupportsSidebar to FSequencerHostCapabilities
 - Disable for Niagara Sequencers

diff --git a/Engine/Plugins/Experimental/Avalanche/Source/AvalancheSequencer/Private/AvaSequencer.cpp b/Engine/Plugins/Experimental/Avalanche/Source/AvalancheSequencer/Private/AvaSequencer.cpp
--- a/Engine/Plugins/Experimental/Avalanche/Source/AvalancheSequencer/Private/AvaSequencer.cpp
+++ b/Engine/Plugins/Experimental/Avalanche/Source/AvalancheSequencer/Private/AvaSequencer.cpp
@@ -50,6 +50,7 @@
 #include "Settings/AvaSequencerSettings.h"
 #include "Sidebar/SidebarDrawerConfig.h"
 #include "Sidebar/SSidebar.h"
+#include "Sidebar/SSidebarContainer.h"
 #include "ToolMenu.h"
 #include "ToolMenuEntry.h"
 #include "ToolMenuSection.h"
@@ -350,6 +351,7 @@ TSharedRef<ISequencer> FAvaSequencer::CreateSequencer()
 		// Host Capabilities
 		SequencerInitParams.HostCapabilities.bSupportsCurveEditor = true;
 		SequencerInitParams.HostCapabilities.bSupportsSaveMovieSceneAsset = false;
+		SequencerInitParams.HostCapabilities.bSupportsSidebar = true;
 	};
 
 	InstancedSequencer = FAvaSequencerUtils::GetSequencerModule().CreateSequencer(SequencerInitParams);
@@ -1488,57 +1490,43 @@ TSharedRef<SWidget> FAvaSequencer::CreateSequenceWidget()
 
 	TSharedRef<ISequencer> Sequencer = GetSequencer();
 
-	const TSharedRef<SOverlay> SidebarOverlay = SNew(SOverlay);
-	SequenceTreeDockLocation = SNew(SBox);
+	UAvaSequencerSettings* const SequencerSettings = GetMutableDefault<UAvaSequencerSettings>();
+	check(IsValid(SequencerSettings));
+	FSidebarState& SidebarState = SequencerSettings->GetSidebarState();
 
-	TSharedRef<SOverlay> OutWidget = SNew(SOverlay)
-		+ SOverlay::Slot()
-		[
-			SNew(SHorizontalBox)
-			+ SHorizontalBox::Slot()
-			.AutoWidth()
-			[
-				SAssignNew(LeftSidebar, SSidebar, SidebarOverlay, SequenceTreeDockLocation.ToSharedRef())
-				.TabLocation(ESidebarTabLocation::Left)
-				.OnDockStateChanged(this, &FAvaSequencer::OnSidebarDockStateChanged)
-			]
-			+ SHorizontalBox::Slot()
-			[
-				SNew(SSplitter)
-				+ SSplitter::Slot()
-				.Resizable(false)
-				.SizeRule(SSplitter::ESizeRule::SizeToContent)
-				.Value(FSidebarDrawerConfig::DefaultSizeCoefficient)
-				.OnSlotResized(this, &FAvaSequencer::OnSidebarSlotResized)
-				.Expose(SidebarSlot)
-				[
-					SequenceTreeDockLocation.ToSharedRef()
-				]
-				+ SSplitter::Slot()
-				.Value(0.8f)
-				[
-					SNew(SOverlay)
-					.AddMetaData<FTagMetaData>(FTagMetaData(TEXT("Sequencer")))
-					+ SOverlay::Slot()
-					[
-						Sequencer->GetSequencerWidget()
-					]
-				]
-			]
-		]
-		+ SOverlay::Slot()
-		[
-			SidebarOverlay
-		];
+	TSharedPtr<SWidget> OutWidget;
+
+	if (SidebarState.IsVisible())
+	{
+		SidebarContainer = SNew(SSidebarContainer);
+
+		LeftSidebar = SNew(SSidebar, SidebarContainer.ToSharedRef())
+			.TabLocation(ESidebarTabLocation::Left)
+			.InitialDrawerSize(SidebarState.GetDrawerSize())
+			.OnStateChanged(this, &FAvaSequencer::OnSidebarStateChanged)
+			.OnGetContent(FOnGetContent::CreateLambda([this]()
+				{
+					return GetSequencer()->GetSequencerWidget();
+				}));
+
+		SidebarContainer->RebuildSidebar(LeftSidebar.ToSharedRef(), SidebarState);
+
+		OutWidget = SidebarContainer.ToSharedRef();
+	}
+	else
+	{
+		OutWidget = Sequencer->GetSequencerWidget();
+	}
 
 	FSidebarDrawerConfig SequenceTreeDrawerConfig;
 	SequenceTreeDrawerConfig.UniqueId = SidebarDrawerId;
 	SequenceTreeDrawerConfig.ButtonText = LOCTEXT("SequenceLabel", "Sequence");
 	SequenceTreeDrawerConfig.ToolTipText = LOCTEXT("SequenceTooltip", "Open the Sequence options panel");
 	SequenceTreeDrawerConfig.Icon = FAppStyle::GetBrush(TEXT("Persona.EditInSequencer"));
-	SequenceTreeDrawerConfig.bInitiallyDocked = true;
+	SequenceTreeDrawerConfig.InitialState = SidebarState.FindOrAddDrawerState(SidebarDrawerId);
+
 	LeftSidebar->RegisterDrawer(MoveTemp(SequenceTreeDrawerConfig));
-	
+
 	// Make sure the sequence tree widget is created
 	GetSequenceTreeWidget();
 
@@ -1547,7 +1535,7 @@ TSharedRef<SWidget> FAvaSequencer::CreateSequenceWidget()
 	LeftSidebar->RegisterDrawerSection(SidebarDrawerId, MakeShared<FAvaSequencePlaybackDetails>(ThisSequencerRef));
 	LeftSidebar->RegisterDrawerSection(SidebarDrawerId, MakeShared<FAvaSequenceSettingsDetails>(ThisSequencerRef));
 
-	return OutWidget;
+	return OutWidget.ToSharedRef();
 }
 
 void FAvaSequencer::OnActorsCopied(FString& InOutCopiedData, TConstArrayView<AActor*> InCopiedActors)
@@ -1804,16 +1792,12 @@ TSharedRef<FAvaEaseCurveTool> FAvaSequencer::GetEaseCurveTool() const
 	return EaseCurveTool.ToSharedRef(); 
 }
 
-void FAvaSequencer::OnSidebarDockStateChanged(const FName InDrawerId)
-{
-	LeftSidebar->UpdateDockedSplitterSlot(SidebarDrawerId, SidebarSlot);
-}
-
-void FAvaSequencer::OnSidebarSlotResized(const float InFillCoefficient)
+void FAvaSequencer::OnSidebarStateChanged(const FSidebarState& InNewState)
 {
-	if (InFillCoefficient == 0.f && LeftSidebar->HasDrawerDocked())
+	UAvaSequencerSettings* const SequencerSettings = GetMutableDefault<UAvaSequencerSettings>();
+	if (IsValid(SequencerSettings))
 	{
-		LeftSidebar->UndockAllDrawers();
+		SequencerSettings->SetSidebarState(InNewState);
 	}
 }
 
diff --git a/Engine/Plugins/Experimental/Avalanche/Source/AvalancheSequencer/Private/AvaSequencer.h b/Engine/Plugins/Experimental/Avalanche/Source/AvalancheSequencer/Private/AvaSequencer.h
--- a/Engine/Plugins/Experimental/Avalanche/Source/AvalancheSequencer/Private/AvaSequencer.h
+++ b/Engine/Plugins/Experimental/Avalanche/Source/AvalancheSequencer/Private/AvaSequencer.h
@@ -29,6 +29,7 @@ class SAvaSequenceTree;
 class SBox;
 class SHeaderRow;
 class SSidebar;
+class SSidebarContainer;
 class SWidget;
 class UAvaSequence;
 class UAvaSequencerSettings;
@@ -43,6 +44,7 @@ struct FAvaSequencerArgs;
 struct FMovieSceneBinding;
 struct FMovieScenePossessable;
 struct FSequencerInitParams;
+struct FSidebarState;
 struct FToolMenuContext;
 template<typename ItemType> class STreeView;
 
@@ -221,8 +223,7 @@ private:
 
 	void OnUpdateCameraCut(UObject* InCameraObject, bool bInJumpCut);
 
-	void OnSidebarDockStateChanged(const FName InDrawerId);
-	void OnSidebarSlotResized(const float InFillCoefficient);
+	void OnSidebarStateChanged(const FSidebarState& InNewState);
 
 	void ExtendSidebarSelectionMenu(FMenuBuilder& OutMenuBuilder);
 	void ExtendSidebarMarkedFramesMenu(FMenuBuilder& OutMenuBuilder);
@@ -278,9 +279,8 @@ private:
 
 	TSharedPtr<FAvaEaseCurveTool> EaseCurveTool;
 
+	TSharedPtr<SSidebarContainer> SidebarContainer;
 	TSharedPtr<SSidebar> LeftSidebar;
-	TSharedPtr<SBox> SequenceTreeDockLocation;
-	SSplitter::FSlot* SidebarSlot = nullptr;
 
 	TSharedPtr<FExtender> SidebarExtender;
 	FDelegateHandle SidebarSelectionExtenderHandle;
diff --git a/Engine/Plugins/Experimental/Avalanche/Source/AvalancheSequencer/Private/Settings/AvaSequencerSettings.h b/Engine/Plugins/Experimental/Avalanche/Source/AvalancheSequencer/Private/Settings/AvaSequencerSettings.h
--- a/Engine/Plugins/Experimental/Avalanche/Source/AvalancheSequencer/Private/Settings/AvaSequencerSettings.h
+++ b/Engine/Plugins/Experimental/Avalanche/Source/AvalancheSequencer/Private/Settings/AvaSequencerSettings.h
@@ -5,6 +5,7 @@
 #include "AvaSequencePreset.h"
 #include "AvaSequencerDisplayRate.h"
 #include "Engine/DeveloperSettings.h"
+#include "Sidebar/SidebarState.h"
 #include "UObject/Object.h"
 #include "AvaSequencerSettings.generated.h"
 
@@ -41,6 +42,16 @@ public:
 		return CustomSequencePresets;
 	}
 
+	FSidebarState& GetSidebarState()
+	{
+		return SidebarState;
+	}
+
+	void SetSidebarState(const FSidebarState& InSidebarState)
+	{
+		SidebarState = InSidebarState;
+	}
+
 private:
 	/** The default display rate to use for new sequences */
 	UPROPERTY(Config, EditAnywhere, Category = "Playback")
@@ -57,4 +68,8 @@ private:
 	/** Sequence Presets that are uniquely identified by their Preset Name */
 	UPROPERTY(Config, EditAnywhere, Category = "Sequencer")
 	TSet<FAvaSequencePreset> CustomSequencePresets;
+
+	/** The state of a sidebar to be restored when Sequencer is initialized */
+	UPROPERTY(Config)
+	FSidebarState SidebarState;
 };
diff --git a/Engine/Plugins/MovieScene/ActorSequence/Source/ActorSequenceEditor/Private/ActorSequenceEditorTabSummoner.cpp b/Engine/Plugins/MovieScene/ActorSequence/Source/ActorSequenceEditor/Private/ActorSequenceEditorTabSummoner.cpp
--- a/Engine/Plugins/MovieScene/ActorSequence/Source/ActorSequenceEditor/Private/ActorSequenceEditorTabSummoner.cpp
+++ b/Engine/Plugins/MovieScene/ActorSequence/Source/ActorSequenceEditor/Private/ActorSequenceEditorTabSummoner.cpp
@@ -399,6 +399,7 @@ public:
 			{
 				SequencerInitParams.ToolkitHost = WeakBlueprintEditor.Pin()->GetToolkitHost();
 				SequencerInitParams.HostCapabilities.bSupportsCurveEditor = true;
+				SequencerInitParams.HostCapabilities.bSupportsSidebar = true;
 			}
 
 			TSharedRef<FExtender> AddMenuExtender = MakeShareable(new FExtender);
diff --git a/Engine/Plugins/MovieScene/LevelSequenceEditor/Source/LevelSequenceEditor/Private/LevelSequenceEditorToolkit.cpp b/Engine/Plugins/MovieScene/LevelSequenceEditor/Source/LevelSequenceEditor/Private/LevelSequenceEditorToolkit.cpp
--- a/Engine/Plugins/MovieScene/LevelSequenceEditor/Source/LevelSequenceEditor/Private/LevelSequenceEditorToolkit.cpp
+++ b/Engine/Plugins/MovieScene/LevelSequenceEditor/Source/LevelSequenceEditor/Private/LevelSequenceEditorToolkit.cpp
@@ -176,6 +176,7 @@ void FLevelSequenceEditorToolkit::Initialize(const EToolkitMode::Type Mode, cons
 		SequencerInitParams.HostCapabilities.bSupportsRecording = true;
 		SequencerInitParams.HostCapabilities.bSupportsRenderMovie = true;
 		SequencerInitParams.HostCapabilities.bSupportsAddFromContentBrowser = true;
+		SequencerInitParams.HostCapabilities.bSupportsSidebar = true;
 	}
 
 	ExtendSequencerToolbar("Sequencer.MainToolBar");
diff --git a/Engine/Source/Developer/ToolWidgets/Private/Sidebar/SSidebar.cpp b/Engine/Source/Developer/ToolWidgets/Private/Sidebar/SSidebar.cpp
--- a/Engine/Source/Developer/ToolWidgets/Private/Sidebar/SSidebar.cpp
+++ b/Engine/Source/Developer/ToolWidgets/Private/Sidebar/SSidebar.cpp
@@ -4,7 +4,9 @@
 #include "Framework/Application/SlateApplication.h"
 #include "SidebarButtonMenuContext.h"
 #include "Sidebar/ISidebarDrawerContent.h"
+#include "Sidebar/SidebarState.h"
 #include "Sidebar/SSidebarButton.h"
+#include "Sidebar/SSidebarContainer.h"
 #include "Sidebar/SSidebarDrawer.h"
 #include "Sidebar/SSidebarDrawerContent.h"
 #include "ToolMenu.h"
@@ -21,20 +23,23 @@ SSidebar::~SSidebar()
 	RemoveAllDrawers();
 }
 
-void SSidebar::Construct(const FArguments& InArgs, const TSharedRef<SOverlay>& InDrawersOverlay, const TSharedRef<SBox>& InDockLocation)
+void SSidebar::Construct(const FArguments& InArgs, const TSharedRef<SSidebarContainer>& InContainerWidget)
 {
-	DrawersOverlayWeak = InDrawersOverlay;
-	DockLocationWeak = InDockLocation;
+	ContainerWidget = InContainerWidget;
 
 	TabLocation = InArgs._TabLocation;
-	bHideWhenDocked = InArgs._HideWhenDocked;
+	InContainerWidget->SidebarSizePercent = InArgs._InitialDrawerSize;
+	OnGetContent = InArgs._OnGetContent;
+	bHideWhenAllDocked = InArgs._HideWhenAllDocked;
 	bAlwaysUseMaxButtonSize = InArgs._AlwaysUseMaxButtonSize;
 	bDisablePin = InArgs._DisablePin;
 	bDisableDock = InArgs._DisableDock;
-	OnDockStateChanged = InArgs._OnDockStateChanged;
+	OnStateChanged = InArgs._OnStateChanged;
+
+	check(OnGetContent.IsBound());
+
+	SetVisibility(EVisibility::Visible);
 
-	SetVisibility(EVisibility::SelfHitTestInvisible);
-	
 	ChildSlot
 	.Padding(FMargin(
 		TabLocation == ESidebarTabLocation::Right ? 2.f : 0.f,
@@ -43,11 +48,11 @@ void SSidebar::Construct(const FArguments& InArgs, const TSharedRef<SOverlay>& I
 		TabLocation == ESidebarTabLocation::Top ? 2.f : 0.f))
 	[
 		SNew(SBorder)
-		.Padding(0.0f)
+		.Padding(0.f)
 		.BorderImage(FAppStyle::Get().GetBrush(TEXT("Docking.Sidebar.Background")))
 		[
 			SAssignNew(TabButtonContainer, SScrollBox)
-			.Orientation(IsHorizontal() ? EOrientation::Orient_Horizontal : EOrientation::Orient_Vertical)
+			.Orientation(IsHorizontal() ? Orient_Horizontal : Orient_Vertical)
 			.ScrollBarAlwaysVisible(false)
 			.ScrollBarVisibility(EVisibility::Collapsed)
 		]
@@ -62,6 +67,7 @@ bool SSidebar::RegisterDrawer(FSidebarDrawerConfig&& InDrawerConfig)
 	}
 
 	const TSharedRef<FSidebarDrawer> NewDrawer = MakeShared<FSidebarDrawer>(MoveTemp(InDrawerConfig));
+	NewDrawer->State = InDrawerConfig.InitialState;
 	NewDrawer->bDisablePin = bDisablePin;
 	NewDrawer->bDisableDock = bDisableDock;
 	NewDrawer->ContentWidget = NewDrawer->Config.OverrideContentWidget.IsValid()
@@ -71,7 +77,7 @@ bool SSidebar::RegisterDrawer(FSidebarDrawerConfig&& InDrawerConfig)
 	// Add tab button
 	TabButtonContainer->AddSlot()
 		[
-			SAssignNew(NewDrawer->ButtonWidget, SSidebarButton, NewDrawer, TabLocation)
+			SAssignNew(NewDrawer->ButtonWidget, SSidebarButton, NewDrawer, GetTabLocation())
 			.MinButtonSize(bAlwaysUseMaxButtonSize ? MaxTabButtonSize : MinTabButtonSize)
 			.MaxButtonSize(MaxTabButtonSize)
 			.ButtonThickness(TabButtonThickness)
@@ -83,37 +89,25 @@ bool SSidebar::RegisterDrawer(FSidebarDrawerConfig&& InDrawerConfig)
 
 	Drawers.Add(NewDrawer);
 
-	// Figure out the size this tab should be when opened later. We do it now when the tab still has valid geometry. Once it is moved to the sidebar it will not.
-	float TargetDrawerSizePct = NewDrawer->SizeCoefficient;
-	if (TargetDrawerSizePct == 0)
-	{
-		TSharedPtr<SWindow> MyWindow = FSlateApplication::Get().FindWidgetWindow(AsShared());
-		if (MyWindow.IsValid() && NewDrawer->ContentWidget.IsValid())
-		{
-			TargetDrawerSizePct = NewDrawer->ContentWidget->GetTickSpaceGeometry().GetLocalSize().X / MyWindow->GetPaintSpaceGeometry().GetLocalSize().X;
-			NewDrawer->SizeCoefficient = TargetDrawerSizePct;
-		}
-	}
+	const FName DrawerId = NewDrawer->GetUniqueId();
 
-	// We don't currently allow more than one pinned tab per sidebar, so enforce that
-	// Note: it's possible to relax this if users actually want multiple pinned tabs
-	if (FindFirstPinnedTab())
+	if (NewDrawer->State.bIsPinned)
 	{
-		SetDrawerPinned(NewDrawer->GetUniqueId(), false);
+		UndockAllDrawers();
+		SetDrawerPinned(DrawerId, true);
 	}
-
-	if (NewDrawer->bIsPinned)
+	else if (NewDrawer->State.bIsDocked)
 	{
-		// If this tab is a pinned tab, then open the drawer automatically after it's added
-		OpenDrawerNextFrame(NewDrawer, /*bAnimateOpen=*/false);
+		SetDrawerDocked(DrawerId, true);
 	}
-	else if (NewDrawer->Config.bInitiallyDocked)
+
+	ContainerWidget->UpdateDrawerTabAppearance();
+
+	if (bHideWhenAllDocked && !AreAllDrawersDocked())
 	{
-		SetDrawerDocked(NewDrawer->GetUniqueId(), true);
+		SetVisibility(EVisibility::Visible);
 	}
 
-	UpdateDrawerAppearance();
-
 	return true;
 }
 
@@ -121,7 +115,7 @@ bool SSidebar::UnregisterDrawer(const FName InDrawerId)
 {
 	if (IsDrawerOpened(InDrawerId))
 	{
-		CloseAllDrawers();
+		ContainerWidget->CloseAllDrawerWidgets(false);
 	}
 	
 	const int32 IndexToRemove = Drawers.IndexOfByPredicate(
@@ -139,7 +133,7 @@ bool SSidebar::UnregisterDrawer(const FName InDrawerId)
 	RemoveDrawer(Drawers[IndexToRemove]);
 	Drawers.RemoveAt(IndexToRemove);
 
-	SummonPinnedTabIfNothingOpened();
+	ContainerWidget->SummonPinnedTabIfNothingOpened();
 
 	// Clear the pinned flag when the tab is removed from the sidebar.
 	// (Users probably expect that pinning a tab, restoring it/closing it,
@@ -152,7 +146,7 @@ bool SSidebar::UnregisterDrawer(const FName InDrawerId)
 	}
 	else
 	{
-		UpdateDrawerAppearance();
+		ContainerWidget->UpdateDrawerTabAppearance();
 	}
 
 	return true;
@@ -224,16 +218,14 @@ bool SSidebar::TryOpenDrawer(const FName InDrawerId)
 		return false;
 	}
 
-	OpenDrawerNextFrame(Drawer.ToSharedRef());
+	ContainerWidget->OpenDrawerNextFrame(Drawer.ToSharedRef());
+
 	return true;
 }
 
 void SSidebar::CloseAllDrawers(const bool bInAnimate)
 {
-	for (const TSharedRef<FSidebarDrawer>& DrawerTab : Drawers)
-	{
-		CloseDrawerInternal(DrawerTab, bInAnimate);
-	}
+	ContainerWidget->CloseAllDrawerWidgets(bInAnimate);
 }
 
 void SSidebar::OnTabDrawerButtonPressed(const TSharedRef<FSidebarDrawer>& InDrawer)
@@ -243,17 +235,17 @@ void SSidebar::OnTabDrawerButtonPressed(const TSharedRef<FSidebarDrawer>& InDraw
 		// When clicking on the button of an active (but unpinned) tab, close that tab drawer
 		if (!IsDrawerPinned(InDrawer->GetUniqueId()))
 		{
-			CloseDrawerInternal(InDrawer);
+			ContainerWidget->CloseDrawer_Internal(InDrawer);
+		}
+		else if (!InDrawer->DrawerWidget->HasKeyboardFocus())
+		{
+			FSlateApplication::Get().SetKeyboardFocus(InDrawer->DrawerWidget);
 		}
 	}
-	else if (!InDrawer->bIsDocked)
+	else if (!InDrawer->State.bIsDocked)
 	{
 		// Otherwise clicking on an inactive tab should open the drawer
-		OpenDrawerInternal(InDrawer);
-	}
-	else if (InDrawer->bIsDocked && InDrawer->DrawerWidget.IsValid())
-	{
-		FSlateApplication::Get().SetKeyboardFocus(InDrawer->DrawerWidget);
+		ContainerWidget->OpenDrawer_Internal(InDrawer);
 	}
 }
 
@@ -264,71 +256,18 @@ void SSidebar::OnDrawerTabPinToggled(const TSharedRef<FSidebarDrawer>& InDrawer,
 	{
 		SetDrawerPinned(DrawerTab->GetUniqueId(), DrawerTab == InDrawer ? bIsPinned : false);
 	}
-
-	// Open any newly-pinned tab
-	if (bIsPinned)
-	{
-		OpenDrawerInternal(InDrawer);
-	}
 }
 
 void SSidebar::OnDrawerTabDockToggled(const TSharedRef<FSidebarDrawer>& InDrawer, const bool bIsDocked)
 {
-	// Undock the previously docked drawer
-	if (DockedDrawerTab.IsValid())
-	{
-		SetDrawerDocked(DockedDrawerTab->GetUniqueId(), false);
-	}
+	SetDrawerDocked(InDrawer->GetUniqueId(), bIsDocked);
 
-	// Dock new drawer if needed
-	if (bIsDocked)
+	if (!bIsDocked)
 	{
-		SetDrawerDocked(InDrawer->GetUniqueId(), bIsDocked);
+		SetWidgetDrawerSize(InDrawer);
 	}
 }
 
-void SSidebar::OnTabDrawerFocusLost(const TSharedRef<SSidebarDrawer>& InDrawerWidget)
-{
-	const TSharedPtr<FSidebarDrawer> Drawer = InDrawerWidget->GetDrawer();
-	if (!Drawer.IsValid() || IsDrawerPinned(Drawer->GetUniqueId()))
-	{
-		return;
-	}
-
-	CloseDrawerInternal(Drawer.ToSharedRef());
-}
-
-void SSidebar::OnOpenAnimationFinish(const TSharedRef<SSidebarDrawer>& InDrawerWidget)
-{
-}
-
-void SSidebar::OnCloseAnimationFinish(const TSharedRef<SSidebarDrawer>& InDrawerWidget)
-{
-	if (const TSharedPtr<SOverlay> DrawersOverlay = DrawersOverlayWeak.Pin())
-	{
-		DrawersOverlay->RemoveSlot(InDrawerWidget);
-	}
-
-	ClosingDrawerWidgets.Remove(InDrawerWidget);
-}
-
-void SSidebar::OnDrawerTargetSizeChanged(const TSharedRef<SSidebarDrawer>& InDrawerWidget, const float InNewSize)
-{
-	const TSharedPtr<SOverlay> DrawersOverlay = DrawersOverlayWeak.Pin();
-	if (!DrawersOverlay.IsValid())
-	{
-		return;
-	}
-	
-	const TSharedPtr<FSidebarDrawer> DrawerWidget = InDrawerWidget->GetDrawer();
-	if (!DrawerWidget.IsValid())
-	{
-		return;
-	}
-	
-	DrawerWidget->SizeCoefficient = InNewSize / DrawersOverlay->GetPaintSpaceGeometry().GetLocalSize().X;
-}
-
 TSharedRef<SWidget> SSidebar::OnGetTabDrawerContextMenuWidget(TSharedRef<FSidebarDrawer> InDrawer)
 {
 	UToolMenus* const ToolMenus = UToolMenus::Get();
@@ -375,7 +314,7 @@ void SSidebar::BuildOptionsMenu(UToolMenu* const InMenu)
 
 	FToolMenuSection& Section = InMenu->FindOrAddSection(TEXT("Options"), LOCTEXT("Options", "Options"));
 
-	if (Drawer->bIsDocked)
+	if (Drawer->State.bIsDocked)
 	{
 		Section.AddMenuEntry(TEXT("Undock"),
 			LOCTEXT("UndockLabel", "Undock"),
@@ -392,7 +331,7 @@ void SSidebar::BuildOptionsMenu(UToolMenu* const InMenu)
 			FUIAction(FExecuteAction::CreateSP(this, &SSidebar::SetDrawerDocked, Drawer->GetUniqueId(), true)));
 	}
 
-	if (Drawer->bIsPinned)
+	if (Drawer->State.bIsPinned)
 	{
 		Section.AddMenuEntry(TEXT("Unpin"),
 			LOCTEXT("UnpinLabel", "Unpin"),
@@ -414,219 +353,26 @@ void SSidebar::RemoveDrawer(const TSharedRef<FSidebarDrawer>& InDrawer)
 {
 	if (InDrawer->DrawerWidget.IsValid())
 	{
-		if (const TSharedPtr<SOverlay> DrawersOverlay = DrawersOverlayWeak.Pin())
-		{
-			DrawersOverlay->RemoveSlot(InDrawer->DrawerWidget.ToSharedRef());
-		}
+		ContainerWidget->RemoveDrawerOverlaySlot(InDrawer, false);
 	}
 
-	InDrawer->bIsOpen = false;
-
 	InDrawer->DrawerClosedDelegate.ExecuteIfBound(InDrawer->GetUniqueId());
 
-	UpdateDrawerAppearance();
+	Drawers.Remove(InDrawer);
+
+	ContainerWidget->UpdateDrawerTabAppearance();
 }
 
 void SSidebar::RemoveAllDrawers()
 {
-	PendingTabToOpen.Reset();
-	bAnimatePendingTabOpen = false;
-
-	// Closing drawers can remove them from the opened drawers list so copy the list first
-	TArray<TSharedRef<SSidebarDrawer>> OpenDrawerWidgetsCopy = OpenDrawerWidgets;
-
-	for (const TSharedRef<SSidebarDrawer>& Drawer : OpenDrawerWidgetsCopy)
+	for (const TSharedRef<FSidebarDrawer>& Drawer : Drawers)
 	{
-		RemoveDrawer(Drawer->GetDrawer().ToSharedRef());
+		RemoveDrawer(Drawer);
 	}
 
 	Drawers.Empty();
 }
 
-EActiveTimerReturnType SSidebar::OnOpenPendingDrawerTimer(const double InCurrentTime, const float InDeltaTime)
-{
-	if (const TSharedPtr<FSidebarDrawer> TabToOpen = PendingTabToOpen.Pin())
-	{
-		// Wait until the drawers overlay has been arranged once to open the drawer
-		// It might not have geometry yet if we're adding back tabs on startup
-		if (const TSharedPtr<SOverlay> DrawersOverlay = DrawersOverlayWeak.Pin())
-		{
-			if (DrawersOverlay->GetTickSpaceGeometry().GetLocalSize().IsZero())
-			{
-				return EActiveTimerReturnType::Continue;
-			}
-		}
-
-		OpenDrawerInternal(TabToOpen.ToSharedRef(), bAnimatePendingTabOpen);
-	}
-
-	PendingTabToOpen.Reset();
-	bAnimatePendingTabOpen = false;
-	OpenPendingDrawerTimerHandle.Reset();
-
-	return EActiveTimerReturnType::Stop;
-}
-
-void SSidebar::OpenDrawerNextFrame(const TSharedRef<FSidebarDrawer>& InDrawer, const bool bInAnimate)
-{
-	PendingTabToOpen = InDrawer;
-	bAnimatePendingTabOpen = bInAnimate;
-
-	if (!OpenPendingDrawerTimerHandle.IsValid())
-	{
-		OpenPendingDrawerTimerHandle = RegisterActiveTimer(0.f, FWidgetActiveTimerDelegate::CreateSP(this, &SSidebar::OnOpenPendingDrawerTimer));
-	}
-}
-
-void SSidebar::OpenDrawerInternal(const TSharedRef<FSidebarDrawer>& InDrawer, const bool bInAnimate)
-{
-	if (OpenDrawerWidgets.Contains(InDrawer->DrawerWidget))
-	{
-		return;
-	}
-
-	const TSharedPtr<SOverlay> DrawersOverlay = DrawersOverlayWeak.Pin();
-	if (!DrawersOverlay.IsValid())
-	{
-		return;
-	}
-
-	PendingTabToOpen.Reset();
-	bAnimatePendingTabOpen = false;
-
-	const FGeometry DrawersOverlayGeometry = DrawersOverlay->GetTickSpaceGeometry();
-	const FGeometry Geometry = GetTickSpaceGeometry();
-
-	// Calculate padding for the drawer itself
-	const float MinDrawerSize = Geometry.GetLocalSize().X - 4.f; // overlap with sidebar border slightly
-	const FVector2D ShadowOffset(8.f, 8.f);
-	const FMargin SlotPadding(
-		TabLocation == ESidebarTabLocation::Left ? MinDrawerSize : 0.f,
-		-ShadowOffset.Y,
-		TabLocation == ESidebarTabLocation::Right ? MinDrawerSize : 0.f,
-		-ShadowOffset.Y);
-	const float AvailableWidth = DrawersOverlayGeometry.GetLocalSize().X - SlotPadding.GetTotalSpaceAlong<EOrientation::Orient_Horizontal>();
-	const float MaxDrawerSize = AvailableWidth * 0.5f;
-
-	float TargetDrawerSizePct = InDrawer->SizeCoefficient;
-	TargetDrawerSizePct = FMath::Clamp(TargetDrawerSizePct, 0.f, 0.5f);
-
-	const float TargetDrawerSize = AvailableWidth * TargetDrawerSizePct;
-
-	if (!InDrawer->DrawerWidget.IsValid())
-	{
-		InDrawer->DrawerWidget =
-			SNew(SSidebarDrawer, InDrawer, TabLocation)
-			.MinDrawerSize(MinDrawerSize)
-			.TargetDrawerSize(TargetDrawerSize)
-			.MaxDrawerSize(MaxDrawerSize)
-			.OnDrawerFocusLost(this, &SSidebar::OnTabDrawerFocusLost)
-			.OnOpenAnimationFinish(this, &SSidebar::OnOpenAnimationFinish)
-			.OnCloseAnimationFinish(this, &SSidebar::OnCloseAnimationFinish)
-			.OnDrawerTargetSizeChanged(this, &SSidebar::OnDrawerTargetSizeChanged);
-	}
-
-	const TSharedRef<SSidebarDrawer> DrawerWidgetRef = InDrawer->DrawerWidget.ToSharedRef();
-
-	if (ClosingDrawerWidgets.Contains(DrawerWidgetRef))
-	{
-		ClosingDrawerWidgets.Remove(DrawerWidgetRef);
-	}
-	else
-	{
-		DrawersOverlay->AddSlot()
-			.Padding(SlotPadding)
-			.HAlign(TabLocation == ESidebarTabLocation::Left ? HAlign_Left : HAlign_Right)
-			[
-				DrawerWidgetRef
-			];
-	}
-
-	OpenDrawerWidgets.Add(DrawerWidgetRef);
-
-	InDrawer->DrawerWidget->Open(bInAnimate);
-
-	for (const TSharedRef<FSidebarDrawer>& DrawerTab : Drawers)
-	{
-		DrawerTab->bIsOpen = false;
-	}
-	InDrawer->bIsOpen = true;
-
-	InDrawer->DrawerOpenedDelegate.ExecuteIfBound(InDrawer->GetUniqueId());
-
-	UpdateDrawerAppearance();
-
-	// This changes the focus and will trigger focus-related events, such as closing other tabs,
-	// so it's important that we only call it after we added the new drawer to OpenedDrawers.
-	FSlateApplication::Get().SetKeyboardFocus(InDrawer->DrawerWidget);
-}
-
-void SSidebar::CloseDrawerInternal(const TSharedRef<FSidebarDrawer>& InDrawer, const bool bInAnimate)
-{
-	const TSharedPtr<SSidebarDrawer> FoundDrawerWidget = FindOpenDrawerWidget(InDrawer);
-	if (OpenDrawerWidgets.Contains(FoundDrawerWidget) && !ClosingDrawerWidgets.Contains(InDrawer->DrawerWidget))
-	{
-		const TSharedRef<SSidebarDrawer> FoundDrawerWidgetRef = FoundDrawerWidget.ToSharedRef();
-
-		FoundDrawerWidgetRef->Close(bInAnimate);
-
-		if (bInAnimate)
-		{
-			ClosingDrawerWidgets.Add(FoundDrawerWidgetRef);
-		}
-		else
-		{
-			if (const TSharedPtr<SOverlay> DrawersOverlay = DrawersOverlayWeak.Pin())
-			{
-				DrawersOverlay->RemoveSlot(FoundDrawerWidgetRef);
-			}
-		}
-
-		OpenDrawerWidgets.Remove(FoundDrawerWidgetRef);
-
-		InDrawer->bIsOpen = false;
-	}
-
-	SummonPinnedTabIfNothingOpened();
-	UpdateDrawerAppearance();
-}
-
-void SSidebar::SummonPinnedTabIfNothingOpened()
-{
-	// If there's already a tab in the foreground, don't bring the pinned tab forward
-	if (GetForegroundTab())
-	{
-		return;
-	}
-
-	// But if there's no current foreground tab, then bring forward a pinned tab (there should be at most one)
-	// This should happen when:
-	// - the current foreground tab is not pinned and loses focus
-	// - the current foreground tab's drawer is manually closed by pressing on the tab button
-	// - closing or restoring the current foreground tab
-	if (const TSharedPtr<FSidebarDrawer> PinnedTab = FindFirstPinnedTab())
-	{
-		OpenDrawerInternal(PinnedTab.ToSharedRef());
-	}
-}
-
-void SSidebar::UpdateDrawerAppearance()
-{
-	TSharedPtr<FSidebarDrawer> OpenedDrawer;
-	if (OpenDrawerWidgets.Num() > 0)
-	{
-		OpenedDrawer = OpenDrawerWidgets.Last()->GetDrawer();
-	}
-
-	for (const TSharedRef<FSidebarDrawer>& Drawer : Drawers)
-	{
-		if (const TSharedPtr<SSidebarButton> TabButton = StaticCastSharedPtr<SSidebarButton>(Drawer->ButtonWidget))
-		{
-			TabButton->UpdateAppearance(OpenedDrawer);
-		}
-	}
-}
-
 TSharedPtr<FSidebarDrawer> SSidebar::FindDrawer(const FName InDrawerId) const
 {
 	const TSharedRef<FSidebarDrawer>* const FoundDrawer = Drawers.FindByPredicate(
@@ -637,38 +383,6 @@ TSharedPtr<FSidebarDrawer> SSidebar::FindDrawer(const FName InDrawerId) const
 	return FoundDrawer ? *FoundDrawer : TSharedPtr<FSidebarDrawer>();
 }
 
-TSharedPtr<FSidebarDrawer> SSidebar::FindFirstPinnedTab() const
-{
-	for (const TSharedRef<FSidebarDrawer>& DrawerTab : Drawers)
-	{
-		if (DrawerTab->bIsPinned)
-		{
-			return DrawerTab;
-		}
-	}
-	return nullptr;
-}
-
-TSharedPtr<FSidebarDrawer> SSidebar::GetForegroundTab() const
-{
-	const int32 Index = OpenDrawerWidgets.FindLastByPredicate(
-		[](const TSharedRef<SSidebarDrawer>& InDrawerWidget)
-		{
-			return InDrawerWidget->IsOpen() && !InDrawerWidget->IsClosing();
-		});
-	return Index == INDEX_NONE ? nullptr : OpenDrawerWidgets[Index]->GetDrawer();
-}
-
-TSharedPtr<SSidebarDrawer> SSidebar::FindOpenDrawerWidget(const TSharedRef<FSidebarDrawer>& InDrawer) const
-{
-	const TSharedRef<SSidebarDrawer>* const OpenDrawWidget = OpenDrawerWidgets.FindByPredicate(
-		[&InDrawer](const TSharedRef<SSidebarDrawer>& Drawer)
-		{
-			return InDrawer == Drawer->GetDrawer();
-		});
-	return OpenDrawWidget ? OpenDrawWidget->ToSharedPtr(): nullptr;
-}
-
 bool SSidebar::HasDrawerOpened() const
 {
 	for (const TSharedRef<FSidebarDrawer>& DrawerTab : Drawers)
@@ -695,20 +409,14 @@ bool SSidebar::IsDrawerOpened(const FName InDrawerId) const
 
 FName SSidebar::GetOpenedDrawerId() const
 {
-	if (OpenDrawerWidgets.IsEmpty())
-	{
-		return NAME_None;
-	}
-	
-	const TSharedRef<SSidebarDrawer> LastOpenDrawerWidget = OpenDrawerWidgets.Last();
-	return LastOpenDrawerWidget->GetDrawer()->GetUniqueId();
+	return ContainerWidget->GetOpenedDrawerId();
 }
 
 bool SSidebar::HasDrawerPinned() const
 {
 	for (const TSharedRef<FSidebarDrawer>& DrawerTab : Drawers)
 	{
-		if (DrawerTab->bIsPinned)
+		if (DrawerTab->State.bIsPinned)
 		{
 			return true;
 		}
@@ -718,131 +426,129 @@ bool SSidebar::HasDrawerPinned() const
 
 bool SSidebar::IsDrawerPinned(const FName InDrawerId) const
 {
-	if (const TSharedPtr<FSidebarDrawer> Drawer = FindDrawer(InDrawerId))
+	const TSharedPtr<FSidebarDrawer> Drawer = FindDrawer(InDrawerId);
+	if (!Drawer.IsValid())
 	{
-		return PinnedDrawerTabs.Contains(Drawer.ToSharedRef());
+		return false;
 	}
-	return false;
+	return Drawer->State.bIsPinned;
+}
+
+TSet<FName> SSidebar::GetPinnedDrawerIds() const
+{
+	TSet<FName> OutDrawerIds;
+
+	for (const TSharedRef<FSidebarDrawer>& Drawer : Drawers)
+	{
+		if (Drawer->State.bIsPinned)
+		{
+			OutDrawerIds.Add(Drawer->GetUniqueId());
+		}
+	}
+
+	return OutDrawerIds;
 }
 
 void SSidebar::SetDrawerPinned(const FName InDrawerId, const bool bInIsPinned)
 {
 	const TSharedPtr<FSidebarDrawer> DrawerToPin = FindDrawer(InDrawerId);
-	if (!DrawerToPin.IsValid() || DrawerToPin->bIsPinned == bInIsPinned)
+	if (!DrawerToPin.IsValid())
 	{
 		return;
 	}
 
 	if (bInIsPinned)
 	{
-		if (DrawerToPin->bIsDocked)
+		if (DrawerToPin->State.bIsDocked)
 		{
 			SetDrawerDocked(InDrawerId, false);
 		}
 
 		if (!DrawerToPin->bIsOpen)
 		{
-			OpenDrawerInternal(DrawerToPin.ToSharedRef(), false);
-		}
-		if (!DrawerToPin->bIsOpen)
-		{
-			return;
+			ContainerWidget->OpenDrawerNextFrame(DrawerToPin.ToSharedRef(), false);
 		}
 
 		// In case two modules attempt to register drawers with initially pinned states
 		for (const TSharedRef<FSidebarDrawer>& Drawer : Drawers)
 		{
-			Drawer->bIsPinned = false;
+			Drawer->State.bIsPinned = false;
 		}
 	}
 
-	DrawerToPin->bIsPinned = bInIsPinned;
-
-	if (bInIsPinned)
-	{
-		PinnedDrawerTabs.AddUnique(DrawerToPin.ToSharedRef());
-	}
-	else
+	DrawerToPin->State.bIsPinned = bInIsPinned;
+	if (DrawerToPin->State.bIsPinned)
 	{
-		PinnedDrawerTabs.Remove(DrawerToPin.ToSharedRef());
+		DrawerToPin->bIsOpen = true;
 	}
+
+	OnStateChanged.ExecuteIfBound(GetState());
 }
 
 bool SSidebar::HasDrawerDocked() const
 {
-	return DockedDrawerTab.IsValid();
+	for (const TSharedRef<FSidebarDrawer>& Drawer : Drawers)
+	{
+		if (Drawer->State.bIsDocked)
+		{
+			return true;
+		}
+	}
+	return false;
 }
 
 bool SSidebar::IsDrawerDocked(const FName InDrawerId) const
 {
-	if (!DockedDrawerTab.IsValid())
+	const TSharedPtr<FSidebarDrawer> Drawer = FindDrawer(InDrawerId);
+	if (!Drawer.IsValid())
 	{
 		return false;
 	}
+	return Drawer->State.bIsDocked;
+}
+
+TSet<FName> SSidebar::GetDockedDrawerIds() const
+{
+	TSet<FName> OutDrawerIds;
 
-	const TSharedPtr<FSidebarDrawer> DrawerConfig = FindDrawer(InDrawerId);
-	if (!DrawerConfig.IsValid())
+	for (const TSharedRef<FSidebarDrawer>& Drawer : Drawers)
 	{
-		return false;
+		if (Drawer->State.bIsDocked)
+		{
+			OutDrawerIds.Add(Drawer->GetUniqueId());
+		}
 	}
 
-	return DrawerConfig->GetUniqueId() == DockedDrawerTab->GetUniqueId();
+	return OutDrawerIds;
 }
 
 void SSidebar::SetDrawerDocked(const FName InDrawerId, const bool bInIsDocked)
 {
-	const TSharedPtr<SBox> DockLocation = DockLocationWeak.Pin();
-	if (!DockLocation.IsValid())
-	{
-		return;
-	}
-
 	const TSharedPtr<FSidebarDrawer> DrawerToDock = FindDrawer(InDrawerId);
-	if (!DrawerToDock.IsValid() || DrawerToDock->bIsDocked == bInIsDocked)
+	if (!DrawerToDock.IsValid())
 	{
 		return;
 	}
 
 	if (bInIsDocked)
 	{
-		if (DrawerToDock->bIsPinned)
+		if (DrawerToDock->State.bIsPinned)
 		{
 			SetDrawerPinned(InDrawerId, false);
 		}
 
-		CloseAllDrawers(false);
-
-		if (DockedDrawerTab.IsValid())
-		{
-			UndockAllDrawers();
-		}
-
-		DockedDrawerTab = DrawerToDock;
-
-		// In case two modules attempt to register drawers with initially docked states
-		for (const TSharedRef<FSidebarDrawer>& Drawer : Drawers)
-		{
-			Drawer->bIsDocked = false;
-		}
-
-		DockedDrawerTab->bIsOpen = true;
-		DockedDrawerTab->bIsPinned = false;
-		DockedDrawerTab->bIsDocked = true;
+		ContainerWidget->DockDrawer_Internal(DrawerToDock.ToSharedRef());
 
 		if (DrawerToDock->ContentWidget.IsValid())
 		{
-			DockLocation->SetContent(DrawerToDock->ContentWidget.ToSharedRef());
-			
-			if (bHideWhenDocked)
+			if (bHideWhenAllDocked && AreAllDrawersDocked())
 			{
 				SetVisibility(EVisibility::Collapsed);
 			}
 		}
 		else
 		{
-			DockLocation->SetContent(SNullWidget::NullWidget);
-			
-			if (bHideWhenDocked)
+			if (bHideWhenAllDocked && !AreAllDrawersDocked())
 			{
 				SetVisibility(EVisibility::Visible);
 			}
@@ -850,52 +556,30 @@ void SSidebar::SetDrawerDocked(const FName InDrawerId, const bool bInIsDocked)
 	}
 	else
 	{
-		if (DockedDrawerTab.IsValid())
-		{
-			DockedDrawerTab->bIsOpen = false;
-			DockedDrawerTab->bIsDocked = false;
-			DockedDrawerTab.Reset();
+		ContainerWidget->UndockDrawer_Internal(DrawerToDock.ToSharedRef());
 
-			DockLocation->SetContent(SNullWidget::NullWidget);
-			if (bHideWhenDocked)
-			{
-				SetVisibility(EVisibility::Visible);
-			}
-		}
+		if (bHideWhenAllDocked && !AreAllDrawersDocked())
+        {
+        	SetVisibility(EVisibility::Visible);
+        }
 	}
 
-	OnDockStateChanged.ExecuteIfBound(InDrawerId);
+	OnStateChanged.ExecuteIfBound(GetState());
 }
 
 void SSidebar::UndockAllDrawers()
 {
-	for (const TSharedRef<FSidebarDrawer>& DrawerTab : Drawers)
+	for (const TSharedRef<FSidebarDrawer>& Drawer : Drawers)
 	{
-		SetDrawerDocked(DrawerTab->GetUniqueId(), false);
+		SetDrawerDocked(Drawer->GetUniqueId(), false);
 	}
 }
 
-void SSidebar::UpdateDockedSplitterSlot(const FName InDrawerId, SSplitter::FSlot* const InSlot, const bool bInAutoUndock, const float InDefaultDockPercent)
+void SSidebar::UnpinAllDrawers()
 {
-	if (!InSlot)
-	{
-		return;
-	}
-	
-	const TSharedPtr<FSidebarDrawer> Drawer = FindDrawer(InDrawerId);
-	if (!Drawer.IsValid())
-	{
-		return;
-	}
-
-	const bool bDocked = IsDrawerDocked(InDrawerId);
-
-	InSlot->SetSizingRule(bDocked ? SSplitter::ESizeRule::FractionOfParent : SSplitter::ESizeRule::SizeToContent);
-	InSlot->SetResizable(bDocked);
-	if (bInAutoUndock && InSlot->GetSizeValue() < 0.01f)
+	for (const TSharedRef<FSidebarDrawer>& Drawer : Drawers)
 	{
-		Drawer->SizeCoefficient = bDocked ? InDefaultDockPercent : 0.f;
-		InSlot->SetSizeValue(Drawer->SizeCoefficient);
+		SetDrawerPinned(Drawer->GetUniqueId(), false);
 	}
 }
 
@@ -928,4 +612,61 @@ bool SSidebar::IsVertical() const
 	return TabLocation == ESidebarTabLocation::Left || TabLocation == ESidebarTabLocation::Right;
 }
 
+FSidebarState SSidebar::GetState() const
+{
+	FSidebarState OutState;
+
+	OutState.SetHidden(false);
+
+	const float CurrentDrawerSize = ContainerWidget->GetCurrentDrawerSize();
+	OutState.SetDrawerSizes(CurrentDrawerSize, 1.f - CurrentDrawerSize);
+
+	for (const TSharedRef<FSidebarDrawer>& Drawer : Drawers)
+	{
+		OutState.FindOrAddDrawerState(Drawer->State);
+	}
+
+	return OutState;
+}
+
+ESidebarTabLocation SSidebar::GetTabLocation() const
+{
+	return TabLocation;
+}
+
+TSharedRef<SWidget> SSidebar::GetMainContent() const
+{
+	return OnGetContent.IsBound() ? OnGetContent.Execute() : SNullWidget::NullWidget;
+}
+
+const TArray<TSharedRef<FSidebarDrawer>>& SSidebar::GetAllDrawers() const
+{
+	return Drawers;
+}
+
+void SSidebar::SetWidgetDrawerSize(const TSharedRef<FSidebarDrawer>& InDrawer)
+{
+	if (!InDrawer->DrawerWidget.IsValid())
+	{
+		return;
+	}
+
+	const float DrawerOverlayWidth = ContainerWidget->GetOverlaySize().X;
+	const float CurrentDrawerSize = ContainerWidget->GetCurrentDrawerSize();
+	const float PixelWidth = CurrentDrawerSize * DrawerOverlayWidth;
+	InDrawer->DrawerWidget->SetCurrentSize(PixelWidth);
+}
+
+bool SSidebar::AreAllDrawersDocked() const
+{
+	for (const TSharedRef<FSidebarDrawer>& Drawer : Drawers)
+	{
+		if (!Drawer->State.bIsDocked)
+		{
+			return false;
+		}
+	}
+	return true;
+}
+
 #undef LOCTEXT_NAMESPACE
diff --git a/Engine/Source/Developer/ToolWidgets/Private/Sidebar/SSidebarButton.cpp b/Engine/Source/Developer/ToolWidgets/Private/Sidebar/SSidebarButton.cpp
--- a/Engine/Source/Developer/ToolWidgets/Private/Sidebar/SSidebarButton.cpp
+++ b/Engine/Source/Developer/ToolWidgets/Private/Sidebar/SSidebarButton.cpp
@@ -5,6 +5,7 @@
 #include "Sidebar/SidebarDrawer.h"
 #include "Sidebar/SSidebar.h"
 #include "Sidebar/SSidebarButtonText.h"
+#include "Sidebar/SSidebarDrawer.h"
 #include "Styling/AppStyle.h"
 #include "Widgets/Colors/SComplexGradient.h"
 #include "Widgets/Images/SImage.h"
@@ -206,7 +207,7 @@ void SSidebarButton::Construct(const FArguments& InArgs, const TSharedRef<FSideb
 			[
 				SAssignNew(MainButton, SButton)
 				.ToolTipText(InDrawer->Config.ToolTipText)
-				.ContentPadding(FMargin(0.0f, DockTabStyle->TabPadding.Top, 0.0f, DockTabStyle->TabPadding.Bottom))
+				.ContentPadding(FMargin(0.f, DockTabStyle->TabPadding.Top, 0.f, DockTabStyle->TabPadding.Bottom))
 				.OnPressed_Lambda([this]()
 					{
 						// Activate tab on mouse down (not mouse down-up) for consistency with non-sidebar tabs
@@ -223,12 +224,13 @@ void SSidebarButton::Construct(const FArguments& InArgs, const TSharedRef<FSideb
 				.Visibility(EVisibility::HitTestInvisible)
 			]
 			+ SOverlay::Slot()
-			.HAlign(TabLocation == ESidebarTabLocation::Left ? HAlign_Left : HAlign_Right)
+			.HAlign(GetHAlignFromTabLocation(TabLocation))
+			.VAlign(GetVAlignFromTabLocation(TabLocation))
 			[
 				SAssignNew(ActiveIndicator, SComplexGradient)
-				.DesiredSizeOverride(FVector2D(1.0f, 1.0f))
+				.DesiredSizeOverride(FVector2D(1.f, 1.f))
 				.GradientColors(GradientStops)
-				.Orientation(EOrientation::Orient_Horizontal)
+				.Orientation(Orient_Horizontal)
 				.Visibility(this, &SSidebarButton::GetActiveTabIndicatorVisibility)
 			]
 		]
@@ -269,8 +271,8 @@ void SSidebarButton::UpdateAppearance(const TSharedPtr<FSidebarDrawer>& InLastDr
 		Label->SetRotation(LabelRotation);
 	}
 
-	// Border when open/docked
-	if (InLastDrawerOpen == ThisDrawer && (!ThisDrawer->bIsDocked && ThisDrawer->bIsOpen))
+	// Border when not docked and open
+	if (InLastDrawerOpen == ThisDrawer && (!ThisDrawer->State.bIsDocked && ThisDrawer->bIsOpen))
 	{
 		OpenBorder->SetVisibility(EVisibility::HitTestInvisible);
 		OpenBorder->SetBorderImage(FAppStyle::Get().GetBrush(FocusBorderBrushName));
@@ -281,7 +283,7 @@ void SSidebarButton::UpdateAppearance(const TSharedPtr<FSidebarDrawer>& InLastDr
 	}
 
 	// Button style
-	if (InLastDrawerOpen == ThisDrawer || ThisDrawer->bIsDocked)
+	if (InLastDrawerOpen == ThisDrawer || ThisDrawer->State.bIsDocked)
 	{
 		// this button is the one with the tab that is actually opened so show the tab border
 		MainButton->SetButtonStyle(&FAppStyle::Get().GetWidgetStyle<FButtonStyle>(TEXT("Docking.SidebarButton.Opened")));
@@ -335,7 +337,10 @@ FSlateColor SSidebarButton::GetForegroundColor() const
 EVisibility SSidebarButton::GetActiveTabIndicatorVisibility() const
 {
 	const TSharedPtr<FSidebarDrawer> Drawer = DrawerWeak.Pin();
-	if (Drawer.IsValid() && Drawer->bIsOpen)
+	if (Drawer.IsValid()
+		&& Drawer->bIsOpen
+		&& Drawer->DrawerWidget.IsValid()
+		&& Drawer->DrawerWidget->HasAnyUserFocusOrFocusedDescendants())
 	{
 		return EVisibility::HitTestInvisible;
 	}
@@ -349,14 +354,14 @@ EVisibility SSidebarButton::GetPinVisibility() const
 	{
 		return EVisibility::Collapsed;
 	}
-	return (Drawer->bIsPinned || IsHovered() || Drawer->bIsOpen)
+	return (Drawer->State.bIsPinned || IsHovered() || Drawer->bIsOpen)
 		? EVisibility::Visible : EVisibility::Hidden;
 }
 
 FText SSidebarButton::GetPinToolTipText() const
 {
 	const TSharedPtr<FSidebarDrawer> Drawer = DrawerWeak.Pin();
-	if (Drawer.IsValid() && Drawer->bIsPinned)
+	if (Drawer.IsValid() && Drawer->State.bIsPinned)
 	{
 		return LOCTEXT("UnpinTabToolTip", "Unpin Tab");
 	}
@@ -366,7 +371,7 @@ FText SSidebarButton::GetPinToolTipText() const
 ECheckBoxState SSidebarButton::IsPinChecked() const
 {
 	const TSharedPtr<FSidebarDrawer> Drawer = DrawerWeak.Pin();
-	if (Drawer.IsValid() && Drawer->bIsPinned)
+	if (Drawer.IsValid() && Drawer->State.bIsPinned)
 	{
 		return ECheckBoxState::Checked;
 	}
@@ -376,7 +381,7 @@ ECheckBoxState SSidebarButton::IsPinChecked() const
 const FSlateBrush* SSidebarButton::GetPinImage() const
 {
 	const TSharedPtr<FSidebarDrawer> Drawer = DrawerWeak.Pin();
-	if (Drawer.IsValid() && Drawer->bIsPinned)
+	if (Drawer.IsValid() && Drawer->State.bIsPinned)
 	{
 		return FAppStyle::Get().GetBrush(TEXT("Icons.Pinned"));
 	}
@@ -395,13 +400,13 @@ EVisibility SSidebarButton::GetDockVisibility() const
 	{
 		return EVisibility::Collapsed;
 	}
-	return (Drawer->bIsDocked || IsHovered() || Drawer->bIsOpen) ? EVisibility::Visible : EVisibility::Hidden;
+	return (Drawer->State.bIsDocked || IsHovered() || Drawer->bIsOpen) ? EVisibility::Visible : EVisibility::Hidden;
 }
 
 FText SSidebarButton::GetDockToolTipText() const
 {
 	const TSharedPtr<FSidebarDrawer> Drawer = DrawerWeak.Pin();
-	if (Drawer.IsValid() && Drawer->bIsDocked)
+	if (Drawer.IsValid() && Drawer->State.bIsDocked)
 	{
 		return LOCTEXT("UndockTabToolTip", "Undock Tab");
 	}
@@ -411,7 +416,7 @@ FText SSidebarButton::GetDockToolTipText() const
 ECheckBoxState SSidebarButton::IsDockChecked() const
 {
 	const TSharedPtr<FSidebarDrawer> Drawer = DrawerWeak.Pin();
-	if (Drawer.IsValid() && Drawer->bIsDocked)
+	if (Drawer.IsValid() && Drawer->State.bIsDocked)
 	{
 		return ECheckBoxState::Checked;
 	}
@@ -421,7 +426,7 @@ ECheckBoxState SSidebarButton::IsDockChecked() const
 const FSlateBrush* SSidebarButton::GetDockImage() const
 {
 	const TSharedPtr<FSidebarDrawer> Drawer = DrawerWeak.Pin();
-	if (Drawer.IsValid() && Drawer->bIsPinned)
+	if (Drawer.IsValid() && Drawer->State.bIsPinned)
 	{
 		return FAppStyle::Get().GetBrush(TEXT("Icons.Layout"));
 	}
@@ -433,4 +438,34 @@ void SSidebarButton::OnDockStateChanged(const ECheckBoxState InNewState)
 	OnDockToggled.ExecuteIfBound(DrawerWeak.Pin().ToSharedRef(), InNewState == ECheckBoxState::Checked);
 }
 
+EHorizontalAlignment SSidebarButton::GetHAlignFromTabLocation(const ESidebarTabLocation InTabLocation)
+{
+	switch (InTabLocation)
+	{
+	case ESidebarTabLocation::Left:
+		return HAlign_Left;
+	case ESidebarTabLocation::Right:
+		return HAlign_Right;
+	case ESidebarTabLocation::Top:
+	case ESidebarTabLocation::Bottom:
+		return HAlign_Fill;
+	}
+	return HAlign_Fill;
+}
+
+EVerticalAlignment SSidebarButton::GetVAlignFromTabLocation(const ESidebarTabLocation InTabLocation)
+{
+	switch (InTabLocation)
+	{
+	case ESidebarTabLocation::Left:
+	case ESidebarTabLocation::Right:
+		return VAlign_Fill;
+	case ESidebarTabLocation::Top:
+		return VAlign_Top;
+	case ESidebarTabLocation::Bottom:
+		return VAlign_Bottom;
+	}
+	return VAlign_Fill;
+}
+
 #undef LOCTEXT_NAMESPACE
diff --git a/Engine/Source/Developer/ToolWidgets/Private/Sidebar/SSidebarButton.h b/Engine/Source/Developer/ToolWidgets/Private/Sidebar/SSidebarButton.h
--- a/Engine/Source/Developer/ToolWidgets/Private/Sidebar/SSidebarButton.h
+++ b/Engine/Source/Developer/ToolWidgets/Private/Sidebar/SSidebarButton.h
@@ -45,7 +45,10 @@ public:
 	void UpdateAppearance(const TSharedPtr<FSidebarDrawer>& InLastDrawerOpen);
 
 	void OnTabRenamed(const TWeakPtr<FSidebarDrawer>& InDrawer);
-	
+
+	static EHorizontalAlignment GetHAlignFromTabLocation(const ESidebarTabLocation InTabLocation);
+	static EVerticalAlignment GetVAlignFromTabLocation(const ESidebarTabLocation InTabLocation);
+
 protected:
 	EVisibility GetActiveTabIndicatorVisibility() const;
 
diff --git a/Engine/Source/Developer/ToolWidgets/Private/Sidebar/SSidebarContainer.cpp b/Engine/Source/Developer/ToolWidgets/Private/Sidebar/SSidebarContainer.cpp
new file mode 100644
--- /dev/null
+++ b/Engine/Source/Developer/ToolWidgets/Private/Sidebar/SSidebarContainer.cpp
@@ -0,0 +1,705 @@
+﻿// Copyright Epic Games, Inc. All Rights Reserved.
+
+#include "Sidebar/SSidebarContainer.h"
+#include "Framework/Application/SlateApplication.h"
+#include "Sidebar/SSidebar.h"
+#include "Sidebar/SSidebarDrawer.h"
+#include "Widgets/Layout/SSplitter.h"
+#include "Widgets/SBoxPanel.h"
+#include "Widgets/SOverlay.h"
+
+#define LOCTEXT_NAMESPACE "SSidebarContainer"
+
+void SSidebarContainer::Construct(const FArguments& InArgs)
+{
+}
+
+void SSidebarContainer::RebuildSidebar(const TSharedRef<SSidebar>& InSidebarWidget, const FSidebarState& InState)
+{
+	SidebarWidget = InSidebarWidget;
+
+	Reconstruct(InState);
+}
+
+void SSidebarContainer::Reconstruct(const FSidebarState& InState)
+{
+	TSharedPtr<SWidget> OutWidget;
+
+	if (InState.IsHidden())
+	{
+		DrawersOverlay.Reset();
+
+		OutWidget = SidebarWidget->GetMainContent();
+	}
+	else if (InState.IsVisible())
+	{
+		OutWidget = SNew(SOverlay)
+			+ SOverlay::Slot()
+			.HAlign(HAlign_Fill)
+			.VAlign(VAlign_Fill)
+			[
+				ConstructBoxPanel(InState)
+			]
+			+ SOverlay::Slot()
+			.HAlign(HAlign_Fill)
+			.VAlign(VAlign_Fill)
+			[
+				SAssignNew(DrawersOverlay, SOverlay)
+			];
+	}
+
+	ChildSlot
+		.HAlign(HAlign_Fill)
+		.VAlign(VAlign_Fill)
+		[
+			OutWidget.ToSharedRef()
+		];
+}
+
+TSharedRef<SWidget> SSidebarContainer::ConstructBoxPanel(const FSidebarState& InState)
+{
+	ConstructSplitterPanel(InState);
+
+	// MainSplitter will be valid here if we have a docked drawer
+	const TSharedRef<SWidget> Content = MainSplitter.IsValid()
+		? MainSplitter.ToSharedRef()
+		: SidebarWidget->GetMainContent();
+
+	if (SidebarWidget->IsVertical())
+	{
+		const TSharedRef<SHorizontalBox> Box = SNew(SHorizontalBox);
+
+		auto BoxContentSlot = [this, &Content, &Box]()
+		{
+			Box->AddSlot()
+				.FillWidth(1.f)
+				[
+					Content
+				];
+		};
+		auto BoxSidebarSlot = [this, &Box]()
+		{
+			Box->AddSlot()
+				.AutoWidth()
+				[
+					SidebarWidget.ToSharedRef()
+				];
+		};
+
+		const ESidebarTabLocation TabLocation = SidebarWidget->GetTabLocation();
+		if (TabLocation == ESidebarTabLocation::Left)
+		{
+			BoxSidebarSlot();
+			BoxContentSlot();
+		}
+		else if (TabLocation == ESidebarTabLocation::Right)
+		{
+			BoxContentSlot();
+			BoxSidebarSlot();
+		}
+
+		return Box;
+	}
+
+	if (SidebarWidget->IsHorizontal())
+	{
+		const TSharedRef<SVerticalBox> Box = SNew(SVerticalBox);
+
+		auto BoxContentSlot = [this, &Content, &Box]()
+		{
+			Box->AddSlot()
+				.FillHeight(1.f)
+				[
+					Content
+				];
+		};
+		auto BoxSidebarSlot = [this, &Box]()
+		{
+			Box->AddSlot()
+				.AutoHeight()
+				[
+					SidebarWidget.ToSharedRef()
+				];
+		};
+
+		const ESidebarTabLocation TabLocation = SidebarWidget->GetTabLocation();
+		if (TabLocation == ESidebarTabLocation::Top)
+		{
+			BoxSidebarSlot();
+			BoxContentSlot();
+		}
+		else if (TabLocation == ESidebarTabLocation::Bottom)
+		{
+			BoxContentSlot();
+			BoxSidebarSlot();
+		}
+
+		return Box;
+	}
+
+	return SNullWidget::NullWidget;
+}
+
+void SSidebarContainer::ConstructSplitterPanel(const FSidebarState& InState)
+{
+	if (InState.IsVisible() && SidebarWidget->HasDrawerDocked())
+	{
+		const TSet<FName> DockedDrawerIds = SidebarWidget->GetDockedDrawerIds();
+		const FName FirstFoundDrawerId = DockedDrawerIds.IsEmpty() ? NAME_None : DockedDrawerIds.Array()[0];
+
+		MainSplitter = SNew(SSplitter)
+			.Orientation(GetSplitterOrientation())
+			.OnSplitterFinishedResizing(this, &SSidebarContainer::OnSplitterResized);
+
+		const ESidebarTabLocation TabLocation = SidebarWidget->GetTabLocation();
+		if (TabLocation == ESidebarTabLocation::Left || TabLocation == ESidebarTabLocation::Top)
+		{
+			AddSidebarDockSlot(FirstFoundDrawerId);
+			AddContentDockSlot();
+		}
+		else if (TabLocation == ESidebarTabLocation::Right || TabLocation == ESidebarTabLocation::Bottom)
+		{
+			AddContentDockSlot();
+			AddSidebarDockSlot(FirstFoundDrawerId);
+		}
+	}
+	else
+	{
+		MainSplitter.Reset();
+	}
+}
+
+void SSidebarContainer::AddContentDockSlot()
+{
+	const bool bDrawerDocked = SidebarWidget->HasDrawerDocked();
+
+	if (bDrawerDocked)
+	{
+		ContentSlotSize = TAttribute<float>::Create([this]()
+			{
+				return ContentSizePercent;
+			});
+	}
+	else
+	{
+		ContentSlotSize = {};
+	}
+
+	MainSplitter->AddSlot()
+		.SizeRule(bDrawerDocked ? SSplitter::FractionOfParent : SSplitter::SizeToContent)
+		.Value(ContentSlotSize)
+		.OnSlotResized(this, &SSidebarContainer::OnContentSlotResizing)
+		[
+			SidebarWidget->GetMainContent()
+		];
+}
+
+void SSidebarContainer::RemoveContentDockSlot()
+{
+	const int32 SlotIndex = GetContentSlotIndex();
+	MainSplitter->RemoveAt(SlotIndex);
+}
+
+TSharedRef<SWidget> SSidebarContainer::GetSidebarDrawerContent(const TSharedRef<FSidebarDrawer>& InDrawer) const
+{
+	if (InDrawer->Config.OverrideContentWidget.IsValid())
+	{
+		return InDrawer->Config.OverrideContentWidget.ToSharedRef();
+	}
+	return InDrawer->ContentWidget.IsValid() ? InDrawer->ContentWidget.ToSharedRef() : SNullWidget::NullWidget;
+}
+
+void SSidebarContainer::AddSidebarDockSlot(const FName InDockDrawerId)
+{
+	const TSharedPtr<FSidebarDrawer> DrawerToDock = SidebarWidget->FindDrawer(InDockDrawerId);
+	if (!DrawerToDock.IsValid())
+	{
+		return;
+	}
+
+	const bool bDrawerDocked = SidebarWidget->HasDrawerDocked();
+
+	if (bDrawerDocked)
+	{
+		SidebarSlotSize = TAttribute<float>::Create([this]()
+			{
+				return SidebarSizePercent;
+			});
+	}
+	else
+	{
+		SidebarSlotSize = {};
+	}
+
+	MainSplitter->AddSlot()
+		.SizeRule(bDrawerDocked ? SSplitter::FractionOfParent : SSplitter::SizeToContent)
+		.Value(SidebarSlotSize)
+		.OnSlotResized(this, &SSidebarContainer::OnSidebarSlotResizing)
+		[
+			GetSidebarDrawerContent(DrawerToDock.ToSharedRef())
+		];
+}
+
+void SSidebarContainer::RemoveSidebarDockSlot()
+{
+	const int32 SlotIndex = GetSidebarSlotIndex();
+	MainSplitter->RemoveAt(SlotIndex);
+}
+
+float SSidebarContainer::GetContentSlotSize() const
+{
+	return ContentSizePercent;
+}
+
+float SSidebarContainer::GetSidebarSlotSize() const
+{
+	return SidebarSizePercent;
+}
+
+int32 SSidebarContainer::GetContentSlotIndex() const
+{
+	switch (SidebarWidget->GetTabLocation())
+	{
+	case ESidebarTabLocation::Right:
+	case ESidebarTabLocation::Bottom:
+		return 0;
+	case ESidebarTabLocation::Left:
+	case ESidebarTabLocation::Top:
+		return 1;
+	}
+	return 0;
+}
+
+int32 SSidebarContainer::GetSidebarSlotIndex() const
+{
+	switch (SidebarWidget->GetTabLocation())
+	{
+	case ESidebarTabLocation::Left:
+	case ESidebarTabLocation::Top:
+		return 0;
+	case ESidebarTabLocation::Right:
+	case ESidebarTabLocation::Bottom:
+		return 1;
+	}
+	return 1;
+}
+
+EOrientation SSidebarContainer::GetSplitterOrientation() const
+{
+	switch (SidebarWidget->GetTabLocation())
+	{
+	case ESidebarTabLocation::Left:
+	case ESidebarTabLocation::Right:
+		return Orient_Horizontal;
+	case ESidebarTabLocation::Top:
+	case ESidebarTabLocation::Bottom:
+		return Orient_Vertical;
+	}
+	return Orient_Horizontal;
+}
+
+ESidebarTabLocation SSidebarContainer::GetTabLocation() const
+{
+	return SidebarWidget->GetTabLocation();
+}
+
+float SSidebarContainer::GetCurrentDrawerSize() const
+{
+	return SidebarSizePercent;
+}
+
+FVector2D SSidebarContainer::GetOverlaySize() const
+{
+	return DrawersOverlay->GetTickSpaceGeometry().GetLocalSize();
+}
+
+bool SSidebarContainer::AddDrawerOverlaySlot(const TSharedRef<FSidebarDrawer>& InDrawer)
+{
+	if (!InDrawer->DrawerWidget)
+	{
+		return false;
+	}
+
+	const TSharedRef<SSidebarDrawer> DrawerWidgetRef = InDrawer->DrawerWidget.ToSharedRef();
+
+	if (ClosingDrawerWidgets.Contains(DrawerWidgetRef))
+	{
+		ClosingDrawerWidgets.Remove(DrawerWidgetRef);
+	}
+	else
+	{
+		const ESidebarTabLocation TabLocation = SidebarWidget->GetTabLocation();
+
+		DrawersOverlay->AddSlot()
+			.Padding(CalculateSlotMargin())
+			.HAlign(SSidebarButton::GetHAlignFromTabLocation(TabLocation))
+			.VAlign(SSidebarButton::GetVAlignFromTabLocation(TabLocation))
+			[
+				DrawerWidgetRef
+			];
+	}
+
+	OpenDrawerWidgets.Add(DrawerWidgetRef);
+
+	return true;
+}
+
+bool SSidebarContainer::RemoveDrawerOverlaySlot(const TSharedRef<FSidebarDrawer>& InDrawer, const bool bInAnimate)
+{
+	if (!InDrawer->DrawerWidget)
+	{
+		return false;
+	}
+
+	const TSharedRef<SSidebarDrawer> DrawerWidgetRef = InDrawer->DrawerWidget.ToSharedRef();
+
+	if (bInAnimate)
+	{
+		ClosingDrawerWidgets.Add(DrawerWidgetRef);
+	}
+	else
+	{
+		ClosingDrawerWidgets.Remove(DrawerWidgetRef);
+
+		DrawersOverlay->RemoveSlot(DrawerWidgetRef);
+	}
+
+	OpenDrawerWidgets.Remove(DrawerWidgetRef);
+
+	return true;
+}
+
+void SSidebarContainer::CloseAllDrawerWidgets(const bool bInAnimate)
+{
+	for (const TSharedRef<FSidebarDrawer>& Drawer : SidebarWidget->GetAllDrawers())
+	{
+		CloseDrawer_Internal(Drawer, bInAnimate);
+	}
+}
+
+EActiveTimerReturnType SSidebarContainer::OnOpenPendingDrawerTimer(const double InCurrentTime, const float InDeltaTime)
+{
+	if (const TSharedPtr<FSidebarDrawer> DrawerToOpen = PendingTabToOpen.Pin())
+	{
+		// Wait until the drawers overlay has been arranged once to open the drawer
+		// It might not have geometry yet if we're adding back tabs on startup
+		if (GetOverlaySize().IsZero())
+		{
+			return EActiveTimerReturnType::Continue;
+		}
+
+		OpenDrawer_Internal(DrawerToOpen.ToSharedRef(), bAnimatePendingTabOpen);
+	}
+
+	PendingTabToOpen.Reset();
+	bAnimatePendingTabOpen = false;
+	OpenPendingDrawerTimerHandle.Reset();
+
+	return EActiveTimerReturnType::Stop;
+}
+
+void SSidebarContainer::OpenDrawerNextFrame(const TSharedRef<FSidebarDrawer>& InDrawer, const bool bInAnimate)
+{
+	if (InDrawer->DrawerWidget.IsValid() && OpenDrawerWidgets.Contains(InDrawer->DrawerWidget))
+	{
+		return;
+	}
+
+	PendingTabToOpen = InDrawer;
+	bAnimatePendingTabOpen = bInAnimate;
+
+	if (!OpenPendingDrawerTimerHandle.IsValid())
+	{
+		OpenPendingDrawerTimerHandle = RegisterActiveTimer(0.f, FWidgetActiveTimerDelegate::CreateSP(this, &SSidebarContainer::OnOpenPendingDrawerTimer));
+	}
+}
+
+FMargin SSidebarContainer::CalculateSlotMargin() const
+{
+	const FGeometry SidebarGeometry = SidebarWidget->GetTickSpaceGeometry();
+	const float MinDrawerSize = SidebarGeometry.GetLocalSize().X - 4.f; // overlap with sidebar border slightly
+	const FVector2D ShadowOffset(8.f, 8.f);
+	const ESidebarTabLocation TabLocation = SidebarWidget->GetTabLocation();
+	return FMargin(
+		TabLocation == ESidebarTabLocation::Left ? MinDrawerSize : 0.f,
+		-ShadowOffset.Y,
+		TabLocation == ESidebarTabLocation::Right ? MinDrawerSize : 0.f,
+		-ShadowOffset.Y);
+}
+
+void SSidebarContainer::CreateDrawerWidget(const TSharedRef<FSidebarDrawer>& InDrawer)
+{
+	// Calculate padding for the drawer itself
+	const FGeometry SidebarGeometry = SidebarWidget->GetTickSpaceGeometry();
+	const float MinDrawerSize = SidebarGeometry.GetLocalSize().X - 4.f; // overlap with sidebar border slightly
+	const FMargin SlotPadding = CalculateSlotMargin();
+	const float AvailableSize = GetOverlaySize().X - SlotPadding.GetTotalSpaceAlong<Orient_Horizontal>();
+	const float MaxDrawerSize = AvailableSize * 0.5f; // max 50% of width or height
+	const float TargetDrawerSize = AvailableSize * SidebarSizePercent;
+
+	InDrawer->DrawerWidget =
+		SNew(SSidebarDrawer, InDrawer, SidebarWidget->GetTabLocation())
+		.MinDrawerSize(MinDrawerSize)
+		.MaxDrawerSize(MaxDrawerSize)
+		.TargetDrawerSize(TargetDrawerSize)
+		.OnDrawerFocusLost(this, &SSidebarContainer::OnTabDrawerFocusLost)
+		.OnOpenAnimationFinish(this, &SSidebarContainer::OnOpenAnimationFinish)
+		.OnCloseAnimationFinish(this, &SSidebarContainer::OnCloseAnimationFinish)
+		.OnDrawerSizeChanged(this, &SSidebarContainer::OnDrawerSizeChanged);
+}
+
+void SSidebarContainer::OpenDrawer_Internal(const TSharedRef<FSidebarDrawer>& InDrawer, const bool bInAnimate)
+{
+	if (InDrawer->DrawerWidget.IsValid() && OpenDrawerWidgets.Contains(InDrawer->DrawerWidget))
+	{
+		return;
+	}
+
+	PendingTabToOpen.Reset();
+	bAnimatePendingTabOpen = false;
+
+	CreateDrawerWidget(InDrawer);
+	AddDrawerOverlaySlot(InDrawer);
+
+	InDrawer->DrawerWidget->Open(bInAnimate);
+
+	for (const TSharedRef<FSidebarDrawer>& Drawer : SidebarWidget->GetAllDrawers())
+	{
+		Drawer->bIsOpen = false;
+	}
+	InDrawer->bIsOpen = true;
+
+	InDrawer->DrawerOpenedDelegate.ExecuteIfBound(InDrawer->GetUniqueId());
+
+	UpdateDrawerTabAppearance();
+
+	// This changes the focus and will trigger focus-related events, such as closing other tabs,
+	// so it's important that we only call it after we added the new drawer to OpenedDrawers.
+	FSlateApplication::Get().SetKeyboardFocus(InDrawer->DrawerWidget);
+}
+
+void SSidebarContainer::CloseDrawer_Internal(const TSharedRef<FSidebarDrawer>& InDrawer, const bool bInAnimate)
+{
+	const TSharedPtr<SSidebarDrawer> FoundDrawerWidget = FindOpenDrawerWidget(InDrawer);
+	if (OpenDrawerWidgets.Contains(FoundDrawerWidget) && !ClosingDrawerWidgets.Contains(InDrawer->DrawerWidget))
+	{
+		const TSharedRef<SSidebarDrawer> FoundDrawerWidgetRef = FoundDrawerWidget.ToSharedRef();
+
+		InDrawer->bIsOpen = false;
+
+		RemoveDrawerOverlaySlot(InDrawer, bInAnimate);
+
+		FoundDrawerWidgetRef->Close(bInAnimate);
+	}
+
+	SummonPinnedTabIfNothingOpened();
+	UpdateDrawerTabAppearance();
+}
+
+void SSidebarContainer::SummonPinnedTabIfNothingOpened()
+{
+	// If there's already a drawer in the foreground, don't bring the pinned tab forward
+	if (GetForegroundDrawer())
+	{
+		return;
+	}
+
+	// But if there's no current foreground tab, then bring forward a pinned tab (there should be at most one)
+	// This should happen when:
+	// - the current foreground tab is not pinned and loses focus
+	// - the current foreground tab's drawer is manually closed by pressing on the tab button
+	// - closing or restoring the current foreground tab
+	if (const TSharedPtr<FSidebarDrawer> PinnedTab = FindFirstPinnedTab())
+	{
+		OpenDrawer_Internal(PinnedTab.ToSharedRef());
+	}
+}
+
+void SSidebarContainer::UpdateDrawerTabAppearance()
+{
+	TSharedPtr<FSidebarDrawer> OpenedDrawer;
+	if (OpenDrawerWidgets.Num() > 0)
+	{
+		OpenedDrawer = OpenDrawerWidgets.Last()->GetDrawer();
+	}
+
+	for (const TSharedRef<FSidebarDrawer>& Drawer : SidebarWidget->GetAllDrawers())
+	{
+		if (const TSharedPtr<SSidebarButton> TabButton = StaticCastSharedPtr<SSidebarButton>(Drawer->ButtonWidget))
+		{
+			TabButton->UpdateAppearance(OpenedDrawer);
+		}
+	}
+}
+
+void SSidebarContainer::DockDrawer_Internal(const TSharedRef<FSidebarDrawer>& InDrawer)
+{
+	for (const TSharedRef<FSidebarDrawer>& Drawer : SidebarWidget->GetAllDrawers())
+	{
+		CloseDrawer_Internal(Drawer, false);
+	}
+
+	InDrawer->bIsOpen = true;
+	InDrawer->State.bIsPinned = false;
+	InDrawer->State.bIsDocked = true;
+
+	Reconstruct();
+
+	UpdateDrawerTabAppearance();
+}
+
+void SSidebarContainer::UndockDrawer_Internal(const TSharedRef<FSidebarDrawer>& InDrawer)
+{
+	InDrawer->bIsOpen = false;
+	InDrawer->State.bIsDocked = false;
+
+	Reconstruct();
+
+	UpdateDrawerTabAppearance();
+}
+
+TSharedPtr<SSidebarDrawer> SSidebarContainer::FindOpenDrawerWidget(const TSharedRef<FSidebarDrawer>& InDrawer) const
+{
+	const TSharedRef<SSidebarDrawer>* const OpenDrawWidget = OpenDrawerWidgets.FindByPredicate(
+		[&InDrawer](const TSharedRef<SSidebarDrawer>& Drawer)
+		{
+			return InDrawer == Drawer->GetDrawer();
+		});
+	return OpenDrawWidget ? OpenDrawWidget->ToSharedPtr(): nullptr;
+}
+
+FName SSidebarContainer::GetOpenedDrawerId() const
+{
+	if (OpenDrawerWidgets.IsEmpty())
+	{
+		return NAME_None;
+	}
+	
+	const TSharedRef<SSidebarDrawer> LastOpenDrawerWidget = OpenDrawerWidgets.Last();
+	return LastOpenDrawerWidget->GetDrawer()->GetUniqueId();
+}
+
+TSharedPtr<FSidebarDrawer> SSidebarContainer::GetForegroundDrawer() const
+{
+	const int32 Index = OpenDrawerWidgets.FindLastByPredicate(
+		[](const TSharedRef<SSidebarDrawer>& InDrawerWidget)
+		{
+			return InDrawerWidget->IsOpen() && !InDrawerWidget->IsClosing();
+		});
+	return Index == INDEX_NONE ? nullptr : OpenDrawerWidgets[Index]->GetDrawer();
+}
+
+void SSidebarContainer::OnTabDrawerFocusLost(const TSharedRef<SSidebarDrawer>& InDrawerWidget)
+{
+	const TSharedPtr<FSidebarDrawer> Drawer = InDrawerWidget->GetDrawer();
+	if (!Drawer.IsValid())
+	{
+		return;
+	}
+
+	// Update to remove the focus marker
+	UpdateDrawerTabAppearance();
+
+	if (Drawer->State.bIsPinned)
+	{
+		return;
+	}
+
+	CloseDrawer_Internal(Drawer.ToSharedRef());
+}
+
+void SSidebarContainer::OnOpenAnimationFinish(const TSharedRef<SSidebarDrawer>& InDrawerWidget)
+{
+}
+
+void SSidebarContainer::OnCloseAnimationFinish(const TSharedRef<SSidebarDrawer>& InDrawerWidget)
+{
+	RemoveDrawerOverlaySlot(InDrawerWidget->GetDrawer().ToSharedRef(), false);
+}
+
+void SSidebarContainer::OnDrawerSizeChanged(const TSharedRef<SSidebarDrawer>& InDrawerWidget, const float InNewPixelSize)
+{
+	if (!DrawersOverlay.IsValid())
+	{
+		return;
+	}
+
+	const TSharedPtr<FSidebarDrawer> Drawer = InDrawerWidget->GetDrawer();
+	if (!Drawer.IsValid())
+	{
+		return;
+	}
+
+	const float DrawerOverlayWidth = GetOverlaySize().X;
+	const float FillPercent = InNewPixelSize / DrawerOverlayWidth;
+	SidebarSizePercent = FillPercent;
+
+	SidebarWidget->OnStateChanged.ExecuteIfBound(SidebarWidget->GetState());
+}
+
+TSharedPtr<FSidebarDrawer> SSidebarContainer::FindDrawer(const FName InDrawerId) const
+{
+	const TSharedRef<FSidebarDrawer>* const FoundDrawer = SidebarWidget->GetAllDrawers().FindByPredicate(
+		[InDrawerId](const TSharedRef<FSidebarDrawer>& InDrawer)
+		{
+			return InDrawerId == InDrawer->GetUniqueId();
+		});
+	return FoundDrawer ? *FoundDrawer : TSharedPtr<FSidebarDrawer>();
+}
+
+TSharedPtr<FSidebarDrawer> SSidebarContainer::FindFirstPinnedTab() const
+{
+	for (const TSharedRef<FSidebarDrawer>& Drawer : SidebarWidget->GetAllDrawers())
+	{
+		if (Drawer->State.bIsPinned)
+		{
+			return Drawer;
+		}
+	}
+	return nullptr;
+}
+
+void SSidebarContainer::OnContentSlotResizing(const float InFillPercent)
+{
+	ContentSizePercent = InFillPercent;
+}
+
+void SSidebarContainer::OnSidebarSlotResizing(const float InFillPercent)
+{
+	if (SidebarSizePercent < FSidebarState::AutoDockThresholdSize)
+	{
+		FSlateApplication::Get().ReleaseAllPointerCapture();
+
+		SidebarWidget->UndockAllDrawers();
+
+		bWantsToAutoDock = true;
+		ContentSizePercent = ContentSizeBeforeResize;
+		SidebarSizePercent = SidebarSizeBeforeResize;
+
+		FSidebarState NewState = SidebarWidget->GetState();
+		NewState.SetDrawerSizes(SidebarSizePercent, ContentSizePercent);
+		SidebarWidget->OnStateChanged.ExecuteIfBound(NewState);
+	}
+	else
+	{
+		// Save the current size when we start resizing
+		if (SidebarSizeBeforeResize == 0.f)
+		{
+			ContentSizeBeforeResize = ContentSizePercent;
+			SidebarSizeBeforeResize = SidebarSizePercent;
+		}
+
+		SidebarSizePercent = InFillPercent;
+	}
+}
+
+void SSidebarContainer::OnSplitterResized()
+{
+	bWantsToAutoDock = false;
+	ContentSizeBeforeResize = 0.f;
+	SidebarSizeBeforeResize = 0.f;
+
+	SidebarWidget->OnStateChanged.ExecuteIfBound(SidebarWidget->GetState());
+}
+
+#undef LOCTEXT_NAMESPACE
diff --git a/Engine/Source/Developer/ToolWidgets/Private/Sidebar/SSidebarDrawer.cpp b/Engine/Source/Developer/ToolWidgets/Private/Sidebar/SSidebarDrawer.cpp
--- a/Engine/Source/Developer/ToolWidgets/Private/Sidebar/SSidebarDrawer.cpp
+++ b/Engine/Source/Developer/ToolWidgets/Private/Sidebar/SSidebarDrawer.cpp
@@ -28,7 +28,7 @@ void SSidebarDrawer::Construct(const FArguments& InArgs, const TSharedRef<FSideb
 	ShadowOffset = InArgs._ShadowOffset;
 	ExpanderHandleSize = InArgs._ExpanderHandleSize;
 	
-	OnDrawerTargetSizeChanged = InArgs._OnDrawerTargetSizeChanged;
+	OnDrawerSizeChanged = InArgs._OnDrawerSizeChanged;
 	OnDrawerFocused = InArgs._OnDrawerFocused;
 	OnDrawerFocusLost = InArgs._OnDrawerFocusLost;
 	OnCloseAnimationFinish = InArgs._OnCloseAnimationFinish;
@@ -165,7 +165,7 @@ void SSidebarDrawer::OnArrangeChildren(const FGeometry& InAllottedGeometry, FArr
 	ArrangedChildren.AddWidget(InAllottedGeometry.MakeChild(ChildSlot.GetWidget(), ChildOffset, LocalSize));
 }
 
-FReply SSidebarDrawer::OnMouseButtonDown(const FGeometry& InAllottedGeometry, const FPointerEvent& InMouseEvent) 
+FReply SSidebarDrawer::OnMouseButtonDown(const FGeometry& InAllottedGeometry, const FPointerEvent& InMouseEvent)
 {
 	FReply Reply = FReply::Unhandled();
 
@@ -187,14 +187,15 @@ FReply SSidebarDrawer::OnMouseButtonDown(const FGeometry& InAllottedGeometry, co
 	return Reply;
 }
 
-FReply SSidebarDrawer::OnMouseButtonUp(const FGeometry& InAllottedGeometry, const FPointerEvent& InMouseEvent) 
+FReply SSidebarDrawer::OnMouseButtonUp(const FGeometry& InAllottedGeometry, const FPointerEvent& InMouseEvent)
 {
 	if (bIsResizing && InMouseEvent.GetEffectingButton() == EKeys::LeftMouseButton)
 	{
 		bIsResizing = false;
 		FSlateThrottleManager::Get().LeaveResponsiveMode(ResizeThrottleHandle);
 
-		OnDrawerTargetSizeChanged.ExecuteIfBound(SharedThis(this), TargetDrawerSize);
+		OnDrawerSizeChanged.ExecuteIfBound(SharedThis(this), TargetDrawerSize);
+
 		return FReply::Handled().ReleaseMouseCapture();
 	}
 
@@ -330,7 +331,7 @@ int32 SSidebarDrawer::OnPaint(const FPaintArgs& InArgs, const FGeometry& InAllot
 	// Main Shadow
 	FSlateDrawElement::MakeBox(
 		OutDrawElements,
-		InLayerId,
+		InLayerId++,
 		RenderTransformedChildGeometry.ToPaintGeometry(),
 		ShadowBrush,
 		ESlateDrawEffect::None,
@@ -339,7 +340,7 @@ int32 SSidebarDrawer::OnPaint(const FPaintArgs& InArgs, const FGeometry& InAllot
 	// Background
 	FSlateDrawElement::MakeBox(
 		OutDrawElements,
-		InLayerId,
+		InLayerId++,
 		OffsetPaintGeom,
 		BackgroundBrush,
 		ESlateDrawEffect::None,
@@ -548,7 +549,7 @@ void SSidebarDrawer::OnGlobalFocusChanging(const FFocusEvent& InFocusEvent
 	const TSharedPtr<FSidebarDrawer> Drawer = DrawerWeak.Pin();
 
 	// Only open drawers that are not docked or pinned need to close the drawer when focus is lost
-	if (!Drawer || Drawer->bIsDocked || Drawer->bIsPinned || !Drawer->bIsOpen)
+	if (!Drawer || Drawer->State.bIsDocked || Drawer->State.bIsPinned || !Drawer->bIsOpen)
 	{
 		return;
 	}
diff --git a/Engine/Source/Developer/ToolWidgets/Private/Sidebar/SSidebarDrawer.h b/Engine/Source/Developer/ToolWidgets/Private/Sidebar/SSidebarDrawer.h
--- a/Engine/Source/Developer/ToolWidgets/Private/Sidebar/SSidebarDrawer.h
+++ b/Engine/Source/Developer/ToolWidgets/Private/Sidebar/SSidebarDrawer.h
@@ -39,8 +39,8 @@ public:
 		SLATE_ARGUMENT(FVector2D, ShadowOffset)
 		/** The size of the handle used to resize the drawer */
 		SLATE_ARGUMENT(float, ExpanderHandleSize)
-		/** Called when the target drawer size is changed by the user */
-		SLATE_EVENT(FOnSidebarDrawerTargetSizeChanged, OnDrawerTargetSizeChanged)
+		/** Called when the drawer size is changed by the user */
+		SLATE_EVENT(FOnSidebarDrawerTargetSizeChanged, OnDrawerSizeChanged)
 		/** Called when the drawer gains focus */
 		SLATE_EVENT(FGenericSidebarDrawerWidgetDelegate, OnDrawerFocused)
 		/** Called when the drawer loses focus */
@@ -75,9 +75,10 @@ public:
 	/** @return True if the drawer is currently playing the close animation. */
 	bool IsClosing() const;
 
-	/** Sets the current size of the drawer, ignoring any open/close animation. */
+	/** Sets the current size of the drawer in pixels, ignoring any open/close animation. */
 	void SetCurrentSize(const float InSize);
 
+	/** @return The drawer associated with this drawer widget. */
 	TSharedPtr<FSidebarDrawer> GetDrawer() const;
 
 protected:
@@ -115,7 +116,7 @@ protected:
 	FVector2D ShadowOffset = FVector2D::ZeroVector;
 	float ExpanderHandleSize = 5.f;
 
-	FOnSidebarDrawerTargetSizeChanged OnDrawerTargetSizeChanged;
+	FOnSidebarDrawerTargetSizeChanged OnDrawerSizeChanged;
 	FGenericSidebarDrawerWidgetDelegate OnDrawerFocused;
 	FGenericSidebarDrawerWidgetDelegate OnDrawerFocusLost;
 	FGenericSidebarDrawerWidgetDelegate OnCloseAnimationFinish;
diff --git a/Engine/Source/Developer/ToolWidgets/Private/Sidebar/SSidebarDrawerContent.cpp b/Engine/Source/Developer/ToolWidgets/Private/Sidebar/SSidebarDrawerContent.cpp
--- a/Engine/Source/Developer/ToolWidgets/Private/Sidebar/SSidebarDrawerContent.cpp
+++ b/Engine/Source/Developer/ToolWidgets/Private/Sidebar/SSidebarDrawerContent.cpp
@@ -133,7 +133,7 @@ void SSidebarDrawerContent::BuildContent()
 		AddContentSlot(SectionPair.Value, 1.f);
 	}
 
-	if (Drawer->SelectedContentSections.IsEmpty())
+	if (Drawer->State.SelectedSections.IsEmpty())
 	{
 		// Find first section that is visible
 		FName FoundSectionName;
@@ -148,7 +148,7 @@ void SSidebarDrawerContent::BuildContent()
 
 		if (FoundSectionName != NAME_None)
 		{
-			Drawer->SelectedContentSections.Add(FoundSectionName);
+			Drawer->State.SelectedSections.Add(FoundSectionName);
 		}
 	}
 }
@@ -168,30 +168,30 @@ void SSidebarDrawerContent::OnSectionSelected(const ECheckBoxState InCheckBoxSta
 	{
 		if (bIsModifierDown)
 		{
-			Drawer->SelectedContentSections.Add(InSectionName);
+			Drawer->State.SelectedSections.Add(InSectionName);
 		}
 		else
 		{
-			Drawer->SelectedContentSections.Reset();
-			Drawer->SelectedContentSections.Add(InSectionName);
+			Drawer->State.SelectedSections.Reset();
+			Drawer->State.SelectedSections.Add(InSectionName);
 		}
 	}
 	else
 	{
 		if (bIsModifierDown)
 		{
-			Drawer->SelectedContentSections.Remove(InSectionName);
+			Drawer->State.SelectedSections.Remove(InSectionName);
 
 			// Force to always have a Selected Section, making it unable to de-select the last one
-			if (Drawer->SelectedContentSections.IsEmpty())
+			if (Drawer->State.SelectedSections.IsEmpty())
 			{
-				Drawer->SelectedContentSections.Add(InSectionName);
+				Drawer->State.SelectedSections.Add(InSectionName);
 			}
 		}
 		else
 		{
-			Drawer->SelectedContentSections.Reset();
-			Drawer->SelectedContentSections.Add(InSectionName);
+			Drawer->State.SelectedSections.Reset();
+			Drawer->State.SelectedSections.Add(InSectionName);
 		}
 	}
 }
@@ -199,7 +199,7 @@ void SSidebarDrawerContent::OnSectionSelected(const ECheckBoxState InCheckBoxSta
 bool SSidebarDrawerContent::IsSectionSelected(const FName InSectionName) const
 {
 	const TSharedPtr<FSidebarDrawer> Drawer = OwnerDrawerWeak.Pin();
-	return Drawer.IsValid() && Drawer->SelectedContentSections.Contains(InSectionName);
+	return Drawer.IsValid() && Drawer->State.SelectedSections.Contains(InSectionName);
 }
 
 bool SSidebarDrawerContent::ShouldShowSection(const TWeakPtr<ISidebarDrawerContent>& InSectionWeak) const
diff --git a/Engine/Source/Developer/ToolWidgets/Private/Sidebar/SidebarDrawer.h b/Engine/Source/Developer/ToolWidgets/Private/Sidebar/SidebarDrawer.h
--- a/Engine/Source/Developer/ToolWidgets/Private/Sidebar/SidebarDrawer.h
+++ b/Engine/Source/Developer/ToolWidgets/Private/Sidebar/SidebarDrawer.h
@@ -19,7 +19,6 @@ class FSidebarDrawer
 public:
 	FSidebarDrawer(FSidebarDrawerConfig&& InDrawerConfig)
 		: Config(MoveTemp(InDrawerConfig))
-		, SizeCoefficient(Config.DefaultSizeCoefficient)
 	{}
 
 	bool operator==(const FName InOtherId) const
@@ -59,16 +58,12 @@ public:
 	TSharedPtr<SWidget> ContentWidget;
 
 	TMap<FName, TSharedRef<ISidebarDrawerContent>> ContentSections;
-	TSet<FName> SelectedContentSections;
-
-	bool bIsOpen = false;
-	bool bIsPinned = false;
-	bool bIsDocked = false;
 
 	bool bDisablePin = false;
 	bool bDisableDock = false;
 
-	float SizeCoefficient;
+	bool bIsOpen = false;
+	FSidebarDrawerState State;
 
 	FOnSidebarDrawerOpened DrawerOpenedDelegate;
 	FOnSidebarDrawerClosed DrawerClosedDelegate;
diff --git a/Engine/Source/Developer/ToolWidgets/Private/Sidebar/SidebarState.cpp b/Engine/Source/Developer/ToolWidgets/Private/Sidebar/SidebarState.cpp
new file mode 100644
--- /dev/null
+++ b/Engine/Source/Developer/ToolWidgets/Private/Sidebar/SidebarState.cpp
@@ -0,0 +1,98 @@
+﻿// Copyright Epic Games, Inc. All Rights Reserved.
+
+#include "Sidebar/SidebarState.h"
+
+bool FSidebarState::IsValid() const
+{
+	return !(bHidden == false && DrawerSize == 0.f && DrawerStates.IsEmpty());
+}
+
+bool FSidebarState::IsHidden() const
+{
+	return bHidden;
+}
+
+bool FSidebarState::IsVisible() const
+{
+	return !bHidden;
+}
+
+void FSidebarState::SetHidden(const bool bInHidden)
+{
+	bHidden = bInHidden;
+}
+
+void FSidebarState::SetVisible(const bool bInVisible)
+{
+	bHidden = !bInVisible;
+}
+
+float FSidebarState::GetDrawerSize() const
+{
+	return (DrawerSize <= MinSize) ? DefaultSize : DrawerSize;
+}
+
+void FSidebarState::SetDrawerSize(const float InSize)
+{
+	DrawerSize = InSize;
+}
+
+void FSidebarState::GetDrawerSizes(float& OutDrawerSize, float& OutContentSize) const
+{
+	OutDrawerSize = (OutDrawerSize <= MinSize) ? DefaultSize : DrawerSize;
+	OutContentSize = (OutContentSize <= MinSize) ? DefaultSize : ContentSize;
+
+	const float TotalSize = OutDrawerSize + OutContentSize;
+	if (TotalSize < 1.f)
+	{
+		OutContentSize = 1.f - OutDrawerSize;
+	}
+}
+
+void FSidebarState::SetDrawerSizes(const float InDrawerSize, const float InContentSize)
+{
+	DrawerSize = InDrawerSize;
+	ContentSize = InContentSize;
+}
+
+const TArray<FSidebarDrawerState>& FSidebarState::GetDrawerStates() const
+{
+	return DrawerStates;
+}
+
+FSidebarDrawerState& FSidebarState::FindOrAddDrawerState(const FSidebarDrawerState& InDrawerState)
+{
+	for (FSidebarDrawerState& State : DrawerStates)
+	{
+		if (State.DrawerId == InDrawerState.DrawerId)
+		{
+			return State;
+		}
+	}
+	return DrawerStates.Add_GetRef(InDrawerState);
+}
+
+const FSidebarDrawerState* FSidebarState::FindDrawerState(const FSidebarDrawerState& InDrawerState)
+{
+	for (FSidebarDrawerState& State : DrawerStates)
+	{
+		if (State.DrawerId == InDrawerState.DrawerId)
+		{
+			return &State;
+		}
+	}
+	return nullptr;
+}
+
+void FSidebarState::SaveDrawerState(const FSidebarDrawerState& InState)
+{
+	for (FSidebarDrawerState& State : DrawerStates)
+	{
+		if (State.DrawerId.IsEqual(InState.DrawerId))
+		{
+			State = InState;
+			return;
+		}
+	}
+	DrawerStates.Add(InState);
+}
diff --git a/Engine/Source/Developer/ToolWidgets/Public/Sidebar/SSidebar.h b/Engine/Source/Developer/ToolWidgets/Public/Sidebar/SSidebar.h
--- a/Engine/Source/Developer/ToolWidgets/Public/Sidebar/SSidebar.h
+++ b/Engine/Source/Developer/ToolWidgets/Public/Sidebar/SSidebar.h
@@ -3,6 +3,7 @@
 #pragma once
 
 #include "SidebarDrawerConfig.h"
+#include "Framework/SlateDelegates.h"
 #include "Widgets/DeclarativeSyntaxSupport.h"
 #include "Widgets/Layout/SSplitter.h"
 #include "Widgets/SCompoundWidget.h"
@@ -14,10 +15,12 @@ class SBox;
 class SOverlay;
 class SScrollBox;
 class SSidebarButton;
+class SSidebarContainer;
 class SSidebarDrawer;
 class STabDrawer;
 class SVerticalBox;
 class UToolMenu;
+struct FSidebarState;
 struct FTabId;
 
 /**
@@ -37,7 +40,7 @@ enum class ESidebarTabLocation : uint8
 	Bottom
 };
 
-DECLARE_DELEGATE_OneParam(FOnSidebarDrawerDockStateChanged, const FName /*InDrawerId*/);
+DECLARE_DELEGATE_OneParam(FOnSidebarStateChanged, const FSidebarState& /*InNewState*/);
 
 /**
  * Static sidebar tab widget that cannot be dragged or moved to a different location. Multiple drawers can be registered
@@ -47,29 +50,34 @@ DECLARE_DELEGATE_OneParam(FOnSidebarDrawerDockStateChanged, const FName /*InDraw
 class TOOLWIDGETS_API SSidebar : public SCompoundWidget
 {
 public:
+	friend class SSidebarContainer;
+
 	static constexpr float MinTabButtonSize = 100.f;
 	static constexpr float MaxTabButtonSize = 200.f;
 	static constexpr float TabButtonThickness = 25.f;
 
 	SLATE_BEGIN_ARGS(SSidebar)
-		: _TabLocation(ESidebarTabLocation::Right)
-		, _HideWhenDocked(false)
+		: _HideWhenAllDocked(false)
 		, _AlwaysUseMaxButtonSize(false)
 		, _DisablePin(false)
 		, _DisableDock(false)
 	{}
 		/** The direction that a tab drawer opens relative to the location of the sidebar. */
 		SLATE_ARGUMENT(ESidebarTabLocation, TabLocation)
-		/** Hides the sidebar when a drawer is docked. NOTE: Must provide a way to manually undock the drawer to restore the sidebar visibility. */
-		SLATE_ARGUMENT(bool, HideWhenDocked)
+		/** The initial size of a drawer for the sidebar. */
+		SLATE_ARGUMENT(float, InitialDrawerSize)
+		/** Delegate used to retrieve the main inner content of the sidebar that will be overlayed. */
+		SLATE_ARGUMENT(FOnGetContent, OnGetContent)
+		/** Hides the sidebar when all drawers is docked. NOTE: Must provide a way to manually undock the drawer to restore the sidebar visibility. */
+		SLATE_ARGUMENT(bool, HideWhenAllDocked)
 		/** Forces the sidebar tab buttons to always be a uniform size of max. */
 		SLATE_ARGUMENT(bool, AlwaysUseMaxButtonSize)
 		/** Disables the ability to pin a drawer. */
 		SLATE_ARGUMENT(bool, DisablePin)
 		/** Disables the ability to dock a drawer. */
 		SLATE_ARGUMENT(bool, DisableDock)
-		/** Event triggered when a drawers dock state changes. */
-		SLATE_EVENT(FOnSidebarDrawerDockStateChanged, OnDockStateChanged)
+		/** Event triggered when the sidebar state changes. */
+		SLATE_EVENT(FOnSidebarStateChanged, OnStateChanged)
 	SLATE_END_ARGS()
 
 	virtual ~SSidebar() override;
@@ -78,10 +86,10 @@ public:
 	 * Constructs the sidebar widget.
 	 * 
 	 * @param InArgs Widget construction arguments
-	 * @param InDrawersOverlay Overlay widget used to display the animating drawer
+	 * @param InContainerWidget Container widget used to manage the drawer sidebar widgets
 	 * @param InDockLocation Parent widget that will contain the drawer content widget when docked
 	 */
-	void Construct(const FArguments& InArgs, const TSharedRef<SOverlay>& InDrawersOverlay, const TSharedRef<SBox>& InDockLocation);
+	void Construct(const FArguments& InArgs, const TSharedRef<SSidebarContainer>& InContainerWidget);
 
 	/**
 	 * Registers and displays a new drawer in the sidebar.
@@ -182,6 +190,9 @@ public:
 	 */
 	bool IsDrawerPinned(const FName InDrawerId) const;
 
+	/** @return List of drawer Ids that are pinned. */
+	TSet<FName> GetPinnedDrawerIds() const;
+
 	/**
 	 * Pins a drawer so it stays open even when focus is lost.
 	 * 
@@ -196,6 +207,9 @@ public:
 	/** @return True if the specified drawer Id is docked. */
 	bool IsDrawerDocked(const FName InDrawerId) const;
 
+	/** @return List of drawer Ids that are docked. */
+	TSet<FName> GetDockedDrawerIds() const;
+
 	/**
 	 * Docks a drawer so it embeds itself into the content.
 	 * 
@@ -207,17 +221,8 @@ public:
 	/** Undocks any drawers that are docked. */
 	void UndockAllDrawers();
 
-	/**
-	 * Helper function to update a splitter slot size based on a drawers state.
-	 * Sets the slots resize-ability, sizing rule, size value.
-	 * 
-	 * @param InDrawerId Unique Id of the drawer to update for
-	 * @param InSlot The splitter slot to modify the size of
-	 * @param bInAutoUndock Automatically undocks the drawer if the size is below the threshold
-	 * @param InDefaultDockPercent Default size co-efficient to reset the drawer size to
-	 */
-	void UpdateDockedSplitterSlot(const FName InDrawerId, SSplitter::FSlot* const InSlot, const bool bInAutoUndock = true
-		, const float InDefaultDockPercent = FSidebarDrawerConfig::DefaultSizeCoefficient);
+	/** Unpins any drawers that are pinned. */
+	void UnpinAllDrawers();
 
 	/** @return True if the sidebar is set to animate horizontally. */
 	bool IsHorizontal() const;
@@ -225,14 +230,19 @@ public:
 	/** @return True if the sidebar is set to animate vertically. */
 	bool IsVertical() const;
 
+	/** @return The current state of the sidebar to save/restore. */
+	FSidebarState GetState() const;
+
+	/** @return The location of the sidebar (left, right, top, bottom). */
+	ESidebarTabLocation GetTabLocation() const;
+
+	TSharedRef<SWidget> GetMainContent() const;
+
 private:
 	void OnTabDrawerButtonPressed(const TSharedRef<FSidebarDrawer>& InDrawer);
 	void OnDrawerTabPinToggled(const TSharedRef<FSidebarDrawer>& InDrawer, const bool bIsPinned);
 	void OnDrawerTabDockToggled(const TSharedRef<FSidebarDrawer>& InDrawer, const bool bIsDocked);
-	void OnTabDrawerFocusLost(const TSharedRef<SSidebarDrawer>& InDrawerWidget);
-	void OnOpenAnimationFinish(const TSharedRef<SSidebarDrawer>& InDrawerWidget);
-	void OnCloseAnimationFinish(const TSharedRef<SSidebarDrawer>& InDrawerWidget);
-	void OnDrawerTargetSizeChanged(const TSharedRef<SSidebarDrawer>& InDrawerWidget, const float InNewSize);
+
 	TSharedRef<SWidget> OnGetTabDrawerContextMenuWidget(TSharedRef<FSidebarDrawer> InDrawer);
 	void BuildOptionsMenu(UToolMenu* const InMenu);
 
@@ -242,54 +252,25 @@ private:
 	/** Removes all drawers instantly (including drawers for pinned tabs). */
 	void RemoveAllDrawers();
 
-	EActiveTimerReturnType OnOpenPendingDrawerTimer(const double InCurrentTime, const float InDeltaTime);
-	void OpenDrawerNextFrame(const TSharedRef<FSidebarDrawer>& InDrawer, const bool bInAnimate = true);
-	void OpenDrawerInternal(const TSharedRef<FSidebarDrawer>& InDrawer, const bool bInAnimate = true);
-	void CloseDrawerInternal(const TSharedRef<FSidebarDrawer>& InDrawer, const bool bInAnimate = true);
-
-	/** Reopens the pinned tab only if there are no other open drawers. This should be used to bring pinned tabs back after other tabs lose focus/are closed. */
-	void SummonPinnedTabIfNothingOpened();
-
-	/** Updates the appearance of open drawers. */
-	void UpdateDrawerAppearance();
-
 	TSharedPtr<FSidebarDrawer> FindDrawer(const FName InDrawerId) const;
 
-	/** Returns the first tab in this sidebar that is marked pinned. */
-	TSharedPtr<FSidebarDrawer> FindFirstPinnedTab() const;
+	void SetWidgetDrawerSize(const TSharedRef<FSidebarDrawer>& InDrawer);
 
-	/** Returns the tab for the last-opened drawer that is still open, excluding any drawers that are in the process of closing. */
-	TSharedPtr<FSidebarDrawer> GetForegroundTab() const;
+	bool AreAllDrawersDocked() const;
 
-	/** Returns the drawer for the given tab if it's open. */
-	TSharedPtr<SSidebarDrawer> FindOpenDrawerWidget(const TSharedRef<FSidebarDrawer>& InDrawer) const;
+	const TArray<TSharedRef<FSidebarDrawer>>& GetAllDrawers() const;
 
-	TWeakPtr<SOverlay> DrawersOverlayWeak;
-	TWeakPtr<SBox> DockLocationWeak;
+	TSharedPtr<SSidebarContainer> ContainerWidget;
 
 	ESidebarTabLocation TabLocation = ESidebarTabLocation::Right;
-	bool bHideWhenDocked = false;
+	FOnGetContent OnGetContent;
+	bool bHideWhenAllDocked = false;
 	bool bAlwaysUseMaxButtonSize = false;
 	bool bDisablePin = false;
 	bool bDisableDock = false;
-	FOnSidebarDrawerDockStateChanged OnDockStateChanged;
+	FOnSidebarStateChanged OnStateChanged;
 
 	TSharedPtr<SScrollBox> TabButtonContainer;
 
 	TArray<TSharedRef<FSidebarDrawer>> Drawers;
-
-	/** Generally speaking one drawer is only ever open at once but we animate any previous drawer
-	 * closing so there could be more than one while an animation is playing. A docked drawer is
-	 * also considered open, along with any user opened/pinned drawers. */
-	TArray<TSharedRef<SSidebarDrawer>> OpenDrawerWidgets;
-
-	TArray<TSharedRef<SSidebarDrawer>> ClosingDrawerWidgets;
-
-	TWeakPtr<FSidebarDrawer> PendingTabToOpen;
-	bool bAnimatePendingTabOpen = false;
-	TSharedPtr<FActiveTimerHandle> OpenPendingDrawerTimerHandle;
-
-	TArray<TSharedRef<FSidebarDrawer>> PinnedDrawerTabs;
-
-	TSharedPtr<FSidebarDrawer> DockedDrawerTab;
 };
diff --git a/Engine/Source/Developer/ToolWidgets/Public/Sidebar/SSidebarContainer.h b/Engine/Source/Developer/ToolWidgets/Public/Sidebar/SSidebarContainer.h
new file mode 100644
--- /dev/null
+++ b/Engine/Source/Developer/ToolWidgets/Public/Sidebar/SSidebarContainer.h
@@ -0,0 +1,130 @@
+﻿// Copyright Epic Games, Inc. All Rights Reserved.
+
+#pragma once
+
+#include "Sidebar/SidebarState.h"
+#include "Sidebar/SSidebar.h"
+#include "Widgets/DeclarativeSyntaxSupport.h"
+#include "Widgets/SCompoundWidget.h"
+
+class SSidebar;
+class SSidebarDrawer;
+
+/**
+ * A container for a sidebar widget that manages the slider drawer overlay widgets
+ * and a default docking location for all drawers.
+ */
+class TOOLWIDGETS_API SSidebarContainer : public SCompoundWidget
+{
+public:
+	friend class SSidebar;
+
+	SLATE_BEGIN_ARGS(SSidebarContainer)
+	{}
+	SLATE_END_ARGS()
+
+	void Construct(const FArguments& InArgs);
+
+	void RebuildSidebar(const TSharedRef<SSidebar>& InSidebarWidget, const FSidebarState& InState);
+
+	float GetContentSlotSize() const;
+	float GetSidebarSlotSize() const;
+
+	EOrientation GetSplitterOrientation() const;
+	ESidebarTabLocation GetTabLocation() const;
+	float GetCurrentDrawerSize() const;
+	FVector2D GetOverlaySize() const;
+
+	void CloseAllDrawerWidgets(const bool bInAnimate);
+
+	/** Reopens the pinned tab only if there are no other open drawers. This should be used to bring pinned tabs back after other tabs lose focus/are closed. */
+	void SummonPinnedTabIfNothingOpened();
+
+	/** Updates the appearance of drawer tabs. */
+	void UpdateDrawerTabAppearance();
+
+	FName GetOpenedDrawerId() const;
+
+protected:
+	void Reconstruct(const FSidebarState& InState = FSidebarState());
+
+	TSharedRef<SWidget> ConstructBoxPanel(const FSidebarState& InState);
+	void ConstructSplitterPanel(const FSidebarState& InState);
+
+	FMargin CalculateSlotMargin() const;
+
+	void CreateDrawerWidget(const TSharedRef<FSidebarDrawer>& InDrawer);
+
+	TSharedRef<SWidget> GetSidebarDrawerContent(const TSharedRef<FSidebarDrawer>& InDrawer) const;
+
+	bool AddDrawerOverlaySlot(const TSharedRef<FSidebarDrawer>& InDrawer);
+	bool RemoveDrawerOverlaySlot(const TSharedRef<FSidebarDrawer>& InDrawer, const bool bInAnimate);
+
+	void AddContentDockSlot();
+	void RemoveContentDockSlot();
+
+	void AddSidebarDockSlot(const FName InDockDrawerId);
+	void RemoveSidebarDockSlot();
+
+	void OnTabDrawerFocusLost(const TSharedRef<SSidebarDrawer>& InDrawerWidget);
+	void OnOpenAnimationFinish(const TSharedRef<SSidebarDrawer>& InDrawerWidget);
+	void OnCloseAnimationFinish(const TSharedRef<SSidebarDrawer>& InDrawerWidget);
+	void OnDrawerSizeChanged(const TSharedRef<SSidebarDrawer>& InDrawerWidget, const float InNewPixelSize);
+
+	EActiveTimerReturnType OnOpenPendingDrawerTimer(const double InCurrentTime, const float InDeltaTime);
+
+	void OpenDrawerNextFrame(const TSharedRef<FSidebarDrawer>& InDrawer, const bool bInAnimate = true);
+	void OpenDrawer_Internal(const TSharedRef<FSidebarDrawer>& InDrawer, const bool bInAnimate = true);
+	void CloseDrawer_Internal(const TSharedRef<FSidebarDrawer>& InDrawer, const bool bInAnimate = true);
+
+	void DockDrawer_Internal(const TSharedRef<FSidebarDrawer>& InDrawer);
+	void UndockDrawer_Internal(const TSharedRef<FSidebarDrawer>& InDrawer);
+
+	TSharedPtr<FSidebarDrawer> FindDrawer(const FName InDrawerId) const;
+	TSharedPtr<FSidebarDrawer> FindFirstPinnedTab() const;
+
+	TSharedPtr<SSidebarDrawer> FindOpenDrawerWidget(const TSharedRef<FSidebarDrawer>& InDrawer) const;
+
+	TSharedPtr<FSidebarDrawer> GetForegroundDrawer() const;
+
+	/** Event that occurs when the docked main content slot is being resized by the user (while mouse down). */
+	void OnContentSlotResizing(const float InFillPercent);
+	/** Event that occurs when the docked sidebar slot is being resized by the user (while mouse down). */
+	void OnSidebarSlotResizing(const float InFillPercent);
+
+	/** Event that occurs when the main splitter widget has finished being resized by the user (mouse up). */
+	void OnSplitterResized();
+
+	int32 GetContentSlotIndex() const;
+	int32 GetSidebarSlotIndex() const;
+
+	/** The sidebar widget associated with this container. One sidebar widget per container. */
+	TSharedPtr<SSidebar> SidebarWidget;
+
+	/** The main splitter widget used when a drawer is docked. */
+	TSharedPtr<SSplitter> MainSplitter;
+
+	/** Overlay used to draw drawer widgets on top of the rest of the content. */
+	TSharedPtr<SOverlay> DrawersOverlay;
+
+	/** Generally speaking one drawer is only ever open at once but we animate any previous drawer
+	 * closing so there could be more than one while an animation is playing. A docked drawer is
+	 * also considered open, along with any user opened/pinned drawers. */
+	TArray<TSharedRef<SSidebarDrawer>> OpenDrawerWidgets;
+
+	TArray<TSharedRef<SSidebarDrawer>> ClosingDrawerWidgets;
+
+	TWeakPtr<FSidebarDrawer> PendingTabToOpen;
+	bool bAnimatePendingTabOpen = false;
+	TSharedPtr<FActiveTimerHandle> OpenPendingDrawerTimerHandle;
+
+	float ContentSizePercent = 0.8f;
+	float SidebarSizePercent = 0.2f;
+
+	TAttribute<float> ContentSlotSize;
+	TAttribute<float> SidebarSlotSize;
+
+	bool bWantsToAutoDock = false;
+	float ContentSizeBeforeResize = 0.f;
+	float SidebarSizeBeforeResize = 0.f;
+};
diff --git a/Engine/Source/Developer/ToolWidgets/Public/Sidebar/SidebarDrawerConfig.h b/Engine/Source/Developer/ToolWidgets/Public/Sidebar/SidebarDrawerConfig.h
--- a/Engine/Source/Developer/ToolWidgets/Public/Sidebar/SidebarDrawerConfig.h
+++ b/Engine/Source/Developer/ToolWidgets/Public/Sidebar/SidebarDrawerConfig.h
@@ -5,6 +5,7 @@
 #include "Containers/Array.h"
 #include "Internationalization/Text.h"
 #include "Misc/Attribute.h"
+#include "SidebarState.h"
 #include "Templates/SharedPointer.h"
 #include "UObject/NameTypes.h"
 
@@ -14,8 +15,6 @@ struct FSlateBrush;
 /** Configuration information used to register a sidebar drawer. */
 struct TOOLWIDGETS_API FSidebarDrawerConfig
 {
-	static constexpr float DefaultSizeCoefficient = 0.25f;
-
 	bool operator==(const FName InOtherId) const
 	{
 		return UniqueId == InOtherId;
@@ -48,11 +47,8 @@ struct TOOLWIDGETS_API FSidebarDrawerConfig
 	/** Icon to display for the drawer tab button */
 	TAttribute<const FSlateBrush*> Icon;
 
-	/** True if the drawers initial state should be docked. */
-	bool bInitiallyDocked = false;
-
-	/** Size co-efficient from 0 to 1 to use for the size of the docked content. */
-	float SizeCoefficient = DefaultSizeCoefficient;
+	/** The state of the sidebar drawer. (Visible, Pinned, Docked, Size) */
+	FSidebarDrawerState InitialState;
 
 	/** Optional content widget to use instead of registering sections. */
 	TSharedPtr<SWidget> OverrideContentWidget;
diff --git a/Engine/Source/Developer/ToolWidgets/Public/Sidebar/SidebarState.h b/Engine/Source/Developer/ToolWidgets/Public/Sidebar/SidebarState.h
new file mode 100644
--- /dev/null
+++ b/Engine/Source/Developer/ToolWidgets/Public/Sidebar/SidebarState.h
@@ -0,0 +1,81 @@
+﻿// Copyright Epic Games, Inc. All Rights Reserved.
+
+#pragma once
+
+#include "SidebarState.generated.h"
+
+/** Represents the state of a sidebar drawer to be saved/restored to/from config. */
+USTRUCT()
+struct TOOLWIDGETS_API FSidebarDrawerState
+{
+	GENERATED_BODY()
+
+	FSidebarDrawerState() {}
+	FSidebarDrawerState(const FName InDrawerId)
+	{
+		DrawerId = InDrawerId;
+	}
+
+	UPROPERTY()
+	FName DrawerId = NAME_None;
+
+	/** Names of all sections that were last selected */
+	UPROPERTY()
+	TSet<FName> SelectedSections;
+
+	UPROPERTY()
+	bool bIsPinned = false;
+
+	UPROPERTY()
+	bool bIsDocked = false;
+};
+
+/** Represents the state of a sidebar to be saved/restored to/from config. */
+USTRUCT()
+struct TOOLWIDGETS_API FSidebarState
+{
+	GENERATED_BODY()
+
+public:
+	static constexpr float DefaultSize = 0.25f;
+	static constexpr float MinSize = 0.005f;
+	static constexpr float MaxSize = 0.5f;
+	static constexpr float AutoDockThresholdSize = 0.05f;
+
+	/** @return True if any property has been changed from default */
+	bool IsValid() const;
+
+	bool IsHidden() const;
+	bool IsVisible() const;
+
+	void SetHidden(const bool bInHidden);
+	void SetVisible(const bool bInVisible);
+
+	float GetDrawerSize() const;
+	void SetDrawerSize(const float InSize);
+
+	void GetDrawerSizes(float& OutDrawerSize, float& OutContentSize) const;
+	void SetDrawerSizes(const float InDrawerSize, const float InContentSize);
+
+	const TArray<FSidebarDrawerState>& GetDrawerStates() const;
+
+	FSidebarDrawerState& FindOrAddDrawerState(const FSidebarDrawerState& InDrawerState);
+	const FSidebarDrawerState* FindDrawerState(const FSidebarDrawerState& InDrawerState);
+
+	/** Saves the state of a drawer. If the drawers state already exists in config, it will be replaced. */
+	void SaveDrawerState(const FSidebarDrawerState& InState);
+
+protected:
+	UPROPERTY()
+	bool bHidden = false;
+
+	UPROPERTY()
+	float DrawerSize = 0.f;
+
+	/** Save the other splitter slot size to exactly restore the size when a drawer is docked in a SSplitter widget. */
+	UPROPERTY()
+	float ContentSize = 0.f;
+
+	UPROPERTY()
+	TArray<FSidebarDrawerState> DrawerStates;
+};
diff --git a/Engine/Source/Editor/Sequencer/Private/SSequencer.cpp b/Engine/Source/Editor/Sequencer/Private/SSequencer.cpp
--- a/Engine/Source/Editor/Sequencer/Private/SSequencer.cpp
+++ b/Engine/Source/Editor/Sequencer/Private/SSequencer.cpp
@@ -128,6 +128,7 @@
 #include "SequencerUtilities.h"
 #include "Sidebar/SidebarDrawerConfig.h"
 #include "Sidebar/SSidebar.h"
+#include "Sidebar/SSidebarContainer.h"
 #include "Widgets/SOverlay.h"
 #include "IContentBrowserSingleton.h"
 #include "ContentBrowserModule.h"
@@ -159,8 +160,6 @@ void SSequencer::Construct(const FArguments& InArgs, TSharedRef<FSequencer> InSe
 	
 	InitializeTrackFilters();
 
-	ISequencerWidgetsModule& SequencerWidgets = FModuleManager::Get().LoadModuleChecked<ISequencerWidgetsModule>( "SequencerWidgets" );
-
 	OnPlaybackRangeBeginDrag = InArgs._OnPlaybackRangeBeginDrag;
 	OnPlaybackRangeEndDrag = InArgs._OnPlaybackRangeEndDrag;
 	OnSelectionRangeBeginDrag = InArgs._OnSelectionRangeBeginDrag;
@@ -254,115 +253,16 @@ void SSequencer::Construct(const FArguments& InArgs, TSharedRef<FSequencer> InSe
 		TimeSliderArgs.NumericTypeInterface = GetNumericTypeInterface();
 	}
 
-	TimeSliderController = MakeShareable( new FSequencerTimeSliderController( TimeSliderArgs, SequencerPtr ) );
-	
-	TSharedRef<FSequencerTimeSliderController> TimeSliderControllerRef = TimeSliderController.ToSharedRef();
-
-	bool bMirrorLabels = false;
-	
-	// Create the top and bottom sliders
-	TopTimeSlider = SequencerWidgets.CreateTimeSlider( TimeSliderControllerRef, bMirrorLabels );
-	bMirrorLabels = true;
-	TSharedRef<ITimeSlider> BottomTimeSlider = SequencerWidgets.CreateTimeSlider( TimeSliderControllerRef, TAttribute<EVisibility>(this, &SSequencer::GetBottomTimeSliderVisibility), bMirrorLabels );
-
-	EShowRange Ranges = EShowRange(EShowRange::WorkingRange | EShowRange::ViewRange);
-	if(InArgs._ShowPlaybackRangeInTimeSlider)
-	{
-		Ranges |= EShowRange::PlaybackRange;
-	}
-
-	// Create bottom time range slider
-	TSharedRef<ITimeSlider> BottomTimeRange = SequencerWidgets.CreateTimeRange(
-		FTimeRangeArgs(
-			Ranges,
-			TimeSliderControllerRef,
-			TAttribute<EVisibility>(this, &SSequencer::GetTimeRangeVisibility),
-			NumericTypeInterface.ToSharedRef()
-		),
-		SequencerWidgets.CreateTimeRangeSlider(TimeSliderControllerRef)
-	);
-
 	OnGetPlaybackSpeeds = InArgs._OnGetPlaybackSpeeds;
 
 	RootCustomization.AddMenuExtender = InArgs._AddMenuExtender;
 	RootCustomization.ToolbarExtender = InArgs._ToolbarExtender;
 
-	ColumnFillCoefficients[0] = 0.3f;
-	ColumnFillCoefficients[1] = 0.7f;
-
-	if (GetSequencerSettings())
-	{
-		const float TreeViewWidth = GetSequencerSettings()->GetTreeViewWidth();
-		const float TimelineWidth = 1.f - TreeViewWidth;
-
-		if (TreeViewWidth > 0.f && TimelineWidth > 0.f)
-		{
-			ColumnFillCoefficients[0] = TreeViewWidth;
-			ColumnFillCoefficients[1] = TimelineWidth;
-		}
-	}
-
-	TAttribute<float> FillCoefficient_0, FillCoefficient_1;
-	{
-		FillCoefficient_0.Bind(TAttribute<float>::FGetter::CreateSP(this, &SSequencer::GetColumnFillCoefficient, 0));
-		FillCoefficient_1.Bind(TAttribute<float>::FGetter::CreateSP(this, &SSequencer::GetColumnFillCoefficient, 1));
-	}
-
-	TSharedRef<SScrollBar> ScrollBar = SNew(SScrollBar)
-		.Thickness(FVector2D(9.0f, 9.0f));
-
-	TSharedRef<SScrollBar> PinnedAreaScrollBar = SNew(SScrollBar)
-		.Thickness(FVector2D(9.0f, 9.0f));
-
-	TSharedPtr<FSequencerEditorViewModel> SequencerViewModel = InSequencer->GetViewModel()->CastThisShared<FSequencerEditorViewModel>();
-	SAssignNew(PinnedTrackArea, SSequencerTrackAreaView, SequencerViewModel->GetPinnedTrackArea(), TimeSliderControllerRef);
-	SAssignNew(PinnedTreeView, SSequencerOutlinerView, SequencerViewModel->GetOutliner(), PinnedTrackArea.ToSharedRef())
-		.Selection(InSequencer->GetViewModel()->GetSelection())
-		.ExternalScrollbar(PinnedAreaScrollBar)
-		.Clipping(EWidgetClipping::ClipToBounds);
-
-	PinnedTrackArea->SetOutliner(PinnedTreeView);
-	PinnedTrackArea->SetShowPinned(true);
-	PinnedTrackArea->SetIsPinned(true);
-	PinnedTreeView->SetShowPinned(true);
-
-	SAssignNew(TrackArea, SSequencerTrackAreaView, SequencerViewModel->GetTrackArea(), TimeSliderControllerRef);
-	SAssignNew(TreeView, SSequencerOutlinerView, SequencerViewModel->GetOutliner(), TrackArea.ToSharedRef())
-		.Selection(InSequencer->GetViewModel()->GetSelection())
-		.ExternalScrollbar(ScrollBar)
-		.Clipping(EWidgetClipping::ClipToBounds);
-
-	TrackArea->SetOutliner(TreeView);
-
-	TreeView->AddPinnedTreeView(PinnedTreeView);
-
-	if (GetSequencerSettings())
-	{
-		InitializeOutlinerColumns();
-	}
-
-	SequencerViewModel->GetTrackArea()->CastThisChecked<FSequencerTrackAreaViewModel>()->InitializeDefaultEditTools(*TrackArea);
-	SequencerViewModel->GetPinnedTrackArea()->CastThisChecked<FSequencerTrackAreaViewModel>()->InitializeDefaultEditTools(*PinnedTrackArea);
-
-
-	if (USequencerSettings* Settings = GetSequencerSettings())
-	{
-		SequencerViewModel->SetViewDensity(Settings->GetViewDensity());
-	}
-
 	PlayTimeDisplay = StaticCastSharedRef<STemporarilyFocusedSpinBox<double>>(SequencerPtr.Pin()->MakePlayTimeDisplay(GetNumericTypeInterface()));
 
-	const TSharedRef<SOverlay> SidebarOverlay = SNew(SOverlay);
-
-	DetailsDockLocation = SNew(SBox);
-
 	TAttribute<FAnimatedRange> ViewRangeAttribute = InArgs._ViewRange;
 
-	const int32 Column0 = 0, Column1 = 1;
-	const int32 Row0 = 0, Row1 = 1, Row2 = 2, Row3 = 3, Row4 = 4;
-
-	const float CommonPadding = 3.f;
-	const FMargin ResizeBarPadding(4.f, 0, 0, 0);
+	GridPanel = ConstructTrackAreaGridPanel(InArgs, InSequencer, TimeSliderArgs);
 
 	ChildSlot
 	[
@@ -506,7 +406,7 @@ void SSequencer::Construct(const FArguments& InArgs, TSharedRef<FSequencer> InSe
 					]
 
 					// Sequence Locking symbol
-					+SHorizontalBox::Slot()
+					+ SHorizontalBox::Slot()
 					.HAlign(HAlign_Right)
 					.VAlign(VAlign_Center)
 					.AutoWidth()
@@ -530,435 +430,554 @@ void SSequencer::Construct(const FArguments& InArgs, TSharedRef<FSequencer> InSe
 			]
 		]
 
+		// Main content body
 		+ SVerticalBox::Slot()
 		[
-			// Overlay for sequencer details sliding drawers
-			SNew(SOverlay)
+			SAssignNew(MainContentContainer, SBox)
+			[
+				ConstructMainContent()
+			]
+		]
+	];
+
+	if (InSequencer->GetHostCapabilities().bSupportsCurveEditor)
+	{
+		TSharedPtr<FSequencerEditorViewModel> SequencerViewModel = InSequencer->GetViewModel()->CastThisShared<FSequencerEditorViewModel>();
+		FCurveEditorExtension* CurveEditorExtension = SequencerViewModel->CastDynamicChecked<FCurveEditorExtension>();
+		CurveEditorExtension->CreateCurveEditor(TimeSliderArgs);
+	}
+
+	ApplySequencerCustomization(RootCustomization);
+
+	InSequencer->GetViewModel()->GetSelection()->KeySelection.OnChanged.AddSP(this, &SSequencer::HandleKeySelectionChanged);
+	InSequencer->GetViewModel()->GetSelection()->Outliner.OnChanged.AddSP(this, &SSequencer::HandleOutlinerNodeSelectionChanged);
+
+	ResetBreadcrumbs();
+}
+
+TSharedRef<SWidget> SSequencer::ConstructMainContent()
+{
+	const TSharedPtr<FSequencer> Sequencer = SequencerPtr.Pin();
+	check(Sequencer.IsValid());
+	const bool bSupportsSidebar = Sequencer->GetHostCapabilities().bSupportsSidebar;
+
+	USequencerSettings* const SequencerSettings = GetSequencerSettings();
+	check(IsValid(SequencerSettings));
+	const FSidebarState& SidebarState = SequencerSettings->GetSidebarState();
+
+	TSharedPtr<SWidget> OutWidget;
+
+	// Create the details sidebar only once to avoid having to re-register drawers
+	if (bSupportsSidebar)
+	{
+		if (!DetailsSidebar.IsValid())
+		{
+			SidebarContainer = SNew(SSidebarContainer);
+
+			DetailsSidebar = SNew(SSidebar, SidebarContainer.ToSharedRef())
+				.TabLocation(ESidebarTabLocation::Right)
+				.InitialDrawerSize(SidebarState.GetDrawerSize())
+				.OnStateChanged(this, &SSequencer::OnSidebarStateChanged)
+				.OnGetContent(FOnGetContent::CreateLambda([this]()
+					{
+						return ConstructGridPanel();
+					}));
+		}
+
+		SidebarContainer->RebuildSidebar(DetailsSidebar.ToSharedRef(), SidebarState);
+	}
+	else
+	{
+		DetailsSidebar.Reset();
+	}
+
+	if (SidebarState.IsHidden() || !bSupportsSidebar)
+	{
+		OutWidget = ConstructGridPanel();
+	}
+	else if (SidebarState.IsVisible())
+	{
+		OutWidget = SidebarContainer.ToSharedRef();
+	}
+
+	return OutWidget.ToSharedRef();
+}
 
-			+ SOverlay::Slot()
+TSharedRef<SWidget> SSequencer::ConstructGridPanel()
+{
+	TAttribute<float> FillCoefficient_0, FillCoefficient_1;
+	{
+		FillCoefficient_0.Bind(TAttribute<float>::FGetter::CreateSP(this, &SSequencer::GetColumnFillCoefficient, 0));
+		FillCoefficient_1.Bind(TAttribute<float>::FGetter::CreateSP(this, &SSequencer::GetColumnFillCoefficient, 1));
+	}
+
+	return
+		SNew(SOverlay)
+		+ SOverlay::Slot()
+		[
+			GridPanel.ToSharedRef()
+		]
+		+ SOverlay::Slot()
+		[
+			// track area virtual splitter overlay
+			SAssignNew(TreeViewSplitter, SSequencerSplitterOverlay)
+			.Style(FAppStyle::Get(), TEXT("Sequencer.AnimationOutliner.Splitter"))
+			.Visibility(EVisibility::SelfHitTestInvisible)
+			.OnSplitterFinishedResizing(this, &SSequencer::OnSplitterFinishedResizing)
+
+			+ SSplitter::Slot()
+			.Value(FillCoefficient_0)
+			// Can't use a minsize here because the grid panel that is actually being used to
+			//   lay out the widgets only supports fill coefficients and this leads to a disparity between the two
+			// .MinSize(200)
+			.OnSlotResized(SSplitter::FOnSlotResized::CreateSP(this, &SSequencer::OnColumnFillCoefficientChanged, 0))
 			[
-				SNew(SHorizontalBox)
+				SNew(SSpacer)
+			]
 
-				+ SHorizontalBox::Slot()
-				[
-					SNew(SSplitter)
-					.Orientation(Orient_Horizontal)
+			+ SSplitter::Slot()
+			.Value(FillCoefficient_1)
+			.OnSlotResized(SSplitter::FOnSlotResized::CreateSP(this, &SSequencer::OnColumnFillCoefficientChanged, 1))
+			[
+				SNew(SSpacer)
+			]
+		];
+}
 
-					+ SSplitter::Slot()
-					.Value(1 - FSidebarDrawerConfig::DefaultSizeCoefficient)
-					[
-						SNew(SOverlay)
+TSharedRef<SGridPanel> SSequencer::ConstructTrackAreaGridPanel(const FArguments& InArgs, const TSharedRef<FSequencer>& InSequencer, const FTimeSliderArgs& InTimeSliderArgs)
+{
+	using namespace UE::Sequencer;
 
-						+ SOverlay::Slot()
-						[
-							// track area grid panel
-							SNew( SGridPanel )
-							.FillRow( 2, 1.f )
-							.FillColumn( 0, FillCoefficient_0 )
-							.FillColumn( 1, FillCoefficient_1 )
+	constexpr int32 Column0 = 0, Column1 = 1;
+	constexpr int32 Row0 = 0, Row1 = 1, Row2 = 2, Row3 = 3, Row4 = 4;
+	const FMargin ResizeBarPadding(4.f, 0, 0, 0);
 
-							+ SGridPanel::Slot( Column0, Row1 )
-							[
-								SNew(SBorder)
-								.BorderImage(FAppStyle::GetBrush("ToolPanel.GroupBorder"))
-								[
-									SNew(SSpacer)
-								]
-							]
+	TAttribute<float> FillCoefficient_0, FillCoefficient_1;
+	{
+		FillCoefficient_0.Bind(TAttribute<float>::FGetter::CreateSP(this, &SSequencer::GetColumnFillCoefficient, 0));
+		FillCoefficient_1.Bind(TAttribute<float>::FGetter::CreateSP(this, &SSequencer::GetColumnFillCoefficient, 1));
+	}
 
-							// outliner search box
-							+ SGridPanel::Slot( Column0, Row1, SGridPanel::Layer(10) )
-							[
-								SNew(SBorder)
-								.BorderImage(FAppStyle::GetBrush("ToolPanel.GroupBorder"))
-								.Padding(FMargin(CommonPadding*2, CommonPadding))
-								.Clipping(EWidgetClipping::ClipToBounds)
-								[
-									SNew(SHorizontalBox)
-									+ SHorizontalBox::Slot()
-									.AutoWidth()
-									.VAlign(VAlign_Center)
-									.Padding(FMargin(0.f, 0.f, CommonPadding, 0.f))
-									[
-										MakeAddButton()
-									]
-						
-									+ SHorizontalBox::Slot()
-									.VAlign(VAlign_Center)
-									[
-										SAssignNew(SearchBox, SSearchBox)
-										.HintText(LOCTEXT("SearchNodesHint", "Search Tracks"))
-										.OnTextChanged( this, &SSequencer::OnOutlinerSearchChanged )
-									]
-
-									+ SHorizontalBox::Slot()
-									.AutoWidth()
-									.VAlign(VAlign_Center)
-									.Padding(FMargin(0.f, 0.f, CommonPadding, 0.f))
-									[
-										MakeFilterButton()
-									]
-								]
-							]
+	USequencerSettings* const SequencerSettings = GetSequencerSettings();
 
-							// main sequencer area
-							+ SGridPanel::Slot( Column0, Row2, SGridPanel::Layer(10) )
-							.ColumnSpan(2)
-							[
-								SAssignNew(MainSequencerArea, SVerticalBox)
+	ColumnFillCoefficients[0] = 0.3f;
+	ColumnFillCoefficients[1] = 0.7f;
 
-								+SVerticalBox::Slot()
-								.AutoHeight()
-								.MaxHeight(TAttribute<float>::Create(TAttribute<float>::FGetter::CreateSP(this, &SSequencer::GetPinnedAreaMaxHeight)))
-								.Padding(FMargin(0.f, 0.f, 0.f, CommonPadding))
-								[
-									SNew(SOverlay)
-									.Visibility(this, &SSequencer::GetPinnedAreaVisibility)
-
-									+ SOverlay::Slot()
-									[
-										SNew(SScrollBorder, PinnedTreeView.ToSharedRef())
-										[
-											SNew(SHorizontalBox)
-
-											// outliner tree
-											+ SHorizontalBox::Slot()
-											.FillWidth( FillCoefficient_0 )
-											[
-												PinnedTreeView.ToSharedRef()
-											]
-
-											// track area
-											+ SHorizontalBox::Slot()
-											.FillWidth( FillCoefficient_1 )
-											[
-												SNew(SBox)
-												.Padding(ResizeBarPadding)
-												.Clipping(EWidgetClipping::ClipToBounds)
-												[
-													PinnedTrackArea.ToSharedRef()
-												]
-											]
-										]
-									]
-
-									+ SOverlay::Slot()
-									.HAlign(HAlign_Right)
-									[
-										PinnedAreaScrollBar
-									]
-								]
+	if (IsValid(SequencerSettings))
+	{
+		const float TreeViewWidth = SequencerSettings->GetTreeViewWidth();
+		const float TimelineWidth = 1.f - TreeViewWidth;
+		if (TreeViewWidth > 0.f && TimelineWidth > 0.f)
+		{
+			ColumnFillCoefficients[0] = TreeViewWidth;
+			ColumnFillCoefficients[1] = TimelineWidth;
+		}
+	}
 
-								+SVerticalBox::Slot()
-								[
-									SNew(SOverlay)
-
-									+ SOverlay::Slot()
-									[
-										SNew(SScrollBorder, TreeView.ToSharedRef())
-										[
-											SNew(SHorizontalBox)
-
-											// outliner tree
-											+ SHorizontalBox::Slot()
-											.FillWidth( FillCoefficient_0 )
-											[
-												TreeView.ToSharedRef()
-											]
-
-											// track area
-											+ SHorizontalBox::Slot()
-											.FillWidth( FillCoefficient_1 )
-											[
-												SNew(SBox)
-												.Padding(ResizeBarPadding)
-												.Clipping(EWidgetClipping::ClipToBounds)
-												[
-													TrackArea.ToSharedRef()
-												]
-											]
-										]
-									]
-
-									+ SOverlay::Slot()
-									.HAlign(HAlign_Right)
-									[
-										ScrollBar
-									]
-
-									+ SOverlay::Slot()
-									.VAlign(VAlign_Bottom)
-									[
-										SNew(SHorizontalBox)
-
-										+ SHorizontalBox::Slot()
-										.AutoWidth()
-										[
-											SAssignNew(SequencerTreeFilterStatusBar, SSequencerTreeFilterStatusBar, InSequencer)
-											.Visibility(EVisibility::Hidden) // Initially hidden, visible on hover of the info button
-										]
-									]
-								]
-							]
+	TimeSliderController = MakeShared<FSequencerTimeSliderController>(InTimeSliderArgs, SequencerPtr);
+	TSharedRef<FSequencerTimeSliderController> TimeSliderControllerRef = TimeSliderController.ToSharedRef();
 
-							// Info Button, Transport Controls and Current Frame
-							+ SGridPanel::Slot( Column0, Row4, SGridPanel::Layer(10) )
-							[
-								SNew(SBorder)
-								.BorderImage(FAppStyle::GetBrush("ToolPanel.GroupBorder"))
-								.Clipping(EWidgetClipping::ClipToBounds)
-								[
-									SNew(SHorizontalBox)
-						
-									+ SHorizontalBox::Slot()
-									.AutoWidth()
-									.VAlign(VAlign_Center)
-									.HAlign(HAlign_Left)
-									[
-										SNew(SButton)
-										.VAlign(EVerticalAlignment::VAlign_Center)
-										.ButtonStyle(FAppStyle::Get(), "SimpleButton")
-										.ToolTipText_Lambda([this] { return LOCTEXT("ShowStatus", "Show Status"); })
-										.ContentPadding(FMargin(1, 0))
-										.Visibility(this, &SSequencer::GetInfoButtonVisibility)
-										.OnHovered_Lambda([this] { SequencerTreeFilterStatusBar->ShowStatusBar(); })
-										.OnUnhovered_Lambda([this] { SequencerTreeFilterStatusBar->FadeOutStatusBar(); })
-										.OnClicked_Lambda([this] { SequencerTreeFilterStatusBar->HideStatusBar(); return FReply::Handled(); })
-										[
-											SNew(SImage)
-											.ColorAndOpacity(FSlateColor::UseForeground())
-											.Image(FAppStyle::Get().GetBrush("Icons.Info.Small"))
-										]
-									]
-
-									+ SHorizontalBox::Slot()
-									[
-										SNew(SBorder)
-										.BorderImage(FAppStyle::GetBrush("ToolPanel.GroupBorder"))
-										.HAlign(HAlign_Center)
-										[
-											SequencerPtr.Pin()->MakeTransportControls(true)
-										]
-									]
-
-									+ SHorizontalBox::Slot()
-									.AutoWidth()
-									.VAlign(VAlign_Center)
-									.HAlign(HAlign_Right)
-									[
-										SNew(SButton)
-										.VAlign(EVerticalAlignment::VAlign_Center)
-										.ButtonStyle(FAppStyle::Get(), "NoBorder")
-										.ContentPadding(FMargin(1, 0))
-										[
-											SNew(SHorizontalBox)
-
-											+ SHorizontalBox::Slot()
-											.AutoWidth()
-											.VAlign(VAlign_Center)
-											.HAlign(HAlign_Right)
-											.Padding(FMargin(CommonPadding, 0.f, 0.f, 0.f))
-											[
-												SNew(SBorder)
-												.BorderImage(nullptr)
-												[
-													PlayTimeDisplay.ToSharedRef()
-												]
-											]
-
-											+ SHorizontalBox::Slot()
-											.AutoWidth()
-											.VAlign(VAlign_Center)
-											.HAlign(HAlign_Right)
-											.Padding(FMargin(CommonPadding, 0.f, 0.f, 0.f))
-											[
-												SNew(SBorder)
-												.BorderImage(nullptr)
-												[
-													// Current loop index, if any
-													SAssignNew(LoopIndexDisplay, STextBlock)
-													.Font(FCoreStyle::GetDefaultFontStyle("Regular", 8))
-													.Text_Lambda([this]() -> FText {
-														TOptional<int32> LoopIndex = SequencerPtr.Pin()->GetLocalLoopIndex();
-														return LoopIndex ? FText::AsNumber(LoopIndex.GetValue()) : FText();
-													})
-												]
-											]
-										]
-									]
-								]
-							]
+	ISequencerWidgetsModule& SequencerWidgets = FModuleManager::Get().LoadModuleChecked<ISequencerWidgetsModule>(TEXT("SequencerWidgets"));
 
-							// Second column
+	TSharedRef<SScrollBar> ScrollBar = SNew(SScrollBar)
+		.Thickness(FVector2D(9.f, 9.f));
 
-							+ SGridPanel::Slot( Column1, Row1 )
-							.Padding(ResizeBarPadding)
-							.RowSpan(3)
-							[
-								SNew(SBorder)
-								.BorderImage(FAppStyle::GetBrush("ToolPanel.GroupBorder"))
-								[
-									SNew(SSpacer)
-								]
-							]
+	TSharedRef<SScrollBar> PinnedAreaScrollBar = SNew(SScrollBar)
+		.Thickness(FVector2D(9.f, 9.f));
 
-							+ SGridPanel::Slot( Column1, Row1, SGridPanel::Layer(10) )
-							.Padding(ResizeBarPadding)
-							[
-								SNew( SBorder )
-								.BorderImage( FAppStyle::GetBrush("ToolPanel.GroupBorder") )
-								.BorderBackgroundColor( FLinearColor(.50f, .50f, .50f, 1.0f ) )
-								.Padding(0)
-								.Clipping(EWidgetClipping::ClipToBounds)
-								[
-									TopTimeSlider.ToSharedRef()
-								]
-							]
+	TSharedPtr<FSequencerEditorViewModel> SequencerViewModel = InSequencer->GetViewModel()->CastThisShared<FSequencerEditorViewModel>();
+	SAssignNew(PinnedTrackArea, SSequencerTrackAreaView, SequencerViewModel->GetPinnedTrackArea(), TimeSliderControllerRef);
+	SAssignNew(PinnedTreeView, SSequencerOutlinerView, SequencerViewModel->GetOutliner(), PinnedTrackArea.ToSharedRef())
+		.Selection(InSequencer->GetViewModel()->GetSelection())
+		.ExternalScrollbar(PinnedAreaScrollBar)
+		.Clipping(EWidgetClipping::ClipToBounds);
 
-							// Overlay that draws the tick lines
-							+ SGridPanel::Slot( Column1, Row2, SGridPanel::Layer(10) )
-							.Padding(ResizeBarPadding)
-							[
-								SNew( SSequencerSectionOverlay, TimeSliderControllerRef )
-								.Visibility(this, &SSequencer::GetShowTickLines)
-								.DisplayScrubPosition( false )
-								.DisplayTickLines( true )
-								.Clipping(EWidgetClipping::ClipToBounds)
-							]
+	PinnedTrackArea->SetOutliner(PinnedTreeView);
+	PinnedTrackArea->SetShowPinned(true);
+	PinnedTrackArea->SetIsPinned(true);
+	PinnedTreeView->SetShowPinned(true);
 
-							// Overlay that draws the scrub position
-							+ SGridPanel::Slot( Column1, Row2, SGridPanel::Layer(20) )
-							.Padding(ResizeBarPadding)
-							[
-								SNew( SSequencerSectionOverlay, TimeSliderControllerRef )
-								.Visibility( EVisibility::HitTestInvisible )
-								.DisplayScrubPosition( true )
-								.DisplayTickLines( false )
-								.DisplayMarkedFrames( true )
-								.PaintPlaybackRangeArgs(this, &SSequencer::GetSectionPlaybackRangeArgs)
-								.Clipping(EWidgetClipping::ClipToBounds)
-							]
+	SAssignNew(TrackArea, SSequencerTrackAreaView, SequencerViewModel->GetTrackArea(), TimeSliderControllerRef);
+	SAssignNew(TreeView, SSequencerOutlinerView, SequencerViewModel->GetOutliner(), TrackArea.ToSharedRef())
+		.Selection(InSequencer->GetViewModel()->GetSelection())
+		.ExternalScrollbar(ScrollBar)
+		.Clipping(EWidgetClipping::ClipToBounds);
 
-							+ SGridPanel::Slot(Column1, Row2, SGridPanel::Layer(30))
-								.Padding(ResizeBarPadding)
-								.HAlign(HAlign_Left)
-								.VAlign(VAlign_Top)
-							[
-								// Transform box
-								SAssignNew(TransformBox, SSequencerTransformBox, SequencerPtr.Pin().ToSharedRef(), *GetSequencerSettings(), NumericTypeInterface.ToSharedRef())
-							]
+	TrackArea->SetOutliner(TreeView);
 
-							+ SGridPanel::Slot(Column1, Row2, SGridPanel::Layer(40))
-								.Padding(ResizeBarPadding)
-								.HAlign(HAlign_Left)
-								.VAlign(VAlign_Top)
-							[
-								// Stretch box
-								SAssignNew(StretchBox, SSequencerStretchBox, SequencerPtr.Pin().ToSharedRef(), *GetSequencerSettings(), NumericTypeInterface.ToSharedRef())
-							]
+	TreeView->AddPinnedTreeView(PinnedTreeView);
 
-							// debug vis
-							+ SGridPanel::Slot( Column1, Row3, SGridPanel::Layer(10) )
-							.Padding(ResizeBarPadding)
-							[
-								SNew(SSequencerDebugVisualizer, InSequencer)
-								.ViewRange(FAnimatedRange::WrapAttribute(InArgs._ViewRange))
-								.Visibility(this, &SSequencer::GetDebugVisualizerVisibility)
-							]
+	if (SequencerSettings)
+	{
+		InitializeOutlinerColumns();
+	}
 
-							// play range sliders
-							+ SGridPanel::Slot( Column1, Row4, SGridPanel::Layer(10) )
-							.Padding(ResizeBarPadding)
-							[
-								SNew( SBorder )
-								.BorderImage( FAppStyle::GetBrush("ToolPanel.GroupBorder") )
-								.BorderBackgroundColor( FLinearColor(.50f, .50f, .50f, 1.0f ) )
-								.Clipping(EWidgetClipping::ClipToBounds)
-								.Padding(0)
-								[
-									SNew( SOverlay )
+	SequencerViewModel->GetTrackArea()->CastThisChecked<FSequencerTrackAreaViewModel>()->InitializeDefaultEditTools(*TrackArea);
+	SequencerViewModel->GetPinnedTrackArea()->CastThisChecked<FSequencerTrackAreaViewModel>()->InitializeDefaultEditTools(*PinnedTrackArea);
 
-									+ SOverlay::Slot()
-									[
-										BottomTimeSlider
-									]
+	if (SequencerSettings)
+	{
+		SequencerViewModel->SetViewDensity(SequencerSettings->GetViewDensity());
+	}
 
-									+ SOverlay::Slot()
-									[
-										BottomTimeRange
-									]
-								]
-							]
+	// Create the top and bottom sliders
+	bool bMirrorLabels = false;
+	TopTimeSlider = SequencerWidgets.CreateTimeSlider(TimeSliderControllerRef, bMirrorLabels);
+	bMirrorLabels = true;
+	TSharedRef<ITimeSlider> BottomTimeSlider = SequencerWidgets.CreateTimeSlider(TimeSliderControllerRef, TAttribute<EVisibility>(this, &SSequencer::GetBottomTimeSliderVisibility), bMirrorLabels);
+	
+	// Create bottom time range slider
+	EShowRange Ranges = EShowRange(EShowRange::WorkingRange | EShowRange::ViewRange);
+	if(InArgs._ShowPlaybackRangeInTimeSlider)
+	{
+		Ranges |= EShowRange::PlaybackRange;
+	}
+
+	TSharedRef<ITimeSlider> BottomTimeRange = SequencerWidgets.CreateTimeRange(
+		FTimeRangeArgs(
+			Ranges,
+			TimeSliderControllerRef,
+			TAttribute<EVisibility>(this, &SSequencer::GetTimeRangeVisibility),
+			NumericTypeInterface.ToSharedRef()
+		),
+		SequencerWidgets.CreateTimeRangeSlider(TimeSliderControllerRef)
+	);
+
+	return SNew(SGridPanel)
+		.FillRow(2, 1.f)
+		.FillColumn(0, FillCoefficient_0)
+		.FillColumn(1, FillCoefficient_1)
+
+		+ SGridPanel::Slot(Column0, Row1)
+		[
+			SNew(SBorder)
+			.BorderImage(FAppStyle::GetBrush(TEXT("ToolPanel.GroupBorder")))
+			[
+				SNew(SSpacer)
+			]
+		]
+
+		// outliner search box
+		+ SGridPanel::Slot(Column0, Row1, SGridPanel::Layer(10))
+		[
+			SNew(SBorder)
+			.BorderImage(FAppStyle::GetBrush(TEXT("ToolPanel.GroupBorder")))
+			.Padding(FMargin(CommonPadding*2, CommonPadding))
+			.Clipping(EWidgetClipping::ClipToBounds)
+			[
+				SNew(SHorizontalBox)
+				+ SHorizontalBox::Slot()
+				.AutoWidth()
+				.VAlign(VAlign_Center)
+				.Padding(FMargin(0.f, 0.f, CommonPadding, 0.f))
+				[
+					MakeAddButton()
+				]
+	
+				+ SHorizontalBox::Slot()
+				.VAlign(VAlign_Center)
+				[
+					SAssignNew(SearchBox, SSearchBox)
+					.HintText(LOCTEXT("SearchNodesHint", "Search Tracks"))
+					.OnTextChanged(this, &SSequencer::OnOutlinerSearchChanged)
+				]
+
+				+ SHorizontalBox::Slot()
+				.AutoWidth()
+				.VAlign(VAlign_Center)
+				.Padding(FMargin(0.f, 0.f, CommonPadding, 0.f))
+				[
+					MakeFilterButton()
+				]
+			]
+		]
+
+		// main sequencer area
+		+ SGridPanel::Slot(Column0, Row2, SGridPanel::Layer(10))
+		.ColumnSpan(2)
+		[
+			SAssignNew(MainSequencerArea, SVerticalBox)
+
+			+ SVerticalBox::Slot()
+			.AutoHeight()
+			.MaxHeight(TAttribute<float>::Create(TAttribute<float>::FGetter::CreateSP(this, &SSequencer::GetPinnedAreaMaxHeight)))
+			.Padding(FMargin(0.f, 0.f, 0.f, CommonPadding))
+			[
+				SNew(SOverlay)
+				.Visibility(this, &SSequencer::GetPinnedAreaVisibility)
+
+				+ SOverlay::Slot()
+				[
+					SNew(SScrollBorder, PinnedTreeView.ToSharedRef())
+					[
+						SNew(SHorizontalBox)
+
+						// outliner tree
+						+ SHorizontalBox::Slot()
+						.FillWidth(FillCoefficient_0)
+						[
+							PinnedTreeView.ToSharedRef()
 						]
 
-						+ SOverlay::Slot()
+						// track area
+						+ SHorizontalBox::Slot()
+						.FillWidth(FillCoefficient_1)
 						[
-							// track area virtual splitter overlay
-							SAssignNew(TreeViewSplitter, SSequencerSplitterOverlay)
-							.Style(FAppStyle::Get(), "Sequencer.AnimationOutliner.Splitter")
-							.Visibility(EVisibility::SelfHitTestInvisible)
-							.OnSplitterFinishedResizing(this, &SSequencer::OnSplitterFinishedResizing)
-
-							+ SSplitter::Slot()
-							.Value(FillCoefficient_0)
-							// Can't use a minsize here because the grid panel that is actually being used to
-							//   lay out the widgets only supports fill coefficients and this leads to a disparity between the two
-							// .MinSize(200)
-							.OnSlotResized(SSplitter::FOnSlotResized::CreateSP(this, &SSequencer::OnColumnFillCoefficientChanged, 0))
+							SNew(SBox)
+							.Padding(ResizeBarPadding)
+							.Clipping(EWidgetClipping::ClipToBounds)
 							[
-								SNew(SSpacer)
+								PinnedTrackArea.ToSharedRef()
 							]
+						]
+					]
+				]
+
+				+ SOverlay::Slot()
+				.HAlign(HAlign_Right)
+				[
+					PinnedAreaScrollBar
+				]
+			]
+
+			+ SVerticalBox::Slot()
+			[
+				SNew(SOverlay)
+
+				+ SOverlay::Slot()
+				[
+					SNew(SScrollBorder, TreeView.ToSharedRef())
+					[
+						SNew(SHorizontalBox)
 
-							+ SSplitter::Slot()
-							.Value(FillCoefficient_1)
-							.OnSlotResized(SSplitter::FOnSlotResized::CreateSP(this, &SSequencer::OnColumnFillCoefficientChanged, 1))
+						// outliner tree
+						+ SHorizontalBox::Slot()
+						.FillWidth(FillCoefficient_0)
+						[
+							TreeView.ToSharedRef()
+						]
+
+						// track area
+						+ SHorizontalBox::Slot()
+						.FillWidth(FillCoefficient_1)
+						[
+							SNew(SBox)
+							.Padding(ResizeBarPadding)
+							.Clipping(EWidgetClipping::ClipToBounds)
 							[
-								SNew(SSpacer)
+								TrackArea.ToSharedRef()
 							]
 						]
 					]
+				]
+
+				+ SOverlay::Slot()
+				.HAlign(HAlign_Right)
+				[
+					ScrollBar
+				]
+
+				+ SOverlay::Slot()
+				.VAlign(VAlign_Bottom)
+				[
+					SNew(SHorizontalBox)
+
+					+ SHorizontalBox::Slot()
+					.AutoWidth()
+					[
+						SAssignNew(SequencerTreeFilterStatusBar, SSequencerTreeFilterStatusBar, InSequencer)
+						.Visibility(EVisibility::Hidden) // Initially hidden, visible on hover of the info button
+					]
+				]
+			]
+		]
+
+		// Info Button, Transport Controls and Current Frame
+		+ SGridPanel::Slot(Column0, Row4, SGridPanel::Layer(10))
+		[
+			SNew(SBorder)
+			.BorderImage(FAppStyle::GetBrush(TEXT("ToolPanel.GroupBorder")))
+			.Clipping(EWidgetClipping::ClipToBounds)
+			[
+				SNew(SHorizontalBox)
+	
+				+ SHorizontalBox::Slot()
+				.AutoWidth()
+				.VAlign(VAlign_Center)
+				.HAlign(HAlign_Left)
+				[
+					SNew(SButton)
+					.VAlign(EVerticalAlignment::VAlign_Center)
+					.ButtonStyle(FAppStyle::Get(), TEXT("SimpleButton"))
+					.ToolTipText_Lambda([this] { return LOCTEXT("ShowStatus", "Show Status"); })
+					.ContentPadding(FMargin(1, 0))
+					.Visibility(this, &SSequencer::GetInfoButtonVisibility)
+					.OnHovered_Lambda([this] { SequencerTreeFilterStatusBar->ShowStatusBar(); })
+					.OnUnhovered_Lambda([this] { SequencerTreeFilterStatusBar->FadeOutStatusBar(); })
+					.OnClicked_Lambda([this] { SequencerTreeFilterStatusBar->HideStatusBar(); return FReply::Handled(); })
+					[
+						SNew(SImage)
+						.ColorAndOpacity(FSlateColor::UseForeground())
+						.Image(FAppStyle::Get().GetBrush(TEXT("Icons.Info.Small")))
+					]
+				]
 
-					// Sequencer details sliding drawer content dock location
-					+ SSplitter::Slot()
-					.Resizable(false)
-					.SizeRule(SSplitter::ESizeRule::SizeToContent)
-					.Value(FSidebarDrawerConfig::DefaultSizeCoefficient)
-					.OnSlotResized(this, &SSequencer::OnSidebarSlotResized)
-					.Expose(SidebarSlot)
+				+ SHorizontalBox::Slot()
+				[
+					SNew(SBorder)
+					.BorderImage(FAppStyle::GetBrush(TEXT("ToolPanel.GroupBorder")))
+					.HAlign(HAlign_Center)
 					[
-						DetailsDockLocation.ToSharedRef()
+						SequencerPtr.Pin()->MakeTransportControls(true)
 					]
 				]
 
-				// Sequencer details sliding drawer sidebar
 				+ SHorizontalBox::Slot()
 				.AutoWidth()
+				.VAlign(VAlign_Center)
+				.HAlign(HAlign_Right)
 				[
-					SAssignNew(DetailsSidebar, SSidebar, SidebarOverlay, DetailsDockLocation.ToSharedRef())
-					.OnDockStateChanged(this, &SSequencer::OnSidebarDockStateChanged)
-					.Visibility_Lambda([this]() { return IsSidebarVisible() ? EVisibility::Visible : EVisibility::Collapsed; })
+					SNew(SButton)
+					.VAlign(EVerticalAlignment::VAlign_Center)
+					.ButtonStyle(FAppStyle::Get(), TEXT("NoBorder"))
+					.ContentPadding(FMargin(1, 0))
+					[
+						SNew(SHorizontalBox)
+
+						+ SHorizontalBox::Slot()
+						.AutoWidth()
+						.VAlign(VAlign_Center)
+						.HAlign(HAlign_Right)
+						.Padding(FMargin(CommonPadding, 0.f, 0.f, 0.f))
+						[
+							SNew(SBorder)
+							.BorderImage(nullptr)
+							[
+								PlayTimeDisplay.ToSharedRef()
+							]
+						]
+
+						+ SHorizontalBox::Slot()
+						.AutoWidth()
+						.VAlign(VAlign_Center)
+						.HAlign(HAlign_Right)
+						.Padding(FMargin(CommonPadding, 0.f, 0.f, 0.f))
+						[
+							SNew(SBorder)
+							.BorderImage(nullptr)
+							[
+								// Current loop index, if any
+								SAssignNew(LoopIndexDisplay, STextBlock)
+								.Font(FCoreStyle::GetDefaultFontStyle(TEXT("Regular"), 8))
+								.Text_Lambda([this]() -> FText {
+									TOptional<int32> LoopIndex = SequencerPtr.Pin()->GetLocalLoopIndex();
+									return LoopIndex ? FText::AsNumber(LoopIndex.GetValue()) : FText();
+								})
+							]
+						]
+					]
 				]
 			]
+		]
+
+		// Second column
 
-			+ SOverlay::Slot()
+		+ SGridPanel::Slot(Column1, Row1)
+		.Padding(ResizeBarPadding)
+		.RowSpan(3)
+		[
+			SNew(SBorder)
+			.BorderImage(FAppStyle::GetBrush(TEXT("ToolPanel.GroupBorder")))
 			[
-				SidebarOverlay
+				SNew(SSpacer)
 			]
 		]
-	];
 
-	if (InSequencer->GetHostCapabilities().bSupportsCurveEditor)
-	{
-		FCurveEditorExtension* CurveEditorExtension = SequencerViewModel->CastDynamicChecked<FCurveEditorExtension>();
-		CurveEditorExtension->CreateCurveEditor(TimeSliderArgs);
-	}
+		+ SGridPanel::Slot(Column1, Row1, SGridPanel::Layer(10))
+		.Padding(ResizeBarPadding)
+		[
+			SNew(SBorder)
+			.BorderImage(FAppStyle::GetBrush(TEXT("ToolPanel.GroupBorder")))
+			.BorderBackgroundColor(FLinearColor(.5f, .5f, .5f, 1.f))
+			.Padding(0)
+			.Clipping(EWidgetClipping::ClipToBounds)
+			[
+				TopTimeSlider.ToSharedRef()
+			]
+		]
 
-	ApplySequencerCustomization(RootCustomization);
+		// Overlay that draws the tick lines
+		+ SGridPanel::Slot(Column1, Row2, SGridPanel::Layer(10))
+		.Padding(ResizeBarPadding)
+		[
+			SNew(SSequencerSectionOverlay, TimeSliderControllerRef)
+			.Visibility(this, &SSequencer::GetShowTickLines)
+			.DisplayScrubPosition(false)
+			.DisplayTickLines(true)
+			.Clipping(EWidgetClipping::ClipToBounds)
+		]
 
-	InSequencer->GetViewModel()->GetSelection()->KeySelection.OnChanged.AddSP(this, &SSequencer::HandleKeySelectionChanged);
-	InSequencer->GetViewModel()->GetSelection()->Outliner.OnChanged.AddSP(this, &SSequencer::HandleOutlinerNodeSelectionChanged);
+		// Overlay that draws the scrub position
+		+ SGridPanel::Slot(Column1, Row2, SGridPanel::Layer(20))
+		.Padding(ResizeBarPadding)
+		[
+			SNew(SSequencerSectionOverlay, TimeSliderControllerRef)
+			.Visibility(EVisibility::HitTestInvisible)
+			.DisplayScrubPosition(true)
+			.DisplayTickLines(false)
+			.DisplayMarkedFrames(true)
+			.PaintPlaybackRangeArgs(this, &SSequencer::GetSectionPlaybackRangeArgs)
+			.Clipping(EWidgetClipping::ClipToBounds)
+		]
 
-	ResetBreadcrumbs();
+		+ SGridPanel::Slot(Column1, Row2, SGridPanel::Layer(30))
+		.Padding(ResizeBarPadding)
+		.HAlign(HAlign_Left)
+		.VAlign(VAlign_Top)
+		[
+			// Transform box
+			SAssignNew(TransformBox, SSequencerTransformBox, SequencerPtr.Pin().ToSharedRef(), *SequencerSettings, NumericTypeInterface.ToSharedRef())
+		]
+
+		+ SGridPanel::Slot(Column1, Row2, SGridPanel::Layer(40))
+		.Padding(ResizeBarPadding)
+		.HAlign(HAlign_Left)
+		.VAlign(VAlign_Top)
+		[
+			// Stretch box
+			SAssignNew(StretchBox, SSequencerStretchBox, SequencerPtr.Pin().ToSharedRef(), *SequencerSettings, NumericTypeInterface.ToSharedRef())
+		]
+
+		// debug vis
+		+ SGridPanel::Slot(Column1, Row3, SGridPanel::Layer(10))
+		.Padding(ResizeBarPadding)
+		[
+			SNew(SSequencerDebugVisualizer, InSequencer)
+			.ViewRange(FAnimatedRange::WrapAttribute(InArgs._ViewRange))
+			.Visibility(this, &SSequencer::GetDebugVisualizerVisibility)
+		]
+
+		// play range sliders
+		+ SGridPanel::Slot(Column1, Row4, SGridPanel::Layer(10))
+		.Padding(ResizeBarPadding)
+		[
+			SNew(SBorder)
+			.BorderImage(FAppStyle::GetBrush(TEXT("ToolPanel.GroupBorder")))
+			.BorderBackgroundColor(FLinearColor(.5f, .5f, .5f, 1.f))
+			.Clipping(EWidgetClipping::ClipToBounds)
+			.Padding(0)
+			[
+				SNew(SOverlay)
+				+ SOverlay::Slot()
+				[
+					BottomTimeSlider
+				]
+				+ SOverlay::Slot()
+				[
+					BottomTimeRange
+				]
+			]
+		];
 }
 UE_ENABLE_OPTIMIZATION_SHIP
 
@@ -2247,6 +2266,11 @@ TSharedRef<SWidget> SSequencer::MakeViewMenu()
 	FMenuBuilder MenuBuilder( true, SequencerPtr.Pin()->GetCommandBindings() );
 	TSharedPtr<FSequencer> Sequencer = SequencerPtr.Pin();
 
+	if (Sequencer->GetHostCapabilities().bSupportsSidebar)
+	{
+		MenuBuilder.AddMenuEntry( FSequencerCommands::Get().ToggleSidebarVisible );
+	}
+
 	if (Sequencer->IsLevelEditorSequencer())
 	{
 		MenuBuilder.AddMenuEntry( FSequencerCommands::Get().TogglePilotCamera );
@@ -4323,22 +4347,43 @@ bool SSequencer::UnregisterDrawerSection(const FName InDrawerId, const FName InS
 	return false;
 }
 
-void SSequencer::OnSidebarDockStateChanged(const FName InDrawerId)
+void SSequencer::OnSidebarStateChanged(const FSidebarState& InNewState)
 {
-	DetailsSidebar->UpdateDockedSplitterSlot(InDrawerId, SidebarSlot);
+	USequencerSettings* const SequencerSettings = GetSequencerSettings();
+	if (IsValid(SequencerSettings))
+	{
+		SequencerSettings->SetSidebarState(InNewState);
+	}
 }
 
-void SSequencer::OnSidebarSlotResized(const float InFillCoefficient)
+bool SSequencer::IsSidebarVisible() const
 {
-	if (InFillCoefficient == 0.f && DetailsSidebar->HasDrawerDocked())
+	USequencerSettings* const SequencerSettings = GetSequencerSettings();
+	if (IsValid(SequencerSettings))
 	{
-		DetailsSidebar->UndockAllDrawers();
+		return SequencerSettings->GetSidebarState().IsVisible();
 	}
+	return false;
 }
 
-bool SSequencer::IsSidebarVisible() const
+void SSequencer::SetSidebarVisible(const bool bInVisible)
+{
+	USequencerSettings* const SequencerSettings = GetSequencerSettings();
+	if (!IsValid(SequencerSettings))
+	{
+		return;
+	}
+
+	FSidebarState SidebarState = SequencerSettings->GetSidebarState();
+	SidebarState.SetVisible(bInVisible);
+	SequencerSettings->SetSidebarState(SidebarState);
+
+	MainContentContainer->SetContent(ConstructMainContent());
+}
+
+void SSequencer::ToggleSidebarVisible()
 {
-	return DetailsSidebar.IsValid() && DetailsSidebar->GetDrawerCount() > 0;
+	SetSidebarVisible(!IsSidebarVisible());
 }
 
 void SSequencer::ToggleSidebarSelectionDrawerOpen()
diff --git a/Engine/Source/Editor/Sequencer/Private/SSequencer.h b/Engine/Source/Editor/Sequencer/Private/SSequencer.h
--- a/Engine/Source/Editor/Sequencer/Private/SSequencer.h
+++ b/Engine/Source/Editor/Sequencer/Private/SSequencer.h
@@ -37,7 +37,9 @@ class SSequencerStretchBox;
 class SCurveEditorPanel;
 class SBox;
 class SDockTab;
+class SGridPanel;
 class SSidebar;
+class SSidebarContainer;
 class SWindow;
 class USequencerSettings;
 class FSequencerTrackFilter;
@@ -46,6 +48,7 @@ class SSequencerTreeFilterStatusBar;
 struct FPaintPlaybackRangeArgs;
 struct FSequencerCustomizationInfo;
 struct FSidebarDrawerConfig;
+struct FTimeSliderArgs;
 
 namespace UE
 {
@@ -389,7 +392,13 @@ protected:
 	virtual void OnMouseLeave(const FPointerEvent& MouseEvent) override;
 
 private:
-	
+
+	static constexpr float CommonPadding = 3.f;
+
+	TSharedRef<SWidget> ConstructMainContent();
+	TSharedRef<SWidget> ConstructGridPanel();
+	TSharedRef<SGridPanel> ConstructTrackAreaGridPanel(const FArguments& InArgs, const TSharedRef<FSequencer>& InSequencer, const FTimeSliderArgs& InTimeSliderArgs);
+
 	/** Initalizes a list of all track filter objects */
 	void InitializeTrackFilters();
 
@@ -700,6 +709,12 @@ public:
 	/** @return True if the sidebar is being displayed. */
 	bool IsSidebarVisible() const;
 
+	/** Set the visibility of the sidebar */
+	void SetSidebarVisible(const bool bInVisible);
+
+	/** Toggle the visibility of the sidebar. */
+	void ToggleSidebarVisible();
+
 	/** Toggles the sidebar "Selection" drawer open or closed. */
 	void ToggleSidebarSelectionDrawerOpen();
 
@@ -720,9 +735,19 @@ private:
 	/** Applies a single customization. */
 	void ApplySequencerCustomization(const FSequencerCustomizationInfo& Customization);
 
-	void OnSidebarDockStateChanged(const FName InDrawerId);
+	void OnSidebarStateChanged(const FSidebarState& InNewState);
 
-	void OnSidebarSlotResized(const float InFillCoefficient);
+	/** This is the main container widget for the Sequencer (minus the toolbar). This is the parent of
+	 * the GridPanel widget below and is used to rebuild content when the sidebar is added/removed. */
+	TSharedPtr<SBox> MainContentContainer;
+
+	/** Grid panel that holds most of the content and contains MainSequencerArea widget below.
+	 * Built once at Sequencer initialize and used when rebuilding the main content container when the sidebar is added/removed. */
+	TSharedPtr<SGridPanel> GridPanel;
+
+	/** Contains the outliner tree and track area. Used by GetPinnedAreaMaxHeight().
+	 * Built once at Sequencer initialize and used when rebuilding the main content container when the sidebar is added/removed. */
+	TSharedPtr<SVerticalBox> MainSequencerArea;
 
 	/** Transform box widget. */
 	TSharedPtr<SSequencerTransformBox> TransformBox;
@@ -730,9 +755,6 @@ private:
 	/** Stretch box widget. */
 	TSharedPtr<SSequencerStretchBox> StretchBox;
 
-	/** Main Sequencer Area*/
-	TSharedPtr<SVerticalBox> MainSequencerArea;
-
 	/** Filter Status Bar */
 	TSharedPtr<SSequencerTreeFilterStatusBar> SequencerTreeFilterStatusBar;
 
@@ -873,8 +895,6 @@ private:
 
 	TSharedPtr<SSequencerGroupManager> NodeGroupManager;
 
-	TSharedPtr<SBox> DetailsDockLocation;
-
+	TSharedPtr<SSidebarContainer> SidebarContainer;
 	TSharedPtr<SSidebar> DetailsSidebar;
-	SSplitter::FSlot* SidebarSlot = nullptr;
 };
diff --git a/Engine/Source/Editor/Sequencer/Private/Sequencer.cpp b/Engine/Source/Editor/Sequencer/Private/Sequencer.cpp
--- a/Engine/Source/Editor/Sequencer/Private/Sequencer.cpp
+++ b/Engine/Source/Editor/Sequencer/Private/Sequencer.cpp
@@ -293,8 +293,8 @@ namespace UE
 	} // namespace Sequencer
 } // namespace UE
 
-bool FSequencer::bSelectionLimited = false;
 const FName FSequencer::SelectionDrawerId = TEXT("SelectionDetails");
+bool FSequencer::bSelectionLimited = false;
 
 void FSequencer::InitSequencer(const FSequencerInitParams& InitParams, const TSharedRef<ISequencerObjectChangeListener>& InObjectChangeListener, const TArray<FOnCreateTrackEditor>& TrackEditorDelegates, const TArray<FOnCreateEditorObjectBinding>& EditorObjectBindingDelegates, const TArray<FOnCreateOutlinerColumn>& OutlinerColumnDelegates)
 {
@@ -669,13 +669,15 @@ void FSequencer::InitSequencer(const FSequencerInitParams& InitParams, const TSh
 
 	OnActivateSequenceEvent.Broadcast(ActiveTemplateIDs[0]);
 
-	{ // Selection Details Drawer
+	// Selection Details Drawer
+	if (HostCapabilities.bSupportsSidebar)
+	{
 		FSidebarDrawerConfig DetailsDrawer;
 		DetailsDrawer.UniqueId = SelectionDrawerId;
 		DetailsDrawer.ButtonText = LOCTEXT("SelectionDetailsPanelLabel", "Selection");
 		DetailsDrawer.ToolTipText = TAttribute<FText>::CreateSP(this, &FSequencer::GetSidebarSelectionDrawerToolTipText);
 		DetailsDrawer.Icon = FAppStyle::GetBrush(TEXT("EditorPreferences.TabIcon"));
-		DetailsDrawer.bInitiallyDocked = true;
+		DetailsDrawer.InitialState = Settings->GetSidebarState().FindOrAddDrawerState(SelectionDrawerId);
 		RegisterDrawer(MoveTemp(DetailsDrawer));
 
 		RegisterDrawerSection(SelectionDrawerId, MakeShared<FSequencerSelectionDrawer>(SharedThis(this)));
@@ -11594,9 +11596,18 @@ void FSequencer::BindCommands()
 		FCanExecuteAction(),
 		FIsActionChecked::CreateSP(this, &FSequencer::IsViewportSelectionLimited));
 
+	if (HostCapabilities.bSupportsSidebar)
+	{
+		SequencerCommandBindings->MapAction(
+			Commands.ToggleSidebarVisible,
+			FExecuteAction::CreateSP(this, &FSequencer::ToggleSidebar),
+			FCanExecuteAction(),
+			FIsActionChecked::CreateSP(this, &FSequencer::IsSidebarVisible));
+	}
+
 	SequencerCommandBindings->MapAction(
 		Commands.ToggleSidebarSelectionDrawerOpen,
-		FExecuteAction::CreateSP( this, &FSequencer::ShowHideSidebarSelectionDrawer));
+		FExecuteAction::CreateSP(this, &FSequencer::ToggleSidebarSelectionDrawer));
 
 	SequencerCommandBindings->MapAction(
 		Commands.ToggleSidebarDrawerDock,
@@ -12179,7 +12190,20 @@ bool FSequencer::UnregisterDrawerSection(const FName InDrawerId, const FName InS
 	return false;
 }
 
-void FSequencer::ShowHideSidebarSelectionDrawer()
+void FSequencer::ToggleSidebar()
+{
+	if (SequencerWidget.IsValid())
+	{
+		return SequencerWidget->ToggleSidebarVisible();
+	}
+}
+
+bool FSequencer::IsSidebarVisible() const
+{
+	return SequencerWidget.IsValid() ? SequencerWidget->IsSidebarVisible() : false;
+}
+
+void FSequencer::ToggleSidebarSelectionDrawer()
 {
 	if (SequencerWidget.IsValid())
 	{
diff --git a/Engine/Source/Editor/Sequencer/Private/Sequencer.h b/Engine/Source/Editor/Sequencer/Private/Sequencer.h
--- a/Engine/Source/Editor/Sequencer/Private/Sequencer.h
+++ b/Engine/Source/Editor/Sequencer/Private/Sequencer.h
@@ -1212,8 +1212,14 @@ public:
 	 */
 	virtual bool UnregisterDrawerSection(const FName InDrawerId, const FName InSectionId) override;
 
+	/** Toggles the visibility of the Sequencer sidebar. */
+	void ToggleSidebar();
+
+	/** @return True if the Sequencer sidebar is currently visible. */
+	bool IsSidebarVisible() const;
+
 	/** Toggles the sidebar "Selection" drawer open or closed. */
-	void ShowHideSidebarSelectionDrawer();
+	void ToggleSidebarSelectionDrawer();
 
 	/** Undocks the docked sidebar drawer if docked or docks the sidebar drawer if there is one open and no currently docked drawer. */
 	void ToggleSidebarDrawerDocked();
diff --git a/Engine/Source/Editor/Sequencer/Private/SequencerCommands.cpp b/Engine/Source/Editor/Sequencer/Private/SequencerCommands.cpp
--- a/Engine/Source/Editor/Sequencer/Private/SequencerCommands.cpp
+++ b/Engine/Source/Editor/Sequencer/Private/SequencerCommands.cpp
@@ -198,6 +198,7 @@ void FSequencerCommands::RegisterCommands()
 
 	UI_COMMAND(ToggleLimitViewportSelection, "Toggle Limit Viewport Selection", "Toggles viewport selectability between Sequencer only objects and all objects.", EUserInterfaceActionType::ToggleButton, FInputChord());
 
+	UI_COMMAND(ToggleSidebarVisible, "Sidebar Details", "Toggles the visibility of the Sequencer sidebar.", EUserInterfaceActionType::ToggleButton, FInputChord());
 	UI_COMMAND(ToggleSidebarSelectionDrawerOpen, "Toggle Sidebar Selection Drawer Open", "Toggles the visible state of the sidebar selection drawer. Undocks and closes the sidebar drawer if docked.", EUserInterfaceActionType::Button, FInputChord(EModifierKey::Control, EKeys::E));
 	UI_COMMAND(ToggleSidebarDrawerDock, "Toggle Sidebar Drawer Docked", "Toggles the docked state of the sidebar current drawer. Undocks the docked sidebar drawer if docked or docks the sidebar drawer if there is one open and no currently docked drawer.", EUserInterfaceActionType::Button, FInputChord(EModifierKey::Alt, EKeys::D));
 }
diff --git a/Engine/Source/Editor/Sequencer/Private/SequencerSettings.cpp b/Engine/Source/Editor/Sequencer/Private/SequencerSettings.cpp
--- a/Engine/Source/Editor/Sequencer/Private/SequencerSettings.cpp
+++ b/Engine/Source/Editor/Sequencer/Private/SequencerSettings.cpp
@@ -1068,3 +1068,15 @@ void USequencerSettings::SetOutlinerColumnVisibility(const TArray<FColumnVisibil
 		SaveConfig();
 	}
 }
+
+FSidebarState& USequencerSettings::GetSidebarState()
+{
+	return SidebarState.FindOrAdd(GetFName());
+}
+
+void USequencerSettings::SetSidebarState(const FSidebarState& InSidebarState)
+{
+	FSidebarState& State = SidebarState.FindOrAdd(GetFName());
+	State = InSidebarState;
+	SaveConfig();
+}
diff --git a/Engine/Source/Editor/Sequencer/Public/ISequencerModule.h b/Engine/Source/Editor/Sequencer/Public/ISequencerModule.h
--- a/Engine/Source/Editor/Sequencer/Public/ISequencerModule.h
+++ b/Engine/Source/Editor/Sequencer/Public/ISequencerModule.h
@@ -150,12 +150,16 @@ struct FSequencerHostCapabilities
 	/** Do we support adding content from the content browser */
 	bool bSupportsAddFromContentBrowser;
 
+	/** Do we support the sidebar widget */
+	bool bSupportsSidebar;
+
 	FSequencerHostCapabilities()
 		: bSupportsSaveMovieSceneAsset(false)
 		, bSupportsCurveEditor(false)
 		, bSupportsRecording(false)
 		, bSupportsRenderMovie(false)
 		, bSupportsAddFromContentBrowser(false)
+		, bSupportsSidebar(false)
 	{}
 };
 
diff --git a/Engine/Source/Editor/Sequencer/Public/SequencerCommands.h b/Engine/Source/Editor/Sequencer/Public/SequencerCommands.h
--- a/Engine/Source/Editor/Sequencer/Public/SequencerCommands.h
+++ b/Engine/Source/Editor/Sequencer/Public/SequencerCommands.h
@@ -478,6 +478,9 @@ public:
 	/** Toggles selectability of viewport objects to Sequencer items only */
 	TSharedPtr< FUICommandInfo > ToggleLimitViewportSelection;
 
+	/** Toggles the visibility of the Sequencer sidebar */
+	TSharedPtr< FUICommandInfo > ToggleSidebarVisible;
+
 	/** Toggles the sidebar "Selection" drawer open or closed */
 	TSharedPtr< FUICommandInfo > ToggleSidebarSelectionDrawerOpen;
 
diff --git a/Engine/Source/Editor/Sequencer/Public/SequencerSettings.h b/Engine/Source/Editor/Sequencer/Public/SequencerSettings.h
--- a/Engine/Source/Editor/Sequencer/Public/SequencerSettings.h
+++ b/Engine/Source/Editor/Sequencer/Public/SequencerSettings.h
@@ -7,6 +7,7 @@
 #include "UObject/Object.h"
 #include "UObject/Package.h"
 #include "FrameNumberDisplayFormat.h"
+#include "Sidebar/SidebarState.h"
 #include "SequencerSettings.generated.h"
 
 struct FPropertyChangedEvent;
@@ -484,6 +485,11 @@ public:
 	/** Sets the visibility of outliner columns in display order */
 	void SetOutlinerColumnVisibility(const TArray<FColumnVisibilitySetting>& InColumnVisibilitySettings);
 
+	/** Gets the last saved sidebar state */
+	FSidebarState& GetSidebarState();
+	/** Sets the sidebar state to be restored on Sequencer initialize */
+	void SetSidebarState(const FSidebarState& InSidebarState);
+
 protected:
 
 	/** The auto change mode (auto-key, auto-track or none). */
@@ -759,6 +765,10 @@ protected:
 	UPROPERTY(config, EditAnywhere, Category = General)
 	TArray<FColumnVisibilitySetting> ColumnVisibilitySettings;
 
+	/** The state of a sidebar to be restored when each Sequencer type is initialized */
+	UPROPERTY(config)
+	TMap<FName, FSidebarState> SidebarState;
+
 	FOnEvaluateSubSequencesInIsolationChanged OnEvaluateSubSequencesInIsolationChangedEvent;
 	FOnShowSelectedNodesOnlyChanged OnShowSelectedNodesOnlyChangedEvent;
 	FOnAllowEditsModeChanged OnAllowEditsModeChangedEvent;
diff --git a/Engine/Source/Editor/UMGEditor/Private/WidgetBlueprintEditor.cpp b/Engine/Source/Editor/UMGEditor/Private/WidgetBlueprintEditor.cpp
--- a/Engine/Source/Editor/UMGEditor/Private/WidgetBlueprintEditor.cpp
+++ b/Engine/Source/Editor/UMGEditor/Private/WidgetBlueprintEditor.cpp
@@ -1455,6 +1455,7 @@ TSharedPtr<ISequencer> FWidgetBlueprintEditor::CreateSequencerWidgetInternal()
 
 		SequencerInitParams.HostCapabilities.bSupportsCurveEditor = true;
 		SequencerInitParams.HostCapabilities.bSupportsAddFromContentBrowser = true;
+		SequencerInitParams.HostCapabilities.bSupportsSidebar = true;
 	};
 
 	TSharedPtr<ISequencer> Sequencer = FModuleManager::LoadModuleChecked<ISequencerModule>("Sequencer").CreateSequencer(SequencerInitParams);
