Sequencer: Fixed changing key interpolation not modifying channel owners

Channel owners are a new concept that were introduced with timewarp - previous code that modified the section now needs to modify the channel owner instead

diff --git a/Engine/Source/Editor/MovieSceneTools/Private/Channels/BuiltInChannelEditors.cpp b/Engine/Source/Editor/MovieSceneTools/Private/Channels/BuiltInChannelEditors.cpp
--- a/Engine/Source/Editor/MovieSceneTools/Private/Channels/BuiltInChannelEditors.cpp
+++ b/Engine/Source/Editor/MovieSceneTools/Private/Channels/BuiltInChannelEditors.cpp
@@ -1225,13 +1225,13 @@ struct TCurveChannelKeyMenuExtension : TSharedFromThis<TCurveChannelKeyMenuExten
 
 		for (const TExtendKeyMenuParams<ChannelType>& Channel : ChannelAndHandles)
 		{
-			UMovieSceneSection* Section = Channel.Section.Get();
-			ChannelType* ChannelPtr = Channel.Channel.Get();
-
-			if (Section && ChannelPtr)
+			if (UMovieSceneSignedObject* OwningObject = Cast<UMovieSceneSignedObject>(Channel.WeakOwner.Get()))
 			{
-				Section->Modify();
+				OwningObject->Modify();
+			}
 
+			if (ChannelType* ChannelPtr = Channel.Channel.Get())
+			{
 				TMovieSceneChannelData<ChannelValueType> ChannelData = ChannelPtr->GetData();
 				TArrayView<ChannelValueType> Values = ChannelData.GetValues();
 
diff --git a/Engine/Source/Editor/Sequencer/Private/SequencerContextMenus.cpp b/Engine/Source/Editor/Sequencer/Private/SequencerContextMenus.cpp
--- a/Engine/Source/Editor/Sequencer/Private/SequencerContextMenus.cpp
+++ b/Engine/Source/Editor/Sequencer/Private/SequencerContextMenus.cpp
@@ -139,9 +139,10 @@ void FKeyContextMenu::PopulateMenu(FMenuBuilder& MenuBuilder, TSharedPtr<FExtend
 		for (FSelectedChannelInfo& ChannelInfo : SelectedKeysByChannel.SelectedChannels)
 		{
 			FExtendKeyMenuParams ExtendKeyMenuParams;
-			ExtendKeyMenuParams.Section = ChannelInfo.OwningSection;
-			ExtendKeyMenuParams.Channel = ChannelInfo.Channel;
-			ExtendKeyMenuParams.Handles = MoveTemp(ChannelInfo.KeyHandles);
+			ExtendKeyMenuParams.Section   = ChannelInfo.OwningSection;
+			ExtendKeyMenuParams.WeakOwner = ChannelInfo.OwningObject;
+			ExtendKeyMenuParams.Channel   = ChannelInfo.Channel;
+			ExtendKeyMenuParams.Handles   = MoveTemp(ChannelInfo.KeyHandles);
 
 			ChannelAndHandlesByType.FindOrAdd(ChannelInfo.Channel.GetChannelTypeName()).Add(MoveTemp(ExtendKeyMenuParams));
 		}
@@ -1010,42 +1011,41 @@ void FSectionContextMenu::SetInterpTangentMode(ERichCurveInterpMode InterpMode,
 	{
 		if (KeyArea.IsValid())
 		{
-			UMovieSceneSection* Section = KeyArea->GetOwningSection();
-			if (Section)
+			if (UMovieSceneSignedObject* OwningObject = Cast<UMovieSceneSignedObject>(KeyArea->GetOwningObject()))
 			{
-				Section->Modify();
-
-				FMovieSceneChannelHandle Handle = KeyArea->GetChannel();
-				if (Handle.GetChannelTypeName() == FMovieSceneFloatChannel::StaticStruct()->GetFName())
-				{
-					FMovieSceneFloatChannel* FloatChannel = static_cast<FMovieSceneFloatChannel*>(Handle.Get());
-					TMovieSceneChannelData<FMovieSceneFloatValue> ChannelData = FloatChannel->GetData();
-					TArrayView<FMovieSceneFloatValue> Values = ChannelData.GetValues();
+				OwningObject->Modify();
+			}
 
-					for (int32 KeyIndex = 0; KeyIndex < FloatChannel->GetNumKeys(); ++KeyIndex)
-					{
-						Values[KeyIndex].InterpMode = InterpMode;
-						Values[KeyIndex].TangentMode = TangentMode;
-						bAnythingChanged = true;
-					}
+			FMovieSceneChannelHandle Handle = KeyArea->GetChannel();
+			if (Handle.GetChannelTypeName() == FMovieSceneFloatChannel::StaticStruct()->GetFName())
+			{
+				FMovieSceneFloatChannel* FloatChannel = static_cast<FMovieSceneFloatChannel*>(Handle.Get());
+				TMovieSceneChannelData<FMovieSceneFloatValue> ChannelData = FloatChannel->GetData();
+				TArrayView<FMovieSceneFloatValue> Values = ChannelData.GetValues();
 
-					FloatChannel->AutoSetTangents();
-				}
-				else if (Handle.GetChannelTypeName() == FMovieSceneDoubleChannel::StaticStruct()->GetFName())
+				for (int32 KeyIndex = 0; KeyIndex < FloatChannel->GetNumKeys(); ++KeyIndex)
 				{
-					FMovieSceneDoubleChannel* DoubleChannel = static_cast<FMovieSceneDoubleChannel*>(Handle.Get());
-					TMovieSceneChannelData<FMovieSceneDoubleValue> ChannelData = DoubleChannel->GetData();
-					TArrayView<FMovieSceneDoubleValue> Values = ChannelData.GetValues();
+					Values[KeyIndex].InterpMode = InterpMode;
+					Values[KeyIndex].TangentMode = TangentMode;
+					bAnythingChanged = true;
+				}
 
-					for (int32 KeyIndex = 0; KeyIndex < DoubleChannel->GetNumKeys(); ++KeyIndex)
-					{
-						Values[KeyIndex].InterpMode = InterpMode;
-						Values[KeyIndex].TangentMode = TangentMode;
-						bAnythingChanged = true;
-					}
+				FloatChannel->AutoSetTangents();
+			}
+			else if (Handle.GetChannelTypeName() == FMovieSceneDoubleChannel::StaticStruct()->GetFName())
+			{
+				FMovieSceneDoubleChannel* DoubleChannel = static_cast<FMovieSceneDoubleChannel*>(Handle.Get());
+				TMovieSceneChannelData<FMovieSceneDoubleValue> ChannelData = DoubleChannel->GetData();
+				TArrayView<FMovieSceneDoubleValue> Values = ChannelData.GetValues();
 
-					DoubleChannel->AutoSetTangents();
+				for (int32 KeyIndex = 0; KeyIndex < DoubleChannel->GetNumKeys(); ++KeyIndex)
+				{
+					Values[KeyIndex].InterpMode = InterpMode;
+					Values[KeyIndex].TangentMode = TangentMode;
+					bAnythingChanged = true;
 				}
+
+				DoubleChannel->AutoSetTangents();
 			}
 		}
 	}
diff --git a/Engine/Source/Editor/Sequencer/Private/SequencerSelectedKey.cpp b/Engine/Source/Editor/Sequencer/Private/SequencerSelectedKey.cpp
--- a/Engine/Source/Editor/Sequencer/Private/SequencerSelectedKey.cpp
+++ b/Engine/Source/Editor/Sequencer/Private/SequencerSelectedKey.cpp
@@ -41,7 +41,7 @@ FSelectedKeysByChannel::FSelectedKeysByChannel(const UE::Sequencer::FKeySelectio
 			const int32* ChannelArrayIndex = KeyAreaToChannelIndex.Find(Channel->GetKeyArea().Get());
 			if (!ChannelArrayIndex)
 			{
-				int32 NewIndex = SelectedChannels.Add(FSelectedChannelInfo(Channel->GetKeyArea()->GetChannel(), Channel->GetSection()));
+				int32 NewIndex = SelectedChannels.Add(FSelectedChannelInfo(Channel->GetKeyArea()->GetChannel(), Channel->GetSection(), Channel->GetOwningObject()));
 				ChannelArrayIndex = &KeyAreaToChannelIndex.Add(Channel->GetKeyArea().Get(), NewIndex);
 			}
 
@@ -70,7 +70,7 @@ FSelectedKeysByChannel::FSelectedKeysByChannel(TArrayView<const FSequencerSelect
 			const int32* ChannelArrayIndex = KeyAreaToChannelIndex.Find(Channel->GetKeyArea().Get());
 			if (!ChannelArrayIndex)
 			{
-				int32 NewIndex = SelectedChannels.Add(FSelectedChannelInfo(Channel->GetKeyArea()->GetChannel(), Channel->GetSection()));
+				int32 NewIndex = SelectedChannels.Add(FSelectedChannelInfo(Channel->GetKeyArea()->GetChannel(), Channel->GetSection(), Channel->GetOwningObject()));
 				ChannelArrayIndex = &KeyAreaToChannelIndex.Add(Channel->GetKeyArea().Get(), NewIndex);
 			}
 
diff --git a/Engine/Source/Editor/Sequencer/Private/SequencerSelectedKey.h b/Engine/Source/Editor/Sequencer/Private/SequencerSelectedKey.h
--- a/Engine/Source/Editor/Sequencer/Private/SequencerSelectedKey.h
+++ b/Engine/Source/Editor/Sequencer/Private/SequencerSelectedKey.h
@@ -83,8 +83,8 @@ public:
  */
 struct FSelectedChannelInfo
 {
-	explicit FSelectedChannelInfo(FMovieSceneChannelHandle InChannel, UMovieSceneSection* InOwningSection)
-		: Channel(InChannel), OwningSection(InOwningSection)
+	explicit FSelectedChannelInfo(FMovieSceneChannelHandle InChannel, UMovieSceneSection* InOwningSection, UObject* InOwningObject)
+		: Channel(InChannel), OwningSection(InOwningSection), OwningObject(InOwningObject)
 	{}
 
 	/** The channel on which the keys are selected */
@@ -93,6 +93,9 @@ struct FSelectedChannelInfo
 	/** The section that owns this channel */
 	UMovieSceneSection* OwningSection;
 
+	/** The object that owns this channel. Maybe the same as OwningSection, maybe not. */
+	UObject* OwningObject;
+
 	/** The key handles that are selected on this channel */
 	TArray<FKeyHandle> KeyHandles;
 
diff --git a/Engine/Source/Editor/Sequencer/Public/ISequencerChannelInterface.h b/Engine/Source/Editor/Sequencer/Public/ISequencerChannelInterface.h
--- a/Engine/Source/Editor/Sequencer/Public/ISequencerChannelInterface.h
+++ b/Engine/Source/Editor/Sequencer/Public/ISequencerChannelInterface.h
@@ -58,6 +58,9 @@ struct FExtendKeyMenuParams
 	/** The section on which the channel resides */
 	TWeakObjectPtr<UMovieSceneSection> Section;
 
+	/** The owning object. Probably the same as Section, but can be different. */
+	TWeakObjectPtr<UObject> WeakOwner;
+
 	/** The channel on which the keys reside */
 	FMovieSceneChannelHandle Channel;
 
diff --git a/Engine/Source/Editor/Sequencer/Public/SequencerChannelInterface.h b/Engine/Source/Editor/Sequencer/Public/SequencerChannelInterface.h
--- a/Engine/Source/Editor/Sequencer/Public/SequencerChannelInterface.h
+++ b/Engine/Source/Editor/Sequencer/Public/SequencerChannelInterface.h
@@ -131,9 +131,10 @@ struct TSequencerChannelInterfaceCommon : ISequencerChannelInterface
 		for (const FExtendKeyMenuParams& Ptr : ChannelsAndHandles)
 		{
 			TExtendKeyMenuParams<ChannelType> TypedChannelAndHandles;
-			TypedChannelAndHandles.Section = Ptr.Section;
-			TypedChannelAndHandles.Handles = Ptr.Handles;
-			TypedChannelAndHandles.Channel = Ptr.Channel.Cast<ChannelType>();
+			TypedChannelAndHandles.Section   = Ptr.Section;
+			TypedChannelAndHandles.WeakOwner = Ptr.WeakOwner;
+			TypedChannelAndHandles.Handles   = Ptr.Handles;
+			TypedChannelAndHandles.Channel   = Ptr.Channel.Cast<ChannelType>();
 
 			TypedChannels.Add(MoveTemp(TypedChannelAndHandles));
 		}
diff --git a/Engine/Source/Editor/Sequencer/Public/SequencerChannelTraits.h b/Engine/Source/Editor/Sequencer/Public/SequencerChannelTraits.h
--- a/Engine/Source/Editor/Sequencer/Public/SequencerChannelTraits.h
+++ b/Engine/Source/Editor/Sequencer/Public/SequencerChannelTraits.h
@@ -53,6 +53,9 @@ struct TExtendKeyMenuParams
 	/** The section on which the channel resides */
 	TWeakObjectPtr<UMovieSceneSection> Section;
 
+	/** The owning object. Probably the same as Section, but can be different. */
+	TWeakObjectPtr<UObject> WeakOwner;
+
 	/** The channel on which the keys reside */
 	TMovieSceneChannelHandle<ChannelType> Channel;
 
