Sequencer: Time-warp channels are now mutable

diff --git a/Engine/Source/Editor/Sequencer/Private/MVVM/ViewModels/ChannelModel.cpp b/Engine/Source/Editor/Sequencer/Private/MVVM/ViewModels/ChannelModel.cpp
--- a/Engine/Source/Editor/Sequencer/Private/MVVM/ViewModels/ChannelModel.cpp
+++ b/Engine/Source/Editor/Sequencer/Private/MVVM/ViewModels/ChannelModel.cpp
@@ -6,6 +6,7 @@
 #include "Channels/MovieSceneChannel.h"
 #include "Channels/MovieSceneChannelHandle.h"
 #include "Channels/MovieSceneSectionChannelOverrideRegistry.h"
+#include "Channels/IMovieSceneChannelOwner.h"
 #include "CurveModel.h"
 #include "DetailsViewArgs.h"
 #include "HAL/PlatformCrt.h"
@@ -45,6 +46,7 @@
 #include "Types/SlateEnums.h"
 #include "Widgets/DeclarativeSyntaxSupport.h"
 #include "Widgets/SBoxPanel.h"
+#include "Algo/AnyOf.h"
 
 class ISequencer;
 class SWidget;
@@ -54,6 +56,92 @@ class SWidget;
 namespace UE::Sequencer
 {
 
+
+struct FDynamicChannelMuteExtension : IDynamicExtension, IMutableExtension
+{
+	TWeakObjectPtr<> WeakOwner;
+	FName ChannelName;
+
+	using Implements = TImplements<IDynamicExtension, IMutableExtension>;
+
+	UE_SEQUENCER_DECLARE_VIEW_MODEL_TYPE_ID(FDynamicChannelMuteExtension)
+
+	FDynamicChannelMuteExtension(TWeakObjectPtr<> InWeakOwner, FName InChannelName)
+		: WeakOwner(InWeakOwner)
+		, ChannelName(InChannelName)
+	{}
+
+	bool IsMuted() const override
+	{
+		IMovieSceneChannelOwner* ChannelOwner = Cast<IMovieSceneChannelOwner>(WeakOwner.Get());
+		return ChannelOwner && ChannelOwner->IsMuted(ChannelName);
+	}
+	void SetIsMuted(bool bIsMuted) override
+	{
+		IMovieSceneChannelOwner* ChannelOwner = Cast<IMovieSceneChannelOwner>(WeakOwner.Get());
+		if (ChannelOwner)
+		{
+			ChannelOwner->SetIsMuted(ChannelName, bIsMuted);
+		}
+	}
+};
+UE_SEQUENCER_DEFINE_VIEW_MODEL_TYPE_ID(FDynamicChannelMuteExtension);
+
+struct FDynamicChannelGroupMuteExtension : IDynamicExtension, IMutableExtension
+{
+	TWeakViewModelPtr<FChannelGroupModel> WeakChannelGroup;
+
+	using Implements = TImplements<IDynamicExtension, IMutableExtension>;
+
+	UE_SEQUENCER_DECLARE_VIEW_MODEL_TYPE_ID(FDynamicChannelGroupMuteExtension)
+
+	FDynamicChannelGroupMuteExtension(TWeakViewModelPtr<FChannelGroupModel> InWeakChannelGroup)
+		: WeakChannelGroup(InWeakChannelGroup)
+	{}
+
+	bool IsInheritable() const override
+	{
+		return false;
+	}
+
+	bool IsMuted() const override
+	{
+		TViewModelPtr<FChannelGroupModel> ChannelGroup = WeakChannelGroup.Pin();
+		if (!ChannelGroup)
+		{
+			return false;
+		}
+
+		for (const TWeakViewModelPtr<FChannelModel>& WeakChannel : ChannelGroup->GetChannels())
+		{
+			TViewModelPtr<IMutableExtension> Mutable = WeakChannel.ImplicitPin();
+			if (Mutable && Mutable->IsMuted())
+			{
+				return true;
+			}
+		}
+		return false;
+	}
+
+	void SetIsMuted(bool bIsMuted) override
+	{
+		TViewModelPtr<FChannelGroupModel> ChannelGroup = WeakChannelGroup.Pin();
+		if (ChannelGroup)
+		{
+			for (const TWeakViewModelPtr<FChannelModel>& WeakChannel : ChannelGroup->GetChannels())
+			{
+				TViewModelPtr<IMutableExtension> Mutable = WeakChannel.ImplicitPin();
+				if (Mutable)
+				{
+					Mutable->SetIsMuted(bIsMuted);
+				}
+			}
+		}
+	}
+};
+UE_SEQUENCER_DEFINE_VIEW_MODEL_TYPE_ID(FDynamicChannelGroupMuteExtension);
+
+
 FChannelModel::FChannelModel(FName InChannelName, TWeakPtr<ISequencerSection> InSection, FMovieSceneChannelHandle InChannel)
 	: KeyArea(MakeShared<IKeyArea>(InSection, InChannel))
 	, ChannelName(InChannelName)
@@ -84,6 +172,17 @@ void FChannelModel::Initialize(TWeakPtr<ISequencerSection> InSection, FMovieScen
 	{
 		KeyArea->Reinitialize(InSection, InChannel);
 	}
+
+	const FMovieSceneChannelMetaData* ChannelMetaData = InChannel.GetMetaData();
+	IMovieSceneChannelOwner*          ChannelOwner    = Cast<IMovieSceneChannelOwner>(KeyArea->GetOwningObject());
+	if (ChannelMetaData && ChannelOwner && ChannelOwner->GetCapabilities(KeyArea->GetName()).bSupportsMute)
+	{
+		AddDynamicExtension<FDynamicChannelMuteExtension>(KeyArea->GetOwningObject(), ChannelMetaData->Name);
+	}
+	else
+	{
+		RemoveDynamicExtension<FDynamicChannelMuteExtension>();
+	}
 }
 
 FMovieSceneChannel* FChannelModel::GetChannel() const
@@ -338,6 +437,19 @@ void FChannelGroupModel::CleanupChannels()
 	if (NumRemoved > 0)
 	{
 		++ChannelsSerialNumber;
+
+		const bool bShouldBeMutable = Algo::AnyOf(Channels, [](TWeakViewModelPtr<FChannelModel> In)
+			{
+				TViewModelPtr<IMutableExtension> Mutable = In.ImplicitPin();
+				return Mutable.IsValid();
+			}
+		);
+
+		const bool bIsMutable = CastDynamic<IMutableExtension>() != nullptr;
+		if (!bShouldBeMutable && bIsMutable)
+		{
+			RemoveDynamicExtension<FDynamicChannelGroupMuteExtension>();
+		}
 	}
 }
 
@@ -786,6 +898,28 @@ FChannelGroupOutlinerModel::FChannelGroupOutlinerModel(FName InChannelName, cons
 FChannelGroupOutlinerModel::~FChannelGroupOutlinerModel()
 {}
 
+void FChannelGroupOutlinerModel::OnUpdated()
+{
+	const bool bShouldBeMutable = Algo::AnyOf(Channels, 
+		[](TWeakViewModelPtr<FChannelModel> In)
+		{
+			TViewModelPtr<IMutableExtension> Mutable = In.ImplicitPin();
+			return Mutable.IsValid();
+		}
+	);
+
+	const bool bIsMutable = CastDynamic<IMutableExtension>() != nullptr;
+
+	if (bShouldBeMutable && !bIsMutable)
+	{
+		AddDynamicExtension<FDynamicChannelGroupMuteExtension>(SharedThis(this));
+	}
+	else if (!bShouldBeMutable && bIsMutable)
+	{
+		RemoveDynamicExtension<FDynamicChannelGroupMuteExtension>();
+	}
+}
+
 FOutlinerSizing FChannelGroupOutlinerModel::RecomputeSizing()
 {
 	FOutlinerSizing MaxSizing;
diff --git a/Engine/Source/Editor/Sequencer/Private/MVVM/ViewModels/TrackModelLayoutBuilder.cpp b/Engine/Source/Editor/Sequencer/Private/MVVM/ViewModels/TrackModelLayoutBuilder.cpp
--- a/Engine/Source/Editor/Sequencer/Private/MVVM/ViewModels/TrackModelLayoutBuilder.cpp
+++ b/Engine/Source/Editor/Sequencer/Private/MVVM/ViewModels/TrackModelLayoutBuilder.cpp
@@ -145,6 +145,10 @@ FTrackModelLayoutBuilder::~FTrackModelLayoutBuilder()
 	{
 		CompoundItem->RecomputeSizing();
 	}
+	for (TSharedPtr<FChannelGroupOutlinerModel> OutlinerChannelGroup : Root->GetDescendantsOfType<FChannelGroupOutlinerModel>())
+	{
+		OutlinerChannelGroup->OnUpdated();
+	}
 }
 
 void FTrackModelLayoutBuilder::RefreshLayout(TSharedPtr<FSectionModel> InSection)
diff --git a/Engine/Source/Editor/Sequencer/Private/MVVM/Views/SChannelView.cpp b/Engine/Source/Editor/Sequencer/Private/MVVM/Views/SChannelView.cpp
--- a/Engine/Source/Editor/Sequencer/Private/MVVM/Views/SChannelView.cpp
+++ b/Engine/Source/Editor/Sequencer/Private/MVVM/Views/SChannelView.cpp
@@ -7,6 +7,7 @@
 #include "MVVM/ViewModels/TrackAreaViewModel.h"
 #include "MVVM/ViewModels/SequencerEditorViewModel.h"
 #include "MVVM/Views/STrackAreaView.h"
+#include "MVVM/Extensions/IMutableExtension.h"
 #include "MVVM/Selection/Selection.h"
 
 #include "IKeyArea.h"
@@ -446,7 +447,6 @@ int32 SChannelView::OnPaint(const FPaintArgs& Args, const FGeometry& AllottedGeo
 {
 	static const FName SelectionColorName("SelectionColor");
 
-	LayerId = DrawLane(Args, AllottedGeometry, MyCullingRect, OutDrawElements, LayerId, InWidgetStyle, bParentEnabled);
 
 	TSharedPtr<FSequencer> Sequencer = LegacyGetSequencer();
 	FViewModelPtr Model = WeakModel.Pin();
@@ -455,6 +455,13 @@ int32 SChannelView::OnPaint(const FPaintArgs& Args, const FGeometry& AllottedGeo
 		return LayerId;
 	}
 
+	{
+		IMutableExtension* Mutable = Model->CastThis<IMutableExtension>();
+		bParentEnabled = bParentEnabled && (!Mutable || !Mutable->IsMuted());
+	}
+
+	LayerId = DrawLane(Args, AllottedGeometry, MyCullingRect, OutDrawElements, LayerId, InWidgetStyle, bParentEnabled);
+
 	const bool bIncludeThis = true;
 
 	TSharedPtr<STrackAreaView>      TrackAreaView = WeakTrackAreaView.Pin();
diff --git a/Engine/Source/Editor/Sequencer/Public/MVVM/ViewModels/ChannelModel.h b/Engine/Source/Editor/Sequencer/Public/MVVM/ViewModels/ChannelModel.h
--- a/Engine/Source/Editor/Sequencer/Public/MVVM/ViewModels/ChannelModel.h
+++ b/Engine/Source/Editor/Sequencer/Public/MVVM/ViewModels/ChannelModel.h
@@ -191,6 +191,10 @@ public:
 
 	void CleanupChannels();
 
+protected:
+
+	void UpdateMutability();
+
 protected:
 
 	TArray<TWeakViewModelPtr<FChannelModel>> Channels;
@@ -239,6 +243,9 @@ public:
 	/*~ ICurveEditorTreeItemExtension */
 	bool HasCurves() const override;
 	void BuildContextMenu(FMenuBuilder& MenuBuilder) override;
+	
+	void OnUpdated();
+	
 	TOptional<FString> GetUniquePathName() const override;
 
 private:
diff --git a/Engine/Source/Editor/SequencerCore/Private/MVVM/Extensions/DynamicExtensionContainer.cpp b/Engine/Source/Editor/SequencerCore/Private/MVVM/Extensions/DynamicExtensionContainer.cpp
--- a/Engine/Source/Editor/SequencerCore/Private/MVVM/Extensions/DynamicExtensionContainer.cpp
+++ b/Engine/Source/Editor/SequencerCore/Private/MVVM/Extensions/DynamicExtensionContainer.cpp
@@ -18,6 +18,19 @@ const void* FDynamicExtensionContainer::CastDynamic(FViewModelTypeID Type) const
 	return nullptr;
 }
 
+void FDynamicExtensionContainer::RemoveDynamicExtension(FViewModelTypeID Type)
+{
+	for (int32 Index = DynamicExtensions.Num()-1; Index >= 0; --Index)
+	{
+		const FDynamicExtensionInfo& DynamicExtension = DynamicExtensions[Index];
+		if (DynamicExtension.TypeTable->Cast(&DynamicExtension.Extension.Get(), Type.GetTypeID()) != nullptr)
+		{
+			DynamicExtensions.RemoveAt(Index, 1);
+			return;
+		}
+	}
+}
+
 FDynamicExtensionContainerIterator::FDynamicExtensionContainerIterator(IteratorType&& InIterator, FViewModelTypeID InType)
 	: CurrentExtension(nullptr)
 	, Iterator(MoveTemp(InIterator))
diff --git a/Engine/Source/Editor/SequencerCore/Private/MVVM/Extensions/IMutableExtension.cpp b/Engine/Source/Editor/SequencerCore/Private/MVVM/Extensions/IMutableExtension.cpp
--- a/Engine/Source/Editor/SequencerCore/Private/MVVM/Extensions/IMutableExtension.cpp
+++ b/Engine/Source/Editor/SequencerCore/Private/MVVM/Extensions/IMutableExtension.cpp
@@ -4,6 +4,17 @@
 
 namespace UE::Sequencer
 {
+ECachedMuteState CombinePropagatedChildFlags(ECachedMuteState ParentFlags, ECachedMuteState CombinedChildFlags)
+{
+	// If this item is mutable, but not muted, do not inherit mute state
+	if ((ParentFlags & (ECachedMuteState::Mutable | ECachedMuteState::Muted)) == ECachedMuteState::Mutable && EnumHasAnyFlags(CombinedChildFlags, ECachedMuteState::Muted))
+	{
+		CombinedChildFlags &= ~ECachedMuteState::Muted;
+		CombinedChildFlags |= ECachedMuteState::PartiallyMutedChildren;
+	}
+
+	return ParentFlags | CombinedChildFlags;
+}
 
 ECachedMuteState FMuteStateCacheExtension::ComputeFlagsForModel(const FViewModelPtr& ViewModel)
 {
@@ -22,6 +33,10 @@ ECachedMuteState FMuteStateCacheExtension::ComputeFlagsForModel(const FViewModel
 		{
 			ThisModelFlags |= ECachedMuteState::Muted;
 		}
+		if (Mutable->IsInheritable())
+		{
+			ThisModelFlags |= ECachedMuteState::Inheritable;
+		}
 	}
 
 	return ThisModelFlags;
@@ -29,6 +44,12 @@ ECachedMuteState FMuteStateCacheExtension::ComputeFlagsForModel(const FViewModel
 
 void FMuteStateCacheExtension::PostComputeChildrenFlags(const FViewModelPtr& ViewModel, ECachedMuteState& OutThisModelFlags, ECachedMuteState& OutPropagateToParentFlags)
 {
+	const bool bIsInheritable          = EnumHasAnyFlags(OutThisModelFlags, ECachedMuteState::Inheritable);
+	if (!bIsInheritable)
+	{
+		return;
+	}
+
 	// --------------------------------------------------------------------
 	// Handle mute state propagation
 	const bool bIsMutable              = EnumHasAnyFlags(OutThisModelFlags, ECachedMuteState::Mutable);
diff --git a/Engine/Source/Editor/SequencerCore/Private/MVVM/Views/OutlinerColumns/SMuteColumnWidget.cpp b/Engine/Source/Editor/SequencerCore/Private/MVVM/Views/OutlinerColumns/SMuteColumnWidget.cpp
--- a/Engine/Source/Editor/SequencerCore/Private/MVVM/Views/OutlinerColumns/SMuteColumnWidget.cpp
+++ b/Engine/Source/Editor/SequencerCore/Private/MVVM/Views/OutlinerColumns/SMuteColumnWidget.cpp
@@ -69,24 +69,18 @@ void SMuteColumnWidget::SetIsActive(const bool bInIsActive)
 void SMuteColumnWidget::SetIsActive(const FViewModelPtr& ViewModel, const bool bInIsActive)
 {
 	TViewModelPtr<IMutableExtension> Mutable = ViewModel.ImplicitCast();
-	if (bInIsActive)
+
+	// If this is mutable, mute only this
+	if (Mutable)
 	{
-		// If this is mutable, mute only this
-		if (Mutable)
-		{
-			Mutable->SetIsMuted(true);
-		}
-		// Otherwise mute mutable children of this (if any)
-		else for (TViewModelPtr<IMutableExtension> Child : ViewModel->GetDescendantsOfType<IMutableExtension>())
-		{
-			Child->SetIsMuted(true);
-		}
+		Mutable->SetIsMuted(bInIsActive);
 	}
-	else
+	// Otherwise mute mutable children of this (if any)
+	else for (TViewModelPtr<IMutableExtension> Child : ViewModel->GetDescendantsOfType<IMutableExtension>())
 	{
-		for (TViewModelPtr<IMutableExtension> Child : ViewModel->GetDescendantsOfType<IMutableExtension>(true))
+		if (Child->IsInheritable())
 		{
-			Child->SetIsMuted(false);
+			Child->SetIsMuted(bInIsActive);
 		}
 	}
 }
diff --git a/Engine/Source/Editor/SequencerCore/Public/MVVM/Extensions/DynamicExtensionContainer.h b/Engine/Source/Editor/SequencerCore/Public/MVVM/Extensions/DynamicExtensionContainer.h
--- a/Engine/Source/Editor/SequencerCore/Public/MVVM/Extensions/DynamicExtensionContainer.h
+++ b/Engine/Source/Editor/SequencerCore/Public/MVVM/Extensions/DynamicExtensionContainer.h
@@ -100,6 +100,7 @@ public:
 	}
 
 	SEQUENCERCORE_API const void* CastDynamic(FViewModelTypeID Type) const;
+	SEQUENCERCORE_API void RemoveDynamicExtension(FViewModelTypeID Type);
 
 	const void* CastDynamicChecked(FViewModelTypeID Type) const
 	{
@@ -137,6 +138,12 @@ protected:
 		return NewExtension.Get();
 	}
 
+	template<typename T>
+	void RemoveDynamicExtension()
+	{
+		RemoveDynamicExtension(T::ID);
+	}
+
 	void PostInitializeExtensions()
 	{
 		for (const FDynamicExtensionInfo& DynamicExtension : DynamicExtensions)
diff --git a/Engine/Source/Editor/SequencerCore/Public/MVVM/Extensions/HierarchicalCacheExtension.h b/Engine/Source/Editor/SequencerCore/Public/MVVM/Extensions/HierarchicalCacheExtension.h
--- a/Engine/Source/Editor/SequencerCore/Public/MVVM/Extensions/HierarchicalCacheExtension.h
+++ b/Engine/Source/Editor/SequencerCore/Public/MVVM/Extensions/HierarchicalCacheExtension.h
@@ -15,6 +15,12 @@ namespace UE::Sequencer
 
 class IHierarchicalCache;
 
+template<typename EnumType>
+EnumType CombinePropagatedChildFlags(EnumType ParentFlags, EnumType CombinedChildFlags)
+{
+	return ParentFlags | CombinedChildFlags;
+}
+
 DECLARE_MULTICAST_DELEGATE_OneParam(FPreUpdateCachesEvent, FViewModelPtr);
 
 class SEQUENCERCORE_API IHierarchicalCache
@@ -111,7 +117,7 @@ protected:
 		FlagsType ThisModelFlags              = IndividualItemFlags.Pop();
 		FlagsType FlagsPropagatedFromChildren = AccumulatedChildFlags.Pop();
 
-		ThisModelFlags |= FlagsPropagatedFromChildren;
+		ThisModelFlags = CombinePropagatedChildFlags(ThisModelFlags, FlagsPropagatedFromChildren);
 		PostComputeChildrenFlags(ViewModel, ThisModelFlags, AccumulatedChildFlags.Last());
 
 		AccumulatedChildFlags.Last() |= (ThisModelFlags & FlagsType::InheritedFromChildren);
diff --git a/Engine/Source/Editor/SequencerCore/Public/MVVM/Extensions/IMutableExtension.h b/Engine/Source/Editor/SequencerCore/Public/MVVM/Extensions/IMutableExtension.h
--- a/Engine/Source/Editor/SequencerCore/Public/MVVM/Extensions/IMutableExtension.h
+++ b/Engine/Source/Editor/SequencerCore/Public/MVVM/Extensions/IMutableExtension.h
@@ -27,6 +27,12 @@ public:
 
 	/** Set this item's mute state */
 	virtual void SetIsMuted(bool bIsMuted) = 0;
+
+	/** Returns whether this mutable can be muted by a parent, and should report its mute state to a parent */
+	virtual bool IsInheritable() const
+	{
+		return true;
+	}
 };
 
 enum class ECachedMuteState
@@ -38,10 +44,12 @@ enum class ECachedMuteState
 	Muted                    = 1 << 2,
 	PartiallyMutedChildren   = 1 << 3,
 	ImplicitlyMutedByParent  = 1 << 4,
+	Inheritable              = 1 << 5,
 
 	InheritedFromChildren = MutableChildren | PartiallyMutedChildren,
 };
 ENUM_CLASS_FLAGS(ECachedMuteState)
+SEQUENCERCORE_API ECachedMuteState CombinePropagatedChildFlags(ECachedMuteState ParentFlags, ECachedMuteState CombinedChildFlags);
 
 class SEQUENCERCORE_API FMuteStateCacheExtension
 	: public TFlagStateCacheExtension<ECachedMuteState>
diff --git a/Engine/Source/Runtime/MovieScene/Private/Evaluation/MovieSceneSequenceTransform.cpp b/Engine/Source/Runtime/MovieScene/Private/Evaluation/MovieSceneSequenceTransform.cpp
--- a/Engine/Source/Runtime/MovieScene/Private/Evaluation/MovieSceneSequenceTransform.cpp
+++ b/Engine/Source/Runtime/MovieScene/Private/Evaluation/MovieSceneSequenceTransform.cpp
@@ -845,6 +845,15 @@ void FMovieSceneSequenceTransform::Add(FFrameTime InOffset, FMovieSceneTimeWarpV
 	}
 	else
 	{
+		if (InTimeWarp.GetType() == EMovieSceneTimeWarpType::Custom)
+		{
+			UMovieSceneTimeWarpGetter* Getter = InTimeWarp.AsCustom();
+			if (Getter && Getter->IsMuted())
+			{
+				return;
+			}
+		}
+
 		NestedTransforms.Emplace(InOffset, MoveTemp(InTimeWarp));
 	}
 }
diff --git a/Engine/Source/Runtime/MovieScene/Private/Variants/MovieSceneTimeWarpGetter.cpp b/Engine/Source/Runtime/MovieScene/Private/Variants/MovieSceneTimeWarpGetter.cpp
--- a/Engine/Source/Runtime/MovieScene/Private/Variants/MovieSceneTimeWarpGetter.cpp
+++ b/Engine/Source/Runtime/MovieScene/Private/Variants/MovieSceneTimeWarpGetter.cpp
@@ -9,6 +9,8 @@ UMovieSceneTimeWarpGetter::UMovieSceneTimeWarpGetter()
 	// Allow time-warps to be accessible across different packages so that 
 	//     we can store them directly inside FMovieSceneSequenceTransforms
 	SetFlags(RF_Public);
+
+	bMuted = false;
 }
 
 EMovieSceneChannelProxyType UMovieSceneTimeWarpGetter::PopulateChannelProxy(FMovieSceneChannelProxyData& OutProxyData, EAllowTopLevelChannels AllowTopLevel)
@@ -20,3 +22,32 @@ bool UMovieSceneTimeWarpGetter::DeleteChannel(FMovieSceneTimeWarpVariant& OutVar
 {
 	return false;
 }
+
+UE::MovieScene::FChannelOwnerCapabilities UMovieSceneTimeWarpGetter::GetCapabilities(FName ChannelName) const
+{
+	UE::MovieScene::FChannelOwnerCapabilities Capabilities;
+	Capabilities.bSupportsMute = true;
+	return Capabilities;
+}
+
+bool UMovieSceneTimeWarpGetter::IsMuted(FName ChannelName) const
+{
+	return bMuted;
+}
+
+void UMovieSceneTimeWarpGetter::SetIsMuted(FName ChannelName, bool bIsMuted)
+{
+	Modify();
+	bMuted = bIsMuted;
+}
+
+bool UMovieSceneTimeWarpGetter::IsMuted() const
+{
+	return bMuted;
+}
+
+void UMovieSceneTimeWarpGetter::SetIsMuted(bool bIsMuted)
+{
+	Modify();
+	bMuted = bIsMuted;
+}
\ No newline at end of file
diff --git a/Engine/Source/Runtime/MovieScene/Public/Channels/IMovieSceneChannelOwner.h b/Engine/Source/Runtime/MovieScene/Public/Channels/IMovieSceneChannelOwner.h
new file mode 100644
--- /dev/null
+++ b/Engine/Source/Runtime/MovieScene/Public/Channels/IMovieSceneChannelOwner.h
@@ -0,0 +1,64 @@
+// Copyright Epic Games, Inc. All Rights Reserved.
+
+#pragma once
+
+#include "CoreTypes.h"
+#include "UObject/Interface.h"
+#include "IMovieSceneChannelOwner.generated.h"
+
+namespace UE::MovieScene
+{
+	struct FChannelOwnerCapabilities
+	{
+		FChannelOwnerCapabilities()
+			: bSupportsMute (false)
+		{}
+
+		uint8 bSupportsMute : 1;
+	};
+} // namespace UE::MovieScene
+
+
+/**
+ * Interface that can be added to a channel owner to implement various opt-in behaviors for that channel
+ */
+UINTERFACE(MinimalAPI)
+class UMovieSceneChannelOwner : public UInterface
+{
+public:
+
+	GENERATED_BODY()
+};
+
+class IMovieSceneChannelOwner
+{
+public:
+
+	GENERATED_BODY()
+
+	/**
+	 * Retrieve the capabilities for the channel on this interface
+	 **/
+	virtual UE::MovieScene::FChannelOwnerCapabilities GetCapabilities(FName ChannelName) const
+	{
+		return UE::MovieScene::FChannelOwnerCapabilities();
+	}
+
+	/**
+	 * Check whether the specified channel is muted.
+	 * @note: Only valid to be called if GetCapabilities(ChannelName).bSupportsMute == true
+	 */
+	virtual bool IsMuted(FName ChannelName) const
+	{
+		return false;
+	}
+
+	/**
+	 * Set the mute state for the specified channel.
+	 * @note: Only valid to be called if GetCapabilities(ChannelName).bSupportsMute == true
+	 */
+	virtual void SetIsMuted(FName ChannelName, bool bIsMuted)
+	{
+	}
+};
+
diff --git a/Engine/Source/Runtime/MovieScene/Public/Variants/MovieSceneTimeWarpGetter.h b/Engine/Source/Runtime/MovieScene/Public/Variants/MovieSceneTimeWarpGetter.h
--- a/Engine/Source/Runtime/MovieScene/Public/Variants/MovieSceneTimeWarpGetter.h
+++ b/Engine/Source/Runtime/MovieScene/Public/Variants/MovieSceneTimeWarpGetter.h
@@ -7,6 +7,7 @@
 #include "Math/Range.h"
 #include "Misc/FrameTime.h"
 #include "Misc/Optional.h"
+#include "Channels/IMovieSceneChannelOwner.h"
 #include "Variants/MovieSceneNumericVariantGetter.h"
 #include "MovieSceneTimeWarpGetter.generated.h"
 
@@ -29,7 +30,9 @@ namespace UE::MovieScene
  * Base class for all dynamic getter implementations of a FMovieSceneTimeWarpVariant
  */
 UCLASS(Abstract, MinimalAPI)
-class UMovieSceneTimeWarpGetter : public UMovieSceneNumericVariantGetter
+class UMovieSceneTimeWarpGetter
+	: public UMovieSceneNumericVariantGetter
+	, public IMovieSceneChannelOwner
 {
 public:
 
@@ -116,6 +119,13 @@ public:
 	 *            End abstract API
 	  ~~~~~~~~~~~~~~~~~~~~~~~~~~~~~~~~~~~~~~~~~~~*/
 
+	MOVIESCENE_API virtual UE::MovieScene::FChannelOwnerCapabilities GetCapabilities(FName ChannelName) const override;
+	MOVIESCENE_API virtual bool IsMuted(FName ChannelName) const override;
+	MOVIESCENE_API virtual void SetIsMuted(FName ChannelName, bool bIsMuted) override;
+
+	MOVIESCENE_API bool IsMuted() const;
+	MOVIESCENE_API void SetIsMuted(bool bIsMuted);
+
 public:
 
 	/**
@@ -137,4 +147,12 @@ public:
 	 * Attempt to delete this time-warp from a channel proxy if it matches the specified name
 	 */
 	MOVIESCENE_API virtual bool DeleteChannel(FMovieSceneTimeWarpVariant& OutVariant, FName ChannelName);
+
+protected:
+
+	/**
+	 * Whether this getter is muted or not. Default: false.
+	 */
+	UPROPERTY()
+	uint8 bMuted : 1;
 };
