Sequencer: Added the ability for channels to customize outliner context menu and columns when all channel models are of the same type

diff --git a/Engine/Source/Editor/Sequencer/Private/MVVM/ViewModels/ChannelModel.cpp b/Engine/Source/Editor/Sequencer/Private/MVVM/ViewModels/ChannelModel.cpp
--- a/Engine/Source/Editor/Sequencer/Private/MVVM/ViewModels/ChannelModel.cpp
+++ b/Engine/Source/Editor/Sequencer/Private/MVVM/ViewModels/ChannelModel.cpp
@@ -900,7 +900,32 @@ FChannelGroupOutlinerModel::~FChannelGroupOutlinerModel()
 
 void FChannelGroupOutlinerModel::OnUpdated()
 {
-	const bool bShouldBeMutable = Algo::AnyOf(Channels, 
+	WeakCommonChannelModel = nullptr;
+
+	// If all channels are the same type, assign the common channel
+	//    model for edit interactions and context menus
+	{
+		TViewModelPtr<FChannelModel> CommonChannel;
+		for (TWeakViewModelPtr<FChannelModel> WeakChannel : Channels)
+		{
+			if (TViewModelPtr<FChannelModel> Channel = WeakChannel.Pin())
+			{
+				if (!CommonChannel)
+				{
+					CommonChannel = Channel;
+				}
+				else if (CommonChannel->GetTypeTable().GetTypeID() != Channel->GetTypeTable().GetTypeID())
+				{
+					CommonChannel = nullptr;
+					break;
+				}
+			}
+		}
+		WeakCommonChannelModel = CommonChannel;
+	}
+
+
+	const bool bShouldBeMutable = Algo::AnyOf(Channels,
 		[](TWeakViewModelPtr<FChannelModel> In)
 		{
 			TViewModelPtr<IMutableExtension> Mutable = In.ImplicitPin();
@@ -973,6 +998,16 @@ TSharedPtr<SWidget> FChannelGroupOutlinerModel::CreateOutlinerViewForColumn(cons
 		return SNullWidget::NullWidget;
 	}
 
+	// Ask a common channel to populate the outliner column view first
+	if (TViewModelPtr<FChannelModel> CommonChannel = WeakCommonChannelModel.Pin())
+	{
+		TSharedPtr<SWidget> Widget = CommonChannel->CreateOutlinerViewForColumn(InParams, InColumnName);
+		if (Widget)
+		{
+			return Widget;
+		}
+	}
+
 	if (InColumnName == FCommonOutlinerNames::Label)
 	{
 		return SNew(SOutlinerItemViewBase, SharedThis(this), InParams.Editor, InParams.TreeViewRow);
@@ -1069,6 +1104,11 @@ void FChannelGroupOutlinerModel::CreateCurveModels(TArray<TUniquePtr<FCurveModel
 
 void FChannelGroupOutlinerModel::BuildContextMenu(FMenuBuilder& MenuBuilder)
 {
+	if (TViewModelPtr<FChannelModel> CommonChannel = WeakCommonChannelModel.Pin())
+	{
+		CommonChannel->BuildContextMenu(MenuBuilder, TViewModelPtr<FChannelGroupOutlinerModel>(this));
+	}
+
 	FOutlinerItemModelMixin::BuildContextMenu(MenuBuilder);
 
 	BuildChannelOverrideMenu(MenuBuilder);
diff --git a/Engine/Source/Editor/Sequencer/Public/MVVM/ViewModels/ChannelModel.h b/Engine/Source/Editor/Sequencer/Public/MVVM/ViewModels/ChannelModel.h
--- a/Engine/Source/Editor/Sequencer/Public/MVVM/ViewModels/ChannelModel.h
+++ b/Engine/Source/Editor/Sequencer/Public/MVVM/ViewModels/ChannelModel.h
@@ -46,6 +46,7 @@ namespace Sequencer
 
 class FSectionModel;
 class FSequenceModel;
+class FChannelGroupOutlinerModel;
 
 /**
  * Model for a single channel inside a section.
@@ -105,6 +106,10 @@ public:
 	void DrawKeys(TArrayView<const FKeyHandle> InKeyHandles, TArrayView<FKeyDrawParams> OutKeyDrawParams) override;
 	TUniquePtr<FCurveModel> CreateCurveModel() override;
 
+	/*~ Begin virtual interface */
+	virtual void BuildContextMenu(FMenuBuilder& MenuBuilder, TViewModelPtr<FChannelGroupOutlinerModel> GroupOwner) {}
+	virtual TSharedPtr<SWidget> CreateOutlinerViewForColumn(const FCreateOutlinerViewParams& InParams, const FName& InColumnName) { return nullptr; }
+	/*~ End virtual interface */
 private:
 
 	FLinearColor GetKeyBarColor() const;
@@ -255,6 +260,7 @@ private:
 private:
 
 	FOutlinerSizing ComputedSizing;
+	TWeakViewModelPtr<FChannelModel> WeakCommonChannelModel;
 };
 
 
