SequencerCore: Added Section parameter to CreateChannelModel overloads

This is actually required for successful creation of FChannelModels. There are no current implementations of this (I will be adding one for TimeWarp channels) so I don't suspect deprecation is required.

diff --git a/Engine/Plugins/Experimental/PropertyAnimator/Source/PropertyAnimatorEditor/Private/Sequencer/PropertyAnimatorEditorCurveChannelInterface.h b/Engine/Plugins/Experimental/PropertyAnimator/Source/PropertyAnimatorEditor/Private/Sequencer/PropertyAnimatorEditorCurveChannelInterface.h
--- a/Engine/Plugins/Experimental/PropertyAnimator/Source/PropertyAnimatorEditor/Private/Sequencer/PropertyAnimatorEditorCurveChannelInterface.h
+++ b/Engine/Plugins/Experimental/PropertyAnimator/Source/PropertyAnimatorEditor/Private/Sequencer/PropertyAnimatorEditorCurveChannelInterface.h
@@ -32,7 +32,7 @@ struct TPropertyAnimatorEditorCurveChannelInterface : ISequencerChannelInterface
 	virtual bool ShouldShowCurve_Raw(const FMovieSceneChannel*, UMovieSceneSection*) const override { return true; }
 	virtual bool SupportsCurveEditorModels_Raw(const FMovieSceneChannelHandle&) const override { return false; }
 	virtual TUniquePtr<FCurveModel> CreateCurveEditorModel_Raw(const FMovieSceneChannelHandle&, const UE::Sequencer::FCreateCurveEditorModelParams& Params) const override { return nullptr; }
-	virtual TSharedPtr<UE::Sequencer::FChannelModel> CreateChannelModel_Raw(const FMovieSceneChannelHandle&, FName) const override { return nullptr; }
+	virtual TSharedPtr<UE::Sequencer::FChannelModel> CreateChannelModel_Raw(const FMovieSceneChannelHandle&, const UE::Sequencer::FSectionModel& InSection, FName) const override { return nullptr; }
 	virtual TSharedPtr<UE::Sequencer::STrackAreaLaneView> CreateChannelView_Raw(const FMovieSceneChannelHandle&, TWeakPtr<UE::Sequencer::FChannelModel>, const UE::Sequencer::FCreateTrackLaneViewParams&) const override { return nullptr; }
 	virtual void ExtendSectionMenu_Raw(FMenuBuilder& MenuBuilder, TSharedPtr<FExtender> InMenuExtender, TConstArrayView<FMovieSceneChannelHandle> InChannels, const TArray<TWeakObjectPtr<UMovieSceneSection>>& InWeakSections, TWeakPtr<ISequencer> InWeakSequencer) const override;
 	virtual TSharedPtr<ISidebarChannelExtension> ExtendSidebarMenu_Raw(FMenuBuilder& MenuBuilder, TSharedPtr<FExtender> InMenuExtender, TConstArrayView<FMovieSceneChannelHandle> InChannels, const TArray<TWeakObjectPtr<UMovieSceneSection>>& InWeakSections, TWeakPtr<ISequencer> InWeakSequencer) const override;
diff --git a/Engine/Source/Editor/MovieSceneTools/Public/Channels/PerlinNoiseChannelInterface.h b/Engine/Source/Editor/MovieSceneTools/Public/Channels/PerlinNoiseChannelInterface.h
--- a/Engine/Source/Editor/MovieSceneTools/Public/Channels/PerlinNoiseChannelInterface.h
+++ b/Engine/Source/Editor/MovieSceneTools/Public/Channels/PerlinNoiseChannelInterface.h
@@ -195,7 +195,7 @@ struct TPerlinNoiseChannelInterface : ISequencerChannelInterface
 		return LayerId + 1;
 	}
 
-	virtual TSharedPtr<UE::Sequencer::FChannelModel> CreateChannelModel_Raw(const FMovieSceneChannelHandle& InChannelHandle, FName InChannelName) const override
+	virtual TSharedPtr<UE::Sequencer::FChannelModel> CreateChannelModel_Raw(const FMovieSceneChannelHandle& InChannelHandle, const UE::Sequencer::FSectionModel& InSection, FName InChannelName) const override
 	{
 		return nullptr;
 	}
diff --git a/Engine/Source/Editor/Sequencer/Private/ISequencerSection.cpp b/Engine/Source/Editor/Sequencer/Private/ISequencerSection.cpp
--- a/Engine/Source/Editor/Sequencer/Private/ISequencerSection.cpp
+++ b/Engine/Source/Editor/Sequencer/Private/ISequencerSection.cpp
@@ -114,7 +114,7 @@ void ISequencerSection::GenerateSectionLayout( ISectionLayoutBuilder& LayoutBuil
 
 	ISequencerModule* SequencerModule = &FModuleManager::LoadModuleChecked<ISequencerModule>("Sequencer");
 
-	auto ChannelFactory = [this, SequencerModule](FName InChannelName, const FMovieSceneChannelHandle& InChannel)
+	auto ChannelFactory = [this, SequencerModule](FName InChannelName, const FSectionModel& InSection, const FMovieSceneChannelHandle& InChannel)
 	{
 		TSharedPtr<FChannelModel> ChannelModel = this->ConstructChannelModel(InChannelName, InChannel);
 		if (!ChannelModel)
@@ -122,7 +122,7 @@ void ISequencerSection::GenerateSectionLayout( ISectionLayoutBuilder& LayoutBuil
 			ISequencerChannelInterface* EditorInterface = SequencerModule->FindChannelEditorInterface(InChannel.GetChannelTypeName());
 			if (EditorInterface)
 			{
-				ChannelModel = EditorInterface->CreateChannelModel_Raw(InChannel, InChannelName);
+				ChannelModel = EditorInterface->CreateChannelModel_Raw(InChannel, InSection, InChannelName);
 			}
 		}
 
diff --git a/Engine/Source/Editor/Sequencer/Private/MVVM/ViewModels/TrackModelLayoutBuilder.cpp b/Engine/Source/Editor/Sequencer/Private/MVVM/ViewModels/TrackModelLayoutBuilder.cpp
--- a/Engine/Source/Editor/Sequencer/Private/MVVM/ViewModels/TrackModelLayoutBuilder.cpp
+++ b/Engine/Source/Editor/Sequencer/Private/MVVM/ViewModels/TrackModelLayoutBuilder.cpp
@@ -134,7 +134,7 @@ void FHierarchicalModelListRefresher::Pop()
 FTrackModelLayoutBuilder::FTrackModelLayoutBuilder(TSharedPtr<FViewModel> InSharedTreeRoot)
 	: Root(InSharedTreeRoot)
 	, OutlinerList(InSharedTreeRoot, InSharedTreeRoot->GetChildList(EViewModelListType::Outliner))
-	, SequencerSection(nullptr)
+	, Section(nullptr)
 {
 }
 
@@ -160,17 +160,17 @@ void FTrackModelLayoutBuilder::RefreshLayout(TSharedPtr<FSectionModel> InSection
 
 	// Start off with the track area list adding to the top-level-channel list
 	TrackAreaList = FHierarchicalModelListRefresher(InSection, InSection->GetChildList(FTrackModel::GetTopLevelChannelType()));
-	SequencerSection = InSection->GetSectionInterface();
+	Section = InSection;
 
-	SequencerSection->GenerateSectionLayout(*this);
+	Section->GetSectionInterface()->GenerateSectionLayout(*this);
 
-	SequencerSection = nullptr;
+	Section = nullptr;
 	TrackAreaList = FHierarchicalModelListRefresher();
 }
 
 void FTrackModelLayoutBuilder::PushCategory(FName CategoryName, const FText& DisplayLabel, FGetMovieSceneTooltipText GetGroupTooltipTextDelegate, TFunction<TSharedPtr<FCategoryModel>(FName, const FText&)> OptionalFactory)
 {
-	check(TrackAreaList.IsValid() && SequencerSection);
+	check(TrackAreaList.IsValid() && Section);
 
 	auto CategoryNamePredicate = [CategoryName](const auto& InModel){ return InModel.GetCategoryName() == CategoryName; };
 
@@ -214,13 +214,13 @@ void FTrackModelLayoutBuilder::PushCategory(FName CategoryName, const FText& Dis
 
 void FTrackModelLayoutBuilder::PopCategory()
 {
-	check(TrackAreaList.IsValid() && SequencerSection);
+	check(TrackAreaList.IsValid() && Section);
 
 	TrackAreaList.Pop();
 	OutlinerList.Pop();
 }
 
-void FTrackModelLayoutBuilder::SetTopLevelChannel(const FMovieSceneChannelHandle& Channel, TFunction<TSharedPtr<UE::Sequencer::FChannelModel>(FName, const FMovieSceneChannelHandle&)> OptionalFactory)
+void FTrackModelLayoutBuilder::SetTopLevelChannel(const FMovieSceneChannelHandle& Channel, TFunction<TSharedPtr<UE::Sequencer::FChannelModel>(FName, const FSectionModel&, const FMovieSceneChannelHandle&)> OptionalFactory)
 {
 	ensureAlwaysMsgf(
 			OutlinerList.GetCurrentParent() == Root,
@@ -236,9 +236,9 @@ void FTrackModelLayoutBuilder::SetTopLevelChannel(const FMovieSceneChannelHandle
 	OutlinerList.Pop();
 }
 
-void FTrackModelLayoutBuilder::AddChannel(const FMovieSceneChannelHandle& Channel, TFunction<TSharedPtr<UE::Sequencer::FChannelModel>(FName, const FMovieSceneChannelHandle&)> OptionalFactory)
+void FTrackModelLayoutBuilder::AddChannel(const FMovieSceneChannelHandle& Channel, TFunction<TSharedPtr<UE::Sequencer::FChannelModel>(FName, const FSectionModel&, const FMovieSceneChannelHandle&)> OptionalFactory)
 {
-	check(TrackAreaList.IsValid() && SequencerSection);
+	check(TrackAreaList.IsValid() && Section);
 
 	// Since we always start off adding to the top level channel, point the track area list at the generic child list now
 	if (TrackAreaList.GetCurrentType() == FTrackModel::GetTopLevelChannelType())
@@ -250,9 +250,9 @@ void FTrackModelLayoutBuilder::AddChannel(const FMovieSceneChannelHandle& Channe
 	AddChannel(Channel, false, OptionalFactory);
 }
 
-void FTrackModelLayoutBuilder::AddChannel(const FMovieSceneChannelHandle& Channel, bool bIsTopLevel, TFunction<TSharedPtr<UE::Sequencer::FChannelModel>(FName, const FMovieSceneChannelHandle&)> OptionalFactory)
+void FTrackModelLayoutBuilder::AddChannel(const FMovieSceneChannelHandle& Channel, bool bIsTopLevel, TFunction<TSharedPtr<UE::Sequencer::FChannelModel>(FName, const FSectionModel&, const FMovieSceneChannelHandle&)> OptionalFactory)
 {
-	check(TrackAreaList.IsValid() && SequencerSection);
+	check(TrackAreaList.IsValid() && Section);
 
 	// @todo: this is all pretty crusty - we're currently linear-searching for both the child node, and the IKeyArea within that node
 	// Performance is generally acceptible however since we are dealing with small numbers of children, but this may need to be revisited.
@@ -311,12 +311,12 @@ void FTrackModelLayoutBuilder::AddChannel(const FMovieSceneChannelHandle& Channe
 		{
 			if (OptionalFactory)
 			{
-				TrackAreaModel = OptionalFactory(ChannelName, Channel);
+				TrackAreaModel = OptionalFactory(ChannelName, *Section, Channel);
 			}
 
 			if (!TrackAreaModel)
 			{
-				TrackAreaModel = MakeShared<FChannelModel>(ChannelName, SequencerSection, Channel);
+				TrackAreaModel = MakeShared<FChannelModel>(ChannelName, Section->GetSectionInterface(), Channel);
 			}
 		}
 
@@ -329,7 +329,7 @@ void FTrackModelLayoutBuilder::AddChannel(const FMovieSceneChannelHandle& Channe
 		{
 			TrackAreaModel->SetLinkedOutlinerItem(OutlinerModel.ImplicitCastChecked());
 		}
-		TrackAreaModel->Initialize(SequencerSection, Channel);
+		TrackAreaModel->Initialize(Section->GetSectionInterface(), Channel);
 
 		TrackAreaList.Link(TrackAreaModel);
 	}
diff --git a/Engine/Source/Editor/Sequencer/Private/MVVM/ViewModels/TrackModelLayoutBuilder.h b/Engine/Source/Editor/Sequencer/Private/MVVM/ViewModels/TrackModelLayoutBuilder.h
--- a/Engine/Source/Editor/Sequencer/Private/MVVM/ViewModels/TrackModelLayoutBuilder.h
+++ b/Engine/Source/Editor/Sequencer/Private/MVVM/ViewModels/TrackModelLayoutBuilder.h
@@ -113,20 +113,20 @@ public:
 	// ISectionLayoutBuilder interface
 
 	virtual void PushCategory( FName CategoryName, const FText& DisplayLabel, FGetMovieSceneTooltipText GetGroupTooltipTextDelegate, TFunction<TSharedPtr<UE::Sequencer::FCategoryModel>(FName, const FText&)> OptionalFactory ) override;
-	virtual void SetTopLevelChannel( const FMovieSceneChannelHandle& Channel, TFunction<TSharedPtr<UE::Sequencer::FChannelModel>(FName, const FMovieSceneChannelHandle&)> OptionalFactory ) override;
-	virtual void AddChannel( const FMovieSceneChannelHandle& Channel, TFunction<TSharedPtr<UE::Sequencer::FChannelModel>(FName, const FMovieSceneChannelHandle&)> OptionalFactory ) override;
+	virtual void SetTopLevelChannel( const FMovieSceneChannelHandle& Channel, TFunction<TSharedPtr<UE::Sequencer::FChannelModel>(FName, const FSectionModel&, const FMovieSceneChannelHandle&)> OptionalFactory ) override;
+	virtual void AddChannel( const FMovieSceneChannelHandle& Channel, TFunction<TSharedPtr<UE::Sequencer::FChannelModel>(FName, const FSectionModel&, const FMovieSceneChannelHandle&)> OptionalFactory ) override;
 	virtual void PopCategory() override;
 
 private:
 
-	void AddChannel( const FMovieSceneChannelHandle& Channel, bool bIsTopLevel, TFunction<TSharedPtr<UE::Sequencer::FChannelModel>(FName, const FMovieSceneChannelHandle&)> OptionalFactory );
+	void AddChannel( const FMovieSceneChannelHandle& Channel, bool bIsTopLevel, TFunction<TSharedPtr<UE::Sequencer::FChannelModel>(FName, const FSectionModel&, const FMovieSceneChannelHandle&)> OptionalFactory );
 
 private:
 
 	TSharedPtr<FViewModel> Root;
 	FHierarchicalModelListRefresher OutlinerList;
 	FHierarchicalModelListRefresher TrackAreaList;
-	TSharedPtr<ISequencerSection> SequencerSection;
+	TSharedPtr<FSectionModel> Section;
 };
 
 } // namespace Sequencer
diff --git a/Engine/Source/Editor/Sequencer/Public/ISectionLayoutBuilder.h b/Engine/Source/Editor/Sequencer/Public/ISectionLayoutBuilder.h
--- a/Engine/Source/Editor/Sequencer/Public/ISectionLayoutBuilder.h
+++ b/Engine/Source/Editor/Sequencer/Public/ISectionLayoutBuilder.h
@@ -14,6 +14,7 @@ namespace UE::Sequencer
 {
 	class FCategoryModel;
 	class FChannelModel;
+	class FSectionModel;
 }
 
 /**
@@ -36,14 +37,14 @@ public:
 	 * Sets the section as a key area itself
 	 * @param Channel		The channel that is to be assigned as the top level channel for this section
 	 */
-	virtual void SetTopLevelChannel( const FMovieSceneChannelHandle& Channel, TFunction<TSharedPtr<UE::Sequencer::FChannelModel>(FName, const FMovieSceneChannelHandle&)> OptionalFactory ) = 0;
+	virtual void SetTopLevelChannel( const FMovieSceneChannelHandle& Channel, TFunction<TSharedPtr<UE::Sequencer::FChannelModel>(FName, const UE::Sequencer::FSectionModel&, const FMovieSceneChannelHandle&)> OptionalFactory ) = 0;
 
 	/**
 	 * Adds a channel onto the layout. If a category is pushed, the key area will appear as a child of the current category
 	 *
 	 * @param Channel		A handle to the channel to be added to the layout
 	 */
-	virtual void AddChannel( const FMovieSceneChannelHandle& Channel, TFunction<TSharedPtr<UE::Sequencer::FChannelModel>(FName, const FMovieSceneChannelHandle&)> OptionalFactory ) = 0;
+	virtual void AddChannel( const FMovieSceneChannelHandle& Channel, TFunction<TSharedPtr<UE::Sequencer::FChannelModel>(FName, const UE::Sequencer::FSectionModel&, const FMovieSceneChannelHandle&)> OptionalFactory ) = 0;
 
 	/**
 	 * Pops a category off the stack
diff --git a/Engine/Source/Editor/Sequencer/Public/ISequencerChannelInterface.h b/Engine/Source/Editor/Sequencer/Public/ISequencerChannelInterface.h
--- a/Engine/Source/Editor/Sequencer/Public/ISequencerChannelInterface.h
+++ b/Engine/Source/Editor/Sequencer/Public/ISequencerChannelInterface.h
@@ -31,6 +31,7 @@ class FTrackInstancePropertyBindings;
 namespace UE::Sequencer
 {
 	class FChannelModel;
+	class FSectionModel;
 	class STrackAreaLaneView;
 	struct FCreateTrackLaneViewParams;
 
@@ -219,10 +220,11 @@ struct ISequencerChannelInterface
 	 * Create a new channel model for this type of channel
 	 *
 	 * @param InChannelHandle    The channel handle to create a model for
+	 * @param InSectionModel     The section that owns this channel model
 	 * @param InChannelName      The identifying name of this channel
 	 * @return (Optional) A new model to be added to a curve editor
 	 */
-	virtual TSharedPtr<UE::Sequencer::FChannelModel> CreateChannelModel_Raw(const FMovieSceneChannelHandle& InChannelHandle, FName InChannelName) const = 0;
+	virtual TSharedPtr<UE::Sequencer::FChannelModel> CreateChannelModel_Raw(const FMovieSceneChannelHandle& InChannelHandle, const UE::Sequencer::FSectionModel& InSection, FName InChannelName) const = 0;
 
 	/**
 	 * Create a new channel view for this type of channel
diff --git a/Engine/Source/Editor/Sequencer/Public/SequencerChannelInterface.h b/Engine/Source/Editor/Sequencer/Public/SequencerChannelInterface.h
--- a/Engine/Source/Editor/Sequencer/Public/SequencerChannelInterface.h
+++ b/Engine/Source/Editor/Sequencer/Public/SequencerChannelInterface.h
@@ -275,13 +275,14 @@ struct TSequencerChannelInterfaceCommon : ISequencerChannelInterface
 	 * Create a new channel model for this type of channel
 	 *
 	 * @param InChannelHandle    The channel handle to create a model for
+	 * @param InSectionModel     The section that owns this channel model
 	 * @param InChannelName      The identifying name of this channel
 	 * @return (Optional) A new model to be added to a curve editor
 	 */
-	virtual TSharedPtr<UE::Sequencer::FChannelModel> CreateChannelModel_Raw(const FMovieSceneChannelHandle& InChannelHandle, FName InChannelName) const override
+	virtual TSharedPtr<UE::Sequencer::FChannelModel> CreateChannelModel_Raw(const FMovieSceneChannelHandle& InChannelHandle, const UE::Sequencer::FSectionModel& InSection, FName InChannelName) const override
 	{
 		using namespace Sequencer;
-		return CreateChannelModel(InChannelHandle.Cast<ChannelType>(), InChannelName);
+		return CreateChannelModel(InChannelHandle.Cast<ChannelType>(), InSection, InChannelName);
 	}
 
 	/**
diff --git a/Engine/Source/Editor/Sequencer/Public/SequencerChannelTraits.h b/Engine/Source/Editor/Sequencer/Public/SequencerChannelTraits.h
--- a/Engine/Source/Editor/Sequencer/Public/SequencerChannelTraits.h
+++ b/Engine/Source/Editor/Sequencer/Public/SequencerChannelTraits.h
@@ -429,10 +429,11 @@ namespace Sequencer
 	 * Create a new channel model for this type of channel
 	 *
 	 * @param InChannelHandle    The channel handle to create a model for
+	 * @param InSectionModel     The section that owns this channel model
 	 * @param InChannelName      The identifying name of this channel
 	 * @return (Optional) A new model to be used as part of the Sequencer MVVM framework
 	 */
-	inline TSharedPtr<UE::Sequencer::FChannelModel> CreateChannelModel(const FMovieSceneChannelHandle& InChannelHandle, FName InChannelName)
+	inline TSharedPtr<UE::Sequencer::FChannelModel> CreateChannelModel(const FMovieSceneChannelHandle& InChannelHandle, const UE::Sequencer::FSectionModel& InSection, FName InChannelName)
 	{
 		return nullptr;
 	}
