Sequencer: Time Warp QoL Improvements

  - Added methods to FSequencerUtilities for populating a menu with the available timewarp curve types
  - Updated Skeletal Animation and Sub Sections to use these methods
  - Changed main + Time Warp Track entry to be a sub menu that lists the available curve types
  - Time Warp can now be customized via the button on the outliner, or the context menu
  - Time Warp Variant customization now shows fixed play rate inline

diff --git a/Engine/Source/Editor/MovieSceneTools/Private/Channels/BuiltInChannelEditors.cpp b/Engine/Source/Editor/MovieSceneTools/Private/Channels/BuiltInChannelEditors.cpp
--- a/Engine/Source/Editor/MovieSceneTools/Private/Channels/BuiltInChannelEditors.cpp
+++ b/Engine/Source/Editor/MovieSceneTools/Private/Channels/BuiltInChannelEditors.cpp
@@ -24,6 +24,8 @@
 #include "CurveKeyEditors/SEnumKeyEditor.h"
 #include "UObject/StructOnScope.h"
 #include "MVVM/Views/KeyDrawParams.h"
+#include "MVVM/ViewModels/TimeWarpChannelModel.h"
+#include "MVVM/ViewModels/SectionModel.h"
 #include "Framework/MultiBox/MultiBoxBuilder.h"
 #include "Channels/MovieSceneChannelProxy.h"
 #include "Channels/MovieSceneChannelEditorData.h"
@@ -2021,4 +2023,10 @@ bool ShouldShowCurve(const FMovieSceneDoubleChannel* Channel, UMovieSceneSection
 	return Channel->GetShowCurve();
 }
 
+TSharedPtr<UE::Sequencer::FChannelModel> CreateChannelModel(const TMovieSceneChannelHandle<FMovieSceneTimeWarpChannel>& InChannelHandle, const UE::Sequencer::FSectionModel& InSection, FName InChannelName)
+{
+	return MakeShared<UE::Sequencer::FTimeWarpChannelModel>(InChannelName, InSection.GetSectionInterface(), InChannelHandle);
+}
+
+
 #undef LOCTEXT_NAMESPACE
diff --git a/Engine/Source/Editor/MovieSceneTools/Private/Channels/BuiltInChannelEditors.h b/Engine/Source/Editor/MovieSceneTools/Private/Channels/BuiltInChannelEditors.h
--- a/Engine/Source/Editor/MovieSceneTools/Private/Channels/BuiltInChannelEditors.h
+++ b/Engine/Source/Editor/MovieSceneTools/Private/Channels/BuiltInChannelEditors.h
@@ -105,4 +105,6 @@ TUniquePtr<FCurveModel> CreateCurveEditorModel(const TMovieSceneChannelHandle<FM
 TUniquePtr<FCurveModel> CreateCurveEditorModel(const TMovieSceneChannelHandle<FMovieSceneTimeWarpChannel>& TimeWarpChannel, const UE::Sequencer::FCreateCurveEditorModelParams& Params);
 
 bool ShouldShowCurve(const FMovieSceneFloatChannel* Channel, UMovieSceneSection* InSection);
-bool ShouldShowCurve(const FMovieSceneDoubleChannel* Channel, UMovieSceneSection* InSection);
\ No newline at end of file
+bool ShouldShowCurve(const FMovieSceneDoubleChannel* Channel, UMovieSceneSection* InSection);
+
+TSharedPtr<UE::Sequencer::FChannelModel> CreateChannelModel(const TMovieSceneChannelHandle<FMovieSceneTimeWarpChannel>& InChannelHandle, const UE::Sequencer::FSectionModel& InSection, FName InChannelName);
\ No newline at end of file
diff --git a/Engine/Source/Editor/MovieSceneTools/Private/MVVM/MovieSceneToolsViewModelTypeIDs.cpp b/Engine/Source/Editor/MovieSceneTools/Private/MVVM/MovieSceneToolsViewModelTypeIDs.cpp
--- a/Engine/Source/Editor/MovieSceneTools/Private/MVVM/MovieSceneToolsViewModelTypeIDs.cpp
+++ b/Engine/Source/Editor/MovieSceneTools/Private/MVVM/MovieSceneToolsViewModelTypeIDs.cpp
@@ -5,6 +5,7 @@
 #include "MVVM/ViewModels/CameraCutTrackModel.h"
 #include "MVVM/ViewModels/CinematicShotTrackModel.h"
 #include "MVVM/ViewModels/BindingLifetimeTrackModel.h"
+#include "MVVM/ViewModels/TimeWarpChannelModel.h"
 
 
 namespace UE
@@ -16,6 +17,7 @@ namespace Sequencer
 UE_SEQUENCER_DEFINE_CASTABLE(FCameraCutTrackModel);
 UE_SEQUENCER_DEFINE_CASTABLE(FCinematicShotTrackModel);
 UE_SEQUENCER_DEFINE_CASTABLE(FBindingLifetimeTrackModel);
+UE_SEQUENCER_DEFINE_CASTABLE(FTimeWarpChannelModel);
 
 } // namespace Sequencer
 } // namespace UE
diff --git a/Engine/Source/Editor/MovieSceneTools/Private/MVVM/ViewModels/TimeWarpChannelModel.cpp b/Engine/Source/Editor/MovieSceneTools/Private/MVVM/ViewModels/TimeWarpChannelModel.cpp
new file mode 100644
--- /dev/null
+++ b/Engine/Source/Editor/MovieSceneTools/Private/MVVM/ViewModels/TimeWarpChannelModel.cpp
@@ -0,0 +1,74 @@
+// Copyright Epic Games, Inc. All Rights Reserved.
+
+#include "MVVM/ViewModels/TimeWarpChannelModel.h"
+
+#include "MVVM/ViewModels/SectionModel.h"
+#include "MVVM/ViewModels/OutlinerColumns/OutlinerColumnTypes.h"
+#include "MVVM/Extensions/ITrackExtension.h"
+#include "MVVM/Views/ViewUtilities.h"
+#include "SequencerUtilities.h"
+
+#include "Framework/MultiBox/MultiBoxBuilder.h"
+
+#define LOCTEXT_NAMESPACE "TimeWarpChannelModel"
+
+namespace UE::Sequencer
+{
+
+FTimeWarpChannelModel::FTimeWarpChannelModel(FName InChannelName, TWeakPtr<ISequencerSection> InSection, FMovieSceneChannelHandle InChannel)
+	: FChannelModel(InChannelName, InSection, InChannel)
+{}
+
+
+TSharedPtr<SWidget> FTimeWarpChannelModel::CreateOutlinerViewForColumn(const FCreateOutlinerViewParams& InParams, const FName& InColumnName)
+{
+	TViewModelPtr<IOutlinerExtension> OutlinerItem = GetLinkedOutlinerItem();
+
+	if (InColumnName == FCommonOutlinerNames::Add)
+	{
+		return MakeButton(
+			LOCTEXT("ChangeTimeWarpToolTip", "Change Time Warp to utilize a different curve type"),
+			FAppStyle::GetBrush("Sequencer.Tracks.Slomo"),
+			FOnGetContent::CreateSP(this, &FTimeWarpChannelModel::BuildReplaceTimeWarpSubMenu),
+			OutlinerItem.AsModel());
+	}
+
+	return nullptr;
+}
+
+
+void FTimeWarpChannelModel::BuildContextMenu(FMenuBuilder& MenuBuilder, TViewModelPtr<FChannelGroupOutlinerModel> GroupOwner)
+{
+	TViewModelPtr<ITrackExtension> Track = GetLinkedOutlinerItem().AsModel()->FindAncestorOfType<ITrackExtension>();
+	if (Track)
+	{
+		MenuBuilder.AddSubMenu(
+			LOCTEXT("ReplaceTimeWarp", "Replace With"),
+			FText(),
+			FNewMenuDelegate::CreateSP(this, &FTimeWarpChannelModel::PopulateReplaceTimeWarpSubMenu)
+		);
+	}
+}
+
+
+TSharedRef<SWidget> FTimeWarpChannelModel::BuildReplaceTimeWarpSubMenu()
+{
+	FMenuBuilder MenuBuilder(true, nullptr);
+	PopulateReplaceTimeWarpSubMenu(MenuBuilder);
+	return MenuBuilder.MakeWidget();
+}
+
+
+void FTimeWarpChannelModel::PopulateReplaceTimeWarpSubMenu(FMenuBuilder& MenuBuilder)
+{
+	TViewModelPtr<ITrackExtension> Track = GetLinkedOutlinerItem().AsModel()->FindAncestorOfType<ITrackExtension>();
+	if (Track)
+	{
+		FSequencerUtilities::PopulateTimeWarpChannelSubMenu(MenuBuilder, Track);
+	}
+}
+
+
+} // namespace UE::Sequencer
+
+#undef LOCTEXT_NAMESPACE
\ No newline at end of file
diff --git a/Engine/Source/Editor/MovieSceneTools/Private/MVVM/ViewModels/TimeWarpChannelModel.h b/Engine/Source/Editor/MovieSceneTools/Private/MVVM/ViewModels/TimeWarpChannelModel.h
new file mode 100644
--- /dev/null
+++ b/Engine/Source/Editor/MovieSceneTools/Private/MVVM/ViewModels/TimeWarpChannelModel.h
@@ -0,0 +1,37 @@
+// Copyright Epic Games, Inc. All Rights Reserved.
+
+#pragma once
+
+#include "Templates/SharedPointerFwd.h"
+
+#include "MVVM/ViewModels/ChannelModel.h"
+
+
+struct FMovieSceneChannelHandle;
+
+class FName;
+class ISequencerSection;
+
+
+namespace UE::Sequencer
+{
+
+struct FTimeWarpChannelModel : FChannelModel
+{
+	UE_SEQUENCER_DECLARE_CASTABLE(FTimeWarpChannelModel, FChannelModel);
+
+	FTimeWarpChannelModel(FName InChannelName, TWeakPtr<ISequencerSection> InSection, FMovieSceneChannelHandle InChannel);
+
+	/*~ Begin FChannelModel API */
+	TSharedPtr<SWidget> CreateOutlinerViewForColumn(const FCreateOutlinerViewParams& InParams, const FName& InColumnName) override;
+	void BuildContextMenu(FMenuBuilder& MenuBuilder, TViewModelPtr<FChannelGroupOutlinerModel> GroupOwner) override;
+	/*~ End FChannelModel API */
+
+private:
+
+	TSharedRef<SWidget> BuildReplaceTimeWarpSubMenu();
+	void PopulateReplaceTimeWarpSubMenu(FMenuBuilder& MenuBuilder);
+};
+
+} // namespace UE::Sequencer
+
diff --git a/Engine/Source/Editor/MovieSceneTools/Private/MovieSceneTimeWarpVariantCustomization.cpp b/Engine/Source/Editor/MovieSceneTools/Private/MovieSceneTimeWarpVariantCustomization.cpp
--- a/Engine/Source/Editor/MovieSceneTools/Private/MovieSceneTimeWarpVariantCustomization.cpp
+++ b/Engine/Source/Editor/MovieSceneTools/Private/MovieSceneTimeWarpVariantCustomization.cpp
@@ -3,13 +3,12 @@
 #include "MovieSceneTimeWarpVariantCustomization.h"
 #include "Variants/MovieSceneTimeWarpVariant.h"
 #include "Variants/MovieSceneTimeWarpGetter.h"
+#include "SequencerUtilities.h"
 
-#include "AssetRegistry/AssetRegistryModule.h"
-#include "ClassViewerModule.h"
-#include "ClassViewerFilter.h"
 #include "ScopedTransaction.h"
 
 #include "UObject/Class.h"
+#include "Templates/SubclassOf.h"
 
 #include "DetailWidgetRow.h"
 #include "IDetailChildrenBuilder.h"
@@ -65,28 +64,32 @@ void FMovieSceneTimeWarpVariantCustomization::CustomizeHeader(TSharedRef<IProper
 		StructPropertyHandle->CreatePropertyNameWidget()
 	]
 	.ValueContent()
+	.HAlign(HAlign_Fill)
+	.MinDesiredWidth(TOptional<float>())
 	.MaxDesiredWidth(TOptional<float>())
 	[
-		SNew(SComboButton)
-		.ForegroundColor(FSlateColor::UseForeground())
-		.ComboButtonStyle(FAppStyle::Get(), "SimpleComboButton")
-		.OnGetMenuContent(this, &FMovieSceneTimeWarpVariantCustomization::BuildTypePickerMenu)
-		.ButtonContent()
+		SNew(SHorizontalBox)
+
+		+ SHorizontalBox::Slot()
+		.HAlign(HAlign_Left)
+		.Padding(FMargin(0.0f, 0.0f, 6.0f, 0.0f))
 		[
-			SNew(STextBlock)
-			.Text(this, &FMovieSceneTimeWarpVariantCustomization::GetTypeComboLabel)
+			SNew(SComboButton)
+			.ForegroundColor(FSlateColor::UseForeground())
+			.OnGetMenuContent(this, &FMovieSceneTimeWarpVariantCustomization::BuildTypePickerMenu)
+			.ButtonContent()
+			[
+				SNew(STextBlock)
+				.Text(this, &FMovieSceneTimeWarpVariantCustomization::GetTypeComboLabel)
+			]
 		]
-	];
-}
 
-void FMovieSceneTimeWarpVariantCustomization::CustomizeChildren(TSharedRef<IPropertyHandle> StructPropertyHandle, IDetailChildrenBuilder& ChildBuilder, IPropertyTypeCustomizationUtils& StructCustomizationUtils)
-{
-	if (bIsFixed)
-	{
-		ChildBuilder.AddCustomRow(FText())
-		.ValueContent()
+		+ SHorizontalBox::Slot()
+		.HAlign(HAlign_Left)
+		.Padding(FMargin(6.0f, 0.0f, 2.0f, 0.0f))
 		[
 			SNew(SSpinBox<double>)
+			.Visibility(this, &FMovieSceneTimeWarpVariantCustomization::GetFixedVisibility)
 			.Style(FAppStyle::Get(), "Sequencer.HyperlinkSpinBox")
 			.Font(FAppStyle::GetFontStyle("Sequencer.FixedFont"))
 			.OnValueCommitted(this, &FMovieSceneTimeWarpVariantCustomization::OnCommitFixedPlayRate)
@@ -95,16 +98,17 @@ void FMovieSceneTimeWarpVariantCustomization::CustomizeChildren(TSharedRef<IProp
 			.MaxValue(TOptional<double>())
 			.OnEndSliderMovement(this, &FMovieSceneTimeWarpVariantCustomization::SetFixedPlayRate)
 			.Value(this, &FMovieSceneTimeWarpVariantCustomization::GetFixedPlayRate)
-		];
-	}
-	else if (Class.IsSet())
-	{
-		// All same external type
-	}
-	else
-	{
-	}
+		]
+	];
+}
+
+void FMovieSceneTimeWarpVariantCustomization::CustomizeChildren(TSharedRef<IPropertyHandle> StructPropertyHandle, IDetailChildrenBuilder& ChildBuilder, IPropertyTypeCustomizationUtils& StructCustomizationUtils)
+{
+}
 
+EVisibility FMovieSceneTimeWarpVariantCustomization::GetFixedVisibility() const
+{
+	return IsFixed() ? EVisibility::Visible : EVisibility::Collapsed;
 }
 
 void FMovieSceneTimeWarpVariantCustomization::OnCommitFixedPlayRate(double InValue, ETextCommit::Type Type)
@@ -239,23 +243,6 @@ void FMovieSceneTimeWarpVariantCustomization::SetFixed()
 
 TSharedRef<SWidget> FMovieSceneTimeWarpVariantCustomization::BuildTypePickerMenu()
 {
-	class FMovieSceneTimeWarpGetterFilter : public IClassViewerFilter
-	{
-	public:
-		bool IsClassAllowed(const FClassViewerInitializationOptions&, const UClass* InClass, TSharedRef<FClassViewerFilterFuncs>) override
-		{
-			return !InClass->HasAnyClassFlags(CLASS_Abstract) && InClass->IsChildOf(UMovieSceneTimeWarpGetter::StaticClass());
-		}
-		
-		bool IsUnloadedClassAllowed(
-			const FClassViewerInitializationOptions&,
-			const TSharedRef<const IUnloadedBlueprintData> InUnloadedClassData,
-			TSharedRef<FClassViewerFilterFuncs>) override
-		{
-			return !InUnloadedClassData->HasAnyClassFlags(CLASS_Abstract) && InUnloadedClassData->IsChildOf(UMovieSceneTimeWarpGetter::StaticClass());
-		}
-	};
-
 	bool bShouldCloseWindowAfterMenuSelection = false;
 	FMenuBuilder MenuBuilder(bShouldCloseWindowAfterMenuSelection, nullptr);
 
@@ -274,15 +261,9 @@ TSharedRef<SWidget> FMovieSceneTimeWarpVariantCustomization::BuildTypePickerMenu
 
 		MenuBuilder.AddSeparator();
 
-		FClassViewerInitializationOptions ClassViewerOptions;
-		ClassViewerOptions.Mode        = EClassViewerMode::ClassPicker;
-		ClassViewerOptions.DisplayMode = EClassViewerDisplayMode::ListView;
-		ClassViewerOptions.ClassFilters.Add(MakeShared<FMovieSceneTimeWarpGetterFilter>());
-
-		FClassViewerModule& ClassViewerModule = FModuleManager::LoadModuleChecked<FClassViewerModule>("ClassViewer");
-		TSharedRef<SWidget> ClassViewer = ClassViewerModule.CreateClassViewer(ClassViewerOptions, FOnClassPicked::CreateSP(this, &FMovieSceneTimeWarpVariantCustomization::ChangeClassType));
-
-		MenuBuilder.AddWidget(ClassViewer, FText(), true, false);
+		FSequencerUtilities::PopulateTimeWarpSubMenu(MenuBuilder, [this](TSubclassOf<UMovieSceneTimeWarpGetter> NewClass){
+			this->ChangeClassType(NewClass.Get());
+		});
 	}
 	MenuBuilder.EndSection();
 
diff --git a/Engine/Source/Editor/MovieSceneTools/Private/MovieSceneTimeWarpVariantCustomization.h b/Engine/Source/Editor/MovieSceneTools/Private/MovieSceneTimeWarpVariantCustomization.h
--- a/Engine/Source/Editor/MovieSceneTools/Private/MovieSceneTimeWarpVariantCustomization.h
+++ b/Engine/Source/Editor/MovieSceneTools/Private/MovieSceneTimeWarpVariantCustomization.h
@@ -5,6 +5,8 @@
 #include "IPropertyTypeCustomization.h"
 #include "Types/SlateEnums.h"
 
+struct EVisibility;
+
 class FText;
 class UClass;
 class SWidget;
@@ -30,6 +32,8 @@ private:
 
 	FText GetTypeComboLabel() const;
 
+	EVisibility GetFixedVisibility() const;
+
 	TSharedRef<SWidget> BuildTypePickerMenu();
 
 	void ChangeClassType(UClass* InClass);
diff --git a/Engine/Source/Editor/MovieSceneTools/Private/TrackEditors/CinematicShotTrackEditor.cpp b/Engine/Source/Editor/MovieSceneTools/Private/TrackEditors/CinematicShotTrackEditor.cpp
--- a/Engine/Source/Editor/MovieSceneTools/Private/TrackEditors/CinematicShotTrackEditor.cpp
+++ b/Engine/Source/Editor/MovieSceneTools/Private/TrackEditors/CinematicShotTrackEditor.cpp
@@ -83,7 +83,7 @@ TSharedPtr<SWidget> FCinematicShotTrackEditor::BuildOutlinerColumnWidget(const F
 	{
 		return UE::Sequencer::MakeAddButton(
 			LOCTEXT("CinematicShotText", "Shot"),
-			FOnGetContent::CreateSP(this, &FCinematicShotTrackEditor::HandleAddSubSequenceComboButtonGetMenuContent, Params.TrackModel->GetTrack()),
+			FOnGetContent::CreateSP(this, &FCinematicShotTrackEditor::HandleAddSubSequenceComboButtonGetMenuContent, Params.TrackModel.AsWeak()),
 			Params.ViewModel);
 	}
 
diff --git a/Engine/Source/Editor/MovieSceneTools/Private/TrackEditors/SkeletalAnimationTrackEditor.cpp b/Engine/Source/Editor/MovieSceneTools/Private/TrackEditors/SkeletalAnimationTrackEditor.cpp
--- a/Engine/Source/Editor/MovieSceneTools/Private/TrackEditors/SkeletalAnimationTrackEditor.cpp
+++ b/Engine/Source/Editor/MovieSceneTools/Private/TrackEditors/SkeletalAnimationTrackEditor.cpp
@@ -27,6 +27,7 @@
 #include "SequencerSettings.h"
 #include "MVVM/Views/ViewUtilities.h"
 #include "MVVM/ViewModels/ViewDensity.h"
+#include "MVVM/Extensions/ITrackExtension.h"
 #include "ISectionLayoutBuilder.h"
 #include "Animation/AnimMontage.h"
 #include "Animation/AnimSequence.h"
@@ -85,6 +86,7 @@
 #include "UObject/SavePackage.h"
 #include "AnimSequencerInstanceProxy.h"
 #include "TimeToPixel.h"
+#include "SequencerUtilities.h"
 #include "SequencerAnimationOverride.h"
 
 int32 FSkeletalAnimationTrackEditor::NumberActive = 0;
@@ -1999,10 +2001,34 @@ void FSkeletalAnimationTrackEditor::BuildObjectBindingTrackMenu(FMenuBuilder& Me
 	}
 }
 
+TSharedRef<SWidget> FSkeletalAnimationTrackEditor::BuildAddAnimationSubMenu(FGuid ObjectBinding, USkeleton* Skeleton, UE::Sequencer::TWeakViewModelPtr<UE::Sequencer::ITrackExtension> WeakTrackModel)
+{
+	FMenuBuilder MenuBuilder(true, nullptr);
+
+	TArray<FGuid> ObjectBindings;
+	ObjectBindings.Add(ObjectBinding);
+
+
+	MenuBuilder.BeginSection(NAME_None, LOCTEXT("TimeWarp_Label", "Time Warp"));
+	{
+		FSequencerUtilities::MakeTimeWarpMenuEntry(MenuBuilder, WeakTrackModel);
+	}
+	MenuBuilder.EndSection();
+
+
+	MenuBuilder.BeginSection(NAME_None, LOCTEXT("AddAnimation_Label", "Add Animation"));
+	{
+		AddAnimationSubMenu(MenuBuilder, ObjectBindings, Skeleton, WeakTrackModel.Pin()->GetTrack());
+	}
+	MenuBuilder.EndSection();
+
+	return MenuBuilder.MakeWidget();
+}
+
 TSharedRef<SWidget> FSkeletalAnimationTrackEditor::BuildAnimationSubMenu(FGuid ObjectBinding, USkeleton* Skeleton, UMovieSceneTrack* Track)
 {
 	FMenuBuilder MenuBuilder(true, nullptr);
-	
+
 	TArray<FGuid> ObjectBindings;
 	ObjectBindings.Add(ObjectBinding);
 
@@ -2315,7 +2341,8 @@ TSharedPtr<SWidget> FSkeletalAnimationTrackEditor::BuildOutlinerEditWidget(const
 
 	if (Skeleton)
 	{
-		return UE::Sequencer::MakeAddButton(LOCTEXT("AnimationText", "Animation"), FOnGetContent::CreateSP(this, &FSkeletalAnimationTrackEditor::BuildAnimationSubMenu, ObjectBinding, Skeleton, Track), Params.ViewModel);
+		FOnGetContent HandleGetAddButtonContent = FOnGetContent::CreateSP(this, &FSkeletalAnimationTrackEditor::BuildAddAnimationSubMenu, ObjectBinding, Skeleton, Params.TrackModel.AsWeak());
+		return UE::Sequencer::MakeAddButton(LOCTEXT("AnimationText", "Animation"), HandleGetAddButtonContent, Params.ViewModel);
 	}
 	else
 	{
diff --git a/Engine/Source/Editor/MovieSceneTools/Private/TrackEditors/SkeletalAnimationTrackEditor.h b/Engine/Source/Editor/MovieSceneTools/Private/TrackEditors/SkeletalAnimationTrackEditor.h
--- a/Engine/Source/Editor/MovieSceneTools/Private/TrackEditors/SkeletalAnimationTrackEditor.h
+++ b/Engine/Source/Editor/MovieSceneTools/Private/TrackEditors/SkeletalAnimationTrackEditor.h
@@ -13,6 +13,12 @@
 #include "ISequencerTrackEditor.h"
 #include "MovieSceneTrackEditor.h"
 #include "EditModes/SkeletalAnimationTrackEditMode.h"
+#include "SequencerCoreFwd.h"
+
+namespace UE::Sequencer
+{
+	class ITrackExtension;
+}
 
 struct FAssetData;
 struct FMovieSceneTimeWarpChannel;
@@ -80,6 +86,7 @@ public:
 private:
 
 	/** Animation sub menu */
+	TSharedRef<SWidget> BuildAddAnimationSubMenu(FGuid ObjectBinding, USkeleton* Skeleton, UE::Sequencer::TWeakViewModelPtr<UE::Sequencer::ITrackExtension> TrackModel);
 	TSharedRef<SWidget> BuildAnimationSubMenu(FGuid ObjectBinding, USkeleton* Skeleton, UMovieSceneTrack* Track);
 	void AddAnimationSubMenu(FMenuBuilder& MenuBuilder, TArray<FGuid> ObjectBindings, USkeleton* Skeleton, UMovieSceneTrack* Track);
 
diff --git a/Engine/Source/Editor/MovieSceneTools/Private/TrackEditors/SubTrackEditor.cpp b/Engine/Source/Editor/MovieSceneTools/Private/TrackEditors/SubTrackEditor.cpp
--- a/Engine/Source/Editor/MovieSceneTools/Private/TrackEditors/SubTrackEditor.cpp
+++ b/Engine/Source/Editor/MovieSceneTools/Private/TrackEditors/SubTrackEditor.cpp
@@ -19,6 +19,7 @@
 #include "IContentBrowserSingleton.h"
 #include "ContentBrowserModule.h"
 #include "MVVM/Views/ViewUtilities.h"
+#include "MVVM/Extensions/ITrackExtension.h"
 #include "SequencerSectionPainter.h"
 #include "TrackEditors/SubTrackEditorBase.h"
 #include "DragAndDrop/AssetDragDropOp.h"
@@ -31,6 +32,7 @@
 #include "EngineAnalytics.h"
 #include "Interfaces/IAnalyticsProvider.h"
 #include "Algo/Accumulate.h"
+#include "SequencerUtilities.h"
 #include "AssetToolsModule.h"
 #include "Interfaces/IMainFrameModule.h"
 #include "IDetailsView.h"
@@ -215,7 +217,7 @@ void FSubTrackEditor::BuildAddTrackMenu(FMenuBuilder& MenuBuilder)
 
 TSharedPtr<SWidget> FSubTrackEditor::BuildOutlinerEditWidget(const FGuid& ObjectBinding, UMovieSceneTrack* Track, const FBuildEditWidgetParams& Params)
 {
-	return UE::Sequencer::MakeAddButton(GetSubTrackName(), FOnGetContent::CreateSP(this, &FSubTrackEditor::HandleAddSubSequenceComboButtonGetMenuContent, Track), Params.ViewModel);
+	return UE::Sequencer::MakeAddButton(GetSubTrackName(), FOnGetContent::CreateSP(this, &FSubTrackEditor::HandleAddSubSequenceComboButtonGetMenuContent, Params.TrackModel.AsWeak()), Params.ViewModel);
 }
 
 
@@ -852,25 +854,40 @@ UMovieSceneSubTrack* FSubTrackEditor::FindOrCreateSubTrack(UMovieScene* MovieSce
 	return SubTrack;
 }
 
-TSharedRef<SWidget> FSubTrackEditor::HandleAddSubSequenceComboButtonGetMenuContent(UMovieSceneTrack* InTrack)
+TSharedRef<SWidget> FSubTrackEditor::HandleAddSubSequenceComboButtonGetMenuContent(UE::Sequencer::TWeakViewModelPtr<UE::Sequencer::ITrackExtension> WeakTrackModel)
 {
+	using namespace UE::Sequencer;
+
+	TViewModelPtr<ITrackExtension> TrackModel = WeakTrackModel.Pin();
+	if (!TrackModel)
+	{
+		return SNullWidget::NullWidget;
+	}
+
+	UMovieSceneTrack* Track = TrackModel->GetTrack();
+
 	FMenuBuilder MenuBuilder(true, nullptr);
 
-	MenuBuilder.AddMenuEntry(
-		FText::Join(FText::FromString(" "), LOCTEXT("InsertText", "Insert"), GetSubTrackName()),
-		LOCTEXT("InsertSectionTooltip", "Insert new sequence at current time"),
-		FSlateIcon(),
-		FUIAction(FExecuteAction::CreateSP(this, &FSubTrackEditor::InsertSection, InTrack))
-	);
+	MenuBuilder.BeginSection(NAME_None, LOCTEXT("TimeWarpCategory", "Time Warp"));
+	{
+		FSequencerUtilities::MakeTimeWarpMenuEntry(MenuBuilder, WeakTrackModel);
+	}
 
-	MenuBuilder.BeginSection(TEXT("ChooseSequence"), LOCTEXT("ChooseSequence", "Choose Sequence"));
+	MenuBuilder.BeginSection(TEXT("ChooseSequence"), LOCTEXT("InsertSequence", "Insert Sequence"));
 	{
-		UMovieSceneSequence* Sequence = GetSequencer() ? GetSequencer()->GetFocusedMovieSceneSequence() : nullptr;
+		MenuBuilder.AddMenuEntry(
+			FText::Format(LOCTEXT("CreateNewText", "Create New {0} Asset"), GetSubTrackName()),
+			FText::Format(LOCTEXT("CreateNewSectionTooltip", "Create new {0} asset and insert it at current time"), GetSubTrackName()),
+			FSlateIcon(),
+			FUIAction(FExecuteAction::CreateSP(this, &FSubTrackEditor::InsertSection, Track))
+		);
+
 
+		UMovieSceneSequence* Sequence = GetSequencer() ? GetSequencer()->GetFocusedMovieSceneSequence() : nullptr;
 		FAssetPickerConfig AssetPickerConfig;
 		{
-			AssetPickerConfig.OnAssetSelected = FOnAssetSelected::CreateRaw( this, &FSubTrackEditor::HandleAddSubSequenceComboButtonMenuEntryExecute, InTrack);
-			AssetPickerConfig.OnAssetEnterPressed = FOnAssetEnterPressed::CreateRaw( this, &FSubTrackEditor::HandleAddSubSequenceComboButtonMenuEntryEnterPressed, InTrack);
+			AssetPickerConfig.OnAssetSelected = FOnAssetSelected::CreateRaw( this, &FSubTrackEditor::HandleAddSubSequenceComboButtonMenuEntryExecute, Track);
+			AssetPickerConfig.OnAssetEnterPressed = FOnAssetEnterPressed::CreateRaw( this, &FSubTrackEditor::HandleAddSubSequenceComboButtonMenuEntryEnterPressed, Track);
 			AssetPickerConfig.bAllowNullSelection = false;
 			AssetPickerConfig.bAddFilterUI = true;
 			AssetPickerConfig.InitialAssetViewType = EAssetViewType::List;
diff --git a/Engine/Source/Editor/MovieSceneTools/Private/TrackEditors/TimeWarpTrackEditor.cpp b/Engine/Source/Editor/MovieSceneTools/Private/TrackEditors/TimeWarpTrackEditor.cpp
--- a/Engine/Source/Editor/MovieSceneTools/Private/TrackEditors/TimeWarpTrackEditor.cpp
+++ b/Engine/Source/Editor/MovieSceneTools/Private/TrackEditors/TimeWarpTrackEditor.cpp
@@ -4,6 +4,8 @@
 #include "ISequencer.h"
 #include "ISequencerSection.h"
 #include "SequencerSectionPainter.h"
+#include "SequencerUtilities.h"
+#include "Variants/MovieSceneTimeWarpGetter.h"
 #include "Framework/MultiBox/MultiBoxBuilder.h"
 
 #include "Sections/MovieSceneTimeWarpSection.h"
@@ -89,13 +91,17 @@ struct FTimeWarpSection : FSequencerSection
 
 void FTimeWarpTrackEditor::BuildAddTrackMenu(FMenuBuilder& MenuBuilder)
 {
-	MenuBuilder.AddMenuEntry(
+	auto HandleAddTimeWarp = [this](TSubclassOf<UMovieSceneTimeWarpGetter> InClass)
+	{
+		this->HandleAddTimeWarpTrack(InClass);
+	};
+
+	MenuBuilder.AddSubMenu(
 		LOCTEXT("AddTimeWarpTrack", "Time Warp"),
 		LOCTEXT("AddTimeWarpTrackTooltip", "Adds a new track that manipulates the time of the current sequence."),
-		FSlateIcon(FAppStyle::GetAppStyleSetName(), "Sequencer.Tracks.Slomo"),
-		FUIAction(
-			FExecuteAction::CreateRaw(this, &FTimeWarpTrackEditor::HandleAddTimeWarpTrack)
-		)
+		FNewMenuDelegate::CreateStatic(FSequencerUtilities::PopulateTimeWarpSubMenu, TFunction<void(TSubclassOf<UMovieSceneTimeWarpGetter>)>(HandleAddTimeWarp)),
+		false,
+		FSlateIcon(FAppStyle::GetAppStyleSetName(), "Sequencer.Tracks.Slomo")
 	);
 }
 
@@ -104,7 +110,7 @@ TSharedRef<ISequencerSection> FTimeWarpTrackEditor::MakeSectionInterface(UMovieS
 	return MakeShared<UE::Sequencer::FTimeWarpSection>(SectionObject);
 }
 
-void FTimeWarpTrackEditor::HandleAddTimeWarpTrack()
+void FTimeWarpTrackEditor::HandleAddTimeWarpTrack(TSubclassOf<UMovieSceneTimeWarpGetter> ClassType)
 {
 	UMovieScene* FocusedMovieScene = GetFocusedMovieScene();
 	if (FocusedMovieScene == nullptr)
@@ -127,8 +133,17 @@ void FTimeWarpTrackEditor::HandleAddTimeWarpTrack()
 
 	FocusedMovieScene->Modify();
 
-	UMovieSceneTimeWarpTrack* NewTrack = NewObject<UMovieSceneTimeWarpTrack>(FocusedMovieScene, NAME_None, RF_Transactional);
-	NewTrack->AddSection(*NewTrack->CreateNewSection());
+	UMovieSceneTimeWarpTrack*   NewTrack   = NewObject<UMovieSceneTimeWarpTrack>(FocusedMovieScene, NAME_None, RF_Transactional);
+	UMovieSceneSection*         NewSection = NewTrack->CreateNewSection();
+	FMovieSceneTimeWarpVariant* TimeWarp   = NewSection->GetTimeWarp();
+
+	check(TimeWarp);
+
+	UMovieSceneTimeWarpGetter* NewGetter = NewObject<UMovieSceneTimeWarpGetter>(NewSection, ClassType.Get(), NAME_None, RF_Transactional);
+	NewGetter->InitializeDefaults();
+	TimeWarp->Set(NewGetter);
+
+	NewTrack->AddSection(*NewSection);
 
 	FocusedMovieScene->AddGivenTrack(NewTrack);
 	SequencerPtr->OnAddTrack(NewTrack, FGuid());
diff --git a/Engine/Source/Editor/MovieSceneTools/Private/TrackEditors/TimeWarpTrackEditor.h b/Engine/Source/Editor/MovieSceneTools/Private/TrackEditors/TimeWarpTrackEditor.h
--- a/Engine/Source/Editor/MovieSceneTools/Private/TrackEditors/TimeWarpTrackEditor.h
+++ b/Engine/Source/Editor/MovieSceneTools/Private/TrackEditors/TimeWarpTrackEditor.h
@@ -5,6 +5,9 @@
 #include "KeyframeTrackEditor.h"
 #include "Tracks/MovieSceneTimeWarpTrack.h"
 
+class UMovieSceneTimeWarpGetter;
+
+template <typename T> class TSubclassOf;
 
 class FTimeWarpTrackEditor
 	: public FKeyframeTrackEditor<UMovieSceneTimeWarpTrack>
@@ -24,5 +27,5 @@ private:
 	virtual void BuildAddTrackMenu(FMenuBuilder& MenuBuilder) override;
 	virtual TSharedRef<ISequencerSection> MakeSectionInterface(UMovieSceneSection& SectionObject, UMovieSceneTrack& Track, FGuid ObjectBinding) override;
 
-	void HandleAddTimeWarpTrack();
+	void HandleAddTimeWarpTrack(TSubclassOf<UMovieSceneTimeWarpGetter> ClassType);
 };
diff --git a/Engine/Source/Editor/MovieSceneTools/Public/TrackEditors/SubTrackEditor.h b/Engine/Source/Editor/MovieSceneTools/Public/TrackEditors/SubTrackEditor.h
--- a/Engine/Source/Editor/MovieSceneTools/Public/TrackEditors/SubTrackEditor.h
+++ b/Engine/Source/Editor/MovieSceneTools/Public/TrackEditors/SubTrackEditor.h
@@ -138,7 +138,7 @@ protected:
 	UMovieSceneSubTrack* FindOrCreateSubTrack(UMovieScene* MovieScene, UMovieSceneTrack* Track) const;
 
 	/** Callback for generating the menu of the "Add Sequence" combo button. */
-	TSharedRef<SWidget> HandleAddSubSequenceComboButtonGetMenuContent(UMovieSceneTrack* InTrack);
+	TSharedRef<SWidget> HandleAddSubSequenceComboButtonGetMenuContent(UE::Sequencer::TWeakViewModelPtr<UE::Sequencer::ITrackExtension> TrackModel);
 
 private:
 
diff --git a/Engine/Source/Editor/Sequencer/Private/SequencerUtilities.cpp b/Engine/Source/Editor/Sequencer/Private/SequencerUtilities.cpp
--- a/Engine/Source/Editor/Sequencer/Private/SequencerUtilities.cpp
+++ b/Engine/Source/Editor/Sequencer/Private/SequencerUtilities.cpp
@@ -58,7 +58,9 @@
 #include "HAL/PlatformApplicationMisc.h"
 #include "LevelSequence.h"
 #include "MVVM/Extensions/IObjectBindingExtension.h"
+#include "MVVM/Extensions/IOutlinerExtension.h"
 #include "MVVM/ViewModels/TrackModel.h"
+#include "MVVM/ViewModels/SectionModel.h"
 #include "MVVM/Views/ViewUtilities.h"
 #include "Framework/Notifications/NotificationManager.h"
 #include "Widgets/Notifications/SNotificationList.h"
@@ -69,11 +71,15 @@
 #include "Tracks/MovieSceneBindingLifetimeTrack.h"
 #include "ClassViewerModule.h"
 #include "ClassViewerFilter.h"
+#include "ClassIconFinder.h"
+#include "Styling/SlateIconFinder.h"
 #include "Framework/Application/SlateApplication.h"
 #include "SequencerCommands.h"
 #include "MVVM/Selection/Selection.h"
 #include "UnrealEdGlobals.h"
 #include "Misc/FeedbackContext.h"
+#include "Variants/MovieSceneTimeWarpVariant.h"
+#include "Variants/MovieSceneTimeWarpGetter.h"
 #include "UObject/UObjectIterator.h"
 #include "ObjectTools.h"
 
@@ -105,6 +111,147 @@ TSharedRef<SWidget> FSequencerUtilities::MakeAddButton(FText HoverText, FOnClick
 	return UE::Sequencer::MakeAddButton(HoverText, OnClicked, HoverState, IsEnabled);
 }
 
+void FSequencerUtilities::MakeTimeWarpMenuEntry(FMenuBuilder& MenuBuilder, UE::Sequencer::TWeakViewModelPtr<UE::Sequencer::ITrackExtension> WeakTrackModel)
+{
+	using namespace UE::Sequencer;
+
+	TViewModelPtr<ITrackExtension> TrackModel = WeakTrackModel.Pin();
+	if (!TrackModel)
+	{
+		return;
+	}
+
+	TOptional<UClass*> CommonClass;
+	for (const TViewModelPtr<FSectionModel>& SectionModel : TrackModel->GetSectionModels().IterateSubList<FSectionModel>())
+	{
+		UMovieSceneSection* Section = SectionModel->GetSection();
+		if (!Section)
+		{
+			continue;
+		}
+
+		FMovieSceneTimeWarpVariant* Variant = Section->GetTimeWarp();
+		UMovieSceneTimeWarpGetter*  Getter  = Variant && Variant->GetType() == EMovieSceneTimeWarpType::Custom 
+			? Variant->AsCustom()
+			: nullptr;
+
+		if (Getter)
+		{
+			if (!CommonClass)
+			{
+				CommonClass = Getter->GetClass();
+			}
+			else if (CommonClass.GetValue() != Getter->GetClass())
+			{
+				CommonClass = nullptr;
+			}
+		}
+	}
+
+	FText TimeWarpLabel = CommonClass.IsSet()
+		? LOCTEXT("ReplaceTimeWarp_Label", "Replace Time Warp")
+		: LOCTEXT("AddTimeWarp_Label", "Add Time Warp");
+	FText TimeWarpToolTip = CommonClass.IsSet()
+		? LOCTEXT("ReplaceTimeWarp_ToolTip", "Replaces the Time Warp implementation with a different kind")
+		: LOCTEXT("AddTimeWarp_ToolTip", "Add Time Warp");
+
+	MenuBuilder.AddSubMenu(
+		TimeWarpLabel,
+		TimeWarpToolTip,
+		FNewMenuDelegate::CreateStatic(PopulateTimeWarpChannelSubMenu, WeakTrackModel)
+	);
+}
+
+void FSequencerUtilities::PopulateTimeWarpSubMenu(FMenuBuilder& MenuBuilder, TFunction<void(TSubclassOf<UMovieSceneTimeWarpGetter>)> OnTimeWarpPicked)
+{
+	using namespace UE::Sequencer;
+
+	FAssetRegistryModule& AssetRegistryModule = FModuleManager::LoadModuleChecked<FAssetRegistryModule>("AssetRegistry");
+
+	TSet<FTopLevelAssetPath> AllTimeWarpClasses;
+	{
+		FTopLevelAssetPath TargetClassPath(UMovieSceneTimeWarpGetter::StaticClass());
+		AssetRegistryModule.Get().GetDerivedClassNames({ TargetClassPath }, TSet<FTopLevelAssetPath>(), AllTimeWarpClasses);
+		AllTimeWarpClasses.Remove(TargetClassPath);
+	}
+
+	if (AllTimeWarpClasses.Num() == 0)
+	{
+		MenuBuilder.AddWidget(SNew(STextBlock).Text(LOCTEXT("NoTimeWarpTypesError", "No Time Warp implementations found")), FText(), true);
+		return;
+	}
+
+	auto HandleTimeWarpSelection = [OnTimeWarpPicked](FTopLevelAssetPath ClassPath)
+	{
+		UClass* Class = FSoftClassPath(ClassPath.ToString()).TryLoadClass<UMovieSceneTimeWarpGetter>();
+		if (Class)
+		{
+			OnTimeWarpPicked(Class);
+		}
+	};
+
+	MenuBuilder.BeginSection(NAME_None, LOCTEXT("TimeWarpCategoryLabel", "Time Warp Types:"));
+
+	for (const FTopLevelAssetPath& ClassPath : AllTimeWarpClasses)
+	{
+		FAssetData AssetData = AssetRegistryModule.Get().GetAssetByObjectPath(FSoftObjectPath(ClassPath.ToString()));
+
+		const UClass* IconClass = FClassIconFinder::GetIconClassForAssetData(AssetData);
+		const UClass* Class     = Cast<const UClass>(AssetData.FastGetAsset());
+
+		MenuBuilder.AddMenuEntry(
+			Class ? Class->GetDisplayNameText() : FText::FromName(ClassPath.GetAssetName()),
+			Class ? Class->GetToolTipText()     : FText(),
+			FSlateIconFinder::FindIconForClass(IconClass),
+			FUIAction(FExecuteAction::CreateLambda(HandleTimeWarpSelection, ClassPath))
+		);
+	}
+
+	MenuBuilder.EndSection();
+}
+
+void FSequencerUtilities::PopulateTimeWarpChannelSubMenu(FMenuBuilder& MenuBuilder, UE::Sequencer::TWeakViewModelPtr<UE::Sequencer::ITrackExtension> WeakTrackModel)
+{
+	using namespace UE::Sequencer;
+
+	auto HandleTimeWarpSelection = [WeakTrackModel](TSubclassOf<UMovieSceneTimeWarpGetter> Class)
+	{
+		TViewModelPtr<ITrackExtension> TrackModel = WeakTrackModel.Pin();
+		if (!TrackModel)
+		{
+			return;
+		}
+
+		FScopedTransaction Transaction(LOCTEXT("ChangeTimeWarpType", "Changed Time Warp type"));
+
+		for (const TViewModelPtr<FSectionModel>& SectionModel : TrackModel->GetSectionModels().IterateSubList<FSectionModel>())
+		{
+			UMovieSceneSection*         Section = SectionModel->GetSection();
+			FMovieSceneTimeWarpVariant* Variant = Section ? Section->GetTimeWarp() : nullptr;
+
+			if (Variant)
+			{
+				Section->Modify();
+
+				UMovieSceneTimeWarpGetter* Getter = NewObject<UMovieSceneTimeWarpGetter>(Section, Class.Get(), NAME_None, RF_Transactional);
+				Getter->InitializeDefaults();
+
+				Variant->Set(Getter);
+
+				Section->InvalidateChannelProxy();
+
+				TViewModelPtr<IOutlinerExtension> Outliner = TrackModel.ImplicitCast();
+				if (Outliner && !Outliner->IsExpanded())
+				{
+					Outliner->SetExpansion(true);
+				}
+			}
+		}
+	};
+
+	PopulateTimeWarpSubMenu(MenuBuilder, HandleTimeWarpSelection);
+}
+
 void FSequencerUtilities::CreateNewSection(UMovieSceneTrack* InTrack, TWeakPtr<ISequencer> InSequencer, int32 InRowIndex, EMovieSceneBlendType InBlendType)
 {
 	TSharedPtr<ISequencer> Sequencer = InSequencer.Pin();
diff --git a/Engine/Source/Editor/Sequencer/Public/SequencerUtilities.h b/Engine/Source/Editor/Sequencer/Public/SequencerUtilities.h
--- a/Engine/Source/Editor/Sequencer/Public/SequencerUtilities.h
+++ b/Engine/Source/Editor/Sequencer/Public/SequencerUtilities.h
@@ -7,15 +7,18 @@
 #include "Framework/SlateDelegates.h"
 #include "Misc/FrameTime.h"
 #include "MovieSceneBindingProxy.h"
+#include "SequencerCoreFwd.h"
 #include "SequencerUtilities.generated.h"
 
 class ACineCameraActor;
-template< typename ObjectType > class TAttribute;
+template <typename> class TAttribute;
+template <typename> class TSubclassOf;
 class UActorFactory;
 class UMovieSceneFolder;
 class UMovieSceneTrack;
 class UMovieSceneSection;
 class UMovieSceneSequence;
+class UMovieSceneTimeWarpGetter;
 class ISequencer;
 class FMenuBuilder;
 struct FMovieSceneBinding;
@@ -29,6 +32,8 @@ enum class EMovieSceneBlendType : uint8;
 namespace UE::Sequencer
 {
 
+class ITrackExtension;
+
 struct FCreateBindingParams
 {
 	UE_DEPRECATED(5.4, "Please use FSequencerUtilitiesCreateBindingParams directly.")
@@ -195,6 +200,12 @@ struct SEQUENCER_API FSequencerUtilities
 	/* Creates a button (used for +Section) that fires a user-defined OnClick response with no sub-menu. */
 	static TSharedRef<SWidget> MakeAddButton(FText HoverText, FOnClicked OnClicked, const TAttribute<bool>& HoverState, TWeakPtr<ISequencer> InSequencer);
 
+	static void MakeTimeWarpMenuEntry(FMenuBuilder& MenuBuilder, UE::Sequencer::TWeakViewModelPtr<UE::Sequencer::ITrackExtension> TrackModel);
+
+	static void PopulateTimeWarpSubMenu(FMenuBuilder& MenuBuilder, TFunction<void(TSubclassOf<UMovieSceneTimeWarpGetter>)> OnTimeWarpPicked);
+
+	static void PopulateTimeWarpChannelSubMenu(FMenuBuilder& MenuBuilder, UE::Sequencer::TWeakViewModelPtr<UE::Sequencer::ITrackExtension> TrackModel);
+
 	static void CreateNewSection(UMovieSceneTrack* InTrack, TWeakPtr<ISequencer> InSequencer, int32 InRowIndex, EMovieSceneBlendType InBlendType);
 
 	static void PopulateMenu_CreateNewSection(FMenuBuilder& MenuBuilder, int32 RowIndex, UMovieSceneTrack* Track, TWeakPtr<ISequencer> InSequencer);
diff --git a/Engine/Source/Runtime/MovieScene/Private/Sections/MovieSceneTimeWarpSection.cpp b/Engine/Source/Runtime/MovieScene/Private/Sections/MovieSceneTimeWarpSection.cpp
--- a/Engine/Source/Runtime/MovieScene/Private/Sections/MovieSceneTimeWarpSection.cpp
+++ b/Engine/Source/Runtime/MovieScene/Private/Sections/MovieSceneTimeWarpSection.cpp
@@ -10,13 +10,6 @@ UMovieSceneTimeWarpSection::UMovieSceneTimeWarpSection(const FObjectInitializer&
 {
 	bSupportsInfiniteRange = true;
 	SetRange(TRange<FFrameNumber>::All());
-
-	if (!HasAnyFlags(RF_ClassDefaultObject))
-	{
-		UMovieScenePlayRateCurve* PlayRateCurve = ObjectInitializer.CreateDefaultSubobject<UMovieScenePlayRateCurve>(this, "PlayRateCurve");
-		PlayRateCurve->InitializeDefaults();
-		TimeWarp.Set(PlayRateCurve);
-	}
 }
 
 FMovieSceneNestedSequenceTransform UMovieSceneTimeWarpSection::GenerateTransform() const
