Sequencer: Add option to toggle viewport selectability button per sequencer.

Disabled by default (ie. Niagara), enabled by [REDACTED_PERSON], Actor Sequencer, UMG.

#rnx

diff --git a/Engine/Plugins/MovieScene/ActorSequence/Source/ActorSequenceEditor/Private/ActorSequenceEditorTabSummoner.cpp b/Engine/Plugins/MovieScene/ActorSequence/Source/ActorSequenceEditor/Private/ActorSequenceEditorTabSummoner.cpp
--- a/Engine/Plugins/MovieScene/ActorSequence/Source/ActorSequenceEditor/Private/ActorSequenceEditorTabSummoner.cpp
+++ b/Engine/Plugins/MovieScene/ActorSequence/Source/ActorSequenceEditor/Private/ActorSequenceEditorTabSummoner.cpp
@@ -400,6 +400,7 @@ public:
 				SequencerInitParams.ToolkitHost = WeakBlueprintEditor.Pin()->GetToolkitHost();
 				SequencerInitParams.HostCapabilities.bSupportsCurveEditor = true;
 				SequencerInitParams.HostCapabilities.bSupportsSidebar = true;
+				SequencerInitParams.HostCapabilities.bSupportsViewportSelectability = true;
 			}
 
 			TSharedRef<FExtender> AddMenuExtender = MakeShareable(new FExtender);
diff --git a/Engine/Plugins/MovieScene/LevelSequenceEditor/Source/LevelSequenceEditor/Private/LevelSequenceEditorToolkit.cpp b/Engine/Plugins/MovieScene/LevelSequenceEditor/Source/LevelSequenceEditor/Private/LevelSequenceEditorToolkit.cpp
--- a/Engine/Plugins/MovieScene/LevelSequenceEditor/Source/LevelSequenceEditor/Private/LevelSequenceEditorToolkit.cpp
+++ b/Engine/Plugins/MovieScene/LevelSequenceEditor/Source/LevelSequenceEditor/Private/LevelSequenceEditorToolkit.cpp
@@ -177,6 +177,7 @@ void FLevelSequenceEditorToolkit::Initialize(const EToolkitMode::Type Mode, cons
 		SequencerInitParams.HostCapabilities.bSupportsRenderMovie = true;
 		SequencerInitParams.HostCapabilities.bSupportsAddFromContentBrowser = true;
 		SequencerInitParams.HostCapabilities.bSupportsSidebar = true;
+		SequencerInitParams.HostCapabilities.bSupportsViewportSelectability = true;
 	}
 
 	ExtendSequencerToolbar("Sequencer.MainToolBar");
diff --git a/Engine/Source/Editor/Sequencer/Private/SSequencer.cpp b/Engine/Source/Editor/Sequencer/Private/SSequencer.cpp
--- a/Engine/Source/Editor/Sequencer/Private/SSequencer.cpp
+++ b/Engine/Source/Editor/Sequencer/Private/SSequencer.cpp
@@ -1644,16 +1644,19 @@ void SSequencer::PopulateToolBar(UToolMenu* InMenu)
 	}
 
 	{
-		FToolMenuSection& SelectionSection = InMenu->AddSection(TEXT("Selection"));
+		if (Sequencer.Pin()->GetHostCapabilities().bSupportsViewportSelectability)
+		{
+			FToolMenuSection& SelectionSection = InMenu->AddSection(TEXT("Selection"));
 
-		FToolMenuEntry ToggleLockViewportSelectionEntry = FToolMenuEntry::InitToolBarButton(FSequencerCommands::Get().ToggleLimitViewportSelection
-			, FText()
-			, FSequencerCommands::Get().ToggleLimitViewportSelection->GetDescription()
-			, FSlateIcon(FAppStyle::GetAppStyleSetName(), TEXT("ContentBrowser.Local.Small")));
-		ToggleLockViewportSelectionEntry.StyleNameOverride = SequencerToolbarStyleName;
-		SelectionSection.AddEntry(ToggleLockViewportSelectionEntry);
+			FToolMenuEntry ToggleLockViewportSelectionEntry = FToolMenuEntry::InitToolBarButton(FSequencerCommands::Get().ToggleLimitViewportSelection
+				, FText()
+				, FSequencerCommands::Get().ToggleLimitViewportSelection->GetDescription()
+				, FSlateIcon(FAppStyle::GetAppStyleSetName(), TEXT("ContentBrowser.Local.Small")));
+			ToggleLockViewportSelectionEntry.StyleNameOverride = SequencerToolbarStyleName;
+			SelectionSection.AddEntry(ToggleLockViewportSelectionEntry);
 
-		SelectionSection.AddSeparator(NAME_None);
+			SelectionSection.AddSeparator(NAME_None);
+		}
 	}
 
 	{
diff --git a/Engine/Source/Editor/Sequencer/Public/ISequencerModule.h b/Engine/Source/Editor/Sequencer/Public/ISequencerModule.h
--- a/Engine/Source/Editor/Sequencer/Public/ISequencerModule.h
+++ b/Engine/Source/Editor/Sequencer/Public/ISequencerModule.h
@@ -153,6 +153,9 @@ struct FSequencerHostCapabilities
 	/** Do we support the sidebar widget */
 	bool bSupportsSidebar;
 
+	/** Do we support selectability options in the viewport */
+	bool bSupportsViewportSelectability;
+
 	FSequencerHostCapabilities()
 		: bSupportsSaveMovieSceneAsset(false)
 		, bSupportsCurveEditor(false)
@@ -160,6 +163,7 @@ struct FSequencerHostCapabilities
 		, bSupportsRenderMovie(false)
 		, bSupportsAddFromContentBrowser(false)
 		, bSupportsSidebar(false)
+		, bSupportsViewportSelectability(false)
 	{}
 };
 
diff --git a/Engine/Source/Editor/UMGEditor/Private/WidgetBlueprintEditor.cpp b/Engine/Source/Editor/UMGEditor/Private/WidgetBlueprintEditor.cpp
--- a/Engine/Source/Editor/UMGEditor/Private/WidgetBlueprintEditor.cpp
+++ b/Engine/Source/Editor/UMGEditor/Private/WidgetBlueprintEditor.cpp
@@ -1456,6 +1456,7 @@ TSharedPtr<ISequencer> FWidgetBlueprintEditor::CreateSequencerWidgetInternal()
 		SequencerInitParams.HostCapabilities.bSupportsCurveEditor = true;
 		SequencerInitParams.HostCapabilities.bSupportsAddFromContentBrowser = true;
 		SequencerInitParams.HostCapabilities.bSupportsSidebar = true;
+		SequencerInitParams.HostCapabilities.bSupportsViewportSelectability = true;
 	};
 
 	TSharedPtr<ISequencer> Sequencer = FModuleManager::LoadModuleChecked<ISequencerModule>("Sequencer").CreateSequencer(SequencerInitParams);
