Sequencer: Fix object being added twice.

Continue if the object was successfully added, otherwise...

Add missing transaction for Add Actors to Binding

#rnx

diff --git a/Engine/Source/Editor/Sequencer/Private/SequencerUtilities.cpp b/Engine/Source/Editor/Sequencer/Private/SequencerUtilities.cpp
--- a/Engine/Source/Editor/Sequencer/Private/SequencerUtilities.cpp
+++ b/Engine/Source/Editor/Sequencer/Private/SequencerUtilities.cpp
@@ -3931,6 +3931,8 @@ FGuid FSequencerUtilities::AssignActor(TSharedRef<ISequencer> Sequencer, AActor*
 
 void FSequencerUtilities::AddActorsToBinding(TSharedRef<ISequencer> Sequencer, const TArray<AActor*>& Actors, const FMovieSceneBindingProxy& ObjectBinding)
 {
+	FScopedTransaction AddActorsToBinding(LOCTEXT("AddActorsToBinding", "Add Actors to Binding"));
+
 	TArray<UObject*> ObjectsToAdd;
 	Algo::Transform(Actors, ObjectsToAdd, [](AActor* Actor) { return Actor;});
 	AddObjectsToBinding(Sequencer, ObjectsToAdd, ObjectBinding, Sequencer->GetPlaybackContext());
@@ -4004,6 +4006,14 @@ void FSequencerUtilities::AddObjectsToBinding(TSharedRef<ISequencer> InSequencer
 			Sequence->BindPossessableObject(Guid, *ObjectToAdd, InResolutionContext);
 		}
 
+		// If the object was added successfully, continue
+		FGuid AddedGuid = InSequencer->GetHandleToObject(ObjectToAdd, false);
+		if (AddedGuid.IsValid())
+		{
+			continue;
+		}
+
+		// Otherwise...
 		if (ObjectClass == nullptr || UClass::FindCommonBase(ObjectToAdd->GetClass(), ObjectClass) != nullptr)
 		{
 			if (ObjectClass == nullptr)
