Sequencer: Fix camera cut track/section not being added when creating a root sequence with shots

FLevelSequenceEditorToolkit::AddShot used to call FSequencer::MakeNewSpawnable which called OnNewCameraAdded. From 32218082, it now calls FSequencerUtilities::CreateBinding, so NewCameraAdded needs to be called from there.

#rnx

diff --git a/Engine/Source/Editor/Sequencer/Private/SequencerUtilities.cpp b/Engine/Source/Editor/Sequencer/Private/SequencerUtilities.cpp
--- a/Engine/Source/Editor/Sequencer/Private/SequencerUtilities.cpp
+++ b/Engine/Source/Editor/Sequencer/Private/SequencerUtilities.cpp
@@ -3691,6 +3691,12 @@ FGuid FSequencerUtilities::CreateOrReplaceBinding(TSharedPtr<ISequencer> Sequenc
 			DestinationFolder->AddChildObjectBinding(RootObjectGuid);
 		}
 	}
+
+	if (ACameraActor* NewCamera = Cast<ACameraActor>(InObject))
+	{
+		NewCameraAdded(Sequencer.ToSharedRef(), NewCamera, BindingGuid);
+	}
+
 	if (Sequencer)
 	{
 		Sequencer->OnAddBinding(BindingGuid, OwnerMovieScene);
