Sequencer: Add a setting on whether the curve editor should be visible for Sequencer so that when you open and close Sequencer, the curve editor visibility will be restored.

This fixes a bug where if the curve editor is vislble and you do a Save As on the level sequence, it doesn't stay visible afterwards.

diff --git a/Engine/Source/Editor/Sequencer/Private/Sequencer.cpp b/Engine/Source/Editor/Sequencer/Private/Sequencer.cpp
--- a/Engine/Source/Editor/Sequencer/Private/Sequencer.cpp
+++ b/Engine/Source/Editor/Sequencer/Private/Sequencer.cpp
@@ -605,6 +605,8 @@ void FSequencer::InitSequencer(const FSequencerInitParams& InitParams, const TSh
 		TSharedPtr<FCurveEditor> CurveEditor = CurveEditorExtension->GetCurveEditor();
 
 		CurveEditor->OnCurveArrayChanged.AddRaw(this, &FSequencer::OnCurveModelDisplayChanged);
+
+		SetShowCurveEditor(GetSequencerSettings()->GetCurveEditorVisible());
 	}
 
 	// When undo occurs, get a notification so we can make sure our view is up to date
@@ -800,6 +802,12 @@ void FSequencer::Close()
 		Runner->Flush();
 	}
 
+	// Save the user's preference for curve editor visibility since it can be closed outside of Sequencer's control (ie. close tab)
+	if (GetHostCapabilities().bSupportsCurveEditor)
+	{
+		GetSequencerSettings()->SetCurveEditorVisible(GetCurveEditorIsVisible());
+	}
+
 	RestorePreAnimatedState();
 
 	for (auto TrackEditor : TrackEditors)
@@ -5942,6 +5950,8 @@ void FSequencer::OnNewActorsDropped(const TArray<UObject*>& DroppedObjects, cons
 void FSequencer::SetShowCurveEditor(bool bInShowCurveEditor)
 {
 	SequencerWidget->OnCurveEditorVisibilityChanged(bInShowCurveEditor);
+
+	GetSequencerSettings()->SetCurveEditorVisible(bInShowCurveEditor);
 }
 
 bool FSequencer::GetCurveEditorIsVisible() const
diff --git a/Engine/Source/Editor/Sequencer/Private/SequencerSettings.cpp b/Engine/Source/Editor/Sequencer/Private/SequencerSettings.cpp
--- a/Engine/Source/Editor/Sequencer/Private/SequencerSettings.cpp
+++ b/Engine/Source/Editor/Sequencer/Private/SequencerSettings.cpp
@@ -37,6 +37,7 @@ USequencerSettings::USequencerSettings( const FObjectInitializer& ObjectInitiali
 	bLinkCurveEditorTimeRange = false;
 	bSynchronizeCurveEditorSelection = true;
 	bIsolateCurveEditorToSelection = true;
+	bCurveEditorVisible = true;
 	LoopMode = ESequencerLoopMode::SLM_NoLoop;
 	bSnapKeysAndSectionsToPlayRange = false;
 	bResetPlayheadWhenNavigating = false;
@@ -542,6 +543,15 @@ void USequencerSettings::IsolateCurveEditorToSelection(bool bInIsolateCurveEdito
 	}
 }
 
+void USequencerSettings::SetCurveEditorVisible(bool bInCurveEditorVisible)
+{
+	if (bCurveEditorVisible != bInCurveEditorVisible)
+	{
+		bCurveEditorVisible = bInCurveEditorVisible;
+		SaveConfig();
+	}
+}
+
 uint8 USequencerSettings::GetZeroPadFrames() const
 {
 	return ZeroPadFrames;
diff --git a/Engine/Source/Editor/Sequencer/Public/SequencerSettings.h b/Engine/Source/Editor/Sequencer/Public/SequencerSettings.h
--- a/Engine/Source/Editor/Sequencer/Public/SequencerSettings.h
+++ b/Engine/Source/Editor/Sequencer/Public/SequencerSettings.h
@@ -278,6 +278,11 @@ public:
 	/** Assign whether we should filter the curve editor tree to only nodes that are relevant to the current sequencer selection */
 	void IsolateCurveEditorToSelection(bool bInIsolateCurveEditorToSelection);
 
+	/** Return true if we should filter the curve editor tree to only nodes that are relevant to the current sequencer selection */
+	bool GetCurveEditorVisible() const { return bCurveEditorVisible; }
+	/** Assign whether we should filter the curve editor tree to only nodes that are relevant to the current sequencer selection */
+	void SetCurveEditorVisible(bool bCurveEditorVisible);
+
 	/** Gets the loop mode. */
 	ESequencerLoopMode GetLoopMode() const;
 	/** Sets the loop mode. */
@@ -611,9 +616,13 @@ protected:
 	bool bSynchronizeCurveEditorSelection;
 
 	/** When enabled, changing the sequencer tree selection will isolate (auto-filter) the selected nodes in the curve editor. */
-	UPROPERTY( config, EditAnywhere, Category=CurveEditor )
+	UPROPERTY(config, EditAnywhere, Category=CurveEditor)
 	bool bIsolateCurveEditorToSelection;
 
+	/** Whether the curve editor is visible */
+	UPROPERTY(config, EditAnywhere, Category=CurveEditor)
+	bool bCurveEditorVisible;
+
 	/** The loop mode of the playback in timeline. */
 	UPROPERTY( config )
 	TEnumAsByte<ESequencerLoopMode> LoopMode;
