Sequencer: Return false if not pasting any object bindings so that other types are not pasted as well

#rnx

diff --git a/Engine/Source/Editor/Sequencer/Private/Sequencer.cpp b/Engine/Source/Editor/Sequencer/Private/Sequencer.cpp
--- a/Engine/Source/Editor/Sequencer/Private/Sequencer.cpp
+++ b/Engine/Source/Editor/Sequencer/Private/Sequencer.cpp
@@ -8240,7 +8240,11 @@ bool FSequencer::DoPaste(bool bClearSelection)
 	}
 	if (EnumHasAnyFlags(PasteSupport, ESequencerPasteSupport::ObjectBindings)) 
 	{
-		bAnythingPasted |= PasteObjectBindings(TextToImport, ParentFolder, PastedFolders, PasteErrors, bClearSelection);
+		// Stop if failed to paste any object bindings (other paste types, tracks and sections can be pasted in bulk)
+		if (!PasteObjectBindings(TextToImport, ParentFolder, PastedFolders, PasteErrors, bClearSelection))
+		{
+			return false;
+		}
 	}
 	if (EnumHasAnyFlags(PasteSupport, ESequencerPasteSupport::Tracks)) 
 	{
