Curve Editor: Check for valid toolkit host and tab manager on close

#rnx

diff --git a/Engine/Source/Editor/Sequencer/Private/MVVM/CurveEditorExtension.cpp b/Engine/Source/Editor/Sequencer/Private/MVVM/CurveEditorExtension.cpp
--- a/Engine/Source/Editor/Sequencer/Private/MVVM/CurveEditorExtension.cpp
+++ b/Engine/Source/Editor/Sequencer/Private/MVVM/CurveEditorExtension.cpp
@@ -518,8 +518,20 @@ bool FCurveEditorExtension::IsCurveEditorOpen() const
 		return false;
 	}
 
+	TSharedPtr<IToolkitHost> ToolkitHost = Sequencer->GetToolkitHost();
+	if (!ToolkitHost)
+	{
+		return false;
+	}
+
+	TSharedPtr<FTabManager> TabManager = ToolkitHost->GetTabManager();
+	if (!TabManager)
+	{
+		return false;
+	}
+
 	FTabId TabId = FTabId(FCurveEditorExtension::CurveEditorTabName);
-	return Sequencer->GetToolkitHost()->GetTabManager()->FindExistingLiveTab(TabId).IsValid();
+	return TabManager->FindExistingLiveTab(TabId).IsValid();
 }
 
 void FCurveEditorExtension::CloseCurveEditor()
