Sequencer: Suppress selection events on create drag.

This fixes a perf issue when clicking and dragging on top level keys which fires selection events for all the keys that are being selected.

diff --git a/Engine/Source/Editor/Sequencer/Private/Tools/SequencerEditTool_Movement.cpp b/Engine/Source/Editor/Sequencer/Private/Tools/SequencerEditTool_Movement.cpp
--- a/Engine/Source/Editor/Sequencer/Private/Tools/SequencerEditTool_Movement.cpp
+++ b/Engine/Source/Editor/Sequencer/Private/Tools/SequencerEditTool_Movement.cpp
@@ -153,6 +153,8 @@ TSharedPtr<UE::Sequencer::ISequencerEditToolDragOperation> FSequencerEditTool_Mo
 	using namespace UE::Sequencer;
 
 	FSequencerSelection& Selection = *Sequencer.GetViewModel()->GetSelection();
+	FSelectionEventSuppressor SuppressEvents = Selection.SuppressEvents();
+
 	TSharedRef<SSequencer> SequencerWidget = StaticCastSharedRef<SSequencer>(Sequencer.GetSequencerWidget());
 
 	GetHotspotTime(OriginalHotspotTime);
@@ -253,6 +255,7 @@ TSharedPtr<UE::Sequencer::ISequencerEditToolDragOperation> FSequencerEditTool_Mo
 				if (!Selection.KeySelection.IsSelected(Key.KeyHandle))
 				{
 					bUniqueDrag = true;
+					break;
 				}
 			};
 
