Sequencer: add more extensibility APIs for the Sequencer UI

Add actions menu and view menu extenders to sequencer customizations.
Split the view menu into sections for easier extensibility.

diff --git a/Engine/Source/Editor/Sequencer/Private/SSequencer.cpp b/Engine/Source/Editor/Sequencer/Private/SSequencer.cpp
--- a/Engine/Source/Editor/Sequencer/Private/SSequencer.cpp
+++ b/Engine/Source/Editor/Sequencer/Private/SSequencer.cpp
@@ -2080,10 +2080,13 @@ TSharedRef<SWidget> SSequencer::MakeActionsMenu()
 	TSharedPtr<FSequencer> Sequencer = SequencerPtr.Pin();
 	TSharedPtr<FSequencerEditorViewModel> EditorViewModel = Sequencer->GetViewModel();
 
+	TArray<TSharedPtr<FExtender>> AllExtenders;
 	ISequencerModule& SequencerModule = FModuleManager::GetModuleChecked<ISequencerModule>("Sequencer");
-	TSharedPtr<FExtender> ActionsMenuExtender = SequencerModule.GetActionsMenuExtensibilityManager()->GetAllExtenders();
+	AllExtenders.Add(SequencerModule.GetActionsMenuExtensibilityManager()->GetAllExtenders());
+	AllExtenders.Append(ActionsMenuExtenders);
 
-	FMenuBuilder MenuBuilder(true, Sequencer->GetCommandBindings(), ActionsMenuExtender);
+	TSharedPtr<FExtender> Extender = FExtender::Combine(AllExtenders);
+	FMenuBuilder MenuBuilder(true, Sequencer->GetCommandBindings(), Extender);
 
 	MenuBuilder.BeginSection("SequenceOptions", LOCTEXT("SequenceOptionsHeader", "Sequence"));
 	{
@@ -2266,9 +2269,16 @@ void SSequencer::FillAdvancedMenu(FMenuBuilder& MenuBuilder)
 
 TSharedRef<SWidget> SSequencer::MakeViewMenu()
 {
-	FMenuBuilder MenuBuilder( true, SequencerPtr.Pin()->GetCommandBindings() );
 	TSharedPtr<FSequencer> Sequencer = SequencerPtr.Pin();
 
+	TArray<TSharedPtr<FExtender>> AllExtenders;
+	ISequencerModule& SequencerModule = FModuleManager::GetModuleChecked<ISequencerModule>("Sequencer");
+	AllExtenders.Add(SequencerModule.GetViewMenuExtensibilityManager()->GetAllExtenders());
+	AllExtenders.Append(ActionsMenuExtenders);
+
+	TSharedPtr<FExtender> Extender = FExtender::Combine(AllExtenders);
+	FMenuBuilder MenuBuilder(true, Sequencer->GetCommandBindings(), Extender);
+	
 	if (Sequencer->GetHostCapabilities().bSupportsSidebar)
 	{
 		MenuBuilder.AddMenuEntry( FSequencerCommands::Get().ToggleSidebarVisible );
@@ -2276,52 +2286,60 @@ TSharedRef<SWidget> SSequencer::MakeViewMenu()
 
 	if (Sequencer->IsLevelEditorSequencer())
 	{
-		MenuBuilder.AddMenuEntry( FSequencerCommands::Get().TogglePilotCamera );
-		MenuBuilder.AddMenuEntry( FSequencerCommands::Get().ToggleRestoreOriginalViewportOnCameraCutUnlock );
+		MenuBuilder.BeginSection("PilotCamera", LOCTEXT("PilotCamera", "Pilot Camera"));
+		{
+			MenuBuilder.AddMenuEntry(FSequencerCommands::Get().TogglePilotCamera);
+			MenuBuilder.AddMenuEntry(FSequencerCommands::Get().ToggleRestoreOriginalViewportOnCameraCutUnlock);
+		}
+		MenuBuilder.EndSection();
 	}
 
-	MenuBuilder.AddMenuEntry( FSequencerCommands::Get().ToggleAutoScroll );
-	MenuBuilder.AddMenuEntry( FSequencerCommands::Get().ToggleShowRangeSlider );
-	MenuBuilder.AddMenuEntry( FSequencerCommands::Get().ToggleShowSelectedNodesOnly );
-	MenuBuilder.AddMenuEntry( FSequencerCommands::Get().ToggleLayerBars );
-	MenuBuilder.AddMenuEntry( FSequencerCommands::Get().ToggleKeyBars );
-	MenuBuilder.AddMenuEntry( FSequencerCommands::Get().ToggleChannelColors );
-	MenuBuilder.AddMenuEntry( FSequencerCommands::Get().ToggleShowInfoButton );
-	MenuBuilder.AddMenuEntry( FSequencerCommands::Get().ToggleShowPreAndPostRoll );
+	MenuBuilder.BeginSection("SequencerSettings", LOCTEXT("SequencerSettings", "Sequencer Settings"));
+	{
+		MenuBuilder.AddMenuEntry( FSequencerCommands::Get().ToggleAutoScroll );
+		MenuBuilder.AddMenuEntry( FSequencerCommands::Get().ToggleShowRangeSlider );
+		MenuBuilder.AddMenuEntry( FSequencerCommands::Get().ToggleShowSelectedNodesOnly );
+		MenuBuilder.AddMenuEntry( FSequencerCommands::Get().ToggleLayerBars );
+		MenuBuilder.AddMenuEntry( FSequencerCommands::Get().ToggleKeyBars );
+		MenuBuilder.AddMenuEntry( FSequencerCommands::Get().ToggleChannelColors );
+		MenuBuilder.AddMenuEntry( FSequencerCommands::Get().ToggleShowInfoButton );
+		MenuBuilder.AddMenuEntry( FSequencerCommands::Get().ToggleShowPreAndPostRoll );
 
-	// Menu entry for zero padding
-	auto OnZeroPadChanged = [this](uint8 NewValue) {
-		GetSequencerSettings()->SetZeroPadFrames(NewValue);
-	};
+		MenuBuilder.AddSubMenu(LOCTEXT("ViewDensityMenuLabel", "View Density"), FText::GetEmpty(), FNewMenuDelegate::CreateRaw(this, &SSequencer::FillViewDensityMenu));
 
-	MenuBuilder.AddSubMenu(LOCTEXT("ViewDensityMenuLabel", "View Density"), FText::GetEmpty(), FNewMenuDelegate::CreateRaw(this, &SSequencer::FillViewDensityMenu));
+		// Menu Entry for Outliner Column Visibilities
+		if (OutlinerColumnVisibilities.Num() > 0)
+		{
+			MenuBuilder.AddSubMenu(LOCTEXT("ColumnVisibilityHeader", "Columns"), FText::GetEmpty(), FNewMenuDelegate::CreateRaw(this, &SSequencer::FillColumnVisibilityMenu));
+		}
 
-	// Menu Entry for Outliner Column Visibilities
-	if (OutlinerColumnVisibilities.Num() > 0)
-	{
-		MenuBuilder.AddSubMenu(LOCTEXT("ColumnVisibilityHeader", "Columns"), FText::GetEmpty(), FNewMenuDelegate::CreateRaw(this, &SSequencer::FillColumnVisibilityMenu));
-	}
+		// Menu entry for zero padding
+		auto OnZeroPadChanged = [this](uint8 NewValue) {
+			GetSequencerSettings()->SetZeroPadFrames(NewValue);
+		};
 
-	MenuBuilder.AddWidget(
-		SNew(SHorizontalBox)
-		+ SHorizontalBox::Slot()
-			[
-				SNew(SSpacer)
-			]
-		+ SHorizontalBox::Slot()
-			.AutoWidth()
-			[
-				SNew(SSpinBox<uint8>)
-				.Style(&FAppStyle::GetWidgetStyle<FSpinBoxStyle>("Sequencer.HyperlinkSpinBox"))
-			.OnValueCommitted_Lambda([=](uint8 Value, ETextCommit::Type) { OnZeroPadChanged(Value); })
-			.OnValueChanged_Lambda(OnZeroPadChanged)
-			.MinValue(0)
-			.MaxValue(8)
-			.Value_Lambda([this]() -> uint8 {
-			return GetSequencerSettings()->GetZeroPadFrames();
-		})
-		],
-		LOCTEXT("ZeroPaddingText", "Zero Pad Frame Numbers"));
+		MenuBuilder.AddWidget(
+			SNew(SHorizontalBox)
+			+ SHorizontalBox::Slot()
+				[
+					SNew(SSpacer)
+				]
+			+ SHorizontalBox::Slot()
+				.AutoWidth()
+				[
+					SNew(SSpinBox<uint8>)
+					.Style(&FAppStyle::GetWidgetStyle<FSpinBoxStyle>("Sequencer.HyperlinkSpinBox"))
+					.OnValueCommitted_Lambda([=](uint8 Value, ETextCommit::Type) { OnZeroPadChanged(Value); })
+					.OnValueChanged_Lambda(OnZeroPadChanged)
+					.MinValue(0)
+					.MaxValue(8)
+					.Value_Lambda([this]() -> uint8 {
+						return GetSequencerSettings()->GetZeroPadFrames();
+					})
+				],
+			LOCTEXT("ZeroPaddingText", "Zero Pad Frame Numbers"));
+	}
+	MenuBuilder.EndSection();
 
 	MenuBuilder.BeginSection("OrganizeAndSort", LOCTEXT("OrganizeAndSortHeader", "Organize and Sort"));
 	{
@@ -4251,6 +4269,8 @@ void SSequencer::ApplySequencerCustomizations(const TArrayView<const FSequencerC
 {
 	AddMenuExtenders.Reset();
 	ToolbarExtenders.Reset();
+	ActionsMenuExtenders.Reset();
+	ViewMenuExtenders.Reset();
 
 	OnReceivedDragOver.Reset();
 	OnReceivedDrop.Reset();
@@ -4278,6 +4298,14 @@ void SSequencer::ApplySequencerCustomization(const FSequencerCustomizationInfo&
 	{
 		ToolbarExtenders.Add(Customization.ToolbarExtender);
 	}
+	if (Customization.ActionsMenuExtender != nullptr)
+	{
+		ActionsMenuExtenders.Add(Customization.ActionsMenuExtender);
+	}
+	if (Customization.ViewMenuExtender != nullptr)
+	{
+		ViewMenuExtenders.Add(Customization.ViewMenuExtender);
+	}
 
 	if (Customization.OnReceivedDragOver.IsBound())
 	{
diff --git a/Engine/Source/Editor/Sequencer/Private/SSequencer.h b/Engine/Source/Editor/Sequencer/Private/SSequencer.h
--- a/Engine/Source/Editor/Sequencer/Private/SSequencer.h
+++ b/Engine/Source/Editor/Sequencer/Private/SSequencer.h
@@ -820,6 +820,12 @@ private:
 	/** Extender to use for the toolbar */
 	TArray<TSharedPtr<FExtender>> ToolbarExtenders;
 
+	/** Extender to use for the actions menu */
+	TArray<TSharedPtr<FExtender>> ActionsMenuExtenders;
+
+	/** Extender to use for the view menu */
+	TArray<TSharedPtr<FExtender>> ViewMenuExtenders;
+
 	/** Numeric type interface used for converting parsing and generating strings from numbers */
 	TSharedPtr<INumericTypeInterface<double>> NumericTypeInterface;
 
diff --git a/Engine/Source/Editor/Sequencer/Private/SequencerModule.cpp b/Engine/Source/Editor/Sequencer/Private/SequencerModule.cpp
--- a/Engine/Source/Editor/Sequencer/Private/SequencerModule.cpp
+++ b/Engine/Source/Editor/Sequencer/Private/SequencerModule.cpp
@@ -427,6 +427,7 @@ public:
 		AddTrackMenuExtensibilityManager = MakeShareable( new FExtensibilityManager );
 		ToolBarExtensibilityManager = MakeShareable(new FExtensibilityManager);
 		ActionsMenuExtensibilityManager = MakeShareable(new FExtensibilityManager);
+		ViewMenuExtensibilityManager = MakeShareable(new FExtensibilityManager);
 		SidebarExtensibilityManager = MakeShareable(new FExtensibilityManager);
 
 		SequencerCustomizationManager = MakeShareable(new FSequencerCustomizationManager);
@@ -515,6 +516,7 @@ public:
 	virtual TSharedPtr<FExtensibilityManager> GetAddTrackMenuExtensibilityManager() const override { return AddTrackMenuExtensibilityManager; }
 	virtual TSharedPtr<FExtensibilityManager> GetToolBarExtensibilityManager() const override { return ToolBarExtensibilityManager; }
 	virtual TSharedPtr<FExtensibilityManager> GetActionsMenuExtensibilityManager() const override { return ActionsMenuExtensibilityManager; }
+	virtual TSharedPtr<FExtensibilityManager> GetViewMenuExtensibilityManager() const override { return ViewMenuExtensibilityManager; }
 	virtual TSharedPtr<FExtensibilityManager> GetSidebarExtensibilityManager() const override { return SidebarExtensibilityManager; }
 
 	virtual TSharedPtr<FSequencerCustomizationManager> GetSequencerCustomizationManager() const override { return SequencerCustomizationManager; }
@@ -628,6 +630,7 @@ private:
 	TSharedPtr<FExtensibilityManager> AddTrackMenuExtensibilityManager;
 	TSharedPtr<FExtensibilityManager> ToolBarExtensibilityManager;
 	TSharedPtr<FExtensibilityManager> ActionsMenuExtensibilityManager;
+	TSharedPtr<FExtensibilityManager> ViewMenuExtensibilityManager;
 	TSharedPtr<FExtensibilityManager> SidebarExtensibilityManager;
 
 	TSharedPtr<FSequencerCustomizationManager> SequencerCustomizationManager;
diff --git a/Engine/Source/Editor/Sequencer/Public/ISequencerModule.h b/Engine/Source/Editor/Sequencer/Public/ISequencerModule.h
--- a/Engine/Source/Editor/Sequencer/Public/ISequencerModule.h
+++ b/Engine/Source/Editor/Sequencer/Public/ISequencerModule.h
@@ -353,7 +353,14 @@ public:
 	virtual TSharedPtr<FExtensibilityManager> GetToolBarExtensibilityManager() const = 0;
 
 	/**
-	 * Get the extensibility manager for toolbars.
+	 * Get the extensibility manager for the view options menu.
+	 *
+	 * @return Toolbar extensibility manager.
+	 */
+	virtual TSharedPtr<FExtensibilityManager> GetViewMenuExtensibilityManager() const = 0;
+
+	/**
+	 * Get the extensibility manager for the actions menu.
 	 *
 	 * @return Toolbar extensibility manager.
 	 */
diff --git a/Engine/Source/Editor/Sequencer/Public/SequencerCustomizationManager.h b/Engine/Source/Editor/Sequencer/Public/SequencerCustomizationManager.h
--- a/Engine/Source/Editor/Sequencer/Public/SequencerCustomizationManager.h
+++ b/Engine/Source/Editor/Sequencer/Public/SequencerCustomizationManager.h
@@ -62,6 +62,12 @@ struct FSequencerCustomizationInfo
 	/** Extender for the sequencer toolbar. */
 	TSharedPtr<FExtender> ToolbarExtender;
 
+	/** Extender for the actions menu. */
+	TSharedPtr<FExtender> ActionsMenuExtender;
+
+	/** Extender for the view menu. */
+	TSharedPtr<FExtender> ViewMenuExtender;
+	
 	/** Extender for the sequencer sidebar. */
 	TSharedPtr<FExtender> SidebarExtender;
 
