Sequencer: Prevent ensure on closing sequencer while sync selection timer is active.

diff --git a/Engine/Source/Editor/Sequencer/Private/MVVM/CurveEditorExtension.cpp b/Engine/Source/Editor/Sequencer/Private/MVVM/CurveEditorExtension.cpp
--- a/Engine/Source/Editor/Sequencer/Private/MVVM/CurveEditorExtension.cpp
+++ b/Engine/Source/Editor/Sequencer/Private/MVVM/CurveEditorExtension.cpp
@@ -604,6 +604,11 @@ void FCurveEditorExtension::RequestSyncSelection()
 		{
 			return;
 		}
+		TSharedPtr<ISequencer> Sequencer = RootViewModel->GetSequencer();
+		if (!Sequencer)
+		{
+			return;
+		}
 
 		FCurveEditorExtension* This = RootViewModel->CastDynamic<FCurveEditorExtension>();
 		if (This)
