[Sequencer] - Reset extension references at the start of rebuilding the sidebar and on Sequencer close

diff --git a/Engine/Source/Editor/Sequencer/Private/Widgets/Sidebar/SequencerSelectionDrawer.cpp b/Engine/Source/Editor/Sequencer/Private/Widgets/Sidebar/SequencerSelectionDrawer.cpp
--- a/Engine/Source/Editor/Sequencer/Private/Widgets/Sidebar/SequencerSelectionDrawer.cpp
+++ b/Engine/Source/Editor/Sequencer/Private/Widgets/Sidebar/SequencerSelectionDrawer.cpp
@@ -132,6 +132,11 @@ TSharedRef<SWidget> FSequencerSelectionDrawer::CreateContentWidget()
 
 			OnSequencerSelectionChanged();
 		}
+
+		Sequencer->OnCloseEvent().AddLambda([this](const TSharedRef<ISequencer> InSequencer)
+			{
+				ResetContent();
+			});
 	}
 
 	return SNew(SBorder)
@@ -147,6 +152,15 @@ TSharedRef<SWidget> FSequencerSelectionDrawer::CreateContentWidget()
 		];
 }
 
+void FSequencerSelectionDrawer::ResetContent()
+{
+	ContentBox->ClearChildren();
+
+	CurveChannelExtension.Reset();
+
+	ChannelExtensions.Reset();
+}
+
 void FSequencerSelectionDrawer::OnSequencerSelectionChanged()
 {
 	if (!ContentBox.IsValid())
@@ -171,9 +185,7 @@ void FSequencerSelectionDrawer::UpdateFromSelectionNextFrame()
 {
 	bWaitingToHandleSelectionChanged = false;
 
-	ContentBox->ClearChildren();
-
-	CurveChannelExtension.Reset();
+	ResetContent();
 
 	const TSharedPtr<FSequencer> Sequencer = WeakSequencer.Pin();
 	if (!Sequencer.IsValid())
@@ -382,11 +394,9 @@ void FSequencerSelectionDrawer::BuildOutlinerDetails(const TSharedRef<FSequencer
 		// Channel Interface Extensions (Perlin Noise, Easing, Wave)
 		if (ChannelInterfaces.Num() > 0)
 		{
-			ChannelExtensions.Empty();
-
 			if (bAllChannelNamesEqual)
 			{
-				if (const TSharedPtr<ISidebarChannelExtension> ChannelExtension = ChannelInterfaces[0]->ExtendSidebarMenu_Raw(MenuBuilder, Extender, ChannelHandles, WeakSceneSections, Sequencer))
+				if (const TSharedPtr<ISidebarChannelExtension> ChannelExtension = ChannelInterfaces[0]->ExtendSidebarMenu_Raw(MenuBuilder, Extender, ChannelHandles, WeakSceneSections, WeakSequencer))
 				{
 					ChannelExtensions.Add(ChannelExtension);
 				}
@@ -396,7 +406,7 @@ void FSequencerSelectionDrawer::BuildOutlinerDetails(const TSharedRef<FSequencer
 				// Display different channels separately and don't allow to edit "all-in-one"
 				for (int32 Index = 0; Index < ChannelInterfaces.Num(); ++Index)
 				{
-					if (const TSharedPtr<ISidebarChannelExtension> ChannelExtension = ChannelInterfaces[Index]->ExtendSidebarMenu_Raw(MenuBuilder, Extender, { ChannelHandles[Index] }, { WeakSceneSections[Index] }, Sequencer))
+					if (const TSharedPtr<ISidebarChannelExtension> ChannelExtension = ChannelInterfaces[Index]->ExtendSidebarMenu_Raw(MenuBuilder, Extender, { ChannelHandles[Index] }, { WeakSceneSections[Index] }, WeakSequencer))
 					{
 						ChannelExtensions.Add(ChannelExtension);
 					}
diff --git a/Engine/Source/Editor/Sequencer/Private/Widgets/Sidebar/SequencerSelectionDrawer.h b/Engine/Source/Editor/Sequencer/Private/Widgets/Sidebar/SequencerSelectionDrawer.h
--- a/Engine/Source/Editor/Sequencer/Private/Widgets/Sidebar/SequencerSelectionDrawer.h
+++ b/Engine/Source/Editor/Sequencer/Private/Widgets/Sidebar/SequencerSelectionDrawer.h
@@ -51,6 +51,8 @@ protected:
 	TSharedPtr<SWidget> CreateKeyFrameDetails(const TSharedRef<UE::Sequencer::FSequencerSelection>& InSequencerSelection);
 	TSharedPtr<SWidget> CreateMarkedFrameDetails(const int32 InMarkedFrameIndex);
 
+	void ResetContent();
+
 	TWeakPtr<FSequencer> WeakSequencer;
 
 	TSharedPtr<SVerticalBox> ContentBox;
