[Sequencer] Advanced Filtering:
 - Adds various text expressions
 - Adds additional filters
 - Add custom filters created from text filter expressions
 - Add Hide/Isolate/Show filter
 - Move all filtering functions from ISequencer to ISequencerTrackFilters
 - Create filter bar widget and backend class
 - Add actions that can be hotkeyed for all filters
 - Text expression help dialog POC
 - Add USequencerTrackFilterTextExpressionExtension to enable outside modules to add custom text filter expressions

diff --git a/Engine/Plugins/Animation/ControlRig/Source/ControlRigEditor/Private/Sequencer/ControlRigSequencerFilter.cpp b/Engine/Plugins/Animation/ControlRig/Source/ControlRigEditor/Private/Sequencer/ControlRigSequencerFilter.cpp
--- a/Engine/Plugins/Animation/ControlRig/Source/ControlRigEditor/Private/Sequencer/ControlRigSequencerFilter.cpp
+++ b/Engine/Plugins/Animation/ControlRig/Source/ControlRigEditor/Private/Sequencer/ControlRigSequencerFilter.cpp
@@ -2,12 +2,18 @@
 
 #include "Sequencer/ControlRigSequencerFilter.h"
 #include "ControlRig.h"
-#include "Editor/ControlRigSkeletalMeshComponent.h"
+#include "ControlRigBlueprint.h"
+#include "Filters/ISequencerTrackFilters.h"
+#include "Filters/SequencerTrackFilterBase.h"
+#include "Framework/Commands/Commands.h"
+#include "Framework/Commands/UICommandInfo.h"
+#include "MVVM/Extensions/IOutlinerExtension.h"
+#include "MVVM/ViewModels/ChannelModel.h"
 #include "Sequencer/MovieSceneControlRigParameterTrack.h"
 #include "Styling/AppStyle.h"
 #include "Styling/SlateIconFinder.h"
-#include "Framework/Commands/Commands.h"
-#include "ISequencer.h"
+
+using namespace UE::Sequencer;
 
 #include UE_INLINE_GENERATED_CPP_BY_NAME(ControlRigSequencerFilter)
 
@@ -20,90 +26,79 @@ class FSequencerTrackFilter_ControlRigControlsCommands
 	: public TCommands<FSequencerTrackFilter_ControlRigControlsCommands>
 {
 public:
-
 	FSequencerTrackFilter_ControlRigControlsCommands()
-		: TCommands<FSequencerTrackFilter_ControlRigControlsCommands>
-		(
-			"FSequencerTrackFilter_ControlRigControls",
-			NSLOCTEXT("Contexts", "FSequencerTrackFilter_ControlRigControls", "FSequencerTrackFilter_ControlRigControls"),
+		: TCommands<FSequencerTrackFilter_ControlRigControlsCommands>(
+			TEXT("FSequencerTrackFilter_ControlRigControls"),
+			LOCTEXT("FSequencerTrackFilter_ControlRigControls", "Control Rig Filters"),
 			NAME_None,
-			FAppStyle::GetAppStyleSetName() // Icon Style Set
-			)
-	{ }
+			FAppStyle::GetAppStyleSetName())
+	{}
 
 	/** Toggle the control rig controls filter */
-	TSharedPtr< FUICommandInfo > ToggleControlRigControls;
+	TSharedPtr< FUICommandInfo > ToggleFilter_ControlRigControls;
 
 	/** Initialize commands */
 	virtual void RegisterCommands() override
 	{
-		UI_COMMAND(ToggleControlRigControls, "Control Rig Controls", "Toggle the filter for Control Rig Controls.", EUserInterfaceActionType::ToggleButton, FInputChord(EKeys::F9));
+		UI_COMMAND(ToggleFilter_ControlRigControls, "Control Rig Controls", "Toggle the filter for Control Rig Controls.", EUserInterfaceActionType::ToggleButton, FInputChord(EKeys::F9));
 	}
 };
 
+//////////////////////////////////////////////////////////////////////////
+//
+
 class FSequencerTrackFilter_ControlRigControls : public FSequencerTrackFilter
 {
 public:
-
-	FSequencerTrackFilter_ControlRigControls()
-		: BindingCount(0)
+	FSequencerTrackFilter_ControlRigControls(ISequencerTrackFilters& InOutFilterInterface, TSharedPtr<FFilterCategory> InCategory = nullptr)
+		: FSequencerTrackFilter(InOutFilterInterface, MoveTemp(InCategory))
+		, BindingCount(0)
 	{
 		FSequencerTrackFilter_ControlRigControlsCommands::Register();
 	}
 
-	~FSequencerTrackFilter_ControlRigControls()
+	virtual ~FSequencerTrackFilter_ControlRigControls() override
 	{
 		BindingCount--;
-
 		if (BindingCount < 1)
 		{
 			FSequencerTrackFilter_ControlRigControlsCommands::Unregister();
 		}
 	}
 
-	virtual FString GetName() const override { return TEXT("ControlRigControlsFilter"); }
-	virtual FText GetDisplayName() const override { return LOCTEXT("SequenceTrackFilter_ControlRigControls", "Control Rig Controls"); }
-	virtual FSlateIcon GetIcon() const { return FSlateIconFinder::FindIconForClass(UControlRig::StaticClass()); }
+	//~ Begin IFilter
 
-	virtual bool PassesFilterWithDisplayName(FTrackFilterType InItem, const FText& InText) const
-	{
-		const UMovieSceneControlRigParameterTrack* Track = Cast<UMovieSceneControlRigParameterTrack>(InItem);
-		return (Track != nullptr);
-	}
+	virtual FString GetName() const override { return TEXT("ControlRigControl"); }
 
-	virtual bool PassesFilter(FTrackFilterType InItem) const override
+	virtual bool PassesFilter(FSequencerTrackFilterType InItem) const override
 	{
-		return true;
+		FSequencerFilterData& FilterData = FilterInterface.GetFilterData();
+		const UMovieSceneTrack* const TrackObject = ResolveMovieSceneTrackObject(InItem, FilterData);
+		const UMovieSceneControlRigParameterTrack* const Track = Cast<UMovieSceneControlRigParameterTrack>(TrackObject);
+		return IsValid(Track);
 	}
 
-	virtual FText GetToolTipText() const override
-	{
-		const FSequencerTrackFilter_ControlRigControlsCommands& Commands = FSequencerTrackFilter_ControlRigControlsCommands::Get();
+	//~ End IFilter
 
-		const TSharedRef<const FInputChord> FirstActiveChord = Commands.ToggleControlRigControls->GetFirstValidChord();
+	//~ Begin FFilterBase
+	virtual FText GetDisplayName() const override { return LOCTEXT("SequenceTrackFilter_ControlRigControl", "Control Rig Control"); }
+	virtual FSlateIcon GetIcon() const override { return FSlateIconFinder::FindIconForClass(UControlRigBlueprint::StaticClass()); }
+	//~ End FFilterBase
 
-		FText Tooltip = LOCTEXT("SequencerTrackFilter_ControlRigControlsTip", "Show Only Control Rig Controls.");
+	//~ Begin FSequencerTrackFilter
 
-		if (FirstActiveChord->IsValidChord())
-		{
-			return FText::Join(FText::FromString(TEXT(" ")), Tooltip, FirstActiveChord->GetInputText());
-		}
-		return Tooltip;
+	virtual FText GetDefaultToolTipText() const override
+	{
+		return LOCTEXT("SequencerTrackFilter_ControlRigControlsTip", "Show only Control Rig Control tracks");
 	}
 
-	virtual void BindCommands(TSharedRef<FUICommandList> SequencerBindings, TSharedRef<FUICommandList> CurveEditorBindings, TWeakPtr<ISequencer> Sequencer) override
+	virtual TSharedPtr<FUICommandInfo> GetToggleCommand() const override
 	{
-		const FSequencerTrackFilter_ControlRigControlsCommands& Commands = FSequencerTrackFilter_ControlRigControlsCommands::Get();
-
-		SequencerBindings->MapAction(
-			Commands.ToggleControlRigControls,
-			FExecuteAction::CreateLambda([this, Sequencer] { Sequencer.Pin()->SetTrackFilterEnabled(GetDisplayName(), !Sequencer.Pin()->IsTrackFilterEnabled(GetDisplayName())); }),
-			FCanExecuteAction::CreateLambda([this, Sequencer] { return true; }),
-			FIsActionChecked::CreateLambda([this, Sequencer] { return Sequencer.Pin()->IsTrackFilterEnabled(GetDisplayName()); }));
-
-		CurveEditorBindings->MapAction(Commands.ToggleControlRigControls, *SequencerBindings->GetActionForCommand(Commands.ToggleControlRigControls));
+		return FSequencerTrackFilter_ControlRigControlsCommands::Get().ToggleFilter_ControlRigControls;
 	}
 
+	//~ End FSequencerTrackFilter
+
 private:
 	mutable uint32 BindingCount;
 };
@@ -115,120 +110,138 @@ class FSequencerTrackFilter_ControlRigSelectedControlsCommands
 	: public TCommands<FSequencerTrackFilter_ControlRigSelectedControlsCommands>
 {
 public:
-
 	FSequencerTrackFilter_ControlRigSelectedControlsCommands()
-		: TCommands<FSequencerTrackFilter_ControlRigSelectedControlsCommands>
-		(
-			"FSequencerTrackFilter_ControlRigSelectedControls",
-			NSLOCTEXT("Contexts", "FSequencerTrackFilter_ControlRigSelectedControls", "FSequencerTrackFilter_ControlRigSelectedControls"),
+		: TCommands<FSequencerTrackFilter_ControlRigSelectedControlsCommands>(
+			TEXT("FSequencerTrackFilter_ControlRigSelectedControls"),
+			LOCTEXT("FSequencerTrackFilter_ControlRigSelectedControls", "Control Rig Selected Control Filters"),
 			NAME_None,
-			FAppStyle::GetAppStyleSetName() // Icon Style Set
-			)
-	{ }
+			FAppStyle::GetAppStyleSetName())
+	{}
 
 	/** Toggle the control rig selected controls filter */
-	TSharedPtr< FUICommandInfo > ToggleControlRigSelectedControls;
+	TSharedPtr<FUICommandInfo> ToggleFilter_ControlRigSelectedControls;
 
 	/** Initialize commands */
 	virtual void RegisterCommands() override
 	{
-		UI_COMMAND(ToggleControlRigSelectedControls, "Control Rig Selected Controls", "Toggle the filter for Control Rig Selected Controls.", EUserInterfaceActionType::ToggleButton, FInputChord(EKeys::F10));
+		UI_COMMAND(ToggleFilter_ControlRigSelectedControls, "Control Rig Selected Controls", "Toggle the filter for Control Rig Selected Controls.", EUserInterfaceActionType::ToggleButton, FInputChord(EKeys::F10));
 	}
 };
 
+//////////////////////////////////////////////////////////////////////////
+//
+
 class FSequencerTrackFilter_ControlRigSelectedControls : public FSequencerTrackFilter
 {
 public:
-
-	FSequencerTrackFilter_ControlRigSelectedControls()
-		: BindingCount(0)
+	FSequencerTrackFilter_ControlRigSelectedControls(ISequencerTrackFilters& InOutFilterInterface, TSharedPtr<FFilterCategory> InCategory = nullptr)
+		:  FSequencerTrackFilter(InOutFilterInterface, MoveTemp(InCategory))
+		, BindingCount(0)
 	{
 		FSequencerTrackFilter_ControlRigSelectedControlsCommands::Register();
 	}
 
-	~FSequencerTrackFilter_ControlRigSelectedControls()
+	virtual ~FSequencerTrackFilter_ControlRigSelectedControls() override
 	{
 		BindingCount--;
-
 		if (BindingCount < 1)
 		{
 			FSequencerTrackFilter_ControlRigSelectedControlsCommands::Unregister();
 		}
 	}
 
-	virtual FString GetName() const override { return TEXT("ControlRigControlsSelectedFilter"); }
-	virtual FText GetDisplayName() const override { return LOCTEXT("SequenceTrackFilter_ControlRigSelectedControls", "Selected Control Rig Controls"); }
-	virtual FSlateIcon GetIcon() const { return FSlateIconFinder::FindIconForClass(UControlRig::StaticClass()); }
+	//~ Begin IFilter
 
-	virtual bool PassesFilter(FTrackFilterType InItem) const override
-	{
-		return true;
-	}
+	virtual FString GetName() const override { return TEXT("SelectedControlRigControl"); }
 
-	virtual bool PassesFilterWithDisplayName(FTrackFilterType InItem, const FText& InText) const
+	virtual bool PassesFilter(FSequencerTrackFilterType InItem) const override
 	{
-		const UMovieSceneControlRigParameterTrack* Track = Cast<UMovieSceneControlRigParameterTrack>(InItem);
-		if (Track)
+		FSequencerFilterData& FilterData = FilterInterface.GetFilterData();
+
+		const TViewModelPtr<IOutlinerExtension> OutlinerExtension = InItem.AsModel()->FindAncestorOfType<IOutlinerExtension>();
+		if (!OutlinerExtension.IsValid())
 		{
-			UControlRig *ControlRig = Track->GetControlRig();
-			if (ControlRig && ControlRig->GetHierarchy())
+			return false;
+		}
+
+		const FString ControlTrackLabel = OutlinerExtension->GetLabel().ToString();
+
+		UMovieSceneTrack* const TrackObject = ResolveMovieSceneTrackObject(InItem, FilterData);
+		URigHierarchy* const ControlRigHierarchy = GetControlRigHierarchyFromTrackObject(TrackObject);
+
+		if (!IsValid(ControlRigHierarchy))
+		{
+			return false;
+		}
+
+		const TArray<const FRigBaseElement*> BaseRigElements = ControlRigHierarchy->GetSelectedElements(ERigElementType::Control);
+		for (const FRigBaseElement* const BaseRigElement : BaseRigElements)
+		{
+			const FString ElementDisplayName = ControlRigHierarchy->GetDisplayNameForUI(BaseRigElement).ToString();
+			if (ElementDisplayName.Equals(ControlTrackLabel))
+			{
+				return true;
+			}
+
+			if (const FRigControlElement* const ControlElement = Cast<FRigControlElement>(BaseRigElement))
 			{
-				TArray<const FRigBaseElement*> SelectedControls = ControlRig->GetHierarchy()->GetSelectedElements(ERigElementType::Control);
-				for (const FRigBaseElement* SelectedControl : SelectedControls)
+				if (ControlElement->CanDriveControls())
 				{
-					if (InText.EqualTo(ControlRig->GetHierarchy()->GetDisplayNameForUI(SelectedControl)))
+					const TArray<FRigElementKey>& DrivenControls = ControlElement->Settings.DrivenControls;
+					for (const FRigElementKey& DrivenKey : DrivenControls)
 					{
-						return true;
-					}
-					if (const FRigControlElement* ControlElement = Cast<FRigControlElement>(SelectedControl))
-					{
-						if (ControlElement->CanDriveControls())
+						if (const FRigBaseElement* const DrivenKeyElement = ControlRigHierarchy->Find(DrivenKey))
 						{
-							const TArray<FRigElementKey>& DrivenControls = ControlElement->Settings.DrivenControls;
-							for (const FRigElementKey& DrivenKey : DrivenControls)
+							const FString DrivenKeyElementDisplayName = ControlRigHierarchy->GetDisplayNameForUI(DrivenKeyElement).ToString();
+							if (DrivenKeyElementDisplayName.Equals(ControlTrackLabel))
 							{
-								if (FRigBaseElement* Element = ControlRig->GetHierarchy()->Find(DrivenKey))
-								{
-									if (InText.EqualTo(ControlRig->GetHierarchy()->GetDisplayNameForUI(Element)))
-									{
-										return true;
-									}
-								}
+								return true;
 							}
 						}
 					}
 				}
 			}
 		}
+
 		return false;
 	}
 
-	virtual FText GetToolTipText() const override
-	{
-		const FSequencerTrackFilter_ControlRigSelectedControlsCommands& Commands = FSequencerTrackFilter_ControlRigSelectedControlsCommands::Get();
+	//~ End IFilter
 
-		const TSharedRef<const FInputChord> FirstActiveChord = Commands.ToggleControlRigSelectedControls->GetFirstValidChord();
+	//~ Begin FFilterBase
+	virtual FText GetDisplayName() const override { return LOCTEXT("SequenceTrackFilter_ControlRigSelectedControl", "Selected Control Rig Control"); }
+	virtual FSlateIcon GetIcon() const override { return FSlateIconFinder::FindIconForClass(UControlRigBlueprint::StaticClass()); }
+	//~ End FFilterBase
 
-		FText Tooltip = LOCTEXT("SequencerTrackFilter_ControlRigSelectedControlsTip", "Show Only Selected Control Rig Controls.");
+	//~ Begin FSequencerTrackFilter
 
-		if (FirstActiveChord->IsValidChord())
-		{
-			return FText::Join(FText::FromString(TEXT(" ")), Tooltip, FirstActiveChord->GetInputText());
-		}
-		return Tooltip;
+	virtual FText GetDefaultToolTipText() const override
+	{
+		return LOCTEXT("SequencerTrackFilter_ControlRigSelectedControlsTip", "Show Only Selected Control Rig Controls.");
 	}
 
-	virtual void BindCommands(TSharedRef<FUICommandList> SequencerBindings, TSharedRef<FUICommandList> CurveEditorBindings, TWeakPtr<ISequencer> Sequencer) override
+	virtual TSharedPtr<FUICommandInfo> GetToggleCommand() const override
 	{
-		const FSequencerTrackFilter_ControlRigSelectedControlsCommands& Commands = FSequencerTrackFilter_ControlRigSelectedControlsCommands::Get();
+		return FSequencerTrackFilter_ControlRigSelectedControlsCommands::Get().ToggleFilter_ControlRigSelectedControls;
+	}
 
-		SequencerBindings->MapAction(
-			Commands.ToggleControlRigSelectedControls,
-			FExecuteAction::CreateLambda([this, Sequencer] { Sequencer.Pin()->SetTrackFilterEnabled(GetDisplayName(), !Sequencer.Pin()->IsTrackFilterEnabled(GetDisplayName())); }),
-			FCanExecuteAction::CreateLambda([this, Sequencer] { return true; }),
-			FIsActionChecked::CreateLambda([this, Sequencer] { return Sequencer.Pin()->IsTrackFilterEnabled(GetDisplayName()); }));
+	//~ End FSequencerTrackFilter
 
-		CurveEditorBindings->MapAction(Commands.ToggleControlRigSelectedControls, *SequencerBindings->GetActionForCommand(Commands.ToggleControlRigSelectedControls));
+	static URigHierarchy* GetControlRigHierarchyFromTrackObject(UMovieSceneTrack* const InTrackObject)
+	{
+		const UMovieSceneControlRigParameterTrack* const Track = Cast<UMovieSceneControlRigParameterTrack>(InTrackObject);
+		if (!IsValid(Track))
+		{
+			return nullptr;
+		}
+
+		UControlRig* const ControlRig = Track->GetControlRig();
+		if (!IsValid(ControlRig))
+		{
+			return nullptr;
+		}
+
+		return ControlRig->GetHierarchy();
 	}
 
 private:
@@ -249,11 +262,12 @@ private:
 //////////////////////////////////////////////////////////////////////////
 //
 
-void UControlRigTrackFilter::AddTrackFilterExtensions(TArray< TSharedRef<class FSequencerTrackFilter> >& InOutFilterList) const
+void UControlRigTrackFilter::AddTrackFilterExtensions(ISequencerTrackFilters& InOutFilterInterface
+	, const TSharedRef<FFilterCategory>& InPreferredCategory
+	, TArray<TSharedRef<FSequencerTrackFilter>>& InOutFilterList) const
 {
-	InOutFilterList.Add(MakeShared<FSequencerTrackFilter_ControlRigControls>());
-	InOutFilterList.Add(MakeShared<FSequencerTrackFilter_ControlRigSelectedControls>());
+	InOutFilterList.Add(MakeShared<FSequencerTrackFilter_ControlRigControls>(InOutFilterInterface, InPreferredCategory));
+	InOutFilterList.Add(MakeShared<FSequencerTrackFilter_ControlRigSelectedControls>(InOutFilterInterface, InPreferredCategory));
 }
 
 #undef LOCTEXT_NAMESPACE
-
diff --git a/Engine/Plugins/Animation/ControlRig/Source/ControlRigEditor/Private/Sequencer/ControlRigSequencerFilter.h b/Engine/Plugins/Animation/ControlRig/Source/ControlRigEditor/Private/Sequencer/ControlRigSequencerFilter.h
--- a/Engine/Plugins/Animation/ControlRig/Source/ControlRigEditor/Private/Sequencer/ControlRigSequencerFilter.h
+++ b/Engine/Plugins/Animation/ControlRig/Source/ControlRigEditor/Private/Sequencer/ControlRigSequencerFilter.h
@@ -2,10 +2,7 @@
 
 #pragma once
 
-#include "CoreMinimal.h"
-#include "UObject/ObjectMacros.h"
-#include "SequencerTrackFilterBase.h"
-#include "SequencerTrackFilterExtension.h"
+#include "Filters/SequencerTrackFilterExtension.h"
 #include "ControlRigSequencerFilter.generated.h"
 
 UCLASS()
@@ -14,7 +11,9 @@ class UControlRigTrackFilter : public USequencerTrackFilterExtension
 public:
 	GENERATED_BODY()
 
-	// USequencerTrackFilterExtension interface
-	virtual void AddTrackFilterExtensions(TArray< TSharedRef<class FSequencerTrackFilter> >& InOutFilterList) const override;
-	// End of USequencerTrackFilterExtension interface
+	//~ Begin USequencerTrackFilterExtension
+	virtual void AddTrackFilterExtensions(ISequencerTrackFilters& InOutFilterInterface
+		, const TSharedRef<FFilterCategory>& InPreferredCategory
+		, TArray<TSharedRef<FSequencerTrackFilter>>& InOutFilterList) const override;
+	//~ End USequencerTrackFilterExtension
 };
diff --git a/Engine/Plugins/Animation/LiveLink/Source/LiveLinkEditor/Private/Sequencer/LiveLinkSequencerFilters.cpp b/Engine/Plugins/Animation/LiveLink/Source/LiveLinkEditor/Private/Sequencer/LiveLinkSequencerFilters.cpp
--- a/Engine/Plugins/Animation/LiveLink/Source/LiveLinkEditor/Private/Sequencer/LiveLinkSequencerFilters.cpp
+++ b/Engine/Plugins/Animation/LiveLink/Source/LiveLinkEditor/Private/Sequencer/LiveLinkSequencerFilters.cpp
@@ -1,42 +1,97 @@
 // Copyright Epic Games, Inc. All Rights Reserved.
 
 #include "LiveLinkSequencerFilters.h"
-#include "MovieScene/MovieSceneLiveLinkTrack.h"
+#include "Filters/SequencerTrackFilterBase.h"
+#include "Framework/Commands/Commands.h"
+#include "Framework/Commands/UICommandInfo.h"
 #include "LiveLinkComponent.h"
-#include "SequencerTrackFilterBase.h"
+#include "MovieScene/MovieSceneLiveLinkTrack.h"
 #include "Styling/SlateIconFinder.h"
 
 #include UE_INLINE_GENERATED_CPP_BY_NAME(LiveLinkSequencerFilters)
 
 #define LOCTEXT_NAMESPACE "LiveLinkSequencerTrackFilters"
 
+class FSequencerTrackFilter_LiveLinkFilterCommands
+	: public TCommands<FSequencerTrackFilter_LiveLinkFilterCommands>
+{
+public:
+	FSequencerTrackFilter_LiveLinkFilterCommands()
+		: TCommands<FSequencerTrackFilter_LiveLinkFilterCommands>(
+			TEXT("FSequencerTrackFilter_LiveLink"),
+			LOCTEXT("FSequencerTrackFilter_LiveLink", "Live Link Filters"),
+			NAME_None,
+			FAppStyle::GetAppStyleSetName())
+	{}
+
+	TSharedPtr<FUICommandInfo> ToggleFilter_LiveLink;
+
+	virtual void RegisterCommands() override
+	{
+		UI_COMMAND(ToggleFilter_LiveLink, "Toggle Live Link Filter", "Toggle the filter for Live Link tracks", EUserInterfaceActionType::ToggleButton, FInputChord());
+	}
+};
+
 //////////////////////////////////////////////////////////////////////////
 //
 
-class FSequencerTrackFilter_LiveLinkTracks : public FSequencerTrackFilter_ClassType<UMovieSceneLiveLinkTrack>
+class FSequencerTrackFilter_LiveLink : public FSequencerTrackFilter_ClassType<UMovieSceneLiveLinkTrack>
 {
-	virtual FString GetName() const override { return TEXT("SequencerLiveLinkTrackFilter"); }
-	virtual FText GetDisplayName() const override { return LOCTEXT("SequencerTrackFilter_LiveLinkTracks", "Live Link"); }
-	virtual FText GetToolTipText() const override { return LOCTEXT("SequencerTrackFilter_LiveLinkTracksToolTip", "Show only Live Link tracks."); }
-	virtual FSlateIcon GetIcon() const { return FSlateIconFinder::FindIconForClass(ULiveLinkComponent::StaticClass()); }
-	virtual bool SupportsSequence(UMovieSceneSequence* InSequence) const override
+public:
+	FSequencerTrackFilter_LiveLink(ISequencerTrackFilters& InFilterInterface, TSharedPtr<FFilterCategory> InCategory = nullptr)
+		:  FSequencerTrackFilter_ClassType<UMovieSceneLiveLinkTrack>(InFilterInterface, InCategory)
+		, BindingCount(0)
+	{
+		FSequencerTrackFilter_LiveLinkFilterCommands::Register();
+	}
+
+	virtual ~FSequencerTrackFilter_LiveLink() override
 	{
-		if (InSequence && InSequence->IsTrackSupported(UMovieSceneLiveLinkTrack::StaticClass()) == ETrackSupport::NotSupported)
+		BindingCount--;
+		if (BindingCount < 1)
 		{
-			return false;
+			FSequencerTrackFilter_LiveLinkFilterCommands::Unregister();
 		}
+	}
 
-		return FSequencerTrackFilter_ClassType<UMovieSceneLiveLinkTrack>::SupportsSequence(InSequence);
+	//~ Begin IFilter
+	virtual FString GetName() const override { return TEXT("LiveLink"); }
+	//~ End IFilter
+
+	//~ Begin FFilterBase
+	virtual FText GetDisplayName() const override { return LOCTEXT("SequencerTrackFilter_LiveLink", "Live Link"); }
+	virtual FSlateIcon GetIcon() const override { return FSlateIconFinder::FindIconForClass(ULiveLinkComponent::StaticClass()); }
+	//~ End FFilterBase
+
+	//~ Begin FSequencerTrackFilter
+
+	virtual FText GetDefaultToolTipText() const override
+	{
+		return LOCTEXT("SequencerTrackFilter_LiveLinkToolTip", "Show only Live Link tracks");
+	}
+
+	virtual TSharedPtr<FUICommandInfo> GetToggleCommand() const override
+	{
+		return FSequencerTrackFilter_LiveLinkFilterCommands::Get().ToggleFilter_LiveLink;
 	}
+
+	virtual bool SupportsSequence(UMovieSceneSequence* const InSequence) const override
+	{
+		return SupportsLevelSequence(InSequence) || IsSequenceTrackSupported<UMovieSceneLiveLinkTrack>(InSequence);
+	}
+
+	//~ End FSequencerTrackFilter
+
+private:
+	mutable uint32 BindingCount;
 };
 
 //////////////////////////////////////////////////////////////////////////
 //
 
-void ULiveLinkSequencerTrackFilter::AddTrackFilterExtensions(TArray< TSharedRef<class FSequencerTrackFilter> >& InOutFilterList) const
+void ULiveLinkSequencerTrackFilter::AddTrackFilterExtensions(ISequencerTrackFilters& InFilterInterface, const TSharedRef<FFilterCategory>& InPreferredCategory, TArray<TSharedRef<FSequencerTrackFilter>>& InOutFilterList) const
 {
-	InOutFilterList.Add(MakeShared<FSequencerTrackFilter_LiveLinkTracks>());
+	InOutFilterList.Add(MakeShared<FSequencerTrackFilter_LiveLink>(InFilterInterface, InPreferredCategory));
 }
 
 #undef LOCTEXT_NAMESPACE
-
diff --git a/Engine/Plugins/Animation/LiveLink/Source/LiveLinkEditor/Private/Sequencer/LiveLinkSequencerFilters.h b/Engine/Plugins/Animation/LiveLink/Source/LiveLinkEditor/Private/Sequencer/LiveLinkSequencerFilters.h
--- a/Engine/Plugins/Animation/LiveLink/Source/LiveLinkEditor/Private/Sequencer/LiveLinkSequencerFilters.h
+++ b/Engine/Plugins/Animation/LiveLink/Source/LiveLinkEditor/Private/Sequencer/LiveLinkSequencerFilters.h
@@ -2,7 +2,7 @@
 
 #pragma once
 
-#include "SequencerTrackFilterExtension.h"
+#include "Filters/SequencerTrackFilterExtension.h"
 #include "LiveLinkSequencerFilters.generated.h"
 
 UCLASS()
@@ -11,7 +11,9 @@ class ULiveLinkSequencerTrackFilter : public USequencerTrackFilterExtension
 public:
 	GENERATED_BODY()
 
-	// USequencerTrackFilterExtension interface
-	virtual void AddTrackFilterExtensions(TArray< TSharedRef<class FSequencerTrackFilter> >& InOutFilterList) const override;
-	// End of USequencerTrackFilterExtension interface
+	//~ Begin USequencerTrackFilterExtension
+	virtual void AddTrackFilterExtensions(ISequencerTrackFilters& InFilterInterface
+		, const TSharedRef<FFilterCategory>& InPreferredCategory
+		, TArray<TSharedRef<FSequencerTrackFilter>>& InOutFilterList) const override;
+	//~ End USequencerTrackFilterExtension
 };
diff --git a/Engine/Plugins/FX/Niagara/Source/NiagaraEditor/Private/Sequencer/NiagaraSequence/NiagaraSequencerFilters.cpp b/Engine/Plugins/FX/Niagara/Source/NiagaraEditor/Private/Sequencer/NiagaraSequence/NiagaraSequencerFilters.cpp
--- a/Engine/Plugins/FX/Niagara/Source/NiagaraEditor/Private/Sequencer/NiagaraSequence/NiagaraSequencerFilters.cpp
+++ b/Engine/Plugins/FX/Niagara/Source/NiagaraEditor/Private/Sequencer/NiagaraSequence/NiagaraSequencerFilters.cpp
@@ -1,6 +1,9 @@
 // Copyright Epic Games, Inc. All Rights Reserved.
 
 #include "NiagaraSequencerFilters.h"
+#include "Filters/SequencerTrackFilterBase.h"
+#include "Framework/Commands/Commands.h"
+#include "Framework/Commands/UICommandInfo.h"
 #include "MovieScene/MovieSceneNiagaraTrack.h"
 #include "NiagaraActor.h"
 #include "NiagaraComponent.h"
@@ -10,49 +13,107 @@
 
 #define LOCTEXT_NAMESPACE "NiagaraSequencerTrackFilters"
 
+class FSequencerTrackFilter_NiagaraFilterCommands
+	: public TCommands<FSequencerTrackFilter_NiagaraFilterCommands>
+{
+public:
+	FSequencerTrackFilter_NiagaraFilterCommands()
+		: TCommands<FSequencerTrackFilter_NiagaraFilterCommands>(
+			TEXT("FSequencerTrackFilter_Niagara"),
+			LOCTEXT("FSequencerTrackFilter_Niagara", "Niagara Filters"),
+			NAME_None,
+			FAppStyle::GetAppStyleSetName())
+	{}
+
+	TSharedPtr<FUICommandInfo> ToggleFilter_Niagara;
+
+	virtual void RegisterCommands() override
+	{
+		UI_COMMAND(ToggleFilter_Niagara, "Toggle Niagara Filter", "Toggle the filter for Niagara tracks", EUserInterfaceActionType::ToggleButton, FInputChord());
+	}
+};
+
 //////////////////////////////////////////////////////////////////////////
 //
 
-class FSequencerTrackFilter_NiagaraTracks : public FSequencerTrackFilter
+class FSequencerTrackFilter_Niagara : public FSequencerTrackFilter
 {
-	virtual FString GetName() const override { return TEXT("SequencerNiagaraTrackFilter"); }
-	virtual FText GetDisplayName() const override { return LOCTEXT("SequencerTrackFilter_NiagaraTracks", "Niagara"); }
-	virtual FText GetToolTipText() const override { return LOCTEXT("SequencerTrackFilter_NiagaraTracksToolTip", "Show only Niagara tracks."); }
-	virtual FSlateIcon GetIcon() const { return FSlateIcon(FAppStyle::GetAppStyleSetName(), "ClassIcon.ParticleSystem"); }
+public:
+	FSequencerTrackFilter_Niagara(ISequencerTrackFilters& InFilterInterface, TSharedPtr<FFilterCategory> InCategory = nullptr)
+		: FSequencerTrackFilter(InFilterInterface, MoveTemp(InCategory))
+		, BindingCount(0)
+	{
+		FSequencerTrackFilter_NiagaraFilterCommands::Register();
+	}
 
-	// IFilter implementation
-	virtual bool PassesFilter(FTrackFilterType InItem) const override
+	virtual ~FSequencerTrackFilter_Niagara() override
 	{
-		if (!InItem)
+		BindingCount--;
+		if (BindingCount < 1)
 		{
-			return false;
+			FSequencerTrackFilter_NiagaraFilterCommands::Unregister();
 		}
+	}
 
-		if (InItem->IsA(UMovieSceneNiagaraTrack::StaticClass()) || InItem->IsA(ANiagaraActor::StaticClass()) || InItem->IsA(UNiagaraComponent::StaticClass()))
-		{
-			return true;
-		}
+	//~ Begin FFilterBase
+	virtual FText GetDisplayName() const override { return LOCTEXT("SequencerTrackFilter_Niagara", "Niagara"); }
+	virtual FSlateIcon GetIcon() const override { return FSlateIcon(FAppStyle::GetAppStyleSetName(), TEXT("PlacementBrowser.Icons.VisualEffects")); } // TEXT("ClassIcon.ParticleSystem")
+	//~ End FFilterBase
+
+	//~ Begin IFilter
+
+	virtual FString GetName() const override { return TEXT("Niagara"); }
+
+	virtual bool PassesFilter(FSequencerTrackFilterType InItem) const override
+	{
+		FSequencerFilterData& FilterData = FilterInterface.GetFilterData();
 
-		const AActor* Actor = Cast<const AActor>(InItem);
-		if (Actor)
+		const UObject* const BoundObject = ResolveTrackBoundObject(InItem, FilterData);
+		if (IsValid(BoundObject))
 		{
-			if (Actor->FindComponentByClass(UNiagaraComponent::StaticClass()))
+			if (BoundObject->IsA(UMovieSceneNiagaraTrack::StaticClass())
+				|| BoundObject->IsA(ANiagaraActor::StaticClass())
+				|| BoundObject->IsA(UNiagaraComponent::StaticClass()))
 			{
 				return true;
 			}
 		}
 
+		const AActor* const Actor = Cast<const AActor>(BoundObject);
+		if (IsValid(Actor) && Actor->FindComponentByClass(UNiagaraComponent::StaticClass()))
+		{
+			return true;
+		}
+
 		return false;
 	}
+
+	//~ End IFilter
+
+	//~ Begin FSequencerTrackFilter
+
+	virtual FText GetDefaultToolTipText() const override
+	{
+		return LOCTEXT("SequencerTrackFilter_NiagaraToolTip", "Show only Niagara tracks");
+	}
+
+	virtual TSharedPtr<FUICommandInfo> GetToggleCommand() const override
+	{
+		return FSequencerTrackFilter_NiagaraFilterCommands::Get().ToggleFilter_Niagara;
+	}
+
+	//~ End FSequencerTrackFilter
+
+private:
+	mutable uint32 BindingCount;
 };
 
 //////////////////////////////////////////////////////////////////////////
 //
 
-void UNiagaraSequencerTrackFilter::AddTrackFilterExtensions(TArray< TSharedRef<class FSequencerTrackFilter> >& InOutFilterList) const
+void UNiagaraSequencerTrackFilter::AddTrackFilterExtensions(ISequencerTrackFilters& InFilterInterface, const TSharedRef<FFilterCategory>& InPreferredCategory, TArray<TSharedRef<FSequencerTrackFilter>>& InOutFilterList) const
 {
-	InOutFilterList.Add(MakeShared<FSequencerTrackFilter_NiagaraTracks>());
+	InOutFilterList.Add(MakeShared<FSequencerTrackFilter_Niagara>(InFilterInterface, InPreferredCategory));
 }
 
 #undef LOCTEXT_NAMESPACE
-
diff --git a/Engine/Plugins/FX/Niagara/Source/NiagaraEditor/Private/Sequencer/NiagaraSequence/NiagaraSequencerFilters.h b/Engine/Plugins/FX/Niagara/Source/NiagaraEditor/Private/Sequencer/NiagaraSequence/NiagaraSequencerFilters.h
--- a/Engine/Plugins/FX/Niagara/Source/NiagaraEditor/Private/Sequencer/NiagaraSequence/NiagaraSequencerFilters.h
+++ b/Engine/Plugins/FX/Niagara/Source/NiagaraEditor/Private/Sequencer/NiagaraSequence/NiagaraSequencerFilters.h
@@ -2,10 +2,7 @@
 
 #pragma once
 
-#include "CoreMinimal.h"
-#include "UObject/ObjectMacros.h"
-#include "SequencerTrackFilterBase.h"
-#include "SequencerTrackFilterExtension.h"
+#include "Filters/SequencerTrackFilterExtension.h"
 #include "NiagaraSequencerFilters.generated.h"
 
 UCLASS()
@@ -14,7 +11,9 @@ class UNiagaraSequencerTrackFilter : public USequencerTrackFilterExtension
 public:
 	GENERATED_BODY()
 
-	// USequencerTrackFilterExtension interface
-	virtual void AddTrackFilterExtensions(TArray< TSharedRef<class FSequencerTrackFilter> >& InOutFilterList) const override;
-	// End of USequencerTrackFilterExtension interface
+	//~ Begin USequencerTrackFilterExtension
+	virtual void AddTrackFilterExtensions(ISequencerTrackFilters& InFilterInterface
+		, const TSharedRef<FFilterCategory>& InPreferredCategory
+		, TArray<TSharedRef<FSequencerTrackFilter>>& InOutFilterList) const override;
+	//~ End USequencerTrackFilterExtension
 };
diff --git a/Engine/Plugins/MovieScene/LevelSequenceEditor/Source/LevelSequenceEditor/Private/LevelSequenceEditorBlueprintLibrary.cpp b/Engine/Plugins/MovieScene/LevelSequenceEditor/Source/LevelSequenceEditor/Private/LevelSequenceEditorBlueprintLibrary.cpp
--- a/Engine/Plugins/MovieScene/LevelSequenceEditor/Source/LevelSequenceEditor/Private/LevelSequenceEditorBlueprintLibrary.cpp
+++ b/Engine/Plugins/MovieScene/LevelSequenceEditor/Source/LevelSequenceEditor/Private/LevelSequenceEditorBlueprintLibrary.cpp
@@ -2,6 +2,7 @@
 
 #include "LevelSequenceEditorBlueprintLibrary.h"
 
+#include "Filters/ISequencerTrackFilters.h"
 #include "ISequencer.h"
 #include "MVVM/ViewModels/ChannelModel.h"
 #include "MVVM/ViewModels/SectionModel.h"
@@ -726,35 +727,38 @@ void ULevelSequenceEditorBlueprintLibrary::SetLockLevelSequence(bool bLock)
 
 bool ULevelSequenceEditorBlueprintLibrary::IsTrackFilterEnabled(const FText& TrackFilterName)
 {
-	if (CurrentSequencer.IsValid())
-	{
-		TSharedPtr<ISequencer> Sequencer = CurrentSequencer.Pin();
+	return IsTrackFilterActive(TrackFilterName);
+}
 
-		return Sequencer->IsTrackFilterEnabled(TrackFilterName);
+bool ULevelSequenceEditorBlueprintLibrary::IsTrackFilterActive(const FText& TrackFilterName)
+{
+	if (const TSharedPtr<ISequencer> Sequencer = CurrentSequencer.Pin())
+	{
+		return Sequencer->GetFilterInterface()->IsFilterActiveByDisplayName(TrackFilterName.ToString());
 	}
 	return false;
 }
 
 void ULevelSequenceEditorBlueprintLibrary::SetTrackFilterEnabled(const FText& TrackFilterName, bool bEnabled)
 {
-	if (CurrentSequencer.IsValid())
-	{
-		TSharedPtr<ISequencer> Sequencer = CurrentSequencer.Pin();
+	return SetTrackFilterActive(TrackFilterName, bEnabled);
+}
 
-		Sequencer->SetTrackFilterEnabled(TrackFilterName, bEnabled);
+void ULevelSequenceEditorBlueprintLibrary::SetTrackFilterActive(const FText& TrackFilterName, bool bActive)
+{
+	if (const TSharedPtr<ISequencer> Sequencer = CurrentSequencer.Pin())
+	{
+		Sequencer->GetFilterInterface()->SetFilterActiveByDisplayName(TrackFilterName.ToString(), bActive);
 	}
 }
 
 TArray<FText> ULevelSequenceEditorBlueprintLibrary::GetTrackFilterNames()
 {
-	if (CurrentSequencer.IsValid())
+	if (const TSharedPtr<ISequencer> Sequencer = CurrentSequencer.Pin())
 	{
-		TSharedPtr<ISequencer> Sequencer = CurrentSequencer.Pin();
-
-		return Sequencer->GetTrackFilterNames();
+		return Sequencer->GetFilterInterface()->GetFilterDisplayNames();
 	}
-
-	return TArray<FText>();
+	return {};
 }
 
 bool ULevelSequenceEditorBlueprintLibrary::HasCustomColorForChannel(UClass* Class, const FString& Identifier)
diff --git a/Engine/Plugins/MovieScene/LevelSequenceEditor/Source/LevelSequenceEditor/Public/LevelSequenceEditorBlueprintLibrary.h b/Engine/Plugins/MovieScene/LevelSequenceEditor/Source/LevelSequenceEditor/Public/LevelSequenceEditorBlueprintLibrary.h
--- a/Engine/Plugins/MovieScene/LevelSequenceEditor/Source/LevelSequenceEditor/Public/LevelSequenceEditorBlueprintLibrary.h
+++ b/Engine/Plugins/MovieScene/LevelSequenceEditor/Source/LevelSequenceEditor/Public/LevelSequenceEditorBlueprintLibrary.h
@@ -268,13 +268,23 @@ public:
 public:
 
 	/** Gets whether the specified track filter is on/off */
+	UE_DEPRECATED(5.5, "Use IsTrackFilterActive")
 	UFUNCTION(BlueprintPure, Category = "Level Sequence Editor")
 	static bool IsTrackFilterEnabled(const FText& TrackFilterName);
 
+	/** Gets whether the specified track filter is on/off */
+	UFUNCTION(BlueprintPure, Category = "Level Sequence Editor")
+	static bool IsTrackFilterActive(const FText& TrackFilterName);
+
 	/** Sets the specified track filter to be on or off */
+	UE_DEPRECATED(5.5, "Use SetTrackFilterActive")
 	UFUNCTION(BlueprintCallable, Category = "Level Sequence Editor")
 	static void SetTrackFilterEnabled(const FText& TrackFilterName, bool bEnabled);
 
+	/** Sets the specified track filter to be on or off */
+	UFUNCTION(BlueprintCallable, Category = "Level Sequence Editor")
+	static void SetTrackFilterActive(const FText& TrackFilterName, bool bActive);
+
 	/** Gets all the available track filter names */
 	UFUNCTION(BlueprintPure, Category = "Level Sequence Editor")
 	static TArray<FText> GetTrackFilterNames();
diff --git a/Engine/Source/Editor/Sequencer/Private/Filters/Filters/SequencerTrackFilterBase.cpp b/Engine/Source/Editor/Sequencer/Private/Filters/Filters/SequencerTrackFilterBase.cpp
new file mode 100644
--- /dev/null
+++ b/Engine/Source/Editor/Sequencer/Private/Filters/Filters/SequencerTrackFilterBase.cpp
@@ -0,0 +1,189 @@
+﻿// Copyright Epic Games, Inc. All Rights Reserved.
+
+#include "Filters/SequencerTrackFilterBase.h"
+#include "Animation/WidgetAnimation.h"
+#include "Filters/SequencerFilterBar.h"
+#include "ISequencer.h"
+#include "LevelSequence.h"
+#include "MovieSceneSequence.h"
+#include "MVVM/ViewModelPtr.h"
+#include "Sequencer.h"
+
+using namespace UE::Sequencer;
+
+FSequencerTrackFilter::FSequencerTrackFilter(ISequencerTrackFilters& InOutFilterInterface, TSharedPtr<FFilterCategory>&& InCategory)
+	: FFilterBase<FSequencerTrackFilterType>(MoveTemp(InCategory))
+	, FilterInterface(InOutFilterInterface)
+{
+}
+
+FText FSequencerTrackFilter::GetToolTipText() const
+{
+	if (const TSharedPtr<FUICommandInfo> ToggleCommand = GetToggleCommand())
+	{
+		return BuildTooltipTextForCommand(GetDefaultToolTipText(), ToggleCommand);
+	}
+	return GetDefaultToolTipText();
+}
+
+bool FSequencerTrackFilter::SupportsSequence(UMovieSceneSequence* const InSequence) const
+{
+	return SupportsLevelSequence(InSequence);
+}
+
+void FSequencerTrackFilter::BindCommands()
+{
+	if (const TSharedPtr<FUICommandInfo> ToggleCommand = GetToggleCommand())
+	{
+		MapToggleAction(ToggleCommand);
+	}
+}
+
+void FSequencerTrackFilter::BroadcastChangedEvent() const
+{
+	ChangedEvent.Broadcast();
+}
+
+ISequencerTrackFilters& FSequencerTrackFilter::GetFilterInterface() const
+{
+	return FilterInterface;
+}
+
+ISequencer& FSequencerTrackFilter::GetSequencer() const
+{
+	return FilterInterface.GetSequencer();
+}
+
+UMovieSceneSequence* FSequencerTrackFilter::GetFocusedMovieSceneSequence() const
+{
+	return GetSequencer().GetFocusedMovieSceneSequence();
+}
+
+UMovieScene* FSequencerTrackFilter::GetFocusedGetMovieScene() const
+{
+	const UMovieSceneSequence* const FocusedMovieSceneSequence = GetFocusedMovieSceneSequence();
+	return IsValid(FocusedMovieSceneSequence) ? FocusedMovieSceneSequence->GetMovieScene() : nullptr;
+}
+
+bool FSequencerTrackFilter::SupportsLevelSequence(UMovieSceneSequence* const InSequence)
+{
+	const UClass* const LevelSequenceClass = ULevelSequence::StaticClass();
+	return IsValid(InSequence)
+		&& IsValid(LevelSequenceClass)
+		&& InSequence->GetClass()->IsChildOf(LevelSequenceClass);
+}
+
+bool FSequencerTrackFilter::SupportsUMGSequence(UMovieSceneSequence* const InSequence)
+{
+	static UClass* const WidgetAnimationClass = FindObject<UClass>(nullptr, TEXT("/Script/UMG.WidgetAnimation"), true);
+	return IsValid(InSequence)
+		&& IsValid(WidgetAnimationClass)
+		&& InSequence->GetClass()->IsChildOf(WidgetAnimationClass);
+}
+
+UMovieSceneTrack* FSequencerTrackFilter::ResolveMovieSceneTrackObject(FSequencerTrackFilterType InNode, FSequencerFilterData& FilterData)
+{
+	if (!InNode.IsValid())
+	{
+		return nullptr;
+	}
+
+	const TWeakViewModelPtr<IOutlinerExtension> WeakOutlinerNode = InNode.ImplicitCast();
+
+	// Use cache version if it exists, otherwise resolve below
+	if (FilterData.ResolvedTrackObjects.Contains(WeakOutlinerNode))
+	{
+		if (FilterData.ResolvedTrackObjects[WeakOutlinerNode].IsValid())
+		{
+			return FilterData.ResolvedTrackObjects[WeakOutlinerNode].Get();
+		}
+
+		FilterData.ResolvedTrackObjects.Remove(WeakOutlinerNode);
+	}
+
+	UMovieSceneTrack* TrackObject = nullptr;
+
+	if (const TViewModelPtr<ITrackExtension> AncestorTrackModel = InNode->FindAncestorOfType<ITrackExtension>(true))
+	{
+		TrackObject = AncestorTrackModel->GetTrack();
+	}
+
+	if (IsValid(TrackObject))
+	{
+		FilterData.ResolvedTrackObjects.Add(WeakOutlinerNode, TrackObject);
+	}
+
+	return TrackObject;
+}
+
+UObject* FSequencerTrackFilter::ResolveTrackBoundObject(ISequencer& InSequencer, FSequencerTrackFilterType InNode, FSequencerFilterData& FilterData)
+{
+	if (!InNode.IsValid())
+	{
+		return nullptr;
+	}
+
+	const TWeakViewModelPtr<IOutlinerExtension> WeakOutlinerNode = InNode.ImplicitCast();
+
+	// Use cache version if it exists, otherwise resolve below
+	if (FilterData.ResolvedBoundObjects.Contains(WeakOutlinerNode))
+	{
+		if (FilterData.ResolvedBoundObjects[WeakOutlinerNode].IsValid())
+		{
+			return FilterData.ResolvedBoundObjects[WeakOutlinerNode].Get();
+		}
+
+		FilterData.ResolvedBoundObjects.Remove(WeakOutlinerNode);
+	}
+
+	UObject* BoundObject = nullptr;
+
+	if (const TViewModelPtr<IObjectBindingExtension> ObjectBindingModel = InNode->FindAncestorOfType<IObjectBindingExtension>(true))
+	{
+		BoundObject = InSequencer.FindSpawnedObjectOrTemplate(ObjectBindingModel->GetObjectGuid());
+	}
+
+	if (IsValid(BoundObject))
+	{
+		FilterData.ResolvedBoundObjects.Add(WeakOutlinerNode, BoundObject);
+	}
+
+	return BoundObject;
+}
+
+UObject* FSequencerTrackFilter::ResolveTrackBoundObject(const TViewModelPtr<FViewModel> InNode, FSequencerFilterData& FilterData) const
+{
+	return ResolveTrackBoundObject(GetSequencer(), InNode, FilterData);
+}
+
+FText FSequencerTrackFilter::BuildTooltipTextForCommand(const FText& InBaseText, const TSharedPtr<FUICommandInfo>& InCommand)
+{
+	const TSharedRef<const FInputChord> FirstValidChord = InCommand->GetFirstValidChord();
+	if (FirstValidChord->IsValidChord())
+	{
+		return FText::Format(NSLOCTEXT("Sequencer", "TrackFilterTooltipText", "{0} ({1})"), InBaseText, FirstValidChord->GetInputText());
+	}
+	return InBaseText;
+}
+
+bool FSequencerTrackFilter::CanToggleFilter() const
+{
+	const FString FilterName = GetDisplayName().ToString();
+	return FilterInterface.IsFilterActiveByDisplayName(FilterName);
+}
+
+void FSequencerTrackFilter::ToggleFilter()
+{
+	const FString FilterName = GetDisplayName().ToString();
+	const bool bNewState = !FilterInterface.IsFilterActiveByDisplayName(FilterName);
+	FilterInterface.SetFilterActiveByDisplayName(FilterName, bNewState);
+}
+
+void FSequencerTrackFilter::MapToggleAction(const TSharedPtr<FUICommandInfo>& InCommand)
+{
+	FilterInterface.GetCommandList()->MapAction(
+		InCommand,
+		FExecuteAction::CreateSP(this, &FSequencerTrackFilter::ToggleFilter),
+		FCanExecuteAction(),
+		FIsActionChecked::CreateSP(this, &FSequencerTrackFilter::CanToggleFilter));
+}
diff --git a/Engine/Source/Editor/Sequencer/Private/Filters/Filters/SequencerTrackFilter_Keyed.cpp b/Engine/Source/Editor/Sequencer/Private/Filters/Filters/SequencerTrackFilter_Keyed.cpp
new file mode 100644
--- /dev/null
+++ b/Engine/Source/Editor/Sequencer/Private/Filters/Filters/SequencerTrackFilter_Keyed.cpp
@@ -0,0 +1,98 @@
+﻿// Copyright Epic Games, Inc. All Rights Reserved.
+
+#include "SequencerTrackFilter_Keyed.h"
+#include "Filters/SequencerFilterBar.h"
+#include "Filters/SequencerTrackFilterCommands.h"
+#include "Framework/Commands/Commands.h"
+#include "MVVM/Extensions/IOutlinerExtension.h"
+#include "MVVM/ViewModelPtr.h"
+#include "MVVM/ViewModels/CategoryModel.h"
+#include "Sequencer.h"
+
+using namespace UE::Sequencer;
+
+#define LOCTEXT_NAMESPACE "SequencerTrackFilter_Keyed"
+
+FSequencerTrackFilter_Keyed::FSequencerTrackFilter_Keyed(ISequencerTrackFilters& InFilterInterface, TSharedPtr<FFilterCategory> InCategory)
+	: FSequencerTrackFilter(InFilterInterface, MoveTemp(InCategory))
+{
+}
+
+bool FSequencerTrackFilter_Keyed::ShouldUpdateOnTrackValueChanged() const
+{
+	return true;
+}
+
+FText FSequencerTrackFilter_Keyed::GetDefaultToolTipText() const
+{
+	return LOCTEXT("SequencerTrackFilter_KeyedTip", "Show only Keyed tracks"); 
+}
+
+TSharedPtr<FUICommandInfo> FSequencerTrackFilter_Keyed::GetToggleCommand() const
+{
+	return FSequencerTrackFilterCommands::Get().ToggleFilter_Keyed;
+}
+
+bool FSequencerTrackFilter_Keyed::SupportsSequence(UMovieSceneSequence* const InSequence) const
+{
+	return FSequencerTrackFilter::SupportsSequence(InSequence) || SupportsUMGSequence(InSequence);
+}
+
+bool FSequencerTrackFilter_Keyed::PassesFilter(FSequencerTrackFilterType InItem) const
+{
+	FSequencerFilterData& FilterData = FilterInterface.GetFilterData();
+
+	for (const TViewModelPtr<IOutlinerExtension>& ChildNode : InItem->GetDescendantsOfType<IOutlinerExtension>(true))
+	{
+		const UMovieSceneTrack* const TrackObject = ResolveMovieSceneTrackObject(InItem, FilterData);
+		if (DoesMovieSceneTrackHaveKeys(TrackObject))
+		{
+			return true;
+		}
+	}
+
+	return false;
+}
+
+FText FSequencerTrackFilter_Keyed::GetDisplayName() const
+{
+	return LOCTEXT("SequenceTrackFilter_Keyed", "Keyed");
+}
+
+FSlateIcon FSequencerTrackFilter_Keyed::GetIcon() const
+{
+	return FSlateIcon(FAppStyle::GetAppStyleSetName(), TEXT("Sequencer.IconKeyUser"));
+}
+
+FString FSequencerTrackFilter_Keyed::GetName() const
+{
+	return TEXT("Keyed");
+}
+
+bool FSequencerTrackFilter_Keyed::DoesMovieSceneTrackHaveKeys(const UMovieSceneTrack* const InTrackObject)
+{
+	if (!IsValid(InTrackObject))
+	{
+		return false;
+	}
+
+	for (const UMovieSceneSection* const Section : InTrackObject->GetAllSections())
+	{
+		for (const FMovieSceneChannelEntry& ChannelEntry : Section->GetChannelProxy().GetAllEntries())
+		{
+			const TConstArrayView<FMovieSceneChannel*> Channels = ChannelEntry.GetChannels();
+			for (int32 ChannelIndex = 0; ChannelIndex < Channels.Num(); ++ChannelIndex)
+			{
+				const FMovieSceneChannel* const Channel = Channels[ChannelIndex];
+				if (Channel && Channel->GetNumKeys() > 0)
+				{
+					return true;
+				}
+			}
+		}
+	}
+
+	return false;
+}
+
+#undef LOCTEXT_NAMESPACE
diff --git a/Engine/Source/Editor/Sequencer/Private/Filters/Filters/SequencerTrackFilter_Selected.cpp b/Engine/Source/Editor/Sequencer/Private/Filters/Filters/SequencerTrackFilter_Selected.cpp
new file mode 100644
--- /dev/null
+++ b/Engine/Source/Editor/Sequencer/Private/Filters/Filters/SequencerTrackFilter_Selected.cpp
@@ -0,0 +1,115 @@
+﻿// Copyright Epic Games, Inc. All Rights Reserved.
+
+#include "SequencerTrackFilter_Selected.h"
+#include "EditorModeManager.h"
+#include "Elements/Framework/TypedElementSelectionSet.h"
+#include "Filters/SequencerFilterBar.h"
+#include "Filters/SequencerTrackFilterCommands.h"
+#include "Selection.h"
+#include "Sequencer.h"
+
+#define LOCTEXT_NAMESPACE "SequencerTrackFilter_Selected"
+
+FSequencerTrackFilter_Selected::FSequencerTrackFilter_Selected(ISequencerTrackFilters& InFilterInterface, TSharedPtr<FFilterCategory> InCategory)
+	: FSequencerTrackFilter(InFilterInterface, MoveTemp(InCategory))
+{
+	if (const TSharedPtr<IToolkitHost> ToolkitHost = GetSequencer().GetToolkitHost())
+	{
+		ToolkitHost->GetEditorModeManager().GetSelectedObjects()->SelectionChangedEvent.AddRaw(this, &FSequencerTrackFilter_Selected::OnSelectionChanged);
+	}
+}
+
+FSequencerTrackFilter_Selected::~FSequencerTrackFilter_Selected()
+{
+	if (const TSharedPtr<IToolkitHost> ToolkitHost = GetSequencer().GetToolkitHost())
+	{
+		ToolkitHost->GetEditorModeManager().GetSelectedObjects()->SelectionChangedEvent.RemoveAll(this);
+	}
+}
+
+bool FSequencerTrackFilter_Selected::ShouldUpdateOnTrackValueChanged() const
+{
+	return true;
+}
+
+FText FSequencerTrackFilter_Selected::GetDefaultToolTipText() const
+{
+	return LOCTEXT("SequencerTrackFilter_SelectedToolTip", "Show only track selected in the viewport");
+}
+
+TSharedPtr<FUICommandInfo> FSequencerTrackFilter_Selected::GetToggleCommand() const
+{
+	return FSequencerTrackFilterCommands::Get().ToggleFilter_Selected;
+}
+
+FText FSequencerTrackFilter_Selected::GetDisplayName() const
+{
+	return LOCTEXT("SequencerTrackFilter_Selected", "Selected");
+}
+
+FSlateIcon FSequencerTrackFilter_Selected::GetIcon() const
+{
+	return FSlateIcon(FAppStyle::GetAppStyleSetName(), TEXT("Icons.SelectInViewport"));
+}
+
+FString FSequencerTrackFilter_Selected::GetName() const
+{
+	return TEXT("Selected");
+}
+
+bool FSequencerTrackFilter_Selected::PassesFilter(FSequencerTrackFilterType InItem) const
+{
+	const FLevelEditorModule& LevelEditorModule = FModuleManager::GetModuleChecked<FLevelEditorModule>(TEXT("LevelEditor"));
+	const TSharedPtr<ILevelEditor> LevelEditor = LevelEditorModule.GetFirstLevelEditor();
+	if (!LevelEditor.IsValid())
+	{
+		return false;
+	}
+
+	const UTypedElementSelectionSet* const Selection = LevelEditor->GetElementSelectionSet();
+	if (!IsValid(Selection))
+	{
+		return false;
+	}
+
+	FSequencerFilterData& FilterData = FilterInterface.GetFilterData();
+
+	const UObject* const TrackObject = ResolveTrackBoundObject(InItem, FilterData);
+	if (!IsValid(TrackObject))
+	{
+		return false;
+	}
+
+	const TArray<UObject*> SelectedObjects = Selection->GetSelectedObjects<UObject>();
+
+	if (SelectedObjects.Contains(TrackObject))
+	{
+		return true;
+	}
+
+	USceneComponent* const Component = TrackObject->GetTypedOuter<USceneComponent>();
+	if (IsValid(Component) && SelectedObjects.Contains(Component))
+	{
+		return true;
+	}
+
+	AActor* const Actor = TrackObject->GetTypedOuter<AActor>();
+	if (IsValid(Actor) && SelectedObjects.Contains(Actor))
+	{
+		return true;
+	}
+
+	return false;
+}
+
+void FSequencerTrackFilter_Selected::OnSelectionChanged(UObject* const InObject)
+{
+	FilterInterface.RequestFilterUpdate();
+}
+
+void FSequencerTrackFilter_Selected::ToggleShowOnlySelectedTracks()
+{
+	SetActive(!IsActive());
+}
+
+#undef LOCTEXT_NAMESPACE
diff --git a/Engine/Source/Editor/Sequencer/Private/Filters/Filters/SequencerTrackFilters.h b/Engine/Source/Editor/Sequencer/Private/Filters/Filters/SequencerTrackFilters.h
new file mode 100644
--- /dev/null
+++ b/Engine/Source/Editor/Sequencer/Private/Filters/Filters/SequencerTrackFilters.h
@@ -0,0 +1,570 @@
+// Copyright Epic Games, Inc. All Rights Reserved.
+
+#pragma once
+
+#include "Camera/CameraComponent.h"
+#include "Components/LightComponentBase.h"
+#include "Components/SkeletalMeshComponent.h"
+#include "Filters/SequencerTrackFilterBase.h"
+#include "Filters/SequencerTrackFilterCommands.h"
+#include "Misc/IFilter.h"
+#include "MovieSceneFolder.h"
+#include "MVVM/ViewModels/FolderModel.h"
+#include "Particles/ParticleSystem.h"
+#include "Styling/AppStyle.h"
+#include "Styling/SlateIconFinder.h"
+#include "Tracks/MovieSceneAudioTrack.h"
+#include "Tracks/MovieSceneCameraCutTrack.h"
+#include "Tracks/MovieSceneCinematicShotTrack.h"
+#include "Tracks/MovieSceneDataLayerTrack.h"
+#include "Tracks/MovieSceneEventTrack.h"
+#include "Tracks/MovieSceneFadeTrack.h"
+#include "Tracks/MovieSceneLevelVisibilityTrack.h"
+#include "Tracks/MovieSceneParticleTrack.h"
+#include "Tracks/MovieSceneSlomoTrack.h"
+#include "Tracks/MovieSceneSubTrack.h"
+#include "Tracks/MovieSceneTimeWarpTrack.h"
+
+#define LOCTEXT_NAMESPACE "SequencerTrackFilters"
+
+class FSequencerTrackFilter_Audio : public FSequencerTrackFilter_ClassType<UMovieSceneAudioTrack>
+{
+public:
+	FSequencerTrackFilter_Audio(ISequencerTrackFilters& InFilterInterface, TSharedPtr<FFilterCategory> InCategory = nullptr)
+		: FSequencerTrackFilter_ClassType<UMovieSceneAudioTrack>(InFilterInterface, InCategory)
+	{}
+
+	//~ Begin IFilter
+	virtual FString GetName() const override { return TEXT("Audio"); }
+	//~ End IFilter
+
+	//~ Begin FFilterBase
+	virtual FText GetDisplayName() const override { return LOCTEXT("SequencerTrackFilter_Audio", "Audio"); }
+	virtual FSlateIcon GetIcon() const override { return FSlateIcon(FAppStyle::GetAppStyleSetName(), TEXT("Sequencer.Tracks.Audio")); }
+	//~ End FFilterBase
+
+	//~ Begin FSequencerTrackFilter
+
+	virtual FText GetDefaultToolTipText() const override
+	{
+		return LOCTEXT("SequencerTrackFilter_AudioToolTip", "Show only Audio tracks");
+	}
+
+	virtual TSharedPtr<FUICommandInfo> GetToggleCommand() const override
+	{
+		return FSequencerTrackFilterCommands::Get().ToggleFilter_Audio;
+	}
+
+	virtual bool SupportsSequence(UMovieSceneSequence* const InSequence) const override
+	{
+		return IsSequenceTrackSupported<UMovieSceneAudioTrack>(InSequence)
+			|| SupportsLevelSequence(InSequence)
+			|| SupportsUMGSequence(InSequence);
+	}
+
+	//~ End FSequencerTrackFilter
+};
+
+//////////////////////////////////////////////////////////////////////////
+//
+
+class FSequencerTrackFilter_Event : public FSequencerTrackFilter_ClassType<UMovieSceneEventTrack>
+{
+public:
+	FSequencerTrackFilter_Event(ISequencerTrackFilters& InFilterInterface, TSharedPtr<FFilterCategory> InCategory = nullptr)
+		: FSequencerTrackFilter_ClassType<UMovieSceneEventTrack>(InFilterInterface, InCategory)
+	{}
+
+	//~ Begin IFilter
+	virtual FString GetName() const override { return TEXT("Event"); }
+	//~ End IFilter
+
+	//~ Begin FFilterBase
+	virtual FText GetDisplayName() const override { return LOCTEXT("SequencerTrackFilter_Event", "Event"); }
+	virtual FSlateIcon GetIcon() const override { return FSlateIcon(FAppStyle::GetAppStyleSetName(), TEXT("Sequencer.Tracks.Event")); }
+	//~ End FFilterBase
+
+	//~ Begin FSequencerTrackFilter
+
+	virtual FText GetDefaultToolTipText() const override
+	{
+		return LOCTEXT("SequencerTrackFilter_EventToolTip", "Show only Event tracks");
+	}
+
+	virtual TSharedPtr<FUICommandInfo> GetToggleCommand() const override
+	{
+		return FSequencerTrackFilterCommands::Get().ToggleFilter_Event;
+	}
+
+	virtual bool SupportsSequence(UMovieSceneSequence* const InSequence) const override
+	{
+		return FSequencerTrackFilter::SupportsSequence(InSequence) || SupportsUMGSequence(InSequence);
+	}
+
+	//~ End FSequencerTrackFilter
+};
+
+//////////////////////////////////////////////////////////////////////////
+//
+
+class FSequencerTrackFilter_LevelVisibility : public FSequencerTrackFilter_ClassType<UMovieSceneLevelVisibilityTrack>
+{
+public:
+	FSequencerTrackFilter_LevelVisibility(ISequencerTrackFilters& InFilterInterface, TSharedPtr<FFilterCategory> InCategory = nullptr)
+		: FSequencerTrackFilter_ClassType<UMovieSceneLevelVisibilityTrack>(InFilterInterface, InCategory)
+	{}
+
+	//~ Begin IFilter
+	virtual FString GetName() const override { return TEXT("LevelVisibility"); }
+	//~ End IFilter
+
+	//~ Begin FFilterBase
+	virtual FText GetDisplayName() const override { return LOCTEXT("SequencerTrackFilter_LevelVisibility", "Level Visibility"); }
+	virtual FSlateIcon GetIcon() const override { return FSlateIcon(FAppStyle::GetAppStyleSetName(), TEXT("Sequencer.Tracks.LevelVisibility")); }
+	//~ End FFilterBase
+
+	//~ Begin FSequencerTrackFilter
+
+	virtual FText GetDefaultToolTipText() const override
+	{
+		return LOCTEXT("SequencerTrackFilter_LevelVisibilityToolTip", "Show only Level Visibility tracks");
+	}
+
+	virtual TSharedPtr<FUICommandInfo> GetToggleCommand() const override
+	{
+		return FSequencerTrackFilterCommands::Get().ToggleFilter_LevelVisibility;
+	}
+
+	//~ End FSequencerTrackFilter
+};
+
+//////////////////////////////////////////////////////////////////////////
+//
+
+class FSequencerTrackFilter_Particle : public FSequencerTrackFilter_ClassType<UMovieSceneParticleTrack>
+{
+public:
+	FSequencerTrackFilter_Particle(ISequencerTrackFilters& InFilterInterface, TSharedPtr<FFilterCategory> InCategory = nullptr)
+		: FSequencerTrackFilter_ClassType<UMovieSceneParticleTrack>(InFilterInterface, InCategory)
+	{}
+
+	//~ Begin IFilter
+	virtual FString GetName() const override { return TEXT("ParticleSystem"); }
+	//~ End IFilter
+
+	//~ Begin FFilterBase
+	virtual FText GetDisplayName() const override { return LOCTEXT("SequencerTrackFilter_Particle", "Particle System"); }
+	virtual FSlateIcon GetIcon() const override { return FSlateIconFinder::FindIconForClass(UParticleSystem::StaticClass()); }
+	//~ End FFilterBase
+
+	//~ Begin FSequencerTrackFilter
+
+	virtual FText GetDefaultToolTipText() const override
+	{
+		return LOCTEXT("SequencerTrackFilter_ParticleToolTip", "Show only Particle System tracks");
+	}
+
+	virtual TSharedPtr<FUICommandInfo> GetToggleCommand() const override
+	{
+		return FSequencerTrackFilterCommands::Get().ToggleFilter_Particle;
+	}
+
+	//~ End FSequencerTrackFilter
+};
+
+//////////////////////////////////////////////////////////////////////////
+//
+
+class FSequencerTrackFilter_CinematicShot : public FSequencerTrackFilter_ClassType<UMovieSceneCinematicShotTrack>
+{
+public:
+	FSequencerTrackFilter_CinematicShot(ISequencerTrackFilters& InFilterInterface, TSharedPtr<FFilterCategory> InCategory = nullptr)
+		: FSequencerTrackFilter_ClassType<UMovieSceneCinematicShotTrack>(InFilterInterface, InCategory)
+	{}
+
+	//~ Begin IFilter
+	virtual FString GetName() const override { return TEXT("CinematicShot"); }
+	//~ End IFilter
+
+	//~ Begin FFilterBase
+	virtual FText GetDisplayName() const override { return LOCTEXT("SequencerTrackFilter_CinematicShot", "Shot"); }
+	virtual FSlateIcon GetIcon() const override { return FSlateIcon(FAppStyle::GetAppStyleSetName(), TEXT("Sequencer.Tracks.CinematicShot")); }
+	//~ End FFilterBase
+
+	//~ Begin FSequencerTrackFilter
+
+	virtual FText GetDefaultToolTipText() const override
+	{
+		return LOCTEXT("SequencerTrackFilter_CinematicShotToolTip", "Show only Shot tracks");
+	}
+
+	virtual TSharedPtr<FUICommandInfo> GetToggleCommand() const override
+	{
+		return FSequencerTrackFilterCommands::Get().ToggleFilter_CinematicShot;
+	}
+
+	//~ End FSequencerTrackFilter
+};
+
+//////////////////////////////////////////////////////////////////////////
+//
+
+class FSequencerTrackFilter_Subsequence : public FSequencerTrackFilter_ClassType<UMovieSceneSubTrack>
+{
+public:
+	FSequencerTrackFilter_Subsequence(ISequencerTrackFilters& InFilterInterface, TSharedPtr<FFilterCategory> InCategory = nullptr)
+		: FSequencerTrackFilter_ClassType<UMovieSceneSubTrack>(InFilterInterface, InCategory)
+	{}
+
+	//~ Begin IFilter
+
+	virtual FString GetName() const override { return TEXT("SubSequence"); }
+
+	virtual bool PassesFilter(FSequencerTrackFilterType InItem) const override
+	{
+		FSequencerFilterData& FilterData = FilterInterface.GetFilterData();
+		const UMovieSceneTrack* const TrackObject = ResolveMovieSceneTrackObject(InItem, FilterData);
+		return IsValid(TrackObject)
+			&& TrackObject->IsA(UMovieSceneSubTrack::StaticClass())
+			&& !TrackObject->IsA(UMovieSceneCinematicShotTrack::StaticClass());
+	}
+
+	//~ End IFilter
+
+	//~ Begin FFilterBase
+	virtual FText GetDisplayName() const override { return LOCTEXT("SequencerTrackFilter_Subsequence", "Subsequence"); }
+	virtual FSlateIcon GetIcon() const override { return FSlateIcon(FAppStyle::GetAppStyleSetName(), TEXT("Sequencer.Tracks.Sub")); }
+	//~ End FFilterBase
+
+	//~ Begin FSequencerTrackFilter
+
+	virtual FText GetDefaultToolTipText() const override
+	{
+		return LOCTEXT("SequencerTrackFilter_SubsequenceToolTip", "Show only Subsequence tracks");
+	}
+
+	virtual TSharedPtr<FUICommandInfo> GetToggleCommand() const override
+	{
+		return FSequencerTrackFilterCommands::Get().ToggleFilter_Subsequence;
+	}
+
+	//~ End FSequencerTrackFilter
+};
+
+//////////////////////////////////////////////////////////////////////////
+//
+
+class FSequencerTrackFilter_SkeletalMesh : public FSequencerTrackFilter_ComponentType<USkeletalMeshComponent>
+{
+public:
+	FSequencerTrackFilter_SkeletalMesh(ISequencerTrackFilters& InFilterInterface, TSharedPtr<FFilterCategory> InCategory = nullptr)
+		: FSequencerTrackFilter_ComponentType<USkeletalMeshComponent>(InFilterInterface, InCategory)
+	{}
+
+	//~ Begin IFilter
+	virtual FString GetName() const override { return TEXT("SkeletalMesh"); }
+	//~ End IFilter
+
+	//~ Begin FFilterBase
+	virtual FText GetDisplayName() const override { return LOCTEXT("SequencerTrackFilter_SkeletalMesh", "Skeletal Mesh"); }
+	virtual FSlateIcon GetIcon() const override { return FSlateIconFinder::FindIconForClass(USkeletalMeshComponent::StaticClass()); }
+	//~ End FFilterBase
+
+	//~ Begin FSequencerTrackFilter
+
+	virtual FText GetDefaultToolTipText() const override
+	{
+		return LOCTEXT("SequencerTrackFilter_SkeletalMeshToolTip", "Show only Skeletal Mesh tracks");
+	}
+
+	virtual TSharedPtr<FUICommandInfo> GetToggleCommand() const override
+	{
+		return FSequencerTrackFilterCommands::Get().ToggleFilter_SkeletalMesh;
+	}
+
+	//~ End FSequencerTrackFilter
+};
+
+//////////////////////////////////////////////////////////////////////////
+//
+
+class FSequencerTrackFilter_Camera : public FSequencerTrackFilter_ComponentType<UCameraComponent>
+{
+public:
+	FSequencerTrackFilter_Camera(ISequencerTrackFilters& InFilterInterface, TSharedPtr<FFilterCategory> InCategory = nullptr)
+		: FSequencerTrackFilter_ComponentType<UCameraComponent>(InFilterInterface, InCategory)
+	{}
+
+	//~ Begin IFilter
+	virtual FString GetName() const override { return TEXT("Camera"); }
+	//~ End IFilter
+
+	//~ Begin FFilterBase
+	virtual FText GetDisplayName() const override { return LOCTEXT("SequencerTrackFilter_Camera", "Camera"); }
+	virtual FSlateIcon GetIcon() const override { return FSlateIconFinder::FindIconForClass(UCameraComponent::StaticClass()); }
+	//~ End FFilterBase
+
+	//~ Begin FSequencerTrackFilter
+
+	virtual FText GetDefaultToolTipText() const override
+	{
+		return LOCTEXT("SequencerTrackFilter_CameraToolTip", "Show only Camera tracks");
+	}
+
+	virtual TSharedPtr<FUICommandInfo> GetToggleCommand() const override
+	{
+		return FSequencerTrackFilterCommands::Get().ToggleFilter_Camera;
+	}
+
+	//~ End FSequencerTrackFilter
+};
+
+//////////////////////////////////////////////////////////////////////////
+//
+
+class FSequencerTrackFilter_Light : public FSequencerTrackFilter_ComponentType<ULightComponentBase>
+{
+public:
+	FSequencerTrackFilter_Light(ISequencerTrackFilters& InFilterInterface, TSharedPtr<FFilterCategory> InCategory = nullptr)
+		: FSequencerTrackFilter_ComponentType<ULightComponentBase>(InFilterInterface, InCategory)
+	{}
+
+	//~ Begin IFilter
+	virtual FString GetName() const override { return TEXT("Light"); }
+	//~ End IFilter
+
+	//~ Begin FFilterBase
+	virtual FText GetDisplayName() const override { return LOCTEXT("SequencerTrackFilter_Light", "Light"); }
+	virtual FSlateIcon GetIcon() const override { return FSlateIcon(FAppStyle::GetAppStyleSetName(), TEXT("ClassIcon.Light")); }
+	//~ End FFilterBase
+
+	//~ Begin FSequencerTrackFilter
+
+	virtual FText GetDefaultToolTipText() const override
+	{
+		return LOCTEXT("SequencerTrackFilter_LightToolTip", "Show only Light tracks");
+	}
+
+	virtual TSharedPtr<FUICommandInfo> GetToggleCommand() const override
+	{
+		return FSequencerTrackFilterCommands::Get().ToggleFilter_Light;
+	}
+
+	//~ End FSequencerTrackFilter
+};
+
+//////////////////////////////////////////////////////////////////////////
+//
+
+class FSequencerTrackFilter_CameraCut : public FSequencerTrackFilter_ClassType<UMovieSceneCameraCutTrack>
+{
+public:
+	FSequencerTrackFilter_CameraCut(ISequencerTrackFilters& InFilterInterface, TSharedPtr<FFilterCategory> InCategory = nullptr)
+		: FSequencerTrackFilter_ClassType<UMovieSceneCameraCutTrack>(InFilterInterface, InCategory)
+	{}
+
+	//~ Begin IFilter
+	virtual FString GetName() const override { return TEXT("CameraCut"); }
+	//~ End IFilter
+
+	//~ Begin FFilterBase
+	virtual FText GetDisplayName() const override { return LOCTEXT("SequencerTrackFilter_CameraCut", "Camera Cut"); }
+	virtual FSlateIcon GetIcon() const override { return FSlateIcon(FAppStyle::GetAppStyleSetName(), TEXT("Sequencer.Tracks.CameraCut")); }
+	//~ End FFilterBase
+
+	//~ Begin FSequencerTrackFilter
+
+	virtual FText GetDefaultToolTipText() const override
+	{
+		return LOCTEXT("SequencerTrackFilter_CameraCutToolTip", "Show only Camera Cut tracks");
+	}
+
+	virtual TSharedPtr<FUICommandInfo> GetToggleCommand() const override
+	{
+		return FSequencerTrackFilterCommands::Get().ToggleFilter_CameraCut;
+	}
+
+	//~ End FSequencerTrackFilter
+};
+
+//////////////////////////////////////////////////////////////////////////
+//
+
+class FSequencerTrackFilter_Fade : public FSequencerTrackFilter_ClassType<UMovieSceneFadeTrack>
+{
+public:
+	FSequencerTrackFilter_Fade(ISequencerTrackFilters& InFilterInterface, TSharedPtr<FFilterCategory> InCategory = nullptr)
+		: FSequencerTrackFilter_ClassType<UMovieSceneFadeTrack>(InFilterInterface, InCategory)
+	{}
+
+	//~ Begin IFilter
+	virtual FString GetName() const override { return TEXT("Fade"); }
+	//~ End IFilter
+
+	//~ Begin FFilterBase
+	virtual FText GetDisplayName() const override { return LOCTEXT("SequencerTrackFilter_Fade", "Fade"); }
+	virtual FSlateIcon GetIcon() const override { return FSlateIcon(FAppStyle::GetAppStyleSetName(), TEXT("Sequencer.Tracks.Fade")); }
+	//~ End FFilterBase
+
+	//~ Begin FSequencerTrackFilter
+
+	virtual FText GetDefaultToolTipText() const override
+	{
+		return LOCTEXT("SequencerTrackFilter_FadeToolTip", "Show only Fade tracks");
+	}
+
+	virtual TSharedPtr<FUICommandInfo> GetToggleCommand() const override
+	{
+		return FSequencerTrackFilterCommands::Get().ToggleFilter_Fade;
+	}
+
+	//~ End FSequencerTrackFilter
+};
+
+//////////////////////////////////////////////////////////////////////////
+//
+
+class FSequencerTrackFilter_DataLayer : public FSequencerTrackFilter_ClassType<UMovieSceneDataLayerTrack>
+{
+public:
+	FSequencerTrackFilter_DataLayer(ISequencerTrackFilters& InFilterInterface, TSharedPtr<FFilterCategory> InCategory = nullptr)
+		: FSequencerTrackFilter_ClassType<UMovieSceneDataLayerTrack>(InFilterInterface, InCategory)
+	{}
+
+	//~ Begin IFilter
+	virtual FString GetName() const override { return TEXT("DataLayer"); }
+	//~ End IFilter
+
+	//~ Begin FFilterBase
+	virtual FText GetDisplayName() const override { return LOCTEXT("SequencerTrackFilter_DataLayer", "Data Layer"); }
+	virtual FSlateIcon GetIcon() const override { return FSlateIcon(FAppStyle::GetAppStyleSetName(), TEXT("Sequencer.Tracks.DataLayer")); }
+	//~ End FFilterBase
+
+	//~ Begin FSequencerTrackFilter
+
+	virtual FText GetDefaultToolTipText() const override
+	{
+		return LOCTEXT("SequencerTrackFilter_DataLayerToolTip", "Show only Data Layer tracks");
+	}
+
+	virtual TSharedPtr<FUICommandInfo> GetToggleCommand() const override
+	{
+		return FSequencerTrackFilterCommands::Get().ToggleFilter_DataLayer;
+	}
+
+	//~ End FSequencerTrackFilter
+};
+
+//////////////////////////////////////////////////////////////////////////
+//
+
+class FSequencerTrackFilter_TimeDilation : public FSequencerTrackFilter_ClassType<UMovieSceneSlomoTrack>
+{
+public:
+	FSequencerTrackFilter_TimeDilation(ISequencerTrackFilters& InFilterInterface, TSharedPtr<FFilterCategory> InCategory = nullptr)
+		: FSequencerTrackFilter_ClassType<UMovieSceneSlomoTrack>(InFilterInterface, InCategory)
+	{}
+
+	//~ Begin IFilter
+	virtual FString GetName() const override { return TEXT("TimeDilation"); }
+	//~ End IFilter
+
+	//~ Begin FFilterBase
+	virtual FText GetDisplayName() const override { return LOCTEXT("SequencerTrackFilter_TimeDilation", "Time Dilation"); }
+	virtual FSlateIcon GetIcon() const override { return FSlateIcon(FAppStyle::GetAppStyleSetName(), TEXT("Sequencer.Tracks.Slomo")); }
+	//~ End FFilterBase
+
+	//~ Begin FSequencerTrackFilter
+
+	virtual FText GetDefaultToolTipText() const override
+	{
+		return LOCTEXT("SequencerTrackFilter_TimeDilationToolTip", "Show only Time Dilation tracks");
+	}
+
+	virtual TSharedPtr<FUICommandInfo> GetToggleCommand() const override
+	{
+		return FSequencerTrackFilterCommands::Get().ToggleFilter_TimeDilation;
+	}
+
+	//~ End FSequencerTrackFilter
+};
+
+//////////////////////////////////////////////////////////////////////////
+//
+
+class FSequencerTrackFilter_TimeWarp : public FSequencerTrackFilter_ClassType<UMovieSceneTimeWarpTrack>
+{
+public:
+	FSequencerTrackFilter_TimeWarp(ISequencerTrackFilters& InFilterInterface, TSharedPtr<FFilterCategory> InCategory = nullptr)
+		: FSequencerTrackFilter_ClassType<UMovieSceneTimeWarpTrack>(InFilterInterface, InCategory)
+	{}
+
+	//~ Begin IFilter
+	virtual FString GetName() const override { return TEXT("TimeWarp"); }
+	//~ End IFilter
+
+	//~ Begin FFilterBase
+	virtual FText GetDisplayName() const override { return LOCTEXT("SequencerTrackFilter_TimeWarp", "Time Warp"); }
+	virtual FSlateIcon GetIcon() const override { return FSlateIcon(FAppStyle::GetAppStyleSetName(), TEXT("Sequencer.Tracks.Slomo")); }
+	//~ End FFilterBase
+
+	//~ Begin FSequencerTrackFilter
+
+	virtual FText GetDefaultToolTipText() const override
+	{
+		return LOCTEXT("SequencerTrackFilter_TimeWarpToolTip", "Show only Time Warp tracks");
+	}
+
+	virtual TSharedPtr<FUICommandInfo> GetToggleCommand() const override
+	{
+		return FSequencerTrackFilterCommands::Get().ToggleFilter_TimeWarp;
+	}
+
+	//~ End FSequencerTrackFilter
+};
+
+//////////////////////////////////////////////////////////////////////////
+//
+
+class FSequencerTrackFilter_Folder : public FSequencerTrackFilter_ModelType<UE::Sequencer::FFolderModel>
+{
+public:
+	FSequencerTrackFilter_Folder(ISequencerTrackFilters& InFilterInterface, TSharedPtr<FFilterCategory> InCategory = nullptr)
+		: FSequencerTrackFilter_ModelType<UE::Sequencer::FFolderModel>(InFilterInterface, InCategory)
+	{}
+
+	//~ Begin IFilter
+	virtual FString GetName() const override { return TEXT("Folder"); }
+	//~ End IFilter
+
+	//~ Begin FFilterBase
+	virtual FText GetDisplayName() const override { return LOCTEXT("SequencerTrackFilter_Folder", "Folder"); }
+	virtual FSlateIcon GetIcon() const override { return FSlateIcon(FAppStyle::GetAppStyleSetName(), TEXT("ContentBrowser.AssetTreeFolderClosed")); }
+	//~ End FFilterBase
+
+	//~ Begin FSequencerTrackFilter
+
+	virtual FText GetDefaultToolTipText() const override
+	{
+		return LOCTEXT("SequencerTrackFilter_FolderToolTip", "Show only Folder tracks");
+	}
+
+	virtual TSharedPtr<FUICommandInfo> GetToggleCommand() const override
+	{
+		return FSequencerTrackFilterCommands::Get().ToggleFilter_Folder;
+	}
+
+	TSubclassOf<UMovieSceneTrack> GetTrackClass() const override
+	{
+		return UMovieSceneFolder::StaticClass();
+	}
+
+	virtual bool SupportsSequence(UMovieSceneSequence* const InSequence) const override
+	{
+		return true;
+	}
+
+	//~ End FSequencerTrackFilter
+};
+
+#undef LOCTEXT_NAMESPACE
diff --git a/Engine/Source/Editor/Sequencer/Private/Filters/Menus/SequencerFilterBarContext.h b/Engine/Source/Editor/Sequencer/Private/Filters/Menus/SequencerFilterBarContext.h
new file mode 100644
--- /dev/null
+++ b/Engine/Source/Editor/Sequencer/Private/Filters/Menus/SequencerFilterBarContext.h
@@ -0,0 +1,28 @@
+﻿// Copyright Epic Games, Inc. All Rights Reserved.
+
+#pragma once
+
+#include "Templates/SharedPointer.h"
+#include "SequencerFilterBarContext.generated.h"
+
+class FSequencerFilterBar;
+
+UCLASS()
+class USequencerFilterBarContext : public UObject
+{
+	GENERATED_BODY()
+
+public:
+	void Init(const TWeakPtr<FSequencerFilterBar>& InFilterBarWeak)
+	{
+		FilterBarWeak = InFilterBarWeak;
+	}
+
+	TSharedPtr<FSequencerFilterBar> GetFilterBar() const
+	{
+		return FilterBarWeak.IsValid() ? FilterBarWeak.Pin() : nullptr;
+	}
+
+protected:
+	TWeakPtr<FSequencerFilterBar> FilterBarWeak;
+};
diff --git a/Engine/Source/Editor/Sequencer/Private/Filters/Menus/SequencerFilterBarContextMenu.h b/Engine/Source/Editor/Sequencer/Private/Filters/Menus/SequencerFilterBarContextMenu.h
new file mode 100644
--- /dev/null
+++ b/Engine/Source/Editor/Sequencer/Private/Filters/Menus/SequencerFilterBarContextMenu.h
@@ -0,0 +1,28 @@
+﻿// Copyright Epic Games, Inc. All Rights Reserved.
+
+#pragma once
+
+#include "Templates/SharedPointer.h"
+#include "UObject/WeakObjectPtr.h"
+
+class FSequencerFilterBar;
+class SWidget;
+class USequencerFilterBarContext;
+class UToolMenu;
+
+class FSequencerFilterBarContextMenu
+{
+public:
+	TSharedRef<SWidget> CreateMenu(const TSharedRef<FSequencerFilterBar>& InFilterBar);
+
+protected:
+	void PopulateMenu(UToolMenu* const InMenu);
+
+	void PopulateOptionsSection(UToolMenu& InMenu);
+	void PopulateFilterBulkOptionsSection(UToolMenu& InMenu);
+
+	void OnActivateAllFilters(const bool bInActivate);
+	void OnResetFilters();
+
+	TWeakObjectPtr<USequencerFilterBarContext> CurrentContext;
+};
diff --git a/Engine/Source/Editor/Sequencer/Private/Filters/Menus/SequencerFilterMenuContext.h b/Engine/Source/Editor/Sequencer/Private/Filters/Menus/SequencerFilterMenuContext.h
new file mode 100644
--- /dev/null
+++ b/Engine/Source/Editor/Sequencer/Private/Filters/Menus/SequencerFilterMenuContext.h
@@ -0,0 +1,33 @@
+﻿// Copyright Epic Games, Inc. All Rights Reserved.
+
+#pragma once
+
+#include "Filters/SequencerTrackFilterBase.h"
+#include "Filters/Widgets/SSequencerFilter.h"
+#include "Templates/SharedPointer.h"
+#include "SequencerFilterMenuContext.generated.h"
+
+UCLASS()
+class USequencerFilterMenuContext : public UObject
+{
+	GENERATED_BODY()
+
+public:
+	void Init(const TWeakPtr<SSequencerFilter>& InWeakFilterWidget)
+	{
+		WeakFilterWidget = InWeakFilterWidget;
+	}
+
+	TSharedPtr<SSequencerFilter> GetFilterWidget() const
+	{
+		return WeakFilterWidget.Pin();
+	}
+
+	TSharedPtr<FSequencerTrackFilter> GetFilter() const
+	{
+		return WeakFilterWidget.IsValid() ? WeakFilterWidget.Pin()->GetFilter() : nullptr;
+	}
+
+protected:
+	TWeakPtr<SSequencerFilter> WeakFilterWidget;
+};
diff --git a/Engine/Source/Editor/Sequencer/Private/Filters/Menus/SequencerTrackFilterContextMenu.cpp b/Engine/Source/Editor/Sequencer/Private/Filters/Menus/SequencerTrackFilterContextMenu.cpp
new file mode 100644
--- /dev/null
+++ b/Engine/Source/Editor/Sequencer/Private/Filters/Menus/SequencerTrackFilterContextMenu.cpp
@@ -0,0 +1,265 @@
+﻿// Copyright Epic Games, Inc. All Rights Reserved.
+
+#include "SequencerTrackFilterContextMenu.h"
+#include "Filters/Filters/SequencerTrackFilter_CustomText.h"
+#include "Filters/Menus/SequencerFilterMenuContext.h"
+#include "Filters/SequencerFilterBar.h"
+#include "Filters/SequencerTrackFilterCommands.h"
+#include "Filters/Widgets/SSequencerCustomTextFilterDialog.h"
+#include "Filters/Widgets/SSequencerFilter.h"
+#include "Sequencer.h"
+#include "SequencerSettings.h"
+#include "ToolMenu.h"
+#include "ToolMenus.h"
+
+#define LOCTEXT_NAMESPACE "SequencerTrackFilterContextMenu"
+
+TSharedRef<SWidget> FSequencerTrackFilterContextMenu::CreateMenuWidget(const TSharedRef<SSequencerFilter>& InFilterWidget)
+{
+	UToolMenus* const ToolMenus = UToolMenus::Get();
+	check(IsValid(ToolMenus));
+
+	const FName FilterMenuName = TEXT("Sequencer.TrackFilterContextMenu");
+	if (!ToolMenus->IsMenuRegistered(FilterMenuName))
+	{
+		UToolMenu* const Menu = ToolMenus->RegisterMenu(FilterMenuName);
+		Menu->AddDynamicSection(TEXT("PopulateMenu"), FNewToolMenuDelegate::CreateRaw(this, &FSequencerTrackFilterContextMenu::PopulateMenu));
+	}
+
+	const TSharedPtr<FSequencerTrackFilter> Filter = InFilterWidget->GetFilter();
+	const TSharedPtr<FUICommandList> CommandList = Filter.IsValid() ? Filter->GetFilterInterface().GetCommandList() : nullptr;
+
+	USequencerFilterMenuContext* const ContextObject = NewObject<USequencerFilterMenuContext>();
+	ContextObject->Init(InFilterWidget);
+
+	const FToolMenuContext MenuContext(CommandList, nullptr, ContextObject);
+
+	return UToolMenus::Get()->GenerateWidget(FilterMenuName, MenuContext);
+}
+
+void FSequencerTrackFilterContextMenu::PopulateMenu(UToolMenu* const InMenu)
+{
+	if (!IsValid(InMenu))
+	{
+		return;
+	}
+
+	UToolMenu& MenuRef = *InMenu;
+
+	CurrentContext = MenuRef.FindContext<USequencerFilterMenuContext>();
+	if (!CurrentContext.IsValid())
+	{
+		return;
+	}
+
+	PopulateFilterOptionsSection(MenuRef);
+	PopulateCustomFilterOptionsSection(MenuRef);
+	PopulateBulkOptionsSection(MenuRef);
+}
+
+void FSequencerTrackFilterContextMenu::PopulateFilterOptionsSection(UToolMenu& InMenu)
+{
+	const TSharedPtr<FSequencerTrackFilter> Filter = CurrentContext->GetFilter();
+	if (!Filter.IsValid())
+	{
+		return;
+	}
+
+	const FText FilterName = GetFilterDisplayName();
+
+	FToolMenuSection& Section = InMenu.FindOrAddSection(TEXT("FilterOptions")
+		, LOCTEXT("FilterOptionsContextHeading", "Filter Options"));
+
+	Section.AddMenuEntry(TEXT("ActivateOnlyThisFilter"),
+		FText::Format(LOCTEXT("ActivateOnlyThisFilter", "Activate Only: {0}"), FilterName),
+		LOCTEXT("ActivateOnlyThisFilterTooltip", "Activate only this filter from the list."),
+		FSlateIcon(FAppStyle::GetAppStyleSetName(), TEXT("Icons.Plus")),
+		FUIAction(FExecuteAction::CreateRaw(this, &FSequencerTrackFilterContextMenu::OnActivateWithFilterException)));
+
+	Section.AddMenuEntry(TEXT("DisableFilter"),
+		FText::Format(LOCTEXT("DisableFilter", "Disable: {0}"), FilterName),
+		LOCTEXT("DisableFilterTooltip", "Disable this filter from the list. It can be added again in the filters menu."),
+		FSlateIcon(FAppStyle::GetAppStyleSetName(), TEXT("Icons.Minus")),
+		FUIAction(FExecuteAction::CreateRaw(this, &FSequencerTrackFilterContextMenu::OnDisableFilter)));
+}
+
+void FSequencerTrackFilterContextMenu::PopulateCustomFilterOptionsSection(UToolMenu& InMenu)
+{
+	const TSharedPtr<FSequencerTrackFilter_CustomText> CustomTextFilter
+		= StaticCastSharedPtr<FSequencerTrackFilter_CustomText>(CurrentContext->GetFilter());
+	if (!CustomTextFilter.IsValid() || !CustomTextFilter->IsCustomTextFilter())
+	{
+		return;
+	}
+
+	FToolMenuSection& Section = InMenu.FindOrAddSection(TEXT("CustomFilterOptions")
+		, LOCTEXT("CustomFilterOptionsContextHeading", "Custom Filter Options"));
+
+	Section.AddMenuEntry(TEXT("EditCustomTextFilter"),
+		FText::Format(LOCTEXT("EditCustomTextFilter", "Edit: {0}"), CustomTextFilter->GetDisplayName()),
+		LOCTEXT("EditCustomTextFilterTooltip", "Edit this custom text filter saved to config."),
+		FSlateIcon(FAppStyle::GetAppStyleSetName(), TEXT("Icons.Edit")),
+		FUIAction(FExecuteAction::CreateRaw(this, &FSequencerTrackFilterContextMenu::OnEditFilter)));
+
+	Section.AddMenuEntry(TEXT("DeleteCustomTextFilter"),
+		FText::Format(LOCTEXT("DeleteCustomTextFilter", "Delete: {0}"), CustomTextFilter->GetDisplayName()),
+		LOCTEXT("DeleteCustomTextFilterTooltip", "Delete this custom text filter from config.\n\nCAUTION: This cannot be undone!"),
+		FSlateIcon(FAppStyle::GetAppStyleSetName(), TEXT("Icons.Delete")),
+		FUIAction(FExecuteAction::CreateRaw(this, &FSequencerTrackFilterContextMenu::OnDeleteFilter)));
+}
+
+void FSequencerTrackFilterContextMenu::PopulateBulkOptionsSection(UToolMenu& InMenu)
+{
+	const FSequencerTrackFilterCommands& TrackFilterCommands = FSequencerTrackFilterCommands::Get();
+
+	const FText FilterName = GetFilterDisplayName();
+
+	FToolMenuSection& Section = InMenu.FindOrAddSection(TEXT("FilterBulkOptions")
+		, LOCTEXT("BulkOptionsContextHeading", "Filter Bulk Options"));
+
+	Section.AddMenuEntry(TEXT("ActivateAllFilters"),
+		LOCTEXT("ActivateAllFilters", "Activate All Filters"),
+		LOCTEXT("ActivateAllFiltersTooltip", "Activates all enabled filters."),
+		FSlateIcon(FAppStyle::GetAppStyleSetName(), TEXT("Icons.Plus")),
+		FUIAction(FExecuteAction::CreateRaw(this, &FSequencerTrackFilterContextMenu::OnActivateAllFilters, true)));
+
+	Section.AddMenuEntry(TEXT("DeactivateAllFilters"),
+		LOCTEXT("DeactivateAllFilters", "Deactivate All Filters"),
+		LOCTEXT("DeactivateAllFiltersTooltip", "Deactivates all enabled filters."),
+		FSlateIcon(FAppStyle::GetAppStyleSetName(), TEXT("Icons.Minus")),
+		FUIAction(FExecuteAction::CreateRaw(this, &FSequencerTrackFilterContextMenu::OnActivateAllFilters, false)));
+
+	Section.AddSeparator(NAME_None);
+
+	Section.AddMenuEntry(TrackFilterCommands.ResetFilters
+		, TrackFilterCommands.ResetFilters->GetLabel()
+		, TrackFilterCommands.ResetFilters->GetDescription()
+		, FSlateIcon(FAppStyle::Get().GetStyleSetName(), TEXT("PropertyWindow.DiffersFromDefault")));
+
+	Section.AddSeparator(NAME_None);
+
+	Section.AddMenuEntry(TrackFilterCommands.ToggleMuteFilters
+		, TrackFilterCommands.ToggleMuteFilters->GetLabel()
+		, TrackFilterCommands.ToggleMuteFilters->GetDescription()
+		, FSlateIcon(FAppStyle::Get().GetStyleSetName(), TEXT("Denied")));
+
+	Section.AddMenuEntry(TrackFilterCommands.DisableAllFilters
+		, TrackFilterCommands.DisableAllFilters->GetLabel()
+		, TrackFilterCommands.DisableAllFilters->GetDescription()
+		, FSlateIcon(FAppStyle::Get().GetStyleSetName(), TEXT("Icons.Minus")));
+}
+
+FText FSequencerTrackFilterContextMenu::GetFilterDisplayName() const
+{
+	const TSharedPtr<FSequencerTrackFilter> Filter = CurrentContext->GetFilter();
+	if (!Filter.IsValid())
+	{
+		return FText::GetEmpty();
+	}
+
+	return Filter.IsValid() ? Filter->GetDisplayName() : FText::GetEmpty();
+}
+
+void FSequencerTrackFilterContextMenu::OnDisableFilter()
+{
+	const TSharedPtr<FSequencerTrackFilter> Filter = CurrentContext->GetFilter();
+	if (!Filter.IsValid())
+	{
+		return;
+	}
+
+	const FString FilterName = Filter->GetDisplayName().ToString();
+	Filter->GetFilterInterface().SetFilterEnabledByDisplayName(FilterName, false);
+}
+
+void FSequencerTrackFilterContextMenu::OnResetFilters()
+{
+	const TSharedPtr<FSequencerTrackFilter> Filter = CurrentContext->GetFilter();
+	if (!Filter.IsValid())
+	{
+		return;
+	}
+
+	Filter->GetFilterInterface().EnableAllFilters(false, {});
+}
+
+void FSequencerTrackFilterContextMenu::OnActivateWithFilterException()
+{
+	const TSharedPtr<FSequencerTrackFilter> Filter = CurrentContext->GetFilter();
+	if (!Filter.IsValid())
+	{
+		return;
+	}
+
+	ISequencerTrackFilters& FilterInterface = Filter->GetFilterInterface();
+	const FString FilterName = Filter->GetDisplayName().ToString();
+
+	FilterInterface.ActivateAllEnabledFilters(false, { FilterName });
+	FilterInterface.SetFilterActiveByDisplayName(FilterName, true);
+}
+
+void FSequencerTrackFilterContextMenu::OnActivateAllFilters(const bool bInActivate)
+{
+	const TSharedPtr<FSequencerTrackFilter> Filter = CurrentContext->GetFilter();
+	if (!Filter.IsValid())
+	{
+		return;
+	}
+
+	Filter->GetFilterInterface().ActivateAllEnabledFilters(bInActivate, {});
+}
+
+void FSequencerTrackFilterContextMenu::OnEditFilter()
+{
+	const TSharedPtr<FSequencerTrackFilter> Filter = CurrentContext->GetFilter();
+	if (!Filter.IsValid() || !Filter->IsCustomTextFilter())
+	{
+		return;
+	}
+
+	const TSharedPtr<FSequencerTrackFilter_CustomText> CustomTextFilter = StaticCastSharedPtr<FSequencerTrackFilter_CustomText>(Filter);
+	if (!CustomTextFilter.IsValid())
+	{
+		return;
+	}
+
+	FSequencerFilterBar& FilterBar = (FSequencerFilterBar&)Filter->GetFilterInterface();
+
+	SSequencerCustomTextFilterDialog::CreateWindow_EditCustomTextFilter(FilterBar.AsShared(), CustomTextFilter);
+}
+
+void FSequencerTrackFilterContextMenu::OnDeleteFilter()
+{
+	const TSharedPtr<FSequencerTrackFilter> Filter = CurrentContext->GetFilter();
+	if (!Filter.IsValid() || !Filter->IsCustomTextFilter())
+	{
+		return;
+	}
+
+	const TSharedPtr<FSequencerTrackFilter_CustomText> CustomTextFilter = StaticCastSharedPtr<FSequencerTrackFilter_CustomText>(Filter);
+	if (!CustomTextFilter.IsValid())
+	{
+		return;
+	}
+
+	ISequencerTrackFilters& FilterInterface = Filter->GetFilterInterface();
+	const FString FilterName = Filter->GetDisplayName().ToString();
+
+	FilterInterface.SetFilterActiveByDisplayName(FilterName, false);
+	FilterInterface.RemoveCustomTextFilter(CustomTextFilter.ToSharedRef(), false);
+
+	USequencerSettings* const SequencerSettings = FilterInterface.GetSequencer().GetSequencerSettings();
+	if (IsValid(SequencerSettings))
+	{
+		FSequencerFilterBarConfig& Config = SequencerSettings->FindOrAddTrackFilterBar(FilterInterface.GetIdentifier(), false);
+
+		if (!Config.RemoveCustomTextFilter(FilterName))
+		{
+			return;
+		}
+
+		SequencerSettings->SaveConfig();
+	}
+}
+
+#undef LOCTEXT_NAMESPACE
diff --git a/Engine/Source/Editor/Sequencer/Private/Filters/Menus/SequencerTrackFilterContextMenu.h b/Engine/Source/Editor/Sequencer/Private/Filters/Menus/SequencerTrackFilterContextMenu.h
new file mode 100644
--- /dev/null
+++ b/Engine/Source/Editor/Sequencer/Private/Filters/Menus/SequencerTrackFilterContextMenu.h
@@ -0,0 +1,38 @@
+﻿// Copyright Epic Games, Inc. All Rights Reserved.
+
+#pragma once
+
+#include "Templates/SharedPointer.h"
+#include "UObject/WeakObjectPtr.h"
+
+class SSequencerFilter;
+class SWidget;
+class USequencerFilterMenuContext;
+class UToolMenu;
+
+class FSequencerTrackFilterContextMenu
+{
+public:
+	TSharedRef<SWidget> CreateMenuWidget(const TSharedRef<SSequencerFilter>& InFilterWidget);
+
+protected:
+	void PopulateMenu(UToolMenu* const InMenu);
+
+	void PopulateFilterOptionsSection(UToolMenu& InMenu);
+	void PopulateCustomFilterOptionsSection(UToolMenu& InMenu);
+	void PopulateBulkOptionsSection(UToolMenu& InMenu);
+
+	FText GetFilterDisplayName() const;
+
+	void OnDisableFilter();
+	void OnResetFilters();
+
+	void OnActivateWithFilterException();
+
+	void OnActivateAllFilters(const bool bInActivate);
+
+	void OnEditFilter();
+	void OnDeleteFilter();
+
+	TWeakObjectPtr<USequencerFilterMenuContext> CurrentContext;
+};
diff --git a/Engine/Source/Editor/Sequencer/Private/Filters/Menus/SequencerViewOptionsMenu.cpp b/Engine/Source/Editor/Sequencer/Private/Filters/Menus/SequencerViewOptionsMenu.cpp
new file mode 100644
--- /dev/null
+++ b/Engine/Source/Editor/Sequencer/Private/Filters/Menus/SequencerViewOptionsMenu.cpp
@@ -0,0 +1,264 @@
+﻿// Copyright Epic Games, Inc. All Rights Reserved.
+
+#include "SequencerViewOptionsMenu.h"
+#include "Filters/Menus/SequencerMenuContext.h"
+#include "Filters/SequencerTrackFilterCommands.h"
+#include "Sequencer.h"
+#include "SequencerCommands.h"
+#include "SSequencer.h"
+#include "ToolMenu.h"
+#include "ToolMenus.h"
+
+#define LOCTEXT_NAMESPACE "SequencerViewOptionsMenu"
+
+TSharedRef<SWidget> FSequencerViewOptionsMenu::CreateMenu(const TWeakPtr<FSequencer>& InSequencerWeak)
+{
+	if (!InSequencerWeak.IsValid())
+	{
+		return SNullWidget::NullWidget;
+	}
+
+	const FName FilterMenuName = TEXT("Sequencer.ViewOptionsMenu");
+	if (!UToolMenus::Get()->IsMenuRegistered(FilterMenuName))
+	{
+		UToolMenu* const Menu = UToolMenus::Get()->RegisterMenu(FilterMenuName);
+		Menu->AddDynamicSection(NAME_None, FNewToolMenuDelegate::CreateRaw(this, &FSequencerViewOptionsMenu::PopulateMenu));
+	}
+
+	USequencerMenuContext* const ContextObject = NewObject<USequencerMenuContext>();
+	ContextObject->Init(InSequencerWeak);
+
+	const FToolMenuContext MenuContext(InSequencerWeak.Pin()->GetFilterInterface()->GetCommandList(), nullptr, ContextObject);
+	return UToolMenus::Get()->GenerateWidget(FilterMenuName, MenuContext);
+}
+
+void FSequencerViewOptionsMenu::PopulateMenu(UToolMenu* const InMenu)
+{
+	if (!IsValid(InMenu))
+	{
+		return;
+	}
+
+	UToolMenu& Menu = *InMenu;
+
+	CurrentContext = InMenu->FindContext<USequencerMenuContext>();
+
+	PopulateFiltersSection(Menu);
+	PopulateSortAndOrganizeSection(Menu);
+	PopulateFilterOptionssSection(Menu);
+	PopulateLayoutSection(Menu);
+}
+
+void FSequencerViewOptionsMenu::PopulateFiltersSection(UToolMenu& InMenu)
+{
+	const FSequencerTrackFilterCommands& TrackFilterCommands = FSequencerTrackFilterCommands::Get();
+
+	FToolMenuSection& HiddenTracksSection = InMenu.FindOrAddSection(TEXT("HiddenTracks"), LOCTEXT("HiddenTracksHeading", "Hidden Tracks"));
+
+	HiddenTracksSection.AddMenuEntry(TrackFilterCommands.HideSelectedTracks);
+	HiddenTracksSection.AddMenuEntry(TrackFilterCommands.ClearHiddenTracks);
+
+	FToolMenuSection& IsolateTracksSection = InMenu.FindOrAddSection(TEXT("IsolatedTracks"), LOCTEXT("IsolatedTracksHeading", "Isolated Tracks"));
+
+	IsolateTracksSection.AddMenuEntry(TrackFilterCommands.IsolateSelectedTracks);
+	IsolateTracksSection.AddMenuEntry(TrackFilterCommands.ClearIsolatedTracks);
+
+	FToolMenuSection& ShowTracksSection = InMenu.FindOrAddSection(TEXT("ShowTracks"), LOCTEXT("ShowTracksHeading", "Show Tracks"));
+
+	ShowTracksSection.AddMenuEntry(TrackFilterCommands.ShowAllTracks);
+	ShowTracksSection.AddSeparator(NAME_None);
+	ShowTracksSection.AddMenuEntry(TrackFilterCommands.ShowLocationCategoryGroups);
+	ShowTracksSection.AddMenuEntry(TrackFilterCommands.ShowRotationCategoryGroups);
+	ShowTracksSection.AddMenuEntry(TrackFilterCommands.ShowScaleCategoryGroups);
+}
+
+void FSequencerViewOptionsMenu::PopulateSortAndOrganizeSection(UToolMenu& InMenu)
+{
+	const TSharedPtr<FSequencer> Sequencer = CurrentContext->GetSequencer();
+	if (!Sequencer.IsValid())
+	{
+		return;
+	}
+
+	const TSharedPtr<FUICommandList> SequencerBindings = Sequencer->GetCommandBindings();
+	const FSequencerCommands& SequencerCommands = FSequencerCommands::Get();
+
+	FToolMenuSection& Section = InMenu.FindOrAddSection(TEXT("OrganizeAndSort"), LOCTEXT("OrganizeAndSortHeader", "Organize and Sort"));
+
+	Section.AddMenuEntryWithCommandList(SequencerCommands.ToggleAutoExpandNodesOnSelection, SequencerBindings);
+	Section.AddMenuEntryWithCommandList(SequencerCommands.ToggleExpandCollapseNodes, SequencerBindings);
+	Section.AddMenuEntryWithCommandList(SequencerCommands.ToggleExpandCollapseNodesAndDescendants, SequencerBindings);
+	Section.AddMenuEntryWithCommandList(SequencerCommands.ExpandAllNodes, SequencerBindings);
+	Section.AddMenuEntryWithCommandList(SequencerCommands.CollapseAllNodes, SequencerBindings);
+	Section.AddMenuEntryWithCommandList(SequencerCommands.SortAllNodesAndDescendants, SequencerBindings);
+}
+
+void FSequencerViewOptionsMenu::PopulateFilterOptionssSection(UToolMenu& InMenu)
+{
+	FToolMenuSection& OptionsSection = InMenu.FindOrAddSection(TEXT("FilterOptions"), LOCTEXT("FilterOptionsHeading", "Filter Options"));
+
+	OptionsSection.AddMenuEntry(TEXT("FilterPinned"),
+		LOCTEXT("FilterPinned", "Filter Pinned"),
+		LOCTEXT("FilterPinnedToolTip", "Toggle inclusion of pinned items when filtering"),
+		FSlateIcon(),
+		FUIAction(
+			FExecuteAction::CreateRaw(this, &FSequencerViewOptionsMenu::ToggleIncludePinnedInFilter),
+			FCanExecuteAction(),
+			FIsActionChecked::CreateRaw(this, &FSequencerViewOptionsMenu::IsIncludePinnedInFilter)
+		),
+		EUserInterfaceActionType::ToggleButton);
+
+	OptionsSection.AddMenuEntry(TEXT("AutoExpandPassedFilterNodes"),
+		LOCTEXT("AutoExpandPassedFilterNodes", "Auto Expand Filtered Items"),
+		LOCTEXT("AutoExpandPassedFilterNodesToolTip", "Toggle expansion of items when a filter is passed"),
+		FSlateIcon(),
+		FUIAction(
+			FExecuteAction::CreateRaw(this, &FSequencerViewOptionsMenu::ToggleAutoExpandPassedFilterNodes),
+			FCanExecuteAction(),
+			FIsActionChecked::CreateRaw(this, &FSequencerViewOptionsMenu::IsAutoExpandPassedFilterNodes)
+		),
+		EUserInterfaceActionType::ToggleButton);
+}
+
+void FSequencerViewOptionsMenu::PopulateLayoutSection(UToolMenu& InMenu)
+{
+	const FSequencerTrackFilterCommands& TrackFilterCommands = FSequencerTrackFilterCommands::Get();
+
+	FToolMenuSection& VisibilitySection = InMenu.FindOrAddSection(TEXT("Visibility"), LOCTEXT("VisibilityHeading", "Filter Bar"));
+
+	VisibilitySection.AddMenuEntry(TrackFilterCommands.ToggleFilterBarVisibility);
+
+	FToolMenuSection& LayoutSection = InMenu.FindOrAddSection(TEXT("Layout"), LOCTEXT("LayoutHeading", "Filter Bar Layout"));
+
+	LayoutSection.AddMenuEntry(
+		TEXT("VerticalLayout"),
+		LOCTEXT("FilterListVerticalLayout", "Vertical"),
+		LOCTEXT("FilterListVerticalLayoutToolTip", "Swap to a vertical layout for the filter bar"),
+		FSlateIcon(),
+		FUIAction(
+			FExecuteAction::CreateRaw(this, &FSequencerViewOptionsMenu::SetFilterLayout, EFilterBarLayout::Vertical),
+			FCanExecuteAction(),
+			FIsActionChecked::CreateRaw(this, &FSequencerViewOptionsMenu::IsFilterLayout, EFilterBarLayout::Vertical)
+		),
+		EUserInterfaceActionType::RadioButton);
+
+	LayoutSection.AddMenuEntry(
+		TEXT("HorizontalLayout"),
+		LOCTEXT("FilterListHorizontalLayout", "Horizontal"),
+		LOCTEXT("FilterListHorizontalLayoutToolTip", "Swap to a Horizontal layout for the filter bar"),
+		FSlateIcon(),
+		FUIAction(
+			FExecuteAction::CreateRaw(this, &FSequencerViewOptionsMenu::SetFilterLayout, EFilterBarLayout::Horizontal),
+			FCanExecuteAction(),
+			FIsActionChecked::CreateRaw(this, &FSequencerViewOptionsMenu::IsFilterLayout, EFilterBarLayout::Horizontal)
+		),
+		EUserInterfaceActionType::RadioButton);
+}
+
+bool FSequencerViewOptionsMenu::IsFilterLayout(const EFilterBarLayout InLayout) const
+{
+	const TSharedPtr<FSequencer> Sequencer = CurrentContext->GetSequencer();
+	if (!Sequencer.IsValid())
+	{
+		return false;
+	}
+
+	const TSharedRef<SSequencer> SequencerWidget = StaticCastSharedRef<SSequencer>(Sequencer->GetSequencerWidget());
+	return InLayout == SequencerWidget->GetFilterBarLayout();
+}
+
+void FSequencerViewOptionsMenu::SetFilterLayout(const EFilterBarLayout InLayout)
+{
+	const TSharedPtr<FSequencer> Sequencer = CurrentContext->GetSequencer();
+	if (!Sequencer.IsValid())
+	{
+		return;
+	}
+
+	const TSharedRef<SSequencer> SequencerWidget = StaticCastSharedRef<SSequencer>(Sequencer->GetSequencerWidget());
+	SequencerWidget->SetFilterBarLayout(InLayout);
+}
+
+bool FSequencerViewOptionsMenu::IsIncludePinnedInFilter() const
+{
+	const TSharedPtr<FSequencer> Sequencer = CurrentContext->GetSequencer();
+	if (!Sequencer.IsValid())
+	{
+		return false;
+	}
+
+	const USequencerSettings* const SequencerSettings = Sequencer->GetSequencerSettings();
+	if (!IsValid(SequencerSettings))
+	{
+		return false;
+	}
+
+	return SequencerSettings->GetIncludePinnedInFilter();
+}
+
+void FSequencerViewOptionsMenu::ToggleIncludePinnedInFilter()
+{
+	const TSharedPtr<ISequencer> Sequencer = CurrentContext->GetSequencer();
+	if (!Sequencer.IsValid())
+	{
+		return;
+	}
+
+	USequencerSettings* const SequencerSettings = Sequencer->GetSequencerSettings();
+	if (!IsValid(SequencerSettings))
+	{
+		return;
+	}
+
+	SequencerSettings->SetIncludePinnedInFilter(!SequencerSettings->GetIncludePinnedInFilter());
+
+	Sequencer->GetFilterInterface()->RequestFilterUpdate();
+}
+
+bool FSequencerViewOptionsMenu::IsAutoExpandPassedFilterNodes() const
+{
+	const TSharedPtr<FSequencer> Sequencer = CurrentContext->GetSequencer();
+	if (!Sequencer.IsValid())
+	{
+		return false;
+	}
+
+	const USequencerSettings* const SequencerSettings = Sequencer->GetSequencerSettings();
+	if (!IsValid(SequencerSettings))
+	{
+		return false;
+	}
+
+	return SequencerSettings->GetAutoExpandNodesOnFilterPass();
+}
+
+void FSequencerViewOptionsMenu::ToggleAutoExpandPassedFilterNodes()
+{
+	const TSharedPtr<ISequencer> Sequencer = CurrentContext->GetSequencer();
+	if (!Sequencer.IsValid())
+	{
+		return;
+	}
+
+	USequencerSettings* const SequencerSettings = Sequencer->GetSequencerSettings();
+	if (!IsValid(SequencerSettings))
+	{
+		return;
+	}
+
+	SequencerSettings->SetAutoExpandNodesOnFilterPass(!SequencerSettings->GetAutoExpandNodesOnFilterPass());
+
+	Sequencer->GetFilterInterface()->RequestFilterUpdate();
+}
+
+TSharedPtr<SSequencer> FSequencerViewOptionsMenu::GetSequencerWidget() const
+{
+	const TSharedPtr<FSequencer> Sequencer = CurrentContext->GetSequencer();
+	if (Sequencer.IsValid())
+	{
+		return StaticCastSharedRef<SSequencer>(Sequencer->GetSequencerWidget());
+	}
+	return nullptr;
+}
+
+#undef LOCTEXT_NAMESPACE
diff --git a/Engine/Source/Editor/Sequencer/Private/Filters/Menus/SequencerViewOptionsMenu.h b/Engine/Source/Editor/Sequencer/Private/Filters/Menus/SequencerViewOptionsMenu.h
new file mode 100644
--- /dev/null
+++ b/Engine/Source/Editor/Sequencer/Private/Filters/Menus/SequencerViewOptionsMenu.h
@@ -0,0 +1,39 @@
+﻿// Copyright Epic Games, Inc. All Rights Reserved.
+
+#pragma once
+
+#include "Templates/SharedPointer.h"
+#include "UObject/WeakObjectPtr.h"
+
+class FSequencer;
+class SSequencer;
+class SWidget;
+class USequencerMenuContext;
+class UToolMenu;
+enum class EFilterBarLayout : uint8;
+
+class FSequencerViewOptionsMenu : public TSharedFromThis<FSequencerViewOptionsMenu>
+{
+public:
+	TSharedRef<SWidget> CreateMenu(const TWeakPtr<FSequencer>& InSequencerWeak);
+
+protected:
+	void PopulateMenu(UToolMenu* const InMenu);
+	void PopulateFiltersSection(UToolMenu& InMenu);
+	void PopulateSortAndOrganizeSection(UToolMenu& InMenu);
+	void PopulateFilterOptionssSection(UToolMenu& InMenu);
+	void PopulateLayoutSection(UToolMenu& InMenu);
+
+	bool IsFilterLayout(const EFilterBarLayout InLayout) const;
+	void SetFilterLayout(const EFilterBarLayout InLayout);
+
+	bool IsIncludePinnedInFilter() const;
+	void ToggleIncludePinnedInFilter();
+
+	bool IsAutoExpandPassedFilterNodes() const;
+	void ToggleAutoExpandPassedFilterNodes();
+
+	TSharedPtr<SSequencer> GetSequencerWidget() const;
+
+	TWeakObjectPtr<USequencerMenuContext> CurrentContext;
+};
