[Sequencer] Advanced Filtering:

diff --git a/Engine/Source/Editor/Sequencer/Private/Filters/SequencerFilterBar.cpp b/Engine/Source/Editor/Sequencer/Private/Filters/SequencerFilterBar.cpp
new file mode 100644
--- /dev/null
+++ b/Engine/Source/Editor/Sequencer/Private/Filters/SequencerFilterBar.cpp
@@ -0,0 +1,1653 @@
+﻿// Copyright Epic Games, Inc. All Rights Reserved.
+
+#include "Filters/SequencerFilterBar.h"
+#include "CurveEditor.h"
+#include "Filters/Filters/SequencerTrackFilter_Keyed.h"
+#include "Filters/Filters/SequencerTrackFilter_CustomText.h"
+#include "Filters/Filters/SequencerTrackFilter_Group.h"
+#include "Filters/Filters/SequencerTrackFilter_HideIsolate.h"
+#include "Filters/Filters/SequencerTrackFilter_Level.h"
+#include "Filters/Filters/SequencerTrackFilter_Modified.h"
+#include "Filters/Filters/SequencerTrackFilter_Selected.h"
+#include "Filters/Filters/SequencerTrackFilter_Text.h"
+#include "Filters/Filters/SequencerTrackFilter_Unbound.h"
+#include "Filters/Filters/SequencerTrackFilters.h"
+#include "Filters/SequencerFilterBarConfig.h"
+#include "Filters/SequencerTrackFilterCollection.h"
+#include "Filters/SequencerTrackFilterCommands.h"
+#include "Filters/SequencerTrackFilterExtension.h"
+#include "Filters/Widgets/SSequencerFilterBar.h"
+#include "Framework/Application/SlateApplication.h"
+#include "Framework/MultiBox/SToolBarButtonBlock.h"
+#include "Menus/SequencerTrackFilterMenu.h"
+#include "MovieScene.h"
+#include "MVVM/CurveEditorExtension.h"
+#include "MVVM/Selection/Selection.h"
+#include "MVVM/ViewModels/CategoryModel.h"
+#include "Sequencer.h"
+#include "SequencerLog.h"
+#include "SSequencer.h"
+#include "Stats/StatsMisc.h"
+#include "UObject/UObjectIterator.h"
+#include "Widgets/SFilterBarIsolateHideShow.h"
+
+using namespace UE::Sequencer;
+
+#define LOCTEXT_NAMESPACE "SequencerFilterBar"
+
+const FName FSequencerFilterBar::SharedIdentifier = TEXT("SharedSequencerTrackFilter");
+
+FSequencerFilterBar::FSequencerFilterBar(FSequencer& InSequencer)
+	: Sequencer(InSequencer)
+	, CommandList(MakeShared<FUICommandList>())
+	, ClassTypeCategory(MakeShared<FFilterCategory>(LOCTEXT("ActorTypeFilterCategory", "Actor Type Filters"), LOCTEXT("ActorTypeFilterCategoryTooltip", "")))
+	, ComponentTypeCategory(MakeShared<FFilterCategory>(LOCTEXT("ObjectTypeFilterCategory", "Object Type Filters"), LOCTEXT("ObjectTypeFilterCategoryTooltip", "")))
+	, MiscCategory(MakeShared<FFilterCategory>(LOCTEXT("MiscFilterCategory", "Misc Filters"), LOCTEXT("MiscFilterCategoryTooltip", "")))
+	, TransientCategory(MakeShared<FFilterCategory>(LOCTEXT("TransientFilterCategory", "Transient Filters"), LOCTEXT("TransientFilterCategoryTooltip", "")))
+	, CommonFilters(MakeShared<FSequencerTrackFilterCollection>(*this))
+	, InternalFilters(MakeShared<FSequencerTrackFilterCollection>(*this))
+	, TextFilter(MakeShared<FSequencerTrackFilter_CustomText>(*this))
+	, HideIsolateFilter(MakeShared<FSequencerTrackFilter_HideIsolate>(*this))
+	, LevelFilter(MakeShared<FSequencerTrackFilter_Level>(*this, TransientCategory))
+	, GroupFilter(MakeShared<FSequencerTrackFilter_Group>(*this, TransientCategory))
+	, SelectedFilter(MakeShared<FSequencerTrackFilter_Selected>(*this, MiscCategory))
+	, ModifiedFilter(MakeShared<FSequencerTrackFilter_Modified>(*this, MiscCategory))
+	, FilterMenu(MakeShared<FSequencerTrackFilterMenu>())
+	, FilterData(FString())
+{
+	FSequencerTrackFilterCommands::Register();
+
+	CommonFilters->OnChanged().AddRaw(this, &FSequencerFilterBar::RequestFilterUpdate);
+	InternalFilters->OnChanged().AddRaw(this, &FSequencerFilterBar::RequestFilterUpdate);
+	TextFilter->OnChanged().AddRaw(this, &FSequencerFilterBar::RequestFilterUpdate);
+	LevelFilter->OnChanged().AddRaw(this, &FSequencerFilterBar::RequestFilterUpdate);
+	HideIsolateFilter->OnChanged().AddRaw(this, &FSequencerFilterBar::RequestFilterUpdate);
+	SelectedFilter->OnChanged().AddRaw(this, &FSequencerFilterBar::RequestFilterUpdate);
+
+	CreateDefaultFilters();
+}
+
+FSequencerFilterBar::~FSequencerFilterBar()
+{
+	FSequencerTrackFilterCommands::Unregister();
+}
+
+TSharedPtr<ICustomTextFilter<FSequencerTrackFilterType>> FSequencerFilterBar::CreateTextFilter()
+{
+	return MakeShared<FSequencerTrackFilter_CustomText>(*this);
+}
+
+void FSequencerFilterBar::CreateDefaultFilters()
+{
+	SetTextFilterString(FString());
+
+	// Add internal filters that won't be saved to config
+	InternalFilters->RemoveAll();
+
+	InternalFilters->Add(LevelFilter);
+	InternalFilters->Add(GroupFilter);
+
+	// Add class type category filters
+	CommonFilters->RemoveAll();
+
+	CommonFilters->Add(MakeShared<FSequencerTrackFilter_Audio>(*this, ClassTypeCategory));
+	CommonFilters->Add(MakeShared<FSequencerTrackFilter_CameraCut>(*this, ClassTypeCategory));
+	CommonFilters->Add(MakeShared<FSequencerTrackFilter_DataLayer>(*this, ClassTypeCategory));
+	CommonFilters->Add(MakeShared<FSequencerTrackFilter_Event>(*this, ClassTypeCategory));
+	CommonFilters->Add(MakeShared<FSequencerTrackFilter_Fade>(*this, ClassTypeCategory));
+	CommonFilters->Add(MakeShared<FSequencerTrackFilter_Folder>(*this, ClassTypeCategory));
+	CommonFilters->Add(MakeShared<FSequencerTrackFilter_LevelVisibility>(*this, ClassTypeCategory));
+	CommonFilters->Add(MakeShared<FSequencerTrackFilter_Particle>(*this, ClassTypeCategory));
+	CommonFilters->Add(MakeShared<FSequencerTrackFilter_CinematicShot>(*this, ClassTypeCategory));
+	CommonFilters->Add(MakeShared<FSequencerTrackFilter_Subsequence>(*this, ClassTypeCategory));
+	CommonFilters->Add(MakeShared<FSequencerTrackFilter_TimeDilation>(*this, ClassTypeCategory));
+	CommonFilters->Add(MakeShared<FSequencerTrackFilter_TimeWarp>(*this, ClassTypeCategory));
+
+	// Add component type category filters
+	CommonFilters->Add(MakeShared<FSequencerTrackFilter_Camera>(*this, ComponentTypeCategory));
+	CommonFilters->Add(MakeShared<FSequencerTrackFilter_Light>(*this, ComponentTypeCategory));
+	CommonFilters->Add(MakeShared<FSequencerTrackFilter_SkeletalMesh>(*this, ComponentTypeCategory));
+
+	// Add misc category filters
+	CommonFilters->Add(MakeShared<FSequencerTrackFilter_Keyed>(*this, MiscCategory));
+	//CommonFilters->Add(ModifiedFilter); // Disabling until clear direction on what this should do
+	CommonFilters->Add(SelectedFilter);
+	CommonFilters->Add(MakeShared<FSequencerTrackFilter_Unbound>(*this, MiscCategory));
+
+	// Add global user-defined filters
+	for (TObjectIterator<USequencerTrackFilterExtension> ExtensionIt(RF_NoFlags); ExtensionIt; ++ExtensionIt)
+	{
+		const USequencerTrackFilterExtension* const PotentialExtension = *ExtensionIt;
+		if (IsValid(PotentialExtension)
+			&& PotentialExtension->HasAnyFlags(RF_ClassDefaultObject)
+			&& !PotentialExtension->GetClass()->HasAnyClassFlags(CLASS_Deprecated | CLASS_Abstract))
+		{
+			TArray<TSharedRef<FSequencerTrackFilter>> ExtendedTrackFilters;
+			PotentialExtension->AddTrackFilterExtensions(*this, ClassTypeCategory, ExtendedTrackFilters);
+
+			for (const TSharedRef<FSequencerTrackFilter>& ExtendedTrackFilter : ExtendedTrackFilters)
+			{
+				CommonFilters->Add(ExtendedTrackFilter);
+			}
+		}
+	}
+
+	CommonFilters->Sort();
+}
+
+void FSequencerFilterBar::BindCommands()
+{
+	const FSequencerTrackFilterCommands& TrackFilterCommands = FSequencerTrackFilterCommands::Get();
+
+	const TSharedRef<SSequencer> SequencerWidget = StaticCastSharedRef<SSequencer>(GetSequencer().GetSequencerWidget());
+
+	CommandList->MapAction(TrackFilterCommands.ToggleFilterBarVisibility,
+		FUIAction(
+			FExecuteAction::CreateSP(SequencerWidget, &SSequencer::ToggleFilterBarVisibility),
+			FCanExecuteAction(),
+			FIsActionChecked::CreateSP(SequencerWidget, &SSequencer::IsFilterBarVisible)
+		));
+
+	CommandList->MapAction(TrackFilterCommands.ResetFilters,
+		FUIAction(
+			FExecuteAction::CreateSP(this, &FSequencerFilterBar::ResetFilters),
+			FCanExecuteAction::CreateSP(this, &FSequencerFilterBar::CanResetFilters)
+		));
+
+	CommandList->MapAction(TrackFilterCommands.ToggleMuteFilters,
+		FUIAction(
+			FExecuteAction::CreateSP(this, &FSequencerFilterBar::ToggleMuteFilters),
+			FCanExecuteAction(),
+			FIsActionChecked::CreateSP(this, &FSequencerFilterBar::AreFiltersMuted)
+		));
+
+	CommandList->MapAction(TrackFilterCommands.DisableAllFilters,
+		FUIAction(
+			FExecuteAction::CreateSP(this, &FSequencerFilterBar::EnableAllFilters, false, TArray<FString>()),
+			FCanExecuteAction::CreateSP(this, &FSequencerFilterBar::HasAnyFilterEnabled)
+		));
+
+	CommandList->MapAction(TrackFilterCommands.ToggleActivateEnabledFilters,
+		FUIAction(
+			FExecuteAction::CreateSP(this, &FSequencerFilterBar::ToggleActivateAllEnabledFilters),
+			FCanExecuteAction::CreateSP(this, &FSequencerFilterBar::HasAnyFilterEnabled)
+		));
+
+	// Bind all filter actions
+	UMovieSceneSequence* const FocusedSequence = Sequencer.GetFocusedMovieSceneSequence();
+	if (!IsValid(FocusedSequence))
+	{
+		return;
+	}
+
+	const TArray<TSharedRef<FSequencerTrackFilter>> AllFilters = GetFilterList(true);
+	for (const TSharedRef<FSequencerTrackFilter>& Filter : AllFilters)
+	{
+		if (Filter->SupportsSequence(FocusedSequence))
+		{
+			Filter->BindCommands();
+		}
+	}
+
+	// Add bindings for curve editor if supported
+	FCurveEditorExtension* CurveEditorExtension = Sequencer.GetViewModel()->CastDynamic<FCurveEditorExtension>();
+	if (CurveEditorExtension && ensure(CurveEditorExtension->GetCurveEditor()))
+	{
+		const TSharedPtr<FUICommandList> SequencerBindings = Sequencer.GetCommandBindings(ESequencerCommandBindings::Sequencer);
+		SequencerBindings->MapAction(TrackFilterCommands.ResetFilters, *CommandList->GetActionForCommand(TrackFilterCommands.ResetFilters));
+		if (const TSharedPtr<FCurveEditor> CurveEditor = CurveEditorExtension->GetCurveEditor())
+		{
+			CurveEditor->GetCommands()->Append(SequencerBindings.ToSharedRef());
+		}
+		
+		const TSharedPtr<FUICommandList> CurveEditorSharedBindings = Sequencer.GetCommandBindings(ESequencerCommandBindings::CurveEditor);
+		CurveEditorSharedBindings->MapAction(TrackFilterCommands.ResetFilters, *CommandList->GetActionForCommand(TrackFilterCommands.ResetFilters));
+		if (const TSharedPtr<FCurveEditor> CurveEditor = CurveEditorExtension->GetCurveEditor())
+		{
+			CurveEditor->GetCommands()->Append(CurveEditorSharedBindings.ToSharedRef());
+		}
+	}
+}
+
+void FSequencerFilterBar::CreateCustomTextFiltersFromConfig()
+{
+	USequencerSettings* const SequencerSettings = Sequencer.GetSequencerSettings();
+	check(IsValid(SequencerSettings));
+
+	CustomTextFilters.Empty();
+
+	FSequencerFilterBarConfig& Config = SequencerSettings->FindOrAddTrackFilterBar(GetIdentifier(), false);
+
+	for (const FCustomTextFilterData& CustomTextFilterData : Config.GetCustomTextFilters())
+	{
+		const TSharedRef<FSequencerTrackFilter_CustomText> NewCustomTextFilter = MakeShared<FSequencerTrackFilter_CustomText>(*this);
+		NewCustomTextFilter->SetFromCustomTextFilterData(CustomTextFilterData);
+		CustomTextFilters.Add(NewCustomTextFilter);
+	}
+}
+
+FSequencer& FSequencerFilterBar::GetSequencer() const
+{
+	return Sequencer;
+}
+
+TSharedPtr<FUICommandList> FSequencerFilterBar::GetCommandList() const
+{
+	return CommandList;
+}
+
+FName FSequencerFilterBar::GetIdentifier() const
+{
+	USequencerSettings* const SequencerSettings = Sequencer.GetSequencerSettings();
+	if (IsValid(SequencerSettings))
+	{
+		return *SequencerSettings->GetName();
+	}
+	return TEXT("SequencerMain");
+}
+
+TSharedPtr<SSequencerFilterBar> FSequencerFilterBar::GetWidget() const
+{
+	return FilterBarWidget;
+}
+
+TSharedRef<SSequencerFilterBar> FSequencerFilterBar::GenerateWidget(const TSharedPtr<SFilterSearchBox>& InSearchBox, const EFilterBarLayout InLayout)
+{
+	if (!FilterBarWidget.IsValid())
+	{
+		FilterBarWidget = SNew(SSequencerFilterBar, SharedThis(this))
+			.FilterBarLayout(InLayout)
+			.AddMetaData<FTagMetaData>(FTagMetaData(TEXT("SequencerTrackFilters")))
+			.FilterSearchBox(InSearchBox);
+	}
+	else
+	{
+		FilterBarWidget->SetLayout(InLayout);
+	}
+
+	return FilterBarWidget.ToSharedRef();
+}
+
+bool FSequencerFilterBar::AreFiltersMuted() const
+{
+	return bFiltersMuted;
+}
+
+void FSequencerFilterBar::MuteFilters(const bool bInMute)
+{
+	bFiltersMuted = bInMute;
+
+	if (FilterBarWidget.IsValid())
+	{
+		FilterBarWidget->SetEnabled(!bFiltersMuted);
+
+		const TWeakPtr<SFilterSearchBox> WeakSearchBox = FilterBarWidget->GetSearchBox();
+		if (WeakSearchBox.IsValid())
+		{
+			WeakSearchBox.Pin()->SetEnabled(!bFiltersMuted);
+		}
+	}
+
+	RequestFilterUpdate();
+}
+
+void FSequencerFilterBar::ToggleMuteFilters()
+{
+	MuteFilters(!AreFiltersMuted());
+}
+
+void FSequencerFilterBar::ResetFilters()
+{
+	EnableAllFilters(false, {});
+	EnableCustomTextFilters(false);
+}
+
+bool FSequencerFilterBar::CanResetFilters() const
+{
+	const TArray<TSharedRef<FFilterCategory>> Categories = { GetClassTypeCategory(), GetComponentTypeCategory(), GetMiscCategory() };
+	const TArray<TSharedRef<FSequencerTrackFilter>> ClassAndCompFilters = GetCommonFilters(Categories);
+	return HasEnabledFilter(ClassAndCompFilters);
+}
+
+FString FSequencerFilterBar::GetTextFilterString() const
+{
+	return TextFilter->GetRawFilterText().ToString();
+}
+
+void FSequencerFilterBar::SetTextFilterString(const FString& InText)
+{
+	TextFilter->SetRawFilterText(FText::FromString(InText));
+
+	if (FilterBarWidget.IsValid())
+	{
+		FilterBarWidget->SetTextFilterString(InText);
+	}
+}
+
+TSharedRef<FSequencerTrackFilter_Text> FSequencerFilterBar::GetTextFilter() const
+{
+	return TextFilter;
+}
+
+FText FSequencerFilterBar::GetFilterErrorText() const
+{
+	return TextFilter->GetFilterErrorText();
+}
+
+TSet<TWeakViewModelPtr<IOutlinerExtension>> FSequencerFilterBar::GetHiddenTracks() const
+{
+	return HideIsolateFilter->GetHiddenTracks();
+}
+
+TSet<TWeakViewModelPtr<IOutlinerExtension>> FSequencerFilterBar::GetIsolatedTracks() const
+{
+	return HideIsolateFilter->GetIsolatedTracks();
+}
+
+void FSequencerFilterBar::HideTracks(const TSet<TWeakViewModelPtr<IOutlinerExtension>>& InTracks, const bool bInAddToExisting)
+{
+	HideIsolateFilter->HideTracks(InTracks, bInAddToExisting);
+}
+
+void FSequencerFilterBar::IsolateTracks(const TSet<TWeakViewModelPtr<IOutlinerExtension>>& InTracks, const bool bInAddToExisting)
+{
+	HideIsolateFilter->IsolateTracks(InTracks, bInAddToExisting);
+}
+
+void FSequencerFilterBar::ShowAllTracks()
+{
+	HideIsolateFilter->ShowAllTracks();
+
+	RequestFilterUpdate();
+}
+
+bool FSequencerFilterBar::HasHiddenTracks() const
+{
+	return HideIsolateFilter->HasHiddenTracks();
+}
+
+bool FSequencerFilterBar::HasIsolatedTracks() const
+{
+	return HideIsolateFilter->HasIsolatedTracks();
+}
+
+void FSequencerFilterBar::EmptyHiddenTracks()
+{
+	HideIsolateFilter->EmptyHiddenTracks();
+
+	RequestFilterUpdate();
+}
+
+void FSequencerFilterBar::EmptyIsolatedTracks()
+{
+	HideIsolateFilter->EmptyIsolatedTracks();
+
+	RequestFilterUpdate();
+}
+
+void FSequencerFilterBar::RequestFilterUpdate()
+{
+	Sequencer.GetNodeTree()->RequestFilterUpdate();
+}
+
+TSharedPtr<FSequencerTrackFilter> FSequencerFilterBar::FindFilterByDisplayName(const FString& InFilterName) const
+{
+	TSharedPtr<FSequencerTrackFilter> OutFilter;
+
+	CommonFilters->ForEachFilter([&InFilterName, &OutFilter]
+		(const TSharedRef<FSequencerTrackFilter>& InFilter)
+		{
+			const FString FilterName = InFilter->GetDisplayName().ToString();
+			if (FilterName.Equals(InFilterName, ESearchCase::IgnoreCase))
+			{
+				OutFilter = InFilter;
+				return false;
+			}
+			return true;
+		}
+		, false);
+
+	return OutFilter;
+}
+
+TSharedPtr<FSequencerTrackFilter_CustomText> FSequencerFilterBar::FindCustomTextFilterByDisplayName(const FString& InFilterName) const
+{
+	TSharedPtr<FSequencerTrackFilter_CustomText> OutFilter;
+
+	for (const TSharedRef<FSequencerTrackFilter_CustomText>& CustomTextFilter : CustomTextFilters)
+	{
+		const FString FilterName = CustomTextFilter->GetDisplayName().ToString();
+		if (FilterName.Equals(InFilterName, ESearchCase::IgnoreCase))
+		{
+			OutFilter = CustomTextFilter;
+			break;
+		}
+	}
+
+	return OutFilter;
+}
+
+bool FSequencerFilterBar::HasAnyFiltersEnabled() const
+{
+	return HasEnabledCommonFilters() || HasEnabledCustomTextFilters();
+}
+
+bool FSequencerFilterBar::IsFilterActiveByDisplayName(const FString InFilterName) const
+{
+	if (const TSharedPtr<FSequencerTrackFilter> Filter = FindFilterByDisplayName(InFilterName))
+	{
+		return IsFilterActive(Filter.ToSharedRef());
+	}
+	return false;
+}
+
+bool FSequencerFilterBar::IsFilterEnabledByDisplayName(const FString InFilterName) const
+{
+	if (const TSharedPtr<FSequencerTrackFilter> Filter = FindFilterByDisplayName(InFilterName))
+	{
+		return IsFilterEnabled(Filter.ToSharedRef());
+	}
+	return false;
+}
+
+bool FSequencerFilterBar::SetFilterActiveByDisplayName(const FString InFilterName, const bool bInActive, const bool bInRequestFilterUpdate)
+{
+	if (const TSharedPtr<FSequencerTrackFilter> Filter = FindFilterByDisplayName(InFilterName))
+	{
+		return SetFilterActive(Filter.ToSharedRef(), bInActive, bInRequestFilterUpdate);
+	}
+
+	if (const TSharedPtr<FSequencerTrackFilter> Filter = FindCustomTextFilterByDisplayName(InFilterName))
+	{
+		return SetFilterActive(Filter.ToSharedRef(), bInActive, bInRequestFilterUpdate);
+	}
+
+	return false;
+}
+
+bool FSequencerFilterBar::SetFilterEnabledByDisplayName(const FString InFilterName, const bool bInEnabled, const bool bInRequestFilterUpdate)
+{
+	if (const TSharedPtr<FSequencerTrackFilter> Filter = FindFilterByDisplayName(InFilterName))
+	{
+		return SetFilterEnabled(Filter.ToSharedRef(), bInEnabled, bInRequestFilterUpdate);
+	}
+
+	if (const TSharedPtr<FSequencerTrackFilter> Filter = FindCustomTextFilterByDisplayName(InFilterName))
+	{
+		return SetFilterEnabled(Filter.ToSharedRef(), bInEnabled, bInRequestFilterUpdate);
+	}
+
+	return false;
+}
+
+bool FSequencerFilterBar::AnyCommonFilterActive() const
+{
+	bool bOutActiveFilter = false;
+
+	CommonFilters->ForEachFilter([this, &bOutActiveFilter]
+		(const TSharedRef<FSequencerTrackFilter>& InFilter)
+		{
+			if (IsFilterActive(InFilter))
+			{
+				bOutActiveFilter = true;
+				return false;
+			}
+			return true;
+		}
+		, false);
+
+	return bOutActiveFilter;
+}
+
+bool FSequencerFilterBar::AnyInternalFilterActive() const
+{
+	const bool bLevelFilterActive = LevelFilter->HasHiddenLevels();
+	const bool bGroupFilterActive = GroupFilter->HasActiveGroupFilter();
+	return bLevelFilterActive || bGroupFilterActive;
+}
+
+bool FSequencerFilterBar::HasAnyFilterActive(const bool bCheckTextFilter
+	, const bool bInCheckHideIsolateFilter
+	, const bool bInCheckCommonFilters
+	, const bool bInCheckInternalFilters
+	, const bool bInCheckCustomTextFilters) const
+{
+	if (bFiltersMuted)
+	{
+		return false;
+	}
+
+	const bool bTextFilterActive = bCheckTextFilter && TextFilter->IsActive();
+	const bool bHideIsolateFilterActive = bInCheckHideIsolateFilter && HideIsolateFilter->IsActive();
+	const bool bCommonFilterActive = bInCheckCommonFilters && AnyCommonFilterActive();
+	const bool bInternalFilterActive = bInCheckInternalFilters && AnyInternalFilterActive();
+	const bool bCustomTextFilterActive = bInCheckCustomTextFilters && AnyCustomTextFilterActive();
+
+	return bTextFilterActive
+		|| bHideIsolateFilterActive
+		|| bCommonFilterActive
+		|| bInternalFilterActive
+		|| bCustomTextFilterActive;
+}
+
+bool FSequencerFilterBar::IsFilterActive(const TSharedRef<FSequencerTrackFilter> InFilter) const
+{
+	USequencerSettings* const SequencerSettings = Sequencer.GetSequencerSettings();
+	check(IsValid(SequencerSettings));
+
+	const FSequencerFilterBarConfig& Config = SequencerSettings->FindOrAddTrackFilterBar(GetIdentifier(), false);
+
+	const FString FilterName = InFilter->GetDisplayName().ToString();
+	return Config.IsFilterActive(FilterName);
+}
+
+bool FSequencerFilterBar::SetFilterActive(const TSharedRef<FSequencerTrackFilter>& InFilter, const bool bInActive, const bool bInRequestFilterUpdate)
+{
+	USequencerSettings* const SequencerSettings = Sequencer.GetSequencerSettings();
+	check(IsValid(SequencerSettings));
+
+	const bool bNewActive = InFilter->IsInverseFilter() ? !bInActive : bInActive;
+
+	FSequencerFilterBarConfig& Config = SequencerSettings->FindOrAddTrackFilterBar(GetIdentifier(), true);
+
+	const FString FilterName = InFilter->GetDisplayName().ToString();
+	const bool bSuccess = Config.SetFilterActive(FilterName, bNewActive);
+
+	if (bSuccess)
+	{
+		SequencerSettings->SaveConfig();
+
+		InFilter->SetActive(bNewActive);
+		InFilter->ActiveStateChanged(bNewActive);
+
+		const ESequencerFilterChange FilterChangeType = bNewActive ? ESequencerFilterChange::Activate : ESequencerFilterChange::Deactivate;
+		FiltersChangedEvent.Broadcast(FilterChangeType, InFilter);
+
+		if (bInRequestFilterUpdate)
+		{
+			RequestFilterUpdate();
+		}
+	}
+
+	return bSuccess;
+}
+
+void FSequencerFilterBar::EnableAllFilters(const bool bInEnable, const TArray<FString> InExceptionFilterNames)
+{
+	TArray<TSharedRef<FSequencerTrackFilter>> ExceptionFilters;
+	TArray<TSharedRef<FSequencerTrackFilter_CustomText>> ExceptionCustomTextFilters;
+
+	for (const FString& FilterName : InExceptionFilterNames)
+	{
+		if (const TSharedPtr<FSequencerTrackFilter> Filter = FindFilterByDisplayName(FilterName))
+		{
+			ExceptionFilters.Add(Filter.ToSharedRef());
+		}
+		else if (const TSharedPtr<FSequencerTrackFilter_CustomText> CustomTextFilter = FindCustomTextFilterByDisplayName(FilterName))
+		{
+			ExceptionCustomTextFilters.Add(CustomTextFilter.ToSharedRef());
+		}
+	}
+
+	EnableFilters(bInEnable, {}, ExceptionFilters);
+	EnableCustomTextFilters(bInEnable, ExceptionCustomTextFilters);
+}
+
+void FSequencerFilterBar::ActivateCommonFilters(const bool bInActivate, const TArray<FString> InExceptionFilterNames)
+{
+	TArray<TSharedRef<FSequencerTrackFilter>> ExceptionFilters;
+
+	for (const FString& FilterName : InExceptionFilterNames)
+	{
+		if (const TSharedPtr<FSequencerTrackFilter> Filter = FindFilterByDisplayName(FilterName))
+		{
+			ExceptionFilters.Add(Filter.ToSharedRef());
+		}
+	}
+
+	return ActivateCommonFilters(bInActivate, {}, ExceptionFilters);
+}
+
+void FSequencerFilterBar::ActivateCommonFilters(const bool bInActivate
+    , const TArray<TSharedRef<FFilterCategory>> InMatchCategories
+    , const TArray<TSharedRef<FSequencerTrackFilter>>& InExceptions)
+{
+	USequencerSettings* const SequencerSettings = Sequencer.GetSequencerSettings();
+	check(IsValid(SequencerSettings));
+
+	FSequencerFilterBarConfig& Config = SequencerSettings->FindOrAddTrackFilterBar(GetIdentifier(), false);
+
+	CommonFilters->ForEachFilter([this, bInActivate, &InExceptions, &Config]
+		(const TSharedRef<FSequencerTrackFilter>& InFilter)
+		{
+			if (InExceptions.IsEmpty() || !InExceptions.Contains(InFilter))
+			{
+				const FString FilterName = InFilter->GetDisplayName().ToString();
+				if (Config.SetFilterActive(FilterName, bInActivate))
+				{
+					const ESequencerFilterChange FilterChangeType = bInActivate
+						? ESequencerFilterChange::Activate
+						: ESequencerFilterChange::Deactivate;
+					FiltersChangedEvent.Broadcast(FilterChangeType, InFilter);
+
+					InFilter->SetActive(bInActivate);
+					InFilter->ActiveStateChanged(bInActivate);
+				}
+			}
+			return true;
+		}
+		, true
+		, InMatchCategories);
+
+	SequencerSettings->SaveConfig();
+
+	RequestFilterUpdate();
+}
+
+bool FSequencerFilterBar::AreAllEnabledFiltersActive(const bool bInActive, const TArray<FString> InExceptionFilterNames) const
+{
+	const TArray<TSharedRef<FSequencerTrackFilter>> EnabledFilters = GetEnabledFilters();
+	for (const TSharedRef<FSequencerTrackFilter>& Filter : EnabledFilters)
+	{
+		const FString FilterName = Filter->GetDisplayName().ToString();
+		if (InExceptionFilterNames.Contains(FilterName))
+		{
+			continue;
+		}
+
+		if (IsFilterActive(Filter) != bInActive)
+		{
+			return false;
+		}
+	}
+
+	const TArray<TSharedRef<FSequencerTrackFilter_CustomText>> EnabledCustomTextFilters = GetEnabledCustomTextFilters();
+	for (const TSharedRef<FSequencerTrackFilter_CustomText>& CustomTextFilter : EnabledCustomTextFilters)
+	{
+		const FString FilterName = CustomTextFilter->GetDisplayName().ToString();
+		if (InExceptionFilterNames.Contains(FilterName))
+		{
+			continue;
+		}
+
+		if (IsFilterActive(CustomTextFilter) != bInActive)
+		{
+			return false;
+		}
+	}
+
+	return true;
+}
+
+void FSequencerFilterBar::ActivateAllEnabledFilters(const bool bInActivate, const TArray<FString> InExceptionFilterNames)
+{
+	const TArray<TSharedRef<FSequencerTrackFilter>> EnabledFilters = GetEnabledFilters();
+	for (const TSharedRef<FSequencerTrackFilter>& Filter : EnabledFilters)
+	{
+		const FString FilterName = Filter->GetDisplayName().ToString();
+		if (InExceptionFilterNames.Contains(FilterName))
+        {
+        	continue;
+        }
+
+		if (IsFilterActive(Filter) != bInActivate)
+		{
+			SetFilterActive(Filter, bInActivate);
+		}
+	}
+
+	const TArray<TSharedRef<FSequencerTrackFilter_CustomText>> EnabledCustomTextFilters = GetEnabledCustomTextFilters();
+	for (const TSharedRef<FSequencerTrackFilter_CustomText>& CustomTextFilter : EnabledCustomTextFilters)
+	{
+		const FString FilterName = CustomTextFilter->GetDisplayName().ToString();
+		if (InExceptionFilterNames.Contains(FilterName))
+		{
+			continue;
+		}
+
+		if (IsFilterActive(CustomTextFilter) != bInActivate)
+		{
+			SetFilterActive(CustomTextFilter, bInActivate);
+		}
+	}
+}
+
+void FSequencerFilterBar::ToggleActivateAllEnabledFilters()
+{
+	const bool bNewActive = !AreAllEnabledFiltersActive(true, {});
+	ActivateAllEnabledFilters(bNewActive, {});
+}
+
+TArray<TSharedRef<FSequencerTrackFilter>> FSequencerFilterBar::GetActiveFilters() const
+{
+	TArray<TSharedRef<FSequencerTrackFilter>> OutFilters;
+
+	CommonFilters->ForEachFilter([this, &OutFilters]
+		(const TSharedRef<FSequencerTrackFilter>& InFilter)
+		{
+			if (IsFilterActive(InFilter))
+			{
+				OutFilters.Add(InFilter);
+			}
+			return true;
+		}
+		, true);
+
+	return OutFilters;
+}
+
+bool FSequencerFilterBar::HasEnabledCommonFilters() const
+{
+	bool bOutReturn = false;
+
+	CommonFilters->ForEachFilter([this, &bOutReturn]
+		(const TSharedRef<FSequencerTrackFilter>& InFilter)
+		{
+			if (IsFilterEnabled(InFilter))
+			{
+				bOutReturn = true;
+				return false;
+			}
+			return true;
+		}
+		, true);
+
+	if (bOutReturn)
+	{
+		return true;
+	}
+
+	InternalFilters->ForEachFilter([this, &bOutReturn]
+		(const TSharedRef<FSequencerTrackFilter>& InFilter)
+		{
+			if (IsFilterEnabled(InFilter))
+			{
+				bOutReturn = true;
+				return false;
+			}
+			return true;
+		}
+		, false);
+
+	return bOutReturn;
+}
+
+bool FSequencerFilterBar::HasEnabledFilter(const TArray<TSharedRef<FSequencerTrackFilter>>& InFilters) const
+{
+	const TArray<TSharedRef<FSequencerTrackFilter>>& Filters = InFilters.IsEmpty() ? GetCommonFilters() : InFilters;
+
+	for (const TSharedRef<FSequencerTrackFilter>& Filter : Filters)
+	{
+		if (IsFilterEnabled(Filter))
+		{
+			return true;
+		}
+	}
+
+	return false;
+}
+
+bool FSequencerFilterBar::HasAnyFilterEnabled() const
+{
+	const bool bCommonFilterEnabled = HasEnabledCommonFilters();
+	const bool bCustomTextFilterEnabled = HasEnabledCustomTextFilters();
+
+	return bCommonFilterEnabled
+		|| bCustomTextFilterEnabled;
+}
+
+bool FSequencerFilterBar::IsFilterEnabled(TSharedRef<FSequencerTrackFilter> InFilter) const
+{
+	USequencerSettings* const SequencerSettings = Sequencer.GetSequencerSettings();
+	check(IsValid(SequencerSettings));
+
+	FSequencerFilterBarConfig& Config = SequencerSettings->FindOrAddTrackFilterBar(GetIdentifier(), false);
+
+	const FString FilterName = InFilter->GetDisplayName().ToString();
+	return Config.IsFilterEnabled(FilterName);
+}
+
+bool FSequencerFilterBar::SetFilterEnabled(const TSharedRef<FSequencerTrackFilter> InFilter, const bool bInEnabled, const bool bInRequestFilterUpdate)
+{
+	USequencerSettings* const SequencerSettings = Sequencer.GetSequencerSettings();
+	check(IsValid(SequencerSettings));
+
+	FSequencerFilterBarConfig& Config = SequencerSettings->FindOrAddTrackFilterBar(GetIdentifier(), true);
+
+	const FString FilterName = InFilter->GetDisplayName().ToString();
+	const bool bSuccess = Config.SetFilterEnabled(FilterName, bInEnabled);
+
+	if (bSuccess)
+	{
+		SequencerSettings->SaveConfig();
+
+		const ESequencerFilterChange FilterChangeType = bInEnabled ? ESequencerFilterChange::Enable : ESequencerFilterChange::Disable;
+		FiltersChangedEvent.Broadcast(FilterChangeType, InFilter);
+
+		if (!bInEnabled && IsFilterActive(InFilter))
+		{
+			InFilter->SetActive(false);
+			InFilter->ActiveStateChanged(false);
+		}
+
+		if (bInRequestFilterUpdate)
+		{
+			RequestFilterUpdate();
+		}
+	}
+
+	return bSuccess;
+}
+
+void FSequencerFilterBar::EnableFilters(const bool bInEnable
+	, const TArray<TSharedRef<FFilterCategory>> InMatchCategories
+	, const TArray<TSharedRef<FSequencerTrackFilter>> InExceptions)
+{
+	USequencerSettings* const SequencerSettings = Sequencer.GetSequencerSettings();
+	check(IsValid(SequencerSettings));
+
+	FSequencerFilterBarConfig& Config = SequencerSettings->FindOrAddTrackFilterBar(GetIdentifier(), true);
+
+	CommonFilters->ForEachFilter([this, bInEnable, &InExceptions, &Config]
+		(const TSharedRef<FSequencerTrackFilter>& InFilter)
+		{
+			if (InExceptions.IsEmpty() || !InExceptions.Contains(InFilter))
+			{
+				const FString FilterName = InFilter->GetDisplayName().ToString();
+				if (Config.SetFilterEnabled(FilterName, bInEnable))
+				{
+					const ESequencerFilterChange FilterChangeType = bInEnable ? ESequencerFilterChange::Enable : ESequencerFilterChange::Disable;
+					FiltersChangedEvent.Broadcast(FilterChangeType, InFilter);
+
+					if (!bInEnable && IsFilterActive(InFilter))
+					{
+						InFilter->SetActive(false);
+						InFilter->ActiveStateChanged(false);
+					}
+				}
+			}
+			return true;
+		}
+		, true
+		, InMatchCategories);
+
+	SequencerSettings->SaveConfig();
+
+	RequestFilterUpdate();
+}
+
+void FSequencerFilterBar::ToggleFilterEnabled(const TSharedRef<FSequencerTrackFilter> InFilter)
+{
+	SetFilterEnabled(InFilter, !IsFilterEnabled(InFilter), true);
+}
+
+TArray<TSharedRef<FSequencerTrackFilter>> FSequencerFilterBar::GetEnabledFilters() const
+{
+	TArray<TSharedRef<FSequencerTrackFilter>> OutFilters;
+
+	CommonFilters->ForEachFilter([this, &OutFilters]
+		(const TSharedRef<FSequencerTrackFilter>& InFilter)
+		{
+			if (IsFilterEnabled(InFilter))
+			{
+				OutFilters.Add(InFilter);
+			}
+			return true;
+		}
+		, true);
+
+	return OutFilters;
+}
+
+bool FSequencerFilterBar::HasAnyCommonFilters() const
+{
+	return !CommonFilters->IsEmpty();
+}
+
+bool FSequencerFilterBar::AddFilter(const TSharedRef<FSequencerTrackFilter>& InFilter)
+{
+	const bool bSuccess = CommonFilters->Add(InFilter) == 1;
+
+	return bSuccess;
+}
+
+bool FSequencerFilterBar::RemoveFilter(const TSharedRef<FSequencerTrackFilter>& InFilter)
+{
+	const bool bSuccess = CommonFilters->Remove(InFilter) == 1;
+
+	if (bSuccess)
+	{
+		FiltersChangedEvent.Broadcast(ESequencerFilterChange::Disable, InFilter);
+	}
+
+	return bSuccess;
+}
+
+TArray<FText> FSequencerFilterBar::GetFilterDisplayNames() const
+{
+	return CommonFilters->GetFilterDisplayNames();
+}
+
+TArray<FText> FSequencerFilterBar::GetCustomTextFilterNames() const
+{
+	TArray<FText> OutLabels;
+
+	for (const TSharedRef<FSequencerTrackFilter_CustomText>& CustomTextFilter : CustomTextFilters)
+	{
+		const FCustomTextFilterData TextFilterData = CustomTextFilter->CreateCustomTextFilterData();
+		OutLabels.Add(TextFilterData.FilterLabel);
+	}
+
+	return OutLabels;
+}
+
+int32 FSequencerFilterBar::GetTotalDisplayNodeCount() const
+{
+	return FilterData.GetTotalNodeCount();
+}
+
+int32 FSequencerFilterBar::GetFilteredDisplayNodeCount() const
+{
+	return FilterData.GetDisplayNodeCount();
+}
+
+TArray<TSharedRef<FSequencerTrackFilter>> FSequencerFilterBar::GetCommonFilters(const TArray<TSharedRef<FFilterCategory>>& InCategories) const
+{
+	return CommonFilters->GetAllFilters(InCategories);
+}
+
+bool FSequencerFilterBar::AnyCustomTextFilterActive() const
+{
+	for (const TSharedRef<FSequencerTrackFilter_CustomText>& Filter : CustomTextFilters)
+	{
+		if (IsFilterActive(Filter))
+		{
+			return true;
+		}
+	}
+
+	return false;
+}
+
+bool FSequencerFilterBar::HasEnabledCustomTextFilters() const
+{
+	for (const TSharedRef<FSequencerTrackFilter_CustomText>& Filter : CustomTextFilters)
+	{
+		if (IsFilterEnabled(Filter))
+		{
+			return true;
+		}
+	}
+	return false;
+}
+
+TArray<TSharedRef<FSequencerTrackFilter_CustomText>> FSequencerFilterBar::GetAllCustomTextFilters() const
+{
+	return CustomTextFilters;
+}
+
+bool FSequencerFilterBar::AddCustomTextFilter(const TSharedRef<FSequencerTrackFilter_CustomText>& InFilter, const bool bInAddToConfig)
+{
+	if (CustomTextFilters.Add(InFilter) != 1)
+	{
+		return false;
+	}
+
+	if (bInAddToConfig)
+	{
+		USequencerSettings* const SequencerSettings = Sequencer.GetSequencerSettings();
+		if (ensure(IsValid(SequencerSettings)))
+		{
+			FSequencerFilterBarConfig& Config = SequencerSettings->FindOrAddTrackFilterBar(GetIdentifier(), false);
+
+			if (Config.AddCustomTextFilter(InFilter->CreateCustomTextFilterData()))
+			{
+				SequencerSettings->SaveConfig();
+			}
+		}
+	}
+
+	FiltersChangedEvent.Broadcast(ESequencerFilterChange::Activate, InFilter);
+
+	return true;
+}
+
+bool FSequencerFilterBar::RemoveCustomTextFilter(const TSharedRef<FSequencerTrackFilter_CustomText>& InFilter, const bool bInAddToConfig)
+{
+	if (CustomTextFilters.Remove(InFilter) != 1)
+	{
+		return false;
+	}
+
+	if (bInAddToConfig)
+	{
+		USequencerSettings* const SequencerSettings = Sequencer.GetSequencerSettings();
+		if (ensure(IsValid(SequencerSettings)))
+		{
+			FSequencerFilterBarConfig& Config = SequencerSettings->FindOrAddTrackFilterBar(GetIdentifier(), false);
+
+			const FString FilterName = InFilter->GetDisplayName().ToString();
+			if (Config.RemoveCustomTextFilter(FilterName))
+			{
+				SequencerSettings->SaveConfig();
+			}
+		}
+	}
+
+	FiltersChangedEvent.Broadcast(ESequencerFilterChange::Disable, InFilter);
+
+	return true;
+}
+
+void FSequencerFilterBar::ActivateCustomTextFilters(const bool bInActivate, const TArray<TSharedRef<FSequencerTrackFilter_CustomText>> InExceptions)
+{
+	USequencerSettings* const SequencerSettings = Sequencer.GetSequencerSettings();
+	check(IsValid(SequencerSettings));
+
+	FSequencerFilterBarConfig& Config = SequencerSettings->FindOrAddTrackFilterBar(GetIdentifier(), false);
+
+	for (const TSharedRef<FSequencerTrackFilter_CustomText>& CustomTextFilter : CustomTextFilters)
+	{
+		if (InExceptions.IsEmpty() || !InExceptions.Contains(CustomTextFilter))
+		{
+			const FString FilterName = CustomTextFilter->GetDisplayName().ToString();
+			if (Config.SetFilterActive(FilterName, bInActivate))
+			{
+				if (!bInActivate && IsFilterActive(CustomTextFilter))
+				{
+					CustomTextFilter->SetActive(false);
+					CustomTextFilter->ActiveStateChanged(false);
+				}
+
+				const ESequencerFilterChange FilterChangeType = bInActivate ? ESequencerFilterChange::Activate : ESequencerFilterChange::Deactivate;
+				FiltersChangedEvent.Broadcast(FilterChangeType, CustomTextFilter);
+			}
+		}
+	}
+
+	SequencerSettings->SaveConfig();
+
+	RequestFilterUpdate();
+}
+
+void FSequencerFilterBar::EnableCustomTextFilters(const bool bInEnable, const TArray<TSharedRef<FSequencerTrackFilter_CustomText>> InExceptions)
+{
+	USequencerSettings* const SequencerSettings = Sequencer.GetSequencerSettings();
+	check(IsValid(SequencerSettings));
+
+	FSequencerFilterBarConfig& Config = SequencerSettings->FindOrAddTrackFilterBar(GetIdentifier(), false);
+
+	for (const TSharedRef<FSequencerTrackFilter_CustomText>& CustomTextFilter : CustomTextFilters)
+	{
+		if (InExceptions.IsEmpty() || !InExceptions.Contains(CustomTextFilter))
+		{
+			const FString FilterName = CustomTextFilter->GetDisplayName().ToString();
+			if (Config.SetFilterEnabled(FilterName, bInEnable))
+			{
+				if (!bInEnable && IsFilterActive(CustomTextFilter))
+				{
+					CustomTextFilter->SetActive(false);
+					CustomTextFilter->ActiveStateChanged(false);
+				}
+
+				const ESequencerFilterChange FilterChangeType = bInEnable ? ESequencerFilterChange::Enable : ESequencerFilterChange::Disable;
+				FiltersChangedEvent.Broadcast(FilterChangeType, CustomTextFilter);
+			}
+		}
+	}
+
+	SequencerSettings->SaveConfig();
+
+	RequestFilterUpdate();
+}
+
+TArray<TSharedRef<FSequencerTrackFilter_CustomText>> FSequencerFilterBar::GetEnabledCustomTextFilters() const
+{
+	TArray<TSharedRef<FSequencerTrackFilter_CustomText>> OutFilters;
+
+	for (const TSharedRef<FSequencerTrackFilter_CustomText>& CustomTextFilter : CustomTextFilters)
+	{
+		if (IsFilterEnabled(CustomTextFilter))
+		{
+			OutFilters.Add(CustomTextFilter);
+		}
+	}
+
+	return OutFilters;
+}
+
+TSet<TSharedRef<FFilterCategory>> FSequencerFilterBar::GetFilterCategories(const TSet<TSharedRef<FSequencerTrackFilter>>* InFilters) const
+{
+	return CommonFilters->GetCategories(InFilters);
+}
+
+TSet<TSharedRef<FFilterCategory>> FSequencerFilterBar::GetConfigCategories() const
+{
+	return { ClassTypeCategory, ComponentTypeCategory, MiscCategory };
+}
+
+TSharedRef<FFilterCategory> FSequencerFilterBar::GetClassTypeCategory() const
+{
+	return ClassTypeCategory;
+}
+
+TSharedRef<FFilterCategory> FSequencerFilterBar::GetComponentTypeCategory() const
+{
+	return ComponentTypeCategory;
+}
+
+TSharedRef<FFilterCategory> FSequencerFilterBar::GetMiscCategory() const
+{
+	return MiscCategory;
+}
+
+void FSequencerFilterBar::ForEachFilter(const TFunctionRef<bool(const TSharedRef<FSequencerTrackFilter>&)>& InFunction
+	, const bool bInCheckSupportsSequence
+	, const TArray<TSharedRef<FFilterCategory>>& InCategories) const
+{
+	CommonFilters->ForEachFilter(InFunction, bInCheckSupportsSequence, InCategories);
+}
+
+bool FSequencerFilterBar::HasActiveLevelFilter() const
+{
+	return LevelFilter->HasHiddenLevels();
+}
+
+bool FSequencerFilterBar::HasAllLevelFiltersActive() const
+{
+	return LevelFilter->HasAllLevelsHidden();
+}
+
+const TSet<FString>& FSequencerFilterBar::GetActiveLevelFilters() const
+{
+	return LevelFilter->GetHiddenLevels();
+}
+
+void FSequencerFilterBar::ActivateLevelFilter(const FString& InLevelName, const bool bInActivate)
+{
+	if (bInActivate)
+	{
+		LevelFilter->UnhideLevel(InLevelName);
+	}
+	else
+	{
+		LevelFilter->HideLevel(InLevelName);
+	}
+}
+
+bool FSequencerFilterBar::IsLevelFilterActive(const FString InLevelName) const
+{
+	return !LevelFilter->IsLevelHidden(InLevelName);
+}
+
+void FSequencerFilterBar::EnableAllLevelFilters(const bool bInEnable)
+{
+	LevelFilter->HideAllLevels(!bInEnable);
+}
+
+bool FSequencerFilterBar::CanEnableAllLevelFilters(const bool bInEnable)
+{
+	return LevelFilter->CanHideAllLevels(!bInEnable);
+}
+
+void FSequencerFilterBar::EnableAllGroupFilters(const bool bInEnable)
+{
+	UMovieSceneSequence* const FocusedMovieSequence = Sequencer.GetFocusedMovieSceneSequence();
+	if (!IsValid(FocusedMovieSequence))
+	{
+		return;
+	}
+
+	UMovieScene* const FocusedMovieScene = FocusedMovieSequence->GetMovieScene();
+	if (!IsValid(FocusedMovieScene))
+	{
+		return;
+	}
+
+	for (UMovieSceneNodeGroup* const NodeGroup : FocusedMovieScene->GetNodeGroups())
+	{
+		NodeGroup->SetEnableFilter(bInEnable);
+	}
+}
+
+bool FSequencerFilterBar::IsGroupFilterActive(const FString& InGroupName) const
+{
+	return GroupFilter->HasActiveGroupFilter();
+}
+
+bool FSequencerFilterBar::PassesAnyCommonFilter(const TViewModelPtr<IOutlinerExtension>& InNode)
+{
+	bool bPassedAnyFilters = false;
+	bool bAnyFilterActive = false;
+
+	// Only one common filter needs to pass for this node to be included in the filtered set
+	CommonFilters->ForEachFilter([this, &InNode, &bPassedAnyFilters, &bAnyFilterActive]
+		(const TSharedRef<FSequencerTrackFilter>& InFilter)
+		{
+			const FString FilterName = InFilter->GetDisplayName().ToString();
+			if (IsFilterActive(InFilter))
+			{
+				bAnyFilterActive = true;
+				
+				if (InFilter->PassesFilter(InNode))
+				{
+					bPassedAnyFilters = true;
+					return false; // Stop processing filters
+				}
+			}
+
+			return true;
+		}
+		, true);
+
+	if (!bAnyFilterActive)
+	{
+		return true;
+	}
+
+	return bPassedAnyFilters;
+}
+
+bool FSequencerFilterBar::PassesAllInternalFilters(const TViewModelPtr<IOutlinerExtension>& InNode)
+{
+	bool bPassedAllFilters = true;
+
+	InternalFilters->ForEachFilter([this, &InNode, &bPassedAllFilters]
+		(const TSharedRef<FSequencerTrackFilter>& InFilter)
+		{
+			if (!InFilter->PassesFilter(InNode))
+			{
+				bPassedAllFilters = false;
+				return false; // Stop processing filters
+			}
+			return true;
+		}
+		, true);
+
+	return bPassedAllFilters;
+}
+
+bool FSequencerFilterBar::PassesAllCustomTextFilters(const TViewModelPtr<IOutlinerExtension>& InNode)
+{
+	for (const TSharedRef<FSequencerTrackFilter_CustomText>& Filter : CustomTextFilters)
+	{
+		if (IsFilterActive(Filter))
+		{
+			if (!Filter->PassesFilter(InNode))
+			{
+				return false;
+			}
+		}
+	}
+
+	return true;
+}
+
+UWorld* FSequencerFilterBar::GetWorld() const
+{
+	UObject* const PlaybackContext = Sequencer.GetPlaybackContext();
+	if (IsValid(PlaybackContext))
+	{
+		return PlaybackContext->GetWorld();
+	}
+	return nullptr;
+}
+
+const FSequencerFilterData& FSequencerFilterBar::FilterNodes()
+{
+	//SCOPE_LOG_TIME_IN_SECONDS(TEXT("FSequencerFilterBar::FilterNodes()"), nullptr);
+
+	// Update the world for the level filter
+	const UObject* const PlaybackContext = Sequencer.GetPlaybackContext();
+	UWorld* const World = PlaybackContext ? PlaybackContext->GetWorld() : nullptr;
+	LevelFilter->UpdateWorld(World);
+
+	// Update the group filter
+	const UMovieSceneSequence* const FocusedMovieSceneSequence = Sequencer.GetFocusedMovieSceneSequence();
+	if (IsValid(FocusedMovieSceneSequence))
+	{
+		GroupFilter->UpdateMovieScene(FocusedMovieSceneSequence->GetMovieScene());
+	}
+
+	// Reset all filter data
+	FilterData.Reset();
+
+	// Always include the bottom spacer
+	TViewModelPtr<IOutlinerExtension> SpacerNode;
+	if (const FSequenceModel* const SequenceModel = Sequencer.GetNodeTree()->GetRootNode()->CastThis<FSequenceModel>())
+	{
+		SpacerNode = CastViewModelChecked<IOutlinerExtension>(SequenceModel->GetBottomSpacer());
+	}
+
+	// Loop through all nodes and filter recursively
+	const bool bHasActiveFilter = HasAnyFilterActive();
+	for (const TViewModelPtr<IOutlinerExtension>& RootNode : Sequencer.GetNodeTree()->GetRootNodes())
+	{
+		FilterNodesRecursive(bHasActiveFilter, RootNode);
+	}
+
+	// Always filter in spacer node
+	SpacerNode->SetFilteredOut(false);
+
+	return FilterData;
+}
+
+FSequencerFilterData& FSequencerFilterBar::GetFilterData()
+{
+	return FilterData;
+}
+
+bool FSequencerFilterBar::FilterNodesRecursive(const bool bInHasActiveFilter, const TViewModelPtr<IOutlinerExtension>& InStartNode)
+{
+	/**
+	 * Main Filtering Logic
+	 *
+	 * - Pinning overrides all other filters
+	 * - Hidden/Isolated tracks will take precedence over common filters
+	 * - Can hide sub tracks of isolated tracks
+	 */
+
+	bool bAnyChildPassed = false;
+
+	// Child nodes should always be processed, as they may force their parents to pass
+	for (const TViewModelPtr<IOutlinerExtension>& Node : InStartNode.AsModel()->GetChildrenOfType<IOutlinerExtension>())
+	{
+		if (FilterNodesRecursive(bInHasActiveFilter, Node))
+		{
+			bAnyChildPassed = true;
+		}
+	}
+
+	// Increment the total node count so we can remove the code to loop again just to count
+	FilterData.IncrementTotalNodeCount();
+
+	// Early out if no filter
+	if (!bInHasActiveFilter)
+	{
+		FilterData.FilterInNode(InStartNode);
+		return false;
+	}
+
+	const USequencerSettings* const SequencerSettings = Sequencer.GetSequencerSettings();
+    check(IsValid(SequencerSettings));
+
+	const TSharedPtr<IPinnableExtension> Pinnable = InStartNode.AsModel()->FindAncestorOfType<IPinnableExtension>(true);
+	const bool bIsPinned = Pinnable && Pinnable->IsPinned();
+
+	// Pinning overrides all other filters
+	if (!SequencerSettings->GetIncludePinnedInFilter() && bIsPinned)
+	{
+		FilterData.FilterInParentChildNodes(InStartNode, true, true, true);
+		return true;
+	}
+
+	const bool bPassedTextFilter = !TextFilter->IsActive() || TextFilter->PassesFilter(InStartNode);
+	const bool bPassedHideIsolateFilter = !HideIsolateFilter->IsActive() || HideIsolateFilter->PassesFilter(InStartNode);
+	const bool bPassedAnyCommonFilters = PassesAnyCommonFilter(InStartNode);
+	const bool bPassedInternalFilters = !AnyInternalFilterActive() || PassesAllInternalFilters(InStartNode);
+	const bool bPassedAnyCustomTextFilters = PassesAllCustomTextFilters(InStartNode);
+
+	const bool bAllFiltersPassed = bPassedTextFilter
+		&& bPassedHideIsolateFilter
+		&& bPassedAnyCommonFilters
+		&& bPassedInternalFilters
+		&& bPassedAnyCustomTextFilters;
+
+	if (bAllFiltersPassed || bAnyChildPassed)
+	{
+		if (SequencerSettings->GetAutoExpandNodesOnFilterPass())
+		{
+			SetTrackParentsExpanded(InStartNode.ImplicitCast(), true);
+		}
+
+		FilterData.FilterInNodeWithAncestors(InStartNode);
+		return true;
+	}
+
+	// After child nodes are processed, fail anything that didn't pass
+	FilterData.FilterOutNode(InStartNode);
+	return false;
+}
+
+TSet<TWeakViewModelPtr<IOutlinerExtension>> FSequencerFilterBar::GetSelectedTracksOrAll() const
+{
+	TSharedPtr<FSequencerEditorViewModel> SequencerViewModel = GetSequencer().GetViewModel();
+	if (!SequencerViewModel.IsValid())
+	{
+		return {};
+	}
+
+	const TSharedPtr<FSequencerSelection> Selection = SequencerViewModel->GetSelection();
+	if (!Selection.IsValid())
+	{
+		return {};
+	}
+
+	const TSet<TWeakViewModelPtr<IOutlinerExtension>> SelectedSet = Selection->Outliner.GetSelected();
+	if (SelectedSet.IsEmpty())
+	{
+		TSet<TWeakViewModelPtr<IOutlinerExtension>> OutTracks;
+		for (const TViewModelPtr<IOutlinerExtension>& TrackModel : SequencerViewModel->GetRootModel()->GetDescendantsOfType<IOutlinerExtension>())
+		{
+			OutTracks.Add(TrackModel);
+		}
+		return OutTracks;
+	}
+
+	return SelectedSet;
+}
+
+bool FSequencerFilterBar::HasSelectedTracks() const
+{
+	return !GetSelectedTracksOrAll().IsEmpty();
+}
+
+void FSequencerFilterBar::AddCategoryGroupIsolatedTracks(const TSet<TWeakViewModelPtr<IOutlinerExtension>>& InTracks
+	, const TSet<FName>& InCategoryNames
+	, const bool bInAddToExisting) const
+{
+	TSharedPtr<FSequencerEditorViewModel> SequencerViewModel = GetSequencer().GetViewModel();
+	if (!SequencerViewModel.IsValid())
+	{
+		return;
+	}
+
+	if (!bInAddToExisting)
+	{
+		HideIsolateFilter->EmptyIsolatedTracks(false);
+	}
+
+	TSet<TWeakViewModelPtr<IOutlinerExtension>> TracksToIsolate;
+	TSet<TViewModelPtr<IOutlinerExtension>> TracksToExpand;
+
+	const FViewModelPtr RootModel = SequencerViewModel->GetRootModel();
+	const TParentFirstChildIterator<IOutlinerExtension> AllNodes = RootModel->GetDescendantsOfType<IOutlinerExtension>();
+	for (const TViewModelPtr<IOutlinerExtension>& Node : AllNodes)
+	{
+		if (const TViewModelPtr<FCategoryGroupModel> CategoryGroupModel = Node.ImplicitCast())
+		{
+			if (InCategoryNames.Contains(CategoryGroupModel->GetCategoryName()))
+			{
+				TracksToIsolate.Add(CategoryGroupModel);
+
+				const TParentModelIterator<IOutlinerExtension> Ancestors = CategoryGroupModel->GetAncestorsOfType<IOutlinerExtension>();
+				TracksToExpand.Append(Ancestors.ToArray());
+				TracksToExpand.Add(CategoryGroupModel);
+			}
+		}
+	}
+
+	HideIsolateFilter->IsolateTracks(TracksToIsolate, true);
+
+	for (const TViewModelPtr<IOutlinerExtension>& Track : TracksToExpand)
+	{
+		Track->SetExpansion(true);
+	}
+}
+
+void FSequencerFilterBar::HideSelectedTracks()
+{
+	const bool bAddToExisting = !FSlateApplication::Get().GetModifierKeys().AreModifersDown(EModifierKey::Shift);
+	const TSet<TWeakViewModelPtr<IOutlinerExtension>> TracksToHide = GetSelectedTracksOrAll();
+	HideIsolateFilter->HideTracks(TracksToHide, bAddToExisting);
+}
+
+void FSequencerFilterBar::IsolateSelectedTracks()
+{
+	const bool bAddToExisting = FSlateApplication::Get().GetModifierKeys().AreModifersDown(EModifierKey::Shift);
+	const TSet<TWeakViewModelPtr<IOutlinerExtension>> TracksToIsolate = GetSelectedTracksOrAll();
+	HideIsolateFilter->IsolateTracks(TracksToIsolate, bAddToExisting);
+}
+
+void FSequencerFilterBar::ShowOnlyLocationCategoryGroups()
+{
+	AddCategoryGroupIsolatedTracks(GetSelectedTracksOrAll(), { TEXT("Location") }, true);
+}
+
+void FSequencerFilterBar::ShowOnlyRotationCategoryGroups()
+{
+	AddCategoryGroupIsolatedTracks(GetSelectedTracksOrAll(), { TEXT("Rotation") }, true);
+}
+
+void FSequencerFilterBar::ShowOnlyScaleCategoryGroups()
+{
+	AddCategoryGroupIsolatedTracks(GetSelectedTracksOrAll(), { TEXT("Scale") }, true);
+}
+
+void FSequencerFilterBar::SetTrackParentsExpanded(const TViewModelPtr<IOutlinerExtension>& InNode, const bool bInExpanded)
+{
+	for (const TViewModelPtr<IOutlinerExtension>& ParentNode : InNode.AsModel()->GetAncestorsOfType<IOutlinerExtension>())
+	{
+		if (!ParentNode->IsExpanded())
+		{
+			ParentNode->SetExpansion(true);
+		}
+	}
+}
+
+FString FSequencerFilterBar::GenerateTextFilterStringFromEnabledFilters() const
+{
+	FString GeneratedFilterString = TextFilter->GetRawFilterText().ToString();
+
+	for (const TSharedRef<FSequencerTrackFilter>& Filter : GetCommonFilters())
+	{
+		if (IsFilterActive(Filter) && IsFilterEnabled(Filter))
+		{
+			const FString AndAddString = GeneratedFilterString.IsEmpty() ? TEXT("") : TEXT(" AND ");
+			const FString ThisFilterGeneratedString = FString::Format(TEXT("{0}{1}==TRUE"), { AndAddString, *Filter->GetName() });
+			GeneratedFilterString.Append(ThisFilterGeneratedString);
+		}
+	}
+
+	return GeneratedFilterString;
+}
+
+TArray<TSharedRef<FSequencerTrackFilter>> FSequencerFilterBar::GetFilterList(const bool bInIncludeCustomTextFilters) const
+{
+	TArray<TSharedRef<FSequencerTrackFilter>> AllFilters;
+
+	AllFilters.Append(CommonFilters->GetAllFilters());
+	AllFilters.Append(InternalFilters->GetAllFilters());
+
+	AllFilters.Add(TextFilter);
+	AllFilters.Add(HideIsolateFilter);
+	AllFilters.Add(LevelFilter);
+	AllFilters.Add(GroupFilter);
+	AllFilters.Add(SelectedFilter);
+	AllFilters.Add(ModifiedFilter);
+
+	if (bInIncludeCustomTextFilters)
+	{
+		for (const TSharedRef<FSequencerTrackFilter_CustomText>& Filter : CustomTextFilters)
+		{
+			AllFilters.Add(Filter);
+		}
+	}
+
+	return AllFilters;
+}
+
+bool FSequencerFilterBar::ShouldUpdateOnTrackValueChanged() const
+{
+	const TArray<TSharedRef<FSequencerTrackFilter>> AllFilters = GetFilterList();
+
+	for (const TSharedRef<FSequencerTrackFilter>& Filter : AllFilters)
+	{
+		if (Filter->ShouldUpdateOnTrackValueChanged())
+		{
+			return true;
+		}
+	}
+
+	return false;
+}
+
+TSharedRef<SFilterBarIsolateHideShow> FSequencerFilterBar::MakeIsolateHideShowPanel()
+{
+	return SNew(SFilterBarIsolateHideShow, SharedThis(this));
+}
+
+TSharedRef<SComboButton> FSequencerFilterBar::MakeAddFilterButton()
+{
+	const TSharedPtr<SLayeredImage> FilterImage = SNew(SLayeredImage)
+		.Image(FAppStyle::Get().GetBrush(TEXT("Icons.Filter")))
+		.ColorAndOpacity_Lambda([this]()
+			{
+				return AreFiltersMuted() ? FLinearColor(1.f, 1.f, 1.f, 0.2f) : FSlateColor::UseForeground();
+			});
+
+	// Badge the filter icon if there are filters enabled or active
+	FilterImage->AddLayer(TAttribute<const FSlateBrush*>::CreateLambda([this]() -> const FSlateBrush*
+		{
+			if (AreFiltersMuted() || !HasAnyFilterEnabled())
+			{
+				return nullptr;
+			}
+
+			if (HasAnyFilterActive(false, false))
+			{
+				return FAppStyle::Get().GetBrush(TEXT("Icons.BadgeModified"));
+			}
+
+			return FAppStyle::Get().GetBrush(TEXT("Icons.Badge"));
+		}));
+
+	const TSharedRef<SToolBarButtonBlock> asdf = SNew(SToolBarButtonBlock);
+	const TSharedRef<SComboButton> ComboButton = SNew(SComboButton)
+		.ComboButtonStyle(&FAppStyle::Get().GetWidgetStyle<FComboButtonStyle>(TEXT("SimpleComboButtonWithIcon")))
+		.ForegroundColor(FSlateColor::UseStyle())
+		.ToolTipText_Lambda([this]()
+			{
+				return FText::Format(LOCTEXT("AddFilterToolTip", "Open the Add Filter Menu to add or manage filters\n\n"
+					"Shift + Click to temporarily mute all active filters\n\n{0}")
+					, SFilterBarIsolateHideShow::MakeLongDisplaySummaryText(*this));
+			})
+		.OnComboBoxOpened_Lambda([this]()
+			{
+				// Don't allow opening the menu if filters are muted or we are toggling the filter mute state
+				if (AreFiltersMuted() || FSlateApplication::Get().GetModifierKeys().IsShiftDown())
+				{
+					FSlateApplication::Get().DismissAllMenus();
+				}
+			})
+		.OnGetMenuContent_Lambda([this]() -> TSharedRef<SWidget>
+			{
+				if (FSlateApplication::Get().GetModifierKeys().IsShiftDown())
+				{
+					MuteFilters(!AreFiltersMuted());
+					FSlateApplication::Get().DismissAllMenus();
+					return SNullWidget::NullWidget;
+				}
+				return FilterMenu->CreateMenu(SharedThis(this));
+			})
+		.ContentPadding(FMargin(1, 0))
+		.ButtonContent()
+		[
+			FilterImage.ToSharedRef()
+		];
+	ComboButton->AddMetadata(MakeShared<FTagMetaData>(TEXT("SequencerTrackFiltersCombo")));
+
+	return ComboButton;
+}
+
+#undef LOCTEXT_NAMESPACE
diff --git a/Engine/Source/Editor/Sequencer/Private/Filters/SequencerFilterBar.h b/Engine/Source/Editor/Sequencer/Private/Filters/SequencerFilterBar.h
new file mode 100644
--- /dev/null
+++ b/Engine/Source/Editor/Sequencer/Private/Filters/SequencerFilterBar.h
@@ -0,0 +1,282 @@
+﻿// Copyright Epic Games, Inc. All Rights Reserved.
+
+#pragma once
+
+#include "Filters/CustomTextFilters.h"
+#include "Filters/ISequencerTrackFilters.h"
+#include "Filters/SequencerFilterData.h"
+#include "Filters/SequencerTrackFilterBase.h"
+
+class FSequencer;
+class FSequencerTrackFilter;
+class FSequencerTrackFilter_CustomText;
+class FSequencerTrackFilter_Group;
+class FSequencerTrackFilter_HideIsolate;
+class FSequencerTrackFilter_Level;
+class FSequencerTrackFilter_Modified;
+class FSequencerTrackFilter_Selected;
+class FSequencerTrackFilter_Text;
+class FSequencerTrackFilterCollection;
+class FSequencerTrackFilterMenu;
+class FUICommandList;
+class SComboButton;
+class SFilterBarIsolateHideShow;
+class SFilterSearchBox;
+class SSequencerFilterBar;
+enum class EFilterBarLayout : uint8;
+
+namespace UE::Sequencer
+{
+	class IOutlinerExtension;
+}
+
+enum class ESequencerFilterChange : uint8
+{
+	Enable,
+	Disable,
+	Activate,
+	Deactivate
+};
+
+/** Holds the Sequencer track filter collection, the current text filter, and hidden/isolated lists. */
+class FSequencerFilterBar : public TSharedFromThis<FSequencerFilterBar>, public ISequencerTrackFilters
+{
+public:
+	/** An identifier shared by all filter bars, used to save and load settings common to every instance */
+	static const FName SharedIdentifier;
+	
+	DECLARE_EVENT_TwoParams(FSequencerFilterBar, FSequencerFiltersChanged, const ESequencerFilterChange /*InChangeType*/, const TSharedRef<FSequencerTrackFilter>& /*InFilter*/);
+	DECLARE_EVENT_TwoParams(FSequencerFilterBar, FSequencerCustomTextFiltersChanged, const ESequencerFilterChange /*InChangeType*/, const TSharedRef<FSequencerTrackFilter_CustomText>& /*InFilter*/);
+
+	FSequencerFilterBar(FSequencer& InSequencer);
+	virtual ~FSequencerFilterBar() override;
+
+	void BindCommands();
+
+	void CreateCustomTextFiltersFromConfig();
+
+	TSharedPtr<ICustomTextFilter<FSequencerTrackFilterType>> CreateTextFilter();
+
+	virtual TSharedPtr<FUICommandList> GetCommandList() const override;
+
+	TSharedPtr<SSequencerFilterBar> GetWidget() const;
+
+	TSharedRef<SSequencerFilterBar> GenerateWidget(const TSharedPtr<SFilterSearchBox>& InSearchBox, const EFilterBarLayout InLayout);
+
+	bool AreFiltersMuted() const;
+	void MuteFilters(const bool bInMute);
+	void ToggleMuteFilters();
+
+	void ResetFilters();
+	bool CanResetFilters() const;
+
+	FSequencerFiltersChanged& GetOnFiltersChanged() { return FiltersChangedEvent; }
+
+	FString GetTextFilterString() const;
+	void SetTextFilterString(const FString& InText);
+
+	TSharedRef<FSequencerTrackFilter_Text> GetTextFilter() const;
+	FText GetFilterErrorText() const;
+
+	/** Hide/Isolate/Show Filter Functions */
+
+	TSet<UE::Sequencer::TWeakViewModelPtr<UE::Sequencer::IOutlinerExtension>> GetHiddenTracks() const;
+	TSet<UE::Sequencer::TWeakViewModelPtr<UE::Sequencer::IOutlinerExtension>> GetIsolatedTracks() const;
+
+	void HideTracks(const TSet<UE::Sequencer::TWeakViewModelPtr<UE::Sequencer::IOutlinerExtension>>& InTracks, const bool bInAddToExisting);
+	void IsolateTracks(const TSet<UE::Sequencer::TWeakViewModelPtr<UE::Sequencer::IOutlinerExtension>>& InTracks, const bool bInAddToExisting);
+	void ShowAllTracks();
+
+	bool HasHiddenTracks() const;
+	bool HasIsolatedTracks() const;
+
+	void EmptyHiddenTracks();
+	void EmptyIsolatedTracks();
+
+	TSharedPtr<FSequencerTrackFilter> FindFilterByDisplayName(const FString& InFilterName) const;
+	TSharedPtr<FSequencerTrackFilter_CustomText> FindCustomTextFilterByDisplayName(const FString& InFilterName) const;
+
+	bool HasAnyFiltersEnabled() const;
+
+	//~ Begin ISequencerTrackFilters
+
+	virtual FName GetIdentifier() const override;
+
+	virtual FSequencer& GetSequencer() const override;
+
+	virtual void RequestFilterUpdate() override;
+
+	virtual void EnableAllFilters(const bool bInEnable, const TArray<FString> InExceptionFilterNames) override;
+
+	virtual void ActivateCommonFilters(const bool bInActivate, const TArray<FString> InExceptionFilterNames) override;
+
+	virtual bool AreAllEnabledFiltersActive(const bool bInActive, const TArray<FString> InExceptionFilterNames) const;
+	virtual void ActivateAllEnabledFilters(const bool bInActivate, const TArray<FString> InExceptionFilterNames) override;
+	virtual void ToggleActivateAllEnabledFilters();
+
+	virtual bool IsFilterActiveByDisplayName(const FString InFilterName) const override;
+	virtual bool IsFilterEnabledByDisplayName(const FString InFilterName) const override;
+	virtual bool SetFilterActiveByDisplayName(const FString InFilterName, const bool bInActive, const bool bInRequestFilterUpdate = true) override;
+	virtual bool SetFilterEnabledByDisplayName(const FString InFilterName, const bool bInEnabled, const bool bInRequestFilterUpdate = true) override;
+
+	virtual TArray<FText> GetFilterDisplayNames() const override;
+	virtual TArray<FText> GetCustomTextFilterNames() const override;
+
+	virtual int32 GetTotalDisplayNodeCount() const override;
+	virtual int32 GetFilteredDisplayNodeCount() const override;
+
+	virtual void HideSelectedTracks() override;
+	virtual void IsolateSelectedTracks() override;
+
+	virtual void ShowOnlyLocationCategoryGroups() override;
+	virtual void ShowOnlyRotationCategoryGroups() override;
+	virtual void ShowOnlyScaleCategoryGroups() override;
+
+	virtual bool HasSelectedTracks() const override;
+
+	virtual FSequencerFilterData& GetFilterData() override;
+
+	//~ End ISequencerTrackFilters
+
+	/** Active Filter Functions */
+
+	bool AnyCommonFilterActive() const;
+	bool AnyInternalFilterActive() const;
+	virtual bool HasAnyFilterActive(const bool bCheckTextFilter = true
+		, const bool bInCheckHideIsolateFilter = true
+		, const bool bInCheckCommonFilters = true
+		, const bool bInCheckInternalFilters = true
+		, const bool bInCheckCustomTextFilters = true) const override;
+	bool IsFilterActive(const TSharedRef<FSequencerTrackFilter> InFilter) const;
+	bool SetFilterActive(const TSharedRef<FSequencerTrackFilter>& InFilter, const bool bInActive, const bool bInRequestFilterUpdate = true);
+	void ActivateCommonFilters(const bool bInActivate
+		, const TArray<TSharedRef<FFilterCategory>> InMatchCategories
+		, const TArray<TSharedRef<FSequencerTrackFilter>>& InExceptions);
+	TArray<TSharedRef<FSequencerTrackFilter>> GetActiveFilters() const;
+
+	/** Enabled Filter Functions */
+
+	bool HasEnabledCommonFilters() const;
+	bool HasEnabledFilter(const TArray<TSharedRef<FSequencerTrackFilter>>& InFilters = {}) const;
+	virtual bool HasAnyFilterEnabled() const override;
+	bool IsFilterEnabled(TSharedRef<FSequencerTrackFilter> InFilter) const;
+	bool SetFilterEnabled(const TSharedRef<FSequencerTrackFilter> InFilter, const bool bInEnabled, const bool bInRequestFilterUpdate = true);
+	void EnableFilters(const bool bInEnable
+		, const TArray<TSharedRef<FFilterCategory>> InMatchCategories = {}
+		, const TArray<TSharedRef<FSequencerTrackFilter>> InExceptions = {});
+	void ToggleFilterEnabled(const TSharedRef<FSequencerTrackFilter> InFilter);
+	TArray<TSharedRef<FSequencerTrackFilter>> GetEnabledFilters() const;
+
+	/** Filter Functions */
+
+	bool HasAnyCommonFilters() const;
+	bool AddFilter(const TSharedRef<FSequencerTrackFilter>& InFilter);
+	bool RemoveFilter(const TSharedRef<FSequencerTrackFilter>& InFilter);
+
+	TArray<TSharedRef<FSequencerTrackFilter>> GetCommonFilters(const TArray<TSharedRef<FFilterCategory>>& InCategories = {}) const;
+
+	/** Custom Text Filter Functions */
+
+	bool AnyCustomTextFilterActive() const;
+	bool HasEnabledCustomTextFilters() const;
+	TArray<TSharedRef<FSequencerTrackFilter_CustomText>> GetAllCustomTextFilters() const;
+	virtual bool AddCustomTextFilter(const TSharedRef<FSequencerTrackFilter_CustomText>& InFilter, const bool bInAddToConfig) override;
+	virtual bool RemoveCustomTextFilter(const TSharedRef<FSequencerTrackFilter_CustomText>& InFilter, const bool bInAddToConfig) override;
+	void ActivateCustomTextFilters(const bool bInActivate, const TArray<TSharedRef<FSequencerTrackFilter_CustomText>> InExceptions = {});
+	void EnableCustomTextFilters(const bool bInEnable, const TArray<TSharedRef<FSequencerTrackFilter_CustomText>> InExceptions = {});
+	TArray<TSharedRef<FSequencerTrackFilter_CustomText>> GetEnabledCustomTextFilters() const;
+
+	/** Filter Category Functions */
+
+	TSet<TSharedRef<FFilterCategory>> GetFilterCategories(const TSet<TSharedRef<FSequencerTrackFilter>>* InFilters = nullptr) const;
+	TSet<TSharedRef<FFilterCategory>> GetConfigCategories() const;
+	TSharedRef<FFilterCategory> GetClassTypeCategory() const;
+	TSharedRef<FFilterCategory> GetComponentTypeCategory() const;
+	TSharedRef<FFilterCategory> GetMiscCategory() const;
+
+	void ForEachFilter(const TFunctionRef<bool(const TSharedRef<FSequencerTrackFilter>&)>& InFunction
+		, const bool bInCheckSupportsSequence
+		, const TArray<TSharedRef<FFilterCategory>>& InCategories = {}) const;
+
+	/** Level Filter Functions */
+
+	bool HasActiveLevelFilter() const;
+	bool HasAllLevelFiltersActive() const;
+	const TSet<FString>& GetActiveLevelFilters() const;
+	void ActivateLevelFilter(const FString& InLevelName, const bool bInActivate);
+	bool IsLevelFilterActive(const FString InLevelName) const;
+	void EnableAllLevelFilters(const bool bInEnable);
+	bool CanEnableAllLevelFilters(const bool bInEnable);
+
+	/** Group Filter Functions */
+
+	void EnableAllGroupFilters(const bool bInEnable);
+	bool IsGroupFilterActive(const FString& InGroupName) const;
+
+	/** Misc Functions */
+
+	TSet<UE::Sequencer::TWeakViewModelPtr<UE::Sequencer::IOutlinerExtension>> GetSelectedTracksOrAll() const;
+
+	void AddCategoryGroupIsolatedTracks(const TSet<UE::Sequencer::TWeakViewModelPtr<UE::Sequencer::IOutlinerExtension>>& InTracks
+		, const TSet<FName>& InCategoryNames
+		, const bool bInAddToExisting) const;
+
+	void SetTrackParentsExpanded(const UE::Sequencer::TViewModelPtr<UE::Sequencer::IOutlinerExtension>& InNode, const bool bInExpanded);
+
+	UWorld* GetWorld() const;
+
+	const FSequencerFilterData& FilterNodes();
+
+	FString GenerateTextFilterStringFromEnabledFilters() const;
+
+	bool ShouldUpdateOnTrackValueChanged() const;
+
+	TSharedRef<SFilterBarIsolateHideShow> MakeIsolateHideShowPanel();
+	TSharedRef<SComboButton> MakeAddFilterButton();
+
+protected:
+	void CreateDefaultFilters();
+
+	bool PassesAnyCommonFilter(const UE::Sequencer::TViewModelPtr<UE::Sequencer::IOutlinerExtension>& InNode);
+	bool PassesAllInternalFilters(const UE::Sequencer::TViewModelPtr<UE::Sequencer::IOutlinerExtension>& InNode);
+	bool PassesAllCustomTextFilters(const UE::Sequencer::TViewModelPtr<UE::Sequencer::IOutlinerExtension>& InNode);
+
+	TArray<TSharedRef<FSequencerTrackFilter>> GetFilterList(const bool bInIncludeCustomTextFilters = false) const;
+
+	/** The Sequencer this filter bar is interacting with */
+	FSequencer& Sequencer;
+
+	TSharedRef<FUICommandList> CommandList;
+
+	/** Global override to enable/disable all filters */
+	bool bFiltersMuted = false;
+
+	TSharedRef<FFilterCategory> ClassTypeCategory;
+	TSharedRef<FFilterCategory> ComponentTypeCategory;
+	TSharedRef<FFilterCategory> MiscCategory;
+	TSharedRef<FFilterCategory> TransientCategory;
+
+	TSharedRef<FSequencerTrackFilterCollection> CommonFilters;
+	TSharedRef<FSequencerTrackFilterCollection> InternalFilters;
+
+	TSharedRef<FSequencerTrackFilter_Text> TextFilter;
+	TSharedRef<FSequencerTrackFilter_HideIsolate> HideIsolateFilter;
+	TSharedRef<FSequencerTrackFilter_Level> LevelFilter;
+	TSharedRef<FSequencerTrackFilter_Group> GroupFilter;
+	TSharedRef<FSequencerTrackFilter_Selected> SelectedFilter;
+	TSharedRef<FSequencerTrackFilter_Modified> ModifiedFilter;
+	
+	TArray<TSharedRef<FSequencerTrackFilter_CustomText>> CustomTextFilters;
+
+	TSharedPtr<SSequencerFilterBar> FilterBarWidget;
+
+	TSharedPtr<FSequencerTrackFilterMenu> FilterMenu;
+
+	FSequencerFiltersChanged FiltersChangedEvent;
+
+	FSequencerFilterData FilterData;
+
+private:
+	/** Do not call directly! Should only be called by FilterNodes(). */
+	bool FilterNodesRecursive(const bool bInHasActiveFilter, const UE::Sequencer::TViewModelPtr<UE::Sequencer::IOutlinerExtension>& InStartNode);
+};
diff --git a/Engine/Source/Editor/Sequencer/Private/Filters/TextExpressions/SequencerTextFilterExpressionContext.cpp b/Engine/Source/Editor/Sequencer/Private/Filters/TextExpressions/SequencerTextFilterExpressionContext.cpp
new file mode 100644
--- /dev/null
+++ b/Engine/Source/Editor/Sequencer/Private/Filters/TextExpressions/SequencerTextFilterExpressionContext.cpp
@@ -0,0 +1,124 @@
+﻿// Copyright Epic Games, Inc. All Rights Reserved.
+
+#include "Filters/SequencerTextFilterExpressionContext.h"
+#include "Filters/SequencerTrackFilterBase.h"
+#include "MVVM/Extensions/IOutlinerExtension.h"
+#include "MVVM/ViewModelPtr.h"
+#include "Sequencer.h"
+#include "MVVM/ViewModels/CategoryModel.h"
+
+using namespace UE::Sequencer;
+
+FSequencerTextFilterExpressionContext::FSequencerTextFilterExpressionContext(ISequencerTrackFilters& InFilterInterface)
+	: FilterInterface(InFilterInterface)
+{
+}
+
+void FSequencerTextFilterExpressionContext::SetFilterItem(FSequencerTrackFilterType InFilterItem, UMovieSceneTrack* const InTrackObject)
+{
+	FilterItem = InFilterItem;
+	WeakTrackObject = InTrackObject;
+}
+
+bool FSequencerTextFilterExpressionContext::TestBasicStringExpression(const FTextFilterString& InValue
+	, const ETextFilterTextComparisonMode InTextComparisonMode) const
+{
+	if (const TViewModelPtr<FChannelGroupOutlinerModel> ChannelGroupOutlinerModel = FilterItem.ImplicitCast())
+	{
+		const FTextFilterString TrackLabel = ChannelGroupOutlinerModel->GetLabel().ToString();
+		if (TextFilterUtils::TestBasicStringExpression(TrackLabel, InValue, ETextFilterTextComparisonMode::Partial))
+		{
+			return true;
+		}
+
+		const FTextFilterString ChannelName = ChannelGroupOutlinerModel->GetChannelName();
+		if (TextFilterUtils::TestBasicStringExpression(ChannelName, InValue, ETextFilterTextComparisonMode::Partial))
+		{
+			return true;
+		}
+	}
+	else if (const TViewModelPtr<FCategoryGroupModel> CategoryGroupModel = FilterItem.ImplicitCast())
+	{
+		const FTextFilterString CategoryName = CategoryGroupModel->GetCategoryName();
+		if (TextFilterUtils::TestBasicStringExpression(CategoryName, InValue, ETextFilterTextComparisonMode::Partial))
+		{
+			return true;
+		}
+	}
+	else if (const TViewModelPtr<FChannelGroupModel> ChannelGroupModel = FilterItem.ImplicitCast())
+	{
+		const FTextFilterString CategoryName = ChannelGroupModel->GetChannelName();
+		if (TextFilterUtils::TestBasicStringExpression(CategoryName, InValue, ETextFilterTextComparisonMode::Partial))
+		{
+			return true;
+		}
+	}
+	else if (const TViewModelPtr<IOutlinerExtension> OutlinerExtension = FilterItem.ImplicitCast())
+	{
+		const FTextFilterString TrackLabel = OutlinerExtension->GetLabel().ToString();
+		if (TextFilterUtils::TestBasicStringExpression(TrackLabel, InValue, ETextFilterTextComparisonMode::Partial))
+		{
+			return true;
+		}
+	}
+
+	return false;
+}
+
+bool FSequencerTextFilterExpressionContext::TestComplexExpression(const FName& InKey
+	, const FTextFilterString& InValue
+	, const ETextFilterComparisonOperation InComparisonOperation
+	, const ETextFilterTextComparisonMode InTextComparisonMode) const
+{
+	if (!FilterItem.IsValid())
+	{
+		return false;
+	}
+
+	const TSet<FName> Keys = GetKeys();
+	if (!Keys.IsEmpty() && !Keys.Contains(InKey))
+	{
+		return false;
+	}
+
+	return !InValue.IsEmpty();
+}
+
+UMovieSceneSequence* FSequencerTextFilterExpressionContext::GetFocusedMovieSceneSequence() const
+{
+	return FilterInterface.GetSequencer().GetFocusedMovieSceneSequence();
+}
+
+UMovieScene* FSequencerTextFilterExpressionContext::GetFocusedGetMovieScene() const
+{
+	const UMovieSceneSequence* const FocusedMovieSceneSequence = GetFocusedMovieSceneSequence();
+	return IsValid(FocusedMovieSceneSequence) ? FocusedMovieSceneSequence->GetMovieScene() : nullptr;
+}
+
+bool FSequencerTextFilterExpressionContext::CompareFStringForExactBool(const FTextFilterString& InValue, const bool bInPassedFilter) const
+{
+	if (InValue.CompareFString(TEXT("TRUE"), ETextFilterTextComparisonMode::Exact))
+	{
+		return bInPassedFilter;
+	}
+	if (InValue.CompareFString(TEXT("FALSE"), ETextFilterTextComparisonMode::Exact))
+	{
+		return !bInPassedFilter;
+	}
+	return true;
+}
+
+bool FSequencerTextFilterExpressionContext::CompareFStringForExactBool(const FTextFilterString& InValue
+	, const ETextFilterComparisonOperation InComparisonOperation
+	, const bool bInPassedFilter) const
+{
+	switch (InComparisonOperation)
+	{
+	case ETextFilterComparisonOperation::Equal:
+		return CompareFStringForExactBool(InValue, bInPassedFilter);
+	case ETextFilterComparisonOperation::NotEqual:
+		return CompareFStringForExactBool(InValue, !bInPassedFilter);
+	default:
+		return true;
+	}
+}
diff --git a/Engine/Source/Editor/Sequencer/Private/Filters/TextExpressions/SequencerTextFilterExpression_Class.cpp b/Engine/Source/Editor/Sequencer/Private/Filters/TextExpressions/SequencerTextFilterExpression_Class.cpp
new file mode 100644
--- /dev/null
+++ b/Engine/Source/Editor/Sequencer/Private/Filters/TextExpressions/SequencerTextFilterExpression_Class.cpp
@@ -0,0 +1,79 @@
+﻿// Copyright Epic Games, Inc. All Rights Reserved.
+
+#include "SequencerTextFilterExpression_Class.h"
+#include "Filters/SequencerFilterData.h"
+#include "Sequencer.h"
+
+using namespace UE::Sequencer;
+
+#define LOCTEXT_NAMESPACE "SequencerTextFilterExpression_Class"
+
+FSequencerTextFilterExpression_Class::FSequencerTextFilterExpression_Class(ISequencerTrackFilters& InFilterInterface)
+	: FSequencerTextFilterExpressionContext(InFilterInterface)
+{
+}
+
+TSet<FName> FSequencerTextFilterExpression_Class::GetKeys() const
+{
+	return { TEXT("CLASS"), TEXT("TYPE") };
+}
+
+ESequencerTextFilterValueType FSequencerTextFilterExpression_Class::GetValueType() const
+{
+	return ESequencerTextFilterValueType::String;
+}
+
+FText FSequencerTextFilterExpression_Class::GetDescription() const
+{
+	return LOCTEXT("ExpressionDescription_Class", "Filter by class name");
+}
+
+bool FSequencerTextFilterExpression_Class::TestComplexExpression(const FName& InKey
+	, const FTextFilterString& InValue
+	, const ETextFilterComparisonOperation InComparisonOperation
+	, const ETextFilterTextComparisonMode InTextComparisonMode) const
+{
+	if (!FSequencerTextFilterExpressionContext::TestComplexExpression(InKey, InValue, InComparisonOperation, InTextComparisonMode))
+	{
+		return true;
+	}
+
+	if (const TViewModelPtr<ITrackExtension> TrackExtension = FilterItem->FindAncestorOfType<ITrackExtension>())
+	{
+		for (const TViewModelPtr<FSectionModel>& SectionModel : TrackExtension->GetSectionModels().IterateSubList<FSectionModel>())
+		{
+			if (const auto SectionInterface = SectionModel->GetSectionInterface())
+			{
+				const FString SectionTitle = SectionInterface->GetSectionTitle().ToString();
+
+				if (TextFilterUtils::TestComplexExpression(SectionTitle, InValue, InComparisonOperation, InTextComparisonMode))
+				{
+					return true;
+				}
+			}
+		}
+	}
+
+	if (WeakTrackObject.IsValid())
+	{
+		const FString TrackClassName = WeakTrackObject->GetClass()->GetName();
+		if (TextFilterUtils::TestComplexExpression(TrackClassName, InValue, InComparisonOperation, InTextComparisonMode))
+		{
+			return true;
+		}
+		return false;
+	}
+
+	ISequencer& Sequencer = FilterInterface.GetSequencer();
+
+	UObject* const BoundObject = FSequencerTrackFilter::ResolveTrackBoundObject(Sequencer, FilterItem, FilterInterface.GetFilterData());
+	if (IsValid(BoundObject)
+		&& !TextFilterUtils::TestComplexExpression(BoundObject->GetClass()->GetName(), InValue, InComparisonOperation, InTextComparisonMode))
+	{
+		return false;
+	}
+
+	return true;
+}
+
+#undef LOCTEXT_NAMESPACE
diff --git a/Engine/Source/Editor/Sequencer/Private/Filters/TextExpressions/SequencerTextFilterExpression_Class.h b/Engine/Source/Editor/Sequencer/Private/Filters/TextExpressions/SequencerTextFilterExpression_Class.h
new file mode 100644
--- /dev/null
+++ b/Engine/Source/Editor/Sequencer/Private/Filters/TextExpressions/SequencerTextFilterExpression_Class.h
@@ -0,0 +1,24 @@
+﻿// Copyright Epic Games, Inc. All Rights Reserved.
+
+#pragma once
+
+#include "Filters/SequencerTextFilterExpressionContext.h"
+
+class FSequencerTextFilterExpression_Class : public FSequencerTextFilterExpressionContext
+{
+public:
+	FSequencerTextFilterExpression_Class(ISequencerTrackFilters& InFilterInterface);
+
+	//~ Begin FSequencerTextFilterExpressionContext
+	virtual TSet<FName> GetKeys() const override;
+	virtual ESequencerTextFilterValueType GetValueType() const override;
+	virtual FText GetDescription() const override;
+	//~ End FSequencerTextFilterExpressionContext
+
+	//~ Begin ITextFilterExpressionContext
+	virtual bool TestComplexExpression(const FName& InKey
+		, const FTextFilterString& InValue
+		, const ETextFilterComparisonOperation InComparisonOperation
+		, const ETextFilterTextComparisonMode InTextComparisonMode) const override;
+	//~ End ITextFilterExpressionContext
+};
diff --git a/Engine/Source/Editor/Sequencer/Private/Filters/Widgets/SFilterBarIsolateHideShow.cpp b/Engine/Source/Editor/Sequencer/Private/Filters/Widgets/SFilterBarIsolateHideShow.cpp
new file mode 100644
--- /dev/null
+++ b/Engine/Source/Editor/Sequencer/Private/Filters/Widgets/SFilterBarIsolateHideShow.cpp
@@ -0,0 +1,304 @@
+﻿// Copyright Epic Games, Inc. All Rights Reserved.
+
+#include "SFilterBarIsolateHideShow.h"
+#include "Filters/SequencerFilterBar.h"
+#include "Filters/SequencerTrackFilterCommands.h"
+#include "Framework/Application/SlateApplication.h"
+#include "Styling/StyleColors.h"
+#include "Widgets/Input/SButton.h"
+#include "Widgets/SBoxPanel.h"
+
+#define LOCTEXT_NAMESPACE "SFilterBarIsolateHideShow"
+
+void SFilterBarIsolateHideShow::Construct(const FArguments& InArgs, const TSharedRef<FSequencerFilterBar>& InFilterBar)
+{
+	WeakFilterBar = InFilterBar;
+
+	ChildSlot
+	[
+		SNew(SHorizontalBox)
+
+		// Isolate Selected Button
+		+ SHorizontalBox::Slot()
+		.AutoWidth()
+		.VAlign(VAlign_Center)
+		.Padding(2.f, 0.f, 0.f, 0.f)
+		[
+			SNew(SBox)
+			.WidthOverride(20.f)
+			[
+				SNew(SButton)
+				.ContentPadding(2.f)
+				.HAlign(HAlign_Center)
+				.VAlign(VAlign_Center)
+				.ButtonStyle(FAppStyle::Get(), TEXT("SimpleButton"))
+				.ToolTipText(this, &SFilterBarIsolateHideShow::GetIsolateTracksButtonTooltipText)
+				.IsEnabled(this, &SFilterBarIsolateHideShow::AreFiltersMuted)
+				.OnClicked(this, &SFilterBarIsolateHideShow::HandleIsolateTracksClick)
+				[
+					SNew(STextBlock)
+					.ColorAndOpacity(this, &SFilterBarIsolateHideShow::GetIsolateTracksButtonTextColor)
+					.Text(LOCTEXT("IsolateSelectedButtonText", "I"))
+				]
+			]
+		]
+
+		// Hide Selected Button
+		+ SHorizontalBox::Slot()
+		.AutoWidth()
+		.VAlign(VAlign_Center)
+		.Padding(2.f, 0.f, 0.f, 0.f)
+		[
+			SNew(SBox)
+			.WidthOverride(20.f)
+			[
+				SNew(SButton)
+				.ContentPadding(2.f)
+				.HAlign(HAlign_Center)
+				.VAlign(VAlign_Center)
+				.ButtonStyle(FAppStyle::Get(), TEXT("SimpleButton"))
+				.ToolTipText(this, &SFilterBarIsolateHideShow::GetHideTracksButtonTooltipText)
+				.IsEnabled(this, &SFilterBarIsolateHideShow::AreFiltersMuted)
+				.OnClicked(this, &SFilterBarIsolateHideShow::HandleHideTracksClick)
+				[
+					SNew(STextBlock)
+					.ColorAndOpacity(this, &SFilterBarIsolateHideShow::GetHideTracksButtonTextColor)
+					.Text(LOCTEXT("HideSelectedButtonText", "H"))
+				]
+			]
+		]
+
+		// Show All Button
+		+ SHorizontalBox::Slot()
+		.AutoWidth()
+		.VAlign(VAlign_Center)
+		.Padding(2.f, 0.f, 0.f, 0.f)
+		[
+			SNew(SBox)
+			.WidthOverride(20.f)
+			[
+				SNew(SButton)
+				.ContentPadding(2.f)
+				.HAlign(HAlign_Center)
+				.VAlign(VAlign_Center)
+				.ButtonStyle(FAppStyle::Get(), TEXT("SimpleButton"))
+				.ToolTipText(this, &SFilterBarIsolateHideShow::GetShowAllTracksButtonTooltipText)
+				.IsEnabled(this, &SFilterBarIsolateHideShow::AreFiltersMuted)
+				.OnClicked(this, &SFilterBarIsolateHideShow::HandleShowAllTracksClick)
+				[
+					SNew(STextBlock)
+					.ColorAndOpacity(this, &SFilterBarIsolateHideShow::GetShowAllTracksButtonTextColor)
+					.Text(LOCTEXT("ShowAllButtonText", "S"))
+				]
+			]
+		]
+	];
+}
+
+bool SFilterBarIsolateHideShow::AreFiltersMuted() const
+{
+	return WeakFilterBar.IsValid() ? !WeakFilterBar.Pin()->AreFiltersMuted() : true;
+}
+
+FReply SFilterBarIsolateHideShow::HandleHideTracksClick()
+{
+	const TSharedPtr<FSequencerFilterBar> FilterBar = WeakFilterBar.Pin();
+	if (!FilterBar.IsValid())
+	{
+		return FReply::Handled();
+	}
+
+	const FModifierKeysState ModifierKeys = FSlateApplication::Get().GetModifierKeys();
+	const bool bControlDown = ModifierKeys.AreModifersDown(EModifierKey::Control);
+	if (bControlDown)
+	{
+		FilterBar->EmptyHiddenTracks();
+	}
+	else
+	{
+		FilterBar->HideSelectedTracks();
+	}
+
+	return FReply::Handled();
+}
+
+FReply SFilterBarIsolateHideShow::HandleIsolateTracksClick()
+{
+	const TSharedPtr<FSequencerFilterBar> FilterBar = WeakFilterBar.Pin();
+	if (!FilterBar.IsValid())
+	{
+		return FReply::Handled();
+	}
+
+	const FModifierKeysState ModifierKeys = FSlateApplication::Get().GetModifierKeys();
+	const bool bControlDown = ModifierKeys.AreModifersDown(EModifierKey::Control);
+	if (bControlDown)
+	{
+		FilterBar->EmptyIsolatedTracks();
+	}
+	else
+	{
+		FilterBar->IsolateSelectedTracks();
+	}
+
+	return FReply::Handled();
+}
+
+FReply SFilterBarIsolateHideShow::HandleShowAllTracksClick()
+{
+	const TSharedPtr<FSequencerFilterBar> FilterBar = WeakFilterBar.Pin();
+	if (!FilterBar.IsValid())
+	{
+		return FReply::Handled();
+	}
+
+	FilterBar->ShowAllTracks();
+
+	return FReply::Handled();
+}
+
+FSlateColor SFilterBarIsolateHideShow::GetHideTracksButtonTextColor() const
+{
+	return WeakFilterBar.IsValid() && WeakFilterBar.Pin()->GetHiddenTracks().Num() == 0
+		? FStyleColors::Foreground : FStyleColors::AccentYellow;
+}
+
+FSlateColor SFilterBarIsolateHideShow::GetIsolateTracksButtonTextColor() const
+{
+	return WeakFilterBar.IsValid() && WeakFilterBar.Pin()->GetIsolatedTracks().Num() == 0
+		? FStyleColors::Foreground : FStyleColors::AccentBlue;
+}
+
+FSlateColor SFilterBarIsolateHideShow::GetShowAllTracksButtonTextColor() const
+{
+	if (const TSharedPtr<FSequencerFilterBar> FilterBar = WeakFilterBar.Pin())
+	{
+		return (FilterBar->GetHiddenTracks().Num() == 0 && FilterBar->GetIsolatedTracks().Num() == 0)
+			? FStyleColors::Foreground : FStyleColors::AccentRed;
+	}
+	return FStyleColors::Foreground;
+}
+
+FText SFilterBarIsolateHideShow::GetHideTracksButtonTooltipText() const
+{
+	const TSharedPtr<FSequencerFilterBar> FilterBar = WeakFilterBar.Pin();
+	if (!FilterBar.IsValid())
+	{
+		return FText::GetEmpty();
+	}
+
+	FText TooltipText = LOCTEXT("HideTracksButtonToolTip", "Hide selected tracks");
+
+	const TSharedRef<const FInputChord> Chord = FSequencerTrackFilterCommands::Get().HideSelectedTracks->GetFirstValidChord();
+	if (Chord->IsValidChord())
+	{
+		TooltipText = FText::Format(LOCTEXT("HideTracksButtonToolTipExtended", "{0} ({1})")
+			, TooltipText, Chord->GetInputText());
+	}
+
+	const FText SummaryWithTotal = MakeHiddenTracksSummaryText(*FilterBar, true);
+	TooltipText = FText::Format(LOCTEXT("HideTracksButtonToolTipExtendedWithTotal", "{0}\n\n"
+		"Use the Control modifier to reset the hidden track list.\n\n{1}")
+		, TooltipText, SummaryWithTotal);
+
+	return TooltipText;
+}
+
+FText SFilterBarIsolateHideShow::GetIsolateTracksButtonTooltipText() const
+{
+	const TSharedPtr<FSequencerFilterBar> FilterBar = WeakFilterBar.Pin();
+	if (!FilterBar.IsValid())
+	{
+		return FText::GetEmpty();
+	}
+
+	FText TooltipText = LOCTEXT("IsolateTracksButtonToolTip", "Isolate selected tracks");
+
+	const TSharedRef<const FInputChord> Chord = FSequencerTrackFilterCommands::Get().IsolateSelectedTracks->GetFirstValidChord();
+	if (Chord->IsValidChord())
+	{
+		TooltipText = FText::Format(LOCTEXT("IsolateTracksButtonToolTipExtended", "{0} ({1})")
+			, TooltipText, Chord->GetInputText());
+	}
+
+	const FText SummaryWithTotal = MakeIsolatedTracksSummaryText(*FilterBar, true);
+	TooltipText = FText::Format(LOCTEXT("IsolateTracksButtonToolTipExtendedWithTotal", "{0}\n\n"
+		"Use the Shift modifier to additively isolate.\n"
+		"Use the Control modifier to reset the isolated track list.\n\n{1}")
+		, TooltipText, SummaryWithTotal);
+
+	return TooltipText;
+}
+
+FText SFilterBarIsolateHideShow::GetShowAllTracksButtonTooltipText() const
+{
+	const TSharedPtr<FSequencerFilterBar> FilterBar = WeakFilterBar.Pin();
+	if (!FilterBar.IsValid())
+	{
+		return FText::GetEmpty();
+	}
+
+	FText TooltipText = LOCTEXT("ShowAllTracksButtonToolTip", "Show all tracks");
+
+	const TSharedRef<const FInputChord> Chord = FSequencerTrackFilterCommands::Get().ShowAllTracks->GetFirstValidChord();
+	if (Chord->IsValidChord())
+	{
+		TooltipText = FText::Format(LOCTEXT("ShowAllTracksButtonToolTipExtended", "{0} ({1})")
+			, TooltipText, Chord->GetInputText());
+	}
+
+	const FText SummaryText = MakeLongDisplaySummaryText(*FilterBar);
+	TooltipText = FText::Format(LOCTEXT("IsolateTracksButtonToolTipExtendedWithTotal", "{0}\n\n{1}")
+		, TooltipText, SummaryText);
+
+	return TooltipText;
+}
+
+FText SFilterBarIsolateHideShow::MakeHiddenTracksSummaryText(FSequencerFilterBar& InFilterBar, const bool bInShowTotalCount)
+{
+	FText Summary = FText::Format(LOCTEXT("HiddenTracksSummary", "{0} hidden tracks")
+		, InFilterBar.GetHiddenTracks().Num());
+
+	if (bInShowTotalCount)
+	{
+		Summary = FText::Format(LOCTEXT("HiddenTracksSummaryWithTotal", "{0} of {1} total tracks")
+			, Summary, InFilterBar.GetFilterData().GetTotalNodeCount());
+	}
+
+	return Summary;
+}
+
+FText SFilterBarIsolateHideShow::MakeIsolatedTracksSummaryText(FSequencerFilterBar& InFilterBar, const bool bInShowTotalCount)
+{
+	FText Summary = FText::Format(LOCTEXT("IsolatedTracksSummary", "{0} isolated tracks")
+		, InFilterBar.GetIsolatedTracks().Num());
+
+	if (bInShowTotalCount)
+	{
+		Summary = FText::Format(LOCTEXT("IsolatedTracksSummaryWithTotal", "{0} of {1} total tracks")
+			, Summary, InFilterBar.GetFilterData().GetTotalNodeCount());
+	}
+
+	return Summary;
+}
+
+FText SFilterBarIsolateHideShow::MakeHideIsolateTracksSummaryText(FSequencerFilterBar& InFilterBar)
+{
+	const FText HiddenTracksSummary = MakeHiddenTracksSummaryText(InFilterBar, false);
+	const FText IsolatedTracksSummary = MakeIsolatedTracksSummaryText(InFilterBar, false);
+	return FText::Format(LOCTEXT("HideIsolateSummary", "{0}, {1}")
+		, HiddenTracksSummary, IsolatedTracksSummary);
+}
+
+FText SFilterBarIsolateHideShow::MakeLongDisplaySummaryText(FSequencerFilterBar& InFilterBar)
+{
+	const FSequencerFilterData& FilterData = InFilterBar.GetFilterData();
+	const int32 FilteredCount = FilterData.GetDisplayNodeCount();
+	const int32 TotalCount = FilterData.GetTotalNodeCount();
+	const FText HideIsolateSummary = MakeHideIsolateTracksSummaryText(InFilterBar);
+
+	return FText::Format(LOCTEXT("LongDisplaySummary", "Showing {0} of {1} total tracks\n{2}"),
+		FilteredCount, TotalCount, HideIsolateSummary);
+}
+
+#undef LOCTEXT_NAMESPACE
diff --git a/Engine/Source/Editor/Sequencer/Private/Filters/Widgets/SFilterBarIsolateHideShow.h b/Engine/Source/Editor/Sequencer/Private/Filters/Widgets/SFilterBarIsolateHideShow.h
new file mode 100644
--- /dev/null
+++ b/Engine/Source/Editor/Sequencer/Private/Filters/Widgets/SFilterBarIsolateHideShow.h
@@ -0,0 +1,39 @@
+﻿// Copyright Epic Games, Inc. All Rights Reserved.
+
+#pragma once
+
+#include "Widgets/SCompoundWidget.h"
+
+class FSequencerFilterBar;
+
+class SFilterBarIsolateHideShow : public SCompoundWidget
+{
+public:
+	static FText MakeHiddenTracksSummaryText(FSequencerFilterBar& InFilterBar, const bool bInShowTotalCount);
+	static FText MakeIsolatedTracksSummaryText(FSequencerFilterBar& InFilterBar, const bool bInShowTotalCount);
+	static FText MakeHideIsolateTracksSummaryText(FSequencerFilterBar& InFilterBar);
+	static FText MakeLongDisplaySummaryText(FSequencerFilterBar& InFilterBar);
+
+	SLATE_BEGIN_ARGS(SFilterBarIsolateHideShow)
+	{}
+	SLATE_END_ARGS()
+
+	void Construct(const FArguments& InArgs, const TSharedRef<FSequencerFilterBar>& InFilterBar);
+
+protected:
+	bool AreFiltersMuted() const;
+
+	FReply HandleHideTracksClick();
+	FReply HandleIsolateTracksClick();
+	FReply HandleShowAllTracksClick();
+
+	FSlateColor GetHideTracksButtonTextColor() const;
+	FSlateColor GetIsolateTracksButtonTextColor() const;
+	FSlateColor GetShowAllTracksButtonTextColor() const;
+
+	FText GetHideTracksButtonTooltipText() const;
+	FText GetIsolateTracksButtonTooltipText() const;
+	FText GetShowAllTracksButtonTooltipText() const;
+
+	TWeakPtr<FSequencerFilterBar> WeakFilterBar;
+};
diff --git a/Engine/Source/Editor/Sequencer/Private/Filters/Widgets/SFilterExpressionHelpDialog.h b/Engine/Source/Editor/Sequencer/Private/Filters/Widgets/SFilterExpressionHelpDialog.h
new file mode 100644
--- /dev/null
+++ b/Engine/Source/Editor/Sequencer/Private/Filters/Widgets/SFilterExpressionHelpDialog.h
@@ -0,0 +1,44 @@
+﻿// Copyright Epic Games, Inc. All Rights Reserved.
+
+#pragma once
+
+#include "Widgets/SWindow.h"
+
+class FSequencerTextFilterExpressionContext;
+enum class ESequencerTextFilterValueType : uint8;
+
+class SFilterExpressionHelpDialog : public SWindow
+{
+public:
+	SLATE_BEGIN_ARGS(SFilterExpressionHelpDialog)
+		: _DialogTitle(NSLOCTEXT("SFilterExpressionHelpDialog", "DialogTitle", "Text Filter Expression Help"))
+		, _DialogSize(FVector2D(320.f, 560.f))
+		, _HeaderText(NSLOCTEXT("SFilterExpressionHelpDialog", "HeaderText", "Text Filter Expressions"))
+		, _DocumentationLink(TEXT("https://dev.epicgames.com/documentation/en-us/unreal-engine/advanced-search-syntax-in-unreal-engine"))
+	{}
+		SLATE_ARGUMENT(FText, DialogTitle)
+		SLATE_ARGUMENT(FVector2D, DialogSize)
+		SLATE_ARGUMENT(FText, HeaderText)
+		SLATE_ARGUMENT(FString, DocumentationLink)
+		SLATE_ARGUMENT(TArray<TSharedRef<FSequencerTextFilterExpressionContext>>, TextFilterExpressionContexts)
+	SLATE_END_ARGS()
+
+	void Construct(const FArguments& InArgs);
+
+protected:
+	static const FSlateColor KeyColor;
+	static const FSlateColor ValueColor;
+
+	TSharedRef<SWidget> ConstructDialogHeader();
+	TSharedRef<SWidget> ConstructExpressionWidgetList();
+	TSharedRef<SWidget> ConstructExpressionWidget(const TSharedPtr<FSequencerTextFilterExpressionContext>& InExpressionContext);
+	TSharedRef<SWidget> ConstructKeysWidget(const TSet<FName>& InKeys);
+	TSharedRef<SWidget> ConstructValueWidget(const ESequencerTextFilterValueType InValueType);
+
+	void OpenDocumentationLink() const;
+
+	FText HeaderText;
+	FString DocumentationLink;
+
+	TArray<TSharedRef<FSequencerTextFilterExpressionContext>> TextFilterExpressionContexts;
+};
diff --git a/Engine/Source/Editor/Sequencer/Private/Filters/Widgets/SSequencerFilter.cpp b/Engine/Source/Editor/Sequencer/Private/Filters/Widgets/SSequencerFilter.cpp
new file mode 100644
--- /dev/null
+++ b/Engine/Source/Editor/Sequencer/Private/Filters/Widgets/SSequencerFilter.cpp
@@ -0,0 +1,165 @@
+﻿// Copyright Epic Games, Inc. All Rights Reserved.
+
+#include "Filters/Widgets/SSequencerFilter.h"
+#include "Filters/Menus/SequencerTrackFilterContextMenu.h"
+#include "Filters/SequencerFilterBar.h"
+#include "Filters/Widgets/SSequencerFilterCheckBox.h"
+
+#define LOCTEXT_NAMESPACE "SSequencerFilter"
+
+void SSequencerFilter::Construct(const FArguments& InArgs
+	, const TSharedRef<FSequencerFilterBar>& InFilterBar
+	, const TSharedRef<FSequencerTrackFilter>& InFilter)
+{
+	FilterBar = InFilterBar;
+	Filter = InFilter;
+
+	ContextMenu = MakeShared<FSequencerTrackFilterContextMenu>();
+
+	TSharedPtr<SWidget> ContentWidget;
+	FName BrushName;
+
+	switch(InArgs._FilterPillStyle)
+	{
+	case EFilterPillStyle::Basic:
+		{
+			ContentWidget = ConstructBasicFilterWidget();
+			BrushName = TEXT("FilterBar.BasicFilterButton");
+			break;
+		}
+	case EFilterPillStyle::Default:
+	default:
+		{
+			ContentWidget = ConstructDefaultFilterWidget();
+			BrushName = TEXT("FilterBar.FilterButton");
+			break;
+		}
+	}
+
+	ChildSlot
+	[
+		SAssignNew(ToggleButtonPtr, SSequencerFilterCheckBox)
+		.Style(FAppStyle::Get(), BrushName)
+		.ToolTipText(TAttribute<FText>::Create(TAttribute<FText>::FGetter::CreateSP(Filter.ToSharedRef(), &FSequencerTrackFilter::GetToolTipText)))
+		.IsChecked(this, &SSequencerFilter::IsChecked)
+		.OnCheckStateChanged(this, &SSequencerFilter::OnFilterToggled)
+		.CheckBoxContentUsesAutoWidth(false)
+		.OnGetMenuContent(this, &SSequencerFilter::GetRightClickMenuContent)
+		[
+			ContentWidget.ToSharedRef()
+		]
+	];
+
+	ToggleButtonPtr->SetOnMouseUp(SSequencerFilterCheckBox::FOnPointerEvent::CreateSP(this, &SSequencerFilter::OnFilterMouseUp));
+	ToggleButtonPtr->SetOnDoubleClick(SSequencerFilterCheckBox::FOnPointerEvent::CreateSP(this, &SSequencerFilter::OnFilterDoubleClicked));
+}
+
+TSharedRef<SWidget> SSequencerFilter::ConstructBasicFilterWidget()
+{
+	return SNew(STextBlock)
+		.Margin(0.f)
+		.TextStyle(FAppStyle::Get(), TEXT("SmallText"))
+		.Text(Filter.ToSharedRef(), &FSequencerTrackFilter::GetDisplayName);
+}
+
+TSharedRef<SWidget> SSequencerFilter::ConstructDefaultFilterWidget()
+{
+	const TSharedRef<FSequencerFilterBar> FilterBarRef = FilterBar.ToSharedRef();
+	const TSharedRef<FSequencerTrackFilter> FilterRef = Filter.ToSharedRef();
+
+	return SNew(SBorder)
+		.Padding(1.f)
+		.BorderImage(FAppStyle::Get().GetBrush(TEXT("FilterBar.FilterBackground")))
+		[
+			SNew(SHorizontalBox)
+			+ SHorizontalBox::Slot()
+			.VAlign(VAlign_Center)
+			.AutoWidth()
+			[
+				SNew(SImage)
+				.DesiredSizeOverride(FVector2D(8, 16)) // was 22
+				.Image(FAppStyle::Get().GetBrush(TEXT("FilterBar.FilterImage")))
+				.ColorAndOpacity(this, &SSequencerFilter::GetFilterImageColorAndOpacity)
+			]
+			+ SHorizontalBox::Slot()
+			.Padding(TAttribute<FMargin>(this, &SSequencerFilter::GetFilterNamePadding))
+			.VAlign(VAlign_Center)
+			[
+				SNew(STextBlock)
+				.Font(FCoreStyle::GetDefaultFontStyle(TEXT("Regular"), 8))
+				.Text(FilterRef, &FSequencerTrackFilter::GetDisplayName)
+				.IsEnabled(FilterBarRef, &FSequencerFilterBar::IsFilterActive, FilterRef)
+			]
+		];
+}
+
+const TSharedPtr<FSequencerTrackFilter> SSequencerFilter::GetFilter() const
+{
+	return Filter;
+}
+
+bool SSequencerFilter::IsActive() const
+{
+	return FilterBar->IsFilterActive(Filter.ToSharedRef());
+}
+
+void SSequencerFilter::OnFilterToggled(const ECheckBoxState NewState)
+{
+	const bool bNewActive = NewState == ECheckBoxState::Checked;
+	FilterBar->SetFilterActive(Filter.ToSharedRef(), bNewActive, true);
+}
+
+FReply SSequencerFilter::OnFilterMouseUp(const FGeometry& InMyGeometry, const FPointerEvent& InMouseEvent)
+{
+	if (InMouseEvent.IsControlDown())
+	{
+		FilterBar->ActivateCommonFilters(true, {}, {});
+	}
+	else if (InMouseEvent.IsAltDown())
+	{
+		FilterBar->ActivateCommonFilters(false, {}, {});
+	}
+	else if(InMouseEvent.GetEffectingButton() == EKeys::MiddleMouseButton)
+	{
+		FilterBar->SetFilterEnabled(Filter.ToSharedRef(), false, true);
+	}
+
+	return FReply::Handled();
+}
+
+FReply SSequencerFilter::OnFilterDoubleClicked(const FGeometry& InMyGeometry, const FPointerEvent& InMouseEvent)
+{
+	FilterBar->ActivateCommonFilters(false, {}, {});
+	FilterBar->ActivateCustomTextFilters(false, {});
+
+	FilterBar->SetFilterActive(Filter.ToSharedRef(), true, true);
+
+	return FReply::Handled();
+}
+
+TSharedRef<SWidget> SSequencerFilter::GetRightClickMenuContent()
+{
+	return ContextMenu->CreateMenuWidget(SharedThis(this));
+}
+
+ECheckBoxState SSequencerFilter::IsChecked() const
+{
+	return IsActive() ? ECheckBoxState::Checked : ECheckBoxState::Unchecked;
+}
+
+FSlateColor SSequencerFilter::GetFilterImageColorAndOpacity() const
+{
+	return IsActive() ? Filter->GetColor() : FAppStyle::Get().GetSlateColor(TEXT("Colors.Recessed"));
+}
+
+EVisibility SSequencerFilter::GetFilterOverlayVisibility() const
+{
+	return IsActive() ? EVisibility::Collapsed : EVisibility::HitTestInvisible;
+}
+
+FMargin SSequencerFilter::GetFilterNamePadding() const
+{
+	return ToggleButtonPtr->IsPressed() ? FMargin(3, 1, 3, 0) : FMargin(3, 0, 3, 0);
+}
+
+#undef LOCTEXT_NAMESPACE
diff --git a/Engine/Source/Editor/Sequencer/Private/Filters/Widgets/SSequencerFilter.h b/Engine/Source/Editor/Sequencer/Private/Filters/Widgets/SSequencerFilter.h
new file mode 100644
--- /dev/null
+++ b/Engine/Source/Editor/Sequencer/Private/Filters/Widgets/SSequencerFilter.h
@@ -0,0 +1,56 @@
+﻿// Copyright Epic Games, Inc. All Rights Reserved.
+
+#pragma once
+
+#include "Filters/SequencerTrackFilterBase.h"
+#include "SSequencerFilterBar.h"
+
+class FSequencerTrackFilterContextMenu;
+class SSequencerFilterCheckBox;
+
+class SSequencerFilter : public SCompoundWidget
+{
+public:
+	DECLARE_DELEGATE_OneParam(FOnSequencerFilterRequest, const TSharedRef<SSequencerFilter>& /*InFilter*/);
+
+	SLATE_BEGIN_ARGS(SSequencerFilter) {}
+		/** Determines how each individual filter pill looks like */
+		SLATE_ARGUMENT(EFilterPillStyle, FilterPillStyle)
+	SLATE_END_ARGS()
+
+	void Construct(const FArguments& InArgs
+		, const TSharedRef<FSequencerFilterBar>& InFilterBar
+		, const TSharedRef<FSequencerTrackFilter>& InFilter);
+
+	const TSharedPtr<FSequencerTrackFilter> GetFilter() const;
+
+protected:
+	TSharedRef<SWidget> ConstructBasicFilterWidget();
+	TSharedRef<SWidget> ConstructDefaultFilterWidget();
+
+	bool IsActive() const;
+
+	void OnFilterToggled(const ECheckBoxState NewState);
+
+	FReply OnFilterMouseUp(const FGeometry& InMyGeometry, const FPointerEvent& InMouseEvent);
+
+	FReply OnFilterDoubleClicked(const FGeometry& InMyGeometry, const FPointerEvent& InMouseEvent);
+
+	TSharedRef<SWidget> GetRightClickMenuContent();
+
+	ECheckBoxState IsChecked() const;
+
+	FSlateColor GetFilterImageColorAndOpacity() const;
+
+	EVisibility GetFilterOverlayVisibility() const;
+
+	FMargin GetFilterNamePadding() const;
+
+	TSharedPtr<FSequencerFilterBar> FilterBar;
+
+	TSharedPtr<FSequencerTrackFilter> Filter;
+
+	TSharedPtr<SSequencerFilterCheckBox> ToggleButtonPtr;
+
+	TSharedPtr<FSequencerTrackFilterContextMenu> ContextMenu;
+};
