[Sequencer] Advanced Filtering:

diff --git a/Engine/Source/Editor/Sequencer/Private/Filters/Widgets/SSequencerFilterBar.cpp b/Engine/Source/Editor/Sequencer/Private/Filters/Widgets/SSequencerFilterBar.cpp
new file mode 100644
--- /dev/null
+++ b/Engine/Source/Editor/Sequencer/Private/Filters/Widgets/SSequencerFilterBar.cpp
@@ -0,0 +1,486 @@
+﻿// Copyright Epic Games, Inc. All Rights Reserved.
+
+#include "Filters/Widgets/SSequencerFilterBar.h"
+#include "Filters/Filters/SequencerTrackFilter_CustomText.h"
+#include "Filters/Filters/SequencerTrackFilter_Text.h"
+#include "Filters/Menus/SequencerFilterBarContextMenu.h"
+#include "Filters/SequencerFilterBar.h"
+#include "Filters/Widgets/SFilterExpressionHelpDialog.h"
+#include "Interfaces/IMainFrameModule.h"
+#include "MovieScene.h"
+#include "Sequencer.h"
+#include "SequencerLog.h"
+#include "SSequencerFilter.h"
+#include "SSequencerCustomTextFilterDialog.h"
+#include "Widgets/SWindow.h"
+
+using namespace UE::Sequencer;
+
+#define LOCTEXT_NAMESPACE "SSequencerFilterBar"
+
+void SSequencerFilterBar::Construct(const FArguments& InArgs, const TSharedRef<FSequencerFilterBar>& InFilterBar)
+{
+	FilterBar = InFilterBar;
+
+	WeakSearchBox = InArgs._FilterSearchBox;
+	FilterBarLayout = InArgs._FilterBarLayout;
+	bCanChangeOrientation = InArgs._CanChangeOrientation;
+	FilterPillStyle = InArgs._FilterPillStyle;
+
+	ChildSlot
+	[
+		SAssignNew(FilterBoxWidget, SWidgetSwitcher)
+		.WidgetIndex_Lambda([this]()
+			{
+				return (FilterBarLayout == EFilterBarLayout::Horizontal) ? 0 : 1;
+			})
+
+		+ SWidgetSwitcher::Slot()
+		.Padding(FMargin(0.f, 2.f, 0.f, 0.f))
+		[
+			SAssignNew(HorizontalContainerWidget, SWrapBox)
+			.UseAllottedSize(true)
+		]
+		+ SWidgetSwitcher::Slot()
+		[
+			SAssignNew(VerticalContainerWidget, SScrollBox)
+			.Visibility_Lambda([this]()
+				{
+				   return HasAnyFilterWidgets() ? EVisibility::Visible : EVisibility::Collapsed;
+				})
+		]
+	];
+
+	AttachFilterSearchBox(InArgs._FilterSearchBox);
+
+	ContextMenu = MakeShared<FSequencerFilterBarContextMenu>();
+
+	CreateFilterWidgetsFromConfig();
+
+	FilterBar->GetOnFiltersChanged().AddSP(this, &SSequencerFilterBar::OnFiltersChanged);
+}
+
+FReply SSequencerFilterBar::OnMouseButtonUp(const FGeometry& MyGeometry, const FPointerEvent& MouseEvent)
+{
+	if (MouseEvent.GetEffectingButton() == EKeys::RightMouseButton)
+	{
+		const FWidgetPath* EventPath = MouseEvent.GetEventPath();
+
+		FSlateApplication::Get().PushMenu(AsShared()
+			, EventPath ? *EventPath : FWidgetPath()
+			, ContextMenu->CreateMenu(FilterBar.ToSharedRef())
+			, MouseEvent.GetScreenSpacePosition()
+			, FPopupTransitionEffect(FPopupTransitionEffect::ContextMenu));
+
+		return FReply::Handled().ReleaseMouseCapture();
+	}
+
+	return FReply::Unhandled();
+}
+
+TSharedPtr<FSequencerFilterBar> SSequencerFilterBar::GetFilterBar() const
+{
+	return FilterBar;
+}
+
+void SSequencerFilterBar::SetTextFilterString(const FString& InText)
+{
+	if (const TSharedPtr<SFilterSearchBox> SearchBox = WeakSearchBox.Pin())
+	{
+		if (!InText.Equals(SearchBox->GetText().ToString()))
+		{
+			SearchBox->SetText(FText::FromString(InText));
+		}
+
+		SearchBox->SetError(FilterBar->GetFilterErrorText());
+	}
+}
+
+FText SSequencerFilterBar::GetFilterErrorText() const
+{
+	return FilterBar->GetFilterErrorText();
+}
+
+EFilterBarLayout SSequencerFilterBar::GetLayout() const
+{
+	return FilterBarLayout;
+}
+
+void SSequencerFilterBar::SetLayout(const EFilterBarLayout InFilterBarLayout)
+{
+	if (!bCanChangeOrientation)
+	{
+		return;
+	}
+
+	FilterBarLayout = InFilterBarLayout;
+
+	if (FilterBarLayout == EFilterBarLayout::Horizontal)
+	{
+		VerticalContainerWidget->ClearChildren();
+
+		FilterBoxWidget->SetActiveWidget(HorizontalContainerWidget.ToSharedRef());
+	}
+	else
+	{
+		HorizontalContainerWidget->ClearChildren();
+
+		FilterBoxWidget->SetActiveWidget(VerticalContainerWidget.ToSharedRef());
+	}
+
+	for (const TSharedRef<SSequencerFilter>& FilterWidget : FilterWidgets)
+	{
+		RemoveWidgetFromLayout(FilterWidget);
+	}
+
+	for (const TSharedRef<SSequencerFilter>& Filter : FilterWidgets)
+	{
+		AddWidgetToLayout(Filter);
+	}
+
+	Invalidate(EInvalidateWidgetReason::Layout);
+}
+
+void SSequencerFilterBar::AttachFilterSearchBox(const TSharedPtr<SFilterSearchBox>& InFilterSearchBox)
+{
+	if (InFilterSearchBox)
+	{
+		InFilterSearchBox->SetOnSaveSearchHandler(
+			 SFilterSearchBox::FOnSaveSearchClicked::CreateSP(this, &SSequencerFilterBar::CreateAddCustomTextFilterWindowFromSearch));
+	}
+}
+
+bool SSequencerFilterBar::HasAnyFilterWidgets() const
+{
+	return FilterWidgets.Num() > 0;
+}
+
+void SSequencerFilterBar::AddWidgetToLayout(const TSharedRef<SWidget>& InWidget)
+{
+	FMargin SlotPadding = FMargin(1); // default editor-wide is FMargin(4, 2) for vertical only
+
+	if (FilterBarLayout == EFilterBarLayout::Horizontal)
+	{
+		switch (FilterPillStyle)
+		{
+		case EFilterPillStyle::Basic:
+			SlotPadding = FMargin(1); // default editor-wide is 2
+			break;
+		case EFilterPillStyle::Default:
+		default:
+			SlotPadding = FMargin(1); // default editor-wide is 3
+		}
+
+		HorizontalContainerWidget->AddSlot()
+			.Padding(SlotPadding)
+			[
+				InWidget
+			];
+	}
+	else
+	{
+		VerticalContainerWidget->AddSlot()
+			.Padding(SlotPadding)
+			[
+				InWidget
+			];
+	}
+}
+
+void SSequencerFilterBar::RemoveWidgetFromLayout(const TSharedRef<SWidget>& InWidget)
+{
+	if (FilterBarLayout == EFilterBarLayout::Horizontal)
+	{
+		HorizontalContainerWidget->RemoveSlot(InWidget);
+	}
+	else
+	{
+		VerticalContainerWidget->RemoveSlot(InWidget);
+	}
+}
+
+TSharedPtr<SSequencerFilter> SSequencerFilterBar::FindFilterWidget(const TSharedRef<FSequencerTrackFilter>& InFilter) const
+{
+	for (const TSharedRef<SSequencerFilter>& FilterWidget : FilterWidgets)
+	{
+		if (FilterWidget->GetFilter() == InFilter)
+		{
+			return FilterWidget;
+		}
+	}
+	return nullptr;
+}
+
+TSharedRef<SSequencerFilter> SSequencerFilterBar::CreateAndAddFilterWidget(const TSharedRef<FSequencerTrackFilter>& InFilter)
+{
+	const TSharedRef<SSequencerFilter> NewFilter = SNew(SSequencerFilter, FilterBar.ToSharedRef(), InFilter)
+		.FilterPillStyle(FilterPillStyle);
+
+	AddFilterWidget(NewFilter);
+
+	return NewFilter;
+}
+
+void SSequencerFilterBar::AddFilterWidget(const TSharedRef<SSequencerFilter>& InFilterWidget)
+{
+	FilterWidgets.Add(InFilterWidget);
+
+	AddWidgetToLayout(InFilterWidget);
+}
+
+void SSequencerFilterBar::RemoveFilterWidget(const TSharedRef<FSequencerTrackFilter>& InFilter, bool ExecuteOnFilterChanged)
+{
+	TSharedPtr<SSequencerFilter> FilterToRemove;
+	for (const TSharedRef<SSequencerFilter>& Filter : FilterWidgets)
+	{
+		if (Filter->GetFilter() == InFilter)
+		{
+			FilterToRemove = Filter;
+			break;
+		}
+	}
+
+	if (FilterToRemove.IsValid())
+	{
+		if (ExecuteOnFilterChanged)
+		{
+			RemoveFilterWidgetAndUpdate(FilterToRemove.ToSharedRef());
+		}
+		else
+		{
+			RemoveFilterWidget(FilterToRemove.ToSharedRef());
+		}
+	}
+}
+
+void SSequencerFilterBar::RemoveFilterWidget(const TSharedRef<SSequencerFilter>& InFilterWidget)
+{
+	FilterWidgets.Remove(InFilterWidget);
+
+	RemoveWidgetFromLayout(InFilterWidget);
+}
+
+void SSequencerFilterBar::RemoveAllFilterWidgets()
+{
+	for (const TSharedRef<SSequencerFilter>& FilterWidget : FilterWidgets)
+	{
+		RemoveWidgetFromLayout(FilterWidget);
+	}
+
+	FilterWidgets.Empty();
+}
+
+void SSequencerFilterBar::RemoveAllFilterWidgetsButThis(const TSharedRef<SSequencerFilter>& InFilterWidget)
+{
+	for (const TSharedRef<SSequencerFilter>& FilterWidget : FilterWidgets)
+	{
+		if (FilterWidget == InFilterWidget)
+		{
+			continue;
+		}
+
+		RemoveWidgetFromLayout(FilterWidget);
+	}
+
+	FilterWidgets.Empty();
+
+	AddFilterWidget(InFilterWidget);
+}
+
+void SSequencerFilterBar::RemoveFilterWidgetAndUpdate(const TSharedRef<SSequencerFilter>& InFilterWidget)
+{
+	RemoveFilterWidget(InFilterWidget);
+}
+
+void SSequencerFilterBar::OnEnableAllGroupFilters(bool bEnableAll)
+{
+	FSequencer& Sequencer = FilterBar->GetSequencer();
+
+	UMovieSceneSequence* const FocusedMovieSequence = Sequencer.GetFocusedMovieSceneSequence();
+	if (!IsValid(FocusedMovieSequence))
+	{
+		return;
+	}
+	
+	UMovieScene* const FocusedMovieScene = FocusedMovieSequence->GetMovieScene();
+	if (!IsValid(FocusedMovieScene))
+	{
+		return;
+	}
+	
+	for (UMovieSceneNodeGroup* const NodeGroup : FocusedMovieScene->GetNodeGroups())
+	{
+		NodeGroup->SetEnableFilter(bEnableAll);
+	}
+}
+
+void SSequencerFilterBar::OnNodeGroupFilterClicked(UMovieSceneNodeGroup* NodeGroup)
+{
+	if (NodeGroup)
+	{
+		NodeGroup->SetEnableFilter(!NodeGroup->GetEnableFilter());
+	}
+}
+
+UWorld* SSequencerFilterBar::GetWorld() const
+{
+	FSequencer& Sequencer = FilterBar->GetSequencer();
+
+	UObject* const PlaybackContext = Sequencer.GetPlaybackContext();
+	if (!IsValid(PlaybackContext))
+	{
+		return nullptr;
+	}
+
+	return PlaybackContext->GetWorld();
+}
+
+TWeakPtr<SFilterSearchBox> SSequencerFilterBar::GetSearchBox() const
+{
+	return WeakSearchBox;
+}
+
+void SSequencerFilterBar::OnFiltersChanged(const ESequencerFilterChange InChangeType, const TSharedRef<FSequencerTrackFilter>& InFilter)
+{
+	switch (InChangeType)
+	{
+	case ESequencerFilterChange::Enable:
+	case ESequencerFilterChange::Activate:
+		{
+			const TSharedPtr<SSequencerFilter> FilterWidget = FindFilterWidget(InFilter);
+			if (!FilterWidget.IsValid())
+			{
+				CreateAndAddFilterWidget(InFilter);
+			}
+			break;
+		}
+	case ESequencerFilterChange::Disable:
+		{
+			RemoveFilterWidget(InFilter);
+			break;
+		}
+	case ESequencerFilterChange::Deactivate:
+		{
+			break;
+		}
+	};
+}
+
+void SSequencerFilterBar::CreateAddCustomTextFilterWindowFromSearch(const FText& InSearchText)
+{
+	FCustomTextFilterData CustomTextFilterData;
+	CustomTextFilterData.FilterLabel = LOCTEXT("NewFilterName", "New Filter Name");
+	CustomTextFilterData.FilterString = InSearchText;
+	SSequencerCustomTextFilterDialog::CreateWindow_AddCustomTextFilter(FilterBar.ToSharedRef(), MoveTemp(CustomTextFilterData));
+}
+
+TSharedRef<FFilterCategory> SSequencerFilterBar::GetClassTypeCategory() const
+{
+	return FilterBar->GetClassTypeCategory();
+}
+
+TSharedRef<FFilterCategory> SSequencerFilterBar::GetComponentTypeCategory() const
+{
+	return FilterBar->GetComponentTypeCategory();
+}
+
+TSharedRef<FFilterCategory> SSequencerFilterBar::GetMiscCategory() const
+{
+	return FilterBar->GetMiscCategory();
+}
+
+void SSequencerFilterBar::OnOpenTextExpressionHelp()
+{
+	if (TextExpressionHelpWindow.IsValid())
+	{
+		TextExpressionHelpWindow->BringToFront();
+	}
+	else
+	{
+		const TSharedRef<SFilterExpressionHelpDialog> HelpDialog = SNew(SFilterExpressionHelpDialog)
+			.DialogTitle(LOCTEXT("SequencerCustomTextFilterHelp", "Sequencer Custom Text Filter Help"))
+			.TextFilterExpressionContexts(FilterBar->GetTextFilter()->GetTextFilterExpressionContexts());
+
+		HelpDialog->GetOnWindowClosedEvent().AddLambda([this](const TSharedRef<SWindow>& InWindow)
+			{
+				TextExpressionHelpWindow.Reset();
+			});
+
+		TSharedPtr<SWindow> ParentWindow;
+		if (FModuleManager::Get().IsModuleLoaded(TEXT("MainFrame")))
+		{
+			IMainFrameModule& MainFrame = FModuleManager::LoadModuleChecked<IMainFrameModule>(TEXT("MainFrame"));
+			ParentWindow = MainFrame.GetParentWindow();
+		}
+
+		if (ParentWindow.IsValid())
+		{
+			TextExpressionHelpWindow = FSlateApplication::Get().AddWindowAsNativeChild(HelpDialog, ParentWindow.ToSharedRef());
+		}
+		else
+		{
+			TextExpressionHelpWindow = FSlateApplication::Get().AddWindow(HelpDialog);
+		}
+	}
+}
+void SSequencerFilterBar::SaveCurrentFilterSetAsCustomTextFilter()
+{
+	FCustomTextFilterData CustomTextFilterData;
+	CustomTextFilterData.FilterString = FText::FromString(FilterBar->GenerateTextFilterStringFromEnabledFilters());
+	if (CustomTextFilterData.FilterLabel.IsEmpty())
+	{
+		CustomTextFilterData.FilterLabel = LOCTEXT("NewFilterName", "New Filter Name");
+	}
+	SSequencerCustomTextFilterDialog::CreateWindow_AddCustomTextFilter(FilterBar.ToSharedRef(), MoveTemp(CustomTextFilterData));
+}
+
+void SSequencerFilterBar::CreateFilterWidgetsFromConfig()
+{
+	USequencerSettings* const SequencerSettings = FilterBar->GetSequencer().GetSequencerSettings();
+	check(IsValid(SequencerSettings));
+
+	const FName InstanceIdentifier = FilterBar->GetIdentifier();
+	FSequencerFilterBarConfig* const Config = SequencerSettings->FindTrackFilterBar(InstanceIdentifier);
+	if (!Config)
+	{
+		UE_LOG(LogSequencer, Error, TEXT("SSequencerFilterBar requires that you specify a FilterBarIdentifier to load settings"));
+		return;
+	}
+
+	RemoveAllFilterWidgets();
+
+	const TSet<TSharedRef<FFilterCategory>> DisplayableCategories = FilterBar->GetConfigCategories();
+
+	auto LoadFilterFromConfig = [this, &Config, &DisplayableCategories](const TSharedRef<FSequencerTrackFilter>& InFilter)
+	{
+		const TSharedPtr<FFilterCategory> FilterCategory = InFilter->GetCategory();
+		if (FilterCategory.IsValid() && !DisplayableCategories.Contains(FilterCategory.ToSharedRef()))
+		{
+			return;
+		}
+
+		const FString FilterName = InFilter->GetDisplayName().ToString();
+		if (!Config->IsFilterEnabled(FilterName))
+		{
+			return;
+		}
+
+		const TSharedPtr<SSequencerFilter> FilterWidget = FindFilterWidget(InFilter);
+		if (!FilterWidget.IsValid())
+		{
+			CreateAndAddFilterWidget(InFilter);
+		}
+	};
+
+	for (const TSharedRef<FSequencerTrackFilter>& Filter : FilterBar->GetCommonFilters())
+	{
+		LoadFilterFromConfig(Filter);
+	}
+
+	for (const TSharedRef<FSequencerTrackFilter_CustomText>& Filter : FilterBar->GetAllCustomTextFilters())
+	{
+		LoadFilterFromConfig(Filter);
+	}
+}
+
+#undef LOCTEXT_NAMESPACE
diff --git a/Engine/Source/Editor/Sequencer/Private/Filters/Widgets/SSequencerFilterBar.h b/Engine/Source/Editor/Sequencer/Private/Filters/Widgets/SSequencerFilterBar.h
new file mode 100644
--- /dev/null
+++ b/Engine/Source/Editor/Sequencer/Private/Filters/Widgets/SSequencerFilterBar.h
@@ -0,0 +1,130 @@
+﻿// Copyright Epic Games, Inc. All Rights Reserved.
+
+#pragma once
+
+#include "Filters/SBasicFilterBar.h"
+#include "Filters/SequencerTrackFilterBase.h"
+
+class FSequencer;
+class FSequencerFilterBarContextMenu;
+class FSequencerFilterBar;
+class FSequencerTrackFilter_CustomText;
+class SSequencerFilter;
+class UMovieSceneNodeGroup;
+enum class ESequencerFilterChange : uint8;
+
+class SSequencerFilterBar : public SCompoundWidget
+{
+public:
+	/** Delegate for when filters have changed */
+	DECLARE_DELEGATE(FOnFilterChanged);
+
+	/** Delegate to create a TTextFilter used to compare FilterType with text queries */
+	DECLARE_DELEGATE_RetVal(TSharedPtr<FSequencerTrackFilter_CustomText>, FCreateTextFilter);
+
+	DECLARE_DELEGATE_OneParam(FOnFilterBarLayoutChanging, EFilterBarLayout /*InNewLayout*/);
+
+	DECLARE_EVENT(FSequencerFilterBar, FSequencerFiltersChanged);
+
+	DECLARE_EVENT_OneParam(FSequencerFilterBar, FExternalCustomTextFilterEvent, const TSharedRef<FSequencerFilterBar>& /*InBroadcastingFilterBar*/);
+
+	SLATE_BEGIN_ARGS(SSequencerFilterBar)
+		: _FilterBarLayout(EFilterBarLayout::Vertical)
+		, _CanChangeOrientation(true)
+		, _FilterPillStyle(EFilterPillStyle::Default)
+		, _UseSectionsForCategories(true)
+	{}
+		/** An SFilterSearchBox that can be attached to this filter bar. When provided along with a CreateTextFilter
+		 *  delegate, allows the user to save searches from the Search Box as text filters for the filter bar.
+		 *	NOTE: Will bind a delegate to SFilterSearchBox::OnClickedAddSearchHistoryButton */
+		SLATE_ARGUMENT(TSharedPtr<SFilterSearchBox>, FilterSearchBox)
+
+		/** The layout that determines how the filters are laid out */
+		SLATE_ARGUMENT(EFilterBarLayout, FilterBarLayout)
+
+		/** If true, allow dynamically changing the orientation and saving in the config */
+		SLATE_ARGUMENT(bool, CanChangeOrientation)
+
+		/** Determines how each individual filter pill looks like */
+		SLATE_ARGUMENT(EFilterPillStyle, FilterPillStyle)
+
+		/** Whether to use submenus or sections for categories in the filter menu */
+		SLATE_ARGUMENT(bool, UseSectionsForCategories)
+
+	SLATE_END_ARGS()
+
+	/** Constructs this widget with InArgs */
+	void Construct(const FArguments& InArgs, const TSharedRef<FSequencerFilterBar>& InFilterBar);
+
+	//~ Begin SWidget
+	virtual FReply OnMouseButtonUp(const FGeometry& MyGeometry, const FPointerEvent& MouseEvent) override;
+	//~ End SWidget
+
+	TSharedPtr<FSequencerFilterBar> GetFilterBar() const;
+
+	void SetTextFilterString(const FString& InFilterString);
+
+	FText GetFilterErrorText() const;
+
+	EFilterBarLayout GetLayout() const;
+	void SetLayout(const EFilterBarLayout InFilterBarLayout);
+
+	void AttachFilterSearchBox(const TSharedPtr<SFilterSearchBox>& InFilterSearchBox);
+
+	bool HasAnyFilterWidgets() const;
+
+	void CreateAddCustomTextFilterWindowFromSearch(const FText& InSearchText);
+
+	TSharedRef<FFilterCategory> GetClassTypeCategory() const;
+	TSharedRef<FFilterCategory> GetComponentTypeCategory() const;
+	TSharedRef<FFilterCategory> GetMiscCategory() const;
+
+	void OnOpenTextExpressionHelp();
+
+	void SaveCurrentFilterSetAsCustomTextFilter();
+
+	TWeakPtr<SFilterSearchBox> GetSearchBox() const;
+
+protected:
+	void AddWidgetToLayout(const TSharedRef<SWidget>& InWidget);
+	void RemoveWidgetFromLayout(const TSharedRef<SWidget>& InWidget);
+
+	TSharedPtr<SSequencerFilter> FindFilterWidget(const TSharedRef<FSequencerTrackFilter>& InFilter) const;
+
+	TSharedRef<SSequencerFilter> CreateAndAddFilterWidget(const TSharedRef<FSequencerTrackFilter>& InFilter);
+	void AddFilterWidget(const TSharedRef<SSequencerFilter>& InFilterWidget);
+
+	void RemoveFilterWidget(const TSharedRef<FSequencerTrackFilter>& InFilter, const bool ExecuteOnFilterChanged = true);
+	void RemoveFilterWidget(const TSharedRef<SSequencerFilter>& InFilterWidget);
+
+	UWorld* GetWorld() const;
+
+	void RemoveAllFilterWidgets();
+	void RemoveAllFilterWidgetsButThis(const TSharedRef<SSequencerFilter>& InFilterWidget);
+	void RemoveFilterWidgetAndUpdate(const TSharedRef<SSequencerFilter>& InFilterWidget);
+
+	void OnEnableAllGroupFilters(bool bEnableAll);
+	void OnNodeGroupFilterClicked(UMovieSceneNodeGroup* NodeGroup);
+
+	void OnFiltersChanged(const ESequencerFilterChange InChangeType, const TSharedRef<FSequencerTrackFilter>& InFilter);
+
+	void CreateFilterWidgetsFromConfig();
+
+	TSharedPtr<FSequencerFilterBar> FilterBar;
+
+	TWeakPtr<SFilterSearchBox> WeakSearchBox;
+
+	TSharedPtr<SWidgetSwitcher> FilterBoxWidget;
+	TSharedPtr<SWrapBox> HorizontalContainerWidget;
+	TSharedPtr<SScrollBox> VerticalContainerWidget;
+
+	EFilterBarLayout FilterBarLayout = EFilterBarLayout::Vertical;
+	bool bCanChangeOrientation = true;
+	EFilterPillStyle FilterPillStyle = EFilterPillStyle::Default;
+
+	TArray<TSharedRef<SSequencerFilter>> FilterWidgets;
+
+	TSharedPtr<SWindow> TextExpressionHelpWindow;
+
+	TSharedPtr<FSequencerFilterBarContextMenu> ContextMenu;
+};
diff --git a/Engine/Source/Editor/Sequencer/Private/Filters/Widgets/SSequencerFilterCheckBox.h b/Engine/Source/Editor/Sequencer/Private/Filters/Widgets/SSequencerFilterCheckBox.h
new file mode 100644
--- /dev/null
+++ b/Engine/Source/Editor/Sequencer/Private/Filters/Widgets/SSequencerFilterCheckBox.h
@@ -0,0 +1,48 @@
+﻿// Copyright Epic Games, Inc. All Rights Reserved.
+
+#pragma once
+
+#include "Widgets/Input/SCheckBox.h"
+
+class SSequencerFilterCheckBox : public SCheckBox
+{
+public:
+	DECLARE_DELEGATE_RetVal_TwoParams(FReply, FOnPointerEvent, const FGeometry& /*InMyGeometry*/, const FPointerEvent& /*InMouseEvent*/);
+
+	void SetOnMouseUp(const FOnPointerEvent& InNewOnMouseUp)
+	{
+		OnMouseUp = InNewOnMouseUp;
+	}
+
+	void SetOnDoubleClick(const FOnPointerEvent& InNewOnDoubleClick)
+	{
+		OnDoubleClick = InNewOnDoubleClick;
+	}
+
+protected:
+	//~ Begin SWidget
+
+	virtual FReply OnMouseButtonUp(const FGeometry& InMyGeometry, const FPointerEvent& InMouseEvent) override
+	{
+		FReply Reply = SCheckBox::OnMouseButtonUp(InMyGeometry, InMouseEvent);
+		if (OnMouseUp.IsBound())
+		{
+			Reply = OnMouseUp.Execute(InMyGeometry, InMouseEvent).ReleaseMouseCapture();
+		}
+		return Reply;
+	}
+
+	virtual FReply OnMouseButtonDoubleClick(const FGeometry& InMyGeometry, const FPointerEvent& InMouseEvent) override
+	{
+		if (OnDoubleClick.IsBound())
+		{
+			return OnDoubleClick.Execute(InMyGeometry, InMouseEvent);
+		}
+		return SCheckBox::OnMouseButtonDoubleClick(InMyGeometry, InMouseEvent);
+	}
+
+	//~ End SWidget
+
+	FOnPointerEvent OnDoubleClick;
+	FOnPointerEvent OnMouseUp;
+};
diff --git a/Engine/Source/Editor/Sequencer/Private/SSequencer.cpp b/Engine/Source/Editor/Sequencer/Private/SSequencer.cpp
--- a/Engine/Source/Editor/Sequencer/Private/SSequencer.cpp
+++ b/Engine/Source/Editor/Sequencer/Private/SSequencer.cpp
@@ -2,6 +2,14 @@
 
 #include "SSequencer.h"
 #include "Engine/Blueprint.h"
+#include "Filters/Filters/SequencerTrackFilter_CustomText.h"
+#include "Filters/Filters/SequencerTrackFilters.h"
+#include "Filters/Menus/SequencerViewOptionsMenu.h"
+#include "Filters/SequencerFilterBar.h"
+#include "Filters/SFilterSearchBox.h"
+#include "Filters/Widgets/SFilterBarIsolateHideShow.h"
+#include "Filters/Widgets/SSequencerCustomTextFilterDialog.h"
+#include "Filters/Widgets/SSequencerFilterBar.h"
 #include "MovieSceneSequence.h"
 #include "MVVM/CurveEditorExtension.h"
 #include "MVVM/ViewModels/CategoryModel.h"
@@ -42,6 +50,7 @@
 #include "Widgets/Layout/SScrollBorder.h"
 #include "Widgets/Input/SComboButton.h"
 #include "Widgets/Layout/SSeparator.h"
+#include "Widgets/Layout/SSplitter.h"
 #include "Widgets/Input/SSpinBox.h"
 #include "Widgets/Input/SCheckBox.h"
 #include "Widgets/Docking/SDockTab.h"
@@ -114,8 +123,6 @@
 #include "IPropertyRowGenerator.h"
 #include "Fonts/FontMeasure.h"
 #include "Compilation/MovieSceneCompiledDataManager.h"
-#include "SequencerTrackFilters.h"
-#include "SequencerTrackFilterExtension.h"
 #include "SequencerCustomizationManager.h"
 #include "EditorActorFolders.h"
 #include "Tracks/MovieSceneEventTrack.h"
@@ -157,8 +164,6 @@ void SSequencer::Construct(const FArguments& InArgs, TSharedRef<FSequencer> InSe
 	bUserIsSelecting = false;
 	CachedClampRange = TRange<double>::Empty();
 	CachedViewRange = TRange<double>::Empty();
-	
-	InitializeTrackFilters();
 
 	OnPlaybackRangeBeginDrag = InArgs._OnPlaybackRangeBeginDrag;
 	OnPlaybackRangeEndDrag = InArgs._OnPlaybackRangeEndDrag;
@@ -262,7 +267,11 @@ void SSequencer::Construct(const FArguments& InArgs, TSharedRef<FSequencer> InSe
 
 	TAttribute<FAnimatedRange> ViewRangeAttribute = InArgs._ViewRange;
 
-	GridPanel = ConstructTrackAreaGridPanel(InArgs, InSequencer, TimeSliderArgs);
+	GridPanel = ConstructTrackAreaGridPanel(InArgs, TimeSliderArgs);
+
+	ViewOptionsMenu = MakeShared<FSequencerViewOptionsMenu>();
+
+	GetFilterBar()->GetOnFiltersChanged().AddSP(this, &SSequencer::OnTrackFiltersChanged);
 
 	ChildSlot
 	[
@@ -432,14 +441,14 @@ void SSequencer::Construct(const FArguments& InArgs, TSharedRef<FSequencer> InSe
 
 		// Main content body
 		+ SVerticalBox::Slot()
+		.FillHeight(1.f)
 		[
 			SAssignNew(MainContentContainer, SBox)
-			[
-				ConstructMainContent()
-			]
 		]
 	];
 
+	RebuildForSidebar();
+
 	if (InSequencer->GetHostCapabilities().bSupportsCurveEditor)
 	{
 		TSharedPtr<FSequencerEditorViewModel> SequencerViewModel = InSequencer->GetViewModel()->CastThisShared<FSequencerEditorViewModel>();
@@ -455,7 +464,12 @@ void SSequencer::Construct(const FArguments& InArgs, TSharedRef<FSequencer> InSe
 	ResetBreadcrumbs();
 }
 
-TSharedRef<SWidget> SSequencer::ConstructMainContent()
+void SSequencer::RebuildForSidebar()
+{
+	MainContentContainer->SetContent(ConstructSidebarContent());
+}
+
+TSharedRef<SWidget> SSequencer::ConstructSidebarContent()
 {
 	const TSharedPtr<FSequencer> Sequencer = SequencerPtr.Pin();
 	check(Sequencer.IsValid());
@@ -467,6 +481,8 @@ TSharedRef<SWidget> SSequencer::ConstructMainContent()
 
 	TSharedPtr<SWidget> OutWidget;
 
+	FilterBarSplitterContainer = SNew(SBox);
+
 	// Create the details sidebar only once to avoid having to re-register drawers
 	if (bSupportsSidebar)
 	{
@@ -480,7 +496,7 @@ TSharedRef<SWidget> SSequencer::ConstructMainContent()
 				.OnStateChanged(this, &SSequencer::OnSidebarStateChanged)
 				.OnGetContent(FOnGetContent::CreateLambda([this]()
 					{
-						return ConstructGridPanel();
+						return FilterBarSplitterContainer.ToSharedRef();
 					}));
 		}
 
@@ -490,10 +506,12 @@ TSharedRef<SWidget> SSequencer::ConstructMainContent()
 	{
 		DetailsSidebar.Reset();
 	}
+	
+	RebuildFilterBarContent();
 
 	if (SidebarState.IsHidden() || !bSupportsSidebar)
 	{
-		OutWidget = ConstructGridPanel();
+		OutWidget = FilterBarSplitterContainer.ToSharedRef();
 	}
 	else if (SidebarState.IsVisible())
 	{
@@ -503,7 +521,73 @@ TSharedRef<SWidget> SSequencer::ConstructMainContent()
 	return OutWidget.ToSharedRef();
 }
 
-TSharedRef<SWidget> SSequencer::ConstructGridPanel()
+void SSequencer::RebuildFilterBarContent()
+{
+	FilterBarSplitterContainer->SetContent(ConstructFilterBarContent());
+}
+
+TSharedRef<SWidget> SSequencer::ConstructFilterBarContent()
+{
+	RebuildSearchAndFilterRow();
+	
+	if (!IsFilterBarVisible() || GetFilterBarLayout() == EFilterBarLayout::Horizontal)
+	{
+		return ConstructGridOverlayContent();
+	}
+
+	const TSharedRef<SSequencerFilterBar> FilterBarWidgetRef = FilterBarWidget.ToSharedRef();
+
+	return SNew(SSplitter)
+		.Orientation(Orient_Horizontal)
+		+ SSplitter::Slot()
+		.Value(TAttribute<float>::CreateLambda([this]()
+			{
+				const TSharedPtr<FSequencerFilterBar> FilterBar = GetFilterBar();
+				if (FilterBar.IsValid()
+					&& (FilterBar->HasEnabledFilter() || FilterBar->HasEnabledCustomTextFilters()))
+				{
+					return GetSequencerSettings()->GetLastFilterBarSizeCoefficient();
+				}
+				return 0.f;
+			}))
+		.OnSlotResized_Lambda([this](const float InNewCoefficient)
+			{
+				GetSequencerSettings()->SetLastFilterBarSizeCoefficient(InNewCoefficient);
+			})
+		.Expose(FilterSplitterSlot)
+		[
+			SNew(SVerticalBox)
+			+ SVerticalBox::Slot()
+			.AutoHeight()
+			[
+				// Vertical Filters List Header
+				SNew(SBorder)
+				.BorderImage(FAppStyle::Get().GetBrush(TEXT("Brushes.Header")))
+				.Padding(FMargin(8.f, 6.f))
+				.OnMouseButtonUp(FilterBarWidgetRef, &SSequencerFilterBar::OnMouseButtonUp)
+				.Content()
+				[
+					SNew(STextBlock)
+					.Text(LOCTEXT("FilterListVerticalHeader", "Filters"))
+					.TextStyle(FAppStyle::Get(), TEXT("ButtonText"))
+					.Font(FAppStyle::Get().GetFontStyle(TEXT("NormalFontBold")))
+				]
+			]
+			+ SVerticalBox::Slot()
+			.FillHeight(1.f)
+			.VAlign(VAlign_Fill)
+			[
+				FilterBarWidgetRef
+			]
+		]
+		+ SSplitter::Slot()
+		.Value(0.94f)
+		[
+			ConstructGridOverlayContent()
+		];
+}
+
+TSharedRef<SWidget> SSequencer::ConstructGridOverlayContent()
 {
 	TAttribute<float> FillCoefficient_0, FillCoefficient_1;
 	{
@@ -544,10 +628,13 @@ TSharedRef<SWidget> SSequencer::ConstructGridPanel()
 		];
 }
 
-TSharedRef<SGridPanel> SSequencer::ConstructTrackAreaGridPanel(const FArguments& InArgs, const TSharedRef<FSequencer>& InSequencer, const FTimeSliderArgs& InTimeSliderArgs)
+TSharedRef<SGridPanel> SSequencer::ConstructTrackAreaGridPanel(const FArguments& InArgs, const FTimeSliderArgs& InTimeSliderArgs)
 {
 	using namespace UE::Sequencer;
 
+	check(SequencerPtr.IsValid());
+	const TSharedRef<FSequencer> SequencerRef = SequencerPtr.Pin().ToSharedRef();
+
 	constexpr int32 Column0 = 0, Column1 = 1;
 	constexpr int32 Row0 = 0, Row1 = 1, Row2 = 2, Row3 = 3, Row4 = 4;
 	const FMargin ResizeBarPadding(4.f, 0, 0, 0);
@@ -579,16 +666,16 @@ TSharedRef<SGridPanel> SSequencer::ConstructTrackAreaGridPanel(const FArguments&
 
 	ISequencerWidgetsModule& SequencerWidgets = FModuleManager::Get().LoadModuleChecked<ISequencerWidgetsModule>(TEXT("SequencerWidgets"));
 
-	TSharedRef<SScrollBar> ScrollBar = SNew(SScrollBar)
+	ScrollBar = SNew(SScrollBar)
 		.Thickness(FVector2D(9.f, 9.f));
 
-	TSharedRef<SScrollBar> PinnedAreaScrollBar = SNew(SScrollBar)
+	PinnedAreaScrollBar = SNew(SScrollBar)
 		.Thickness(FVector2D(9.f, 9.f));
 
-	TSharedPtr<FSequencerEditorViewModel> SequencerViewModel = InSequencer->GetViewModel()->CastThisShared<FSequencerEditorViewModel>();
+	TSharedPtr<FSequencerEditorViewModel> SequencerViewModel = SequencerRef->GetViewModel()->CastThisShared<FSequencerEditorViewModel>();
 	SAssignNew(PinnedTrackArea, SSequencerTrackAreaView, SequencerViewModel->GetPinnedTrackArea(), TimeSliderControllerRef);
 	SAssignNew(PinnedTreeView, SSequencerOutlinerView, SequencerViewModel->GetOutliner(), PinnedTrackArea.ToSharedRef())
-		.Selection(InSequencer->GetViewModel()->GetSelection())
+		.Selection(SequencerRef->GetViewModel()->GetSelection())
 		.ExternalScrollbar(PinnedAreaScrollBar)
 		.Clipping(EWidgetClipping::ClipToBounds);
 
@@ -599,7 +686,7 @@ TSharedRef<SGridPanel> SSequencer::ConstructTrackAreaGridPanel(const FArguments&
 
 	SAssignNew(TrackArea, SSequencerTrackAreaView, SequencerViewModel->GetTrackArea(), TimeSliderControllerRef);
 	SAssignNew(TreeView, SSequencerOutlinerView, SequencerViewModel->GetOutliner(), TrackArea.ToSharedRef())
-		.Selection(InSequencer->GetViewModel()->GetSelection())
+		.Selection(SequencerRef->GetViewModel()->GetSelection())
 		.ExternalScrollbar(ScrollBar)
 		.Clipping(EWidgetClipping::ClipToBounds);
 
@@ -624,7 +711,7 @@ TSharedRef<SGridPanel> SSequencer::ConstructTrackAreaGridPanel(const FArguments&
 	bool bMirrorLabels = false;
 	TopTimeSlider = SequencerWidgets.CreateTimeSlider(TimeSliderControllerRef, bMirrorLabels);
 	bMirrorLabels = true;
-	TSharedRef<ITimeSlider> BottomTimeSlider = SequencerWidgets.CreateTimeSlider(TimeSliderControllerRef, TAttribute<EVisibility>(this, &SSequencer::GetBottomTimeSliderVisibility), bMirrorLabels);
+	BottomTimeSlider = SequencerWidgets.CreateTimeSlider(TimeSliderControllerRef, TAttribute<EVisibility>(this, &SSequencer::GetBottomTimeSliderVisibility), bMirrorLabels);
 	
 	// Create bottom time range slider
 	EShowRange Ranges = EShowRange(EShowRange::WorkingRange | EShowRange::ViewRange);
@@ -633,7 +720,7 @@ TSharedRef<SGridPanel> SSequencer::ConstructTrackAreaGridPanel(const FArguments&
 		Ranges |= EShowRange::PlaybackRange;
 	}
 
-	TSharedRef<ITimeSlider> BottomTimeRange = SequencerWidgets.CreateTimeRange(
+	BottomTimeRange = SequencerWidgets.CreateTimeRange(
 		FTimeRangeArgs(
 			Ranges,
 			TimeSliderControllerRef,
@@ -665,30 +752,7 @@ TSharedRef<SGridPanel> SSequencer::ConstructTrackAreaGridPanel(const FArguments&
 			.Padding(FMargin(CommonPadding*2, CommonPadding))
 			.Clipping(EWidgetClipping::ClipToBounds)
 			[
-				SNew(SHorizontalBox)
-				+ SHorizontalBox::Slot()
-				.AutoWidth()
-				.VAlign(VAlign_Center)
-				.Padding(FMargin(0.f, 0.f, CommonPadding, 0.f))
-				[
-					MakeAddButton()
-				]
-	
-				+ SHorizontalBox::Slot()
-				.VAlign(VAlign_Center)
-				[
-					SAssignNew(SearchBox, SSearchBox)
-					.HintText(LOCTEXT("SearchNodesHint", "Search Tracks"))
-					.OnTextChanged(this, &SSequencer::OnOutlinerSearchChanged)
-				]
-
-				+ SHorizontalBox::Slot()
-				.AutoWidth()
-				.VAlign(VAlign_Center)
-				.Padding(FMargin(0.f, 0.f, CommonPadding, 0.f))
-				[
-					MakeFilterButton()
-				]
+				SAssignNew(SearchAndFilterRow, SVerticalBox)
 			]
 		]
 
@@ -736,7 +800,7 @@ TSharedRef<SGridPanel> SSequencer::ConstructTrackAreaGridPanel(const FArguments&
 				+ SOverlay::Slot()
 				.HAlign(HAlign_Right)
 				[
-					PinnedAreaScrollBar
+					PinnedAreaScrollBar.ToSharedRef()
 				]
 			]
 
@@ -774,7 +838,7 @@ TSharedRef<SGridPanel> SSequencer::ConstructTrackAreaGridPanel(const FArguments&
 				+ SOverlay::Slot()
 				.HAlign(HAlign_Right)
 				[
-					ScrollBar
+					ScrollBar.ToSharedRef()
 				]
 
 				+ SOverlay::Slot()
@@ -785,13 +849,12 @@ TSharedRef<SGridPanel> SSequencer::ConstructTrackAreaGridPanel(const FArguments&
 					+ SHorizontalBox::Slot()
 					.AutoWidth()
 					[
-						SAssignNew(SequencerTreeFilterStatusBar, SSequencerTreeFilterStatusBar, InSequencer)
+						SAssignNew(SequencerTreeFilterStatusBar, SSequencerTreeFilterStatusBar, SequencerRef)
 						.Visibility(EVisibility::Hidden) // Initially hidden, visible on hover of the info button
 					]
 				]
 			]
 		]
-
 		// Info Button, Transport Controls and Current Frame
 		+ SGridPanel::Slot(Column0, Row4, SGridPanel::Layer(10))
 		[
@@ -936,7 +999,7 @@ TSharedRef<SGridPanel> SSequencer::ConstructTrackAreaGridPanel(const FArguments&
 		.VAlign(VAlign_Top)
 		[
 			// Transform box
-			SAssignNew(TransformBox, SSequencerTransformBox, SequencerPtr.Pin().ToSharedRef(), *SequencerSettings, NumericTypeInterface.ToSharedRef())
+			SAssignNew(TransformBox, SSequencerTransformBox, SequencerRef, *SequencerSettings, NumericTypeInterface.ToSharedRef())
 		]
 
 		+ SGridPanel::Slot(Column1, Row2, SGridPanel::Layer(40))
@@ -945,14 +1008,14 @@ TSharedRef<SGridPanel> SSequencer::ConstructTrackAreaGridPanel(const FArguments&
 		.VAlign(VAlign_Top)
 		[
 			// Stretch box
-			SAssignNew(StretchBox, SSequencerStretchBox, SequencerPtr.Pin().ToSharedRef(), *SequencerSettings, NumericTypeInterface.ToSharedRef())
+			SAssignNew(StretchBox, SSequencerStretchBox, SequencerRef, *SequencerSettings, NumericTypeInterface.ToSharedRef())
 		]
 
 		// debug vis
 		+ SGridPanel::Slot(Column1, Row3, SGridPanel::Layer(10))
 		.Padding(ResizeBarPadding)
 		[
-			SNew(SSequencerDebugVisualizer, InSequencer)
+			SNew(SSequencerDebugVisualizer, SequencerRef)
 			.ViewRange(FAnimatedRange::WrapAttribute(InArgs._ViewRange))
 			.Visibility(this, &SSequencer::GetDebugVisualizerVisibility)
 		]
@@ -970,11 +1033,11 @@ TSharedRef<SGridPanel> SSequencer::ConstructTrackAreaGridPanel(const FArguments&
 				SNew(SOverlay)
 				+ SOverlay::Slot()
 				[
-					BottomTimeSlider
+					BottomTimeSlider.ToSharedRef()
 				]
 				+ SOverlay::Slot()
 				[
-					BottomTimeRange
+					BottomTimeRange.ToSharedRef()
 				]
 			]
 		];
@@ -1064,14 +1127,9 @@ void SSequencer::BindCommands(TSharedRef<FUICommandList> SequencerCommandBinding
 		FExecuteAction::CreateSP(this, &SSequencer::OpenNodeGroupsManager)
 	);
 
-	TSharedPtr<FSequencer> Sequencer = SequencerPtr.Pin();
-
-	for (TSharedRef<FSequencerTrackFilter> TrackFilter : AllTrackFilters)
+	if (const TSharedPtr<FSequencerFilterBar> FilterBar = GetFilterBar())
 	{
-		if (TrackFilter->SupportsSequence(Sequencer->GetFocusedMovieSceneSequence()))
-		{
-			TrackFilter->BindCommands(SequencerCommandBindings, CurveEditorSharedBindings, SequencerPtr.Pin());
-		}
+		FilterBar->BindCommands();
 	}
 }
 
@@ -1120,46 +1178,6 @@ TSharedRef<INumericTypeInterface<double>> SSequencer::GetNumericTypeInterface()
 	return NumericTypeInterface.ToSharedRef();
 }
 
-void SSequencer::InitializeTrackFilters()
-{
-	// Add all built-in track filters here
-	AllTrackFilters.Add(MakeShared<FSequencerTrackFilter_AudioTracks>());
-	AllTrackFilters.Add(MakeShared<FSequencerTrackFilter_EventTracks>());
-	AllTrackFilters.Add(MakeShared<FSequencerTrackFilter_LevelVisibilityTracks>());
-	AllTrackFilters.Add(MakeShared<FSequencerTrackFilter_ParticleTracks>());
-	AllTrackFilters.Add(MakeShared<FSequencerTrackFilter_CinematicShotTracks>());
-	AllTrackFilters.Add(MakeShared<FSequencerTrackFilter_SubTracks>());
-
-	AllTrackFilters.Add(MakeShared<FSequencerTrackFilter_CameraObjects>());
-	AllTrackFilters.Add(MakeShared<FSequencerTrackFilter_LightObjects>());
-	AllTrackFilters.Add(MakeShared<FSequencerTrackFilter_SkeletalMeshObjects>());
-	AllTrackFilters.Add(MakeShared<FSequencerTrackFilter_Animated>());
-
-	// Add any global user-defined frontend filters
-	for (TObjectIterator<USequencerTrackFilterExtension> ExtensionIt(RF_NoFlags); ExtensionIt; ++ExtensionIt)
-	{
-		if (USequencerTrackFilterExtension* PotentialExtension = *ExtensionIt)
-		{
-			if (PotentialExtension->HasAnyFlags(RF_ClassDefaultObject) && !PotentialExtension->GetClass()->HasAnyClassFlags(CLASS_Deprecated | CLASS_Abstract))
-			{
-				// Grab the filters
-				TArray< TSharedRef<FSequencerTrackFilter> > ExtendedTrackFilters;
-				PotentialExtension->AddTrackFilterExtensions(ExtendedTrackFilters);
-				AllTrackFilters.Append(ExtendedTrackFilters);
-			}
-		}
-	}
-
-	for (TSharedRef<FSequencerTrackFilter> TrackFilter : AllTrackFilters)
-	{
-		const bool bEnabled = GetSequencerSettings()->IsTrackFilterEnabled(TrackFilter->GetDisplayName().ToString());
-		SetTrackFilterEnabled(TrackFilter->GetDisplayName(), bEnabled);
-	}
-
-	// Sort by display name
-	AllTrackFilters.Sort([](const TSharedRef<FSequencerTrackFilter>& LHS, const TSharedRef<FSequencerTrackFilter>& RHS) { return LHS->GetDisplayName().ToString() < RHS->GetDisplayName().ToString(); });
-}
-
 void SSequencer::UpdateOutlinerViewColumns()
 {
 	using namespace UE::Sequencer;
@@ -1301,47 +1319,124 @@ void SSequencer::HandleOutlinerNodeSelectionChanged()
 	}
 }
 
-TSharedRef<SWidget> SSequencer::MakeAddButton()
+TSharedRef<SWidget> SSequencer::ConstructSearchAndFilterRow()
 {
-	return SNew(SPositiveActionButton)
-	.OnGetMenuContent(this, &SSequencer::MakeAddMenu)
-	.Icon(FAppStyle::Get().GetBrush("Icons.Plus"))
-	.Text(LOCTEXT("Add", "Add"))
-	.IsEnabled_Lambda([this]() { return !SequencerPtr.Pin()->IsReadOnly(); });
-}
+	const TSharedPtr<FSequencer> Sequencer = SequencerPtr.Pin();
+	if (!Sequencer.IsValid())
+	{
+		return SNullWidget::NullWidget;
+	}
 
-TSharedRef<SWidget> SSequencer::MakeFilterButton()
-{
-	TSharedPtr<SLayeredImage> FilterImage = SNew(SLayeredImage)
-		.Image(FAppStyle::Get().GetBrush("Icons.Filter"))
-		.ColorAndOpacity(FSlateColor::UseForeground());
+	const TSharedPtr<FSequencerFilterBar> FilterBar = GetFilterBar();
+	if (!FilterBar.IsValid())
+	{
+		return SNullWidget::NullWidget;
+	}
 
-	const TAttribute<const FSlateBrush*> ModifiedIcon = TAttribute<const FSlateBrush*>::CreateLambda([this]()
-		{
-			for (TSharedRef<FSequencerTrackFilter> TrackFilter : AllTrackFilters)
-			{
-				if (IsTrackFilterActive(TrackFilter))
+	if (!SearchBox.IsValid())
+	{
+		SearchBox = SNew(SFilterSearchBox)
+			.AddMetaData<FTagMetaData>(FTagMetaData(TEXT("SequencerFilterSearch")))
+			.HintText(LOCTEXT("FilterSearch", "Search..."))
+			.ToolTipText(LOCTEXT("FilterSearchHint", "Type here to search (pressing enter selects the results)"))
+			.DelayChangeNotificationsWhileTyping(true)
+			.ShowSearchHistory(true)
+			.OnTextChanged(this, &SSequencer::OnOutlinerSearchChanged)
+			.OnTextCommitted(this, &SSequencer::OnOutlinerSearchCommitted)
+			.OnSaveSearchClicked(this, &SSequencer::OnOutlinerSearchSaved);
+	}
+
+	FilterBarWidget = Sequencer->GetFilterBar()->GenerateWidget(SearchBox, GetFilterBarLayout());
+
+	if (!FilterComboButtonWidget.IsValid())
+	{
+		FilterComboButtonWidget = FilterBar->MakeAddFilterButton();
+	}
+
+	return SNew(SHorizontalBox)
+
+		// Add Button
+		+ SHorizontalBox::Slot()
+		.AutoWidth()
+		.VAlign(VAlign_Center)
+		[
+			SNew(SPositiveActionButton)
+			.OnGetMenuContent(this, &SSequencer::MakeAddMenu)
+			.Icon(FAppStyle::Get().GetBrush("Icons.Plus"))
+			.Text(LOCTEXT("Add", "Add"))
+			.IsEnabled_Lambda([this]() { return !SequencerPtr.Pin()->IsReadOnly(); })
+		]
+
+		// Isolate / Hide / Show
+		+ SHorizontalBox::Slot()
+		.AutoWidth()
+		.VAlign(VAlign_Center)
+		.Padding(2.f, 0.f, 0.f, 0.f)
+		[
+			FilterBar->MakeIsolateHideShowPanel()
+		]
+
+		// Advanced Search Filter Combo Button
+		+ SHorizontalBox::Slot()
+		.AutoWidth()
+		.VAlign(VAlign_Center)
+		.Padding(2.f, 0.f, 0.f, 0.f)
+		[
+			FilterComboButtonWidget.ToSharedRef()
+		]
+
+		// Advanced Search Box
+		+ SHorizontalBox::Slot()
+		.VAlign(VAlign_Center)
+		.Padding(2.f, 0.f, 0.f, 0.f)
+		[
+			SearchBox.ToSharedRef()
+		]
+
+		// Create Folder Button
+		+ SHorizontalBox::Slot()
+		.AutoWidth()
+		.VAlign(VAlign_Center)
+		.Padding(2.f, 0.f, 0.f, 0.f)
+		[
+			SNew(SButton)
+			.ContentPadding(2.f)
+			.HAlign(HAlign_Center)
+			.VAlign(VAlign_Center)
+			.ButtonStyle(FAppStyle::Get(), TEXT("SimpleButton"))
+			.ToolTipText(LOCTEXT("CreateFolderToolTip", "Creates a new folder containing the current selection or adds a new root folder if there is no selection"))
+			.OnClicked(this, &SSequencer::HandleCreateFolderTrackClick)
+			[
+				SNew(SImage)
+				.ColorAndOpacity(FSlateColor::UseForeground())
+				.Image(FAppStyle::Get().GetBrush(TEXT("SceneOutliner.NewFolderIcon")))
+			]
+		]
+
+		// View Options Combo Button
+		+ SHorizontalBox::Slot()
+		.AutoWidth()
+		.VAlign(VAlign_Center)
+		.Padding(2.f, 0.f, 0.f, 0.f)
+		[
+			SAssignNew(ViewOptionsComboButton, SComboButton)
+			.ContentPadding(2.f)
+			.HAlign(HAlign_Center)
+			.VAlign(VAlign_Center)
+			.ToolTipText(LOCTEXT("ViewOptionsToolTip", "View Options"))
+			.ComboButtonStyle(FAppStyle::Get(), TEXT("SimpleComboButtonWithIcon"))
+			.OnGetMenuContent_Lambda([this]() -> TSharedRef<SWidget>
 				{
-					return FAppStyle::Get().GetBrush("Icons.BadgeModified");
-				}
-			}
-			return (const FSlateBrush*)nullptr;
-		});
-
-	FilterImage->AddLayer(ModifiedIcon);
-
-	return SNew(SComboButton)
-	.ComboButtonStyle(&FAppStyle::Get().GetWidgetStyle<FComboButtonStyle>("SimpleComboButtonWithIcon"))
-	.ForegroundColor(FSlateColor::UseStyle())
-	.ContentPadding(0)
-	.ToolTipText(LOCTEXT("AddTrackFilterToolTip", "Add a track filter."))
-	.OnGetMenuContent(this, &SSequencer::MakeFilterMenu)
-	.ContentPadding(FMargin(1, 0))
-	.HasDownArrow(false)
-	.ButtonContent()
-	[
-		FilterImage.ToSharedRef()
-	];
+					return ViewOptionsMenu->CreateMenu(SequencerPtr);
+				})
+			.HasDownArrow(false)
+			.ButtonContent()
+			[
+				SNew(SImage)
+				.ColorAndOpacity(FSlateColor::UseForeground())
+				.Image(FAppStyle::Get().GetBrush(TEXT("Icons.Settings")))
+			]
+		];
 }
 
 TSharedRef<SWidget> SSequencer::MakeToolBar()
@@ -1728,351 +1823,6 @@ TSharedRef<SWidget> SSequencer::MakeAddMenu()
 	return MenuBuilder.MakeWidget();
 }
 
-TSharedRef<SWidget> SSequencer::MakeFilterMenu()
-{
-	TSharedPtr<FExtender> Extender = FExtender::Combine(AddMenuExtenders);
-	FMenuBuilder MenuBuilder(false, SequencerPtr.Pin()->GetCommandBindings(), Extender);
-
-	// let track editors & object bindings populate the menu
-	TSharedPtr<FSequencer> Sequencer = SequencerPtr.Pin();
-
-	MenuBuilder.BeginSection("SequencerTracksResetFilters");
-	{
-		MenuBuilder.AddMenuEntry(
-			LOCTEXT("FilterListEnableAll", "Enable All"),
-			LOCTEXT("FilterListEnableAllToolTip", "Selects all filters"),
-			FSlateIcon(),
-			FUIAction(FExecuteAction::CreateSP(this, &SSequencer::OnEnableAllFilters)));
-
-		MenuBuilder.AddMenuEntry(FSequencerCommands::Get().ResetFilters);
-	}
-	MenuBuilder.EndSection();
-
-	UMovieSceneSequence* FocusedMovieSequence = Sequencer->GetFocusedMovieSceneSequence();
-	UMovieScene* FocusedMovieScene = nullptr;
-	if (IsValid(FocusedMovieSequence))
-	{
-		FocusedMovieScene = FocusedMovieSequence->GetMovieScene();
-		if (IsValid(FocusedMovieScene))
-		{
-			if (FocusedMovieScene->GetNodeGroups().Num() > 0)
-			{
-				MenuBuilder.BeginSection("NodeGroupFilters");
-				MenuBuilder.AddSubMenu(LOCTEXT("NodeGroupFilters", "Group Filters"), LOCTEXT("NodeGroupFiltersToolTip", "Enables using this group as a filter for all items"), FNewMenuDelegate::CreateRaw(this, &SSequencer::FillNodeGroupsFilterMenu), false, FSlateIcon(), false);
-				MenuBuilder.EndSection();
-			}
-		}
-	}
-
-	UObject* PlaybackContext = Sequencer->GetPlaybackContext();
-	UWorld* World = PlaybackContext ? PlaybackContext->GetWorld() : nullptr;
-
-	if (World && World->GetLevels().Num() > 1)
-	{
-		MenuBuilder.BeginSection("TrackLevelFilters");
-		MenuBuilder.AddSubMenu(LOCTEXT("LevelFilters", "Level Filters"), LOCTEXT("LevelFiltersToolTip", "Filter object tracks by level"), FNewMenuDelegate::CreateRaw(this, &SSequencer::FillLevelFilterMenu), false, FSlateIcon(), false);
-		MenuBuilder.EndSection();
-	}
-
-	MenuBuilder.BeginSection("TrackFilters");
-	
-	for (TSharedRef<FSequencerTrackFilter> TrackFilter : AllTrackFilters)
-	{
-		if (TrackFilter->SupportsSequence(Sequencer->GetFocusedMovieSceneSequence()))
-		{
-			MenuBuilder.AddMenuEntry(
-				TrackFilter->GetDisplayName(),
-				TrackFilter->GetToolTipText(),
-				TrackFilter->GetIcon(),
-				FUIAction(
-				FExecuteAction::CreateSP(this, &SSequencer::OnTrackFilterClicked, TrackFilter),
-				FCanExecuteAction(),
-				FIsActionChecked::CreateSP(this, &SSequencer::IsTrackFilterActive, TrackFilter)),
-				NAME_None,
-				EUserInterfaceActionType::ToggleButton
-			);
-		}
-	}
-
-	MenuBuilder.EndSection();
-
-	return MenuBuilder.MakeWidget();
-}
-
-void SSequencer::FillNodeGroupsFilterMenu(FMenuBuilder& InMenuBarBuilder)
-{
-	TSharedPtr<FSequencer> Sequencer = SequencerPtr.Pin();
-	UMovieSceneSequence* FocusedMovieSequence = Sequencer->GetFocusedMovieSceneSequence();
-	UMovieScene* FocusedMovieScene = nullptr;
-
-	if (IsValid(FocusedMovieSequence))
-	{
-		FocusedMovieScene = FocusedMovieSequence->GetMovieScene();
-		if (IsValid(FocusedMovieScene))
-		{
-			InMenuBarBuilder.BeginSection("SequencerTracksResetNodeGroupFilters");
-
-			InMenuBarBuilder.AddMenuEntry(
-				LOCTEXT("EnableAllNodeGroupFilters", "Enable All"),
-				LOCTEXT("EnableAllNodeGroupFiltersToolTip", "Enables all group filters"),
-				FSlateIcon(),
-				FUIAction(FExecuteAction::CreateSP(this, &SSequencer::OnEnableAllNodeGroupFilters, true)));
-
-			InMenuBarBuilder.AddMenuEntry(
-				LOCTEXT("ResetNodeGroupFilters", "Reset Filters"),
-				LOCTEXT("ResetNodeGroupFiltersToolTip", "Disables all group filters"),
-				FSlateIcon(),
-				FUIAction(FExecuteAction::CreateSP(this, &SSequencer::OnEnableAllNodeGroupFilters, false))
-			);
-			
-			InMenuBarBuilder.EndSection();
-
-			bool bIsReadOnly = GetIsSequenceReadOnly();
-
-			InMenuBarBuilder.BeginSection("SequencerTracksNodeGroupFilters");
-			for (UMovieSceneNodeGroup* NodeGroup : FocusedMovieScene->GetNodeGroups())
-			{
-				InMenuBarBuilder.AddMenuEntry(
-					FText::FromName(NodeGroup->GetName()),
-					FText::FromName(NodeGroup->GetName()),
-					FSlateIcon(),
-					FUIAction(
-						FExecuteAction::CreateSP(this, &SSequencer::OnNodeGroupFilterClicked, NodeGroup),
-						FCanExecuteAction::CreateLambda([bIsReadOnly]() { return !bIsReadOnly; }),
-						FIsActionChecked::CreateLambda([NodeGroup]() { return NodeGroup->GetEnableFilter(); })),
-					NAME_None,
-					EUserInterfaceActionType::ToggleButton
-				);
-
-			}
-			InMenuBarBuilder.EndSection();
-		}
-	}
-}
-
-void SSequencer::OnEnableAllNodeGroupFilters(bool bEnableAll)
-{
-	TSharedPtr<FSequencer> Sequencer = SequencerPtr.Pin();
-	UMovieSceneSequence* FocusedMovieSequence = Sequencer->GetFocusedMovieSceneSequence();
-	UMovieScene* FocusedMovieScene = nullptr;
-	if (IsValid(FocusedMovieSequence))
-	{
-		FocusedMovieScene = FocusedMovieSequence->GetMovieScene();
-		if (IsValid(FocusedMovieScene))
-		{
-			for (UMovieSceneNodeGroup* NodeGroup : FocusedMovieScene->GetNodeGroups())
-			{
-				NodeGroup->SetEnableFilter(bEnableAll);
-			}
-		}
-	}
-}
-
-void SSequencer::OnNodeGroupFilterClicked(UMovieSceneNodeGroup* NodeGroup)
-{
-	if (NodeGroup)
-	{
-		NodeGroup->SetEnableFilter(!NodeGroup->GetEnableFilter());
-	}
-	
-}
-
-void SSequencer::FillLevelFilterMenu(FMenuBuilder& InMenuBarBuilder)
-{
-	TSharedPtr<FSequencer> Sequencer = SequencerPtr.Pin();
-	UObject* PlaybackContext = Sequencer->GetPlaybackContext();
-	UWorld* World = PlaybackContext ? PlaybackContext->GetWorld() : nullptr;
-
-	if (World)
-	{
-		const TArray<ULevel*> Levels = World->GetLevels();
-
-		if (Levels.Num() > 0)
-		{
-			InMenuBarBuilder.BeginSection("SequencerTracksResetLevelFilters");
-
-			InMenuBarBuilder.AddMenuEntry(
-				LOCTEXT("EnableAllLevelFilters", "Enable All"),
-				LOCTEXT("EnableAllLevelFiltersToolTip", "Enables all level filters"),
-				FSlateIcon(),
-				FUIAction(FExecuteAction::CreateSP(this, &SSequencer::OnEnableAllLevelFilters, true)));
-
-			InMenuBarBuilder.AddMenuEntry(
-				LOCTEXT("ResetLevelFilters", "Reset Filters"),
-				LOCTEXT("ResetLevelFiltersToolTip", "Resets current level filter selection"),
-				FSlateIcon(),
-				FUIAction(FExecuteAction::CreateSP(this, &SSequencer::OnEnableAllLevelFilters, false)));
-
-			InMenuBarBuilder.EndSection();
-		}
-
-		InMenuBarBuilder.BeginSection("SequencerTracksLevelFilters");
-		for (ULevel* Level : Levels)
-		{
-			FString LevelName = FPackageName::GetShortName(Level->GetOutermost()->GetName());
-			InMenuBarBuilder.AddMenuEntry(
-				FText::FromString(LevelName),
-				FText::FromString(Level->GetOutermost()->GetName()),
-				FSlateIcon(),
-				FUIAction(
-					FExecuteAction::CreateSP(this, &SSequencer::OnTrackLevelFilterClicked, LevelName),
-					FCanExecuteAction(),
-					FIsActionChecked::CreateSP(this, &SSequencer::IsTrackLevelFilterActive, LevelName)),
-				NAME_None,
-				EUserInterfaceActionType::ToggleButton
-			);
-		}
-		InMenuBarBuilder.EndSection();
-	}
-}
-
-void SSequencer::SetTrackFilterEnabled(const FText& InTrackFilterName, bool bEnabled)
-{
-	TSharedPtr<FSequencer> Sequencer = SequencerPtr.Pin();
-
-	for (TSharedRef<FSequencerTrackFilter> TrackFilter : AllTrackFilters)
-	{
-		if (TrackFilter->SupportsSequence(Sequencer->GetFocusedMovieSceneSequence()))
-		{
-			if (InTrackFilterName.EqualToCaseIgnored(TrackFilter->GetDisplayName()))
-			{
-				if (bEnabled && !IsTrackFilterActive(TrackFilter))
-				{
-					Sequencer->GetNodeTree()->AddFilter(TrackFilter);
-				}
-				else if (!bEnabled && IsTrackFilterActive(TrackFilter))
-				{
-					Sequencer->GetNodeTree()->RemoveFilter(TrackFilter);
-				}
-			}
-		}
-	}
-}
-
-bool SSequencer::IsTrackFilterEnabled(const FText& InTrackFilterName) const
-{
-	TSharedPtr<FSequencer> Sequencer = SequencerPtr.Pin();
-
-	for (TSharedRef<FSequencerTrackFilter> TrackFilter : AllTrackFilters)
-	{
-		if (TrackFilter->SupportsSequence(Sequencer->GetFocusedMovieSceneSequence()))
-		{
-			if (InTrackFilterName.EqualToCaseIgnored(TrackFilter->GetDisplayName()))
-			{
-				return Sequencer->GetNodeTree()->IsTrackFilterActive(TrackFilter);
-			}
-		}
-	}
-
-	return false;
-}
-
-TArray<FText> SSequencer::GetTrackFilterNames() const
-{
-	TSharedPtr<FSequencer> Sequencer = SequencerPtr.Pin();
-
-	TArray<FText> Filters;
-	for (TSharedRef<FSequencerTrackFilter> TrackFilter : AllTrackFilters)
-	{
-		if (TrackFilter->SupportsSequence(Sequencer->GetFocusedMovieSceneSequence()))
-		{
-			Filters.Add(TrackFilter->GetDisplayName());
-		}
-	}
-
-	return Filters;
-}
-
-void SSequencer::ResetFilters()
-{
-	TSharedPtr<FSequencer> Sequencer = SequencerPtr.Pin();
-	Sequencer->GetNodeTree()->RemoveAllFilters();
-	OnEnableAllNodeGroupFilters(false);
-}
-
-void SSequencer::OnEnableAllFilters()
-{
-	TSharedPtr<FSequencer> Sequencer = SequencerPtr.Pin();
-
-	for (TSharedRef<FSequencerTrackFilter> TrackFilter : AllTrackFilters)
-	{
-		if (TrackFilter->SupportsSequence(Sequencer->GetFocusedMovieSceneSequence()))
-		{
-			if (!Sequencer->GetNodeTree()->IsTrackFilterActive(TrackFilter))
-			{
-				Sequencer->GetNodeTree()->AddFilter(TrackFilter);
-			}
-		}
-	}
-	OnEnableAllNodeGroupFilters(true);
-	OnEnableAllLevelFilters(true);
-}
-
-void SSequencer::OnTrackFilterClicked(TSharedRef<FSequencerTrackFilter> TrackFilter)
-{
-	TSharedPtr<FSequencer> Sequencer = SequencerPtr.Pin();
-	if (IsTrackFilterActive(TrackFilter))
-	{
-		Sequencer->GetNodeTree()->RemoveFilter(TrackFilter);
-	}
-	else
-	{
-		Sequencer->GetNodeTree()->AddFilter(TrackFilter);
-	}
-}
-
-bool SSequencer::IsTrackFilterActive(TSharedRef<FSequencerTrackFilter> TrackFilter) const
-{
-	TSharedPtr<FSequencer> Sequencer = SequencerPtr.Pin();
-	return Sequencer->GetNodeTree()->IsTrackFilterActive(TrackFilter);
-}
-
-void SSequencer::OnEnableAllLevelFilters(bool bEnableAll)
-{
-	TSharedPtr<FSequencer> Sequencer = SequencerPtr.Pin();
-	UObject* PlaybackContext = Sequencer->GetPlaybackContext();
-	UWorld* World = PlaybackContext ? PlaybackContext->GetWorld() : nullptr;
-
-	if (World)
-	{
-		const TArray<ULevel*> Levels = World->GetLevels();
-
-		for (ULevel* Level : Levels)
-		{
-			FString LevelName = FPackageName::GetShortName(Level->GetOutermost()->GetName());
-
-			if (bEnableAll)
-			{
-				Sequencer->GetNodeTree()->AddLevelFilter(LevelName);
-			}
-			else
-			{
-				Sequencer->GetNodeTree()->RemoveLevelFilter(LevelName);
-			}
-		}
-	}
-}
-
-void SSequencer::OnTrackLevelFilterClicked(const FString LevelName)
-{
-	TSharedPtr<FSequencer> Sequencer = SequencerPtr.Pin();
-	if (IsTrackLevelFilterActive(LevelName))
-	{
-		Sequencer->GetNodeTree()->RemoveLevelFilter(LevelName);
-	}
-	else
-	{
-		Sequencer->GetNodeTree()->AddLevelFilter(LevelName);
-	}
-}
-
-bool SSequencer::IsTrackLevelFilterActive(const FString LevelName) const
-{
-	TSharedPtr<FSequencer> Sequencer = SequencerPtr.Pin();
-	return Sequencer->GetNodeTree()->IsTrackLevelFilterActive(LevelName);
-}
-
-
 TSharedRef<SWidget> SSequencer::MakeActionsMenu()
 {
 	using namespace UE::Sequencer;
@@ -2298,7 +2048,6 @@ TSharedRef<SWidget> SSequencer::MakeViewMenu()
 	{
 		MenuBuilder.AddMenuEntry( FSequencerCommands::Get().ToggleAutoScroll );
 		MenuBuilder.AddMenuEntry( FSequencerCommands::Get().ToggleShowRangeSlider );
-		MenuBuilder.AddMenuEntry( FSequencerCommands::Get().ToggleShowSelectedNodesOnly );
 		MenuBuilder.AddMenuEntry( FSequencerCommands::Get().ToggleLayerBars );
 		MenuBuilder.AddMenuEntry( FSequencerCommands::Get().ToggleKeyBars );
 		MenuBuilder.AddMenuEntry( FSequencerCommands::Get().ToggleChannelColors );
@@ -2341,17 +2090,6 @@ TSharedRef<SWidget> SSequencer::MakeViewMenu()
 	}
 	MenuBuilder.EndSection();
 
-	MenuBuilder.BeginSection("OrganizeAndSort", LOCTEXT("OrganizeAndSortHeader", "Organize and Sort"));
-	{
-		MenuBuilder.AddMenuEntry(FSequencerCommands::Get().ToggleAutoExpandNodesOnSelection);
-		MenuBuilder.AddMenuEntry(FSequencerCommands::Get().ToggleExpandCollapseNodes);
-		MenuBuilder.AddMenuEntry(FSequencerCommands::Get().ToggleExpandCollapseNodesAndDescendants);
-		MenuBuilder.AddMenuEntry(FSequencerCommands::Get().ExpandAllNodes);
-		MenuBuilder.AddMenuEntry(FSequencerCommands::Get().CollapseAllNodes);
-		MenuBuilder.AddMenuEntry(FSequencerCommands::Get().SortAllNodesAndDescendants);
-	}
-	MenuBuilder.EndSection();
-
 	MenuBuilder.BeginSection("MarkedFrames", LOCTEXT("MarkedFramesHeader", "Marked Frames"));
 	{
 		MenuBuilder.AddMenuEntry(FSequencerCommands::Get().ToggleShowMarkedFrames);
@@ -3278,24 +3016,48 @@ void SSequencer::PopBreadcrumb()
 	BreadcrumbTrail->PopCrumb();
 }
 
-void SSequencer::SetSearchText(const FText& InSearchText)
+FText SSequencer::GetSearchText() const
 {
-	SearchBox->SetText(InSearchText);
+	if (FilterBarWidget.IsValid())
+	{
+		if (const TSharedPtr<FSequencerFilterBar> FilterBar = FilterBarWidget->GetFilterBar())
+		{
+			return FText::FromString(FilterBar->GetTextFilterString());
+		}
+	}
+	return FText::GetEmpty();
 }
 
-void SSequencer::OnOutlinerSearchChanged( const FText& Filter )
+void SSequencer::SetSearchText(const FText& InSearchText)
 {
-	TSharedPtr<FSequencer> Sequencer = SequencerPtr.Pin();
-	if ( Sequencer.IsValid())
+	if (const TSharedPtr<FSequencer> Sequencer = SequencerPtr.Pin())
 	{
-		const FString FilterString = Filter.ToString();
-
-		Sequencer->GetNodeTree()->FilterNodes( FilterString );
+		Sequencer->GetNodeTree()->SetTextFilterString(InSearchText.ToString());
 
 		TreeView->Refresh();
 	}
 }
 
+void SSequencer::OnOutlinerSearchChanged(const FText& InFilter)
+{
+	SetSearchText(InFilter);
+}
+
+void SSequencer::OnOutlinerSearchCommitted(const FText& InFilter, ETextCommit::Type InCommitInfo)
+{
+	SetSearchText(InFilter);
+}
+
+void SSequencer::OnOutlinerSearchSaved(const FText& InFilterText)
+{
+	if (FilterBarWidget.IsValid())
+	{
+		FCustomTextFilterData CustomTextFilterData;
+		CustomTextFilterData.FilterString = InFilterText;
+		SSequencerCustomTextFilterDialog::CreateWindow_AddCustomTextFilter(FilterBarWidget->GetFilterBar().ToSharedRef(), CustomTextFilterData);
+	}
+}
+
 void SSequencer::OnDragEnter( const FGeometry& MyGeometry, const FDragDropEvent& DragDropEvent )
 {
 	// @todo sequencer: Add drop validity cue
@@ -3421,9 +3183,20 @@ FReply SSequencer::OnKeyDown( const FGeometry& MyGeometry, const FKeyEvent& InKe
 {
 	// A toolkit tab is active, so direct all command processing to it
 	TSharedPtr<FSequencer> SequencerPin = SequencerPtr.Pin();
-	if ( SequencerPin.IsValid() && SequencerPin->GetCommandBindings()->ProcessCommandBindings( InKeyEvent ) )
+	if (SequencerPin.IsValid())
+	{
+		if (SequencerPin->GetCommandBindings()->ProcessCommandBindings(InKeyEvent))
+		{
+			return FReply::Handled();
+		}
+	}
+
+	if (const TSharedPtr<FSequencerFilterBar> FilterBar = FilterBarWidget->GetFilterBar())
 	{
-		return FReply::Handled();
+		if (FilterBar->GetCommandList()->ProcessCommandBindings(InKeyEvent))
+		{
+			return FReply::Handled();
+		}
 	}
 
 	return FReply::Unhandled();
@@ -3776,7 +3549,7 @@ FText SSequencer::GetBreadcrumbTextForSequence(TWeakObjectPtr<UMovieSceneSequenc
 {
 	UMovieSceneSequence* SequencePtr = Sequence.Get();
 
-	bool bIsDirty = SequencePtr->GetMovieScene()->GetOuter()->GetOutermost()->IsDirty();
+	bool bIsDirty = SequencePtr->GetMovieScene()->GetOuter()->GetPackage()->IsDirty();
 
 	if (bIsActive)
 	{
@@ -4342,6 +4115,7 @@ USequencerSettings* SSequencer::GetSequencerSettings() const
 	return nullptr;
 }
 
+
 bool SSequencer::RegisterDrawer(FSidebarDrawerConfig&& InDrawerConfig)
 {
 	if (DetailsSidebar.IsValid())
@@ -4409,7 +4183,7 @@ void SSequencer::SetSidebarVisible(const bool bInVisible)
 	SidebarState.SetVisible(bInVisible);
 	SequencerSettings->SetSidebarState(SidebarState);
 
-	MainContentContainer->SetContent(ConstructMainContent());
+	RebuildForSidebar();
 }
 
 void SSequencer::ToggleSidebarVisible()
@@ -4479,4 +4253,136 @@ void SSequencer::EnableCurveEditorPendingFocusOnHovering(const bool InEnabled) c
 	CurveEditorPanel->EnablePendingFocusOnHovering(InEnabled);
 }
 
+FReply SSequencer::HandleCreateFolderTrackClick()
+{
+	if (const TSharedPtr<FSequencer> Sequencer = SequencerPtr.Pin())
+	{
+		if (Sequencer->GetSelectedNodesToMove().Num() > 0)
+		{
+			Sequencer->MoveSelectedNodesToNewFolder();
+		}
+		else
+		{
+			Sequencer->AddFolder();
+		}
+	}
+	return FReply::Handled();
+}
+
+TSharedPtr<FSequencerFilterBar> SSequencer::GetFilterBar() const
+{
+	if (const TSharedPtr<FSequencer> Sequencer = SequencerPtr.Pin())
+	{
+		return Sequencer->GetFilterBar();
+	}
+	return nullptr;
+}
+
+bool SSequencer::IsFilterBarVisible() const
+{
+	if (!FilterBarWidget.IsValid())
+	{
+		return false;
+	}
+
+	USequencerSettings* const SequencerSettings = GetSequencerSettings();
+	if (ensure(SequencerSettings))
+	{
+		if (!SequencerSettings->IsFilterBarVisible())
+		{
+			return false;
+		}
+	}
+
+	const TSharedPtr<FSequencerFilterBar> FilterBar = FilterBarWidget->GetFilterBar();
+	if (!FilterBar.IsValid())
+	{
+		return false;
+	}
+
+	return FilterBar->HasAnyFiltersEnabled();
+}
+
+void SSequencer::ToggleFilterBarVisibility()
+{
+	bool bNewFilterBarVisible = true;
+
+	USequencerSettings* const SequencerSettings = GetSequencerSettings();
+	if (ensure(SequencerSettings))
+	{
+		bNewFilterBarVisible = !SequencerSettings->IsFilterBarVisible();
+		SequencerSettings->SetFilterBarVisible(bNewFilterBarVisible);
+	}
+
+	if (FilterBarWidget.IsValid())
+	{
+		const EVisibility NewVisibility = bNewFilterBarVisible ? EVisibility::Visible : EVisibility::Collapsed;
+		FilterBarWidget->SetVisibility(NewVisibility);
+	}
+
+	RebuildFilterBarContent();
+}
+
+TSharedPtr<SSequencerFilterBar> SSequencer::GetFilterBarWidget() const
+{
+	return FilterBarWidget;
+}
+
+EFilterBarLayout SSequencer::GetFilterBarLayout() const
+{
+	USequencerSettings* const SequencerSettings = GetSequencerSettings();
+	if (ensure(SequencerSettings))
+	{
+		return SequencerSettings->GetFilterBarLayout();
+	}
+	return EFilterBarLayout::Vertical;
+}
+
+void SSequencer::SetFilterBarLayout(const EFilterBarLayout InLayout)
+{
+	USequencerSettings* const SequencerSettings = GetSequencerSettings();
+	if (ensure(SequencerSettings))
+	{
+		SequencerSettings->SetFilterBarLayout(InLayout);
+	}
+
+	RebuildFilterBarContent();
+}
+
+void SSequencer::OnTrackFiltersChanged(const ESequencerFilterChange InChangeType, const TSharedRef<FSequencerTrackFilter>& InFilter)
+{
+	RebuildFilterBarContent();
+
+	if (SequencerTreeFilterStatusBar.IsValid())
+	{
+		SequencerTreeFilterStatusBar->UpdateText();
+	}
+}
+
+void SSequencer::RebuildSearchAndFilterRow()
+{
+	SearchAndFilterRow->ClearChildren();
+
+	SearchAndFilterRow->AddSlot()
+		.AutoHeight()
+		[
+			ConstructSearchAndFilterRow()
+		];
+
+	if (FilterBarWidget.IsValid())
+	{
+		if (const TSharedPtr<FSequencerFilterBar> FilterBar = FilterBarWidget->GetFilterBar())
+		{
+			if (IsFilterBarVisible() && GetFilterBarLayout() == EFilterBarLayout::Horizontal)
+			{
+				SearchAndFilterRow->AddSlot()
+					.AutoHeight()
+					[
+						FilterBarWidget.ToSharedRef()
+					];
+			}
+		}
+	}
+}
+
 #undef LOCTEXT_NAMESPACE
diff --git a/Engine/Source/Editor/Sequencer/Private/SSequencer.h b/Engine/Source/Editor/Sequencer/Private/SSequencer.h
--- a/Engine/Source/Editor/Sequencer/Private/SSequencer.h
+++ b/Engine/Source/Editor/Sequencer/Private/SSequencer.h
@@ -37,14 +37,19 @@ class SSequencerStretchBox;
 class SCurveEditorPanel;
 class SBox;
 class SDockTab;
+class SFilterSearchBox;
 class SGridPanel;
 class SSidebar;
 class SSidebarContainer;
 class SWindow;
 class USequencerSettings;
-class FSequencerTrackFilter;
 class SSequencerGroupManager;
 class SSequencerTreeFilterStatusBar;
+class SFilterSearchBox;
+class FSequencerFilterMenu;
+class FSequencerTrackFilter;
+class FSequencerViewOptionsMenu;
+enum class ESequencerFilterChange : uint8;
 struct FPaintPlaybackRangeArgs;
 struct FSequencerCustomizationInfo;
 struct FSidebarDrawerConfig;
@@ -356,17 +361,8 @@ public:
 	/** Sets the play time for the sequence. Will extend the working range if out of bounds. */
 	void SetPlayTime(double Frame);
 
-	/** Sets the specified track filter to be on or off */
-	void SetTrackFilterEnabled(const FText& InTrackFilterName, bool bEnabled);
-
-	/** Gets whether the specified track filter is on/off */
-	bool IsTrackFilterEnabled(const FText& InTrackFilterName) const;
-
-	/** Reset all enabled filters */
-	void ResetFilters();
-
-	/** Gets all the available track filter names */
-	TArray<FText> GetTrackFilterNames() const;
+	/** Gets the text to search by */
+	FText GetSearchText() const;
 
 	/** Sets the text to search by */
 	void SetSearchText(const FText& InSearchText);
@@ -395,13 +391,16 @@ private:
 
 	static constexpr float CommonPadding = 3.f;
 
-	TSharedRef<SWidget> ConstructMainContent();
-	TSharedRef<SWidget> ConstructGridPanel();
-	TSharedRef<SGridPanel> ConstructTrackAreaGridPanel(const FArguments& InArgs, const TSharedRef<FSequencer>& InSequencer, const FTimeSliderArgs& InTimeSliderArgs);
-
-	/** Initalizes a list of all track filter objects */
-	void InitializeTrackFilters();
+	TSharedRef<SWidget> ConstructSidebarContent();
+	TSharedRef<SWidget> ConstructFilterBarContent();
+	TSharedRef<SWidget> ConstructSearchAndFilterRow();
+	TSharedRef<SWidget> ConstructGridOverlayContent();
+	TSharedRef<SGridPanel> ConstructTrackAreaGridPanel(const FArguments& InArgs, const FTimeSliderArgs& InTimeSliderArgs);
 
+	/** Can be called anytime after Construct() to rebuild most of the main content (except for the toolbar). */
+	void RebuildForSidebar();
+	void RebuildFilterBarContent();
+	void RebuildSearchAndFilterRow();
 
 	/** Initializes outliner column list from settings and SequencerCore */
 	void InitializeOutlinerColumns();
@@ -415,9 +414,6 @@ private:
 	/** Empty active timer to ensure Slate ticks during Sequencer playback */
 	EActiveTimerReturnType EnsureSlateTickDuringPlayback(double InCurrentTime, float InDeltaTime);	
 
-	/** Get context menu contents. */
-	void GetContextMenuContent(FMenuBuilder& MenuBuilder);
-
 	TWeakPtr<FSequencer> GetSequencer() { return SequencerPtr; }
 
 	static void PopulateToolBar(UToolMenu* InMenu);
@@ -425,17 +421,9 @@ private:
 	/** Makes the toolbar. */
 	TSharedRef<SWidget> MakeToolBar();
 
-	/** Makes add button. */
-	TSharedRef<SWidget> MakeAddButton();
-
-	/** Makes filter button */
-	TSharedRef<SWidget> MakeFilterButton();
-
 	/** Makes the add menu for the toolbar. */
 	TSharedRef<SWidget> MakeAddMenu();
 
-	TSharedRef<SWidget> MakeFilterMenu();
-
 	/** Makes the actions menu for the toolbar. */
 	TSharedRef<SWidget> MakeActionsMenu();
 
@@ -494,20 +482,6 @@ public:
 
 private:
 
-	void OnEnableAllFilters();
-	void OnTrackFilterClicked(TSharedRef<FSequencerTrackFilter> TrackFilter);
-	bool IsTrackFilterActive(TSharedRef<FSequencerTrackFilter> TrackFilter) const;
-
-	void OnEnableAllLevelFilters(bool bEnableAll);
-	void OnTrackLevelFilterClicked(const FString LevelName);
-	bool IsTrackLevelFilterActive(const FString LevelName) const;
-
-	void FillLevelFilterMenu(FMenuBuilder& InMenuBarBuilder);
-	void FillNodeGroupsFilterMenu(FMenuBuilder& InMenuBarBuilder);
-
-	void OnEnableAllNodeGroupFilters(bool bEnableAll);
-	void OnNodeGroupFilterClicked(UMovieSceneNodeGroup* NodeGroup);
-
 	/**
 	 * Called when any outliner column's visibily is modified.
 	 * Updates SequencerSettings and visible outliner columns in Outliner View.
@@ -536,7 +510,11 @@ private:
 	/**
 	 * Called when the outliner search terms change                                                              
 	 */
-	void OnOutlinerSearchChanged( const FText& Filter );
+	void OnOutlinerSearchChanged(const FText& InFilter);
+
+	void OnOutlinerSearchCommitted(const FText& InFilter, ETextCommit::Type InCommitInfo);
+	
+	void OnOutlinerSearchSaved(const FText& InFilterText);
 
 	/**
 	 * @return The fill percentage of the animation outliner
@@ -724,6 +702,16 @@ public:
 	/** Enable/disable pending focus in sequencer */
 	SEQUENCER_API void EnablePendingFocusOnHovering(const bool InEnabled);
 	
+	TSharedPtr<FSequencerFilterBar> GetFilterBar() const;
+
+	bool IsFilterBarVisible() const;
+	void ToggleFilterBarVisibility();
+
+	TSharedPtr<SSequencerFilterBar> GetFilterBarWidget() const;
+
+	EFilterBarLayout GetFilterBarLayout() const;
+	void SetFilterBarLayout(const EFilterBarLayout InLayout);
+
 private:
 
 	/** Pending focus handler */
@@ -737,6 +725,10 @@ private:
 
 	void OnSidebarStateChanged(const FSidebarState& InNewState);
 
+	FReply HandleCreateFolderTrackClick();
+
+	void OnTrackFiltersChanged(const ESequencerFilterChange InChangeType, const TSharedRef<FSequencerTrackFilter>& InFilter);
+
 	/** This is the main container widget for the Sequencer (minus the toolbar). This is the parent of
 	 * the GridPanel widget below and is used to rebuild content when the sidebar is added/removed. */
 	TSharedPtr<SBox> MainContentContainer;
@@ -749,6 +741,9 @@ private:
 	 * Built once at Sequencer initialize and used when rebuilding the main content container when the sidebar is added/removed. */
 	TSharedPtr<SVerticalBox> MainSequencerArea;
 
+	/** The splitter container that holds the filter bar (if visible). */
+	TSharedPtr<SBox> FilterBarSplitterContainer;
+
 	/** Transform box widget. */
 	TSharedPtr<SSequencerTransformBox> TransformBox;
 
@@ -771,7 +766,7 @@ private:
 	TSharedPtr<SBreadcrumbTrail<FSequencerBreadcrumb>> BreadcrumbTrail;
 
 	/** The search box for filtering tracks. */
-	TSharedPtr<SSearchBox> SearchBox;
+	TSharedPtr<SFilterSearchBox> SearchBox;
 
 	/** The current playback time display. */
 	TSharedPtr<STemporarilyFocusedSpinBox<double>> PlayTimeDisplay;
@@ -785,6 +780,9 @@ private:
 	/** The sequencer tree view for pinned tracks */
 	TSharedPtr<UE::Sequencer::SOutlinerView> PinnedTreeView;
 
+	TSharedPtr<SScrollBar> ScrollBar;
+	TSharedPtr<SScrollBar> PinnedAreaScrollBar;
+
 	/** Dropdown for selecting breadcrumbs */
 	TSharedPtr<class SComboButton> BreadcrumbPickerButton;
 
@@ -794,6 +792,11 @@ private:
 	/** The top time slider widget */
 	TSharedPtr<ITimeSlider> TopTimeSlider;
 
+	/** The bottom time slider widget */
+	TSharedPtr<ITimeSlider> BottomTimeSlider;
+
+	TSharedPtr<ITimeSlider> BottomTimeRange;
+
 	/** Container for the toolbar, so that we can re-create it as needed. */
 	TSharedPtr<SBox> ToolbarContainer;
 
@@ -892,9 +895,6 @@ private:
 
 	TWeakPtr<SWindow> WeakTickResolutionOptionsWindow;
 
-	/** All possible track filter objects */
-	TArray< TSharedRef<FSequencerTrackFilter> > AllTrackFilters;
-
 	TWeakPtr<SWindow> WeakExposedBindingsWindow;
 
 	TWeakPtr<SWindow> WeakNodeGroupWindow;
@@ -903,4 +903,17 @@ private:
 
 	TSharedPtr<SSidebarContainer> SidebarContainer;
 	TSharedPtr<SSidebar> DetailsSidebar;
+
+	TSharedPtr<SSequencerFilterBar> FilterBarWidget;
+	TSharedPtr<SComboButton> FilterComboButtonWidget;
+
+	/** The button that displays view options */
+	TSharedPtr<SComboButton> ViewOptionsComboButton;
+	TSharedPtr<FSequencerViewOptionsMenu> ViewOptionsMenu;
+
+	TSharedPtr<FSequencerFilterMenu> FilterMenu;
+	
+	SSplitter::FSlot* FilterSplitterSlot = nullptr;
+
+	TSharedPtr<SVerticalBox> SearchAndFilterRow;
 };
diff --git a/Engine/Source/Editor/Sequencer/Private/SSequencerTreeFilterStatusBar.cpp b/Engine/Source/Editor/Sequencer/Private/SSequencerTreeFilterStatusBar.cpp
--- a/Engine/Source/Editor/Sequencer/Private/SSequencerTreeFilterStatusBar.cpp
+++ b/Engine/Source/Editor/Sequencer/Private/SSequencerTreeFilterStatusBar.cpp
@@ -57,8 +57,7 @@ void SSequencerTreeFilterStatusBar::Construct(const FArguments& InArgs, TSharedP
 
 void SSequencerTreeFilterStatusBar::ClearFilters()
 {
-	TSharedPtr<FSequencer> Sequencer = WeakSequencer.Pin();
-	if (Sequencer)
+	if (const TSharedPtr<FSequencer> Sequencer = WeakSequencer.Pin())
 	{
 		Sequencer->ClearFilters();
 	}
@@ -66,16 +65,19 @@ void SSequencerTreeFilterStatusBar::ClearFilters()
 
 EVisibility SSequencerTreeFilterStatusBar::GetVisibilityFromFilter() const
 {
-	TSharedPtr<FSequencer> Sequencer = WeakSequencer.Pin();
-	return Sequencer && Sequencer->GetNodeTree()->HasActiveFilter() ? EVisibility::Visible : EVisibility::Collapsed;
+	if (const TSharedPtr<FSequencer> Sequencer = WeakSequencer.Pin())
+	{
+		return Sequencer->GetFilterInterface()->HasAnyFilterActive() ? EVisibility::Visible : EVisibility::Collapsed;
+	}
+	return EVisibility::Collapsed;
 }
 
 void SSequencerTreeFilterStatusBar::UpdateText()
 {
 	using namespace UE::Sequencer;
 
-	TSharedPtr<FSequencer> Sequencer = WeakSequencer.Pin();
-	if (!ensureAlways(Sequencer))
+	const TSharedPtr<FSequencer> Sequencer = WeakSequencer.Pin();
+	if (!Sequencer.IsValid())
 	{
 		return;
 	}
@@ -90,7 +92,7 @@ void SSequencerTreeFilterStatusBar::UpdateText()
 	NamedArgs.Add("Total", NodeTree->GetTotalDisplayNodeCount());
 
 	const bool bHasSelection = SelectedOutlinerItems.Num() != 0;
-	const bool bHasFilter = NodeTree->HasActiveFilter();
+	const bool bHasFilter = Sequencer->GetFilterInterface()->HasAnyFilterActive();
 	const int32 NumFiltered = NodeTree->GetFilteredDisplayNodeCount();
 
 	if (bHasSelection)
diff --git a/Engine/Source/Editor/Sequencer/Private/Sequencer.cpp b/Engine/Source/Editor/Sequencer/Private/Sequencer.cpp
--- a/Engine/Source/Editor/Sequencer/Private/Sequencer.cpp
+++ b/Engine/Source/Editor/Sequencer/Private/Sequencer.cpp
@@ -172,6 +172,7 @@
 #include "Bindings/MovieSceneSpawnableBinding.h"
 #include "Bindings/MovieSceneSpawnableActorBinding.h"
 #include "Bindings/MovieSceneReplaceableActorBinding.h"
+#include "Filters/SequencerFilterBar.h"
 
 #define LOCTEXT_NAMESPACE "Sequencer"
 
@@ -552,6 +553,8 @@ void FSequencer::InitSequencer(const FSequencerInitParams& InitParams, const TSh
 	UpdateTimeBases();
 	PlayPosition.Reset(ConvertFrameTime(GetPlaybackRange().GetLowerBoundValue(), GetRootTickResolution(), PlayPosition.GetInputRate()));
 
+	FilterBar->CreateCustomTextFiltersFromConfig();
+
 	// Make internal widgets
 	SequencerWidget = SNew( SSequencer, SharedThis( this ) )
 		.ViewRange( this, &FSequencer::GetViewRange )
@@ -638,7 +641,6 @@ void FSequencer::InitSequencer(const FSequencerInitParams& InitParams, const TSh
 		TimeUndoRedoHandler.SetSequencer(SharedThis(this));
 		TimeUndoRedoHandler.UndoRedoProxy->SetFlags(RF_Transactional | RF_Transient);
 	}
-	
 
 	// Update the view range to the new current time
 	UpdateTimeBoundsToFocusedMovieScene();
@@ -739,6 +741,7 @@ FSequencer::FSequencer()
 	, bUpdatingExternalSelection( false )
 	, bNeedsEvaluate(false)
 	, bNeedsInvalidateCachedData(false)
+	, FilterBar(MakeShared<FSequencerFilterBar>(*this))
 {
 	// Exposes the sequencer and curve editor command lists to subscribers from other systems
 	FInputBindingManager::Get().RegisterCommandList(FSequencerCommands::Get().GetContextName(), SequencerCommandBindings);
@@ -5674,6 +5677,11 @@ bool FSequencer::IsReadOnly() const
 	return bReadOnly || (GetFocusedMovieSceneSequence() && GetFocusedMovieSceneSequence()->GetMovieScene() && GetFocusedMovieSceneSequence()->GetMovieScene()->IsReadOnly());
 }
 
+void FSequencer::ExternalSelectionHasChanged()
+{
+	SynchronizeSequencerSelectionWithExternalSelection();
+}
+
 FGuid FSequencer::MakeNewSpawnable(UObject& Object, UActorFactory* ActorFactory, bool bSetupDefaults)
 {
 	const FScopedTransaction Transaction(LOCTEXT("UndoAddingObject", "Add Object to MovieScene"));
@@ -6401,7 +6409,7 @@ void FSequencer::AddNodesToExistingNodeGroup(TArrayView<const UE::Sequencer::TWe
 void FSequencer::ClearFilters()
 {
 	SequencerWidget->SetSearchText(FText::GetEmpty());
-	GetNodeTree()->RemoveAllFilters();
+	FilterBar->EnableFilters(false);
 	GetSequencerSettings()->SetShowSelectedNodesOnly(false);
 
 	UMovieSceneSequence* FocusedMovieSequence = GetFocusedMovieSceneSequence();
@@ -6751,6 +6759,8 @@ void FSequencer::SynchronizeSequencerSelectionWithExternalSelection()
 			}
 		}
 	}
+
+	FilterBar->RequestFilterUpdate();
 }
 
 void FSequencer::SelectNodesByPath(const TSet<FString>& NodePaths)
@@ -8516,19 +8526,26 @@ void FSequencer::ObjectImplicitlyRemoved(UObject* InObject) const
 	}
 }
 
-void FSequencer::SetTrackFilterEnabled(const FText& InTrackFilterName, bool bEnabled)
+void FSequencer::SetTrackFilterEnabled(const FText& InFilterName, bool bInEnabled)
 {
-	SequencerWidget->SetTrackFilterEnabled(InTrackFilterName, bEnabled);
+	if (const TSharedPtr<FSequencerTrackFilter> Filter = FilterBar->FindFilterByDisplayName(InFilterName.ToString()))
+	{
+		FilterBar->SetFilterEnabled(Filter.ToSharedRef(), bInEnabled, true);
+	}
 }
 
-bool FSequencer::IsTrackFilterEnabled(const FText& InTrackFilterName) const
+bool FSequencer::IsTrackFilterEnabled(const FText& InFilterName) const
 {
-	return SequencerWidget->IsTrackFilterEnabled(InTrackFilterName);
+	if (const TSharedPtr<FSequencerTrackFilter> Filter = FilterBar->FindFilterByDisplayName(InFilterName.ToString()))
+	{
+		return FilterBar->IsFilterEnabled(Filter.ToSharedRef());
+	}
+	return false;
 }
 
 TArray<FText> FSequencer::GetTrackFilterNames() const
 {
-	return SequencerWidget->GetTrackFilterNames();
+	return FilterBar->GetFilterDisplayNames();
 }
 
 void FSequencer::ToggleNodeLocked()
@@ -8983,7 +9000,7 @@ void FSequencer::OnLoadRecordedData()
 
 }
 
-void FSequencer::OnAddFolder()
+void FSequencer::AddFolder()
 {
 	using namespace UE::Sequencer;
 
@@ -9681,11 +9698,6 @@ void FSequencer::CollapseAllNodes()
 	SequencerWidget->GetTreeView()->ToggleExpandCollapseNodes(ETreeRecursion::Recursive, bExpandAll, bCollapseAll);
 }
 
-void FSequencer::ResetFilters()
-{
-	SequencerWidget->ResetFilters();
-}
-
 void FSequencer::AddSelectedActors()
 {
 	USelection* ActorSelection = GEditor->GetSelectedActors();
@@ -10778,10 +10790,6 @@ void FSequencer::BindCommands()
 		Commands.CollapseAllNodes,
 		FExecuteAction::CreateSP(this, &FSequencer::CollapseAllNodes));
 
-	SequencerCommandBindings->MapAction(
-		Commands.ResetFilters,
-		FExecuteAction::CreateSP(this, &FSequencer::ResetFilters));
-
 	SequencerCommandBindings->MapAction(
 		Commands.AddActorsToSequencer,
 		FExecuteAction::CreateSP( this, &FSequencer::AddSelectedActors));
@@ -11014,14 +11022,6 @@ void FSequencer::BindCommands()
 		FCanExecuteAction::CreateLambda([] { return true; }),
 		FIsActionChecked::CreateLambda([this] { return Settings->GetShowInfoButton(); }));
 
-	SequencerCommandBindings->MapAction(
-		Commands.ToggleShowSelectedNodesOnly,
-		FExecuteAction::CreateLambda( [this]{
-			Settings->SetShowSelectedNodesOnly( !Settings->GetShowSelectedNodesOnly() );
-		} ),
-		FCanExecuteAction::CreateLambda( []{ return true; } ),
-		FIsActionChecked::CreateLambda( [this]{ return Settings->GetShowSelectedNodesOnly(); } ) );
-
 	SequencerCommandBindings->MapAction(
 		Commands.ChangeTimeDisplayFormat,
 		FExecuteAction::CreateLambda( [this]{
@@ -11675,14 +11675,14 @@ void FSequencer::BindCommands()
 		CurveEditorSharedBindings->MapAction(Commands.AddRotationKey, *SequencerCommandBindings->GetActionForCommand(Commands.AddRotationKey));
 		CurveEditorSharedBindings->MapAction(Commands.AddScaleKey, *SequencerCommandBindings->GetActionForCommand(Commands.AddScaleKey));
 
-		CurveEditorSharedBindings->MapAction(Commands.ResetFilters, *SequencerCommandBindings->GetActionForCommand(Commands.ResetFilters));
-
 		TSharedPtr<FCurveEditor> CurveEditor = CurveEditorExtension->GetCurveEditor();
 		CurveEditor->GetCommands()->Append(CurveEditorSharedBindings);
 	}
 
 	// bind widget specific commands
 	SequencerWidget->BindCommands(SequencerCommandBindings, CurveEditorSharedBindings);
+
+	FilterBar->BindCommands();
 }
 
 void FSequencer::BuildAddTrackMenu(class FMenuBuilder& MenuBuilder)
@@ -11700,7 +11700,7 @@ void FSequencer::BuildAddTrackMenu(class FMenuBuilder& MenuBuilder)
 		LOCTEXT( "AddFolder", "Add Folder" ),
 		LOCTEXT( "AddFolderToolTip", "Adds a new folder." ),
 		FSlateIcon( FAppStyle::GetAppStyleSetName(), "ContentBrowser.AssetTreeFolderOpen" ),
-		FUIAction( FExecuteAction::CreateRaw( this, &FSequencer::OnAddFolder ) ) );
+		FUIAction( FExecuteAction::CreateRaw( this, &FSequencer::AddFolder ) ) );
 
 	for (int32 i = 0; i < TrackEditors.Num(); ++i)
 	{
@@ -12278,4 +12278,14 @@ FText FSequencer::GetSidebarSelectionDrawerToolTipText() const
 	return ToolTipText;
 }
 
+TSharedRef<ISequencerTrackFilters> FSequencer::GetFilterInterface() const
+{
+	return FilterBar;
+}
+
+TSharedRef<FSequencerFilterBar> FSequencer::GetFilterBar() const
+{
+	return FilterBar;
+}
+
 #undef LOCTEXT_NAMESPACE
diff --git a/Engine/Source/Editor/Sequencer/Private/Sequencer.h b/Engine/Source/Editor/Sequencer/Private/Sequencer.h
--- a/Engine/Source/Editor/Sequencer/Private/Sequencer.h
+++ b/Engine/Source/Editor/Sequencer/Private/Sequencer.h
@@ -40,6 +40,7 @@
 #include "Curves/RichCurve.h"
 #include "Sections/MovieScene3DTransformSection.h"
 #include "SequencerTimeChangeUndoRedoProxy.h"
+#include "Filters/ISequencerTrackFilters.h"
 
 class AActor;
 class ACameraActor;
@@ -808,7 +809,7 @@ public:
 	virtual UObject* FindSpawnedObjectOrTemplate(const FGuid& BindingId) override;
 	virtual FGuid MakeNewSpawnable(UObject& SourceObject, UActorFactory* ActorFactory = nullptr, bool bSetupDefaults = true) override;
 	virtual bool IsReadOnly() const override;
-	virtual void ExternalSelectionHasChanged() override { SynchronizeSequencerSelectionWithExternalSelection(); }
+	virtual void ExternalSelectionHasChanged() override;
 	virtual TSharedPtr<ISequencerTrackEditor> GetTrackEditor(UMovieSceneTrack* InTrack) override;
 	virtual void ObjectImplicitlyAdded(UObject* InObject) const override;
 	virtual void ObjectImplicitlyRemoved(UObject* InObject) const override;
@@ -818,8 +819,8 @@ public:
 	virtual void SetSequencerSettings(USequencerSettings* InSettings) override;
 	virtual TSharedPtr<class ITimeSlider> GetTopTimeSliderWidget() const override;
 	virtual void ResetTimeController() override;
-	virtual void SetTrackFilterEnabled(const FText& InTrackFilterName, bool bEnabled) override;
-	virtual bool IsTrackFilterEnabled(const FText& InTrackFilterName) const override;
+	virtual void SetTrackFilterEnabled(const FText& InFilterName, bool bInEnabled) override;
+	virtual bool IsTrackFilterEnabled(const FText& InFilterName) const override;
 	virtual TArray<FText> GetTrackFilterNames() const override;
 
 public:
@@ -1012,12 +1013,13 @@ protected:
 	/** Expand or collapse selected nodes and descendants*/
 	void ToggleExpandCollapseNodesAndDescendants();
 
+public:
+
 	/** Expand or collapse all nodes and descendants*/
 	void ExpandAllNodes();
 	void CollapseAllNodes();
 
-	/** Reset all enabled filters */
-	void ResetFilters();
+protected:
 
 	/** Sort all nodes and their descendants by category then alphabetically */
 	void SortAllNodesAndDescendants();
@@ -1068,6 +1070,9 @@ public:
 
 	void ClearFilters();
 
+	/** Handles adding a new folder to the outliner tree. */
+	void AddFolder();
+
 private:
 
 	/** Updates viewport clients' actor locks if they relate to sequencer cameras */
@@ -1076,9 +1081,6 @@ private:
 	/** Internal function to render movie for a given start/end time */
 	void RenderMovieInternal(TRange<FFrameNumber> Range, bool bSetFrameOverrides = false) const;
 
-	/** Handles adding a new folder to the outliner tree. */
-	void OnAddFolder();
-
 	/** Handles loading in previously recorded data. */
 	void OnLoadRecordedData();
 	
@@ -1227,6 +1229,10 @@ public:
 	/** Undocks the docked sidebar drawer if docked or docks the sidebar drawer if there is one open and no currently docked drawer. */
 	void ToggleSidebarDrawerDocked();
 
+	virtual TSharedRef<ISequencerTrackFilters> GetFilterInterface() const override;
+
+	TSharedRef<FSequencerFilterBar> GetFilterBar() const;
+
 private:
 
 	/** Update the time bases for the current movie scene */
@@ -1568,4 +1574,6 @@ private:
 
 	/** Delegate that is called when selection limiting has been toggled on or off */
 	FOnViewportSelectionLimitedChanged OnSelectionLimitedChangedDelegate;
+
+	TSharedRef<FSequencerFilterBar> FilterBar;
 };
diff --git a/Engine/Source/Editor/Sequencer/Private/SequencerCommands.cpp b/Engine/Source/Editor/Sequencer/Private/SequencerCommands.cpp
--- a/Engine/Source/Editor/Sequencer/Private/SequencerCommands.cpp
+++ b/Engine/Source/Editor/Sequencer/Private/SequencerCommands.cpp
@@ -64,7 +64,6 @@ void FSequencerCommands::RegisterCommands()
 	UI_COMMAND( ExpandAllNodes, "Expand All Nodes", "Expand all nodes and descendants", EUserInterfaceActionType::Button, FInputChord() );
 	UI_COMMAND( CollapseAllNodes, "Collapse All Nodes", "Collapse all nodes and descendants", EUserInterfaceActionType::Button, FInputChord() );
 	UI_COMMAND( SortAllNodesAndDescendants, "Sort All Nodes", "Sorts all nodes by type and then alphabetically.", EUserInterfaceActionType::Button, FInputChord());
-	UI_COMMAND( ResetFilters, "Reset Filters", "Reset all enabled filters.", EUserInterfaceActionType::Button, FInputChord());
 
 	UI_COMMAND( SetSelectionRangeEnd, "Set Selection End", "Sets the end of the selection range", EUserInterfaceActionType::Button, FInputChord(EKeys::O) );
 	UI_COMMAND( SetSelectionRangeStart, "Set Selection Start", "Sets the start of the selection range", EUserInterfaceActionType::Button, FInputChord(EKeys::I) );
@@ -154,8 +153,7 @@ void FSequencerCommands::RegisterCommands()
 	UI_COMMAND( ToggleLayerBars, "Layer Bars", "Show/hide the layer bars to edit keyframes in bulk", EUserInterfaceActionType::ToggleButton, FInputChord() );
 	UI_COMMAND( ToggleKeyBars, "Key Bars", "Show/hide key bar connectors for quickly retiming pairs of keys", EUserInterfaceActionType::ToggleButton, FInputChord() );
 	UI_COMMAND( ToggleChannelColors, "Channel Colors", "Show/hide the channel colors for the key bars", EUserInterfaceActionType::ToggleButton, FInputChord() );
-	UI_COMMAND( ToggleShowSelectedNodesOnly, "Selected Nodes Only", "Show selected nodes only", EUserInterfaceActionType::ToggleButton, FInputChord() );
-	
+
 	UI_COMMAND( ToggleShowCurveEditor, "Curve Editor", "Show the animation keys in a curve editor", EUserInterfaceActionType::ToggleButton, FInputChord() );
 	UI_COMMAND( ToggleLinkCurveEditorTimeRange, "Link Curve Editor Time Range", "Link the curve editor time range to the sequence", EUserInterfaceActionType::ToggleButton, FInputChord() );
 
