[Sequencer] Advanced Filtering:

diff --git a/Engine/Source/Editor/Sequencer/Private/SequencerNodeTree.cpp b/Engine/Source/Editor/Sequencer/Private/SequencerNodeTree.cpp
--- a/Engine/Source/Editor/Sequencer/Private/SequencerNodeTree.cpp
+++ b/Engine/Source/Editor/Sequencer/Private/SequencerNodeTree.cpp
@@ -1,6 +1,8 @@
 // Copyright Epic Games, Inc. All Rights Reserved.
 
 #include "SequencerNodeTree.h"
+#include "Filters/SequencerFilterBar.h"
+#include "Filters/Widgets/SSequencerFilterBar.h"
 #include "MVVM/ViewModels/ViewModelHierarchy.h"
 #include "MVVM/Selection/Selection.h"
 #include "MovieSceneBinding.h"
@@ -34,34 +36,21 @@
 #include "ISequencerTrackEditor.h"
 #include "Widgets/Views/STableRow.h"
 #include "CurveEditor.h"
-#include "SequencerTrackFilters.h"
 #include "Channels/MovieSceneChannel.h"
 #include "ScopedTransaction.h"
 #include "SequencerUtilities.h"
 #include "SequencerLog.h"
 #include "SequencerCommonHelpers.h"
+#include "SSequencer.h"
 
 FSequencerNodeTree::~FSequencerNodeTree()
 {
-	if (TrackFilters.IsValid())
-	{
-		TrackFilters->OnChanged().RemoveAll(this);
-	}
-	if (TrackFilterLevelFilter.IsValid())
-	{
-		TrackFilterLevelFilter->OnChanged().RemoveAll(this);
-	}
 }
 
 FSequencerNodeTree::FSequencerNodeTree(FSequencer& InSequencer)
 	: Sequencer(InSequencer)
-	, DisplayNodeCount(0)
 	, bFilterUpdateRequested(false)
 {
-	TrackFilters = MakeShared<FSequencerTrackFilterCollection>();
-	TrackFilters->OnChanged().AddRaw(this, &FSequencerNodeTree::RequestFilterUpdate);
-	TrackFilterLevelFilter = MakeShared< FSequencerTrackFilter_LevelFilter>();
-	TrackFilterLevelFilter->OnChanged().AddRaw(this, &FSequencerNodeTree::RequestFilterUpdate);
 }
 
 TSharedPtr<UE::Sequencer::FObjectBindingModel> FSequencerNodeTree::FindObjectBindingNode(const FGuid& BindingID) const
@@ -78,15 +67,6 @@ TSharedPtr<UE::Sequencer::FObjectBindingModel> FSequencerNodeTree::FindObjectBin
 	return nullptr;
 }
 
-bool FSequencerNodeTree::HasActiveFilter() const
-{
-	return (!FilterString.IsEmpty()
-		|| TrackFilters->Num() > 0
-		|| TrackFilterLevelFilter->IsActive()
-		|| Sequencer.GetSequencerSettings()->GetShowSelectedNodesOnly()
-		|| Sequencer.GetFocusedMovieSceneSequence()->GetMovieScene()->GetNodeGroups().HasAnyActiveFilter());
-}
-
 bool FSequencerNodeTree::UpdateFiltersOnTrackValueChanged()
 {
 	// If filters are already scheduled for update, we can defer until the next update
@@ -95,18 +75,15 @@ bool FSequencerNodeTree::UpdateFiltersOnTrackValueChanged()
 		return false;
 	}
 
-	for (TSharedPtr< FSequencerTrackFilter > TrackFilter : *TrackFilters)
+	if (Sequencer.GetFilterBar()->ShouldUpdateOnTrackValueChanged())
 	{
-		if (TrackFilter->ShouldUpdateOnTrackValueChanged())
-		{
-			// UpdateFilters will only run if bFilterUpdateRequested is true
-			bFilterUpdateRequested = true;
-			bool bFiltersUpdated = UpdateFilters();
+		// UpdateFilters will only run if bFilterUpdateRequested is true
+		bFilterUpdateRequested = true;
+		bool bFiltersUpdated = UpdateFilters();
 
-			// If the filter list was modified, set bFilterUpdateRequested to suppress excessive re-filters between tree update
-			bFilterUpdateRequested = bFiltersUpdated;
-			return bFiltersUpdated;
-		}
+		// If the filter list was modified, set bFilterUpdateRequested to suppress excessive re-filters between tree update
+		bFilterUpdateRequested = bFiltersUpdated;
+		return bFiltersUpdated;
 	}
 
 	return false;
@@ -118,8 +95,6 @@ void FSequencerNodeTree::Update()
 
 	FViewModelHierarchyOperation UpdateOp(RootNode->GetSharedData());
 
-	FilteredNodes.Empty();
-
 	if (!ensure(RootNode))
 	{
 		return;
@@ -136,7 +111,6 @@ void FSequencerNodeTree::Update()
 
 	// Re-filter the tree after updating 
 	// @todo sequencer: Newly added sections may need to be visible even when there is a filter
-	FilterNodes(FilterString);
 	bFilterUpdateRequested = true;
 	UpdateFilters();
 
@@ -282,51 +256,6 @@ void FSequencerNodeTree::SortAllNodesAndDescendants()
 	}
 }
 
-void FSequencerNodeTree::AddFilter(TSharedPtr<FSequencerTrackFilter> TrackFilter)
-{
-	GetSequencer().GetSequencerSettings()->SetTrackFilterEnabled(TrackFilter->GetDisplayName().ToString(), true);
-
-	TrackFilters->Add(TrackFilter);
-}
-
-int32 FSequencerNodeTree::RemoveFilter(TSharedPtr<FSequencerTrackFilter> TrackFilter)
-{
-	GetSequencer().GetSequencerSettings()->SetTrackFilterEnabled(TrackFilter->GetDisplayName().ToString(), false);
-
-	return TrackFilters->Remove(TrackFilter);
-}
-
-void FSequencerNodeTree::RemoveAllFilters()
-{
-	for (TSharedPtr< FSequencerTrackFilter > TrackFilter : *TrackFilters)
-	{
-		GetSequencer().GetSequencerSettings()->SetTrackFilterEnabled(TrackFilter->GetDisplayName().ToString(), false);
-	}
-
-	TrackFilters->RemoveAll();
-	TrackFilterLevelFilter->ResetFilter();
-}
-
-bool FSequencerNodeTree::IsTrackFilterActive(TSharedPtr<FSequencerTrackFilter> TrackFilter) const
-{
-	return TrackFilters->Contains(TrackFilter);
-}
-
-void FSequencerNodeTree::AddLevelFilter(const FString& LevelName)
-{
-	TrackFilterLevelFilter->UnhideLevel(LevelName);
-}
-
-void FSequencerNodeTree::RemoveLevelFilter(const FString& LevelName)
-{
-	TrackFilterLevelFilter->HideLevel(LevelName);
-}
-
-bool FSequencerNodeTree::IsTrackLevelFilterActive(const FString& LevelName) const
-{
-	return !TrackFilterLevelFilter->IsLevelHidden(LevelName);
-}
-
 void FSequencerNodeTree::SaveExpansionState(const UE::Sequencer::FViewModel& Node, bool bExpanded)
 {
 	using namespace UE::Sequencer;
@@ -382,8 +311,8 @@ bool FSequencerNodeTree::IsNodeFiltered(const TSharedPtr<UE::Sequencer::FViewMod
 {
 	using namespace UE::Sequencer;
 
-	TViewModelPtr<IOutlinerExtension> OutlinerItem = CastViewModel<IOutlinerExtension>(Node);
-	return OutlinerItem && FilteredNodes.Contains(OutlinerItem);
+	const TViewModelPtr<IOutlinerExtension> OutlinerItem = CastViewModel<IOutlinerExtension>(Node);
+	return OutlinerItem && !OutlinerItem->IsFilteredOut();
 }
 
 TSharedPtr<UE::Sequencer::FSectionModel> FSequencerNodeTree::GetSectionModel(const UMovieSceneSection* Section) const
@@ -398,365 +327,6 @@ TSharedPtr<UE::Sequencer::FSectionModel> FSequencerNodeTree::GetSectionModel(con
 	return nullptr;
 }
 
-static void AddChildNodes(const UE::Sequencer::TViewModelPtr<UE::Sequencer::IOutlinerExtension>& StartNode, TSet<UE::Sequencer::TWeakViewModelPtr<UE::Sequencer::IOutlinerExtension>>& OutFilteredNodes)
-{
-	using namespace UE::Sequencer;
-
-	for (TViewModelPtr<IOutlinerExtension> ChildNode : StartNode.AsModel()->GetDescendantsOfType<IOutlinerExtension>())
-	{
-		OutFilteredNodes.Add(ChildNode);
-	}
-}
-
-static void AddParentNodes(const UE::Sequencer::TViewModelPtr<UE::Sequencer::IOutlinerExtension>& StartNode, TSet<UE::Sequencer::TWeakViewModelPtr<UE::Sequencer::IOutlinerExtension>>& OutFilteredNodes)
-{
-	using namespace UE::Sequencer;
-
-	// Gather parent folders up the chain
-	for (TViewModelPtr<IOutlinerExtension> ParentNode : StartNode.AsModel()->GetAncestorsOfType<IOutlinerExtension>())
-	{
-		OutFilteredNodes.Add(ParentNode);
-	}
-}
-
-static bool PassesFilterStrings(FSequencer& Sequencer, const UE::Sequencer::TViewModelPtr<UE::Sequencer::IOutlinerExtension>& StartNode, const TArray<FString>& FilterStrings)
-{
-	using namespace UE::Sequencer;
-
-	// If we have a filter string, make sure we match
-	if (FilterStrings.Num() > 0)
-	{
-		TSharedPtr<FOutlinerViewModel> Outliner = Sequencer.GetViewModel()->GetOutliner();
-
-		if (!Outliner)
-		{
-			return false;
-		}
-
-		FString DisplayLabel = StartNode->GetLabel().ToString();
-
-		// check each string in the filter strings list against 
-		for (const FString& String : FilterStrings)
-		{
-			if (!DisplayLabel.Contains(String))
-			{
-				return false;
-			}
-		}
-
-		for (TViewModelPtr<IOutlinerExtension> Parent : StartNode.AsModel()->GetAncestorsOfType<IOutlinerExtension>())
-		{
-			if (!Parent->IsExpanded())
-			{
-				Parent->SetExpansion(true);
-			}
-		}
-	}
-
-	return true;
-}
-
-/**
- * Recursively filters nodes
- *
- * @param StartNode			The node to start from
- * @param Filters			The filter collection to test against
- * @param OutFilteredNodes	The list of all filtered nodes
- *
- * @return Whether this node passed filtering
-  */
-
-static bool FilterNodesRecursive(
-	  FSequencer& Sequencer
-	, const UE::Sequencer::TViewModelPtr<UE::Sequencer::IOutlinerExtension>& StartNode
-	, TSharedPtr<FSequencerTrackFilterCollection> Filters, const TArray<FString>& FilterStrings
-	, TSharedPtr<FSequencerTrackFilter_LevelFilter> LevelTrackFilter
-	, TSet<UE::Sequencer::TWeakViewModelPtr<UE::Sequencer::IOutlinerExtension>>& OutFilteredNodes)
-{
-	using namespace UE::Sequencer;
-
-	bool bAnyChildPassed = false;
-
-	// Special case: If a parent node matches an active text search filter, the children should also pass the text search filter
-	// so we stop filtering them based on text to eliminate special case conflicts with non-object binding nodes
-	// with potential child object bindings nodes when only showing selected bindings. This is also faster.
-	bool bPassedFilterStrings = PassesFilterStrings(Sequencer, StartNode, FilterStrings);
-	const TArray<FString>& ChildFilterStrings = bPassedFilterStrings ? TArray<FString>() : FilterStrings;
-
-	// Special case: Child nodes should always be processed, as they may force their parents to pass
-	for (TViewModelPtr<IOutlinerExtension> Node : StartNode.AsModel()->GetChildrenOfType<IOutlinerExtension>())
-	{
-		if (FilterNodesRecursive(Sequencer, Node, Filters, ChildFilterStrings, LevelTrackFilter, OutFilteredNodes))
-		{
-			bAnyChildPassed = true;
-		}
-	}
-
-	// After child nodes are processed, if this node didn't pass text filtering, fail it
-	if (!bPassedFilterStrings)
-	{
-		return bAnyChildPassed;
-	}
-
-	// TODO: we can probably completely rewrite this logic by taking advantage of the new modular extensions.
-
-	bool bPassedAnyFilters = false;
-	bool bIsTrackOrObjectBinding = false;
-
-	TSharedPtr<IPinnableExtension> Pinnable = StartNode.ImplicitCast();
-	const bool bIsPinned = Pinnable && Pinnable->IsPinned();
-
-	if (TSharedPtr<FTrackModel> TrackModel = StartNode.ImplicitCast())
-	{
-		bIsTrackOrObjectBinding = true;
-
-		UMovieSceneTrack* Track = TrackModel->GetTrack();
-
-		for (TSharedPtr<FChannelGroupModel> ChannelGroupModel : StartNode.AsModel()->GetDescendantsOfType<FChannelGroupModel>())
-		{
-			for (const TWeakViewModelPtr<FChannelModel>& WeakChannel : ChannelGroupModel->GetChannels())
-			{
-				if (TViewModelPtr<FChannelModel> ChannelModel = WeakChannel.Pin())
-				{
-					const TSharedPtr<IKeyArea>& KeyArea = ChannelModel->GetKeyArea();
-					FMovieSceneChannel* Channel = KeyArea->ResolveChannel();
-					if (Channel)
-					{
-						if (Filters->Num() == 0 || Filters->PassesAnyFilters(Channel))
-						{
-							bPassedAnyFilters = true;
-							break;
-						}
-					}
-				}
-			}
-		}
-	
-		if (bPassedAnyFilters || Filters->Num() == 0 || Filters->PassesAnyFilters(Track, TrackModel->GetLabel()))
-		{
-			bPassedAnyFilters = true;
-
-			// Track nodes do not belong to a level, but might be a child of an objectbinding node that does
-			if (LevelTrackFilter->IsActive())
-			{
-				TSharedPtr<IObjectBindingExtension> ObjectNode = StartNode.AsModel()->FindAncestorOfType<IObjectBindingExtension>();
-				if (ObjectNode)
-				{
-					// The track belongs to an objectbinding node, start by assuming it doesn't match the level filter
-					bPassedAnyFilters = false;
-
-					for (TWeakObjectPtr<>& Object : Sequencer.FindObjectsInCurrentSequence(ObjectNode->GetObjectGuid()))
-					{
-						if (Object.IsValid() && LevelTrackFilter->PassesFilter(Object.Get()))
-						{
-							// If at least one of the objects on the objectbinding node pass the level filter, show the track
-							bPassedAnyFilters = true;
-							break;
-						}
-					}
-				}
-			}
-
-			if (bPassedAnyFilters && Sequencer.GetSequencerSettings()->GetShowSelectedNodesOnly())
-			{
-				TSharedPtr<IObjectBindingExtension> ObjectNode = StartNode.AsModel()->FindAncestorOfType<IObjectBindingExtension>();
-				// Always show pinned items
-				if (!bIsPinned && ObjectNode)
-				{
-					const FMovieSceneBinding* Binding = Sequencer.GetFocusedMovieSceneSequence()->GetMovieScene()->FindBinding(ObjectNode->GetObjectGuid());
-					if (!(Binding && Sequencer.IsBindingVisible(*Binding)))
-					{
-						return bAnyChildPassed;
-					}
-				}
-			}
-		}
-	}
-	else if (TViewModelPtr<FObjectBindingModel> ObjectNode = StartNode.ImplicitCast())
-	{
-		bIsTrackOrObjectBinding = true;
-
-		for (TWeakObjectPtr<>& Object : Sequencer.FindObjectsInCurrentSequence(ObjectNode->GetObjectGuid()))
-		{
-			if (Object.IsValid() && (Filters->Num() == 0 || Filters->PassesAnyFilters(Object.Get(), StartNode->GetLabel()))
-				&& LevelTrackFilter->PassesFilter(Object.Get()))
-			{
-				bPassedAnyFilters = true;
-				break;
-			}
-		}
-
-		if (bPassedAnyFilters && Sequencer.GetSequencerSettings()->GetShowSelectedNodesOnly() && !bIsPinned)
-		{
-			UMovieScene* MovieScene = Sequencer.GetFocusedMovieSceneSequence()->GetMovieScene();
-			const FMovieSceneBinding* Binding = Sequencer.GetFocusedMovieSceneSequence()->GetMovieScene()->FindBinding(ObjectNode->GetObjectGuid());
-			if (Binding && !Sequencer.IsBindingVisible(*Binding))
-			{
-				return bAnyChildPassed;
-			}
-		}
-	}
-	else if (TViewModelPtr<FCategoryModel> CategoryModel = StartNode.ImplicitCast())
-	{
-		if (TSharedPtr<FTrackModel> TrackNode = CategoryModel->FindAncestorOfType<FTrackModel>())
-		{
-			UMovieSceneTrack* Track = TrackNode->GetTrack();
-			if (Filters->Num() == 0 || Filters->PassesAnyFilters(Track, StartNode->GetLabel()))
-			{
-				bPassedAnyFilters = true;
-			}
-		}
-		if (bPassedAnyFilters && Sequencer.GetSequencerSettings()->GetShowSelectedNodesOnly())
-		{
-			TSharedPtr<IObjectBindingExtension> ParentObjectNode = StartNode.AsModel()->FindAncestorOfType<IObjectBindingExtension>();
-			// Always show pinned items
-			if (!bIsPinned && ParentObjectNode)
-			{
-				const FMovieSceneBinding* Binding = Sequencer.GetFocusedMovieSceneSequence()->GetMovieScene()->FindBinding(ParentObjectNode->GetObjectGuid());
-				if (!(Binding && Sequencer.IsBindingVisible(*Binding)))
-				{
-					return bAnyChildPassed;
-				}
-			}
-		}
-	}
-	else if (TViewModelPtr<FCategoryGroupModel> CategoryGroupModel = StartNode.ImplicitCast())
-	{
-		if (TSharedPtr<FTrackModel> TrackNode = CategoryGroupModel->FindAncestorOfType<FTrackModel>())
-		{
-			UMovieSceneTrack* Track = TrackNode->GetTrack();
-			if (Filters->Num() == 0 || Filters->PassesAnyFilters(Track, FText::FromName(CategoryGroupModel->GetCategoryName())))
-			{
-				bPassedAnyFilters = true;
-			}
-		}
-		if (bPassedAnyFilters && Sequencer.GetSequencerSettings()->GetShowSelectedNodesOnly())
-		{
-			TSharedPtr<IObjectBindingExtension> ParentObjectNode = StartNode.AsModel()->FindAncestorOfType<IObjectBindingExtension>();
-			// Always show pinned items
-			if (!bIsPinned && ParentObjectNode)
-			{
-				const FMovieSceneBinding* Binding = Sequencer.GetFocusedMovieSceneSequence()->GetMovieScene()->FindBinding(ParentObjectNode->GetObjectGuid());
-				if (!(Binding && Sequencer.IsBindingVisible(*Binding)))
-				{
-					return bAnyChildPassed;
-				}
-			}
-		}
-	}
-	else if (TViewModelPtr<FChannelGroupModel> ChannelGroupModel = StartNode.ImplicitCast())
-	{
-		for (const TWeakViewModelPtr<FChannelModel>& WeakChannel : ChannelGroupModel->GetChannels())
-		{
-			if (TViewModelPtr<FChannelModel> ChannelModel = WeakChannel.Pin())
-			{
-				const TSharedPtr<IKeyArea>& KeyArea = ChannelModel->GetKeyArea();
-				FMovieSceneChannel* Channel = KeyArea->ResolveChannel();
-				if (Channel)
-				{
-					if (Filters->Num() == 0 || Filters->PassesAnyFilters(Channel))
-					{
-						bPassedAnyFilters = true;
-						break;
-					}
-				}
-			}
-		}
-		if (bPassedAnyFilters && Sequencer.GetSequencerSettings()->GetShowSelectedNodesOnly())
-		{
-			TSharedPtr<IObjectBindingExtension> ParentObjectNode = StartNode.AsModel()->FindAncestorOfType<IObjectBindingExtension>();
-			// Always show pinned items
-			if (!bIsPinned && ParentObjectNode)
-			{
-				const FMovieSceneBinding* Binding = Sequencer.GetFocusedMovieSceneSequence()->GetMovieScene()->FindBinding(ParentObjectNode->GetObjectGuid());
-				if (!(Binding && Sequencer.IsBindingVisible(*Binding)))
-				{
-					bPassedAnyFilters = false;
-				}
-			}
-		}
-	}
-	else if (TViewModelPtr<FFolderModel> FolderModel = StartNode.ImplicitCast())
-	{
-		// Special case: If we're pinned, then we should pass regardless
-		if (bIsPinned)
-		{
-			bPassedAnyFilters = true;
-		}
-
-		// Special case: If we're only filtering on text search, include folders and key areas in the search
-		if (!bPassedAnyFilters && Filters->Num() == 0 && FilterStrings.Num() > 0)
-		{
-			bPassedAnyFilters = true;
-
-			// Special case: but don't include if only showing selected bindings and we don't have child that passed
-			if (Sequencer.GetSequencerSettings()->GetShowSelectedNodesOnly() && !bAnyChildPassed)
-			{
-				bPassedAnyFilters = false;
-
-				// Special case: unless we're the child of a node that is a selected binding
-				TSharedPtr<IObjectBindingExtension> ParentObjectNode = StartNode.AsModel()->FindAncestorOfType<IObjectBindingExtension>();
-				// Always show pinned items
-				if (ParentObjectNode)
-				{
-					const FMovieSceneBinding* Binding = Sequencer.GetFocusedMovieSceneSequence()->GetMovieScene()->FindBinding(ParentObjectNode->GetObjectGuid());
-					if (Binding && Sequencer.IsBindingVisible(*Binding))
-					{
-						bPassedAnyFilters = true;
-					}
-				}
-			}
-		}
-	}
-
-	if (bPassedAnyFilters)
-	{
-		// If filtering on selection set is enabled, we need to run another pass to verify we're in an enabled node group
-		UMovieSceneNodeGroupCollection& NodeGroups = Sequencer.GetFocusedMovieSceneSequence()->GetMovieScene()->GetNodeGroups();
-		if (NodeGroups.HasAnyActiveFilter())
-		{
-			bPassedAnyFilters = false;
-			
-			constexpr bool bIncludeThis = true;
-			for (TViewModelPtr<IOutlinerExtension> Parent : StartNode.AsModel()->GetAncestorsOfType<IOutlinerExtension>(bIncludeThis))
-			{
-				// Special case: Pinned tracks should be visible whether in the node group or not
-				if (bIsPinned)
-				{
-					bPassedAnyFilters = true;
-					break;
-				}
-
-				for (const UMovieSceneNodeGroup* NodeGroup : NodeGroups)
-				{
-					if (NodeGroup->GetEnableFilter() && NodeGroup->ContainsNode(IOutlinerExtension::GetPathName(Parent.AsModel())))
-					{
-						bPassedAnyFilters = true;
-						break;
-					}
-				}
-			}
-		}
-	}
-
-	if (bPassedAnyFilters)
-	{
-		OutFilteredNodes.Add(StartNode);
-		AddParentNodes(StartNode, OutFilteredNodes);
-
-		// Special case: When only showing selected bindings, and a non-object passes text filtering
-		// don't add it's children, as they may be a binding that is not seleceted. Selected child nodes will add themselves.
-		if (!(Sequencer.GetSequencerSettings()->GetShowSelectedNodesOnly() && FilterStrings.Num() > 0 && !bIsTrackOrObjectBinding))
-		{
-			AddChildNodes(StartNode, OutFilteredNodes);
-		}
-
-		return true;
-	}
-
-	return bAnyChildPassed;
-}
-
 bool FSequencerNodeTree::UpdateFilters()
 {
 	using namespace UE::Sequencer;
@@ -766,52 +336,21 @@ bool FSequencerNodeTree::UpdateFilters()
 		return false;
 	}
 
-	TSet<TWeakViewModelPtr<IOutlinerExtension>> PreviousFilteredNodes(FilteredNodes);
-
-	FilteredNodes.Empty();
-	const bool bHasActiveFilter = HasActiveFilter();
-
-	UObject* PlaybackContext = Sequencer.GetPlaybackContext();
-	UWorld* World = PlaybackContext ? PlaybackContext->GetWorld() : nullptr;
-	TrackFilterLevelFilter->UpdateWorld(World);
-
-	if (bHasActiveFilter)
+	const TSharedPtr<FSequencerFilterBar> FilterBar = Sequencer.GetFilterBar();
+	if (!FilterBar.IsValid())
 	{
-		// Build a list of strings that must be matched
-		TArray<FString> FilterStrings;
-
-		// Remove whitespace from the front and back of the string
-		FilterString.TrimStartAndEndInline();
-		FilterString.ParseIntoArray(FilterStrings, TEXT(" "), true /*bCullEmpty*/);
-
-		for (const TViewModelPtr<IOutlinerExtension>& Node : GetRootNodes())
-		{
-			// Recursively filter all nodes, matching them against the list of filter strings.  All filter strings must be matched
-			FilterNodesRecursive(Sequencer, Node, TrackFilters, FilterStrings, TrackFilterLevelFilter, FilteredNodes);
-		}
+		return false;
 	}
 
+	const FSequencerFilterData& PreviousFilterData = FilterBar->GetFilterData();
+	const FSequencerFilterData& FilterData = FilterBar->FilterNodes();
+
 	bFilteringOnNodeGroups = Sequencer.GetFocusedMovieSceneSequence()->GetMovieScene()->GetNodeGroups().HasAnyActiveFilter();
 	bFilterUpdateRequested = false;
 
-	// Always include the bottom spacer
-	if (FSequenceModel* SequenceModel = RootNode->CastThis<FSequenceModel>())
-	{
-		FilteredNodes.Add(CastViewModelChecked<IOutlinerExtension>(SequenceModel->GetBottomSpacer()));
-	}
-
-	// Count the total number of display nodes, and update filtered state.
-	DisplayNodeCount = 0;
-	for (const TViewModelPtr<IOutlinerExtension>& Item : RootNode->GetDescendantsOfType<IOutlinerExtension>())
-	{
-		const bool bIsNodeFilteredIn = !bHasActiveFilter || FilteredNodes.Contains(Item);
-		Item->SetFilteredOut(!bIsNodeFilteredIn);
-
-		++DisplayNodeCount;
-	}
-
 	// Return whether the new list of FilteredNodes is different than the previous list
-	return (PreviousFilteredNodes.Num() != FilteredNodes.Num() || !PreviousFilteredNodes.Includes(FilteredNodes));
+	return (PreviousFilterData.GetDisplayNodeCount() != FilterData.GetDisplayNodeCount()
+		|| PreviousFilterData != FilterData);
 }
 
 void FSequencerNodeTree::CleanupMuteSolo(UMovieScene* MovieScene)
@@ -839,22 +378,22 @@ void FSequencerNodeTree::CleanupMuteSolo(UMovieScene* MovieScene)
 
 int32 FSequencerNodeTree::GetTotalDisplayNodeCount() const 
 { 
-	// Subtract 1 for the spacer node which is always added
-	return DisplayNodeCount - 1; 
+	return Sequencer.GetFilterBar()->GetFilterData().GetTotalNodeCount();
 }
 
 int32 FSequencerNodeTree::GetFilteredDisplayNodeCount() const 
 { 
-	// Subtract 1 for the spacer node which is always added
-	return FilteredNodes.Num() - 1;
+	return Sequencer.GetFilterBar()->GetFilterData().GetDisplayNodeCount();
 }
 
-void FSequencerNodeTree::FilterNodes(const FString& InFilter)
+void FSequencerNodeTree::SetTextFilterString(const FString& InFilter)
 {
+	const TSharedRef<FSequencerFilterBar> FilterBar = Sequencer.GetFilterBar();
+	const FString FilterString = FilterBar->GetTextFilterString();
 	if (InFilter != FilterString)
 	{
-		FilterString = InFilter;
 		bFilterUpdateRequested = true;
+		FilterBar->SetTextFilterString(InFilter);
 	}
 }
 
@@ -887,4 +426,3 @@ void FSequencerNodeTree::GetAllNodes(TArray<TSharedRef<UE::Sequencer::FViewModel
 		OutNodes.Add(It.AsModel().ToSharedRef());
 	}
 }
-
diff --git a/Engine/Source/Editor/Sequencer/Private/SequencerNodeTree.h b/Engine/Source/Editor/Sequencer/Private/SequencerNodeTree.h
--- a/Engine/Source/Editor/Sequencer/Private/SequencerNodeTree.h
+++ b/Engine/Source/Editor/Sequencer/Private/SequencerNodeTree.h
@@ -17,10 +17,11 @@
 #include "MovieSceneSequence.h"
 
 class FSequencer;
-class FSequencerTrackFilter;
+class FSequencerFilterBar;
 class FSequencerTrackFilterCollection;
-class FSequencerTrackFilter_LevelFilter;
+class FSequencerTrackFilter_Level;
 class ISequencerTrackEditor;
+class SSequencerFilterBar;
 class UMovieScene;
 class UMovieSceneFolder;
 class UMovieSceneTrack;
@@ -65,9 +66,6 @@ public:
 	 */
 	TArray<UE::Sequencer::TViewModelPtr<UE::Sequencer::IOutlinerExtension>> GetRootNodes() const;
 
-	/** @return Whether or not there is an active filter */
-	bool HasActiveFilter() const;
-
 	/** 
 	 * Checks if filters should be updated on track value changes, and if so updates them.
 	 * @return Whether the filtered node list was modified
@@ -96,7 +94,7 @@ public:
 	 *
 	 * @param InFilter	The filter terms
 	 */
-	void FilterNodes( const FString& InFilter );
+	void SetTextFilterString( const FString& InFilter );
 
 	/** Called when the active MovieScene's node group colletion has been modifed */
 	void NodeGroupsCollectionChanged();
@@ -177,20 +175,6 @@ public:
 	 */
 	TSharedPtr<FSectionModel> GetSectionModel(const UMovieSceneSection* Section) const;
 
-	void AddFilter(TSharedPtr<FSequencerTrackFilter> TrackFilter);
-	int32 RemoveFilter(TSharedPtr<FSequencerTrackFilter> TrackFilter);
-	void RemoveAllFilters();
-	bool IsTrackFilterActive(TSharedPtr<FSequencerTrackFilter> TrackFilter) const;
-
-	void AddLevelFilter(const FString& LevelName);
-	void RemoveLevelFilter(const FString& LevelName);
-	bool IsTrackLevelFilterActive(const FString& LevelName) const;
-
-private:
-
-	/** Returns whether this NodeTree should only display selected nodes */
-	bool ShowSelectedNodesOnly() const;
-
 private:
 
 	/**
@@ -216,25 +200,11 @@ private:
 	/** Symbolic root node that contains the actual displayed root nodes as children */
 	UE::Sequencer::FViewModelPtr RootNode;
 
-	/** Set of all filtered nodes */
-	TSet<UE::Sequencer::TWeakViewModelPtr<UE::Sequencer::IOutlinerExtension>> FilteredNodes;
-	/** Active filter string if any */
-	FString FilterString;
 	/** Sequencer interface */
 	FSequencer& Sequencer;
 	/** A multicast delegate which is called whenever the node tree has been updated. */
 	FOnUpdated OnUpdatedDelegate;
 
-	/** Active track filters */
-	TSharedPtr<FSequencerTrackFilterCollection> TrackFilters;
-	
-	/** Level based track filtering */
-	TSharedPtr<FSequencerTrackFilter_LevelFilter> TrackFilterLevelFilter;
-
-	/** The total number of DisplayNodes in the tree, both displayed and hidden */
-	uint32 DisplayNodeCount;
-
 	bool bFilterUpdateRequested;
 	bool bFilteringOnNodeGroups;
 };
-
diff --git a/Engine/Source/Editor/Sequencer/Private/SequencerSettings.cpp b/Engine/Source/Editor/Sequencer/Private/SequencerSettings.cpp
--- a/Engine/Source/Editor/Sequencer/Private/SequencerSettings.cpp
+++ b/Engine/Source/Editor/Sequencer/Private/SequencerSettings.cpp
@@ -71,6 +71,11 @@ USequencerSettings::USequencerSettings( const FObjectInitializer& ObjectInitiali
 	bShowSequencerToolbar = true;
 	bShowMarkedFrames = true;
 	ViewDensity = "Relaxed";
+	bIncludePinnedInFilter = false;
+	bAutoExpandNodesOnFilterPass = false;
+	bUseFilterSubmenusForCategories = false;
+	LastFilterBarLayout = EFilterBarLayout::Horizontal;
+	LastFilterBarSizeCoefficient = 0.1f;
 
 	SectionColorTints.Add(FColor(88, 102, 142, 255)); // blue
 	SectionColorTints.Add(FColor(99, 137, 132, 255)); // blue-green
@@ -1045,29 +1050,16 @@ void USequencerSettings::SetViewDensity(FName InViewDensity)
 	}
 }
 
-bool USequencerSettings::IsTrackFilterEnabled(const FString& TrackFilter) const
+FSidebarState& USequencerSettings::GetSidebarState()
 {
-	return TrackFilters.Contains(TrackFilter);
+	return SidebarState.FindOrAdd(GetFName());
 }
 
-void USequencerSettings::SetTrackFilterEnabled(const FString & TrackFilter, bool bEnabled)
+void USequencerSettings::SetSidebarState(const FSidebarState& InSidebarState)
 {
-	if (bEnabled)
-	{
-		if (!TrackFilters.Contains(TrackFilter))
-		{
-			TrackFilters.Add(TrackFilter);
-			SaveConfig();
-		}
-	}
-	else
-	{
-		if (TrackFilters.Contains(TrackFilter))
-		{
-			TrackFilters.Remove(TrackFilter);
-			SaveConfig();
-		}
-	}
+	FSidebarState& State = SidebarState.FindOrAdd(GetFName());
+	State = InSidebarState;
+	SaveConfig();
 }
 
 void USequencerSettings::SetOutlinerColumnVisibility(const TArray<FColumnVisibilitySetting>& InColumnVisibilitySettings)
@@ -1079,14 +1071,92 @@ void USequencerSettings::SetOutlinerColumnVisibility(const TArray<FColumnVisibil
 	}
 }
 
-FSidebarState& USequencerSettings::GetSidebarState()
+FSequencerFilterBarConfig& USequencerSettings::FindOrAddTrackFilterBar(const FName InIdentifier, const bool bInSaveConfig)
 {
-	return SidebarState.FindOrAdd(GetFName());
+	FSequencerFilterBarConfig& FilterBarSettings = TrackFilterBars.FindOrAdd(InIdentifier);
+
+	if (bInSaveConfig)
+	{
+		SaveConfig();
+	}
+
+	return FilterBarSettings;
 }
 
-void USequencerSettings::SetSidebarState(const FSidebarState& InSidebarState)
+FSequencerFilterBarConfig* USequencerSettings::FindTrackFilterBar(const FName InIdentifier)
 {
-	FSidebarState& State = SidebarState.FindOrAdd(GetFName());
-	State = InSidebarState;
+	return TrackFilterBars.Find(InIdentifier);
+}
+
+bool USequencerSettings::RemoveTrackFilterBar(const FName InIdentifier)
+{
+	const int32 RemovedCount = TrackFilterBars.Remove(InIdentifier) > 0;
+	SaveConfig();
+	return RemovedCount > 0;
+}
+
+bool USequencerSettings::GetIncludePinnedInFilter() const
+{
+	return bIncludePinnedInFilter;
+}
+
+void USequencerSettings::SetIncludePinnedInFilter(const bool bInIncludePinned)
+{
+	bIncludePinnedInFilter = bInIncludePinned;
+	SaveConfig();
+}
+
+bool USequencerSettings::GetAutoExpandNodesOnFilterPass() const
+{
+	return bAutoExpandNodesOnFilterPass;
+}
+
+void USequencerSettings::SetAutoExpandNodesOnFilterPass(const bool bInIncludeParents)
+{
+	bAutoExpandNodesOnFilterPass = bInIncludeParents;
+	SaveConfig();
+}
+
+bool USequencerSettings::GetUseFilterSubmenusForCategories() const
+{
+	return bUseFilterSubmenusForCategories;
+}
+
+void USequencerSettings::SetUseFilterSubmenusForCategories(const bool bInUseFilterSubmenusForCategories)
+{
+	bUseFilterSubmenusForCategories = bInUseFilterSubmenusForCategories;
+	SaveConfig();
+}
+
+bool USequencerSettings::IsFilterBarVisible() const
+{
+	return bFilterBarVisible;
+}
+
+void USequencerSettings::SetFilterBarVisible(const bool bInVisible)
+{
+	bFilterBarVisible = bInVisible;
+	SaveConfig();
+}
+
+EFilterBarLayout USequencerSettings::GetFilterBarLayout() const
+{
+	return LastFilterBarLayout;
+}
+
+void USequencerSettings::SetFilterBarLayout(const EFilterBarLayout InLayout)
+{
+	LastFilterBarLayout = InLayout;
+	SaveConfig();
+}
+
+float USequencerSettings::GetLastFilterBarSizeCoefficient() const
+{
+	return LastFilterBarSizeCoefficient;
+}
+
+void USequencerSettings::SetLastFilterBarSizeCoefficient(const float bInSizeCoefficient)
+{
+	LastFilterBarSizeCoefficient = bInSizeCoefficient;
 	SaveConfig();
 }
diff --git a/Engine/Source/Editor/Sequencer/Public/Filters/ISequencerTrackFilters.h b/Engine/Source/Editor/Sequencer/Public/Filters/ISequencerTrackFilters.h
new file mode 100644
--- /dev/null
+++ b/Engine/Source/Editor/Sequencer/Public/Filters/ISequencerTrackFilters.h
@@ -0,0 +1,64 @@
+﻿// Copyright Epic Games, Inc. All Rights Reserved.
+
+#pragma once
+
+#include "Framework/Commands/UICommandList.h"
+
+class FSequencerTrackFilter_CustomText;
+class FSequencer;
+class FString;
+class FText;
+struct FSequencerFilterData;
+
+class ISequencerTrackFilters
+{
+public:
+	virtual ~ISequencerTrackFilters() = default;
+
+	SEQUENCER_API virtual FName GetIdentifier() const = 0;
+
+	SEQUENCER_API virtual FSequencer& GetSequencer() const = 0;
+
+	SEQUENCER_API virtual TSharedPtr<FUICommandList> GetCommandList() const = 0;
+
+	SEQUENCER_API virtual bool HasAnyFilterActive(const bool bCheckTextFilter = true
+		, const bool bInCheckHideIsolateFilter = true
+		, const bool bInCheckCommonFilters = true
+		, const bool bInCheckInternalFilters = true
+		, const bool bInCheckCustomTextFilters = true) const = 0;
+
+	SEQUENCER_API virtual bool HasAnyFilterEnabled() const = 0;
+
+	SEQUENCER_API virtual void RequestFilterUpdate() = 0;
+
+	SEQUENCER_API virtual bool AddCustomTextFilter(const TSharedRef<FSequencerTrackFilter_CustomText>& InFilter, const bool bInAddToConfig) = 0;
+	SEQUENCER_API virtual bool RemoveCustomTextFilter(const TSharedRef<FSequencerTrackFilter_CustomText>& InFilter, const bool bInAddToConfig) = 0;
+
+	SEQUENCER_API virtual void EnableAllFilters(const bool bInEnable, const TArray<FString> InExceptionFilterNames) = 0;
+
+	SEQUENCER_API virtual void ActivateCommonFilters(const bool bInActivate, const TArray<FString> InExceptionFilterNames) = 0;
+
+	SEQUENCER_API virtual void ActivateAllEnabledFilters(const bool bInActivate, const TArray<FString> InExceptionFilterNames) = 0;
+
+	SEQUENCER_API virtual bool IsFilterActiveByDisplayName(const FString InFilterName) const = 0;
+	SEQUENCER_API virtual bool IsFilterEnabledByDisplayName(const FString InFilterName) const = 0;
+	SEQUENCER_API virtual bool SetFilterActiveByDisplayName(const FString InFilterName, const bool bInActive, const bool bInRequestFilterUpdate = true) = 0;
+	SEQUENCER_API virtual bool SetFilterEnabledByDisplayName(const FString InFilterName, const bool bInEnabled, const bool bInRequestFilterUpdate = true) = 0;
+
+	SEQUENCER_API virtual TArray<FText> GetFilterDisplayNames() const = 0;
+	SEQUENCER_API virtual TArray<FText> GetCustomTextFilterNames() const = 0;
+
+	SEQUENCER_API virtual int32 GetTotalDisplayNodeCount() const = 0;
+	SEQUENCER_API virtual int32 GetFilteredDisplayNodeCount() const  = 0;
+
+	SEQUENCER_API virtual void HideSelectedTracks() = 0;
+	SEQUENCER_API virtual void IsolateSelectedTracks() = 0;
+
+	SEQUENCER_API virtual void ShowOnlyLocationCategoryGroups() = 0;
+	SEQUENCER_API virtual void ShowOnlyRotationCategoryGroups() = 0;
+	SEQUENCER_API virtual void ShowOnlyScaleCategoryGroups() = 0;
+
+	SEQUENCER_API virtual bool HasSelectedTracks() const = 0;
+
+	SEQUENCER_API virtual FSequencerFilterData& GetFilterData() = 0;
+};
diff --git a/Engine/Source/Editor/Sequencer/Public/Filters/SequencerFilterBarConfig.h b/Engine/Source/Editor/Sequencer/Public/Filters/SequencerFilterBarConfig.h
new file mode 100644
--- /dev/null
+++ b/Engine/Source/Editor/Sequencer/Public/Filters/SequencerFilterBarConfig.h
@@ -0,0 +1,66 @@
+﻿// Copyright Epic Games, Inc. All Rights Reserved.
+
+#pragma once
+
+#include "Filters/SBasicFilterBar.h"
+#include "SequencerFilterBarConfig.generated.h"
+
+USTRUCT()
+struct FSequencerFilterSet
+{
+	GENERATED_BODY()
+
+public:
+	UPROPERTY()
+	FString Label;
+
+	/** Enabled and active states of common filters. Enabled if in the map. Active if the value of the key is true. */
+	UPROPERTY()
+	TMap<FString, bool> EnabledStates;
+
+	UPROPERTY()
+	FString TextFilterString;
+};
+
+USTRUCT()
+struct FSequencerFilterBarConfig
+{
+	GENERATED_BODY()
+
+public:
+	/** Common Filters */
+
+	bool IsFilterEnabled(const FString& InFilterName) const;
+	bool SetFilterEnabled(const FString& InFilterName, const bool bInActive);
+
+	bool IsFilterActive(const FString& InFilterName) const;
+	bool SetFilterActive(const FString& InFilterName, const bool bInActive);
+
+	const FSequencerFilterSet& GetCommonActiveSet() const;
+
+	/** Custom Text Filters */
+
+	TArray<FCustomTextFilterData>& GetCustomTextFilters();
+	bool HasCustomTextFilter(const FString& InFilterName) const;
+	FCustomTextFilterData* FindCustomTextFilter(const FString& InFilterName);
+	bool AddCustomTextFilter(FCustomTextFilterData InFilterData);
+	bool RemoveCustomTextFilter(const FString& InFilterName);
+
+	/** Filter Bar Layout */
+
+	EFilterBarLayout GetFilterBarLayout() const;
+	void SetFilterBarLayout(const EFilterBarLayout InLayout);
+
+protected:
+	/** The currently active set of common and custom text filters that should be restored on editor load */
+	UPROPERTY()
+	FSequencerFilterSet ActiveFilters;
+
+	/** User created custom text filters */
+	UPROPERTY()
+	TArray<FCustomTextFilterData> CustomTextFilters;
+
+	/** The layout style for the filter bar widget */
+	UPROPERTY()
+	EFilterBarLayout FilterBarLayout = EFilterBarLayout::Vertical;
+};
diff --git a/Engine/Source/Editor/Sequencer/Public/Filters/SequencerFilterData.h b/Engine/Source/Editor/Sequencer/Public/Filters/SequencerFilterData.h
new file mode 100644
--- /dev/null
+++ b/Engine/Source/Editor/Sequencer/Public/Filters/SequencerFilterData.h
@@ -0,0 +1,55 @@
+﻿// Copyright Epic Games, Inc. All Rights Reserved.
+
+#pragma once
+
+#include "MovieSceneTrack.h"
+#include "MVVM/ViewModelPtr.h"
+
+namespace UE::Sequencer
+{
+	class IOutlinerExtension;
+}
+
+/** Represents a cache between nodes for a filter operation. */
+struct SEQUENCER_API FSequencerFilterData
+{
+	FSequencerFilterData(const FString& InRawFilterText);
+
+	bool operator==(const FSequencerFilterData& InRhs) const;
+	bool operator!=(const FSequencerFilterData& InRhs) const;
+
+	void Reset();
+
+	FString GetRawFilterText() const;
+
+	uint32 GetDisplayNodeCount() const;
+	uint32 GetTotalNodeCount() const;
+
+	uint32 GetFilterInCount() const;
+	uint32 GetFilterOutCount() const;
+
+	void IncrementTotalNodeCount();
+
+	void FilterInNode(UE::Sequencer::TWeakViewModelPtr<UE::Sequencer::IOutlinerExtension> InNodeWeak);
+	void FilterOutNode(UE::Sequencer::TWeakViewModelPtr<UE::Sequencer::IOutlinerExtension> InNodeWeak);
+
+	void FilterInParentChildNodes(const UE::Sequencer::TViewModelPtr<UE::Sequencer::IOutlinerExtension>& InNode
+		, const bool bInIncludeSelf, const bool bInIncludeParents, const bool bInIncludeChildren = false);
+
+	void FilterInNodeWithAncestors(const UE::Sequencer::TViewModelPtr<UE::Sequencer::IOutlinerExtension>& InNode);
+
+	bool ContainsFilterInNodes(const FSequencerFilterData& InOtherData) const;
+
+	bool IsFilteredOut(const UE::Sequencer::TViewModelPtr<UE::Sequencer::IOutlinerExtension>& InNode) const;
+
+	TMap<UE::Sequencer::TWeakViewModelPtr<UE::Sequencer::IOutlinerExtension>, TWeakObjectPtr<UMovieSceneTrack>> ResolvedTrackObjects;
+	TMap<UE::Sequencer::TWeakViewModelPtr<UE::Sequencer::IOutlinerExtension>, TWeakObjectPtr<>> ResolvedBoundObjects;
+
+protected:
+	FString RawFilterText;
+
+	uint32 TotalNodeCount = 0;
+
+	/** Nodes to be displayed in the UI */
+	TSet<UE::Sequencer::TWeakViewModelPtr<UE::Sequencer::IOutlinerExtension>> FilterInNodes;
+};
diff --git a/Engine/Source/Editor/Sequencer/Public/Filters/SequencerTextFilterExpressionContext.h b/Engine/Source/Editor/Sequencer/Public/Filters/SequencerTextFilterExpressionContext.h
new file mode 100644
--- /dev/null
+++ b/Engine/Source/Editor/Sequencer/Public/Filters/SequencerTextFilterExpressionContext.h
@@ -0,0 +1,60 @@
+﻿// Copyright Epic Games, Inc. All Rights Reserved.
+
+#pragma once
+
+#include "Filters/SequencerTrackFilterBase.h"
+#include "Misc/TextFilterExpressionEvaluator.h"
+
+class UMovieScene;
+class UMovieSceneSequence;
+class UMovieSceneTrack;
+
+enum class SEQUENCER_API ESequencerTextFilterValueType : uint8
+{
+	String,
+	Boolean,
+	Integer
+};
+
+/** Text expression context to test the given asset data against the current text filter */
+class FSequencerTextFilterExpressionContext : public ITextFilterExpressionContext
+{ 
+public:
+	SEQUENCER_API FSequencerTextFilterExpressionContext(ISequencerTrackFilters& InFilterInterface);
+
+	void SetFilterItem(FSequencerTrackFilterType InFilterItem, UMovieSceneTrack* const InTrackObject);
+
+	//~ Begin FSequencerTextFilterExpressionContext
+	virtual TSet<FName> GetKeys() const = 0;
+	virtual ESequencerTextFilterValueType GetValueType() const = 0;
+	virtual FText GetDescription() const = 0;
+	//~ End FSequencerTextFilterExpressionContext
+
+	//~ Begin ITextFilterExpressionContext
+
+	SEQUENCER_API virtual bool TestBasicStringExpression(const FTextFilterString& InValue
+		, const ETextFilterTextComparisonMode InTextComparisonMode) const override;
+
+	SEQUENCER_API virtual bool TestComplexExpression(const FName& InKey
+		, const FTextFilterString& InValue
+		, const ETextFilterComparisonOperation InComparisonOperation
+		, const ETextFilterTextComparisonMode InTextComparisonMode) const override;
+
+	//~ End ITextFilterExpressionContext
+
+protected:
+	UMovieSceneSequence* GetFocusedMovieSceneSequence() const;
+	UMovieScene* GetFocusedGetMovieScene() const;
+
+	bool CompareFStringForExactBool(const FTextFilterString& InValue, const bool bInPassedFilter) const;
+
+	bool CompareFStringForExactBool(const FTextFilterString& InValue
+		, const ETextFilterComparisonOperation InComparisonOperation
+		, const bool bInPassedFilter) const;
+
+	ISequencerTrackFilters& FilterInterface;
+
+	FSequencerTrackFilterType FilterItem;
+
+	TWeakObjectPtr<UMovieSceneTrack> WeakTrackObject;
+};
diff --git a/Engine/Source/Editor/Sequencer/Public/Filters/SequencerTrackFilterBase.h b/Engine/Source/Editor/Sequencer/Public/Filters/SequencerTrackFilterBase.h
new file mode 100644
--- /dev/null
+++ b/Engine/Source/Editor/Sequencer/Public/Filters/SequencerTrackFilterBase.h
@@ -0,0 +1,221 @@
+﻿// Copyright Epic Games, Inc. All Rights Reserved.
+
+#pragma once
+
+#include "GameFramework/Actor.h"
+#include "Filters/FilterBase.h"
+#include "Filters/ISequencerTrackFilters.h"
+#include "Misc/IFilter.h"
+#include "MovieSceneTrack.h"
+#include "MovieSceneTrackEditor.h"
+#include "MVVM/ViewModelPtr.h"
+#include "Textures/SlateIcon.h"
+
+class ISequencer;
+class FMenuBuilder;
+class UMovieScene;
+class UMovieSceneSequence;
+struct FSequencerFilterData;
+
+namespace UE::Sequencer
+{
+	class IOutlinerExtension;
+
+	namespace ExtensionHooks
+	{
+		static FName Hierarchy(TEXT("Hierarchy"));
+		static FName Show(TEXT("Show"));
+	}
+}
+
+using FSequencerTrackFilterType = UE::Sequencer::FViewModelPtr;
+
+//////////////////////////////////////////////////////////////////////////
+//
+
+class FSequencerTrackFilter : public FFilterBase<FSequencerTrackFilterType>
+	, public TSharedFromThis<FSequencerTrackFilter>
+{
+public:
+	SEQUENCER_API static bool SupportsLevelSequence(UMovieSceneSequence* const InSequence);
+	SEQUENCER_API static bool SupportsUMGSequence(UMovieSceneSequence* const InSequence);
+
+	template<typename InTrackClassType>
+	static bool IsSequenceTrackSupported(UMovieSceneSequence* const InSequence)
+	{
+		const ETrackSupport Support = IsValid(InSequence)
+			? InSequence->IsTrackSupported(InTrackClassType::StaticClass())
+			: ETrackSupport::NotSupported;
+		return Support == ETrackSupport::Supported;
+	}
+
+	SEQUENCER_API static UMovieSceneTrack* ResolveMovieSceneTrackObject(FSequencerTrackFilterType InNode, FSequencerFilterData& FilterData);
+	SEQUENCER_API static UObject* ResolveTrackBoundObject(ISequencer& InSequencer, FSequencerTrackFilterType InNode, FSequencerFilterData& FilterData);
+
+	SEQUENCER_API static FText BuildTooltipTextForCommand(const FText& InBaseText, const TSharedPtr<FUICommandInfo>& InCommand);
+
+	SEQUENCER_API FSequencerTrackFilter(ISequencerTrackFilters& InOutFilterInterface, TSharedPtr<FFilterCategory>&& InCategory = nullptr);
+
+	//~ Begin IFilter
+	SEQUENCER_API virtual bool PassesFilter(FSequencerTrackFilterType InItem) const override { return true; };
+	//~ End IFilter
+
+	//~ Begin FFilterBase
+
+	virtual FText GetDisplayName() const override { return FText::GetEmpty(); }
+	SEQUENCER_API virtual FText GetToolTipText() const override;
+	virtual FLinearColor GetColor() const override { return FLinearColor::White; }
+	virtual bool IsInverseFilter() const override { return false; }
+
+	virtual void ActiveStateChanged(bool bActive) override {}
+
+	virtual void ModifyContextMenu(FMenuBuilder& MenuBuilder) override {}
+
+	virtual void SaveSettings(const FString& IniFilename, const FString& IniSection, const FString& SettingsString) const override {}
+	virtual void LoadSettings(const FString& IniFilename, const FString& IniSection, const FString& SettingsString) override {}
+
+	//~ End FFilterBase
+
+	//~ Begin FSequencerTrackFilter
+
+	SEQUENCER_API virtual FText GetDefaultToolTipText() const { return FText(); }
+	SEQUENCER_API virtual FSlateIcon GetIcon() const { return FSlateIcon(); }
+
+	SEQUENCER_API virtual bool IsCustomTextFilter() const { return false; }
+
+	/** Returns whether this filter needs reevaluating any time track values have been modified, not just tree changes */
+	SEQUENCER_API virtual bool ShouldUpdateOnTrackValueChanged() const { return false; }
+
+	/** Returns whether the filter supports the sequence type */
+	SEQUENCER_API virtual bool SupportsSequence(UMovieSceneSequence* const InSequence) const;
+
+	SEQUENCER_API virtual void BindCommands();
+
+	SEQUENCER_API virtual TSharedPtr<FUICommandInfo> GetToggleCommand() const { return nullptr; }
+
+	virtual TSubclassOf<UMovieSceneTrack> GetTrackClass() const { return nullptr; };
+
+	//~ End FSequencerTrackFilter
+
+	SEQUENCER_API void BroadcastChangedEvent() const;
+
+	SEQUENCER_API ISequencerTrackFilters& GetFilterInterface() const;
+	SEQUENCER_API ISequencer& GetSequencer() const;
+
+	bool CanToggleFilter() const;
+	void ToggleFilter();
+
+	SEQUENCER_API void MapToggleAction(const TSharedPtr<FUICommandInfo>& InCommand);
+
+protected:
+	UMovieSceneSequence* GetFocusedMovieSceneSequence() const;
+	UMovieScene* GetFocusedGetMovieScene() const;
+
+	SEQUENCER_API UObject* ResolveTrackBoundObject(FSequencerTrackFilterType InNode, FSequencerFilterData& FilterData) const;
+
+	ISequencerTrackFilters& FilterInterface;
+
+private:
+	//~ Begin IFilter
+	/** Hide behind private to force use of GetIcon() instead */
+	virtual FName GetIconName() const override { return FName(); }
+	//~ End IFilter
+};
+
+//////////////////////////////////////////////////////////////////////////
+//
+
+/** Base filter for filtering Sequencer tracks based track model */
+template<typename InModelType>
+class FSequencerTrackFilter_ModelType : public FSequencerTrackFilter
+{
+public:
+	FSequencerTrackFilter_ModelType(ISequencerTrackFilters& InOutFilterInterface, TSharedPtr<FFilterCategory> InCategory)
+		: FSequencerTrackFilter(InOutFilterInterface, MoveTemp(InCategory))
+	{}
+
+	//~ Begin IFilter
+	virtual bool PassesFilter(FSequencerTrackFilterType InItem) const override
+	{
+		const UE::Sequencer::TViewModelPtr<InModelType> Model = InItem->FindAncestorOfType<InModelType>();
+		return Model.IsValid(); // show child tracks
+		//return InItem->IsA<InModelType>(); // don't show child tracks
+	}
+	//~ End IFilter
+};
+
+//////////////////////////////////////////////////////////////////////////
+//
+
+/** Base filter for filtering Sequencer tracks based object class type */
+template<typename InClassType>
+class FSequencerTrackFilter_ClassType : public FSequencerTrackFilter
+{
+public:
+	FSequencerTrackFilter_ClassType(ISequencerTrackFilters& InOutFilterInterface, TSharedPtr<FFilterCategory> InCategory)
+		: FSequencerTrackFilter(InOutFilterInterface, MoveTemp(InCategory))
+	{}
+
+	//~ Begin IFilter
+	virtual bool PassesFilter(FSequencerTrackFilterType InItem) const override
+	{
+		FSequencerFilterData& FilterData = FilterInterface.GetFilterData();
+		const UMovieSceneTrack* const TrackObject = ResolveMovieSceneTrackObject(InItem, FilterData);
+		return IsValid(TrackObject) && TrackObject->IsA(InClassType::StaticClass());
+	}
+	//~ End IFilter
+
+	//~ Begin FSequencerTrackFilter
+	TSubclassOf<UMovieSceneTrack> GetTrackClass() const override
+	{
+		return InClassType::StaticClass();
+	}
+	//~ End FSequencerTrackFilter
+};
+
+//////////////////////////////////////////////////////////////////////////
+//
+
+/** Base filter for filtering Sequencer tracks based object component type */
+template<typename InComponentType>
+class FSequencerTrackFilter_ComponentType : public FSequencerTrackFilter
+{
+public:
+	FSequencerTrackFilter_ComponentType(ISequencerTrackFilters& InOutFilterInterface, TSharedPtr<FFilterCategory> InCategory)
+		: FSequencerTrackFilter(InOutFilterInterface, MoveTemp(InCategory))
+	{}
+
+	//~ Begin IFilter
+	virtual bool PassesFilter(FSequencerTrackFilterType InItem) const override
+	{
+		FSequencerFilterData& FilterData = FilterInterface.GetFilterData();
+
+		const UMovieSceneTrack* const TrackObject = ResolveMovieSceneTrackObject(InItem, FilterData);
+		if (IsValid(TrackObject) && TrackObject->IsA(InComponentType::StaticClass()))
+		{
+			return true;
+		}
+
+		const TWeakObjectPtr<> BoundObject = ResolveTrackBoundObject(GetSequencer(), InItem, FilterData);
+		if (BoundObject.IsValid() && BoundObject->IsA(InComponentType::StaticClass()))
+		{
+			return true;
+		}
+
+		const AActor* const Actor = Cast<const AActor>(BoundObject.Get());
+		if (IsValid(Actor) && Actor->FindComponentByClass(InComponentType::StaticClass()))
+		{
+			return true;
+		}
+
+		return false;
+	}
+	//~ End IFilter
+
+	//~ Begin FSequencerTrackFilter
+	TSubclassOf<UMovieSceneTrack> GetTrackClass() const override
+	{
+		return InComponentType::StaticClass();
+	}
+	//~ End FSequencerTrackFilter
+};
diff --git a/Engine/Source/Editor/Sequencer/Public/ISequencer.h b/Engine/Source/Editor/Sequencer/Public/ISequencer.h
--- a/Engine/Source/Editor/Sequencer/Public/ISequencer.h
+++ b/Engine/Source/Editor/Sequencer/Public/ISequencer.h
@@ -10,6 +10,7 @@
 #include "Delegates/Delegate.h"
 #include "ViewRangeInterpolation.h"
 #include "Evaluation/MovieSceneSequenceTransform.h"
+#include "Filters/ISequencerTrackFilters.h"
 #include "HAL/Platform.h"
 #include "IMovieScenePlayer.h"
 #include "IMovieScenePlayer.h"
@@ -825,14 +826,20 @@ public:
 public:
 
 	/** Sets the specified track filter to be on or off */
+	UE_DEPRECATED(5.5, "SetTrackFilterEnabled() has been deprecated. Use GetFilterInterface()->SetFilterActiveByDisplayName() instead.")
 	virtual void SetTrackFilterEnabled(const FText& InTrackFilterName, bool bEnabled) = 0;
 
 	/** Gets whether the specified track filter is on/off */
+	UE_DEPRECATED(5.5, "IsTrackFilterEnabled() has been deprecated. Use GetFilterInterface()->IsFilterActiveByDisplayName() instead.")
 	virtual bool IsTrackFilterEnabled(const FText& InTrackFilterName) const = 0;
 
 	/** Gets all the available track filter names */
+	UE_DEPRECATED(5.5, "GetTrackFilterNames() has been deprecated. Use GetFilterInterface()->GetFilterDisplayNames() instead.")
 	virtual TArray<FText> GetTrackFilterNames() const = 0;
 
+	/** Gets the Sequencer filter interface used to manage filters */
+	SEQUENCER_API virtual TSharedRef<ISequencerTrackFilters> GetFilterInterface() const = 0;
+
 public:
 
 	/**
diff --git a/Engine/Source/Editor/Sequencer/Public/SequencerCommands.h b/Engine/Source/Editor/Sequencer/Public/SequencerCommands.h
--- a/Engine/Source/Editor/Sequencer/Public/SequencerCommands.h
+++ b/Engine/Source/Editor/Sequencer/Public/SequencerCommands.h
@@ -175,9 +175,6 @@ public:
 	/** Sort all nodes and descendants */
 	TSharedPtr< FUICommandInfo > SortAllNodesAndDescendants;
 
-	/** Reset all enabled filters */
-	TSharedPtr< FUICommandInfo > ResetFilters;
-
 	/** Sets the upper bound of the selection range */
 	TSharedPtr< FUICommandInfo > SetSelectionRangeEnd;
 
@@ -397,9 +394,6 @@ public:
 	/** Turns auto scroll on and off. */
 	TSharedPtr< FUICommandInfo > ToggleAutoScroll;
 
-	/** Toggles whether or not to show selected nodes only. */
-	TSharedPtr< FUICommandInfo > ToggleShowSelectedNodesOnly;
-
 	/** Toggles whether or not the curve editor should be shown. */
 	TSharedPtr< FUICommandInfo > ToggleShowCurveEditor;
 
diff --git a/Engine/Source/Editor/Sequencer/Public/SequencerSettings.h b/Engine/Source/Editor/Sequencer/Public/SequencerSettings.h
--- a/Engine/Source/Editor/Sequencer/Public/SequencerSettings.h
+++ b/Engine/Source/Editor/Sequencer/Public/SequencerSettings.h
@@ -3,6 +3,7 @@
 #pragma once
 
 #include "CoreMinimal.h"
+#include "Filters/SequencerFilterBarConfig.h"
 #include "UObject/ObjectMacros.h"
 #include "UObject/Object.h"
 #include "UObject/Package.h"
@@ -480,11 +481,6 @@ public:
 	/** Sets the saved view density */
 	void SetViewDensity(FName InViewDensity);
 
-	/** Gets whether the given track filter is enabled */
-	bool IsTrackFilterEnabled(const FString& TrackFilter) const;
-	/** Sets whether the track filter should be enabled/disabled */
-	void SetTrackFilterEnabled(const FString& TrackFilter, bool bEnabled);
-
 	/** Get outliner column visibility in display order */
 	TArray<FColumnVisibilitySetting> GetOutlinerColumnSettings() const { return ColumnVisibilitySettings; }
 	/** Sets the visibility of outliner columns in display order */
@@ -495,6 +491,28 @@ public:
 	/** Sets the sidebar state to be restored on Sequencer initialize */
 	void SetSidebarState(const FSidebarState& InSidebarState);
 
+	FSequencerFilterBarConfig& FindOrAddTrackFilterBar(const FName InIdentifier, const bool bInSaveConfig);
+	FSequencerFilterBarConfig* FindTrackFilterBar(const FName InIdentifier);
+	bool RemoveTrackFilterBar(const FName InIdentifier);
+
+	bool GetIncludePinnedInFilter() const;
+	void SetIncludePinnedInFilter(const bool bInIncludePinned);
+
+	bool GetAutoExpandNodesOnFilterPass() const;
+	void SetAutoExpandNodesOnFilterPass(const bool bInIncludePinned);
+
+	bool GetUseFilterSubmenusForCategories() const;
+	void SetUseFilterSubmenusForCategories(const bool bInUseFilterSubmenusForCategories);
+
+	bool IsFilterBarVisible() const;
+	void SetFilterBarVisible(const bool bInVisible);
+
+	EFilterBarLayout GetFilterBarLayout() const;
+	void SetFilterBarLayout(const EFilterBarLayout InLayout);
+
+	float GetLastFilterBarSizeCoefficient() const;
+	void SetLastFilterBarSizeCoefficient(const float bInSizeCoefficient);
+
 protected:
 
 	/** The auto change mode (auto-key, auto-track or none). */
@@ -766,10 +784,6 @@ protected:
 	UPROPERTY(config, EditAnywhere, Category = General)
 	FName ViewDensity;
 
-	/** The track filters that are enabled */
-	UPROPERTY(config, EditAnywhere, Category = General)
-	TArray<FString> TrackFilters;
-
 	/** List of all columns and their visibility, in the order to be displayed in the outliner view */
 	UPROPERTY(config, EditAnywhere, Category = General)
 	TArray<FColumnVisibilitySetting> ColumnVisibilitySettings;
@@ -778,6 +792,34 @@ protected:
 	UPROPERTY(config)
 	TMap<FName, FSidebarState> SidebarState;
 
+	/** Saved settings for each unique filter bar instance mapped by instance identifier */
+	UPROPERTY(config, EditAnywhere, Category = Filtering)
+	TMap<FName, FSequencerFilterBarConfig> TrackFilterBars;
+
+	/** Apply filtering to pinned tracks that would otherwise ignore filters */
+	UPROPERTY(config, EditAnywhere, Category = Filtering)
+	bool bIncludePinnedInFilter;
+
+	/** Automatically expand tracks that pass filters */
+	UPROPERTY(config, EditAnywhere, Category = Filtering)
+	bool bAutoExpandNodesOnFilterPass;
+
+	/** Display the filter menu categories as submenus instead of sections */
+	UPROPERTY(config, EditAnywhere, Category = Filtering)
+	bool bUseFilterSubmenusForCategories;
+
+	/** Last saved visibility of the filter bar to restore after closed */
+	UPROPERTY(config, EditAnywhere, Category = Filtering)
+	bool bFilterBarVisible;
+
+	/** Last saved layout orientation of the filter bar to restore after closed */
+	UPROPERTY(config, EditAnywhere, Category = Filtering)
+	EFilterBarLayout LastFilterBarLayout;
+
+	/** Last saved size of the filter bar to restore after closed */
+	UPROPERTY(config)
+	float LastFilterBarSizeCoefficient;
+
 	FOnEvaluateSubSequencesInIsolationChanged OnEvaluateSubSequencesInIsolationChangedEvent;
 	FOnShowSelectedNodesOnlyChanged OnShowSelectedNodesOnlyChangedEvent;
 	FOnAllowEditsModeChanged OnAllowEditsModeChangedEvent;
diff --git a/Engine/Source/Editor/Sequencer/Public/SequencerTrackFilterBase.h b/Engine/Source/Editor/Sequencer/Public/SequencerTrackFilterBase.h
--- a/Engine/Source/Editor/Sequencer/Public/SequencerTrackFilterBase.h
+++ b/Engine/Source/Editor/Sequencer/Public/SequencerTrackFilterBase.h
@@ -1,109 +1,13 @@
 // Copyright Epic Games, Inc. All Rights Reserved.
 
+// HEADER_UNIT_SKIP - Deprecated
+/**
+ * DEPRECATED SequencerTrackFilterBase.h Use Filters/SequencerTrackFilterBase.h instead
+ */
+ 
 #pragma once
 
-#include "CoreMinimal.h"
-#include "Textures/SlateIcon.h"
-#include "GameFramework/Actor.h"
-#include "Misc/IFilter.h"
-#include "MovieSceneSequence.h"
+#include "Misc/CoreMiscDefines.h"
 
-class FMenuBuilder;
-class FUICommandList;
-class ISequencer;
-
-typedef const UObject* FTrackFilterType;
-
-struct FMovieSceneChannel;
-
-class FSequencerTrackFilter : public IFilter<FTrackFilterType>
-{
-public:
-	/** Returns the system name for this filter */
-	virtual FString GetName() const = 0;
-
-	/** Returns the human readable name for this filter */
-	virtual FText GetDisplayName() const = 0;
-
-	/** Returns the tooltip for this filter, shown in the filters menu */
-	virtual FText GetToolTipText() const = 0;
-
-	/** Returns the icon to use in menu entries */
-	virtual FSlateIcon GetIcon() const { return FSlateIcon(); }
-
-	/** Returns whether the filter supports the sequence type */
-	virtual bool SupportsSequence(UMovieSceneSequence* InSequence) const
-	{
-		static UClass* LevelSequenceClass = FindObject<UClass>(nullptr, TEXT("/Script/LevelSequence.LevelSequence"), true);
-		return InSequence != nullptr && (LevelSequenceClass != nullptr && InSequence->GetClass()->IsChildOf(LevelSequenceClass));
-	}
-
-	/** Returns whether this filter needs reevaluating any time track values have been modified, not just tree changes */
-	virtual bool ShouldUpdateOnTrackValueChanged() const { return false; }
-
-	// IFilter implementation
-	DECLARE_DERIVED_EVENT(FSequencerTrackFilter, IFilter<FTrackFilterType>::FChangedEvent, FChangedEvent);
-	virtual FChangedEvent& OnChanged() override { return ChangedEvent; }
-
-	/** Returns whether the specified Item passes the Filter's restrictions */
-	virtual bool PassesFilterWithDisplayName(FTrackFilterType InItem, const FText& InText) const 
-	{
-		return PassesFilter(InItem);
-	}
-
-	virtual bool PassesFilterChannel(FMovieSceneChannel* InMovieSceneChannel) const { return false; }
-
-	virtual void BindCommands(TSharedRef<FUICommandList> SequencerBindings, TSharedRef<FUICommandList> CurveEditorBindings, TWeakPtr<ISequencer> Sequencer) {}
-
-protected:
-	void BroadcastChangedEvent() const { ChangedEvent.Broadcast(); }
-
-private:
-	FChangedEvent ChangedEvent;
-};
-
-/** Helper template for building Sequencer track filters based on a track or object type */
-template<typename TrackType>
-class FSequencerTrackFilter_ClassType : public FSequencerTrackFilter
-{
-public:
-	// IFilter implementation
-	virtual bool PassesFilter(FTrackFilterType InItem) const override
-	{
-		if (InItem->IsA(TrackType::StaticClass()))
-		{
-			return true;
-		}
-		return false;
-	}
-};
-
-/** Helper template for building Sequencer track filters based on a component */
-template<typename ComponentType>
-class FSequencerTrackFilter_ComponentType : public FSequencerTrackFilter
-{
-public:
-	// IFilter implementation
-	virtual bool PassesFilter(FTrackFilterType InItem) const override
-	{
-		if (!InItem)
-		{
-			return false;
-		}
-
-		if (InItem->IsA(ComponentType::StaticClass()))
-		{
-			return true;
-		}
-
-		const AActor* Actor = Cast<const AActor>(InItem);
-		if (Actor)
-		{
-			if (Actor->FindComponentByClass(ComponentType::StaticClass()))
-			{
-				return true;
-			}
-		}
-		return false;
-	}
-};
+UE_DEPRECATED_HEADER(5.5, "Use Filters/SequencerTrackFilterBase.h instead of SequencerTrackFilterBase.h.")
+#include "Filters/SequencerTrackFilterBase.h"
diff --git a/Engine/Source/Editor/Sequencer/Public/SequencerTrackFilterExtension.h b/Engine/Source/Editor/Sequencer/Public/SequencerTrackFilterExtension.h
--- a/Engine/Source/Editor/Sequencer/Public/SequencerTrackFilterExtension.h
+++ b/Engine/Source/Editor/Sequencer/Public/SequencerTrackFilterExtension.h
@@ -1,20 +1,13 @@
 // Copyright Epic Games, Inc. All Rights Reserved.
 
+// HEADER_UNIT_SKIP - Deprecated
+/**
+ * DEPRECATED SequencerTrackFilterExtension.h Use Filters/SequencerTrackFilterExtension.h instead
+ */
+ 
 #pragma once
 
-#include "CoreMinimal.h"
-#include "UObject/ObjectMacros.h"
-#include "UObject/Object.h"
-#include "SequencerTrackFilterBase.h"
-#include "SequencerTrackFilterExtension.generated.h"
+#include "Misc/CoreMiscDefines.h"
 
-// Override this class in order to make an additional track filter available in Sequencer
-UCLASS(Abstract)
-class SEQUENCER_API USequencerTrackFilterExtension : public UObject
-{
-	GENERATED_BODY()
-
-public:
-	// Override this method to add new track filters
-	virtual void AddTrackFilterExtensions(TArray< TSharedRef<class FSequencerTrackFilter> >& InOutFilterList) const { }
-};
+UE_DEPRECATED_HEADER(5.5, "Use Filters/SequencerTrackFilterExtension.h instead of SequencerTrackFilterExtension.h.")
+#include "Filters/SequencerTrackFilterExtension.h"
diff --git a/Engine/Source/Editor/Sequencer/Sequencer.Build.cs b/Engine/Source/Editor/Sequencer/Sequencer.Build.cs
--- a/Engine/Source/Editor/Sequencer/Sequencer.Build.cs
+++ b/Engine/Source/Editor/Sequencer/Sequencer.Build.cs
@@ -48,6 +48,7 @@ public class Sequencer : ModuleRules
 				"ToolMenus",
 				"ToolWidgets",
 				"TypedElementFramework",
+				"TypedElementRuntime",
 				"UniversalObjectLocator",
 				"UniversalObjectLocatorEditor",
 				}
diff --git a/Engine/Source/Runtime/SlateCore/Private/Styling/StarshipCoreStyle.cpp b/Engine/Source/Runtime/SlateCore/Private/Styling/StarshipCoreStyle.cpp
--- a/Engine/Source/Runtime/SlateCore/Private/Styling/StarshipCoreStyle.cpp
+++ b/Engine/Source/Runtime/SlateCore/Private/Styling/StarshipCoreStyle.cpp
@@ -415,8 +415,9 @@ TSharedRef<ISlateStyle> FStarshipCoreStyle::Create()
 		Style->Set("Icons.Layout", new IMAGE_BRUSH_SVG("Starship/Common/Layout", Icon16x16));
 		Style->Set("Icons.Recent", new IMAGE_BRUSH_SVG("Starship/Common/Recent", Icon16x16));
 
+		Style->Set("Icons.Badge", new IMAGE_BRUSH_SVG("Starship/Common/badge", Icon16x16));
 		Style->Set("Icons.BadgeModified", new IMAGE_BRUSH_SVG("Starship/Common/badge-modified", Icon16x16));
-		
+
 		// Toolbar Size Icons
 		Style->Set("Icons.Toolbar.Settings", new IMAGE_BRUSH_SVG("Starship/Common/settings", Icon20x20));
 
