[Sequencer] Advanced Filtering: Change filter bar references to weak pointers

diff --git a/Engine/Source/Editor/Sequencer/Private/Filters/Filters/SequencerTrackFilterBase.cpp b/Engine/Source/Editor/Sequencer/Private/Filters/Filters/SequencerTrackFilterBase.cpp
--- a/Engine/Source/Editor/Sequencer/Private/Filters/Filters/SequencerTrackFilterBase.cpp
+++ b/Engine/Source/Editor/Sequencer/Private/Filters/Filters/SequencerTrackFilterBase.cpp
@@ -39,11 +39,6 @@ void FSequencerTrackFilter::BindCommands()
 	}
 }
 
-void FSequencerTrackFilter::BroadcastChangedEvent() const
-{
-	ChangedEvent.Broadcast();
-}
-
 ISequencerTrackFilters& FSequencerTrackFilter::GetFilterInterface() const
 {
 	return FilterInterface;
diff --git a/Engine/Source/Editor/Sequencer/Private/Filters/Menus/SequencerFilterBarContext.h b/Engine/Source/Editor/Sequencer/Private/Filters/Menus/SequencerFilterBarContext.h
--- a/Engine/Source/Editor/Sequencer/Private/Filters/Menus/SequencerFilterBarContext.h
+++ b/Engine/Source/Editor/Sequencer/Private/Filters/Menus/SequencerFilterBarContext.h
@@ -15,14 +15,14 @@ class USequencerFilterBarContext : public UObject
 public:
 	void Init(const TWeakPtr<FSequencerFilterBar>& InFilterBarWeak)
 	{
-		FilterBarWeak = InFilterBarWeak;
+		WeakFilterBar = InFilterBarWeak;
 	}
 
 	TSharedPtr<FSequencerFilterBar> GetFilterBar() const
 	{
-		return FilterBarWeak.IsValid() ? FilterBarWeak.Pin() : nullptr;
+		return WeakFilterBar.IsValid() ? WeakFilterBar.Pin() : nullptr;
 	}
 
 protected:
-	TWeakPtr<FSequencerFilterBar> FilterBarWeak;
+	TWeakPtr<FSequencerFilterBar> WeakFilterBar;
 };
diff --git a/Engine/Source/Editor/Sequencer/Private/Filters/Menus/SequencerFilterBarContextMenu.cpp b/Engine/Source/Editor/Sequencer/Private/Filters/Menus/SequencerFilterBarContextMenu.cpp
--- a/Engine/Source/Editor/Sequencer/Private/Filters/Menus/SequencerFilterBarContextMenu.cpp
+++ b/Engine/Source/Editor/Sequencer/Private/Filters/Menus/SequencerFilterBarContextMenu.cpp
@@ -28,16 +28,27 @@ TSharedRef<SWidget> FSequencerFilterBarContextMenu::CreateMenu(const TSharedRef<
 
 void FSequencerFilterBarContextMenu::PopulateMenu(UToolMenu* const InMenu)
 {
-	UToolMenu& MenuRef = *InMenu;
+	if (!IsValid(InMenu))
+	{
+		return;
+	}
+
+	USequencerFilterBarContext* const Context = InMenu->FindContext<USequencerFilterBarContext>();
+	if (!IsValid(Context))
+	{
+		return;
+	}
 
-	CurrentContext = MenuRef.FindContext<USequencerFilterBarContext>();
+	WeakFilterBar = Context->GetFilterBar();
+
+	UToolMenu& MenuRef = *InMenu;
 
 	PopulateOptionsSection(MenuRef);
 }
 
 void FSequencerFilterBarContextMenu::PopulateOptionsSection(UToolMenu& InMenu)
 {
-	const TSharedPtr<FSequencerFilterBar> FilterBar = CurrentContext->GetFilterBar();
+	const TSharedPtr<FSequencerFilterBar> FilterBar = WeakFilterBar.Pin();
 	if (!FilterBar.IsValid())
 	{
 		return;
@@ -107,7 +118,7 @@ void FSequencerFilterBarContextMenu::PopulateFilterBulkOptionsSection(UToolMenu&
 
 void FSequencerFilterBarContextMenu::OnActivateAllFilters(const bool bInActivate)
 {
-	const TSharedPtr<FSequencerFilterBar> FilterBar = CurrentContext->GetFilterBar();
+	const TSharedPtr<FSequencerFilterBar> FilterBar = WeakFilterBar.Pin();
 	if (!FilterBar.IsValid())
 	{
 		return;
@@ -118,7 +129,7 @@ void FSequencerFilterBarContextMenu::OnActivateAllFilters(const bool bInActivate
 
 void FSequencerFilterBarContextMenu::OnResetFilters()
 {
-	const TSharedPtr<FSequencerFilterBar> FilterBar = CurrentContext->GetFilterBar();
+	const TSharedPtr<FSequencerFilterBar> FilterBar = WeakFilterBar.Pin();
 	if (!FilterBar.IsValid())
 	{
 		return;
diff --git a/Engine/Source/Editor/Sequencer/Private/Filters/Menus/SequencerFilterBarContextMenu.h b/Engine/Source/Editor/Sequencer/Private/Filters/Menus/SequencerFilterBarContextMenu.h
--- a/Engine/Source/Editor/Sequencer/Private/Filters/Menus/SequencerFilterBarContextMenu.h
+++ b/Engine/Source/Editor/Sequencer/Private/Filters/Menus/SequencerFilterBarContextMenu.h
@@ -3,11 +3,9 @@
 #pragma once
 
 #include "Templates/SharedPointer.h"
-#include "UObject/WeakObjectPtr.h"
 
 class FSequencerFilterBar;
 class SWidget;
-class USequencerFilterBarContext;
 class UToolMenu;
 
 class FSequencerFilterBarContextMenu
@@ -24,5 +22,5 @@ protected:
 	void OnActivateAllFilters(const bool bInActivate);
 	void OnResetFilters();
 
-	TWeakObjectPtr<USequencerFilterBarContext> CurrentContext;
+	TWeakPtr<FSequencerFilterBar> WeakFilterBar;
 };
diff --git a/Engine/Source/Editor/Sequencer/Private/Filters/Menus/SequencerTrackFilterContextMenu.cpp b/Engine/Source/Editor/Sequencer/Private/Filters/Menus/SequencerTrackFilterContextMenu.cpp
--- a/Engine/Source/Editor/Sequencer/Private/Filters/Menus/SequencerTrackFilterContextMenu.cpp
+++ b/Engine/Source/Editor/Sequencer/Private/Filters/Menus/SequencerTrackFilterContextMenu.cpp
@@ -44,14 +44,16 @@ void FSequencerTrackFilterContextMenu::PopulateMenu(UToolMenu* const InMenu)
 		return;
 	}
 
-	UToolMenu& MenuRef = *InMenu;
-
-	CurrentContext = MenuRef.FindContext<USequencerFilterMenuContext>();
-	if (!CurrentContext.IsValid())
+	USequencerFilterMenuContext* const Context = InMenu->FindContext<USequencerFilterMenuContext>();
+	if (!IsValid(Context))
 	{
 		return;
 	}
 
+	WeakFilterWidget = Context->GetFilterWidget();
+
+	UToolMenu& MenuRef = *InMenu;
+
 	PopulateFilterOptionsSection(MenuRef);
 	PopulateCustomFilterOptionsSection(MenuRef);
 	PopulateBulkOptionsSection(MenuRef);
@@ -59,7 +61,7 @@ void FSequencerTrackFilterContextMenu::PopulateMenu(UToolMenu* const InMenu)
 
 void FSequencerTrackFilterContextMenu::PopulateFilterOptionsSection(UToolMenu& InMenu)
 {
-	const TSharedPtr<FSequencerTrackFilter> Filter = CurrentContext->GetFilter();
+	const TSharedPtr<FSequencerTrackFilter> Filter = GetFilter();
 	if (!Filter.IsValid())
 	{
 		return;
@@ -86,7 +88,7 @@ void FSequencerTrackFilterContextMenu::PopulateFilterOptionsSection(UToolMenu& I
 void FSequencerTrackFilterContextMenu::PopulateCustomFilterOptionsSection(UToolMenu& InMenu)
 {
 	const TSharedPtr<FSequencerTrackFilter_CustomText> CustomTextFilter
-		= StaticCastSharedPtr<FSequencerTrackFilter_CustomText>(CurrentContext->GetFilter());
+		= StaticCastSharedPtr<FSequencerTrackFilter_CustomText>(GetFilter());
 	if (!CustomTextFilter.IsValid() || !CustomTextFilter->IsCustomTextFilter())
 	{
 		return;
@@ -151,7 +153,7 @@ void FSequencerTrackFilterContextMenu::PopulateBulkOptionsSection(UToolMenu& InM
 
 FText FSequencerTrackFilterContextMenu::GetFilterDisplayName() const
 {
-	const TSharedPtr<FSequencerTrackFilter> Filter = CurrentContext->GetFilter();
+	const TSharedPtr<FSequencerTrackFilter> Filter = GetFilter();
 	if (!Filter.IsValid())
 	{
 		return FText::GetEmpty();
@@ -162,7 +164,7 @@ FText FSequencerTrackFilterContextMenu::GetFilterDisplayName() const
 
 void FSequencerTrackFilterContextMenu::OnDisableFilter()
 {
-	const TSharedPtr<FSequencerTrackFilter> Filter = CurrentContext->GetFilter();
+	const TSharedPtr<FSequencerTrackFilter> Filter = GetFilter();
 	if (!Filter.IsValid())
 	{
 		return;
@@ -174,7 +176,7 @@ void FSequencerTrackFilterContextMenu::OnDisableFilter()
 
 void FSequencerTrackFilterContextMenu::OnResetFilters()
 {
-	const TSharedPtr<FSequencerTrackFilter> Filter = CurrentContext->GetFilter();
+	const TSharedPtr<FSequencerTrackFilter> Filter = GetFilter();
 	if (!Filter.IsValid())
 	{
 		return;
@@ -185,7 +187,7 @@ void FSequencerTrackFilterContextMenu::OnResetFilters()
 
 void FSequencerTrackFilterContextMenu::OnActivateWithFilterException()
 {
-	const TSharedPtr<FSequencerTrackFilter> Filter = CurrentContext->GetFilter();
+	const TSharedPtr<FSequencerTrackFilter> Filter = GetFilter();
 	if (!Filter.IsValid())
 	{
 		return;
@@ -200,7 +202,7 @@ void FSequencerTrackFilterContextMenu::OnActivateWithFilterException()
 
 void FSequencerTrackFilterContextMenu::OnActivateAllFilters(const bool bInActivate)
 {
-	const TSharedPtr<FSequencerTrackFilter> Filter = CurrentContext->GetFilter();
+	const TSharedPtr<FSequencerTrackFilter> Filter = GetFilter();
 	if (!Filter.IsValid())
 	{
 		return;
@@ -211,7 +213,7 @@ void FSequencerTrackFilterContextMenu::OnActivateAllFilters(const bool bInActiva
 
 void FSequencerTrackFilterContextMenu::OnEditFilter()
 {
-	const TSharedPtr<FSequencerTrackFilter> Filter = CurrentContext->GetFilter();
+	const TSharedPtr<FSequencerTrackFilter> Filter = GetFilter();
 	if (!Filter.IsValid() || !Filter->IsCustomTextFilter())
 	{
 		return;
@@ -230,7 +232,7 @@ void FSequencerTrackFilterContextMenu::OnEditFilter()
 
 void FSequencerTrackFilterContextMenu::OnDeleteFilter()
 {
-	const TSharedPtr<FSequencerTrackFilter> Filter = CurrentContext->GetFilter();
+	const TSharedPtr<FSequencerTrackFilter> Filter = GetFilter();
 	if (!Filter.IsValid() || !Filter->IsCustomTextFilter())
 	{
 		return;
@@ -262,4 +264,9 @@ void FSequencerTrackFilterContextMenu::OnDeleteFilter()
 	}
 }
 
+const TSharedPtr<FSequencerTrackFilter> FSequencerTrackFilterContextMenu::GetFilter() const
+{
+	return WeakFilterWidget.IsValid() ? WeakFilterWidget.Pin()->GetFilter() : nullptr;
+}
+
 #undef LOCTEXT_NAMESPACE
diff --git a/Engine/Source/Editor/Sequencer/Private/Filters/Menus/SequencerTrackFilterContextMenu.h b/Engine/Source/Editor/Sequencer/Private/Filters/Menus/SequencerTrackFilterContextMenu.h
--- a/Engine/Source/Editor/Sequencer/Private/Filters/Menus/SequencerTrackFilterContextMenu.h
+++ b/Engine/Source/Editor/Sequencer/Private/Filters/Menus/SequencerTrackFilterContextMenu.h
@@ -3,11 +3,10 @@
 #pragma once
 
 #include "Templates/SharedPointer.h"
-#include "UObject/WeakObjectPtr.h"
 
+class FSequencerTrackFilter;
 class SSequencerFilter;
 class SWidget;
-class USequencerFilterMenuContext;
 class UToolMenu;
 
 class FSequencerTrackFilterContextMenu
@@ -34,5 +33,7 @@ protected:
 	void OnEditFilter();
 	void OnDeleteFilter();
 
-	TWeakObjectPtr<USequencerFilterMenuContext> CurrentContext;
+	const TSharedPtr<FSequencerTrackFilter> GetFilter() const;
+
+	TWeakPtr<SSequencerFilter> WeakFilterWidget;
 };
diff --git a/Engine/Source/Editor/Sequencer/Private/Filters/Menus/SequencerTrackFilterMenu.cpp b/Engine/Source/Editor/Sequencer/Private/Filters/Menus/SequencerTrackFilterMenu.cpp
--- a/Engine/Source/Editor/Sequencer/Private/Filters/Menus/SequencerTrackFilterMenu.cpp
+++ b/Engine/Source/Editor/Sequencer/Private/Filters/Menus/SequencerTrackFilterMenu.cpp
@@ -36,20 +36,21 @@ TSharedRef<SWidget> FSequencerTrackFilterMenu::CreateMenu(const TSharedRef<FSequ
 
 void FSequencerTrackFilterMenu::PopulateMenu(UToolMenu* const InMenu)
 {
+	if (!IsValid(InMenu))
+	{
+		return;
+	}
+
 	USequencerFilterBarContext* const Context = InMenu->FindContext<USequencerFilterBarContext>();
 	if (!IsValid(Context))
 	{
 		return;
 	}
 
+	WeakFilterBar = Context->GetFilterBar();
+
 	UToolMenu& Menu = *InMenu;
 
-	FilterBar = Context->GetFilterBar();
-	if (!FilterBar.IsValid())
-	{
-		return;
-	}
-	
 	PopulateFilterOptionsSection(Menu);
 	PopulateCustomsSection(Menu);
 	PopulateCommonFilterSections(Menu);
@@ -58,7 +59,11 @@ void FSequencerTrackFilterMenu::PopulateMenu(UToolMenu* const InMenu)
 
 void FSequencerTrackFilterMenu::PopulateCustomsSection(UToolMenu& InMenu)
 {
-	const TSharedRef<FSequencerFilterBar> FilterBarRef = FilterBar.ToSharedRef();
+	const TSharedPtr<FSequencerFilterBar> FilterBar = WeakFilterBar.Pin();
+	if (!FilterBar.IsValid())
+	{
+		return;
+	}
 
 	FToolMenuSection& Section = InMenu.FindOrAddSection(TEXT("UserCreated"));
 
@@ -73,6 +78,12 @@ void FSequencerTrackFilterMenu::PopulateCustomsSection(UToolMenu& InMenu)
 
 void FSequencerTrackFilterMenu::PopulateFilterOptionsSection(UToolMenu& InMenu)
 {
+	const TSharedPtr<FSequencerFilterBar> FilterBar = WeakFilterBar.Pin();
+	if (!FilterBar.IsValid())
+	{
+		return;
+	}
+
 	FToolMenuSection& Section = InMenu.FindOrAddSection(TEXT("FilterOptions"), LOCTEXT("FilterOptionsHeading", "Filters"));
 
 	Section.AddMenuEntry(TEXT("ResetFilters"),
@@ -94,6 +105,12 @@ void FSequencerTrackFilterMenu::PopulateFilterOptionsSection(UToolMenu& InMenu)
 
 void FSequencerTrackFilterMenu::PopulateCommonFilterSections(UToolMenu& InMenu)
 {
+	const TSharedPtr<FSequencerFilterBar> FilterBar = WeakFilterBar.Pin();
+	if (!FilterBar.IsValid())
+	{
+		return;
+	}
+
 	const TSharedPtr<SSequencerFilterBar> FilterBarWidget = FilterBar->GetWidget();
 	if (!FilterBarWidget.IsValid())
 	{
@@ -139,6 +156,12 @@ void FSequencerTrackFilterMenu::PopulateCommonFilterSections(UToolMenu& InMenu)
 
 void FSequencerTrackFilterMenu::PopulateOtherFilterSections(UToolMenu& InMenu)
 {
+	const TSharedPtr<FSequencerFilterBar> FilterBar = WeakFilterBar.Pin();
+	if (!FilterBar.IsValid())
+	{
+		return;
+	}
+
 	FToolMenuSection& Section = InMenu.FindOrAddSection(TEXT("OtherFilters"), LOCTEXT("OtherFiltersHeading", "Other Filters"));
 
 	const UMovieSceneSequence* const FocusedSequence = FilterBar->GetSequencer().GetFocusedMovieSceneSequence();
@@ -182,6 +205,12 @@ void FSequencerTrackFilterMenu::FillLevelFilterMenu(UToolMenu* const InMenu)
 		return;
 	}
 
+	const TSharedPtr<FSequencerFilterBar> FilterBar = WeakFilterBar.Pin();
+	if (!FilterBar.IsValid())
+	{
+		return;
+	}
+
 	const TSharedRef<FSequencerFilterBar> FilterBarRef = FilterBar.ToSharedRef();
 
 	FToolMenuSection& LevelsSection = InMenu->FindOrAddSection(TEXT("Levels"), LOCTEXT("LevelFilters", "Level Filters"));
@@ -215,6 +244,12 @@ void FSequencerTrackFilterMenu::FillGroupFilterMenu(UToolMenu* const InMenu)
 		return;
 	}
 
+	const TSharedPtr<FSequencerFilterBar> FilterBar = WeakFilterBar.Pin();
+	if (!FilterBar.IsValid())
+	{
+		return;
+	}
+
 	FSequencer& Sequencer = FilterBar->GetSequencer();
 
 	const TSharedRef<SSequencer> SequencerWidget = StaticCastSharedRef<SSequencer>(Sequencer.GetSequencerWidget());
@@ -265,6 +300,12 @@ void FSequencerTrackFilterMenu::FillCustomTextFiltersMenu(UToolMenu* const InMen
 		return;
 	}
 
+	const TSharedPtr<FSequencerFilterBar> FilterBar = WeakFilterBar.Pin();
+	if (!FilterBar.IsValid())
+	{
+		return;
+	}
+
 	const TSharedPtr<SSequencerFilterBar> FilterBarWidget = FilterBar->GetWidget();
 	if (!FilterBarWidget.IsValid())
 	{
@@ -295,6 +336,12 @@ void FSequencerTrackFilterMenu::FillCustomTextFiltersMenu(UToolMenu* const InMen
 		FSlateIcon(FAppStyle::Get().GetStyleSetName(), TEXT("Icons.PlusCircle")),
 		FUIAction(FExecuteAction::CreateLambda([this]()
 			{
+				const TSharedPtr<FSequencerFilterBar> FilterBar = WeakFilterBar.Pin();
+				if (!FilterBar.IsValid())
+				{
+					return;
+				}
+
 				SSequencerCustomTextFilterDialog::CreateWindow_AddCustomTextFilter(FilterBar.ToSharedRef(), FCustomTextFilterData(), nullptr);
 			})));
 
@@ -323,6 +370,12 @@ void FSequencerTrackFilterMenu::FillCustomTextFiltersMenu(UToolMenu* const InMen
 
 void FSequencerTrackFilterMenu::FillFiltersMenuCategory(FToolMenuSection& InOutSection, const TSharedRef<FFilterCategory> InMenuCategory)
 {
+	const TSharedPtr<FSequencerFilterBar> FilterBar = WeakFilterBar.Pin();
+	if (!FilterBar.IsValid())
+	{
+		return;
+	}
+
 	const TSharedRef<FSequencerFilterBar> FilterBarRef = FilterBar.ToSharedRef();
 	UMovieSceneSequence* const FocusedSequence = FilterBar->GetSequencer().GetFocusedMovieSceneSequence();
 
@@ -348,7 +401,7 @@ void FSequencerTrackFilterMenu::FillFiltersMenuCategory(FToolMenuSection& InOutS
 				FExecuteAction::CreateSP(this, &FSequencerTrackFilterMenu::OnFilterClicked, Filter),
 				FCanExecuteAction::CreateLambda([this]()
 					{
-						return !FilterBar->AreFiltersMuted();
+						return WeakFilterBar.IsValid() ? !WeakFilterBar.Pin()->AreFiltersMuted() : false;
 					}),
 				FIsActionChecked::CreateSP(FilterBarRef, &FSequencerFilterBar::IsFilterEnabled, Filter)
 			),
@@ -424,6 +477,12 @@ TSharedRef<SWidget> FSequencerTrackFilterMenu::ConstructCustomMenuItemWidget(con
 
 void FSequencerTrackFilterMenu::OnFilterCategoryClicked(const TSharedRef<FFilterCategory> InMenuCategory)
 {
+	const TSharedPtr<FSequencerFilterBar> FilterBar = WeakFilterBar.Pin();
+	if (!FilterBar.IsValid())
+	{
+		return;
+	}
+
 	const ECheckBoxState CategoryCheckState = GetFilterCategoryCheckedState(InMenuCategory);
 	const bool bIsCategoryEnabled = (CategoryCheckState == ECheckBoxState::Checked);
 	FilterBar->ActivateCommonFilters(bIsCategoryEnabled, { InMenuCategory }, {});
@@ -431,6 +490,12 @@ void FSequencerTrackFilterMenu::OnFilterCategoryClicked(const TSharedRef<FFilter
 
 ECheckBoxState FSequencerTrackFilterMenu::GetFilterCategoryCheckedState(const TSharedRef<FFilterCategory> InMenuCategory) const
 {
+	const TSharedPtr<FSequencerFilterBar> FilterBar = WeakFilterBar.Pin();
+	if (!FilterBar.IsValid())
+	{
+		return ECheckBoxState::Unchecked;
+	}
+
 	for (const TSharedRef<FSequencerTrackFilter>& Filter : FilterBar->GetCommonFilters())
 	{
 		if (Filter->GetCategory() == InMenuCategory && !FilterBar->IsFilterEnabled(Filter))
@@ -438,11 +503,18 @@ ECheckBoxState FSequencerTrackFilterMenu::GetFilterCategoryCheckedState(const TS
 			return ECheckBoxState::Unchecked;
 		}
 	}
+
 	return ECheckBoxState::Checked;
 }
 
 void FSequencerTrackFilterMenu::OnFilterClicked(const TSharedRef<FSequencerTrackFilter> InFilter)
 {
+	const TSharedPtr<FSequencerFilterBar> FilterBar = WeakFilterBar.Pin();
+	if (!FilterBar.IsValid())
+	{
+		return;
+	}
+
 	if (FilterBar->IsFilterEnabled(InFilter))
 	{
 		FilterBar->SetFilterEnabled(InFilter, false, true);
@@ -455,6 +527,12 @@ void FSequencerTrackFilterMenu::OnFilterClicked(const TSharedRef<FSequencerTrack
 
 void FSequencerTrackFilterMenu::OnCustomTextFilterClicked(const FText InFilterLabel)
 {
+	const TSharedPtr<FSequencerFilterBar> FilterBar = WeakFilterBar.Pin();
+	if (!FilterBar.IsValid())
+	{
+		return;
+	}
+
 	const TSharedPtr<FSequencerTrackFilter_CustomText> CustomTextFilter = FilterBar->FindCustomTextFilterByDisplayName(InFilterLabel.ToString());
 	if (!CustomTextFilter.IsValid())
 	{
@@ -475,15 +553,28 @@ void FSequencerTrackFilterMenu::OnCustomTextFilterClicked(const FText InFilterLa
 
 ECheckBoxState FSequencerTrackFilterMenu::GetCustomTextFilerCheckState(const FText InFilterLabel) const
 {
+	const TSharedPtr<FSequencerFilterBar> FilterBar = WeakFilterBar.Pin();
+	if (!FilterBar.IsValid())
+	{
+		return ECheckBoxState::Unchecked;
+	}
+
 	if (const TSharedPtr<FSequencerTrackFilter_CustomText> CustomTextFilter = FilterBar->FindCustomTextFilterByDisplayName(InFilterLabel.ToString()))
 	{
 		return FilterBar->IsFilterEnabled(CustomTextFilter.ToSharedRef()) ? ECheckBoxState::Checked : ECheckBoxState::Unchecked;
 	}
+
 	return ECheckBoxState::Unchecked;
 }
 
 void FSequencerTrackFilterMenu::OnEditCustomTextFilterClicked(const FText InFilterLabel)
 {
+	const TSharedPtr<FSequencerFilterBar> FilterBar = WeakFilterBar.Pin();
+	if (!FilterBar.IsValid())
+	{
+		return;
+	}
+
 	if (const TSharedPtr<FSequencerTrackFilter_CustomText> CustomTextFilter = FilterBar->FindCustomTextFilterByDisplayName(InFilterLabel.ToString()))
 	{
 		SSequencerCustomTextFilterDialog::CreateWindow_EditCustomTextFilter(FilterBar.ToSharedRef(), CustomTextFilter);
@@ -492,12 +583,24 @@ void FSequencerTrackFilterMenu::OnEditCustomTextFilterClicked(const FText InFilt
 
 void FSequencerTrackFilterMenu::OnTrackLevelFilterClicked(const FString InLevelName)
 {
+	const TSharedPtr<FSequencerFilterBar> FilterBar = WeakFilterBar.Pin();
+	if (!FilterBar.IsValid())
+	{
+		return;
+	}
+
 	const bool bNewActiveState = !FilterBar->IsLevelFilterActive(InLevelName);
 	FilterBar->ActivateLevelFilter(InLevelName, bNewActiveState);
 }
 
 void FSequencerTrackFilterMenu::ToggleAllLevelFilters()
 {
+	const TSharedPtr<FSequencerFilterBar> FilterBar = WeakFilterBar.Pin();
+	if (!FilterBar.IsValid())
+	{
+		return;
+	}
+
 	const ECheckBoxState AllState = GetAllLevelsCheckState();
 	switch (AllState)
 	{
@@ -513,10 +616,17 @@ void FSequencerTrackFilterMenu::ToggleAllLevelFilters()
 
 ECheckBoxState FSequencerTrackFilterMenu::GetAllLevelsCheckState() const
 {
+	const TSharedPtr<FSequencerFilterBar> FilterBar = WeakFilterBar.Pin();
+	if (!FilterBar.IsValid())
+	{
+		return ECheckBoxState::Unchecked;
+	}
+
 	if (FilterBar->HasAllLevelFiltersActive())
 	{
 		return ECheckBoxState::Unchecked;
 	}
+
 	return FilterBar->HasActiveLevelFilter() ? ECheckBoxState::Undetermined : ECheckBoxState::Checked;
 }
 
@@ -530,6 +640,12 @@ void FSequencerTrackFilterMenu::OnNodeGroupFilterClicked(UMovieSceneNodeGroup* c
 
 void FSequencerTrackFilterMenu::ToggleAllGroupFilters()
 {
+	const TSharedPtr<FSequencerFilterBar> FilterBar = WeakFilterBar.Pin();
+	if (!FilterBar.IsValid())
+	{
+		return;
+	}
+
 	const ECheckBoxState AllState = GetAllGroupsCheckState();
 	switch (AllState)
 	{
@@ -545,6 +661,12 @@ void FSequencerTrackFilterMenu::ToggleAllGroupFilters()
 
 ECheckBoxState FSequencerTrackFilterMenu::GetAllGroupsCheckState() const
 {
+	const TSharedPtr<FSequencerFilterBar> FilterBar = WeakFilterBar.Pin();
+	if (!FilterBar.IsValid())
+	{
+		return ECheckBoxState::Unchecked;
+	}
+
 	const UMovieSceneSequence* const FocusedMovieSequence = FilterBar->GetSequencer().GetFocusedMovieSceneSequence();
 	if (!IsValid(FocusedMovieSequence))
 	{
@@ -583,24 +705,25 @@ ECheckBoxState FSequencerTrackFilterMenu::GetAllGroupsCheckState() const
 
 bool FSequencerTrackFilterMenu::CanResetFilters() const
 {
-	if (FilterBar.IsValid())
+	const TSharedPtr<FSequencerFilterBar> FilterBar = WeakFilterBar.Pin();
+	if (!FilterBar.IsValid())
 	{
-		if (FilterBar->AreFiltersMuted())
-		{
-			return true;
-		}
-
-		const TArray<TSharedRef<FFilterCategory>> Categories = {
-			FilterBar->GetClassTypeCategory(),
-			FilterBar->GetComponentTypeCategory(),
-			FilterBar->GetMiscCategory()
-		};
-		const TArray<TSharedRef<FSequencerTrackFilter>> ClassAndCompFilters = FilterBar->GetCommonFilters(Categories);
+		return false;
+	}
 
-		return FilterBar->HasEnabledFilter(ClassAndCompFilters) || FilterBar->HasEnabledCustomTextFilters();
+	if (FilterBar->AreFiltersMuted())
+	{
+		return true;
 	}
 
-	return false;
+	const TArray<TSharedRef<FFilterCategory>> Categories = {
+		FilterBar->GetClassTypeCategory(),
+		FilterBar->GetComponentTypeCategory(),
+		FilterBar->GetMiscCategory()
+	};
+	const TArray<TSharedRef<FSequencerTrackFilter>> ClassAndCompFilters = FilterBar->GetCommonFilters(Categories);
+
+	return FilterBar->HasEnabledFilter(ClassAndCompFilters) || FilterBar->HasEnabledCustomTextFilters();
 }
 
 #undef LOCTEXT_NAMESPACE
diff --git a/Engine/Source/Editor/Sequencer/Private/Filters/Menus/SequencerTrackFilterMenu.h b/Engine/Source/Editor/Sequencer/Private/Filters/Menus/SequencerTrackFilterMenu.h
--- a/Engine/Source/Editor/Sequencer/Private/Filters/Menus/SequencerTrackFilterMenu.h
+++ b/Engine/Source/Editor/Sequencer/Private/Filters/Menus/SequencerTrackFilterMenu.h
@@ -9,7 +9,6 @@
 class FFilterCategory;
 class FSequencerFilterBar;
 class FSequencerTrackFilter;
-class SWidget;
 class UMovieSceneNodeGroup;
 class UToolMenu;
 struct FToolMenuSection;
@@ -62,5 +61,5 @@ protected:
 	bool CanResetFilters() const;
 
 private:
-	TSharedPtr<FSequencerFilterBar> FilterBar;
+	TWeakPtr<FSequencerFilterBar> WeakFilterBar;
 };
diff --git a/Engine/Source/Editor/Sequencer/Private/Filters/Menus/SequencerViewOptionsMenu.cpp b/Engine/Source/Editor/Sequencer/Private/Filters/Menus/SequencerViewOptionsMenu.cpp
--- a/Engine/Source/Editor/Sequencer/Private/Filters/Menus/SequencerViewOptionsMenu.cpp
+++ b/Engine/Source/Editor/Sequencer/Private/Filters/Menus/SequencerViewOptionsMenu.cpp
@@ -39,13 +39,19 @@ void FSequencerViewOptionsMenu::PopulateMenu(UToolMenu* const InMenu)
 		return;
 	}
 
-	UToolMenu& Menu = *InMenu;
+	USequencerMenuContext* const Context = InMenu->FindContext<USequencerMenuContext>();
+	if (!IsValid(Context))
+	{
+		return;
+	}
+
+	WeakSequencer = Context->GetSequencer();
 
-	CurrentContext = InMenu->FindContext<USequencerMenuContext>();
+	UToolMenu& Menu = *InMenu;
 
 	PopulateFiltersSection(Menu);
 	PopulateSortAndOrganizeSection(Menu);
-	PopulateFilterOptionssSection(Menu);
+	PopulateFilterOptionsSection(Menu);
 	PopulateLayoutSection(Menu);
 }
 
@@ -74,7 +80,7 @@ void FSequencerViewOptionsMenu::PopulateFiltersSection(UToolMenu& InMenu)
 
 void FSequencerViewOptionsMenu::PopulateSortAndOrganizeSection(UToolMenu& InMenu)
 {
-	const TSharedPtr<FSequencer> Sequencer = CurrentContext->GetSequencer();
+	const TSharedPtr<FSequencer> Sequencer = WeakSequencer.Pin();
 	if (!Sequencer.IsValid())
 	{
 		return;
@@ -93,7 +99,7 @@ void FSequencerViewOptionsMenu::PopulateSortAndOrganizeSection(UToolMenu& InMenu
 	Section.AddMenuEntryWithCommandList(SequencerCommands.SortAllNodesAndDescendants, SequencerBindings);
 }
 
-void FSequencerViewOptionsMenu::PopulateFilterOptionssSection(UToolMenu& InMenu)
+void FSequencerViewOptionsMenu::PopulateFilterOptionsSection(UToolMenu& InMenu)
 {
 	FToolMenuSection& OptionsSection = InMenu.FindOrAddSection(TEXT("FilterOptions"), LOCTEXT("FilterOptionsHeading", "Filter Options"));
 
@@ -157,7 +163,7 @@ void FSequencerViewOptionsMenu::PopulateLayoutSection(UToolMenu& InMenu)
 
 bool FSequencerViewOptionsMenu::IsFilterLayout(const EFilterBarLayout InLayout) const
 {
-	const TSharedPtr<FSequencer> Sequencer = CurrentContext->GetSequencer();
+	const TSharedPtr<FSequencer> Sequencer = WeakSequencer.Pin();
 	if (!Sequencer.IsValid())
 	{
 		return false;
@@ -169,7 +175,7 @@ bool FSequencerViewOptionsMenu::IsFilterLayout(const EFilterBarLayout InLayout)
 
 void FSequencerViewOptionsMenu::SetFilterLayout(const EFilterBarLayout InLayout)
 {
-	const TSharedPtr<FSequencer> Sequencer = CurrentContext->GetSequencer();
+	const TSharedPtr<FSequencer> Sequencer = WeakSequencer.Pin();
 	if (!Sequencer.IsValid())
 	{
 		return;
@@ -181,7 +187,7 @@ void FSequencerViewOptionsMenu::SetFilterLayout(const EFilterBarLayout InLayout)
 
 bool FSequencerViewOptionsMenu::IsIncludePinnedInFilter() const
 {
-	const TSharedPtr<FSequencer> Sequencer = CurrentContext->GetSequencer();
+	const TSharedPtr<FSequencer> Sequencer = WeakSequencer.Pin();
 	if (!Sequencer.IsValid())
 	{
 		return false;
@@ -198,7 +204,7 @@ bool FSequencerViewOptionsMenu::IsIncludePinnedInFilter() const
 
 void FSequencerViewOptionsMenu::ToggleIncludePinnedInFilter()
 {
-	const TSharedPtr<ISequencer> Sequencer = CurrentContext->GetSequencer();
+	const TSharedPtr<FSequencer> Sequencer = WeakSequencer.Pin();
 	if (!Sequencer.IsValid())
 	{
 		return;
@@ -217,7 +223,7 @@ void FSequencerViewOptionsMenu::ToggleIncludePinnedInFilter()
 
 bool FSequencerViewOptionsMenu::IsAutoExpandPassedFilterNodes() const
 {
-	const TSharedPtr<FSequencer> Sequencer = CurrentContext->GetSequencer();
+	const TSharedPtr<FSequencer> Sequencer = WeakSequencer.Pin();
 	if (!Sequencer.IsValid())
 	{
 		return false;
@@ -234,7 +240,7 @@ bool FSequencerViewOptionsMenu::IsAutoExpandPassedFilterNodes() const
 
 void FSequencerViewOptionsMenu::ToggleAutoExpandPassedFilterNodes()
 {
-	const TSharedPtr<ISequencer> Sequencer = CurrentContext->GetSequencer();
+	const TSharedPtr<FSequencer> Sequencer = WeakSequencer.Pin();
 	if (!Sequencer.IsValid())
 	{
 		return;
@@ -253,7 +259,7 @@ void FSequencerViewOptionsMenu::ToggleAutoExpandPassedFilterNodes()
 
 TSharedPtr<SSequencer> FSequencerViewOptionsMenu::GetSequencerWidget() const
 {
-	const TSharedPtr<FSequencer> Sequencer = CurrentContext->GetSequencer();
+	const TSharedPtr<FSequencer> Sequencer = WeakSequencer.Pin();
 	if (Sequencer.IsValid())
 	{
 		return StaticCastSharedRef<SSequencer>(Sequencer->GetSequencerWidget());
diff --git a/Engine/Source/Editor/Sequencer/Private/Filters/Menus/SequencerViewOptionsMenu.h b/Engine/Source/Editor/Sequencer/Private/Filters/Menus/SequencerViewOptionsMenu.h
--- a/Engine/Source/Editor/Sequencer/Private/Filters/Menus/SequencerViewOptionsMenu.h
+++ b/Engine/Source/Editor/Sequencer/Private/Filters/Menus/SequencerViewOptionsMenu.h
@@ -3,12 +3,10 @@
 #pragma once
 
 #include "Templates/SharedPointer.h"
-#include "UObject/WeakObjectPtr.h"
 
 class FSequencer;
 class SSequencer;
 class SWidget;
-class USequencerMenuContext;
 class UToolMenu;
 enum class EFilterBarLayout : uint8;
 
@@ -21,7 +19,7 @@ protected:
 	void PopulateMenu(UToolMenu* const InMenu);
 	void PopulateFiltersSection(UToolMenu& InMenu);
 	void PopulateSortAndOrganizeSection(UToolMenu& InMenu);
-	void PopulateFilterOptionssSection(UToolMenu& InMenu);
+	void PopulateFilterOptionsSection(UToolMenu& InMenu);
 	void PopulateLayoutSection(UToolMenu& InMenu);
 
 	bool IsFilterLayout(const EFilterBarLayout InLayout) const;
@@ -35,5 +33,5 @@ protected:
 
 	TSharedPtr<SSequencer> GetSequencerWidget() const;
 
-	TWeakObjectPtr<USequencerMenuContext> CurrentContext;
+	TWeakPtr<FSequencer> WeakSequencer;
 };
diff --git a/Engine/Source/Editor/Sequencer/Private/Filters/SequencerFilterBar.cpp b/Engine/Source/Editor/Sequencer/Private/Filters/SequencerFilterBar.cpp
--- a/Engine/Source/Editor/Sequencer/Private/Filters/SequencerFilterBar.cpp
+++ b/Engine/Source/Editor/Sequencer/Private/Filters/SequencerFilterBar.cpp
@@ -70,6 +70,13 @@ FSequencerFilterBar::FSequencerFilterBar(FSequencer& InSequencer)
 FSequencerFilterBar::~FSequencerFilterBar()
 {
 	FSequencerTrackFilterCommands::Unregister();
+	
+	CommonFilters->OnChanged().RemoveAll(this);
+    InternalFilters->OnChanged().RemoveAll(this);
+    TextFilter->OnChanged().RemoveAll(this);
+    LevelFilter->OnChanged().RemoveAll(this);
+    HideIsolateFilter->OnChanged().RemoveAll(this);
+    SelectedFilter->OnChanged().RemoveAll(this);
 }
 
 TSharedPtr<ICustomTextFilter<FSequencerTrackFilterType>> FSequencerFilterBar::CreateTextFilter()
diff --git a/Engine/Source/Editor/Sequencer/Private/Filters/Widgets/SSequencerCustomTextFilterDialog.cpp b/Engine/Source/Editor/Sequencer/Private/Filters/Widgets/SSequencerCustomTextFilterDialog.cpp
--- a/Engine/Source/Editor/Sequencer/Private/Filters/Widgets/SSequencerCustomTextFilterDialog.cpp
+++ b/Engine/Source/Editor/Sequencer/Private/Filters/Widgets/SSequencerCustomTextFilterDialog.cpp
@@ -16,7 +16,7 @@ TSharedPtr<SSequencerCustomTextFilterDialog> SSequencerCustomTextFilterDialog::D
 
 void SSequencerCustomTextFilterDialog::Construct(const FArguments& InArgs, const TSharedRef<FSequencerFilterBar>& InFilterBar)
 {
-	FilterBar = InFilterBar;
+	WeakFilterBar = InFilterBar;
 
 	CustomTextFilter = InArgs._CustomTextFilter;
 	CustomTextFilterData = InArgs._CustomTextFilterData;
@@ -321,6 +321,12 @@ TSharedRef<SWidget> SSequencerCustomTextFilterDialog::ConstructButtonRow()
 
 bool SSequencerCustomTextFilterDialog::CheckFilterNameValidity() const
 {
+	const TSharedPtr<FSequencerFilterBar> FilterBar = WeakFilterBar.Pin();
+	if (!FilterBar.IsValid())
+	{
+		return false;
+	}
+
 	USequencerSettings* const SequencerSettings = FilterBar->GetSequencer().GetSequencerSettings();
 	if (!IsValid(SequencerSettings))
 	{
@@ -376,6 +382,12 @@ FReply SSequencerCustomTextFilterDialog::OnColorBlockMouseButtonDown(const FGeom
 
 void SSequencerCustomTextFilterDialog::OnCreateCustomTextFilter(const bool bInApplyFilter)
 {
+	const TSharedPtr<FSequencerFilterBar> FilterBar = WeakFilterBar.Pin();
+	if (!FilterBar.IsValid())
+	{
+		return;
+	}
+
 	USequencerSettings* const SequencerSettings = FilterBar->GetSequencer().GetSequencerSettings();
 	if (!IsValid(SequencerSettings))
 	{
@@ -412,6 +424,12 @@ void SSequencerCustomTextFilterDialog::OnCreateCustomTextFilter(const bool bInAp
 
 void SSequencerCustomTextFilterDialog::OnModifyCustomTextFilter()
 {
+	const TSharedPtr<FSequencerFilterBar> FilterBar = WeakFilterBar.Pin();
+	if (!FilterBar.IsValid())
+	{
+		return;
+	}
+
 	USequencerSettings* const SequencerSettings = FilterBar->GetSequencer().GetSequencerSettings();
 	if (!IsValid(SequencerSettings))
 	{
@@ -478,6 +496,12 @@ FReply SSequencerCustomTextFilterDialog::OnDeleteButtonClick()
 		return FReply::Handled();
 	}
 
+	const TSharedPtr<FSequencerFilterBar> FilterBar = WeakFilterBar.Pin();
+	if (!FilterBar.IsValid())
+	{
+		return FReply::Handled();
+	}
+
 	USequencerSettings* const SequencerSettings = FilterBar->GetSequencer().GetSequencerSettings();
 	if (!IsValid(SequencerSettings))
 	{
diff --git a/Engine/Source/Editor/Sequencer/Private/Filters/Widgets/SSequencerCustomTextFilterDialog.h b/Engine/Source/Editor/Sequencer/Private/Filters/Widgets/SSequencerCustomTextFilterDialog.h
--- a/Engine/Source/Editor/Sequencer/Private/Filters/Widgets/SSequencerCustomTextFilterDialog.h
+++ b/Engine/Source/Editor/Sequencer/Private/Filters/Widgets/SSequencerCustomTextFilterDialog.h
@@ -59,7 +59,7 @@ private:
 
 	void HandleWindowClosed(const TSharedRef<SWindow>& InWindow);
 
-	TSharedPtr<FSequencerFilterBar> FilterBar;
+	TWeakPtr<FSequencerFilterBar> WeakFilterBar;
 
 	TSharedPtr<FSequencerTrackFilter_CustomText> CustomTextFilter;
 
diff --git a/Engine/Source/Editor/Sequencer/Private/Filters/Widgets/SSequencerFilter.cpp b/Engine/Source/Editor/Sequencer/Private/Filters/Widgets/SSequencerFilter.cpp
--- a/Engine/Source/Editor/Sequencer/Private/Filters/Widgets/SSequencerFilter.cpp
+++ b/Engine/Source/Editor/Sequencer/Private/Filters/Widgets/SSequencerFilter.cpp
@@ -11,8 +11,8 @@ void SSequencerFilter::Construct(const FArguments& InArgs
 	, const TSharedRef<FSequencerFilterBar>& InFilterBar
 	, const TSharedRef<FSequencerTrackFilter>& InFilter)
 {
-	FilterBar = InFilterBar;
-	Filter = InFilter;
+	WeakFilterBar = InFilterBar;
+	WeakFilter = InFilter;
 
 	ContextMenu = MakeShared<FSequencerTrackFilterContextMenu>();
 
@@ -40,7 +40,7 @@ void SSequencerFilter::Construct(const FArguments& InArgs
 	[
 		SAssignNew(ToggleButtonPtr, SSequencerFilterCheckBox)
 		.Style(FAppStyle::Get(), BrushName)
-		.ToolTipText(TAttribute<FText>::Create(TAttribute<FText>::FGetter::CreateSP(Filter.ToSharedRef(), &FSequencerTrackFilter::GetToolTipText)))
+		.ToolTipText(TAttribute<FText>::Create(TAttribute<FText>::FGetter::CreateSP(InFilter, &FSequencerTrackFilter::GetToolTipText)))
 		.IsChecked(this, &SSequencerFilter::IsChecked)
 		.OnCheckStateChanged(this, &SSequencerFilter::OnFilterToggled)
 		.CheckBoxContentUsesAutoWidth(false)
@@ -59,13 +59,22 @@ TSharedRef<SWidget> SSequencerFilter::ConstructBasicFilterWidget()
 	return SNew(STextBlock)
 		.Margin(0.f)
 		.TextStyle(FAppStyle::Get(), TEXT("SmallText"))
-		.Text(Filter.ToSharedRef(), &FSequencerTrackFilter::GetDisplayName);
+		.Text(this, &SSequencerFilter::GetFilterDisplayName);
 }
 
 TSharedRef<SWidget> SSequencerFilter::ConstructDefaultFilterWidget()
 {
-	const TSharedRef<FSequencerFilterBar> FilterBarRef = FilterBar.ToSharedRef();
-	const TSharedRef<FSequencerTrackFilter> FilterRef = Filter.ToSharedRef();
+	const TSharedPtr<FSequencerFilterBar> FilterBar = WeakFilterBar.Pin();
+	if (!FilterBar.IsValid())
+	{
+		return SNullWidget::NullWidget;
+	}
+
+	const TSharedPtr<FSequencerTrackFilter> Filter = WeakFilter.Pin();
+	if (!Filter.IsValid())
+	{
+		return SNullWidget::NullWidget;
+	}
 
 	return SNew(SBorder)
 		.Padding(1.f)
@@ -87,30 +96,60 @@ TSharedRef<SWidget> SSequencerFilter::ConstructDefaultFilterWidget()
 			[
 				SNew(STextBlock)
 				.Font(FCoreStyle::GetDefaultFontStyle(TEXT("Regular"), 8))
-				.Text(FilterRef, &FSequencerTrackFilter::GetDisplayName)
-				.IsEnabled(FilterBarRef, &FSequencerFilterBar::IsFilterActive, FilterRef)
+				.Text(this, &SSequencerFilter::GetFilterDisplayName)
+				.IsEnabled(this, &SSequencerFilter::IsButtonEnabled)
 			]
 		];
 }
 
 const TSharedPtr<FSequencerTrackFilter> SSequencerFilter::GetFilter() const
 {
-	return Filter;
+	return WeakFilter.IsValid() ? WeakFilter.Pin() : nullptr;
 }
 
 bool SSequencerFilter::IsActive() const
 {
+	const TSharedPtr<FSequencerFilterBar> FilterBar = WeakFilterBar.Pin();
+	if (!FilterBar.IsValid())
+	{
+		return false;
+	}
+
+	const TSharedPtr<FSequencerTrackFilter> Filter = WeakFilter.Pin();
+	if (!Filter.IsValid())
+	{
+		return false;
+	}
+
 	return FilterBar->IsFilterActive(Filter.ToSharedRef());
 }
 
 void SSequencerFilter::OnFilterToggled(const ECheckBoxState NewState)
 {
+	const TSharedPtr<FSequencerFilterBar> FilterBar = WeakFilterBar.Pin();
+	if (!FilterBar.IsValid())
+	{
+		return;
+	}
+
+	const TSharedPtr<FSequencerTrackFilter> Filter = WeakFilter.Pin();
+	if (!Filter.IsValid())
+	{
+		return;
+	}
+
 	const bool bNewActive = NewState == ECheckBoxState::Checked;
 	FilterBar->SetFilterActive(Filter.ToSharedRef(), bNewActive, true);
 }
 
 FReply SSequencerFilter::OnFilterMouseUp(const FGeometry& InMyGeometry, const FPointerEvent& InMouseEvent)
 {
+	const TSharedPtr<FSequencerFilterBar> FilterBar = WeakFilterBar.Pin();
+	if (!FilterBar.IsValid())
+	{
+		return FReply::Handled();
+	}
+
 	if (InMouseEvent.IsControlDown())
 	{
 		FilterBar->ActivateCommonFilters(true, {}, {});
@@ -121,7 +160,10 @@ FReply SSequencerFilter::OnFilterMouseUp(const FGeometry& InMyGeometry, const FP
 	}
 	else if(InMouseEvent.GetEffectingButton() == EKeys::MiddleMouseButton)
 	{
-		FilterBar->SetFilterEnabled(Filter.ToSharedRef(), false, true);
+		if (const TSharedPtr<FSequencerTrackFilter> Filter = WeakFilter.Pin())
+		{
+			FilterBar->SetFilterEnabled(Filter.ToSharedRef(), false, true);
+		}
 	}
 
 	return FReply::Handled();
@@ -129,10 +171,19 @@ FReply SSequencerFilter::OnFilterMouseUp(const FGeometry& InMyGeometry, const FP
 
 FReply SSequencerFilter::OnFilterDoubleClicked(const FGeometry& InMyGeometry, const FPointerEvent& InMouseEvent)
 {
+	const TSharedPtr<FSequencerFilterBar> FilterBar = WeakFilterBar.Pin();
+	if (!FilterBar.IsValid())
+	{
+		return FReply::Handled();
+	}
+
 	FilterBar->ActivateCommonFilters(false, {}, {});
 	FilterBar->ActivateCustomTextFilters(false, {});
 
-	FilterBar->SetFilterActive(Filter.ToSharedRef(), true, true);
+	if (const TSharedPtr<FSequencerTrackFilter> Filter = WeakFilter.Pin())
+	{
+		FilterBar->SetFilterActive(Filter.ToSharedRef(), true, true);
+	}
 
 	return FReply::Handled();
 }
@@ -149,7 +200,12 @@ ECheckBoxState SSequencerFilter::IsChecked() const
 
 FSlateColor SSequencerFilter::GetFilterImageColorAndOpacity() const
 {
-	return IsActive() ? Filter->GetColor() : FAppStyle::Get().GetSlateColor(TEXT("Colors.Recessed"));
+	const TSharedPtr<FSequencerTrackFilter> Filter = WeakFilter.Pin();
+	if (!Filter.IsValid() || !IsActive())
+	{
+		return FAppStyle::Get().GetSlateColor(TEXT("Colors.Recessed"));;
+	}
+	return Filter->GetColor();
 }
 
 EVisibility SSequencerFilter::GetFilterOverlayVisibility() const
@@ -162,4 +218,26 @@ FMargin SSequencerFilter::GetFilterNamePadding() const
 	return ToggleButtonPtr->IsPressed() ? FMargin(3, 1, 3, 0) : FMargin(3, 0, 3, 0);
 }
 
+FText SSequencerFilter::GetFilterDisplayName() const
+{
+	return WeakFilter.IsValid() ? WeakFilter.Pin()->GetDisplayName() : FText::GetEmpty();
+}
+
+bool SSequencerFilter::IsButtonEnabled() const
+{
+	const TSharedPtr<FSequencerFilterBar> FilterBar = WeakFilterBar.Pin();
+	if (!FilterBar.IsValid())
+	{
+		return false;
+	}
+
+	const TSharedPtr<FSequencerTrackFilter> Filter = WeakFilter.Pin();
+	if (!Filter.IsValid())
+	{
+		return false;
+	}
+
+	return FilterBar->IsFilterActive(Filter.ToSharedRef());
+}
+
 #undef LOCTEXT_NAMESPACE
diff --git a/Engine/Source/Editor/Sequencer/Private/Filters/Widgets/SSequencerFilter.h b/Engine/Source/Editor/Sequencer/Private/Filters/Widgets/SSequencerFilter.h
--- a/Engine/Source/Editor/Sequencer/Private/Filters/Widgets/SSequencerFilter.h
+++ b/Engine/Source/Editor/Sequencer/Private/Filters/Widgets/SSequencerFilter.h
@@ -46,9 +46,13 @@ protected:
 
 	FMargin GetFilterNamePadding() const;
 
-	TSharedPtr<FSequencerFilterBar> FilterBar;
+	FText GetFilterDisplayName() const;
 
-	TSharedPtr<FSequencerTrackFilter> Filter;
+	bool IsButtonEnabled() const;
+
+	TWeakPtr<FSequencerFilterBar> WeakFilterBar;
+
+	TWeakPtr<FSequencerTrackFilter> WeakFilter;
 
 	TSharedPtr<SSequencerFilterCheckBox> ToggleButtonPtr;
 
diff --git a/Engine/Source/Editor/Sequencer/Private/Filters/Widgets/SSequencerFilterBar.cpp b/Engine/Source/Editor/Sequencer/Private/Filters/Widgets/SSequencerFilterBar.cpp
--- a/Engine/Source/Editor/Sequencer/Private/Filters/Widgets/SSequencerFilterBar.cpp
+++ b/Engine/Source/Editor/Sequencer/Private/Filters/Widgets/SSequencerFilterBar.cpp
@@ -20,7 +20,7 @@ using namespace UE::Sequencer;
 
 void SSequencerFilterBar::Construct(const FArguments& InArgs, const TSharedRef<FSequencerFilterBar>& InFilterBar)
 {
-	FilterBar = InFilterBar;
+	WeakFilterBar = InFilterBar;
 
 	WeakSearchBox = InArgs._FilterSearchBox;
 	FilterBarLayout = InArgs._FilterBarLayout;
@@ -34,7 +34,6 @@ void SSequencerFilterBar::Construct(const FArguments& InArgs, const TSharedRef<F
 			{
 				return (FilterBarLayout == EFilterBarLayout::Horizontal) ? 0 : 1;
 			})
-
 		+ SWidgetSwitcher::Slot()
 		.Padding(FMargin(0.f, 2.f, 0.f, 0.f))
 		[
@@ -57,11 +56,17 @@ void SSequencerFilterBar::Construct(const FArguments& InArgs, const TSharedRef<F
 
 	CreateFilterWidgetsFromConfig();
 
-	FilterBar->GetOnFiltersChanged().AddSP(this, &SSequencerFilterBar::OnFiltersChanged);
+	InFilterBar->GetOnFiltersChanged().AddSP(this, &SSequencerFilterBar::OnFiltersChanged);
 }
 
 FReply SSequencerFilterBar::OnMouseButtonUp(const FGeometry& MyGeometry, const FPointerEvent& MouseEvent)
 {
+	const TSharedPtr<FSequencerFilterBar> FilterBar = WeakFilterBar.Pin();
+	if (!FilterBar.IsValid())
+	{
+		return FReply::Unhandled();
+	}
+
 	if (MouseEvent.GetEffectingButton() == EKeys::RightMouseButton)
 	{
 		const FWidgetPath* EventPath = MouseEvent.GetEventPath();
@@ -80,11 +85,17 @@ FReply SSequencerFilterBar::OnMouseButtonUp(const FGeometry& MyGeometry, const F
 
 TSharedPtr<FSequencerFilterBar> SSequencerFilterBar::GetFilterBar() const
 {
-	return FilterBar;
+	return WeakFilterBar.Pin();
 }
 
 void SSequencerFilterBar::SetTextFilterString(const FString& InText)
 {
+	const TSharedPtr<FSequencerFilterBar> FilterBar = WeakFilterBar.Pin();
+	if (!FilterBar.IsValid())
+	{
+		return;
+	}
+
 	if (const TSharedPtr<SFilterSearchBox> SearchBox = WeakSearchBox.Pin())
 	{
 		if (!InText.Equals(SearchBox->GetText().ToString()))
@@ -98,7 +109,7 @@ void SSequencerFilterBar::SetTextFilterString(const FString& InText)
 
 FText SSequencerFilterBar::GetFilterErrorText() const
 {
-	return FilterBar->GetFilterErrorText();
+	return WeakFilterBar.IsValid() ? WeakFilterBar.Pin()->GetFilterErrorText() : FText::GetEmpty();
 }
 
 EFilterBarLayout SSequencerFilterBar::GetLayout() const
@@ -180,6 +191,7 @@ void SSequencerFilterBar::AddWidgetToLayout(const TSharedRef<SWidget>& InWidget)
 	else
 	{
 		VerticalContainerWidget->AddSlot()
+			.AutoSize()
 			.Padding(SlotPadding)
 			[
 				InWidget
@@ -211,14 +223,18 @@ TSharedPtr<SSequencerFilter> SSequencerFilterBar::FindFilterWidget(const TShared
 	return nullptr;
 }
 
-TSharedRef<SSequencerFilter> SSequencerFilterBar::CreateAndAddFilterWidget(const TSharedRef<FSequencerTrackFilter>& InFilter)
+void SSequencerFilterBar::CreateAndAddFilterWidget(const TSharedRef<FSequencerTrackFilter>& InFilter)
 {
+	const TSharedPtr<FSequencerFilterBar> FilterBar = WeakFilterBar.Pin();
+	if (!FilterBar.IsValid())
+	{
+		return;
+	}
+
 	const TSharedRef<SSequencerFilter> NewFilter = SNew(SSequencerFilter, FilterBar.ToSharedRef(), InFilter)
 		.FilterPillStyle(FilterPillStyle);
 
 	AddFilterWidget(NewFilter);
-
-	return NewFilter;
 }
 
 void SSequencerFilterBar::AddFilterWidget(const TSharedRef<SSequencerFilter>& InFilterWidget)
@@ -294,9 +310,13 @@ void SSequencerFilterBar::RemoveFilterWidgetAndUpdate(const TSharedRef<SSequence
 
 void SSequencerFilterBar::OnEnableAllGroupFilters(bool bEnableAll)
 {
-	FSequencer& Sequencer = FilterBar->GetSequencer();
+	const TSharedPtr<FSequencerFilterBar> FilterBar = WeakFilterBar.Pin();
+	if (!FilterBar.IsValid())
+	{
+		return;
+	}
 
-	UMovieSceneSequence* const FocusedMovieSequence = Sequencer.GetFocusedMovieSceneSequence();
+	UMovieSceneSequence* const FocusedMovieSequence = FilterBar->GetSequencer().GetFocusedMovieSceneSequence();
 	if (!IsValid(FocusedMovieSequence))
 	{
 		return;
@@ -324,9 +344,13 @@ void SSequencerFilterBar::OnNodeGroupFilterClicked(UMovieSceneNodeGroup* NodeGro
 
 UWorld* SSequencerFilterBar::GetWorld() const
 {
-	FSequencer& Sequencer = FilterBar->GetSequencer();
+	const TSharedPtr<FSequencerFilterBar> FilterBar = WeakFilterBar.Pin();
+	if (!FilterBar.IsValid())
+	{
+		return nullptr;
+	}
 
-	UObject* const PlaybackContext = Sequencer.GetPlaybackContext();
+	UObject* const PlaybackContext = FilterBar->GetSequencer().GetPlaybackContext();
 	if (!IsValid(PlaybackContext))
 	{
 		return nullptr;
@@ -368,29 +392,27 @@ void SSequencerFilterBar::OnFiltersChanged(const ESequencerFilterChange InChange
 
 void SSequencerFilterBar::CreateAddCustomTextFilterWindowFromSearch(const FText& InSearchText)
 {
+	const TSharedPtr<FSequencerFilterBar> FilterBar = WeakFilterBar.Pin();
+	if (!FilterBar.IsValid())
+	{
+		return;
+	}
+
 	FCustomTextFilterData CustomTextFilterData;
 	CustomTextFilterData.FilterLabel = LOCTEXT("NewFilterName", "New Filter Name");
 	CustomTextFilterData.FilterString = InSearchText;
-	SSequencerCustomTextFilterDialog::CreateWindow_AddCustomTextFilter(FilterBar.ToSharedRef(), MoveTemp(CustomTextFilterData));
-}
-
-TSharedRef<FFilterCategory> SSequencerFilterBar::GetClassTypeCategory() const
-{
-	return FilterBar->GetClassTypeCategory();
-}
-
-TSharedRef<FFilterCategory> SSequencerFilterBar::GetComponentTypeCategory() const
-{
-	return FilterBar->GetComponentTypeCategory();
-}
 
-TSharedRef<FFilterCategory> SSequencerFilterBar::GetMiscCategory() const
-{
-	return FilterBar->GetMiscCategory();
+	SSequencerCustomTextFilterDialog::CreateWindow_AddCustomTextFilter(FilterBar.ToSharedRef(), MoveTemp(CustomTextFilterData));
 }
 
 void SSequencerFilterBar::OnOpenTextExpressionHelp()
 {
+	const TSharedPtr<FSequencerFilterBar> FilterBar = WeakFilterBar.Pin();
+	if (!FilterBar.IsValid())
+	{
+		return;
+	}
+
 	if (TextExpressionHelpWindow.IsValid())
 	{
 		TextExpressionHelpWindow->BringToFront();
@@ -425,17 +447,30 @@ void SSequencerFilterBar::OnOpenTextExpressionHelp()
 }
 void SSequencerFilterBar::SaveCurrentFilterSetAsCustomTextFilter()
 {
+	const TSharedPtr<FSequencerFilterBar> FilterBar = WeakFilterBar.Pin();
+	if (!FilterBar.IsValid())
+	{
+		return;
+	}
+
 	FCustomTextFilterData CustomTextFilterData;
 	CustomTextFilterData.FilterString = FText::FromString(FilterBar->GenerateTextFilterStringFromEnabledFilters());
 	if (CustomTextFilterData.FilterLabel.IsEmpty())
 	{
 		CustomTextFilterData.FilterLabel = LOCTEXT("NewFilterName", "New Filter Name");
 	}
+
 	SSequencerCustomTextFilterDialog::CreateWindow_AddCustomTextFilter(FilterBar.ToSharedRef(), MoveTemp(CustomTextFilterData));
 }
 
 void SSequencerFilterBar::CreateFilterWidgetsFromConfig()
 {
+	const TSharedPtr<FSequencerFilterBar> FilterBar = WeakFilterBar.Pin();
+	if (!FilterBar.IsValid())
+	{
+		return;
+	}
+
 	USequencerSettings* const SequencerSettings = FilterBar->GetSequencer().GetSequencerSettings();
 	check(IsValid(SequencerSettings));
 
diff --git a/Engine/Source/Editor/Sequencer/Private/Filters/Widgets/SSequencerFilterBar.h b/Engine/Source/Editor/Sequencer/Private/Filters/Widgets/SSequencerFilterBar.h
--- a/Engine/Source/Editor/Sequencer/Private/Filters/Widgets/SSequencerFilterBar.h
+++ b/Engine/Source/Editor/Sequencer/Private/Filters/Widgets/SSequencerFilterBar.h
@@ -75,10 +75,6 @@ public:
 
 	void CreateAddCustomTextFilterWindowFromSearch(const FText& InSearchText);
 
-	TSharedRef<FFilterCategory> GetClassTypeCategory() const;
-	TSharedRef<FFilterCategory> GetComponentTypeCategory() const;
-	TSharedRef<FFilterCategory> GetMiscCategory() const;
-
 	void OnOpenTextExpressionHelp();
 
 	void SaveCurrentFilterSetAsCustomTextFilter();
@@ -91,7 +87,7 @@ protected:
 
 	TSharedPtr<SSequencerFilter> FindFilterWidget(const TSharedRef<FSequencerTrackFilter>& InFilter) const;
 
-	TSharedRef<SSequencerFilter> CreateAndAddFilterWidget(const TSharedRef<FSequencerTrackFilter>& InFilter);
+	void CreateAndAddFilterWidget(const TSharedRef<FSequencerTrackFilter>& InFilter);
 	void AddFilterWidget(const TSharedRef<SSequencerFilter>& InFilterWidget);
 
 	void RemoveFilterWidget(const TSharedRef<FSequencerTrackFilter>& InFilter, const bool ExecuteOnFilterChanged = true);
@@ -110,7 +106,7 @@ protected:
 
 	void CreateFilterWidgetsFromConfig();
 
-	TSharedPtr<FSequencerFilterBar> FilterBar;
+	TWeakPtr<FSequencerFilterBar> WeakFilterBar;
 
 	TWeakPtr<SFilterSearchBox> WeakSearchBox;
 
diff --git a/Engine/Source/Editor/Sequencer/Public/Filters/SequencerTrackFilterBase.h b/Engine/Source/Editor/Sequencer/Public/Filters/SequencerTrackFilterBase.h
--- a/Engine/Source/Editor/Sequencer/Public/Filters/SequencerTrackFilterBase.h
+++ b/Engine/Source/Editor/Sequencer/Public/Filters/SequencerTrackFilterBase.h
@@ -97,8 +97,6 @@ public:
 
 	//~ End FSequencerTrackFilter
 
-	SEQUENCER_API void BroadcastChangedEvent() const;
-
 	SEQUENCER_API ISequencerTrackFilters& GetFilterInterface() const;
 	SEQUENCER_API ISequencer& GetSequencer() const;
 
