[Sequencer]
 - Add SFilterBarClippingHorizontalBox to avoid excessively stretching the play head handle in the UX when there are many filters enabled.
 - Change "Disable Filter" to "Remove Filter" in menus to be easier to understand.

diff --git a/Engine/Source/Editor/Sequencer/Private/Filters/Menus/SequencerTrackFilterContextMenu.cpp b/Engine/Source/Editor/Sequencer/Private/Filters/Menus/SequencerTrackFilterContextMenu.cpp
--- a/Engine/Source/Editor/Sequencer/Private/Filters/Menus/SequencerTrackFilterContextMenu.cpp
+++ b/Engine/Source/Editor/Sequencer/Private/Filters/Menus/SequencerTrackFilterContextMenu.cpp
@@ -1,4 +1,4 @@
-﻿// Copyright Epic Games, Inc. All Rights Reserved.
+// Copyright Epic Games, Inc. All Rights Reserved.
 
 #include "SequencerTrackFilterContextMenu.h"
 #include "Filters/Filters/SequencerTrackFilter_CustomText.h"
@@ -79,8 +79,8 @@ void FSequencerTrackFilterContextMenu::PopulateFilterOptionsSection(UToolMenu& I
 		FUIAction(FExecuteAction::CreateRaw(this, &FSequencerTrackFilterContextMenu::OnActivateWithFilterException)));
 
 	Section.AddMenuEntry(TEXT("DisableFilter"),
-		FText::Format(LOCTEXT("DisableFilter", "Disable: {0}"), FilterName),
-		LOCTEXT("DisableFilterTooltip", "Disable this filter from the list. It can be added again in the filters menu."),
+		FText::Format(LOCTEXT("DisableFilter", "Remove: {0}"), FilterName),
+		LOCTEXT("DisableFilterTooltip", "Disable this filter and remove it from the list. It can be added again in the filters menu."),
 		FSlateIcon(FAppStyle::GetAppStyleSetName(), TEXT("Icons.Minus")),
 		FUIAction(FExecuteAction::CreateRaw(this, &FSequencerTrackFilterContextMenu::OnDisableFilter)));
 }
diff --git a/Engine/Source/Editor/Sequencer/Private/Filters/SequencerTrackFilterCommands.cpp b/Engine/Source/Editor/Sequencer/Private/Filters/SequencerTrackFilterCommands.cpp
--- a/Engine/Source/Editor/Sequencer/Private/Filters/SequencerTrackFilterCommands.cpp
+++ b/Engine/Source/Editor/Sequencer/Private/Filters/SequencerTrackFilterCommands.cpp
@@ -21,7 +21,7 @@ void FSequencerTrackFilterCommands::RegisterCommands()
 
 	UI_COMMAND(ToggleMuteFilters, "Mute Filters", "Mute or Unmute all active filters", EUserInterfaceActionType::ToggleButton, FInputChord());
 
-	UI_COMMAND(DisableAllFilters, "Disable All Filters", "Disables all filters and removes them from the filter bar", EUserInterfaceActionType::Button, FInputChord());
+	UI_COMMAND(DisableAllFilters, "Remove All Filters", "Disables all filters and removes them from the filter bar", EUserInterfaceActionType::Button, FInputChord());
 
 	UI_COMMAND(ToggleActivateEnabledFilters, "Toggle Active State", "Activates or deactivates all enabled filters", EUserInterfaceActionType::Button, FInputChord());
 	
diff --git a/Engine/Source/Editor/Sequencer/Private/Filters/Widgets/SFilterBarClippingHorizontalBox.cpp b/Engine/Source/Editor/Sequencer/Private/Filters/Widgets/SFilterBarClippingHorizontalBox.cpp
new file mode 100644
--- /dev/null
+++ b/Engine/Source/Editor/Sequencer/Private/Filters/Widgets/SFilterBarClippingHorizontalBox.cpp
@@ -0,0 +1,180 @@
+﻿// Copyright Epic Games, Inc. All Rights Reserved.
+
+#include "SFilterBarClippingHorizontalBox.h"
+#include "Framework/Application/SlateApplication.h"
+#include "Layout/ArrangedChildren.h"
+#include "Styling/CoreStyle.h"
+#include "Styling/ToolBarStyle.h"
+#include "Templates/SharedPointer.h"
+#include "Widgets/Images/SImage.h"
+#include "Widgets/Input/SComboButton.h"
+#include "Widgets/Layout/SBorder.h"
+
+#define LOCTEXT_NAMESPACE "SFilterBarClippingHorizontalBox"
+
+void SFilterBarClippingHorizontalBox::Construct(const FArguments& InArgs)
+{
+	OnWrapButtonClicked = InArgs._OnWrapButtonClicked;
+	bIsFocusable = InArgs._IsFocusable;
+}
+
+void SFilterBarClippingHorizontalBox::OnArrangeChildren(const FGeometry& InAllottedGeometry, FArrangedChildren& ArrangedChildren) const
+{
+	// If WrapButton hasn't been initialized, that means AddWrapButton() hasn't 
+	// been called and this method isn't going to behave properly
+	if (!ensure(WeakWrapButton.IsValid()))
+	{
+		SHorizontalBox::OnArrangeChildren(InAllottedGeometry, ArrangedChildren);
+		return;
+	}
+
+	const TSharedPtr<SComboButton> WrapButton = WeakWrapButton.Pin();
+	constexpr int32 ParentReductionSize = 20; // Arbitrary size to make things look nice
+
+	LastClippedIndex = ClippedIndex;
+
+	NumClippedChildren = 0;
+	SHorizontalBox::OnArrangeChildren(InAllottedGeometry, ArrangedChildren);
+
+	// Remove children that are clipped by the allotted geometry
+	const int32 NumChildren = ArrangedChildren.Num(); 
+	int32 IndexClippedAt = NumChildren;
+	for (int32 ChildIndex = NumChildren - 1; ChildIndex >= 0; --ChildIndex)
+	{
+		const FArrangedWidget& ChildWidget = ArrangedChildren[ChildIndex];
+
+		auto CeilToInt = [](const float InPosition)
+		{
+			return FMath::CeilToInt(InPosition - KINDA_SMALL_NUMBER);
+		};
+
+		// Ceil (Minus a tad for float precision) to ensure contents are not a sub-pixel larger than the box, which will create an unnecessary wrap button
+		const FVector2D AbsChildPos = ChildWidget.Geometry.LocalToAbsolute(ChildWidget.Geometry.GetLocalSize());
+		const FVector2D AbsParentPos = FVector2D(InAllottedGeometry.AbsolutePosition) + InAllottedGeometry.GetLocalSize() * InAllottedGeometry.Scale;
+		const FVector2D ChildPosRounded = FVector2D(CeilToInt(AbsChildPos.X), CeilToInt(AbsChildPos.Y));
+		const FVector2D ParentPosRounded = FVector2D(CeilToInt(AbsParentPos.X), CeilToInt(AbsParentPos.Y));
+
+		if (ChildPosRounded.X > ParentPosRounded.X + ParentReductionSize)
+		{
+			++NumClippedChildren;
+			ArrangedChildren.Remove(ChildIndex);
+			IndexClippedAt = ChildIndex;
+		}
+	}
+
+	if (IndexClippedAt == NumChildren)
+	{
+		WrapButton->SetVisibility(EVisibility::Collapsed);
+	}
+	else if (ArrangedChildren.Num() > 0)
+	{
+		WrapButton->SetVisibility(EVisibility::Visible);
+
+		// Further remove any children that the wrap button overlaps with
+		for (int32 ChildIndex = IndexClippedAt - 1; ChildIndex >= 0; --ChildIndex)
+		{
+			const FArrangedWidget& ChildWidget = ArrangedChildren[ChildIndex];
+
+			const int32 ChildRightX = FMath::TruncToInt(ChildWidget.Geometry.AbsolutePosition.X + ChildWidget.Geometry.GetLocalSize().X * ChildWidget.Geometry.Scale);
+			const int32 ParentSizeX = FMath::TruncToInt(InAllottedGeometry.GetLocalSize().X * InAllottedGeometry.Scale);
+
+			if (ChildRightX > ParentSizeX + ParentReductionSize)
+			{
+				++NumClippedChildren;
+				ArrangedChildren.Remove(ChildIndex);
+			}
+		}
+	}
+
+	const int32 NewChildNum =  ArrangedChildren.Num();
+	ClippedIndex = NewChildNum < 0 ? 0 : NewChildNum;
+}
+
+TSharedRef<SComboButton> SFilterBarClippingHorizontalBox::CreateWrapButton()
+{
+	const FToolBarStyle& ToolBarStyle = FCoreStyle::Get().GetWidgetStyle<FToolBarStyle>(TEXT("SlimToolBar"));
+
+	// Always allow this to be focusable to prevent the menu from collapsing during interaction
+	const TSharedRef<SComboButton> WrapButton =
+		SNew(SComboButton)
+		.Visibility(EVisibility::Collapsed)
+		.HasDownArrow(false)
+		.ButtonStyle(&ToolBarStyle.ButtonStyle)
+		.ContentPadding(FMargin(-2.f, 0.f))
+		.ToolTipText(LOCTEXT("ExpandFilterBar", "Click to expand the filter bar"))
+		.OnGetMenuContent(OnWrapButtonClicked)
+		.Cursor(EMouseCursor::Default)
+		.OnMenuOpenChanged(this, &SFilterBarClippingHorizontalBox::OnWrapButtonOpenChanged)
+		.IsFocusable(true)
+		.ButtonContent()
+		[
+			SNew(SImage)
+			.DesiredSizeOverride(FVector2D(10.f))
+			.Image(&ToolBarStyle.ExpandBrush)
+		];
+
+	WeakWrapButton = WrapButton;
+
+	return WrapButton;
+}
+
+void SFilterBarClippingHorizontalBox::OnWrapButtonOpenChanged(const bool bInIsOpen)
+{
+	if (bInIsOpen && !WrapButtonOpenTimer.IsValid())
+	{
+		WrapButtonOpenTimer = RegisterActiveTimer(0.0f, FWidgetActiveTimerDelegate::CreateSP(this, &SFilterBarClippingHorizontalBox::UpdateWrapButtonStatus));
+	}
+	else if(!bInIsOpen && WrapButtonOpenTimer.IsValid())
+	{
+		UnRegisterActiveTimer(WrapButtonOpenTimer.ToSharedRef());
+		WrapButtonOpenTimer.Reset();
+	}
+}
+
+EActiveTimerReturnType SFilterBarClippingHorizontalBox::UpdateWrapButtonStatus(const double InCurrentTime, const float InDeltaTime)
+{
+	const TSharedPtr<SComboButton> WrapButton = WeakWrapButton.Pin();
+	if (!WrapButton.IsValid())
+	{
+		return EActiveTimerReturnType::Stop;
+	}
+
+	if (LastClippedIndex != ClippedIndex || !WrapButton->IsOpen())
+	{
+		WrapButton->SetIsOpen(false);
+		WrapButtonOpenTimer.Reset();
+
+		return EActiveTimerReturnType::Stop;
+	}
+
+	return EActiveTimerReturnType::Continue;
+}
+
+TSharedRef<SWidget> SFilterBarClippingHorizontalBox::WrapVerticalListWithHeading(const TSharedRef<SWidget>& InWidget, const FPointerEventHandler InMouseButtonUpEvent)
+{
+	return SNew(SVerticalBox)
+		+ SVerticalBox::Slot()
+		.AutoHeight()
+		[
+			// Vertical Filters List Header
+			SNew(SBorder)
+			.BorderImage(FAppStyle::Get().GetBrush(TEXT("Brushes.Header")))
+			.Padding(FMargin(8.f, 6.f))
+			.OnMouseButtonUp(InMouseButtonUpEvent)
+			.Content()
+			[
+				SNew(STextBlock)
+				.Text(LOCTEXT("FilterListVerticalHeader", "Filters"))
+				.TextStyle(FAppStyle::Get(), TEXT("ButtonText"))
+				.Font(FAppStyle::Get().GetFontStyle(TEXT("NormalFontBold")))
+			]
+		]
+		+ SVerticalBox::Slot()
+		.FillHeight(1.f)
+		.VAlign(VAlign_Fill)
+		[
+			InWidget
+		];
+}
+
+#undef LOCTEXT_NAMESPACE
diff --git a/Engine/Source/Editor/Sequencer/Private/Filters/Widgets/SFilterBarClippingHorizontalBox.h b/Engine/Source/Editor/Sequencer/Private/Filters/Widgets/SFilterBarClippingHorizontalBox.h
new file mode 100644
--- /dev/null
+++ b/Engine/Source/Editor/Sequencer/Private/Filters/Widgets/SFilterBarClippingHorizontalBox.h
@@ -0,0 +1,58 @@
+﻿// Copyright Epic Games, Inc. All Rights Reserved.
+
+#pragma once
+
+#include "Widgets/DeclarativeSyntaxSupport.h"
+#include "Framework/SlateDelegates.h"
+#include "Widgets/SBoxPanel.h"
+
+class SComboButton;
+
+/**
+ * Specialized widget for handling horizontal clipping of filter widgets that go beyond the filter bar widget width.
+ */
+class SFilterBarClippingHorizontalBox : public SHorizontalBox
+{
+public:
+	static TSharedRef<SWidget> WrapVerticalListWithHeading(const TSharedRef<SWidget>& InWidget, const FPointerEventHandler InMouseButtonUpEvent);
+
+	SLATE_BEGIN_ARGS(SFilterBarClippingHorizontalBox) 
+		: _IsFocusable(true)
+	{}
+		SLATE_ARGUMENT(FOnGetContent, OnWrapButtonClicked)
+		SLATE_ARGUMENT(bool, IsFocusable)
+	SLATE_END_ARGS()
+
+	void Construct(const FArguments& InArgs);
+
+	TSharedRef<SComboButton> CreateWrapButton();
+
+	/** Returns to index of the first clipped child */
+	int32 GetClippedIndex() const { return ClippedIndex; }
+
+protected:
+	//~ Begin SPanel
+	virtual void OnArrangeChildren(const FGeometry& InAllottedGeometry, FArrangedChildren& ArrangedChildren) const override;
+	//~ End SPanel
+
+	void OnWrapButtonOpenChanged(const bool bInIsOpen);
+	EActiveTimerReturnType UpdateWrapButtonStatus(const double InCurrentTime, const float InDeltaTime);
+
+	/** The button that is displayed when a toolbar or menubar is clipped */
+	TWeakPtr<SComboButton> WeakWrapButton;
+
+	/** Callback for when the wrap button is clicked */
+	FOnGetContent OnWrapButtonClicked;
+
+	/** Index of the first clipped child/block */
+	mutable int32 ClippedIndex = INDEX_NONE;
+	/** Index of the last clipped child/block */
+	mutable int32 LastClippedIndex = INDEX_NONE;
+
+	/** Number of clipped children not including the wrap button */
+	mutable int32 NumClippedChildren = 0;
+
+	TSharedPtr<FActiveTimerHandle> WrapButtonOpenTimer;
+
+	bool bIsFocusable = true;
+};
diff --git a/Engine/Source/Editor/Sequencer/Private/Filters/Widgets/SSequencerFilterBar.cpp b/Engine/Source/Editor/Sequencer/Private/Filters/Widgets/SSequencerFilterBar.cpp
--- a/Engine/Source/Editor/Sequencer/Private/Filters/Widgets/SSequencerFilterBar.cpp
+++ b/Engine/Source/Editor/Sequencer/Private/Filters/Widgets/SSequencerFilterBar.cpp
@@ -1,10 +1,11 @@
-﻿// Copyright Epic Games, Inc. All Rights Reserved.
+// Copyright Epic Games, Inc. All Rights Reserved.
 
 #include "Filters/Widgets/SSequencerFilterBar.h"
 #include "Filters/Filters/SequencerTrackFilter_CustomText.h"
 #include "Filters/Filters/SequencerTrackFilter_Text.h"
 #include "Filters/Menus/SequencerFilterBarContextMenu.h"
 #include "Filters/SequencerFilterBar.h"
+#include "Filters/Widgets/SFilterBarClippingHorizontalBox.h"
 #include "Filters/Widgets/SFilterExpressionHelpDialog.h"
 #include "Interfaces/IMainFrameModule.h"
 #include "MovieScene.h"
@@ -27,6 +28,10 @@ void SSequencerFilterBar::Construct(const FArguments& InArgs, const TSharedRef<F
 	bCanChangeOrientation = InArgs._CanChangeOrientation;
 	FilterPillStyle = InArgs._FilterPillStyle;
 
+	HorizontalContainerWidget = SNew(SFilterBarClippingHorizontalBox)
+		.OnWrapButtonClicked(FOnGetContent::CreateSP(this, &SSequencerFilterBar::OnWrapButtonClicked))
+		.IsFocusable(false);
+
 	ChildSlot
 	[
 		SAssignNew(FilterBoxWidget, SWidgetSwitcher)
@@ -37,8 +42,16 @@ void SSequencerFilterBar::Construct(const FArguments& InArgs, const TSharedRef<F
 		+ SWidgetSwitcher::Slot()
 		.Padding(FMargin(0.f, 2.f, 0.f, 0.f))
 		[
-			SAssignNew(HorizontalContainerWidget, SWrapBox)
-			.UseAllottedSize(true)
+			SNew(SHorizontalBox)
+			+ SHorizontalBox::Slot()
+			[
+				HorizontalContainerWidget.ToSharedRef()
+			]
+			+ SHorizontalBox::Slot()
+			.AutoWidth()
+			[
+				HorizontalContainerWidget->CreateWrapButton()
+			]
 		]
 		+ SWidgetSwitcher::Slot()
 		[
@@ -126,36 +139,21 @@ void SSequencerFilterBar::SetLayout(const EFilterBarLayout InFilterBarLayout)
 
 	FilterBarLayout = InFilterBarLayout;
 
-	if (FilterBarLayout == EFilterBarLayout::Horizontal)
-	{
-		VerticalContainerWidget->ClearChildren();
-
-		FilterBoxWidget->SetActiveWidget(HorizontalContainerWidget.ToSharedRef());
-	}
-	else
-	{
-		HorizontalContainerWidget->ClearChildren();
-
-		FilterBoxWidget->SetActiveWidget(VerticalContainerWidget.ToSharedRef());
-	}
+	HorizontalContainerWidget->ClearChildren();
+	VerticalContainerWidget->ClearChildren();
 
 	for (const TSharedRef<SSequencerFilter>& FilterWidget : FilterWidgets)
 	{
-		RemoveWidgetFromLayout(FilterWidget);
-	}
-
-	for (const TSharedRef<SSequencerFilter>& Filter : FilterWidgets)
-	{
-		AddWidgetToLayout(Filter);
+		AddWidgetToLayout(FilterWidget);
 	}
-
-	Invalidate(EInvalidateWidgetReason::Layout);
 }
 
 void SSequencerFilterBar::AttachFilterSearchBox(const TSharedPtr<SFilterSearchBox>& InFilterSearchBox)
 {
 	if (InFilterSearchBox)
 	{
+		WeakSearchBox = InFilterSearchBox;
+
 		InFilterSearchBox->SetOnSaveSearchHandler(
 			 SFilterSearchBox::FOnSaveSearchClicked::CreateSP(this, &SSequencerFilterBar::CreateAddCustomTextFilterWindowFromSearch));
 	}
@@ -183,6 +181,7 @@ void SSequencerFilterBar::AddWidgetToLayout(const TSharedRef<SWidget>& InWidget)
 		}
 
 		HorizontalContainerWidget->AddSlot()
+			.AutoWidth()
 			.Padding(SlotPadding)
 			[
 				InWidget
@@ -518,4 +517,37 @@ void SSequencerFilterBar::CreateFilterWidgetsFromConfig()
 	}
 }
 
+TSharedRef<SWidget> SSequencerFilterBar::OnWrapButtonClicked()
+{
+	const TSharedRef<SVerticalBox> VerticalContainer = SNew(SVerticalBox);
+
+	const int32 NumSlots = HorizontalContainerWidget->NumSlots();
+	int32 SlotIndex = HorizontalContainerWidget->GetClippedIndex();
+	for (; SlotIndex < NumSlots; ++SlotIndex)
+	{
+		SHorizontalBox::FSlot& Slot = HorizontalContainerWidget->GetSlot(SlotIndex);
+
+		VerticalContainer->AddSlot()
+			.AutoHeight()
+			.Padding(1.f)
+			[
+				Slot.GetWidget()
+			];
+	}
+
+	const TSharedRef<SBorder> ContainerBorder = SNew(SBorder)
+		.BorderImage(FAppStyle::GetBrush(TEXT("Brushes.Panel")))
+		.Padding(2.f)
+		[
+			VerticalContainer
+		];
+
+	return SNew(SBox)
+		.Padding(8.f)
+		[
+			HorizontalContainerWidget->WrapVerticalListWithHeading(ContainerBorder
+				, FPointerEventHandler::CreateSP(this, &SSequencerFilterBar::OnMouseButtonUp))
+		];
+}
+
 #undef LOCTEXT_NAMESPACE
diff --git a/Engine/Source/Editor/Sequencer/Private/Filters/Widgets/SSequencerFilterBar.h b/Engine/Source/Editor/Sequencer/Private/Filters/Widgets/SSequencerFilterBar.h
--- a/Engine/Source/Editor/Sequencer/Private/Filters/Widgets/SSequencerFilterBar.h
+++ b/Engine/Source/Editor/Sequencer/Private/Filters/Widgets/SSequencerFilterBar.h
@@ -1,4 +1,4 @@
-﻿// Copyright Epic Games, Inc. All Rights Reserved.
+// Copyright Epic Games, Inc. All Rights Reserved.
 
 #pragma once
 
@@ -9,6 +9,7 @@ class FSequencer;
 class FSequencerFilterBarContextMenu;
 class FSequencerFilterBar;
 class FSequencerTrackFilter_CustomText;
+class SFilterBarClippingHorizontalBox;
 class SSequencerFilter;
 class UMovieSceneNodeGroup;
 enum class ESequencerFilterChange : uint8;
@@ -106,12 +107,14 @@ protected:
 
 	void CreateFilterWidgetsFromConfig();
 
+	TSharedRef<SWidget> OnWrapButtonClicked();
+
 	TWeakPtr<FSequencerFilterBar> WeakFilterBar;
 
 	TWeakPtr<SFilterSearchBox> WeakSearchBox;
 
 	TSharedPtr<SWidgetSwitcher> FilterBoxWidget;
-	TSharedPtr<SWrapBox> HorizontalContainerWidget;
+	TSharedPtr<SFilterBarClippingHorizontalBox> HorizontalContainerWidget;
 	TSharedPtr<SScrollBox> VerticalContainerWidget;
 
 	EFilterBarLayout FilterBarLayout = EFilterBarLayout::Vertical;
diff --git a/Engine/Source/Editor/Sequencer/Private/SSequencer.cpp b/Engine/Source/Editor/Sequencer/Private/SSequencer.cpp
--- a/Engine/Source/Editor/Sequencer/Private/SSequencer.cpp
+++ b/Engine/Source/Editor/Sequencer/Private/SSequencer.cpp
@@ -7,6 +7,7 @@
 #include "Filters/Menus/SequencerViewOptionsMenu.h"
 #include "Filters/SequencerFilterBar.h"
 #include "Filters/SFilterSearchBox.h"
+#include "Filters/Widgets/SFilterBarClippingHorizontalBox.h"
 #include "Filters/Widgets/SFilterBarIsolateHideShow.h"
 #include "Filters/Widgets/SSequencerCustomTextFilterDialog.h"
 #include "Filters/Widgets/SSequencerFilterBar.h"
@@ -556,29 +557,8 @@ TSharedRef<SWidget> SSequencer::ConstructFilterBarContent()
 			})
 		.Expose(FilterSplitterSlot)
 		[
-			SNew(SVerticalBox)
-			+ SVerticalBox::Slot()
-			.AutoHeight()
-			[
-				// Vertical Filters List Header
-				SNew(SBorder)
-				.BorderImage(FAppStyle::Get().GetBrush(TEXT("Brushes.Header")))
-				.Padding(FMargin(8.f, 6.f))
-				.OnMouseButtonUp(FilterBarWidgetRef, &SSequencerFilterBar::OnMouseButtonUp)
-				.Content()
-				[
-					SNew(STextBlock)
-					.Text(LOCTEXT("FilterListVerticalHeader", "Filters"))
-					.TextStyle(FAppStyle::Get(), TEXT("ButtonText"))
-					.Font(FAppStyle::Get().GetFontStyle(TEXT("NormalFontBold")))
-				]
-			]
-			+ SVerticalBox::Slot()
-			.FillHeight(1.f)
-			.VAlign(VAlign_Fill)
-			[
-				FilterBarWidgetRef
-			]
+			SFilterBarClippingHorizontalBox::WrapVerticalListWithHeading(FilterBarWidgetRef
+				, FPointerEventHandler::CreateSP(FilterBarWidgetRef, &SSequencerFilterBar::OnMouseButtonUp))
 		]
 		+ SSplitter::Slot()
 		.Value(0.94f)
@@ -4377,6 +4357,7 @@ void SSequencer::RebuildSearchAndFilterRow()
 			{
 				SearchAndFilterRow->AddSlot()
 					.AutoHeight()
+					.Padding(0)
 					[
 						FilterBarWidget.ToSharedRef()
 					];
