Sequencer: Hit test marked frames only if visible

diff --git a/Engine/Source/Editor/Sequencer/Private/SequencerTimeSliderController.cpp b/Engine/Source/Editor/Sequencer/Private/SequencerTimeSliderController.cpp
--- a/Engine/Source/Editor/Sequencer/Private/SequencerTimeSliderController.cpp
+++ b/Engine/Source/Editor/Sequencer/Private/SequencerTimeSliderController.cpp
@@ -1966,6 +1966,11 @@ bool FSequencerTimeSliderController::HitTestMark(const FGeometry& AllottedGeomet
 		*OutMarkFrameNumber = 0;
 	}
 
+	if (!WeakSequencer.Pin()->GetSequencerSettings()->GetShowMarkedFrames())
+	{
+		return false;
+	}
+
 	const TArray<FMovieSceneMarkedFrame> & MarkedFrames = TimeSliderArgs.MarkedFrames.Get();
 	if (MarkedFrames.Num() < 1)
 	{
@@ -2136,6 +2141,11 @@ void FSequencerTimeSliderController::HandleMarkSelection(int32 InMarkIndex)
 	using namespace UE::Sequencer;
 
 	TSharedPtr<FSequencer> Sequencer = WeakSequencer.Pin();
+	if (!Sequencer->GetSequencerSettings()->GetShowMarkedFrames())
+	{
+		return;
+	}
+
 	FSequencerSelection& SequencerSelection = Sequencer->GetSelection();
 
 	FModifierKeysState ModifierKeys = FSlateApplication::Get().GetModifierKeys();
