Sequencer: Fix partial submit.

diff --git a/Engine/Source/Editor/Sequencer/Private/SequencerCommonHelpers.cpp b/Engine/Source/Editor/Sequencer/Private/SequencerCommonHelpers.cpp
--- a/Engine/Source/Editor/Sequencer/Private/SequencerCommonHelpers.cpp
+++ b/Engine/Source/Editor/Sequencer/Private/SequencerCommonHelpers.cpp
@@ -387,8 +387,8 @@ void SequencerHelpers::BuildEditSectionMenu(const TWeakPtr<FSequencer>& InWeakSe
 				return MakeShared<FMovieSceneSectionDetailsCustomization>(NumericTypeInterfafce.Pin(), CurrentScene.Get());
 			}));
 
-	DetailsView->RegisterInstancedCustomPropertyTypeLayout("MovieSceneConditionContainer", FOnGetPropertyTypeCustomizationInstance::CreateLambda([=]() {
-		return FMovieSceneConditionCustomization::MakeInstance(Sequence); }));
+	DetailsView->RegisterInstancedCustomPropertyTypeLayout("MovieSceneConditionContainer", FOnGetPropertyTypeCustomizationInstance::CreateLambda([Sequence, InWeakSequencer]() {
+		return FMovieSceneConditionCustomization::MakeInstance(Sequence, InWeakSequencer); }));
 
 	DetailsView->RegisterInstancedCustomPropertyTypeLayout("MovieSceneDirectorBlueprintConditionData", FOnGetPropertyTypeCustomizationInstance::CreateLambda([=]() {
 		return FMovieSceneDirectorBlueprintConditionCustomization::MakeInstance(CurrentScene.Get());}));
@@ -497,8 +497,8 @@ void SequencerHelpers::BuildEditTrackMenu(const TWeakPtr<FSequencer>& InWeakSequ
 
 		TSharedRef<IDetailsView> DetailsView = FModuleManager::GetModuleChecked<FPropertyEditorModule>("PropertyEditor").CreateDetailView(DetailsViewArgs);
 
-		DetailsView->RegisterInstancedCustomPropertyTypeLayout("MovieSceneConditionContainer", FOnGetPropertyTypeCustomizationInstance::CreateLambda([=]() {
-			return FMovieSceneConditionCustomization::MakeInstance(Sequence); }));
+		DetailsView->RegisterInstancedCustomPropertyTypeLayout("MovieSceneConditionContainer", FOnGetPropertyTypeCustomizationInstance::CreateLambda([Sequence, InWeakSequencer]() {
+			return FMovieSceneConditionCustomization::MakeInstance(Sequence, InWeakSequencer); }));
 
 		DetailsView->RegisterInstancedCustomPropertyTypeLayout("MovieSceneDirectorBlueprintConditionData", FOnGetPropertyTypeCustomizationInstance::CreateLambda([=]() {
 			return FMovieSceneDirectorBlueprintConditionCustomization::MakeInstance(CurrentScene); }));
