[Sequencer] Bind to selection change event through level module instead of through potential invalid Sequencer reference.

diff --git a/Engine/Source/Editor/Sequencer/Private/Filters/Filters/SequencerTrackFilter_Selected.cpp b/Engine/Source/Editor/Sequencer/Private/Filters/Filters/SequencerTrackFilter_Selected.cpp
--- a/Engine/Source/Editor/Sequencer/Private/Filters/Filters/SequencerTrackFilter_Selected.cpp
+++ b/Engine/Source/Editor/Sequencer/Private/Filters/Filters/SequencerTrackFilter_Selected.cpp
@@ -13,20 +13,42 @@
 FSequencerTrackFilter_Selected::FSequencerTrackFilter_Selected(ISequencerTrackFilters& InFilterInterface, TSharedPtr<FFilterCategory> InCategory)
 	: FSequencerTrackFilter(InFilterInterface, MoveTemp(InCategory))
 {
-	if (const TSharedPtr<IToolkitHost> ToolkitHost = GetSequencer().GetToolkitHost())
+	if (const FEditorModeTools* const EditorMode = GetEditorModeManager())
 	{
-		ToolkitHost->GetEditorModeManager().GetSelectedObjects()->SelectionChangedEvent.AddRaw(this, &FSequencerTrackFilter_Selected::OnSelectionChanged);
+		EditorMode->GetSelectedObjects()->SelectionChangedEvent.AddRaw(this, &FSequencerTrackFilter_Selected::OnSelectionChanged);
 	}
 }
 
 FSequencerTrackFilter_Selected::~FSequencerTrackFilter_Selected()
 {
-	if (const TSharedPtr<IToolkitHost> ToolkitHost = GetSequencer().GetToolkitHost())
+	if (const FEditorModeTools* const EditorMode = GetEditorModeManager())
 	{
-		ToolkitHost->GetEditorModeManager().GetSelectedObjects()->SelectionChangedEvent.RemoveAll(this);
+		EditorMode->GetSelectedObjects()->SelectionChangedEvent.RemoveAll(this);
 	}
 }
 
+TSharedPtr<ILevelEditor> FSequencerTrackFilter_Selected::GetLevelEditor() const
+{
+	const FLevelEditorModule* const LevelEditorModule = FModuleManager::GetModulePtr<FLevelEditorModule>(TEXT("LevelEditor"));
+	if (!LevelEditorModule)
+	{
+		return nullptr;
+	}
+
+	return LevelEditorModule->GetLevelEditorInstance().Pin();
+}
+
+FEditorModeTools* FSequencerTrackFilter_Selected::GetEditorModeManager() const
+{
+	const TSharedPtr<ILevelEditor> LevelEditor = GetLevelEditor();
+	if (!LevelEditor.IsValid())
+	{
+		return nullptr;
+	}
+
+	return &LevelEditor->GetEditorModeManager();
+}
+
 bool FSequencerTrackFilter_Selected::ShouldUpdateOnTrackValueChanged() const
 {
 	return true;
@@ -59,8 +81,7 @@ FString FSequencerTrackFilter_Selected::GetName() const
 
 bool FSequencerTrackFilter_Selected::PassesFilter(FSequencerTrackFilterType InItem) const
 {
-	const FLevelEditorModule& LevelEditorModule = FModuleManager::GetModuleChecked<FLevelEditorModule>(TEXT("LevelEditor"));
-	const TSharedPtr<ILevelEditor> LevelEditor = LevelEditorModule.GetFirstLevelEditor();
+	const TSharedPtr<ILevelEditor> LevelEditor = GetLevelEditor();
 	if (!LevelEditor.IsValid())
 	{
 		return false;
@@ -87,13 +108,13 @@ bool FSequencerTrackFilter_Selected::PassesFilter(FSequencerTrackFilterType InIt
 		return true;
 	}
 
-	USceneComponent* const Component = TrackObject->GetTypedOuter<USceneComponent>();
+	const USceneComponent* const Component = TrackObject->GetTypedOuter<USceneComponent>();
 	if (IsValid(Component) && SelectedObjects.Contains(Component))
 	{
 		return true;
 	}
 
-	AActor* const Actor = TrackObject->GetTypedOuter<AActor>();
+	const AActor* const Actor = TrackObject->GetTypedOuter<AActor>();
 	if (IsValid(Actor) && SelectedObjects.Contains(Actor))
 	{
 		return true;
diff --git a/Engine/Source/Editor/Sequencer/Private/Filters/Filters/SequencerTrackFilter_Selected.h b/Engine/Source/Editor/Sequencer/Private/Filters/Filters/SequencerTrackFilter_Selected.h
--- a/Engine/Source/Editor/Sequencer/Private/Filters/Filters/SequencerTrackFilter_Selected.h
+++ b/Engine/Source/Editor/Sequencer/Private/Filters/Filters/SequencerTrackFilter_Selected.h
@@ -4,6 +4,9 @@
 
 #include "Filters/SequencerTrackFilterBase.h"
 
+class FEditorModeTools;
+class ILevelEditor;
+
 class FSequencerTrackFilter_Selected : public FSequencerTrackFilter
 {
 public:
@@ -29,5 +32,8 @@ public:
 	void ToggleShowOnlySelectedTracks();
 
 protected:
+	TSharedPtr<ILevelEditor> GetLevelEditor() const;
+	FEditorModeTools* GetEditorModeManager() const;
+
 	void OnSelectionChanged(UObject* const InObject);
 };
