[Sequencer]
 - Change ownership of filter bar widget to Sequencer widget instead of filter bar backend instance
 - Add PopulateMenu delegate
 - Minor code cleanups

diff --git a/Engine/Source/Editor/Sequencer/Private/Filters/Menus/SequencerFilterBarContext.h b/Engine/Source/Editor/Sequencer/Private/Filters/Menus/SequencerFilterBarContext.h
--- a/Engine/Source/Editor/Sequencer/Private/Filters/Menus/SequencerFilterBarContext.h
+++ b/Engine/Source/Editor/Sequencer/Private/Filters/Menus/SequencerFilterBarContext.h
@@ -6,6 +6,9 @@
 #include "SequencerFilterBarContext.generated.h"
 
 class FSequencerFilterBar;
+class UToolMenu;
+
+DECLARE_DELEGATE_OneParam(FOnPopulateFilterBarMenu, UToolMenu*);
 
 UCLASS()
 class USequencerFilterBarContext : public UObject
@@ -23,6 +26,8 @@ public:
 		return WeakFilterBar.IsValid() ? WeakFilterBar.Pin() : nullptr;
 	}
 
+	FOnPopulateFilterBarMenu OnPopulateFilterBarMenu;
+
 protected:
 	TWeakPtr<FSequencerFilterBar> WeakFilterBar;
 };
diff --git a/Engine/Source/Editor/Sequencer/Private/Filters/Menus/SequencerTrackFilterMenu.cpp b/Engine/Source/Editor/Sequencer/Private/Filters/Menus/SequencerTrackFilterMenu.cpp
--- a/Engine/Source/Editor/Sequencer/Private/Filters/Menus/SequencerTrackFilterMenu.cpp
+++ b/Engine/Source/Editor/Sequencer/Private/Filters/Menus/SequencerTrackFilterMenu.cpp
@@ -24,11 +24,18 @@ TSharedRef<SWidget> FSequencerTrackFilterMenu::CreateMenu(const TSharedRef<FSequ
 	{
 		UToolMenu* const Menu = UToolMenus::Get()->RegisterMenu(FilterMenuName);
 		Menu->bShouldCloseWindowAfterMenuSelection = false;
-		Menu->AddDynamicSection(TEXT("FilterMenu"), FNewToolMenuDelegate::CreateSP(this, &FSequencerTrackFilterMenu::PopulateMenu));
+		Menu->AddDynamicSection(NAME_None, FNewToolMenuDelegate::CreateLambda([this](UToolMenu* const InMenu)
+			{
+				if (USequencerFilterBarContext* const Context = InMenu->FindContext<USequencerFilterBarContext>())
+				{
+					Context->OnPopulateFilterBarMenu.ExecuteIfBound(InMenu);
+				}
+			}));
 	}
 
 	USequencerFilterBarContext* const ContextObject = NewObject<USequencerFilterBarContext>();
 	ContextObject->Init(InFilterBar);
+	ContextObject->OnPopulateFilterBarMenu = FOnPopulateFilterBarMenu::CreateSP(this, &FSequencerTrackFilterMenu::PopulateMenu);
 
 	const FToolMenuContext MenuContext(InFilterBar->GetCommandList(), nullptr, ContextObject);
 	return UToolMenus::Get()->GenerateWidget(FilterMenuName, MenuContext);
@@ -91,7 +98,7 @@ void FSequencerTrackFilterMenu::PopulateFilterOptionsSection(UToolMenu& InMenu)
 		LOCTEXT("FilterListResetToolTip", "Resets current filter selection"),
 		FSlateIcon(FAppStyle::Get().GetStyleSetName(), TEXT("PropertyWindow.DiffersFromDefault")),
 		FUIAction(
-			FExecuteAction::CreateSP(FilterBar.ToSharedRef(), &FSequencerFilterBar::ResetFilters),
+			FExecuteAction::CreateSP(this, &FSequencerTrackFilterMenu::ResetFilters),
 			FCanExecuteAction::CreateSP(this, &FSequencerTrackFilterMenu::CanResetFilters)
 		));
 
@@ -111,12 +118,6 @@ void FSequencerTrackFilterMenu::PopulateCommonFilterSections(UToolMenu& InMenu)
 		return;
 	}
 
-	const TSharedPtr<SSequencerFilterBar> FilterBarWidget = FilterBar->GetWidget();
-	if (!FilterBarWidget.IsValid())
-	{
-		return;
-	}
-
 	const USequencerSettings* const SequencerSettings = FilterBar->GetSequencer().GetSequencerSettings();
 	if (!IsValid(SequencerSettings))
 	{
@@ -211,8 +212,6 @@ void FSequencerTrackFilterMenu::FillLevelFilterMenu(UToolMenu* const InMenu)
 		return;
 	}
 
-	const TSharedRef<FSequencerFilterBar> FilterBarRef = FilterBar.ToSharedRef();
-
 	FToolMenuSection& LevelsSection = InMenu->FindOrAddSection(TEXT("Levels"), LOCTEXT("LevelFilters", "Level Filters"));
 
 	const UWorld* const World = FilterBar->GetWorld();
@@ -231,7 +230,7 @@ void FSequencerTrackFilterMenu::FillLevelFilterMenu(UToolMenu* const InMenu)
 			FUIAction(
 				FExecuteAction::CreateSP(this, &FSequencerTrackFilterMenu::OnTrackLevelFilterClicked, LevelName),
 				FCanExecuteAction(),
-				FIsActionChecked::CreateSP(FilterBarRef, &FSequencerFilterBar::IsLevelFilterActive, LevelName)
+				FIsActionChecked::CreateSP(this, &FSequencerTrackFilterMenu::IsLevelFilterActive, LevelName)
 			),
 			EUserInterfaceActionType::ToggleButton);
 	}
@@ -252,8 +251,6 @@ void FSequencerTrackFilterMenu::FillGroupFilterMenu(UToolMenu* const InMenu)
 
 	FSequencer& Sequencer = FilterBar->GetSequencer();
 
-	const TSharedRef<SSequencer> SequencerWidget = StaticCastSharedRef<SSequencer>(Sequencer.GetSequencerWidget());
-
 	FToolMenuSection& GroupFilterOptionsSection = InMenu->FindOrAddSection(TEXT("GroupFilterOptions")
 		, FText(), FToolMenuInsert(NAME_None, EToolMenuInsertType::First));
 
@@ -261,7 +258,7 @@ void FSequencerTrackFilterMenu::FillGroupFilterMenu(UToolMenu* const InMenu)
 		LOCTEXT("OpenGroupFilters", "Open Group Filters..."),
 		LOCTEXT("OpenGroupFiltersToolTip", "Opens the group filter dialog for managing groups"),
 		FSlateIcon(FAppStyle::GetAppStyleSetName(), TEXT("Icons.GroupActors")),
-		FUIAction(FExecuteAction::CreateSP(SequencerWidget, &SSequencer::OpenNodeGroupsManager)));
+		FUIAction(FExecuteAction::CreateSP(this, &FSequencerTrackFilterMenu::OpenNodeGroupsManager)));
 
 	const UMovieSceneSequence* const FocusedMovieSequence = Sequencer.GetFocusedMovieSceneSequence();
 	if (!IsValid(FocusedMovieSequence))
@@ -306,14 +303,6 @@ void FSequencerTrackFilterMenu::FillCustomTextFiltersMenu(UToolMenu* const InMen
 		return;
 	}
 
-	const TSharedPtr<SSequencerFilterBar> FilterBarWidget = FilterBar->GetWidget();
-	if (!FilterBarWidget.IsValid())
-	{
-		return;
-	}
-
-	const TSharedRef<SSequencerFilterBar> FilterBarWidgetRef = FilterBarWidget.ToSharedRef();
-
 	FToolMenuSection& CustomTextOptionsSection = InMenu->FindOrAddSection(TEXT("CustomTextFilterOptions")
 		, LOCTEXT("CustomTextFilterOptions", "Custom Text Filter Options")
 		, FToolMenuInsert(NAME_None, EToolMenuInsertType::First));
@@ -322,28 +311,19 @@ void FSequencerTrackFilterMenu::FillCustomTextFiltersMenu(UToolMenu* const InMen
 		LOCTEXT("TextExpressionHelp", "Text Expression Help"),
 		LOCTEXT("TextExpressionHelpToolTip", "Opens the help dialog for the advanced search syntax text expressions"),
 		FSlateIcon(FAppStyle::GetAppStyleSetName(), TEXT("Icons.Help")),
-		FUIAction(FExecuteAction::CreateSP(FilterBarWidgetRef, &SSequencerFilterBar::OnOpenTextExpressionHelp)));
+		FUIAction(FExecuteAction::CreateSP(this, &FSequencerTrackFilterMenu::OnOpenTextExpressionHelp)));
 
 	CustomTextOptionsSection.AddMenuEntry(TEXT("SaveCurrentAsNewTextFilter"),
 		LOCTEXT("SaveCurrentAsNewTextFilter", "Save Current as New Filter"),
 		LOCTEXT("SaveCurrentAsNewTextFilterToolTip", "Saves the enabled and active set of common filters as a custom text filter"),
 		FSlateIcon(FAppStyle::GetAppStyleSetName(), TEXT("LevelEditor.SaveAs")),
-		FUIAction(FExecuteAction::CreateSP(FilterBarWidgetRef, &SSequencerFilterBar::SaveCurrentFilterSetAsCustomTextFilter)));
+		FUIAction(FExecuteAction::CreateSP(this, &FSequencerTrackFilterMenu::SaveCurrentFilterSetAsCustomTextFilter)));
 
 	CustomTextOptionsSection.AddMenuEntry(TEXT("CreateNewTextFilter"),
 		LOCTEXT("CreateNewTextFilter", "Create New Filter"),
 		LOCTEXT("CreateNewTextFilterTooltip", "Create a new text filter"),
 		FSlateIcon(FAppStyle::Get().GetStyleSetName(), TEXT("Icons.PlusCircle")),
-		FUIAction(FExecuteAction::CreateLambda([this]()
-			{
-				const TSharedPtr<FSequencerFilterBar> FilterBar = WeakFilterBar.Pin();
-				if (!FilterBar.IsValid())
-				{
-					return;
-				}
-
-				SSequencerCustomTextFilterDialog::CreateWindow_AddCustomTextFilter(FilterBar.ToSharedRef(), FCustomTextFilterData(), nullptr);
-			})));
+		FUIAction(FExecuteAction::CreateSP(this, &FSequencerTrackFilterMenu::CreateNewTextFilter)));
 
 	USequencerSettings* const SequencerSettings = FilterBar->GetSequencer().GetSequencerSettings();
 	check(SequencerSettings);
@@ -376,7 +356,6 @@ void FSequencerTrackFilterMenu::FillFiltersMenuCategory(FToolMenuSection& InOutS
 		return;
 	}
 
-	const TSharedRef<FSequencerFilterBar> FilterBarRef = FilterBar.ToSharedRef();
 	UMovieSceneSequence* const FocusedSequence = FilterBar->GetSequencer().GetFocusedMovieSceneSequence();
 
 	for (const TSharedRef<FSequencerTrackFilter>& Filter : FilterBar->GetCommonFilters())
@@ -403,7 +382,10 @@ void FSequencerTrackFilterMenu::FillFiltersMenuCategory(FToolMenuSection& InOutS
 					{
 						return WeakFilterBar.IsValid() ? !WeakFilterBar.Pin()->AreFiltersMuted() : false;
 					}),
-				FIsActionChecked::CreateSP(FilterBarRef, &FSequencerFilterBar::IsFilterEnabled, Filter)
+				FIsActionChecked::CreateLambda([this, Filter]()
+					{
+						return WeakFilterBar.IsValid() ? WeakFilterBar.Pin()->IsFilterEnabled(Filter) : false;
+					})
 			),
 			EUserInterfaceActionType::ToggleButton);
 	}
@@ -726,4 +708,96 @@ bool FSequencerTrackFilterMenu::CanResetFilters() const
 	return FilterBar->HasEnabledFilter(ClassAndCompFilters) || FilterBar->HasEnabledCustomTextFilters();
 }
 
+void FSequencerTrackFilterMenu::ResetFilters()
+{
+	const TSharedPtr<FSequencerFilterBar> FilterBar = WeakFilterBar.Pin();
+	if (!FilterBar.IsValid())
+	{
+		return;
+	}
+
+	FilterBar->ResetFilters();
+}
+
+bool FSequencerTrackFilterMenu::IsLevelFilterActive(const FString InLevelName) const
+{
+	const TSharedPtr<FSequencerFilterBar> FilterBar = WeakFilterBar.Pin();
+	if (!FilterBar.IsValid())
+	{
+		return false;
+	}
+
+	return FilterBar->IsLevelFilterActive(InLevelName);
+}
+
+void FSequencerTrackFilterMenu::OpenNodeGroupsManager()
+{
+	const TSharedPtr<FSequencerFilterBar> FilterBar = WeakFilterBar.Pin();
+	if (!FilterBar.IsValid())
+	{
+		return;
+	}
+
+	FSequencer& Sequencer = FilterBar->GetSequencer();
+
+	const TSharedPtr<SSequencer> SequencerWidget = StaticCastSharedRef<SSequencer>(Sequencer.GetSequencerWidget());
+	if (!SequencerWidget.IsValid())
+	{
+		return;
+	}
+
+	SequencerWidget->OpenNodeGroupsManager();
+}
+
+void FSequencerTrackFilterMenu::OnOpenTextExpressionHelp()
+{
+	const TSharedPtr<FSequencerFilterBar> FilterBar = WeakFilterBar.Pin();
+	if (!FilterBar.IsValid())
+	{
+		return;
+	}
+
+	const TSharedPtr<SSequencer> SequencerWidget = StaticCastSharedRef<SSequencer>(FilterBar->GetSequencer().GetSequencerWidget());
+	if (!SequencerWidget.IsValid())
+	{
+		return;
+	}
+
+	const TSharedPtr<SSequencerFilterBar> FilterBarWidget = FilterBar->GetWidget();
+	if (!FilterBarWidget.IsValid())
+	{
+		return;
+	}
+
+	FilterBarWidget->OnOpenTextExpressionHelp();
+}
+
+void FSequencerTrackFilterMenu::SaveCurrentFilterSetAsCustomTextFilter()
+{
+	const TSharedPtr<FSequencerFilterBar> FilterBar = WeakFilterBar.Pin();
+	if (!FilterBar.IsValid())
+	{
+		return;
+	}
+
+	const TSharedPtr<SSequencerFilterBar> FilterBarWidget = FilterBar->GetWidget();
+	if (!FilterBarWidget.IsValid())
+	{
+		return;
+	}
+
+	FilterBarWidget->SaveCurrentFilterSetAsCustomTextFilter();
+}
+
+void FSequencerTrackFilterMenu::CreateNewTextFilter()
+{
+	const TSharedPtr<FSequencerFilterBar> FilterBar = WeakFilterBar.Pin();
+	if (!FilterBar.IsValid())
+	{
+		return;
+	}
+
+	SSequencerCustomTextFilterDialog::CreateWindow_AddCustomTextFilter(FilterBar.ToSharedRef(), FCustomTextFilterData(), nullptr);
+}
+
 #undef LOCTEXT_NAMESPACE
diff --git a/Engine/Source/Editor/Sequencer/Private/Filters/Menus/SequencerTrackFilterMenu.h b/Engine/Source/Editor/Sequencer/Private/Filters/Menus/SequencerTrackFilterMenu.h
--- a/Engine/Source/Editor/Sequencer/Private/Filters/Menus/SequencerTrackFilterMenu.h
+++ b/Engine/Source/Editor/Sequencer/Private/Filters/Menus/SequencerTrackFilterMenu.h
@@ -59,6 +59,16 @@ protected:
 	ECheckBoxState GetAllGroupsCheckState() const;
 
 	bool CanResetFilters() const;
+	void ResetFilters();
+
+	bool IsLevelFilterActive(const FString InLevelName) const;
+
+	void OpenNodeGroupsManager();
+
+	void OnOpenTextExpressionHelp();
+
+	void SaveCurrentFilterSetAsCustomTextFilter();
+	void CreateNewTextFilter();
 
 private:
 	TWeakPtr<FSequencerFilterBar> WeakFilterBar;
diff --git a/Engine/Source/Editor/Sequencer/Private/Filters/SequencerFilterBar.cpp b/Engine/Source/Editor/Sequencer/Private/Filters/SequencerFilterBar.cpp
--- a/Engine/Source/Editor/Sequencer/Private/Filters/SequencerFilterBar.cpp
+++ b/Engine/Source/Editor/Sequencer/Private/Filters/SequencerFilterBar.cpp
@@ -86,8 +86,6 @@ TSharedPtr<ICustomTextFilter<FSequencerTrackFilterType>> FSequencerFilterBar::Cr
 
 void FSequencerFilterBar::CreateDefaultFilters()
 {
-	SetTextFilterString(FString());
-
 	// Add internal filters that won't be saved to config
 	InternalFilters->RemoveAll();
 
@@ -253,26 +251,12 @@ FName FSequencerFilterBar::GetIdentifier() const
 	return TEXT("SequencerMain");
 }
 
-TSharedPtr<SSequencerFilterBar> FSequencerFilterBar::GetWidget() const
-{
-	return FilterBarWidget;
-}
-
 TSharedRef<SSequencerFilterBar> FSequencerFilterBar::GenerateWidget(const TSharedPtr<SFilterSearchBox>& InSearchBox, const EFilterBarLayout InLayout)
 {
-	if (!FilterBarWidget.IsValid())
-	{
-		FilterBarWidget = SNew(SSequencerFilterBar, SharedThis(this))
-			.FilterBarLayout(InLayout)
-			.AddMetaData<FTagMetaData>(FTagMetaData(TEXT("SequencerTrackFilters")))
-			.FilterSearchBox(InSearchBox);
-	}
-	else
-	{
-		FilterBarWidget->SetLayout(InLayout);
-	}
-
-	return FilterBarWidget.ToSharedRef();
+	return SNew(SSequencerFilterBar, SharedThis(this))
+		.FilterBarLayout(InLayout)
+		.AddMetaData<FTagMetaData>(FTagMetaData(TEXT("SequencerTrackFilters")))
+		.FilterSearchBox(InSearchBox);
 }
 
 bool FSequencerFilterBar::AreFiltersMuted() const
@@ -284,7 +268,7 @@ void FSequencerFilterBar::MuteFilters(const bool bInMute)
 {
 	bFiltersMuted = bInMute;
 
-	if (FilterBarWidget.IsValid())
+	if (const TSharedPtr<SSequencerFilterBar> FilterBarWidget = GetWidget())
 	{
 		FilterBarWidget->SetEnabled(!bFiltersMuted);
 
@@ -325,7 +309,7 @@ void FSequencerFilterBar::SetTextFilterString(const FString& InText)
 {
 	TextFilter->SetRawFilterText(FText::FromString(InText));
 
-	if (FilterBarWidget.IsValid())
+	if (const TSharedPtr<SSequencerFilterBar> FilterBarWidget = GetWidget())
 	{
 		FilterBarWidget->SetTextFilterString(InText);
 	}
@@ -1657,4 +1641,14 @@ TSharedRef<SComboButton> FSequencerFilterBar::MakeAddFilterButton()
 	return ComboButton;
 }
 
+TSharedPtr<SSequencerFilterBar> FSequencerFilterBar::GetWidget() const
+{
+	const TSharedPtr<SSequencer> SequencerWidget = StaticCastSharedRef<SSequencer>(GetSequencer().GetSequencerWidget());
+	if (!SequencerWidget.IsValid())
+	{
+		return nullptr;
+	}
+	return SequencerWidget->GetFilterBarWidget();
+}
+
 #undef LOCTEXT_NAMESPACE
diff --git a/Engine/Source/Editor/Sequencer/Private/Filters/SequencerFilterBar.h b/Engine/Source/Editor/Sequencer/Private/Filters/SequencerFilterBar.h
--- a/Engine/Source/Editor/Sequencer/Private/Filters/SequencerFilterBar.h
+++ b/Engine/Source/Editor/Sequencer/Private/Filters/SequencerFilterBar.h
@@ -59,8 +59,6 @@ public:
 
 	virtual TSharedPtr<FUICommandList> GetCommandList() const override;
 
-	TSharedPtr<SSequencerFilterBar> GetWidget() const;
-
 	TSharedRef<SSequencerFilterBar> GenerateWidget(const TSharedPtr<SFilterSearchBox>& InSearchBox, const EFilterBarLayout InLayout);
 
 	bool AreFiltersMuted() const;
@@ -234,6 +232,9 @@ public:
 	TSharedRef<SFilterBarIsolateHideShow> MakeIsolateHideShowPanel();
 	TSharedRef<SComboButton> MakeAddFilterButton();
 
+	/** Attempts to get the filter bar widget from the Sequencer widget */
+	TSharedPtr<SSequencerFilterBar> GetWidget() const;
+
 protected:
 	void CreateDefaultFilters();
 
@@ -268,8 +269,6 @@ protected:
 	
 	TArray<TSharedRef<FSequencerTrackFilter_CustomText>> CustomTextFilters;
 
-	TSharedPtr<SSequencerFilterBar> FilterBarWidget;
-
 	TSharedPtr<FSequencerTrackFilterMenu> FilterMenu;
 
 	FSequencerFiltersChanged FiltersChangedEvent;
diff --git a/Engine/Source/Editor/Sequencer/Private/Filters/Widgets/SSequencerFilterBar.cpp b/Engine/Source/Editor/Sequencer/Private/Filters/Widgets/SSequencerFilterBar.cpp
--- a/Engine/Source/Editor/Sequencer/Private/Filters/Widgets/SSequencerFilterBar.cpp
+++ b/Engine/Source/Editor/Sequencer/Private/Filters/Widgets/SSequencerFilterBar.cpp
@@ -19,6 +19,11 @@ using namespace UE::Sequencer;
 
 #define LOCTEXT_NAMESPACE "SSequencerFilterBar"
 
+SSequencerFilterBar::~SSequencerFilterBar()
+{
+	ContextMenu.Reset();
+}
+
 void SSequencerFilterBar::Construct(const FArguments& InArgs, const TSharedRef<FSequencerFilterBar>& InFilterBar)
 {
 	WeakFilterBar = InFilterBar;
diff --git a/Engine/Source/Editor/Sequencer/Private/Filters/Widgets/SSequencerFilterBar.h b/Engine/Source/Editor/Sequencer/Private/Filters/Widgets/SSequencerFilterBar.h
--- a/Engine/Source/Editor/Sequencer/Private/Filters/Widgets/SSequencerFilterBar.h
+++ b/Engine/Source/Editor/Sequencer/Private/Filters/Widgets/SSequencerFilterBar.h
@@ -54,6 +54,8 @@ public:
 
 	SLATE_END_ARGS()
 
+	virtual ~SSequencerFilterBar() override;
+
 	/** Constructs this widget with InArgs */
 	void Construct(const FArguments& InArgs, const TSharedRef<FSequencerFilterBar>& InFilterBar);
 
diff --git a/Engine/Source/Editor/Sequencer/Private/SSequencer.h b/Engine/Source/Editor/Sequencer/Private/SSequencer.h
--- a/Engine/Source/Editor/Sequencer/Private/SSequencer.h
+++ b/Engine/Source/Editor/Sequencer/Private/SSequencer.h
@@ -30,8 +30,11 @@ class FAssetDragDropOp;
 class FClassDragDropOp;
 class FMovieSceneClipboard;
 class FSequencerTimeSliderController;
+class FSequencerTrackFilter;
+class FSequencerViewOptionsMenu;
 class ISidebarDrawerContent;
 class SCurveEditorTree;
+class SSequencerFilterBar;
 class SSequencerTransformBox;
 class SSequencerStretchBox;
 class SCurveEditorPanel;
@@ -46,9 +49,6 @@ class USequencerSettings;
 class SSequencerGroupManager;
 class SSequencerTreeFilterStatusBar;
 class SFilterSearchBox;
-class FSequencerFilterMenu;
-class FSequencerTrackFilter;
-class FSequencerViewOptionsMenu;
 enum class ESequencerFilterChange : uint8;
 struct FPaintPlaybackRangeArgs;
 struct FSequencerCustomizationInfo;
@@ -703,12 +703,11 @@ public:
 	SEQUENCER_API void EnablePendingFocusOnHovering(const bool InEnabled);
 	
 	TSharedPtr<FSequencerFilterBar> GetFilterBar() const;
+	TSharedPtr<SSequencerFilterBar> GetFilterBarWidget() const;
 
 	bool IsFilterBarVisible() const;
 	void ToggleFilterBarVisibility();
 
-	TSharedPtr<SSequencerFilterBar> GetFilterBarWidget() const;
-
 	EFilterBarLayout GetFilterBarLayout() const;
 	void SetFilterBarLayout(const EFilterBarLayout InLayout);
 
@@ -911,8 +910,6 @@ private:
 	TSharedPtr<SComboButton> ViewOptionsComboButton;
 	TSharedPtr<FSequencerViewOptionsMenu> ViewOptionsMenu;
 
-	TSharedPtr<FSequencerFilterMenu> FilterMenu;
-	
 	SSplitter::FSlot* FilterSplitterSlot = nullptr;
 
 	TSharedPtr<SVerticalBox> SearchAndFilterRow;
diff --git a/Engine/Source/Editor/Sequencer/Private/Sequencer.cpp b/Engine/Source/Editor/Sequencer/Private/Sequencer.cpp
--- a/Engine/Source/Editor/Sequencer/Private/Sequencer.cpp
+++ b/Engine/Source/Editor/Sequencer/Private/Sequencer.cpp
@@ -553,6 +553,7 @@ void FSequencer::InitSequencer(const FSequencerInitParams& InitParams, const TSh
 	UpdateTimeBases();
 	PlayPosition.Reset(ConvertFrameTime(GetPlaybackRange().GetLowerBoundValue(), GetRootTickResolution(), PlayPosition.GetInputRate()));
 
+	FilterBar = MakeShared<FSequencerFilterBar>(*this);
 	FilterBar->CreateCustomTextFiltersFromConfig();
 
 	// Make internal widgets
@@ -741,7 +742,6 @@ FSequencer::FSequencer()
 	, bUpdatingExternalSelection( false )
 	, bNeedsEvaluate(false)
 	, bNeedsInvalidateCachedData(false)
-	, FilterBar(MakeShared<FSequencerFilterBar>(*this))
 {
 	// Exposes the sequencer and curve editor command lists to subscribers from other systems
 	FInputBindingManager::Get().RegisterCommandList(FSequencerCommands::Get().GetContextName(), SequencerCommandBindings);
@@ -779,6 +779,8 @@ FSequencer::~FSequencer()
 
 	TrackEditors.Empty();
 
+	FilterBar.Reset();
+
 	RootTemplateInstance.TearDown();
 }
 
@@ -821,6 +823,8 @@ void FSequencer::Close()
 	SequencerWidget.Reset();
 	TrackEditors.Empty();
 
+	FilterBar.Reset();
+
 	GUnrealEd->UpdatePivotLocationForSelection();
 
 	// Redraw viewports after restoring pre animated state in case viewports are not set to realtime
@@ -12306,12 +12310,12 @@ FText FSequencer::GetSidebarSelectionDrawerToolTipText() const
 
 TSharedRef<ISequencerTrackFilters> FSequencer::GetFilterInterface() const
 {
-	return FilterBar;
+	return FilterBar.ToSharedRef();
 }
 
 TSharedRef<FSequencerFilterBar> FSequencer::GetFilterBar() const
 {
-	return FilterBar;
+	return FilterBar.ToSharedRef();
 }
 
 #undef LOCTEXT_NAMESPACE
diff --git a/Engine/Source/Editor/Sequencer/Private/Sequencer.h b/Engine/Source/Editor/Sequencer/Private/Sequencer.h
--- a/Engine/Source/Editor/Sequencer/Private/Sequencer.h
+++ b/Engine/Source/Editor/Sequencer/Private/Sequencer.h
@@ -1577,5 +1577,5 @@ private:
 	/** Delegate that is called when selection limiting has been toggled on or off */
 	FOnViewportSelectionLimitedChanged OnSelectionLimitedChangedDelegate;
 
-	TSharedRef<FSequencerFilterBar> FilterBar;
+	TSharedPtr<FSequencerFilterBar> FilterBar;
 };
diff --git a/Engine/Source/Editor/Sequencer/Private/SequencerNodeTree.cpp b/Engine/Source/Editor/Sequencer/Private/SequencerNodeTree.cpp
--- a/Engine/Source/Editor/Sequencer/Private/SequencerNodeTree.cpp
+++ b/Engine/Source/Editor/Sequencer/Private/SequencerNodeTree.cpp
@@ -2,7 +2,6 @@
 
 #include "SequencerNodeTree.h"
 #include "Filters/SequencerFilterBar.h"
-#include "Filters/Widgets/SSequencerFilterBar.h"
 #include "MVVM/ViewModels/ViewModelHierarchy.h"
 #include "MVVM/Selection/Selection.h"
 #include "MovieSceneBinding.h"
diff --git a/Engine/Source/Editor/Sequencer/Private/SequencerNodeTree.h b/Engine/Source/Editor/Sequencer/Private/SequencerNodeTree.h
--- a/Engine/Source/Editor/Sequencer/Private/SequencerNodeTree.h
+++ b/Engine/Source/Editor/Sequencer/Private/SequencerNodeTree.h
@@ -21,7 +21,6 @@ class FSequencerFilterBar;
 class FSequencerTrackFilterCollection;
 class FSequencerTrackFilter_Level;
 class ISequencerTrackEditor;
-class SSequencerFilterBar;
 class UMovieScene;
 class UMovieSceneFolder;
 class UMovieSceneTrack;
