[Sequencer] Update filter bar context menu to use OnPopulateFilterBarMenu

diff --git a/Engine/Source/Editor/Sequencer/Private/Filters/Menus/SequencerFilterBarContextMenu.cpp b/Engine/Source/Editor/Sequencer/Private/Filters/Menus/SequencerFilterBarContextMenu.cpp
--- a/Engine/Source/Editor/Sequencer/Private/Filters/Menus/SequencerFilterBarContextMenu.cpp
+++ b/Engine/Source/Editor/Sequencer/Private/Filters/Menus/SequencerFilterBarContextMenu.cpp
@@ -16,11 +16,18 @@ TSharedRef<SWidget> FSequencerFilterBarContextMenu::CreateMenu(const TSharedRef<
 	if (!UToolMenus::Get()->IsMenuRegistered(FilterMenuName))
 	{
 		UToolMenu* const Menu = UToolMenus::Get()->RegisterMenu(FilterMenuName);
-		Menu->AddDynamicSection(NAME_None, FNewToolMenuDelegate::CreateRaw(this, &FSequencerFilterBarContextMenu::PopulateMenu));
+		Menu->AddDynamicSection(NAME_None, FNewToolMenuDelegate::CreateLambda([this](UToolMenu* const InMenu)
+			{
+				if (USequencerFilterBarContext* const Context = InMenu->FindContext<USequencerFilterBarContext>())
+				{
+					Context->OnPopulateFilterBarMenu.ExecuteIfBound(InMenu);
+				}
+			}));
 	}
 
 	USequencerFilterBarContext* const ContextObject = NewObject<USequencerFilterBarContext>();
 	ContextObject->Init(InFilterBar);
+	ContextObject->OnPopulateFilterBarMenu = FOnPopulateFilterBarMenu::CreateRaw(this, &FSequencerFilterBarContextMenu::PopulateMenu);
 
 	const FToolMenuContext MenuContext(InFilterBar->GetCommandList(), nullptr, ContextObject);
 	return UToolMenus::Get()->GenerateWidget(FilterMenuName, MenuContext);
