[Sequencer] Selectability: Remove FocusedTemplate check

diff --git a/Engine/Source/Editor/Sequencer/Private/Sequencer.cpp b/Engine/Source/Editor/Sequencer/Private/Sequencer.cpp
--- a/Engine/Source/Editor/Sequencer/Private/Sequencer.cpp
+++ b/Engine/Source/Editor/Sequencer/Private/Sequencer.cpp
@@ -12145,12 +12145,6 @@ bool FSequencer::IsObjectSelectableInViewport(UObject* const InObject)
 		return SelectableObject->IsSelectable();
 	}
 
-	const FGuid ObjectGuid = FindObjectId(*InObject, GetFocusedTemplateID());
-	if (ObjectGuid.IsValid())
-	{
-		return true;
-	}
-
 	UMovieSceneSequence* const FocusedSequence = GetFocusedMovieSceneSequence();
 	if (!IsValid(FocusedSequence))
 	{
