Sequencer: Check for weak section model

To fix a crash on collapsing sections.

diff --git a/Engine/Source/Editor/Sequencer/Private/SSequencerSection.cpp b/Engine/Source/Editor/Sequencer/Private/SSequencerSection.cpp
--- a/Engine/Source/Editor/Sequencer/Private/SSequencerSection.cpp
+++ b/Engine/Source/Editor/Sequencer/Private/SSequencerSection.cpp
@@ -442,7 +442,8 @@ struct FSequencerSectionPainterImpl : FSequencerSectionPainter
 		{
 			for (const FEasingAreaHandle& Easing : EasingAreaHotspot->Easings)
 			{
-				if (Easing.WeakSectionModel.Pin()->GetSection() == InSection)
+				TSharedPtr<FSectionModel> EasingSectionModel = Easing.WeakSectionModel.Pin();
+				if (EasingSectionModel && EasingSectionModel->GetSection() == InSection)
 				{
 					if (Easing.EasingType == ESequencerEasingType::In)
 					{
@@ -1096,7 +1097,10 @@ bool SSequencerSection::CheckForEasingHandleInteraction( const FPointerEvent& Mo
 
 	// Gather all underlapping sections
 	TArray<TSharedPtr<FSectionModel>> AllUnderlappingSections;
-	AllUnderlappingSections.Add(WeakSectionModel.Pin());
+	if (TSharedPtr<FSectionModel> SectionModel = WeakSectionModel.Pin())
+	{
+		AllUnderlappingSections.Add(SectionModel);
+	}
 	for (const FOverlappingSections& Segment : UnderlappingSegments)
 	{
 		for (const TWeakPtr<FSectionModel>& Section : Segment.Sections)
@@ -1165,7 +1169,10 @@ bool SSequencerSection::CheckForEdgeInteraction( const FPointerEvent& MouseEvent
 	}
 
 	TArray<TSharedPtr<FSectionModel>> AllUnderlappingSections;
-	AllUnderlappingSections.Add(WeakSectionModel.Pin());
+	if (TSharedPtr<FSectionModel> SectionModel = WeakSectionModel.Pin())
+	{
+		AllUnderlappingSections.Add(SectionModel);
+	}
 	for (const FOverlappingSections& Segment : UnderlappingSegments)
 	{
 		for (const TWeakPtr<FSectionModel>& Section : Segment.Sections)
@@ -1588,7 +1595,10 @@ void SSequencerSection::PaintEasingHandles( FSequencerSectionPainter& InPainter,
 	TArray<TSharedPtr<FSectionModel>> AllUnderlappingSections;
 	if (IsSectionHighlighted(SectionInterface->GetSectionObject(), Hotspot))
 	{
-		AllUnderlappingSections.Add(WeakSectionModel.Pin());
+		if (SectionModel)
+		{
+			AllUnderlappingSections.Add(SectionModel);
+		}
 	}
 
 	for (const FOverlappingSections& Segment : UnderlappingSegments)
@@ -1629,7 +1639,8 @@ void SSequencerSection::PaintEasingHandles( FSequencerSectionPainter& InPainter,
 			{
 				for (const FEasingAreaHandle& Easing : EasingAreaHotspot->Easings)
 				{
-					if (Easing.WeakSectionModel.Pin()->GetSection() == UnderlappingSectionObj)
+					TSharedPtr<FSectionModel> EasingSectionModel = Easing.WeakSectionModel.Pin();
+					if (EasingSectionModel && EasingSectionModel->GetSection() == UnderlappingSectionObj)
 					{
 						if (Easing.EasingType == ESequencerEasingType::In)
 						{
@@ -1784,7 +1795,10 @@ void SSequencerSection::DrawSectionHandles( const FGeometry& AllottedGeometry, F
 	TArray<TSharedPtr<FSectionModel>> AllUnderlappingSections;
 	if (IsSectionHighlighted(SectionInterface->GetSectionObject(), Hotspot))
 	{
-		AllUnderlappingSections.Add(WeakSectionModel.Pin());
+		if (TSharedPtr<FSectionModel> SectionModel = WeakSectionModel.Pin())
+		{
+			AllUnderlappingSections.Add(SectionModel);
+		}
 	}
 
 	for (const FOverlappingSections& Segment : UnderlappingSegments)
