Gather text from source, resolve macro has an empty source text (.cpp files)
Resolve duplicate keys(identifiers)
Resolve macro has an empty source text

diff --git a/Engine/Plugins/Experimental/Mutable/Source/CustomizableObjectEditor/Private/MuCOE/CustomizableObjectCompiler.cpp b/Engine/Plugins/Experimental/Mutable/Source/CustomizableObjectEditor/Private/MuCOE/CustomizableObjectCompiler.cpp
--- a/Engine/Plugins/Experimental/Mutable/Source/CustomizableObjectEditor/Private/MuCOE/CustomizableObjectCompiler.cpp
+++ b/Engine/Plugins/Experimental/Mutable/Source/CustomizableObjectEditor/Private/MuCOE/CustomizableObjectCompiler.cpp
@@ -757,7 +757,7 @@ mu::NodeObjectPtr FCustomizableObjectCompiler::GenerateMutableRoot(
 				{
 					if (ComponentInfo->NodeComponentMesh)
 					{
-						FText Msg = FText::Format(LOCTEXT("NoReferenceMeshObjectTab", "Already exist a Component node with the same name in Customizable Object [{0}]"), FText::FromString(GetRootObject(*ComponentInfo->NodeComponentMesh)->GetName()));
+						FText Msg = FText::Format(LOCTEXT("ComponentNodeWithSameNameExists", "Already exist a Component node with the same name in Customizable Object [{0}]"), FText::FromString(GetRootObject(*ComponentInfo->NodeComponentMesh)->GetName()));
 						CompilerLog(Msg, Node, EMessageSeverity::Error);
 						return nullptr;
 					}
diff --git a/Engine/Plugins/Experimental/Mutable/Source/CustomizableObjectEditor/Private/MuCOE/Nodes/CustomizableObjectNodeComponentPassthroughMesh.cpp b/Engine/Plugins/Experimental/Mutable/Source/CustomizableObjectEditor/Private/MuCOE/Nodes/CustomizableObjectNodeComponentPassthroughMesh.cpp
--- a/Engine/Plugins/Experimental/Mutable/Source/CustomizableObjectEditor/Private/MuCOE/Nodes/CustomizableObjectNodeComponentPassthroughMesh.cpp
+++ b/Engine/Plugins/Experimental/Mutable/Source/CustomizableObjectEditor/Private/MuCOE/Nodes/CustomizableObjectNodeComponentPassthroughMesh.cpp
@@ -139,11 +139,11 @@ FText UCustomizableObjectNodeComponentPassthroughMesh::GetNodeTitle(ENodeTitleTy
 		FFormatNamedArguments Args;
 		Args.Add(TEXT("MeshName"), FText::FromString(SkeletalMesh->GetName()));
 
-		return FText::Format(LOCTEXT("ComponentMesh_Title", "{MeshName}\n Passthrough Mesh Component"), Args);
+		return FText::Format(LOCTEXT("ComponentPassthroughMesh_Title", "{MeshName}\n Passthrough Mesh Component"), Args);
 	}
 	else
 	{
-		return LOCTEXT("Component_Mesh", "Passthrough Mesh Component");
+		return LOCTEXT("ComponentPassthroughMesh", "Passthrough Mesh Component");
 	}
 }
 
@@ -196,7 +196,7 @@ FSkeletalMaterial* UCustomizableObjectNodeComponentPassthroughMesh::GetSkeletalM
 
 FText UCustomizableObjectNodeComponentPassthroughMesh::GetTooltipText() const
 {
-	return LOCTEXT("Component_Mesh_Tooltip", "Define a new object component based on a Skeletal Mesh.");
+	return LOCTEXT("ComponentPassthroughMesh_Tooltip", "Define a new object component based on a Skeletal Mesh.");
 }
 
 
diff --git a/Engine/Plugins/Runtime/Metasound/Source/MetasoundEditor/Private/MetasoundEditor.cpp b/Engine/Plugins/Runtime/Metasound/Source/MetasoundEditor/Private/MetasoundEditor.cpp
--- a/Engine/Plugins/Runtime/Metasound/Source/MetasoundEditor/Private/MetasoundEditor.cpp
+++ b/Engine/Plugins/Runtime/Metasound/Source/MetasoundEditor/Private/MetasoundEditor.cpp
@@ -692,7 +692,7 @@ namespace Metasound
 								const FText DisplayName = GraphMember->GetDisplayName();
 								const FText GetterToolTip = FText::Format(LOCTEXT("DropTargetGetterVariableToolTipFormat", "{0}\nAdd:\n* Get (Drop)\n* Get Delayed (Alt+Drop)\n"), DisplayName);
 								static const FText GetJumpToToolTip = LOCTEXT("JumpToGettersToolTip", "Get (Ctrl+Drop)");
-								static const FText AddOrJumpToSetToolTip = LOCTEXT("AddOrJumpToSetToolTip", "");
+								static const FText AddOrJumpToSetToolTip;
 								FConstNodeHandle MutatorNodeHandle = Variable->GetConstVariableHandle()->FindMutatorNode();
 								if (MutatorNodeHandle->IsValid())
 								{
diff --git a/Engine/Source/Editor/MaterialEditor/Private/MaterialEditorViewportToolbarSections.cpp b/Engine/Source/Editor/MaterialEditor/Private/MaterialEditorViewportToolbarSections.cpp
--- a/Engine/Source/Editor/MaterialEditor/Private/MaterialEditorViewportToolbarSections.cpp
+++ b/Engine/Source/Editor/MaterialEditor/Private/MaterialEditorViewportToolbarSections.cpp
@@ -53,7 +53,7 @@ FToolMenuEntry UE::MaterialEditor::CreateShowSubmenu(TWeakPtr<SMaterialEditor3DP
 					return;
 				}
 
-				FToolMenuSection& UnnamedSection = Submenu->FindOrAddSection("", LOCTEXT("UnnamedLabel", ""));
+				FToolMenuSection& UnnamedSection = Submenu->FindOrAddSection("", FText());
 
 				UnnamedSection.AddEntry(FToolMenuEntry::InitWidget(
 					"ShowMenuItems", UE::MaterialEditor::CreateShowMenuWidget(Viewport.ToSharedRef()), FText(), true
diff --git a/Engine/Source/Editor/Sequencer/Private/Filters/SequencerFilterBar.cpp b/Engine/Source/Editor/Sequencer/Private/Filters/SequencerFilterBar.cpp
--- a/Engine/Source/Editor/Sequencer/Private/Filters/SequencerFilterBar.cpp
+++ b/Engine/Source/Editor/Sequencer/Private/Filters/SequencerFilterBar.cpp
@@ -40,10 +40,10 @@ const FName FSequencerFilterBar::SharedIdentifier = TEXT("SharedSequencerTrackFi
 FSequencerFilterBar::FSequencerFilterBar(FSequencer& InSequencer)
 	: Sequencer(InSequencer)
 	, CommandList(MakeShared<FUICommandList>())
-	, ClassTypeCategory(MakeShared<FFilterCategory>(LOCTEXT("ActorTypeFilterCategory", "Actor Type Filters"), LOCTEXT("ActorTypeFilterCategoryTooltip", "")))
-	, ComponentTypeCategory(MakeShared<FFilterCategory>(LOCTEXT("ObjectTypeFilterCategory", "Object Type Filters"), LOCTEXT("ObjectTypeFilterCategoryTooltip", "")))
-	, MiscCategory(MakeShared<FFilterCategory>(LOCTEXT("MiscFilterCategory", "Misc Filters"), LOCTEXT("MiscFilterCategoryTooltip", "")))
-	, TransientCategory(MakeShared<FFilterCategory>(LOCTEXT("TransientFilterCategory", "Transient Filters"), LOCTEXT("TransientFilterCategoryTooltip", "")))
+	, ClassTypeCategory(MakeShared<FFilterCategory>(LOCTEXT("ActorTypeFilterCategory", "Actor Type Filters"), FText()))
+	, ComponentTypeCategory(MakeShared<FFilterCategory>(LOCTEXT("ObjectTypeFilterCategory", "Object Type Filters"), FText()))
+	, MiscCategory(MakeShared<FFilterCategory>(LOCTEXT("MiscFilterCategory", "Misc Filters"), FText()))
+	, TransientCategory(MakeShared<FFilterCategory>(LOCTEXT("TransientFilterCategory", "Transient Filters"), FText()))
 	, CommonFilters(MakeShared<FSequencerTrackFilterCollection>(*this))
 	, InternalFilters(MakeShared<FSequencerTrackFilterCollection>(*this))
 	, TextFilter(MakeShared<FSequencerTrackFilter_CustomText>(*this))
diff --git a/Engine/Source/Editor/Sequencer/Private/Filters/Widgets/SFilterBarIsolateHideShow.cpp b/Engine/Source/Editor/Sequencer/Private/Filters/Widgets/SFilterBarIsolateHideShow.cpp
--- a/Engine/Source/Editor/Sequencer/Private/Filters/Widgets/SFilterBarIsolateHideShow.cpp
+++ b/Engine/Source/Editor/Sequencer/Private/Filters/Widgets/SFilterBarIsolateHideShow.cpp
@@ -248,7 +248,7 @@ FText SFilterBarIsolateHideShow::GetShowAllTracksButtonTooltipText() const
 	}
 
 	const FText SummaryText = MakeLongDisplaySummaryText(*FilterBar);
-	TooltipText = FText::Format(LOCTEXT("IsolateTracksButtonToolTipExtendedWithTotal", "{0}\n\n{1}")
+	TooltipText = FText::Format(LOCTEXT("ShowAllTracksButtonToolTipExtendedWithSummary", "{0}\n\n{1}")
 		, TooltipText, SummaryText);
 
 	return TooltipText;
diff --git a/Engine/Source/Editor/Sequencer/Private/Filters/Widgets/SFilterExpressionHelpDialog.cpp b/Engine/Source/Editor/Sequencer/Private/Filters/Widgets/SFilterExpressionHelpDialog.cpp
--- a/Engine/Source/Editor/Sequencer/Private/Filters/Widgets/SFilterExpressionHelpDialog.cpp
+++ b/Engine/Source/Editor/Sequencer/Private/Filters/Widgets/SFilterExpressionHelpDialog.cpp
@@ -187,8 +187,8 @@ TSharedRef<SWidget> SFilterExpressionHelpDialog::ConstructKeysWidget(const TSet<
 
 TSharedRef<SWidget> SFilterExpressionHelpDialog::ConstructValueWidget(const ESequencerTextFilterValueType InValueType)
 {
-	static const FText LessThanText = LOCTEXT("Comma", "<");
-	static const FText GreaterThanText = LOCTEXT("Comma", ">");
+	static const FText LessThanText = LOCTEXT("CommaLT", "<");
+	static const FText GreaterThanText = LOCTEXT("CommaGT", ">");
 
 	const TSharedRef<SHorizontalBox> HorizontalBox = SNew(SHorizontalBox);
 
diff --git a/Engine/Source/Editor/StaticMeshEditor/Private/StaticMeshEditorViewportToolbarSections.cpp b/Engine/Source/Editor/StaticMeshEditor/Private/StaticMeshEditorViewportToolbarSections.cpp
--- a/Engine/Source/Editor/StaticMeshEditor/Private/StaticMeshEditorViewportToolbarSections.cpp
+++ b/Engine/Source/Editor/StaticMeshEditor/Private/StaticMeshEditorViewportToolbarSections.cpp
@@ -57,14 +57,14 @@ FToolMenuEntry UE::StaticMeshEditor::CreateLODSubmenu()
 					InDynamicSection.AddSubMenu(
 						"LOD",
 						Label,
-						LOCTEXT("LODSubmenuTooltip", ""),
+						FText(),
 						FNewToolMenuDelegate::CreateLambda(
 							[StaticMeshEditorViewportWeak](UToolMenu* Submenu) -> void
 							{
 								if (TSharedPtr<SStaticMeshEditorViewport> Viewport = StaticMeshEditorViewportWeak.Pin())
 								{
 									FToolMenuSection& UnnamedSection =
-										Submenu->FindOrAddSection("", LOCTEXT("UnnamedLabel", ""));
+										Submenu->FindOrAddSection("", FText());
 									TSharedRef<SWidget> LODMenuWidget =
 										UE::StaticMeshEditor::GenerateLODMenuWidget(Viewport);
 									FToolMenuEntry LODSubmenu = FToolMenuEntry::InitWidget("LOD", LODMenuWidget, FText());
