Sequencer Curve Editor: fixed sequencer event binding

#rnx

diff --git a/Engine/Source/Editor/Sequencer/Private/MVVM/CurveEditorExtension.cpp b/Engine/Source/Editor/Sequencer/Private/MVVM/CurveEditorExtension.cpp
--- a/Engine/Source/Editor/Sequencer/Private/MVVM/CurveEditorExtension.cpp
+++ b/Engine/Source/Editor/Sequencer/Private/MVVM/CurveEditorExtension.cpp
@@ -120,7 +120,7 @@ public:
 	{
 		if (TSharedPtr< FSequencer> Sequencer = WeakSequencer.Pin())
 		{
-			Sequencer->OnActivateSequence().AddRaw(this, &FSequencerCurveEditor::HandleSequenceActivated);
+			Sequencer->OnActivateSequence().RemoveAll(this);
 		}
 	}
 
