Sequencer- New Indicator column to the left of existing columns. Currently just will show an indicator for the presence of a condition on that Track or Track Row. Shows at 50% opacity if the condition is currently evaluating to false. Shows at 25% opacity if collapsed and a child has a condition, or the section has a condition.

diff --git a/Engine/Source/Editor/EditorStyle/Private/StarshipStyle.cpp b/Engine/Source/Editor/EditorStyle/Private/StarshipStyle.cpp
--- a/Engine/Source/Editor/EditorStyle/Private/StarshipStyle.cpp
+++ b/Engine/Source/Editor/EditorStyle/Private/StarshipStyle.cpp
@@ -2397,6 +2397,7 @@ void FStarshipEditorStyle::FStyle::SetupSequencerStyles()
 	// Sequencer outliner column UI
 	if (IncludeEditorSpecificStyles())
 	{
+		Set("Sequencer.Column.Condition", new IMAGE_BRUSH_SVG("Sequencer/Column_Widgets/SequencerCondition", Icon14x14));
 		Set("Sequencer.Column.Mute", new IMAGE_BRUSH_SVG("Sequencer/Column_Widgets/SequencerMute", Icon14x14));
 		Set("Sequencer.Column.Locked", new IMAGE_BRUSH_SVG("Sequencer/Column_Widgets/SequencerLocked", Icon14x14));
 		Set("Sequencer.Column.Solo", new IMAGE_BRUSH_SVG("Sequencer/Column_Widgets/SequencerSolo", Icon14x14));
diff --git a/Engine/Source/Editor/Sequencer/Private/MVVM/Extensions/IConditionableExtension.cpp b/Engine/Source/Editor/Sequencer/Private/MVVM/Extensions/IConditionableExtension.cpp
new file mode 100644
--- /dev/null
+++ b/Engine/Source/Editor/Sequencer/Private/MVVM/Extensions/IConditionableExtension.cpp
@@ -0,0 +1,84 @@
+// Copyright Epic Games, Inc. All Rights Reserved.
+
+#include "MVVM/Extensions/IConditionableExtension.h"
+#include "MVVM/Extensions/ITrackExtension.h"
+#include "MVVM/ViewModels/TrackRowModel.h"
+#include "MovieSceneSection.h"
+
+namespace UE::Sequencer
+{
+
+	ECachedConditionState FConditionStateCacheExtension::ComputeFlagsForModel(const FViewModelPtr& ViewModel)
+	{
+		ECachedConditionState& ParentFlags = IndividualItemFlags.Last();
+		ECachedConditionState ThisModelFlags = ECachedConditionState::None;
+
+		if (EnumHasAnyFlags(ParentFlags, ECachedConditionState::HasCondition | ECachedConditionState::ParentHasCondition))
+		{
+			ThisModelFlags |= ECachedConditionState::ParentHasCondition;
+
+			if (EnumHasAnyFlags(ParentFlags, ECachedConditionState::ConditionEvaluatingTrue | ECachedConditionState::ParentHasConditionEvaluatingTrue))
+			{
+				ThisModelFlags |= ECachedConditionState::ParentHasConditionEvaluatingTrue;
+			}
+		}
+
+		if (TViewModelPtr<IConditionableExtension> Conditionable = ViewModel.ImplicitCast())
+		{
+			EConditionableConditionState ConditionState = Conditionable->GetConditionState();
+			if (ConditionState != EConditionableConditionState::None)
+			{
+				ThisModelFlags |= ECachedConditionState::HasCondition;
+			}
+
+			// Special case- if we're a track or track row, and our section has a condition, mark that.
+			// This allows us to surface on the track row level that a section on that row has a condition
+
+			int32 RowIndex = INDEX_NONE;
+			if (TViewModelPtr<FTrackRowModel> TrackRowModel = ViewModel.ImplicitCast())
+			{
+				RowIndex = TrackRowModel->GetRowIndex();
+			}
+
+			if (TViewModelPtr<ITrackExtension> Track = ViewModel.ImplicitCast())
+			{
+				for (UMovieSceneSection* Section : Track->GetSections())
+				{
+					if (Section 
+						&& (RowIndex == INDEX_NONE || Section->GetRowIndex() == RowIndex)
+						&& Section->ConditionContainer.Condition)
+					{
+						ThisModelFlags |= ECachedConditionState::SectionHasCondition;
+					}
+				}
+			}
+
+			if (ConditionState == EConditionableConditionState::HasConditionEvaluatingTrue || ConditionState == EConditionableConditionState::HasConditionEditorForceTrue)
+			{
+				ThisModelFlags |= ECachedConditionState::ConditionEvaluatingTrue;
+				if (ConditionState == EConditionableConditionState::HasConditionEditorForceTrue)
+				{
+					ThisModelFlags |= ECachedConditionState::EditorForceTrue;
+				}
+			}
+		}
+
+		return ThisModelFlags;
+	}
+
+	void FConditionStateCacheExtension::PostComputeChildrenFlags(const FViewModelPtr& ViewModel, ECachedConditionState& OutThisModelFlags, ECachedConditionState& OutPropagateToParentFlags)
+	{
+		// --------------------------------------------------------------------
+		// Handle condition state propagation
+		const bool bHasCondition = EnumHasAnyFlags(OutThisModelFlags, ECachedConditionState::HasCondition);
+		const bool bSectionHasCondition = EnumHasAnyFlags(OutThisModelFlags, ECachedConditionState::SectionHasCondition);
+		const bool bChildrenHaveCondition = EnumHasAnyFlags(OutThisModelFlags, ECachedConditionState::ChildHasCondition);
+
+		if (bHasCondition || bSectionHasCondition || bChildrenHaveCondition)
+		{
+			OutPropagateToParentFlags |= ECachedConditionState::ChildHasCondition;
+		}
+	}
+
+} // namespace UE::Sequencer
+
diff --git a/Engine/Source/Editor/Sequencer/Private/MVVM/SequencerViewModelTypeIDs.cpp b/Engine/Source/Editor/Sequencer/Private/MVVM/SequencerViewModelTypeIDs.cpp
--- a/Engine/Source/Editor/Sequencer/Private/MVVM/SequencerViewModelTypeIDs.cpp
+++ b/Engine/Source/Editor/Sequencer/Private/MVVM/SequencerViewModelTypeIDs.cpp
@@ -18,6 +18,7 @@
 #include "MVVM/Extensions/IObjectBindingExtension.h"
 #include "MVVM/Extensions/ISnappableExtension.h"
 #include "MVVM/Extensions/IBindingLifetimeExtension.h"
+#include "MVVM/Extensions/IConditionableExtension.h"
 
 #include "MVVM/ViewModels/SequencerOutlinerViewModel.h"
 #include "MVVM/ViewModels/SequencerEditorViewModel.h"
@@ -59,6 +60,7 @@ UE_SEQUENCER_DEFINE_CASTABLE(FTrackRowModel);
 
 // Interface types
 UE_SEQUENCER_DEFINE_VIEW_MODEL_TYPE_ID(ITrackExtension);
+UE_SEQUENCER_DEFINE_VIEW_MODEL_TYPE_ID(IConditionableExtension);
 
 // View model types
 UE_SEQUENCER_DEFINE_CASTABLE(FSequencerEditorViewModel);
@@ -79,6 +81,7 @@ UE_SEQUENCER_DEFINE_VIEW_MODEL_TYPE_ID(IObjectBindingExtension);
 UE_SEQUENCER_DEFINE_VIEW_MODEL_TYPE_ID(ISnappableExtension);
 UE_SEQUENCER_DEFINE_VIEW_MODEL_TYPE_ID(IBindingLifetimeExtension);
 UE_SEQUENCER_DEFINE_VIEW_MODEL_TYPE_ID(FOutlinerCacheExtension);
+UE_SEQUENCER_DEFINE_VIEW_MODEL_TYPE_ID(FConditionStateCacheExtension);
 
 } // namespace Sequencer
 } // namespace UE
diff --git a/Engine/Source/Editor/Sequencer/Private/MVVM/ViewModels/OutlinerColumns/IndicatorOutlinerColumn.cpp b/Engine/Source/Editor/Sequencer/Private/MVVM/ViewModels/OutlinerColumns/IndicatorOutlinerColumn.cpp
new file mode 100644
--- /dev/null
+++ b/Engine/Source/Editor/Sequencer/Private/MVVM/ViewModels/OutlinerColumns/IndicatorOutlinerColumn.cpp
@@ -0,0 +1,73 @@
+// Copyright Epic Games, Inc. All Rights Reserved.
+
+#include "MVVM/ViewModels/OutlinerColumns/IndicatorOutlinerColumn.h"
+#include "Widgets/OutlinerColumns/SConditionColumnWidget.h"
+#include "MVVM/Extensions/IConditionableExtension.h"
+#include "MVVM/SharedViewModelData.h"
+#include "Widgets/Layout/SBorder.h"
+
+#define LOCTEXT_NAMESPACE "FIndicatorOutlinerColumn"
+
+namespace UE::Sequencer
+{
+	FIndicatorOutlinerColumn::FIndicatorOutlinerColumn()
+	{
+		Name = FCommonOutlinerNames::Indicator;
+		Label = LOCTEXT("IndicatorColumnLabel", "Indicators");
+		Position = FOutlinerColumnPosition{ 0, EOutlinerColumnGroup::FarLeftGutter };
+		Layout = FOutlinerColumnLayout{ 14, FMargin(4.f, 0.f), HAlign_Center, VAlign_Center, EOutlinerColumnSizeMode::Fixed, EOutlinerColumnFlags::None };
+	}
+
+	bool FIndicatorOutlinerColumn::IsItemCompatibleWithColumn(const FCreateOutlinerColumnParams& InParams) const
+	{
+		// TODO: For now only show condition info, but eventually this should show all indicators
+		if (FConditionStateCacheExtension* ConditionStateCache = InParams.OutlinerExtension.AsModel()->GetSharedData()->CastThis<FConditionStateCacheExtension>())
+		{
+			return EnumHasAnyFlags(ConditionStateCache->GetCachedFlags(InParams.OutlinerExtension), ECachedConditionState::HasCondition | ECachedConditionState::ParentHasCondition | ECachedConditionState::ChildHasCondition | ECachedConditionState::SectionHasCondition);
+		}
+
+		return false;
+	}
+
+	TSharedPtr<SWidget> FIndicatorOutlinerColumn::CreateColumnWidget(const FCreateOutlinerColumnParams& InParams, const TSharedRef<ISequencerTreeViewRow>& TreeViewRow)
+	{
+		//// TODO: For now only show condition info, but eventually this should show all indicators- probably wrap various widgets in one outer widget.
+		
+		// Currently not working attempt at getting a background color behind the column widget 
+		//return SNew(SOverlay)
+		//	+ SOverlay::Slot()
+		//	.HAlign(HAlign_Fill)
+		//	.VAlign(VAlign_Fill)
+		//	[
+		//		SNew(SBorder)
+		//		.BorderImage(FAppStyle::GetBrush("Menu.Background"))
+		//		.ColorAndOpacity(this, &FIndicatorOutlinerColumn::GetColumnBackgroundColorAndOpacity)
+		//		.HAlign(HAlign_Fill)
+		//		.VAlign(VAlign_Fill)
+		//	]
+		//	+ SOverlay::Slot()
+		//	[
+		//		SAssignNew(ColumnWidget, SConditionColumnWidget, SharedThis(this), InParams)
+		//	];
+
+		return SAssignNew(ColumnWidget, SConditionColumnWidget, SharedThis(this), InParams);
+
+	}
+
+	FLinearColor FIndicatorOutlinerColumn::GetColumnBackgroundColorAndOpacity() const
+	{
+		static const FLinearColor ConditionBGColor(92.0f, 220.0f, 205.0f);
+		
+		FLinearColor BGColor = ConditionBGColor;
+
+		if (ColumnWidget.IsValid())
+		{
+			FSlateColor WidgetColor = ColumnWidget->GetImageColorAndOpacity();
+			BGColor.A = WidgetColor.GetSpecifiedColor().A;
+		}
+		return BGColor;
+	}
+
+} // namespace UE::Sequencer
+
+#undef LOCTEXT_NAMESPACE
\ No newline at end of file
diff --git a/Engine/Source/Editor/Sequencer/Private/MVVM/ViewModels/OutlinerColumns/IndicatorOutlinerColumn.h b/Engine/Source/Editor/Sequencer/Private/MVVM/ViewModels/OutlinerColumns/IndicatorOutlinerColumn.h
new file mode 100644
--- /dev/null
+++ b/Engine/Source/Editor/Sequencer/Private/MVVM/ViewModels/OutlinerColumns/IndicatorOutlinerColumn.h
@@ -0,0 +1,33 @@
+// Copyright Epic Games, Inc. All Rights Reserved.
+
+#pragma once
+
+#include "MVVM/ViewModels/OutlinerColumns/OutlinerColumnBase.h"
+#include "Styling/SlateColor.h"
+
+namespace UE::Sequencer
+{
+
+	class SConditionColumnWidget;
+
+	/**
+	 * A column for showing various indicators on rows based on the presence of features (e.g. conditions, time warp) on that row.
+	 */
+	class FIndicatorOutlinerColumn
+		: public FOutlinerColumnBase
+	{
+	public:
+
+		FIndicatorOutlinerColumn();
+
+		bool IsItemCompatibleWithColumn(const FCreateOutlinerColumnParams& InParams) const override;
+		TSharedPtr<SWidget> CreateColumnWidget(const FCreateOutlinerColumnParams& InParams, const TSharedRef<ISequencerTreeViewRow>& TreeViewRow) override;
+
+	private:
+		FLinearColor GetColumnBackgroundColorAndOpacity() const;
+
+	private:
+		TSharedPtr<SConditionColumnWidget> ColumnWidget;
+	};
+
+} // namespace UE::Sequencer
\ No newline at end of file
diff --git a/Engine/Source/Editor/Sequencer/Private/MVVM/ViewModels/SectionModel.cpp b/Engine/Source/Editor/Sequencer/Private/MVVM/ViewModels/SectionModel.cpp
--- a/Engine/Source/Editor/Sequencer/Private/MVVM/ViewModels/SectionModel.cpp
+++ b/Engine/Source/Editor/Sequencer/Private/MVVM/ViewModels/SectionModel.cpp
@@ -12,6 +12,7 @@
 #include "MovieSceneSection.h"
 #include "Channels/MovieSceneChannel.h"
 #include "Channels/MovieSceneChannelProxy.h"
+#include "ScopedTransaction.h"
 
 namespace UE
 {
@@ -542,6 +543,67 @@ TArray<FOverlappingSections> FSectionModel::GetEasingSegments()
 	return Result;
 }
 
+const UMovieSceneCondition* FSectionModel::GetCondition() const
+{
+	UMovieSceneSection* Section = GetSection();
+	if (IsValid(Section))
+	{
+		return Section->ConditionContainer.Condition;
+	}
+	return nullptr;
+}
+
+EConditionableConditionState FSectionModel::GetConditionState() const
+{
+	UMovieSceneSection* Section = GetSection();
+	if (IsValid(Section))
+	{
+		if (Section->ConditionContainer.Condition)
+		{
+			if (Section->ConditionContainer.Condition->bEditorForceTrue)
+			{
+				return EConditionableConditionState::HasConditionEditorForceTrue;
+			}
+
+			TSharedPtr<FSequenceModel> SequenceModel = FindAncestorOfType<FSequenceModel>();
+			TSharedPtr<ISequencer> Sequencer = SequenceModel ? SequenceModel->GetSequencer() : nullptr;
+			if (Sequencer)
+			{
+				FGuid BindingID;
+
+				if (TSharedPtr<IObjectBindingExtension> ParentBinding = FindAncestorOfType<IObjectBindingExtension>())
+				{
+					BindingID = ParentBinding->GetObjectGuid();
+				}
+
+				if (MovieSceneHelpers::EvaluateSequenceCondition(BindingID, Sequencer->GetFocusedTemplateID(), Section->ConditionContainer.Condition, Section, Sequencer->GetSharedPlaybackState()))
+				{
+					return EConditionableConditionState::HasConditionEvaluatingTrue;
+				}
+				else
+				{
+					return EConditionableConditionState::HasConditionEvaluatingFalse;
+				}
+			}
+		}
+	}
+	return EConditionableConditionState::None;
+}
+
+void FSectionModel::SetConditionEditorForceTrue(bool bEditorForceTrue)
+{
+	UMovieSceneSection* Section = GetSection();
+	if (IsValid(Section))
+	{
+		if (Section->ConditionContainer.Condition)
+		{
+			const FScopedTransaction Transaction(NSLOCTEXT("SequencerTrackNode", "ConditionEditorForceTrue", "Set Condition Editor Force True"));
+			Section->ConditionContainer.Condition->Modify();
+			Section->ConditionContainer.Condition->bEditorForceTrue = bEditorForceTrue;
+		}
+	}
+}
+
 } // namespace Sequencer
 } // namespace UE
 
diff --git a/Engine/Source/Editor/Sequencer/Private/MVVM/ViewModels/SequenceModel.cpp b/Engine/Source/Editor/Sequencer/Private/MVVM/ViewModels/SequenceModel.cpp
--- a/Engine/Source/Editor/Sequencer/Private/MVVM/ViewModels/SequenceModel.cpp
+++ b/Engine/Source/Editor/Sequencer/Private/MVVM/ViewModels/SequenceModel.cpp
@@ -51,6 +51,7 @@ void FSequenceModel::InitializeExtensions()
 	NewSharedData->AddDynamicExtension(FMuteStateCacheExtension::ID);
 	NewSharedData->AddDynamicExtension(FSoloStateCacheExtension::ID);
 	NewSharedData->AddDynamicExtension(FLockStateCacheExtension::ID);
+	NewSharedData->AddDynamicExtension(FConditionStateCacheExtension::ID);
 
 	// Add our hierarchical cache processor
 	TSharedPtr<FOutlinerCacheExtension> OutlinerCache = NewSharedData->CastThisSharedChecked<FOutlinerCacheExtension>();
diff --git a/Engine/Source/Editor/Sequencer/Private/MVVM/ViewModels/TrackModel.cpp b/Engine/Source/Editor/Sequencer/Private/MVVM/ViewModels/TrackModel.cpp
--- a/Engine/Source/Editor/Sequencer/Private/MVVM/ViewModels/TrackModel.cpp
+++ b/Engine/Source/Editor/Sequencer/Private/MVVM/ViewModels/TrackModel.cpp
@@ -666,6 +666,67 @@ void FTrackModel::SetIsLocked(bool bInIsLocked)
 	}
 }
 
+const UMovieSceneCondition* FTrackModel::GetCondition() const
+{
+	UMovieSceneTrack* const Track = GetTrack();
+	if (IsValid(Track))
+	{
+		return Track->ConditionContainer.Condition;
+	}
+	return nullptr;
+}
+
+EConditionableConditionState FTrackModel::GetConditionState() const
+{
+	UMovieSceneTrack* const Track = GetTrack();
+	if (IsValid(Track))
+	{
+		if (Track->ConditionContainer.Condition)
+		{
+			if (Track->ConditionContainer.Condition->bEditorForceTrue)
+			{
+				return EConditionableConditionState::HasConditionEditorForceTrue;
+			}
+
+			TSharedPtr<FSequenceModel> SequenceModel = FindAncestorOfType<FSequenceModel>();
+			TSharedPtr<ISequencer> Sequencer = SequenceModel ? SequenceModel->GetSequencer() : nullptr;
+			if (Sequencer)
+			{
+				FGuid BindingID;
+
+				if (TSharedPtr<IObjectBindingExtension> ParentBinding = FindAncestorOfType<IObjectBindingExtension>())
+				{
+					BindingID = ParentBinding->GetObjectGuid();
+				}
+				
+				if (MovieSceneHelpers::EvaluateSequenceCondition(BindingID, Sequencer->GetFocusedTemplateID(), Track->ConditionContainer.Condition, Track, Sequencer->GetSharedPlaybackState()))
+				{
+					return EConditionableConditionState::HasConditionEvaluatingTrue;
+				}
+				else
+				{
+					return EConditionableConditionState::HasConditionEvaluatingFalse;
+				}
+			}
+		}
+	}
+	return EConditionableConditionState::None;
+}
+
+void FTrackModel::SetConditionEditorForceTrue(bool bEditorForceTrue)
+{
+	UMovieSceneTrack* const Track = GetTrack();
+	if (IsValid(Track))
+	{
+		if (Track->ConditionContainer.Condition)
+		{
+			const FScopedTransaction Transaction(NSLOCTEXT("SequencerTrackNode", "ConditionEditorForceTrue", "Set Condition Editor Force True"));
+			Track->ConditionContainer.Condition->Modify();
+			Track->ConditionContainer.Condition->bEditorForceTrue = bEditorForceTrue;
+		}
+	}
+}
+
 bool FTrackModel::CanDrag() const
 {
 	// Can only drag root tracks at the moment
diff --git a/Engine/Source/Editor/Sequencer/Private/MVVM/ViewModels/TrackRowModel.cpp b/Engine/Source/Editor/Sequencer/Private/MVVM/ViewModels/TrackRowModel.cpp
--- a/Engine/Source/Editor/Sequencer/Private/MVVM/ViewModels/TrackRowModel.cpp
+++ b/Engine/Source/Editor/Sequencer/Private/MVVM/ViewModels/TrackRowModel.cpp
@@ -212,6 +212,83 @@ void FTrackRowModel::SetIsLocked(bool bInIsLocked)
 	}
 }
 
+const UMovieSceneCondition* FTrackRowModel::GetCondition() const
+{
+	UMovieSceneTrack* Track = GetTrack();
+
+	if (Track)
+	{
+		if (const FMovieSceneTrackRowMetadata* TrackRowMetadata = Track->FindTrackRowMetadata(GetRowIndex()))
+		{
+			return TrackRowMetadata->ConditionContainer.Condition;
+		}
+	}
+
+	return nullptr;
+}
+
+
+EConditionableConditionState FTrackRowModel::GetConditionState() const
+{
+	UMovieSceneTrack* Track = GetTrack();
+
+	if (Track)
+	{
+		FGuid BindingID;
+		FMovieSceneSequenceID SequenceID = MovieSceneSequenceID::Root;
+		if (TViewModelPtr<FObjectBindingModel> ObjectBindingModel = FindAncestorOfType<FObjectBindingModel>())
+		{
+			BindingID = ObjectBindingModel->GetObjectGuid();
+		}
+
+		if (TViewModelPtr<FSequenceModel> SequenceModel = FindAncestorOfType<FSequenceModel>())
+		{
+			SequenceID = SequenceModel->GetSequenceID();
+
+			if (TSharedPtr<FSequencerEditorViewModel> SequencerModel = SequenceModel->GetEditor())
+			{
+				if (const FMovieSceneTrackRowMetadata* TrackRowMetadata = Track->FindTrackRowMetadata(GetRowIndex()))
+				{
+					if (TrackRowMetadata->ConditionContainer.Condition)
+					{
+						if (TrackRowMetadata->ConditionContainer.Condition->bEditorForceTrue)
+						{
+							return EConditionableConditionState::HasConditionEditorForceTrue;
+						}
+						else if (MovieSceneHelpers::EvaluateSequenceCondition(BindingID, SequencerModel->GetSequencer()->GetFocusedTemplateID(), TrackRowMetadata->ConditionContainer.Condition, Track, SequencerModel->GetSequencer()->GetSharedPlaybackState()))
+						{
+							return EConditionableConditionState::HasConditionEvaluatingTrue;
+						}
+						else
+						{
+							return EConditionableConditionState::HasConditionEvaluatingFalse;
+						}
+					}
+				}
+			}
+		}
+	}
+	return EConditionableConditionState::None;
+}
+
+void FTrackRowModel::SetConditionEditorForceTrue(bool bEditorForceTrue)
+{
+	UMovieSceneTrack* Track = GetTrack();
+
+	if (Track)
+	{
+		if (const FMovieSceneTrackRowMetadata* TrackRowMetadata = Track->FindTrackRowMetadata(GetRowIndex()))
+		{
+			if (TrackRowMetadata->ConditionContainer.Condition)
+			{
+				const FScopedTransaction Transaction(NSLOCTEXT("SequencerTrackRowNode", "ConditionEditorForceTrue", "Set Condition Editor Force True"));
+				TrackRowMetadata->ConditionContainer.Condition->Modify();
+				TrackRowMetadata->ConditionContainer.Condition->bEditorForceTrue = bEditorForceTrue;
+			}
+		}
+	}
+}
+
 FSlateFontInfo FTrackRowModel::GetLabelFont() const
 {
 	bool bAllAnimated = false;
diff --git a/Engine/Source/Editor/Sequencer/Private/SequencerCommonHelpers.cpp b/Engine/Source/Editor/Sequencer/Private/SequencerCommonHelpers.cpp
--- a/Engine/Source/Editor/Sequencer/Private/SequencerCommonHelpers.cpp
+++ b/Engine/Source/Editor/Sequencer/Private/SequencerCommonHelpers.cpp
@@ -305,7 +305,7 @@ public:
 	{
 		if (const TSharedPtr<ISequencer> Sequencer = WeakSequencer.Pin())
 		{
-			Sequencer->NotifyMovieSceneDataChanged(EMovieSceneDataChangeType::TrackValueChanged);
+			Sequencer->NotifyMovieSceneDataChanged(EMovieSceneDataChangeType::MovieSceneStructureItemsChanged);
 		}
 	}
 
diff --git a/Engine/Source/Editor/Sequencer/Private/SequencerModule.cpp b/Engine/Source/Editor/Sequencer/Private/SequencerModule.cpp
--- a/Engine/Source/Editor/Sequencer/Private/SequencerModule.cpp
+++ b/Engine/Source/Editor/Sequencer/Private/SequencerModule.cpp
@@ -33,6 +33,7 @@
 #include "MVVM/ViewModels/SequenceModel.h"
 #include "MVVM/ViewModels/SequencerEditorViewModel.h"
 
+#include "MVVM/ViewModels/OutlinerColumns/IndicatorOutlinerColumn.h"
 #include "MVVM/ViewModels/OutlinerColumns/LockOutlinerColumn.h"
 #include "MVVM/ViewModels/OutlinerColumns/MuteOutlinerColumn.h"
 #include "MVVM/ViewModels/OutlinerColumns/PinOutlinerColumn.h"
@@ -479,6 +480,9 @@ public:
 			FPropertyEditorModule& EditModule = FModuleManager::Get().GetModuleChecked<FPropertyEditorModule>("PropertyEditor");
 			OnGetGlobalRowExtensionHandle = EditModule.GetGlobalRowExtensionDelegate().AddStatic(&RegisterKeyframeExtensionHandler);
 
+			// Register far left gutter columns
+			IndicatorOutlinerColumnHandle = RegisterOutlinerColumn(FOnCreateOutlinerColumn::CreateStatic([] { return TSharedRef<IOutlinerColumn>(MakeShared<FIndicatorOutlinerColumn>()); }));
+
 			// Register left gutter columns
 			PinOutlinerColumnHandle  = RegisterOutlinerColumn(FOnCreateOutlinerColumn::CreateStatic([]{ return TSharedRef<IOutlinerColumn>(MakeShared<FPinOutlinerColumn>()); }));
 			MuteOutlinerColumnHandle = RegisterOutlinerColumn(FOnCreateOutlinerColumn::CreateStatic([]{ return TSharedRef<IOutlinerColumn>(MakeShared<FMuteOutlinerColumn>()); }));
@@ -542,6 +546,7 @@ public:
 			FEditorModeRegistry::Get().UnregisterMode(FSequencerEdMode::EM_SequencerMode);
 
 			// unregister outliner columns
+			UnregisterOutlinerColumn(IndicatorOutlinerColumnHandle);
 			UnregisterOutlinerColumn(PinOutlinerColumnHandle);
 			UnregisterOutlinerColumn(MuteOutlinerColumnHandle);
 			UnregisterOutlinerColumn(LockOutlinerColumnHandle);
@@ -738,6 +743,7 @@ private:
 	TArray<FMovieRendererEntry> MovieRenderers;
 
 	// Outliner Column Delegate Handles
+	FDelegateHandle IndicatorOutlinerColumnHandle;
 	FDelegateHandle PinOutlinerColumnHandle;
 	FDelegateHandle MuteOutlinerColumnHandle;
 	FDelegateHandle LockOutlinerColumnHandle;
diff --git a/Engine/Source/Editor/Sequencer/Private/Widgets/OutlinerColumns/SConditionColumnWidget.cpp b/Engine/Source/Editor/Sequencer/Private/Widgets/OutlinerColumns/SConditionColumnWidget.cpp
new file mode 100644
--- /dev/null
+++ b/Engine/Source/Editor/Sequencer/Private/Widgets/OutlinerColumns/SConditionColumnWidget.cpp
@@ -0,0 +1,91 @@
+// Copyright Epic Games, Inc. All Rights Reserved.
+
+#include "Widgets/OutlinerColumns/SConditionColumnWidget.h"
+
+#include "MVVM/SharedViewModelData.h"
+#include "MVVM/ViewModels/ViewModelIterators.h"
+#include "MVVM/ViewModels/EditorViewModel.h"
+#include "MVVM/ViewModels/OutlinerColumns/IOutlinerColumn.h"
+#include "MVVM/Extensions/IConditionableExtension.h"
+#include "MVVM/Selection/SequencerCoreSelection.h"
+#include "MVVM/Selection/SequencerOutlinerSelection.h"
+
+namespace UE::Sequencer
+{
+	void SConditionColumnWidget::Construct(const FArguments& InArgs, const TWeakPtr<IOutlinerColumn> InWeakOutlinerColumn, const FCreateOutlinerColumnParams& InParams)
+	{
+		SColumnToggleWidget::Construct(
+			SColumnToggleWidget::FArguments(),
+			InWeakOutlinerColumn,
+			InParams
+		);
+
+		WeakConditionStateCacheExtension = CastViewModel<FConditionStateCacheExtension>(InParams.OutlinerExtension.AsModel()->GetSharedData());
+	}
+
+	bool SConditionColumnWidget::IsActive() const
+	{
+		if (TViewModelPtr<FConditionStateCacheExtension> StateCache = WeakConditionStateCacheExtension.Pin())
+		{
+			return EnumHasAnyFlags(StateCache->GetCachedFlags(ModelID), ECachedConditionState::HasCondition);
+		}
+
+		return false;
+	}
+
+
+	const FSlateBrush* SConditionColumnWidget::GetActiveBrush() const
+	{
+		static const FName NAME_ConditionBrush = TEXT("Sequencer.Column.Condition");
+		return FAppStyle::Get().GetBrush(NAME_ConditionBrush);
+	}
+
+	FSlateColor SConditionColumnWidget::GetImageColorAndOpacity() const
+	{
+		TSharedPtr<FEditorViewModel> Editor = WeakEditor.Pin();
+		TViewModelPtr<IOutlinerExtension> OutlinerItem = WeakOutlinerExtension.Pin();
+		
+		bool bActive = IsActive();
+		bool bEvaluatingTrue = false;
+		bool bChildHasCondition = false;
+		bool bSectionHasCondition = false;
+		if (TViewModelPtr<FConditionStateCacheExtension> StateCache = WeakConditionStateCacheExtension.Pin())
+		{
+			bEvaluatingTrue = EnumHasAnyFlags(StateCache->GetCachedFlags(ModelID), ECachedConditionState::ConditionEvaluatingTrue);
+			bChildHasCondition = EnumHasAnyFlags(StateCache->GetCachedFlags(ModelID), ECachedConditionState::ChildHasCondition);
+			bSectionHasCondition = EnumHasAnyFlags(StateCache->GetCachedFlags(ModelID), ECachedConditionState::SectionHasCondition);
+		}
+		FLinearColor OutColor = FLinearColor::White;
+
+		if (!Editor || !OutlinerItem)
+		{
+			return OutColor;
+		}
+
+		float Opacity = 0.0f;
+
+		if (IsActive())
+		{
+			if (bEvaluatingTrue)
+			{
+				Opacity = 1.0f;
+			}
+			else
+			{
+				Opacity = 0.5f;
+			}
+		}
+		else if ((bChildHasCondition && !OutlinerItem->IsExpanded()) || bSectionHasCondition)
+		{
+			Opacity = 0.25f;
+		}
+		else
+		{
+			// Not active, invisible
+			Opacity = 0.0f;
+		}
+		OutColor.A = Opacity;
+		return OutColor;
+	}
+
+} // namespace UE::Sequencer
diff --git a/Engine/Source/Editor/Sequencer/Private/Widgets/OutlinerColumns/SConditionColumnWidget.h b/Engine/Source/Editor/Sequencer/Private/Widgets/OutlinerColumns/SConditionColumnWidget.h
new file mode 100644
--- /dev/null
+++ b/Engine/Source/Editor/Sequencer/Private/Widgets/OutlinerColumns/SConditionColumnWidget.h
@@ -0,0 +1,59 @@
+// Copyright Epic Games, Inc. All Rights Reserved.
+
+#pragma once
+
+#include "MVVM/Views/OutlinerColumns/SColumnToggleWidget.h"
+
+namespace UE::Sequencer
+{
+	class FConditionStateCacheExtension;
+	/**
+	 * A widget that shows an indicator based on the presence and state of conditions in a row
+	 */
+	class SConditionColumnWidget
+		: public SColumnToggleWidget
+	{
+	public:
+		SLATE_BEGIN_ARGS(SConditionColumnWidget) {}
+		SLATE_END_ARGS()
+
+		void Construct(const FArguments& InArgs, const TWeakPtr<IOutlinerColumn> InWeakOutlinerColumn, const FCreateOutlinerColumnParams& InParams);
+
+		/** Get the color and opacity of the column toggle widget. */
+		virtual FSlateColor GetImageColorAndOpacity() const override;
+
+	protected:
+
+		/** Returns whether or not the item is directly active or not. */
+		virtual bool IsActive() const override;
+
+		/** Sets the active state of this item to the input value. */
+		virtual void SetIsActive(const bool bInIsActive) override { /* no-op for now, no toggle*/ }
+
+		/** Returns true if a child of this item is active. */
+		virtual bool IsChildActive() const override { return false; }
+
+		/** Returns true if this item is implicitly active, but not directly active */
+		virtual bool IsImplicitlyActive() const override { return false; }
+
+		/** Returns the brush to be used to represent a widget is active. */
+		virtual const FSlateBrush* GetActiveBrush() const override;
+
+	protected:
+
+		/** Get the image this widget displays. */
+		virtual const FSlateBrush* GetBrush() const override { return GetActiveBrush(); }
+
+	protected:
+
+		/** Handles left click-type inputs and potentially begins drag and drop operation. */
+		FReply HandleClick() { /* no - op for now, no toggle*/ }
+
+	private:
+
+		/** Weak cache extension ptr (can be null). */
+		TWeakViewModelPtr<FConditionStateCacheExtension> WeakConditionStateCacheExtension;
+
+	};
+
+} // namespace UE::Sequencer
\ No newline at end of file
diff --git a/Engine/Source/Editor/Sequencer/Public/MVVM/Extensions/IConditionableExtension.h b/Engine/Source/Editor/Sequencer/Public/MVVM/Extensions/IConditionableExtension.h
new file mode 100644
--- /dev/null
+++ b/Engine/Source/Editor/Sequencer/Public/MVVM/Extensions/IConditionableExtension.h
@@ -0,0 +1,75 @@
+// Copyright Epic Games, Inc. All Rights Reserved.
+
+#pragma once
+
+#include "CoreMinimal.h"
+#include "MVVM/ViewModelTypeID.h"
+#include "MVVM/Extensions/HierarchicalCacheExtension.h"
+
+class UMovieSceneCondition;
+
+namespace UE
+{
+	namespace Sequencer
+	{
+		enum class EConditionableConditionState
+		{
+			None,
+			HasConditionEvaluatingFalse,
+			HasConditionEvaluatingTrue,
+			HasConditionEditorForceTrue,
+		};
+
+		/**
+		 * An extension for models that can have conditions
+		 */
+		class SEQUENCER_API IConditionableExtension
+		{
+		public:
+
+			UE_SEQUENCER_DECLARE_VIEW_MODEL_TYPE_ID(IConditionableExtension)
+
+			virtual ~IConditionableExtension() {}
+
+			/* Returns the condition in question for this model if applicable. */
+			virtual const UMovieSceneCondition* GetCondition() const = 0;
+
+			/** Returns condition state for this item */
+			virtual EConditionableConditionState GetConditionState() const = 0;
+
+			/** Set this item's condition to evaluate true for editor preview purposes */
+			virtual void SetConditionEditorForceTrue(bool bEditorForceTrue) = 0;
+		};
+
+		enum class ECachedConditionState
+		{
+			None = 0,
+
+			HasCondition = 1 << 0,
+			ConditionEvaluatingTrue = 1 << 1,
+			EditorForceTrue = 1 << 2,
+			ChildHasCondition = 1 << 3,
+			SectionHasCondition = 1 << 4,
+			ParentHasCondition = 1 << 5,
+			ParentHasConditionEvaluatingTrue = 1 << 6,
+
+			InheritedFromChildren = ChildHasCondition,
+		};
+		ENUM_CLASS_FLAGS(ECachedConditionState)
+
+		class SEQUENCER_API FConditionStateCacheExtension
+			: public TFlagStateCacheExtension<ECachedConditionState>
+		{
+		public:
+
+			UE_SEQUENCER_DECLARE_VIEW_MODEL_TYPE_ID(FConditionStateCacheExtension);
+
+		private:
+
+			ECachedConditionState ComputeFlagsForModel(const FViewModelPtr& ViewModel) override;
+			void PostComputeChildrenFlags(const FViewModelPtr& ViewModel, ECachedConditionState& OutThisModelFlags, ECachedConditionState& OutPropagateToParentFlags) override;
+		};
+
+	} // namespace Sequencer
+} // namespace UE
+
diff --git a/Engine/Source/Editor/Sequencer/Public/MVVM/ViewModels/SectionModel.h b/Engine/Source/Editor/Sequencer/Public/MVVM/ViewModels/SectionModel.h
--- a/Engine/Source/Editor/Sequencer/Public/MVVM/ViewModels/SectionModel.h
+++ b/Engine/Source/Editor/Sequencer/Public/MVVM/ViewModels/SectionModel.h
@@ -13,6 +13,7 @@
 #include "MVVM/Extensions/IDraggableTrackAreaExtension.h"
 #include "MVVM/Extensions/IStretchableExtension.h"
 #include "MVVM/Extensions/LinkedOutlinerExtension.h"
+#include "MVVM/Extensions/IConditionableExtension.h"
 #include "Delegates/DelegateCombinations.h"
 #include "EventHandlers/ISignedObjectEventHandler.h"
 #include "EventHandlers/ISectionEventHandler.h"
@@ -41,6 +42,7 @@ class SEQUENCER_API FSectionModel
 	, public ISnappableExtension
 	, public IDraggableTrackAreaExtension
 	, public IStretchableExtension
+	, public IConditionableExtension
 	, public UE::MovieScene::TIntrusiveEventHandler<UE::MovieScene::ISignedObjectEventHandler>
 	, public UE::MovieScene::TIntrusiveEventHandler<UE::MovieScene::ISectionEventHandler>
 {
@@ -54,6 +56,7 @@ public:
 		, ISnappableExtension
 		, IDraggableTrackAreaExtension
 		, IStretchableExtension
+		, IConditionableExtension
 	);
 
 	DECLARE_MULTICAST_DELEGATE_OneParam(FOnModelUpdated, FSectionModel*)
@@ -120,6 +123,11 @@ public:
 	void OnStretch(const IStretchOperation& StretchOperation, const FStretchScreenParameters& ScreenParameters, FStretchParameters* InOutParameters) override;
 	void OnEndStretch(const IStretchOperation& StretchOperation, const FStretchScreenParameters& ScreenParameters, FStretchParameters* InOutParameters) override;
 
+	/*~ IConditionableExtension Interface */
+	const UMovieSceneCondition* GetCondition() const override;
+	EConditionableConditionState GetConditionState() const override;
+	void SetConditionEditorForceTrue(bool bEditorForceTrue) override;
+
 private:
 
 	void UpdateCachedData();
diff --git a/Engine/Source/Editor/Sequencer/Public/MVVM/ViewModels/TrackModel.h b/Engine/Source/Editor/Sequencer/Public/MVVM/ViewModels/TrackModel.h
--- a/Engine/Source/Editor/Sequencer/Public/MVVM/ViewModels/TrackModel.h
+++ b/Engine/Source/Editor/Sequencer/Public/MVVM/ViewModels/TrackModel.h
@@ -15,6 +15,7 @@
 #include "MVVM/Extensions/IDeletableExtension.h"
 #include "MVVM/Extensions/IDraggableOutlinerExtension.h"
 #include "MVVM/Extensions/IDeletableExtension.h"
+#include "MVVM/Extensions/IConditionableExtension.h"
 #include "MVVM/ViewModels/ViewModelHierarchy.h"
 #include "EventHandlers/ISignedObjectEventHandler.h"
 
@@ -46,6 +47,7 @@ class SEQUENCER_API FTrackModel
 	, public ISortableExtension
 	, public IDraggableOutlinerExtension
 	, public IDeletableExtension
+	, public IConditionableExtension
 	, public UE::MovieScene::TIntrusiveEventHandler<UE::MovieScene::ISignedObjectEventHandler>
 	, public UE::MovieScene::IDeferredSignedObjectFlushSignal
 {
@@ -61,7 +63,8 @@ public:
 		, IGroupableExtension
 		, ISortableExtension
 		, IDraggableOutlinerExtension
-		, IDeletableExtension);
+		, IDeletableExtension
+		, IConditionableExtension);
 
 	explicit FTrackModel(UMovieSceneTrack* Track);
 	~FTrackModel();
@@ -107,6 +110,11 @@ public:
 	ELockableLockState GetLockState() const override;
 	void SetIsLocked(bool bIsLocked) override;
 
+	/*~ IConditionableExtension Interface */
+	const UMovieSceneCondition* GetCondition() const override;
+	EConditionableConditionState GetConditionState() const override;
+	void SetConditionEditorForceTrue(bool bEditorForceTrue) override;
+
 	/*~ ITrackAreaExtension */
 	FTrackAreaParameters GetTrackAreaParameters() const override;
 	FViewModelVariantIterator GetTrackAreaModelList() const override;
diff --git a/Engine/Source/Editor/Sequencer/Public/MVVM/ViewModels/TrackRowModel.h b/Engine/Source/Editor/Sequencer/Public/MVVM/ViewModels/TrackRowModel.h
--- a/Engine/Source/Editor/Sequencer/Public/MVVM/ViewModels/TrackRowModel.h
+++ b/Engine/Source/Editor/Sequencer/Public/MVVM/ViewModels/TrackRowModel.h
@@ -10,6 +10,7 @@
 #include "MVVM/Extensions/IRenameableExtension.h"
 #include "MVVM/Extensions/IResizableExtension.h"
 #include "MVVM/Extensions/IDeletableExtension.h"
+#include "MVVM/Extensions/IConditionableExtension.h"
 
 #include "UObject/WeakObjectPtr.h"
 #include "UObject/WeakObjectPtrTemplates.h"
@@ -33,10 +34,11 @@ class SEQUENCER_API FTrackRowModel
 	, public IRenameableExtension
 	, public IResizableExtension
 	, public IDeletableExtension
+	, public IConditionableExtension
 {
 public:
 
-	UE_SEQUENCER_DECLARE_CASTABLE(FTrackRowModel, FMuteSoloOutlinerItemModel, ITrackAreaExtension, ILockableExtension, ITrackExtension, IDeletableExtension, IMutableExtension, ISoloableExtension, IRenameableExtension);
+	UE_SEQUENCER_DECLARE_CASTABLE(FTrackRowModel, FMuteSoloOutlinerItemModel, ITrackAreaExtension, ILockableExtension, ITrackExtension, IDeletableExtension, IMutableExtension, ISoloableExtension, IRenameableExtension, IConditionableExtension);
 
 	explicit FTrackRowModel(UMovieSceneTrack* InTrack, int32 InRowIndex);
 	~FTrackRowModel();
@@ -89,6 +91,11 @@ public:
 	ELockableLockState GetLockState() const override;
 	void SetIsLocked(bool bIsLocked) override;
 
+	/*~ IConditionableExtension Interface */
+	const UMovieSceneCondition* GetCondition() const override;
+	EConditionableConditionState GetConditionState() const override;
+	void SetConditionEditorForceTrue(bool bEditorForceTrue) override;
+
 private:
 
 	FViewModelListHead SectionList;
diff --git a/Engine/Source/Editor/SequencerCore/Private/MVVM/ViewModels/OutlinerColumns/OutlinerColumnTypes.cpp b/Engine/Source/Editor/SequencerCore/Private/MVVM/ViewModels/OutlinerColumns/OutlinerColumnTypes.cpp
--- a/Engine/Source/Editor/SequencerCore/Private/MVVM/ViewModels/OutlinerColumns/OutlinerColumnTypes.cpp
+++ b/Engine/Source/Editor/SequencerCore/Private/MVVM/ViewModels/OutlinerColumns/OutlinerColumnTypes.cpp
@@ -5,6 +5,7 @@
 namespace UE::Sequencer
 {
 
+FName FCommonOutlinerNames::Indicator("Indicator");
 FName FCommonOutlinerNames::Pin("Pin");
 FName FCommonOutlinerNames::Lock("Lock");
 FName FCommonOutlinerNames::Mute("Mute");
diff --git a/Engine/Source/Editor/SequencerCore/Private/MVVM/Views/SOutlinerView.cpp b/Engine/Source/Editor/SequencerCore/Private/MVVM/Views/SOutlinerView.cpp
--- a/Engine/Source/Editor/SequencerCore/Private/MVVM/Views/SOutlinerView.cpp
+++ b/Engine/Source/Editor/SequencerCore/Private/MVVM/Views/SOutlinerView.cpp
@@ -378,24 +378,30 @@ void SOutlinerView::UpdateOutlinerColumns()
 
 	// ----------------------------------------------------------------------------------------------------------
 	// Populate columns
-	const int32 NumLeftGutter  = CreateOutlinerColumnsForGroup(0,                         EOutlinerColumnGroup::LeftGutter);
-	const int32 NumCenter      = CreateOutlinerColumnsForGroup(NumLeftGutter,             EOutlinerColumnGroup::Center);
-	const int32 NumRightGutter = CreateOutlinerColumnsForGroup(NumLeftGutter + NumCenter, EOutlinerColumnGroup::RightGutter);
+	const int32 NumFarLeftGutter	= CreateOutlinerColumnsForGroup(0,											  EOutlinerColumnGroup::FarLeftGutter);
+	const int32 NumLeftGutter		= CreateOutlinerColumnsForGroup(NumFarLeftGutter,							  EOutlinerColumnGroup::LeftGutter);
+	const int32 NumCenter			= CreateOutlinerColumnsForGroup(NumFarLeftGutter + NumLeftGutter,             EOutlinerColumnGroup::Center);
+	const int32 NumRightGutter		= CreateOutlinerColumnsForGroup(NumFarLeftGutter + NumLeftGutter + NumCenter, EOutlinerColumnGroup::RightGutter);
 
 	// ----------------------------------------------------------------------------------------------------------
 	// Add some padding to the leading and trailing edge of the first and last columns in each group respectively
 	//    This is implemented this way because columns can be turned on and off dynamically, but we must
 	//    always have a consistent padding within the group. Separators dynamically appear based on the presence of
 	//    each group so we can't put padding on those
+	if (NumFarLeftGutter > 0)
+	{
+		ColumnMetaData->Columns[0].CellPadding.Left += 4.f;
+		ColumnMetaData->Columns[NumFarLeftGutter - 1].CellPadding.Right += 4.f;
+	}
 	if (NumLeftGutter > 0)
 	{
-		ColumnMetaData->Columns[0              ].CellPadding.Left  += 4.f;
-		ColumnMetaData->Columns[NumLeftGutter-1].CellPadding.Right += 4.f;
+		ColumnMetaData->Columns[NumFarLeftGutter].CellPadding.Left  += 4.f;
+		ColumnMetaData->Columns[NumFarLeftGutter+NumLeftGutter-1].CellPadding.Right += 4.f;
 	}
 	if (NumCenter > 0)
 	{
-		ColumnMetaData->Columns[NumLeftGutter            ].CellPadding.Left  += 4.f;
-		ColumnMetaData->Columns[NumLeftGutter+NumCenter-1].CellPadding.Right += 4.f;
+		ColumnMetaData->Columns[NumFarLeftGutter+NumLeftGutter            ].CellPadding.Left  += 4.f;
+		ColumnMetaData->Columns[NumFarLeftGutter+NumLeftGutter+NumCenter-1].CellPadding.Right += 4.f;
 	}
 
 	// No additional padding on the right gutter intentionally
@@ -405,10 +411,19 @@ void SOutlinerView::UpdateOutlinerColumns()
 	//      Only add separators if they actually separate columns.
 	int32 NumSeparators = 0;
 
-	int32 InsertIndex = NumLeftGutter;
+	int32 InsertIndex = NumFarLeftGutter;
 	if (InsertIndex < ColumnMetaData->Columns.Num())
 	{
 		InsertSeparatorColumn(InsertIndex++, ++NumSeparators);
+		if (NumLeftGutter > 0)
+		{
+			InsertIndex += NumLeftGutter;
+			if (InsertIndex < ColumnMetaData->Columns.Num())
+			{
+				InsertSeparatorColumn(InsertIndex++, ++NumSeparators);
+			}
+		}
+
 		if (NumCenter > 0)
 		{
 			InsertIndex += NumCenter;
diff --git a/Engine/Source/Editor/SequencerCore/Public/MVVM/ViewModels/OutlinerColumns/OutlinerColumnTypes.h b/Engine/Source/Editor/SequencerCore/Public/MVVM/ViewModels/OutlinerColumns/OutlinerColumnTypes.h
--- a/Engine/Source/Editor/SequencerCore/Public/MVVM/ViewModels/OutlinerColumns/OutlinerColumnTypes.h
+++ b/Engine/Source/Editor/SequencerCore/Public/MVVM/ViewModels/OutlinerColumns/OutlinerColumnTypes.h
@@ -23,6 +23,8 @@ ENUM_CLASS_FLAGS(EOutlinerColumnFlags)
 
 enum class EOutlinerColumnGroup : uint8
 {
+	FarLeftGutter,
+
 	LeftGutter,
 
 	Center,
@@ -63,6 +65,7 @@ struct FOutlinerColumnLayout
 
 struct FCommonOutlinerNames
 {
+	static SEQUENCERCORE_API FName Indicator;
 	static SEQUENCERCORE_API FName Pin;
 	static SEQUENCERCORE_API FName Lock;
 	static SEQUENCERCORE_API FName Mute;
