Sequencer- Adding condition functionality to the advanced text filters. Users can now use:

CONDITION/HASCONDITION =true/false, for searching for the presence of a condition
CONDITIONPASSES=true/false for searching for the presence of a succeeding/failing condition
CONDITIONCLASS/CONDITIONTYPE=typename for searching for the presence of a condition of a given class
CONDITIONFUNC=funcname for searching for a director blueprint condition with its endpoint matching the name

Also added a top-level filter called Condition which is just a duplicate of the CONDITION/HASCONDITION text filter.

diff --git a/Engine/Source/Editor/Sequencer/Private/Filters/Filters/SequencerTrackFilter_Condition.cpp b/Engine/Source/Editor/Sequencer/Private/Filters/Filters/SequencerTrackFilter_Condition.cpp
new file mode 100644
--- /dev/null
+++ b/Engine/Source/Editor/Sequencer/Private/Filters/Filters/SequencerTrackFilter_Condition.cpp
@@ -0,0 +1,61 @@
+// Copyright Epic Games, Inc. All Rights Reserved.
+
+#include "SequencerTrackFilter_Condition.h"
+#include "Filters/SequencerFilterBar.h"
+#include "Filters/SequencerTrackFilterCommands.h"
+#include "MVVM/ViewModels/ObjectBindingModel.h"
+#include "Sequencer.h"
+
+using namespace UE::Sequencer;
+
+#define LOCTEXT_NAMESPACE "SequencerTrackFilter_Condition"
+
+FSequencerTrackFilter_Condition::FSequencerTrackFilter_Condition(ISequencerTrackFilters& InFilterInterface, TSharedPtr<FFilterCategory> InCategory)
+	: FSequencerTrackFilter(InFilterInterface, MoveTemp(InCategory))
+{
+}
+
+bool FSequencerTrackFilter_Condition::ShouldUpdateOnTrackValueChanged() const
+{
+	return true;
+}
+
+FText FSequencerTrackFilter_Condition::GetDefaultToolTipText() const
+{
+	return LOCTEXT("SequencerTrackFilter_ConditionToolTip", "Show only tracks and sections with conditions");
+}
+
+TSharedPtr<FUICommandInfo> FSequencerTrackFilter_Condition::GetToggleCommand() const
+{
+	return FSequencerTrackFilterCommands::Get().ToggleFilter_Condition;
+}
+
+FText FSequencerTrackFilter_Condition::GetDisplayName() const
+{
+	return LOCTEXT("SequencerTrackFilter_Condition", "Condition");
+}
+
+FSlateIcon FSequencerTrackFilter_Condition::GetIcon() const
+{
+	return FSlateIcon(FAppStyle::GetAppStyleSetName(), TEXT("Sequencer.Column.Condition"));
+}
+
+FString FSequencerTrackFilter_Condition::GetName() const
+{
+	return TEXT("Condition");
+}
+
+bool FSequencerTrackFilter_Condition::PassesFilter(FSequencerTrackFilterType InItem) const
+{
+	const TViewModelPtr<IConditionableExtension> Conditionable = InItem.ImplicitCast();
+	if (!Conditionable)
+	{
+		return false;
+	}
+	else
+	{
+		return Conditionable->GetConditionState() != EConditionableConditionState::None;
+	}
+}
+
+#undef LOCTEXT_NAMESPACE
diff --git a/Engine/Source/Editor/Sequencer/Private/Filters/Filters/SequencerTrackFilter_Condition.h b/Engine/Source/Editor/Sequencer/Private/Filters/Filters/SequencerTrackFilter_Condition.h
new file mode 100644
--- /dev/null
+++ b/Engine/Source/Editor/Sequencer/Private/Filters/Filters/SequencerTrackFilter_Condition.h
@@ -0,0 +1,27 @@
+// Copyright Epic Games, Inc. All Rights Reserved.
+
+#pragma once
+
+#include "Filters/SequencerTrackFilterBase.h"
+
+class FSequencerTrackFilter_Condition : public FSequencerTrackFilter
+{
+public:
+	FSequencerTrackFilter_Condition(ISequencerTrackFilters& InFilterInterface, TSharedPtr<FFilterCategory> InCategory = nullptr);
+
+	//~ Begin FSequencerTrackFilter
+	virtual bool ShouldUpdateOnTrackValueChanged() const override;
+	virtual FText GetDefaultToolTipText() const override;
+	virtual TSharedPtr<FUICommandInfo> GetToggleCommand() const override;
+	//~ End FSequencerTrackFilter
+
+	//~ Begin FFilterBase
+	virtual FText GetDisplayName() const override;
+	virtual FSlateIcon GetIcon() const override;
+	//~ End FFilterBase
+
+	//~ Begin IFilter
+	virtual FString GetName() const override;
+	virtual bool PassesFilter(FSequencerTrackFilterType InItem) const override;
+	//~ End IFilter
+};
diff --git a/Engine/Source/Editor/Sequencer/Private/Filters/Filters/SequencerTrackFilter_Text.cpp b/Engine/Source/Editor/Sequencer/Private/Filters/Filters/SequencerTrackFilter_Text.cpp
--- a/Engine/Source/Editor/Sequencer/Private/Filters/Filters/SequencerTrackFilter_Text.cpp
+++ b/Engine/Source/Editor/Sequencer/Private/Filters/Filters/SequencerTrackFilter_Text.cpp
@@ -1,4 +1,4 @@
-﻿// Copyright Epic Games, Inc. All Rights Reserved.
+// Copyright Epic Games, Inc. All Rights Reserved.
 
 #include "SequencerTrackFilter_Text.h"
 #include "Filters/SequencerFilterBar.h"
@@ -7,6 +7,10 @@
 #include "Misc/TextFilterExpressionEvaluator.h"
 #include "MVVM/ViewModelPtr.h"
 #include "Filters/TextExpressions/SequencerTextFilterExpression_Keyed.h"
+#include "Filters/TextExpressions/SequencerTextFilterExpression_Condition.h"
+#include "Filters/TextExpressions/SequencerTextFilterExpression_ConditionClass.h"
+#include "Filters/TextExpressions/SequencerTextFilterExpression_ConditionFunc.h"
+#include "Filters/TextExpressions/SequencerTextFilterExpression_ConditionPasses.h"
 #include "Filters/TextExpressions/SequencerTextFilterExpression_Class.h"
 #include "Filters/TextExpressions/SequencerTextFilterExpression_Group.h"
 #include "Filters/TextExpressions/SequencerTextFilterExpression_Level.h"
@@ -37,6 +41,10 @@ FSequencerTrackFilter_Text::FSequencerTrackFilter_Text(ISequencerTrackFilters& I
 	// Ordered by importance and most often used. This will dictate the order of display in the text expressions help dialog.
 	TextFilterExpressionContexts.Add(MakeShared<FSequencerTextFilterExpression_Name>(InFilterInterface));
 	TextFilterExpressionContexts.Add(MakeShared<FSequencerTextFilterExpression_Class>(InFilterInterface));
+	TextFilterExpressionContexts.Add(MakeShared<FSequencerTextFilterExpression_Condition>(InFilterInterface));
+	TextFilterExpressionContexts.Add(MakeShared<FSequencerTextFilterExpression_ConditionClass>(InFilterInterface));
+	TextFilterExpressionContexts.Add(MakeShared<FSequencerTextFilterExpression_ConditionFunc>(InFilterInterface));
+	TextFilterExpressionContexts.Add(MakeShared<FSequencerTextFilterExpression_ConditionPasses>(InFilterInterface));
 	TextFilterExpressionContexts.Add(MakeShared<FSequencerTextFilterExpression_Keyed>(InFilterInterface));
 	TextFilterExpressionContexts.Add(MakeShared<FSequencerTextFilterExpression_Selected>(InFilterInterface));
 	TextFilterExpressionContexts.Add(MakeShared<FSequencerTextFilterExpression_Unbound>(InFilterInterface));
diff --git a/Engine/Source/Editor/Sequencer/Private/Filters/SequencerFilterBar.cpp b/Engine/Source/Editor/Sequencer/Private/Filters/SequencerFilterBar.cpp
--- a/Engine/Source/Editor/Sequencer/Private/Filters/SequencerFilterBar.cpp
+++ b/Engine/Source/Editor/Sequencer/Private/Filters/SequencerFilterBar.cpp
@@ -1,8 +1,9 @@
-﻿// Copyright Epic Games, Inc. All Rights Reserved.
+// Copyright Epic Games, Inc. All Rights Reserved.
 
 #include "Filters/SequencerFilterBar.h"
 #include "CurveEditor.h"
 #include "Filters/Filters/SequencerTrackFilter_Keyed.h"
+#include "Filters/Filters/SequencerTrackFilter_Condition.h"
 #include "Filters/Filters/SequencerTrackFilter_CustomText.h"
 #include "Filters/Filters/SequencerTrackFilter_Group.h"
 #include "Filters/Filters/SequencerTrackFilter_HideIsolate.h"
@@ -118,6 +119,7 @@ void FSequencerFilterBar::CreateDefaultFilters()
 	//CommonFilters->Add(ModifiedFilter); // Disabling until clear direction on what this should do
 	CommonFilters->Add(SelectedFilter);
 	CommonFilters->Add(MakeShared<FSequencerTrackFilter_Unbound>(*this, MiscCategory));
+	CommonFilters->Add(MakeShared<FSequencerTrackFilter_Condition>(*this, MiscCategory));
 
 	// Add global user-defined filters
 	for (TObjectIterator<USequencerTrackFilterExtension> ExtensionIt(RF_NoFlags); ExtensionIt; ++ExtensionIt)
diff --git a/Engine/Source/Editor/Sequencer/Private/Filters/SequencerTrackFilterCommands.cpp b/Engine/Source/Editor/Sequencer/Private/Filters/SequencerTrackFilterCommands.cpp
--- a/Engine/Source/Editor/Sequencer/Private/Filters/SequencerTrackFilterCommands.cpp
+++ b/Engine/Source/Editor/Sequencer/Private/Filters/SequencerTrackFilterCommands.cpp
@@ -1,4 +1,4 @@
-﻿// Copyright Epic Games, Inc. All Rights Reserved.
+// Copyright Epic Games, Inc. All Rights Reserved.
 
 #include "SequencerTrackFilterCommands.h"
 #include "Styling/AppStyle.h"
@@ -54,6 +54,7 @@ void FSequencerTrackFilterCommands::RegisterCommands()
 	UI_COMMAND(ToggleFilter_Light, "Toggle Light Filter", "Toggle the filter for Light tracks", EUserInterfaceActionType::Button, FInputChord());
 	UI_COMMAND(ToggleFilter_SkeletalMesh, "Toggle Skeletal Mesh Filter", "Toggle the filter for Skeletal Mesh tracks", EUserInterfaceActionType::Button, FInputChord());
 
+	UI_COMMAND(ToggleFilter_Condition, "Toggle Condition Filter", "Toggle the filter for tracks with conditions", EUserInterfaceActionType::Button, FInputChord());
 	UI_COMMAND(ToggleFilter_Keyed, "Toggle Keyed Filter", "Toggle the filter for Keyed tracks", EUserInterfaceActionType::Button, FInputChord(EKeys::U));
 	UI_COMMAND(ToggleFilter_Modified, "Toggle Modified Filter", "Toggle the filter for Modified tracks", EUserInterfaceActionType::Button, FInputChord());
 	UI_COMMAND(ToggleFilter_Selected, "Toggle Selected In Viewport Filter", "Toggle the filter for Selected In Viewport tracks", EUserInterfaceActionType::Button, FInputChord());
diff --git a/Engine/Source/Editor/Sequencer/Private/Filters/SequencerTrackFilterCommands.h b/Engine/Source/Editor/Sequencer/Private/Filters/SequencerTrackFilterCommands.h
--- a/Engine/Source/Editor/Sequencer/Private/Filters/SequencerTrackFilterCommands.h
+++ b/Engine/Source/Editor/Sequencer/Private/Filters/SequencerTrackFilterCommands.h
@@ -1,4 +1,4 @@
-﻿// Copyright Epic Games, Inc. All Rights Reserved.
+// Copyright Epic Games, Inc. All Rights Reserved.
 
 #pragma once
 
@@ -60,6 +60,7 @@ public:
 	TSharedPtr<FUICommandInfo> ToggleFilter_Light;
 	TSharedPtr<FUICommandInfo> ToggleFilter_SkeletalMesh;
 
+	TSharedPtr<FUICommandInfo> ToggleFilter_Condition;
 	TSharedPtr<FUICommandInfo> ToggleFilter_Keyed;
 	TSharedPtr<FUICommandInfo> ToggleFilter_Modified;
 	TSharedPtr<FUICommandInfo> ToggleFilter_Selected;
diff --git a/Engine/Source/Editor/Sequencer/Private/Filters/TextExpressions/SequencerTextFilterExpression_Condition.cpp b/Engine/Source/Editor/Sequencer/Private/Filters/TextExpressions/SequencerTextFilterExpression_Condition.cpp
new file mode 100644
--- /dev/null
+++ b/Engine/Source/Editor/Sequencer/Private/Filters/TextExpressions/SequencerTextFilterExpression_Condition.cpp
@@ -0,0 +1,59 @@
+// Copyright Epic Games, Inc. All Rights Reserved.
+
+#include "SequencerTextFilterExpression_Condition.h"
+#include "MVVM/Extensions/IConditionableExtension.h"
+#include "MVVM/ViewModels/ViewModelIterators.h"
+#include "MVVM/ViewModelPtr.h"
+#include "MVVM/SharedViewModelData.h"
+
+using namespace UE::Sequencer;
+
+#define LOCTEXT_NAMESPACE "SequencerTextFilterExpression_Condition"
+
+FSequencerTextFilterExpression_Condition::FSequencerTextFilterExpression_Condition(ISequencerTrackFilters& InFilterInterface)
+	: FSequencerTextFilterExpressionContext(InFilterInterface)
+{
+}
+
+TSet<FName> FSequencerTextFilterExpression_Condition::GetKeys() const
+{
+	return { TEXT("CONDITION"), TEXT("HASCONDITION") };
+}
+
+ESequencerTextFilterValueType FSequencerTextFilterExpression_Condition::GetValueType() const
+{
+	return ESequencerTextFilterValueType::Boolean;
+}
+
+FText FSequencerTextFilterExpression_Condition::GetDescription() const
+{
+	return LOCTEXT("ExpressionDescription_Condition", "Filters by the presence of a condition");
+}
+
+bool FSequencerTextFilterExpression_Condition::TestComplexExpression(const FName& InKey
+	, const FTextFilterString& InValue
+	, const ETextFilterComparisonOperation InComparisonOperation
+	, const ETextFilterTextComparisonMode InTextComparisonMode) const
+{
+	if (!FSequencerTextFilterExpressionContext::TestComplexExpression(InKey, InValue, InComparisonOperation, InTextComparisonMode))
+	{
+		return true;
+	}
+
+	const TViewModelPtr<IConditionableExtension> ConditionableExtension = FilterItem->FindAncestorOfType<IConditionableExtension>(true);
+	if (!ConditionableExtension.IsValid())
+	{
+		return false;
+	}
+
+	bool bPassed = ConditionableExtension->GetConditionState() != EConditionableConditionState::None;
+
+	if (TViewModelPtr<FConditionStateCacheExtension> StateCache = CastViewModel<FConditionStateCacheExtension>(FilterItem.AsModel()->GetSharedData()))
+	{
+		bPassed |= EnumHasAnyFlags(StateCache->GetCachedFlags(FilterItem->GetModelID()), ECachedConditionState::ParentHasCondition);
+	}
+
+	return CompareFStringForExactBool(InValue, InComparisonOperation, bPassed);
+}
+
+#undef LOCTEXT_NAMESPACE
diff --git a/Engine/Source/Editor/Sequencer/Private/Filters/TextExpressions/SequencerTextFilterExpression_Condition.h b/Engine/Source/Editor/Sequencer/Private/Filters/TextExpressions/SequencerTextFilterExpression_Condition.h
new file mode 100644
--- /dev/null
+++ b/Engine/Source/Editor/Sequencer/Private/Filters/TextExpressions/SequencerTextFilterExpression_Condition.h
@@ -0,0 +1,24 @@
+// Copyright Epic Games, Inc. All Rights Reserved.
+
+#pragma once
+
+#include "Filters/SequencerTextFilterExpressionContext.h"
+
+class FSequencerTextFilterExpression_Condition : public FSequencerTextFilterExpressionContext
+{
+public:
+	FSequencerTextFilterExpression_Condition(ISequencerTrackFilters& InFilterInterface);
+
+	//~ Begin FSequencerTextFilterExpressionContext
+	virtual TSet<FName> GetKeys() const override;
+	virtual ESequencerTextFilterValueType GetValueType() const override;
+	virtual FText GetDescription() const override;
+	//~ End FSequencerTextFilterExpressionContext
+
+	//~ Begin ITextFilterExpressionContext
+	virtual bool TestComplexExpression(const FName& InKey
+		, const FTextFilterString& InValue
+		, const ETextFilterComparisonOperation InComparisonOperation
+		, const ETextFilterTextComparisonMode InTextComparisonMode) const override;
+	//~ End ITextFilterExpressionContext
+};
diff --git a/Engine/Source/Editor/Sequencer/Private/Filters/TextExpressions/SequencerTextFilterExpression_ConditionClass.cpp b/Engine/Source/Editor/Sequencer/Private/Filters/TextExpressions/SequencerTextFilterExpression_ConditionClass.cpp
new file mode 100644
--- /dev/null
+++ b/Engine/Source/Editor/Sequencer/Private/Filters/TextExpressions/SequencerTextFilterExpression_ConditionClass.cpp
@@ -0,0 +1,62 @@
+// Copyright Epic Games, Inc. All Rights Reserved.
+
+#include "SequencerTextFilterExpression_ConditionClass.h"
+#include "MVVM/Extensions/IConditionableExtension.h"
+#include "MVVM/ViewModels/ViewModelIterators.h"
+#include "MVVM/ViewModelPtr.h"
+#include "MVVM/SharedViewModelData.h"
+
+using namespace UE::Sequencer;
+
+#define LOCTEXT_NAMESPACE "SequencerTextFilterExpression_ConditionClass"
+
+FSequencerTextFilterExpression_ConditionClass::FSequencerTextFilterExpression_ConditionClass(ISequencerTrackFilters& InFilterInterface)
+	: FSequencerTextFilterExpressionContext(InFilterInterface)
+{
+}
+
+TSet<FName> FSequencerTextFilterExpression_ConditionClass::GetKeys() const
+{
+	return { TEXT("CONDITIONCLASS"), TEXT("CONDITIONTYPE") };
+}
+
+ESequencerTextFilterValueType FSequencerTextFilterExpression_ConditionClass::GetValueType() const
+{
+	return ESequencerTextFilterValueType::String;
+}
+
+FText FSequencerTextFilterExpression_ConditionClass::GetDescription() const
+{
+	return LOCTEXT("ExpressionDescription_ConditionClass", "Filters based on the presence of a condition with the given type");
+}
+
+bool FSequencerTextFilterExpression_ConditionClass::TestComplexExpression(const FName& InKey
+	, const FTextFilterString& InValue
+	, const ETextFilterComparisonOperation InComparisonOperation
+	, const ETextFilterTextComparisonMode InTextComparisonMode) const
+{
+	if (!FSequencerTextFilterExpressionContext::TestComplexExpression(InKey, InValue, InComparisonOperation, InTextComparisonMode))
+	{
+		return true;
+	}
+
+	const TViewModelPtr<IConditionableExtension> ConditionableExtension = FilterItem->FindAncestorOfType<IConditionableExtension>(true);
+	if (!ConditionableExtension.IsValid())
+	{
+		return false;
+	}
+
+	if (ConditionableExtension->GetConditionState() == EConditionableConditionState::None)
+	{
+		return false;
+	}
+
+	if (const UMovieSceneCondition* Condition = ConditionableExtension->GetCondition())
+	{
+		return TextFilterUtils::TestComplexExpression(Condition->GetClass()->GetName(), InValue, InComparisonOperation, InTextComparisonMode);
+	}
+
+	return false;
+}
+
+#undef LOCTEXT_NAMESPACE
diff --git a/Engine/Source/Editor/Sequencer/Private/Filters/TextExpressions/SequencerTextFilterExpression_ConditionClass.h b/Engine/Source/Editor/Sequencer/Private/Filters/TextExpressions/SequencerTextFilterExpression_ConditionClass.h
new file mode 100644
--- /dev/null
+++ b/Engine/Source/Editor/Sequencer/Private/Filters/TextExpressions/SequencerTextFilterExpression_ConditionClass.h
@@ -0,0 +1,24 @@
+// Copyright Epic Games, Inc. All Rights Reserved.
+
+#pragma once
+
+#include "Filters/SequencerTextFilterExpressionContext.h"
+
+class FSequencerTextFilterExpression_ConditionClass : public FSequencerTextFilterExpressionContext
+{
+public:
+	FSequencerTextFilterExpression_ConditionClass(ISequencerTrackFilters& InFilterInterface);
+
+	//~ Begin FSequencerTextFilterExpressionContext
+	virtual TSet<FName> GetKeys() const override;
+	virtual ESequencerTextFilterValueType GetValueType() const override;
+	virtual FText GetDescription() const override;
+	//~ End FSequencerTextFilterExpressionContext
+
+	//~ Begin ITextFilterExpressionContext
+	virtual bool TestComplexExpression(const FName& InKey
+		, const FTextFilterString& InValue
+		, const ETextFilterComparisonOperation InComparisonOperation
+		, const ETextFilterTextComparisonMode InTextComparisonMode) const override;
+	//~ End ITextFilterExpressionContext
+};
diff --git a/Engine/Source/Editor/Sequencer/Private/Filters/TextExpressions/SequencerTextFilterExpression_ConditionFunc.cpp b/Engine/Source/Editor/Sequencer/Private/Filters/TextExpressions/SequencerTextFilterExpression_ConditionFunc.cpp
new file mode 100644
--- /dev/null
+++ b/Engine/Source/Editor/Sequencer/Private/Filters/TextExpressions/SequencerTextFilterExpression_ConditionFunc.cpp
@@ -0,0 +1,60 @@
+// Copyright Epic Games, Inc. All Rights Reserved.
+
+#include "SequencerTextFilterExpression_ConditionFunc.h"
+#include "MVVM/Extensions/IConditionableExtension.h"
+#include "MVVM/ViewModels/ViewModelIterators.h"
+#include "MVVM/ViewModelPtr.h"
+#include "MVVM/SharedViewModelData.h"
+#include "Conditions/MovieSceneDirectorBlueprintCondition.h"
+
+using namespace UE::Sequencer;
+
+#define LOCTEXT_NAMESPACE "SequencerTextFilterExpression_ConditionFunc"
+
+FSequencerTextFilterExpression_ConditionFunc::FSequencerTextFilterExpression_ConditionFunc(ISequencerTrackFilters& InFilterInterface)
+	: FSequencerTextFilterExpressionContext(InFilterInterface)
+{
+}
+
+TSet<FName> FSequencerTextFilterExpression_ConditionFunc::GetKeys() const
+{
+	return { TEXT("CONDITIONFUNC"), TEXT("CONDITIONENDPOINT") };
+}
+
+ESequencerTextFilterValueType FSequencerTextFilterExpression_ConditionFunc::GetValueType() const
+{
+	return ESequencerTextFilterValueType::String;
+}
+
+FText FSequencerTextFilterExpression_ConditionFunc::GetDescription() const
+{
+	return LOCTEXT("ExpressionDescription_ConditionFunc", "Filters by the presence of a director blueprint condition with the given function/endponit name");
+}
+
+bool FSequencerTextFilterExpression_ConditionFunc::TestComplexExpression(const FName& InKey
+	, const FTextFilterString& InValue
+	, const ETextFilterComparisonOperation InComparisonOperation
+	, const ETextFilterTextComparisonMode InTextComparisonMode) const
+{
+	if (!FSequencerTextFilterExpressionContext::TestComplexExpression(InKey, InValue, InComparisonOperation, InTextComparisonMode))
+	{
+		return true;
+	}
+
+	const TViewModelPtr<IConditionableExtension> ConditionableExtension = FilterItem->FindAncestorOfType<IConditionableExtension>(true);
+	if (!ConditionableExtension.IsValid())
+	{
+		return false;
+	}
+
+	if (const UMovieSceneDirectorBlueprintCondition* Condition = Cast<UMovieSceneDirectorBlueprintCondition>(ConditionableExtension->GetCondition()))
+	{
+		if (Condition->DirectorBlueprintConditionData.Function && Condition->DirectorBlueprintConditionData.Function->Next)
+		{
+			return TextFilterUtils::TestComplexExpression(Condition->DirectorBlueprintConditionData.Function->Next->GetName(), InValue, InComparisonOperation, InTextComparisonMode);
+		}
+	}
+	return false;
+}
+
+#undef LOCTEXT_NAMESPACE
diff --git a/Engine/Source/Editor/Sequencer/Private/Filters/TextExpressions/SequencerTextFilterExpression_ConditionFunc.h b/Engine/Source/Editor/Sequencer/Private/Filters/TextExpressions/SequencerTextFilterExpression_ConditionFunc.h
new file mode 100644
--- /dev/null
+++ b/Engine/Source/Editor/Sequencer/Private/Filters/TextExpressions/SequencerTextFilterExpression_ConditionFunc.h
@@ -0,0 +1,24 @@
+// Copyright Epic Games, Inc. All Rights Reserved.
+
+#pragma once
+
+#include "Filters/SequencerTextFilterExpressionContext.h"
+
+class FSequencerTextFilterExpression_ConditionFunc : public FSequencerTextFilterExpressionContext
+{
+public:
+	FSequencerTextFilterExpression_ConditionFunc(ISequencerTrackFilters& InFilterInterface);
+
+	//~ Begin FSequencerTextFilterExpressionContext
+	virtual TSet<FName> GetKeys() const override;
+	virtual ESequencerTextFilterValueType GetValueType() const override;
+	virtual FText GetDescription() const override;
+	//~ End FSequencerTextFilterExpressionContext
+
+	//~ Begin ITextFilterExpressionContext
+	virtual bool TestComplexExpression(const FName& InKey
+		, const FTextFilterString& InValue
+		, const ETextFilterComparisonOperation InComparisonOperation
+		, const ETextFilterTextComparisonMode InTextComparisonMode) const override;
+	//~ End ITextFilterExpressionContext
+};
diff --git a/Engine/Source/Editor/Sequencer/Private/Filters/TextExpressions/SequencerTextFilterExpression_ConditionPasses.cpp b/Engine/Source/Editor/Sequencer/Private/Filters/TextExpressions/SequencerTextFilterExpression_ConditionPasses.cpp
new file mode 100644
--- /dev/null
+++ b/Engine/Source/Editor/Sequencer/Private/Filters/TextExpressions/SequencerTextFilterExpression_ConditionPasses.cpp
@@ -0,0 +1,65 @@
+// Copyright Epic Games, Inc. All Rights Reserved.
+
+#include "SequencerTextFilterExpression_ConditionPasses.h"
+#include "MVVM/Extensions/IConditionableExtension.h"
+#include "MVVM/ViewModels/ViewModelIterators.h"
+#include "MVVM/ViewModelPtr.h"
+#include "MVVM/SharedViewModelData.h"
+
+using namespace UE::Sequencer;
+
+#define LOCTEXT_NAMESPACE "SequencerTextFilterExpression_Condition"
+
+FSequencerTextFilterExpression_ConditionPasses::FSequencerTextFilterExpression_ConditionPasses(ISequencerTrackFilters& InFilterInterface)
+	: FSequencerTextFilterExpressionContext(InFilterInterface)
+{
+}
+
+TSet<FName> FSequencerTextFilterExpression_ConditionPasses::GetKeys() const
+{
+	return { TEXT("CONDITIONPASSES"), TEXT("CONDITIONEVALUATES") };
+}
+
+ESequencerTextFilterValueType FSequencerTextFilterExpression_ConditionPasses::GetValueType() const
+{
+	return ESequencerTextFilterValueType::Boolean;
+}
+
+FText FSequencerTextFilterExpression_ConditionPasses::GetDescription() const
+{
+	return LOCTEXT("ExpressionDescription_ConditionPasses", "Filters based on the presence of a condition that is passing/failing");
+}
+
+bool FSequencerTextFilterExpression_ConditionPasses::TestComplexExpression(const FName& InKey
+	, const FTextFilterString& InValue
+	, const ETextFilterComparisonOperation InComparisonOperation
+	, const ETextFilterTextComparisonMode InTextComparisonMode) const
+{
+	if (!FSequencerTextFilterExpressionContext::TestComplexExpression(InKey, InValue, InComparisonOperation, InTextComparisonMode))
+	{
+		return true;
+	}
+
+	const TViewModelPtr<IConditionableExtension> ConditionableExtension = FilterItem->FindAncestorOfType<IConditionableExtension>(true);
+	if (!ConditionableExtension.IsValid())
+	{
+		return false;
+	}
+
+	if (ConditionableExtension->GetConditionState() == EConditionableConditionState::None)
+	{
+		return false;
+	}
+
+	if (TViewModelPtr<FConditionStateCacheExtension> StateCache = CastViewModel<FConditionStateCacheExtension>(FilterItem.AsModel()->GetSharedData()))
+	{
+		bool bPassed = EnumHasAnyFlags(StateCache->GetCachedFlags(FilterItem->GetModelID()), ECachedConditionState::ConditionEvaluatingTrue)
+						|| EnumHasAnyFlags(StateCache->GetCachedFlags(FilterItem->GetModelID()), ECachedConditionState::ParentHasConditionEvaluatingTrue);
+		
+		return CompareFStringForExactBool(InValue, InComparisonOperation, bPassed);
+	}
+
+	return false;
+}
+
+#undef LOCTEXT_NAMESPACE
diff --git a/Engine/Source/Editor/Sequencer/Private/Filters/TextExpressions/SequencerTextFilterExpression_ConditionPasses.h b/Engine/Source/Editor/Sequencer/Private/Filters/TextExpressions/SequencerTextFilterExpression_ConditionPasses.h
new file mode 100644
--- /dev/null
+++ b/Engine/Source/Editor/Sequencer/Private/Filters/TextExpressions/SequencerTextFilterExpression_ConditionPasses.h
@@ -0,0 +1,24 @@
+// Copyright Epic Games, Inc. All Rights Reserved.
+
+#pragma once
+
+#include "Filters/SequencerTextFilterExpressionContext.h"
+
+class FSequencerTextFilterExpression_ConditionPasses : public FSequencerTextFilterExpressionContext
+{
+public:
+	FSequencerTextFilterExpression_ConditionPasses(ISequencerTrackFilters& InFilterInterface);
+
+	//~ Begin FSequencerTextFilterExpressionContext
+	virtual TSet<FName> GetKeys() const override;
+	virtual ESequencerTextFilterValueType GetValueType() const override;
+	virtual FText GetDescription() const override;
+	//~ End FSequencerTextFilterExpressionContext
+
+	//~ Begin ITextFilterExpressionContext
+	virtual bool TestComplexExpression(const FName& InKey
+		, const FTextFilterString& InValue
+		, const ETextFilterComparisonOperation InComparisonOperation
+		, const ETextFilterTextComparisonMode InTextComparisonMode) const override;
+	//~ End ITextFilterExpressionContext
+};
