Sequencer: Fix mac os build.

diff --git a/Engine/Source/Editor/Sequencer/Private/Widgets/OutlinerColumns/SConditionColumnWidget.h b/Engine/Source/Editor/Sequencer/Private/Widgets/OutlinerColumns/SConditionColumnWidget.h
--- a/Engine/Source/Editor/Sequencer/Private/Widgets/OutlinerColumns/SConditionColumnWidget.h
+++ b/Engine/Source/Editor/Sequencer/Private/Widgets/OutlinerColumns/SConditionColumnWidget.h
@@ -47,7 +47,7 @@ namespace UE::Sequencer
 	protected:
 
 		/** Handles left click-type inputs and potentially begins drag and drop operation. */
-		FReply HandleClick() { /* no - op for now, no toggle*/ }
+		FReply HandleClick() { return FReply::Handled(); }
 
 	private:
 
