[Sequencer]
 - Remove menu actions from sidebar
 - Add sidebar section details for FTrackRow
 - Rebuild sidebar after overriding channels

diff --git a/Engine/Plugins/MovieScene/LevelSequenceEditor/Source/LevelSequenceEditor/Private/LevelSequenceEditorSubsystem.cpp b/Engine/Plugins/MovieScene/LevelSequenceEditor/Source/LevelSequenceEditor/Private/LevelSequenceEditorSubsystem.cpp
--- a/Engine/Plugins/MovieScene/LevelSequenceEditor/Source/LevelSequenceEditor/Private/LevelSequenceEditorSubsystem.cpp
+++ b/Engine/Plugins/MovieScene/LevelSequenceEditor/Source/LevelSequenceEditor/Private/LevelSequenceEditorSubsystem.cpp
@@ -630,12 +630,6 @@ void ULevelSequenceEditorSubsystem::Initialize(FSubsystemCollectionBase& Collect
 				}
 
 				AddBindingPropertiesSidebar(MenuBuilder);
-
-				FFormatNamedArguments Args;
-				MenuBuilder.AddSubMenu(
-					FText::Format(LOCTEXT("RebindComponent", "Rebind Component"), Args),
-					FText::Format(LOCTEXT("RebindComponentTooltip", "Rebind component by moving the tracks from one component to another component."), Args),
-					FNewMenuDelegate::CreateUObject(this, &ULevelSequenceEditorSubsystem::RebindComponentMenu));
 			}));
 
 	SidebarMenuExtender->AddMenuExtension(TEXT("CustomBinding"), EExtensionHook::First, CommandList,
diff --git a/Engine/Source/Developer/ToolWidgets/Private/Sidebar/SSidebar.cpp b/Engine/Source/Developer/ToolWidgets/Private/Sidebar/SSidebar.cpp
--- a/Engine/Source/Developer/ToolWidgets/Private/Sidebar/SSidebar.cpp
+++ b/Engine/Source/Developer/ToolWidgets/Private/Sidebar/SSidebar.cpp
@@ -71,8 +71,8 @@ bool SSidebar::RegisterDrawer(FSidebarDrawerConfig&& InDrawerConfig)
 	NewDrawer->bDisablePin = bDisablePin;
 	NewDrawer->bDisableDock = bDisableDock;
 	NewDrawer->ContentWidget = NewDrawer->Config.OverrideContentWidget.IsValid()
-		? NewDrawer->ContentWidget = NewDrawer->Config.OverrideContentWidget
-		: NewDrawer->ContentWidget = SNew(SSidebarDrawerContent, NewDrawer);
+		? NewDrawer->Config.OverrideContentWidget
+		: SNew(SSidebarDrawerContent, NewDrawer);
 
 	// Add tab button
 	TabButtonContainer->AddSlot()
@@ -639,6 +639,28 @@ TSharedRef<SWidget> SSidebar::GetMainContent() const
 	return OnGetContent.IsBound() ? OnGetContent.Execute() : SNullWidget::NullWidget;
 }
 
+void SSidebar::RebuildDrawer(const FName InDrawerId, const bool bInOnlyIfOpen)
+{
+	const TSharedPtr<FSidebarDrawer> Drawer = FindDrawer(InDrawerId);
+	if (!Drawer.IsValid() || Drawer->Config.OverrideContentWidget.IsValid())
+	{
+		return;
+	}
+
+	if (bInOnlyIfOpen && !IsDrawerOpened(InDrawerId))
+	{
+		return;
+	}
+
+	const TSharedPtr<SSidebarDrawerContent> DrawerContent = StaticCastSharedPtr<SSidebarDrawerContent>(Drawer->ContentWidget);
+	if (!DrawerContent.IsValid())
+	{
+		return;
+	}
+
+	DrawerContent->BuildContent();
+}
+
 const TArray<TSharedRef<FSidebarDrawer>>& SSidebar::GetAllDrawers() const
 {
 	return Drawers;
diff --git a/Engine/Source/Developer/ToolWidgets/Public/Sidebar/SSidebar.h b/Engine/Source/Developer/ToolWidgets/Public/Sidebar/SSidebar.h
--- a/Engine/Source/Developer/ToolWidgets/Public/Sidebar/SSidebar.h
+++ b/Engine/Source/Developer/ToolWidgets/Public/Sidebar/SSidebar.h
@@ -238,6 +238,9 @@ public:
 
 	TSharedRef<SWidget> GetMainContent() const;
 
+	/** Rebuilds a drawer's content (unless it has been overriden). */
+	void RebuildDrawer(const FName InDrawerId, const bool bInOnlyIfOpen = true);
+
 private:
 	void OnTabDrawerButtonPressed(const TSharedRef<FSidebarDrawer>& InDrawer);
 	void OnDrawerTabPinToggled(const TSharedRef<FSidebarDrawer>& InDrawer, const bool bIsPinned);
diff --git a/Engine/Source/Editor/Sequencer/Private/Menus/CurveChannelSectionSidebarExtension.cpp b/Engine/Source/Editor/Sequencer/Private/Menus/CurveChannelSectionSidebarExtension.cpp
--- a/Engine/Source/Editor/Sequencer/Private/Menus/CurveChannelSectionSidebarExtension.cpp
+++ b/Engine/Source/Editor/Sequencer/Private/Menus/CurveChannelSectionSidebarExtension.cpp
@@ -43,10 +43,6 @@ void FCurveChannelSectionSidebarExtension::AddSections(const TArray<TWeakObjectP
 
 TSharedPtr<ISidebarChannelExtension> FCurveChannelSectionSidebarExtension::ExtendMenu(FMenuBuilder& MenuBuilder, const bool bInSubMenu)
 {
-	AddExtrapolationMenu(MenuBuilder, true);
-	AddExtrapolationMenu(MenuBuilder, false);
-	AddDisplayOptionsMenu(MenuBuilder);
-
 	return SharedThis(this);
 }
 
diff --git a/Engine/Source/Editor/Sequencer/Private/SSequencer.cpp b/Engine/Source/Editor/Sequencer/Private/SSequencer.cpp
--- a/Engine/Source/Editor/Sequencer/Private/SSequencer.cpp
+++ b/Engine/Source/Editor/Sequencer/Private/SSequencer.cpp
@@ -4173,6 +4173,11 @@ void SSequencer::ToggleSidebarVisible()
 
 void SSequencer::ToggleSidebarSelectionDrawerOpen()
 {
+	if (!DetailsSidebar.IsValid())
+	{
+		return;
+	}
+
 	if (DetailsSidebar->IsDrawerDocked(FSequencer::SelectionDrawerId))
 	{
 		DetailsSidebar->UndockAllDrawers();
@@ -4189,6 +4194,11 @@ void SSequencer::ToggleSidebarSelectionDrawerOpen()
 
 void SSequencer::ToggleSidebarDrawerDock()
 {
+	if (!DetailsSidebar.IsValid())
+	{
+		return;
+	}
+
 	const FName OpenedDrawerId = DetailsSidebar->GetOpenedDrawerId();
 	if (!OpenedDrawerId.IsNone())
 	{
diff --git a/Engine/Source/Editor/Sequencer/Private/SequencerCommonHelpers.cpp b/Engine/Source/Editor/Sequencer/Private/SequencerCommonHelpers.cpp
--- a/Engine/Source/Editor/Sequencer/Private/SequencerCommonHelpers.cpp
+++ b/Engine/Source/Editor/Sequencer/Private/SequencerCommonHelpers.cpp
@@ -427,7 +427,7 @@ void SequencerHelpers::BuildEditSectionMenu(const TWeakPtr<FSequencer>& InWeakSe
 		DetailsNotifyWrapper->SetDetailsAndSequencer(DetailsView, InWeakSequencer);
 		DetailsNotifyWrapper->SetEnabled(!Sequencer->IsReadOnly());
 
-		LambdaMenuBuilder.BeginSection(TEXT("TrackSection"), NSLOCTEXT("Sequencer", "TrackSectionMenuSection", "Track Section"));
+		LambdaMenuBuilder.BeginSection(TEXT("TrackSection"));
         {
 			LambdaMenuBuilder.AddWidget(DetailsNotifyWrapper, FText::GetEmpty(), true);
         }
@@ -515,7 +515,7 @@ void SequencerHelpers::BuildEditTrackMenu(const TWeakPtr<FSequencer>& InWeakSequ
 		DetailsNotifyWrapper->SetDetailsAndSequencer(DetailsView, Sequencer);
 		DetailsNotifyWrapper->SetEnabled(!Sequencer->IsReadOnly());
 
-		LambdaMenuBuilder.BeginSection(TEXT("Track"), NSLOCTEXT("Sequencer", "TrackMenuSection", "Track"));
+		LambdaMenuBuilder.BeginSection(TEXT("Track"));
 		{
 			LambdaMenuBuilder.AddWidget(DetailsNotifyWrapper, FText::GetEmpty(), true);
 		}
diff --git a/Engine/Source/Editor/Sequencer/Private/Widgets/Sidebar/SMarkedFrameDetails.h b/Engine/Source/Editor/Sequencer/Private/Widgets/Sidebar/SMarkedFrameDetails.h
new file mode 100644
--- /dev/null
+++ b/Engine/Source/Editor/Sequencer/Private/Widgets/Sidebar/SMarkedFrameDetails.h
@@ -0,0 +1,136 @@
+﻿// Copyright Epic Games, Inc. All Rights Reserved.
+
+#pragma once
+
+#include "FrameNumberDetailsCustomization.h"
+#include "IStructureDetailsView.h"
+#include "Misc/NotifyHook.h"
+#include "MovieScene.h"
+#include "MovieSceneMarkedFrame.h"
+#include "PropertyEditorModule.h"
+#include "Sequencer.h"
+#include "Widgets/SCompoundWidget.h"
+
+class SMarkedFrameDetails : public SCompoundWidget, public FNotifyHook
+{
+public:
+	SLATE_BEGIN_ARGS(SMarkedFrameDetails)
+	{}
+	SLATE_END_ARGS()
+
+	void Construct(const FArguments& InArgs, const int32 InMarkedFrameIndex, const TWeakPtr<FSequencer>& InWeakSequencer)
+	{
+		WeakSequencer = InWeakSequencer;
+
+		const TSharedPtr<FSequencer> Sequencer = WeakSequencer.Pin();
+		if (!Sequencer.IsValid())
+		{
+			return;
+		}
+
+		const UMovieSceneSequence* const FocusedMovieSceneSequence = Sequencer->GetFocusedMovieSceneSequence();
+		if (!IsValid(FocusedMovieSceneSequence))
+		{
+			return;
+		}
+
+		UMovieScene* const FocusedMovieScene = FocusedMovieSceneSequence->GetMovieScene();
+		if (!IsValid(FocusedMovieScene))
+		{
+			return;
+		}
+
+		if (FocusedMovieScene->GetMarkedFrames().Num() == 0)
+		{
+			return;
+		}
+
+		WeakMovieSceneToModify = FocusedMovieScene;
+
+		FDetailsViewArgs DetailsViewArgs;
+		DetailsViewArgs.bAllowSearch = false;
+		DetailsViewArgs.bShowScrollBar = false;
+		DetailsViewArgs.NameAreaSettings = FDetailsViewArgs::HideNameArea;
+		DetailsViewArgs.NotifyHook = this;
+
+		FStructureDetailsViewArgs StructureDetailsViewArgs;
+		StructureDetailsViewArgs.bShowObjects = true;
+		StructureDetailsViewArgs.bShowAssets = true;
+		StructureDetailsViewArgs.bShowClasses = true;
+		StructureDetailsViewArgs.bShowInterfaces = true;
+		
+		const TSharedPtr<FStructOnScope> StructOnScope = MakeShared<FStructOnScope>(FMovieSceneMarkedFrame::StaticStruct()
+			, (uint8*)&FocusedMovieScene->GetMarkedFrames()[InMarkedFrameIndex]);
+
+		FPropertyEditorModule& PropertyEditorModule = FModuleManager::GetModuleChecked<FPropertyEditorModule>(TEXT("PropertyEditor"));
+
+		DetailsView = PropertyEditorModule.CreateStructureDetailView(DetailsViewArgs, StructureDetailsViewArgs, nullptr);
+
+		DetailsView->GetDetailsView()->RegisterInstancedCustomPropertyTypeLayout(TEXT("FrameNumber"),
+			FOnGetPropertyTypeCustomizationInstance::CreateLambda([this]()
+			{
+				return MakeShared<FFrameNumberDetailsCustomization>(WeakSequencer.Pin()->GetNumericTypeInterface());
+			}));
+
+		DetailsView->SetStructureData(StructOnScope);
+
+		ChildSlot
+		[
+			DetailsView->GetWidget().ToSharedRef()
+		];
+
+		SetEnabled(!AreMarkedFramesLocked());
+	}
+
+	//~ Begin FNotifyHook
+
+	virtual void NotifyPreChange(FProperty* InPropertyAboutToChange) override
+	{
+		if (WeakMovieSceneToModify.IsValid())
+		{
+			WeakMovieSceneToModify->Modify();
+		}
+	}
+
+	virtual void NotifyPreChange(FEditPropertyChain* InPropertyAboutToChange) override
+	{
+		if (WeakMovieSceneToModify.IsValid())
+		{
+			WeakMovieSceneToModify->Modify();
+		}
+	}
+
+	//~ End FNotifyHook
+
+	bool AreMarkedFramesLocked() const
+	{
+		const TSharedPtr<FSequencer> Sequencer = WeakSequencer.Pin();
+		if (!Sequencer.IsValid())
+		{
+			return false;
+		}
+
+		if (Sequencer->IsReadOnly())
+		{
+			return true;
+		}
+
+		const UMovieSceneSequence* const FocusedMovieSceneSequence = Sequencer->GetFocusedMovieSceneSequence();
+		if (IsValid(FocusedMovieSceneSequence))
+		{
+			const UMovieScene* const MovieScene = FocusedMovieSceneSequence->GetMovieScene();
+			if (IsValid(MovieScene))
+			{
+				return MovieScene->IsReadOnly() ? true : MovieScene->AreMarkedFramesLocked();
+			}
+		}
+
+		return false;
+	};
+
+private:
+	TWeakObjectPtr<UMovieScene> WeakMovieSceneToModify;
+	TWeakPtr<FSequencer> WeakSequencer;
+
+	TSharedPtr<IStructureDetailsView> DetailsView;
+};
diff --git a/Engine/Source/Editor/Sequencer/Private/Widgets/Sidebar/STrackDetails.h b/Engine/Source/Editor/Sequencer/Private/Widgets/Sidebar/STrackDetails.h
new file mode 100644
--- /dev/null
+++ b/Engine/Source/Editor/Sequencer/Private/Widgets/Sidebar/STrackDetails.h
@@ -0,0 +1,140 @@
+﻿// Copyright Epic Games, Inc. All Rights Reserved.
+
+#pragma once
+
+#include "IDetailsView.h"
+#include "Misc/NotifyHook.h"
+#include "PropertyEditorDelegates.h"
+#include "PropertyEditorModule.h"
+#include "PropertyPermissionList.h"
+#include "Sequencer.h"
+#include "Widgets/SCompoundWidget.h"
+
+struct FTrackDetailsWidgetCustomizations
+{
+	TMap<UStruct*, FOnGetDetailCustomizationInstance> DetailCustomizationInstances;
+	TMap<FName, FOnGetPropertyTypeCustomizationInstance> PropertyTypeCustomizationInstances;
+};
+
+class STrackDetails : public SCompoundWidget, public FNotifyHook
+{
+public:
+	SLATE_BEGIN_ARGS(STrackDetails)
+	{}
+		SLATE_ARGUMENT(TOptional<FDetailsViewArgs>, ViewArgs)
+		SLATE_ARGUMENT(FTrackDetailsWidgetCustomizations, Customizations)
+		SLATE_ARGUMENT(bool, NotifyMovieSceneDataChanged)
+	SLATE_END_ARGS()
+
+	void Construct(const FArguments& InArgs, const TArray<TWeakObjectPtr<>>& InWeakObjects, const TWeakPtr<FSequencer>& InWeakSequencer)
+	{
+		WeakObjectsToModify = InWeakObjects;
+		WeakSequencer = InWeakSequencer;
+
+		bNotifyMovieSceneDataChanged = InArgs._NotifyMovieSceneDataChanged;
+
+		const TSharedPtr<FSequencer> Sequencer = InWeakSequencer.Pin();
+		check(Sequencer.IsValid());
+
+		FPropertyEditorModule& PropertyEditorModule = FModuleManager::GetModuleChecked<FPropertyEditorModule>(TEXT("PropertyEditor"));
+
+		FDetailsViewArgs DetailsViewArgs;
+		if (InArgs._ViewArgs.IsSet())
+		{
+			DetailsViewArgs = InArgs._ViewArgs.GetValue();
+		}
+		else
+		{
+			DetailsViewArgs.bAllowSearch = false;
+			DetailsViewArgs.bCustomFilterAreaLocation = true;
+			DetailsViewArgs.bCustomNameAreaLocation = true;
+			DetailsViewArgs.bHideSelectionTip = true;
+			DetailsViewArgs.bLockable = false;
+			DetailsViewArgs.bSearchInitialKeyFocus = true;
+			DetailsViewArgs.bUpdatesFromSelection = false;
+			DetailsViewArgs.bShowOptions = false;
+			DetailsViewArgs.bShowModifiedPropertiesOption = false;
+			DetailsViewArgs.bShowScrollBar = false;
+			DetailsViewArgs.NotifyHook = this;
+			DetailsViewArgs.ColumnWidth = 0.45f;
+		}
+
+		DetailsView = PropertyEditorModule.CreateDetailView(DetailsViewArgs);
+
+		for (const TPair<UStruct*, FOnGetDetailCustomizationInstance>& Instance : InArgs._Customizations.DetailCustomizationInstances)
+		{
+			DetailsView->RegisterInstancedCustomPropertyLayout(Instance.Key, Instance.Value);
+		}
+		for (const TPair<FName, FOnGetPropertyTypeCustomizationInstance>& Instance : InArgs._Customizations.PropertyTypeCustomizationInstances)
+		{
+			DetailsView->RegisterInstancedCustomPropertyTypeLayout(Instance.Key, Instance.Value);
+		}
+
+		DetailsView->SetIsPropertyVisibleDelegate(
+			FIsPropertyVisible::CreateLambda([](const FPropertyAndParent& PropertyAndParent)
+			{			
+				return FPropertyEditorPermissionList::Get().DoesPropertyPassFilter(PropertyAndParent.Property.GetOwnerStruct(), PropertyAndParent.Property.GetFName());
+			}));
+
+		Sequencer->OnInitializeDetailsPanel().Broadcast(DetailsView.ToSharedRef(), Sequencer.ToSharedRef());
+		DetailsView->SetObjects(InWeakObjects);
+
+		ChildSlot
+		[
+			DetailsView.ToSharedRef()
+		];
+
+		SetEnabled(!Sequencer->IsReadOnly());
+	}
+
+	//~ Begin FNotifyHook
+
+	virtual void NotifyPreChange(FProperty* const InPropertyAboutToChange) override
+	{
+		ModifyObjects();
+	}
+
+	virtual void NotifyPreChange(FEditPropertyChain* const InPropertyAboutToChange) override
+	{
+		ModifyObjects();
+	}
+
+	virtual void NotifyPostChange(const FPropertyChangedEvent& InPropertyChangedEvent, FProperty* const InPropertyThatChanged) override
+	{
+		if (!bNotifyMovieSceneDataChanged)
+		{
+			return;
+		}
+
+		if (const TSharedPtr<ISequencer> Sequencer = WeakSequencer.Pin())
+		{
+			Sequencer->NotifyMovieSceneDataChanged(EMovieSceneDataChangeType::TrackValueChanged);
+		}
+	}
+
+	//~ End FNotifyHook
+
+	TSharedPtr<IDetailsView> GetDetailsView() const
+	{
+		return DetailsView;
+	}
+
+protected:
+	void ModifyObjects()
+	{
+		for (const TWeakObjectPtr<>& WeakObject : WeakObjectsToModify)
+		{
+			if (WeakObject.IsValid())
+			{
+				WeakObject->Modify();
+			}
+		}
+	}
+
+	TArray<TWeakObjectPtr<>> WeakObjectsToModify;
+	TWeakPtr<FSequencer> WeakSequencer;
+
+	bool bNotifyMovieSceneDataChanged = false;
+
+	TSharedPtr<IDetailsView> DetailsView;
+};
diff --git a/Engine/Source/Editor/Sequencer/Private/Widgets/Sidebar/SequencerSelectionDrawer.cpp b/Engine/Source/Editor/Sequencer/Private/Widgets/Sidebar/SequencerSelectionDrawer.cpp
--- a/Engine/Source/Editor/Sequencer/Private/Widgets/Sidebar/SequencerSelectionDrawer.cpp
+++ b/Engine/Source/Editor/Sequencer/Private/Widgets/Sidebar/SequencerSelectionDrawer.cpp
@@ -2,31 +2,31 @@
 
 #include "SequencerSelectionDrawer.h"
 #include "DetailsViewArgs.h"
-#include "FrameNumberDetailsCustomization.h"
 #include "Framework/MultiBox/MultiBoxBuilder.h"
 #include "IKeyArea.h"
 #include "ISequencerSection.h"
-#include "IStructureDetailsView.h"
 #include "Menus/CurveChannelSectionSidebarExtension.h"
 #include "Modules/ModuleManager.h"
 #include "MovieScene.h"
-#include "MovieSceneMarkedFrame.h"
-#include "MovieSceneSequence.h"
+#include "MovieSceneFolder.h"
 #include "MVVM/Extensions/ITrackExtension.h"
 #include "MVVM/Selection/Selection.h"
 #include "MVVM/ViewModelPtr.h"
 #include "MVVM/ViewModels/CategoryModel.h"
-#include "MVVM/ViewModels/LayerBarModel.h"
+#include "MVVM/ViewModels/FolderModel.h"
 #include "MVVM/ViewModels/SectionModel.h"
 #include "MVVM/ViewModels/SequencerEditorViewModel.h"
+#include "MVVM/ViewModels/TrackRowModel.h"
 #include "PropertyEditorModule.h"
 #include "Sequencer.h"
 #include "SequencerCommonHelpers.h"
-#include "SequencerContextMenus.h"
+#include "SequencerUtilities.h"
 #include "SKeyEditInterface.h"
 #include "Styling/AppStyle.h"
 #include "Templates/SharedPointer.h"
 #include "Widgets/Layout/SBorder.h"
+#include "Widgets/Sidebar/SMarkedFrameDetails.h"
+#include "Widgets/Sidebar/STrackDetails.h"
 
 #define LOCTEXT_NAMESPACE "SequencerSelectionDrawer"
 
@@ -102,6 +102,15 @@ FSequencerSelectionDrawer::FSequencerSelectionDrawer(const TWeakPtr<FSequencer>&
 {
 }
 
+FSequencerSelectionDrawer::~FSequencerSelectionDrawer()
+{
+	if (const TSharedPtr<ISequencer> Sequencer = WeakSequencer.Pin())
+	{
+		Sequencer->OnActorAddedToSequencer().RemoveAll(this);
+		Sequencer->OnMovieSceneDataChanged().RemoveAll(this);
+	}
+}
+
 FName FSequencerSelectionDrawer::GetUniqueId() const
 {
 	return UniqueId;
@@ -126,17 +135,25 @@ TSharedRef<SWidget> FSequencerSelectionDrawer::CreateContentWidget()
 				OnSequencerSelectionChanged();
 			});
 
+		Sequencer->OnMovieSceneDataChanged().AddLambda([this](const EMovieSceneDataChangeType InChangeType)
+			{
+				if (InChangeType == EMovieSceneDataChangeType::ActiveMovieSceneChanged)
+				{
+					OnSequencerSelectionChanged();
+				}
+			});
+
+		Sequencer->OnCloseEvent().AddLambda([this](const TSharedRef<ISequencer> InSequencer)
+			{
+				ResetContent();
+			});
+
 		if (const TSharedPtr<FSequencerSelection> SequencerSelection = Private::GetSelection(*Sequencer.Get()))
 		{
 			SequencerSelection->OnChanged.AddSP(this, &FSequencerSelectionDrawer::OnSequencerSelectionChanged);
 
 			OnSequencerSelectionChanged();
 		}
-
-		Sequencer->OnCloseEvent().AddLambda([this](const TSharedRef<ISequencer> InSequencer)
-			{
-				ResetContent();
-			});
 	}
 
 	return SNew(SBorder)
@@ -271,18 +288,6 @@ void FSequencerSelectionDrawer::BuildKeySelectionDetails(const TSharedRef<FSeque
 		MenuBuilder.AddWidget(CreateKeyFrameDetails(InSelection).ToSharedRef(), FText::GetEmpty(), /*bInNoIndent=*/true);
 	}
 	MenuBuilder.EndSection();
-
-	// Show the section for the keys if they are all part of the same section
-	TArray<TViewModelPtr<FChannelModel>> Channels;
-	for (const FKeyHandle KeyHandle : InSelection->KeySelection)
-	{
-		const TViewModelPtr<FChannelModel> Channel = InSelection->KeySelection.GetModelForKey(KeyHandle);
-		Channels.AddUnique(Channel);
-	}
-	if (Channels.Num() == 1)
-	{
-		FSectionContextMenu::BuildKeyEditMenu(MenuBuilder, Sequencer, Sequencer->GetLastEvaluatedLocalTime());
-	}
 }
 
 void FSequencerSelectionDrawer::BuildTrackAreaDetails(const TSharedRef<FSequencerSelection>& InSelection, FMenuBuilder& MenuBuilder)
@@ -291,26 +296,8 @@ void FSequencerSelectionDrawer::BuildTrackAreaDetails(const TSharedRef<FSequence
 
 	for (const FViewModelPtr TrackAreaItem : InSelection->TrackArea)
 	{
-		if (const TViewModelPtr<FLayerBarModel> LayerBarModel = TrackAreaItem.ImplicitCast())
+		if (const TViewModelPtr<FSectionModel> SectionModel = TrackAreaItem.ImplicitCast())
 		{
-			const TViewModelPtr<IOutlinerExtension> LinkedOutlinerItem = LayerBarModel->GetLinkedOutlinerItem();
-
-			if (const TViewModelPtr<FOutlinerItemModel> OutlinerItemModel = LinkedOutlinerItem.ImplicitCast())
-			{
-				OutlinerItemModel->BuildSidebarMenu(MenuBuilder);
-			}
-		}
-		else if (const TViewModelPtr<FSectionModel> SectionModel = TrackAreaItem.ImplicitCast())
-		{
-			if (InSelection->TrackArea.Num() == 1)
-			{
-				if (const TSharedPtr<ISequencerSection> SectionInterface = SectionModel->GetSectionInterface())
-				{
-					const TSharedPtr<IObjectBindingExtension> ObjectBinding = SectionModel->FindAncestorOfType<IObjectBindingExtension>();
-					SectionInterface->BuildSectionSidebarMenu(MenuBuilder, ObjectBinding.IsValid() ? ObjectBinding->GetObjectGuid() : FGuid());
-				}
-			}
-
 			AllSectionObjects.Add(SectionModel->GetSection());
 		}
 	}
@@ -334,90 +321,72 @@ void FSequencerSelectionDrawer::BuildOutlinerDetails(const TSharedRef<FSequencer
 		return;
 	}
 
+	TArray<TWeakObjectPtr<>> WeakFolderObjects;
+	TArray<TWeakObjectPtr<>> WeakSectionObjects;
+	TArray<TWeakObjectPtr<>> WeakTrackObjects;
+	TSet<TViewModelPtr<FObjectBindingModel>> ObjectBindings;
 	TSet<TViewModelPtr<FChannelGroupOutlinerModel>> ChannelGroups;
 
 	for (const FViewModelPtr OutlinerItem : InSelection->Outliner)
 	{
-		if (const TViewModelPtr<FOutlinerItemModel> OutlinerItemModel = OutlinerItem.ImplicitCast())
+		if (const TViewModelPtr<FTrackModel> TrackModel = OutlinerItem.ImplicitCast())
+		{
+			WeakSectionObjects.Append(TrackModel->GetSections());
+			WeakTrackObjects.Add(TrackModel->GetTrack());
+		}
+		else if (const TViewModelPtr<FTrackRowModel> TrackRowModel = OutlinerItem.ImplicitCast())
+		{
+			WeakSectionObjects.Append(TrackRowModel->GetSections());
+			WeakTrackObjects.Add(TrackRowModel->GetTrack());
+		}
+		else if (const TViewModelPtr<FObjectBindingModel> ObjectBindingModel = OutlinerItem.ImplicitCast())
+		{
+			ObjectBindings.Add(ObjectBindingModel);
+		}
+		else if (const TViewModelPtr<FFolderModel> FolderModel = OutlinerItem.ImplicitCast())
 		{
-			OutlinerItemModel->BuildSidebarMenu(MenuBuilder);
+			WeakFolderObjects.Add(FolderModel->GetFolder());
 		}
 		// Ex. "Location.X", "Rotation.Roll", "Color.R", etc.
         else if (const TViewModelPtr<FChannelGroupOutlinerModel> ChannelGroupOutlinerModel = OutlinerItem.ImplicitCast())
         {
-        	ChannelGroupOutlinerModel->BuildSidebarMenu(MenuBuilder);
         	ChannelGroups.Add(ChannelGroupOutlinerModel);
         }
 	}
 
-	if (!ChannelGroups.IsEmpty())
+	if (!WeakFolderObjects.IsEmpty())
 	{
-		const ISequencerModule& SequencerModule = FModuleManager::Get().LoadModuleChecked<ISequencerModule>(TEXT("Sequencer"));
-		const TSharedPtr<FExtensibilityManager> SidebarExtensibilityManager = SequencerModule.GetSidebarExtensibilityManager();
-		const TSharedPtr<FExtender> Extender = SidebarExtensibilityManager->GetAllExtenders();
+		const TSharedRef<STrackDetails> TrackDetails = SNew(STrackDetails, WeakFolderObjects, Sequencer.ToSharedRef())
+			.NotifyMovieSceneDataChanged(true);
+		MenuBuilder.AddWidget(TrackDetails, FText::GetEmpty(), true);
+	}
 
-		TArray<FName> ChannelTypeNames;
-		TArray<ISequencerChannelInterface*> ChannelInterfaces;
-		TArray<FMovieSceneChannelHandle> ChannelHandles;
-		TArray<TWeakObjectPtr<UMovieSceneSection>> WeakSceneSections;
+	if (!ObjectBindings.IsEmpty())
+	{
+		MenuBuilder.BeginSection(TEXT("Possessable"));
+		MenuBuilder.EndSection();
 
-		for (const TViewModelPtr<FChannelGroupOutlinerModel>& ChannelModel : ChannelGroups)
-		{
-			for (const TSharedRef<IKeyArea>& KeyArea : ChannelModel->GetAllKeyAreas())
-			{
-				if (ISequencerChannelInterface* const SequencerChannelIterface = KeyArea->FindChannelEditorInterface())
-				{
-					const FMovieSceneChannelHandle& Channel = KeyArea->GetChannel();
+		// Shows duplicate information as above?
+		//MenuBuilder.BeginSection(TEXT("CustomBinding"));
+		//MenuBuilder.EndSection();
 
-					ChannelTypeNames.Add(Channel.GetChannelTypeName());
-					ChannelInterfaces.Add(SequencerChannelIterface);
-					ChannelHandles.Add(Channel);
-					WeakSceneSections.Add(KeyArea->GetOwningSection());
-				}
-			}
-		}
+		MenuBuilder.BeginSection(TEXT("TrackRowMetadata"));
+		MenuBuilder.EndSection();
+	}
 
-		// Need to make sure all channels are the same type to allow editing of multiple channels as one
-		bool bAllChannelNamesEqual = true;
-		if (!ChannelTypeNames.IsEmpty())
-		{
-			for (int32 Index = 0; Index < ChannelTypeNames.Num(); ++Index)
-			{
-				if (Index > 0 && ChannelTypeNames[Index] != ChannelTypeNames[0])
-				{
-					bAllChannelNamesEqual = false;
-					break;
-				}
-			}
-		}
+	if (!ChannelGroups.IsEmpty())
+	{
+		BuildExtensionDetails(ChannelGroups, MenuBuilder);
+	}
 
-		// Channel Interface Extensions (Perlin Noise, Easing, Wave)
-		if (ChannelInterfaces.Num() > 0)
-		{
-			if (bAllChannelNamesEqual)
-			{
-				if (const TSharedPtr<ISidebarChannelExtension> ChannelExtension = ChannelInterfaces[0]->ExtendSidebarMenu_Raw(MenuBuilder, Extender, ChannelHandles, WeakSceneSections, WeakSequencer))
-				{
-					ChannelExtensions.Add(ChannelExtension);
-				}
-			}
-			else
-			{
-				// Display different channels separately and don't allow to edit "all-in-one"
-				for (int32 Index = 0; Index < ChannelInterfaces.Num(); ++Index)
-				{
-					if (const TSharedPtr<ISidebarChannelExtension> ChannelExtension = ChannelInterfaces[Index]->ExtendSidebarMenu_Raw(MenuBuilder, Extender, { ChannelHandles[Index] }, { WeakSceneSections[Index] }, WeakSequencer))
-					{
-						ChannelExtensions.Add(ChannelExtension);
-					}
-				}
-			}
-		}
+	if (!WeakTrackObjects.IsEmpty())
+	{
+		SequencerHelpers::BuildEditTrackMenu(Sequencer, WeakTrackObjects, MenuBuilder, false);
+	}
 
-		// Curve Channel Options (Pre-Finity, Post-Finity, etc.)
-		CurveChannelExtension = MakeShared<FCurveChannelSectionSidebarExtension>(Sequencer);
-		CurveChannelExtension->AddSections(WeakSceneSections);
-		CurveChannelExtension->ExtendMenu(MenuBuilder, false);
+	if (!WeakSectionObjects.IsEmpty())
+	{
+		SequencerHelpers::BuildEditSectionMenu(WeakSequencer, WeakSectionObjects, MenuBuilder, false);
 	}
 }
 
@@ -429,15 +398,81 @@ void FSequencerSelectionDrawer::BuildMarkedFrameDetails(const TSharedRef<FSequen
 	}
 
 	MenuBuilder.BeginSection(TEXT("MarkedFrames"), LOCTEXT("MarkedFramesMenuSection", "Marked Frames"));
-	
+
 	for (const int32 MarkIndex : InSelection->MarkedFrames)
 	{
-		MenuBuilder.AddWidget(CreateMarkedFrameDetails(MarkIndex).ToSharedRef(), FText::GetEmpty(), /*bInNoIndent=*/true);
+		const TSharedRef<SMarkedFrameDetails> MarkedFrameDetails = SNew(SMarkedFrameDetails, MarkIndex, WeakSequencer);
+		MenuBuilder.AddWidget(MarkedFrameDetails, FText::GetEmpty(), /*bInNoIndent=*/true);
 	}
 
 	MenuBuilder.EndSection();
 }
 
+void FSequencerSelectionDrawer::BuildExtensionDetails(const TSet<TViewModelPtr<FChannelGroupOutlinerModel>>& InChannelGroups, FMenuBuilder& MenuBuilder)
+{
+	const TSharedPtr<ISequencer> Sequencer = WeakSequencer.Pin();
+	if (!Sequencer.IsValid())
+	{
+		return;
+	}
+
+	const ISequencerModule& SequencerModule = FModuleManager::Get().LoadModuleChecked<ISequencerModule>(TEXT("Sequencer"));
+	const TSharedPtr<FExtensibilityManager> SidebarExtensibilityManager = SequencerModule.GetSidebarExtensibilityManager();
+	const TSharedPtr<FExtender> Extender = SidebarExtensibilityManager->GetAllExtenders();
+
+	TArray<FName> ChannelTypeNames;
+	TArray<ISequencerChannelInterface*> ChannelInterfaces;
+	TArray<FMovieSceneChannelHandle> ChannelHandles;
+	TArray<TWeakObjectPtr<UMovieSceneSection>> WeakSceneSections;
+
+	for (const TViewModelPtr<FChannelGroupOutlinerModel>& ChannelModel : InChannelGroups)
+	{
+		for (const TSharedRef<IKeyArea>& KeyArea : ChannelModel->GetAllKeyAreas())
+		{
+			if (ISequencerChannelInterface* const SequencerChannelIterface = KeyArea->FindChannelEditorInterface())
+			{
+				const FMovieSceneChannelHandle& Channel = KeyArea->GetChannel();
+
+				ChannelTypeNames.Add(Channel.GetChannelTypeName());
+				ChannelInterfaces.Add(SequencerChannelIterface);
+				ChannelHandles.Add(Channel);
+				WeakSceneSections.Add(KeyArea->GetOwningSection());
+			}
+		}
+	}
+
+	// Need to make sure all channels are the same type to allow editing of multiple channels as one
+	bool bAllChannelNamesEqual = AreAllSameNames(ChannelTypeNames);
+
+	// Channel Interface Extensions (Perlin Noise, Easing, Wave)
+	if (ChannelInterfaces.Num() > 0)
+	{
+		if (bAllChannelNamesEqual)
+		{
+			if (const TSharedPtr<ISidebarChannelExtension> ChannelExtension = ChannelInterfaces[0]->ExtendSidebarMenu_Raw(MenuBuilder, Extender, ChannelHandles, WeakSceneSections, WeakSequencer))
+			{
+				ChannelExtensions.Add(ChannelExtension);
+			}
+		}
+		else
+		{
+			// Display different channels separately and don't allow to edit "all-in-one"
+			for (int32 Index = 0; Index < ChannelInterfaces.Num(); ++Index)
+			{
+				if (const TSharedPtr<ISidebarChannelExtension> ChannelExtension = ChannelInterfaces[Index]->ExtendSidebarMenu_Raw(MenuBuilder, Extender, { ChannelHandles[Index] }, { WeakSceneSections[Index] }, WeakSequencer))
+				{
+					ChannelExtensions.Add(ChannelExtension);
+				}
+			}
+		}
+	}
+
+	// Curve Channel Options (Pre-Finity, Post-Finity, etc.)
+	CurveChannelExtension = MakeShared<FCurveChannelSectionSidebarExtension>(Sequencer);
+	CurveChannelExtension->AddSections(WeakSceneSections);
+	CurveChannelExtension->ExtendMenu(MenuBuilder, false);
+}
+
 TSharedRef<SWidget> FSequencerSelectionDrawer::CreateHintText(const FText& InMessage)
 {
 	return SNew(SBox)
@@ -490,110 +525,16 @@ TSharedPtr<SWidget> FSequencerSelectionDrawer::CreateKeyFrameDetails(const TShar
 	return CreateHintText(LOCTEXT("InvalidKeyCombination", "Selected keys must belong to the same section."));
 }
 
-TSharedPtr<SWidget> FSequencerSelectionDrawer::CreateMarkedFrameDetails(const int32 InMarkedFrameIndex)
+bool FSequencerSelectionDrawer::AreAllSameNames(const TArray<FName>& InNames) const
 {
-	const TSharedPtr<FSequencer> Sequencer = WeakSequencer.Pin();
-	if (!Sequencer.IsValid())
+	for (int32 Index = 0; Index < InNames.Num(); ++Index)
 	{
-		return nullptr;
-	}
-
-	const UMovieSceneSequence* const FocusedMovieSceneSequence = Sequencer->GetFocusedMovieSceneSequence();
-	if (!IsValid(FocusedMovieSceneSequence))
-	{
-		return nullptr;
-	}
-
-	UMovieScene* const FocusedMovieScene = FocusedMovieSceneSequence->GetMovieScene();
-	if (!IsValid(FocusedMovieScene))
-	{
-		return nullptr;
-	}
-
-	if (FocusedMovieScene->GetMarkedFrames().Num() == 0)
-	{
-		return nullptr;
-	}
-
-	class SMarkedFramePropertyWidget : public SCompoundWidget, public FNotifyHook
-	{
-	public:
-		SLATE_BEGIN_ARGS(SMarkedFramePropertyWidget) {}
-		SLATE_END_ARGS()
-
-		void Construct(const FArguments& InArgs, UMovieScene* const InMovieScene, const int32 InMarkedFrameIndex, const TWeakPtr<FSequencer>& InWeakSequencer)
-		{
-			MovieSceneToModify = InMovieScene;
-			WeakSequencer = InWeakSequencer;
-
-			FDetailsViewArgs DetailsViewArgs;
-			DetailsViewArgs.bAllowSearch = false;
-			DetailsViewArgs.bShowScrollBar = false;
-			DetailsViewArgs.NameAreaSettings = FDetailsViewArgs::HideNameArea;
-			DetailsViewArgs.NotifyHook = this;
-
-			FStructureDetailsViewArgs StructureDetailsViewArgs;
-			StructureDetailsViewArgs.bShowObjects = true;
-			StructureDetailsViewArgs.bShowAssets = true;
-			StructureDetailsViewArgs.bShowClasses = true;
-			StructureDetailsViewArgs.bShowInterfaces = true;
-			
-			const TSharedPtr<FStructOnScope> StructOnScope = MakeShared<FStructOnScope>(FMovieSceneMarkedFrame::StaticStruct(), (uint8*)&InMovieScene->GetMarkedFrames()[InMarkedFrameIndex]);
-
-			FPropertyEditorModule& PropertyEditorModule = FModuleManager::GetModuleChecked<FPropertyEditorModule>(TEXT("PropertyEditor"));
-
-			DetailsView = PropertyEditorModule.CreateStructureDetailView(DetailsViewArgs, StructureDetailsViewArgs, nullptr);
-			DetailsView->GetDetailsView()->RegisterInstancedCustomPropertyTypeLayout(TEXT("FrameNumber"), FOnGetPropertyTypeCustomizationInstance::CreateLambda([this]() {
-				return MakeShared<FFrameNumberDetailsCustomization>(WeakSequencer.Pin()->GetNumericTypeInterface()); }));
-			DetailsView->SetStructureData(StructOnScope);
-			
-			ChildSlot
-			[
-				DetailsView->GetWidget().ToSharedRef()
-			];
-		}
-
-		virtual void NotifyPreChange(FProperty* InPropertyAboutToChange) override
-		{
-			MovieSceneToModify->Modify();
-		}
-
-		virtual void NotifyPreChange(FEditPropertyChain* InPropertyAboutToChange) override
-		{
-			MovieSceneToModify->Modify();
-		}
-
-	private:
-		TObjectPtr<UMovieScene> MovieSceneToModify;
-		TSharedPtr<IStructureDetailsView> DetailsView;
-		TWeakPtr<FSequencer> WeakSequencer;
-	};
-
-	auto AreMarkedFramesLocked = [&Sequencer]() -> bool
-	{
-		if (Sequencer->IsReadOnly())
+		if (Index > 0 && InNames[Index] != InNames[0])
 		{
-			return true;
+			return false;
 		}
-
-		const UMovieSceneSequence* const FocusedMovieSceneSequence = Sequencer->GetFocusedMovieSceneSequence();
-		if (FocusedMovieSceneSequence != nullptr)
-		{
-			const UMovieScene* const MovieScene = FocusedMovieSceneSequence->GetMovieScene();
-			if (MovieScene->IsReadOnly())
-			{
-				return true;
-			}
-			return MovieScene->AreMarkedFramesLocked();
-		}
-
-		return false;
-	};
-
-	const TSharedRef<SMarkedFramePropertyWidget> Widget = SNew(SMarkedFramePropertyWidget, FocusedMovieScene, InMarkedFrameIndex, WeakSequencer);
-	Widget->SetEnabled(!AreMarkedFramesLocked());
-
-	return Widget;
+	}
+	return true;
 }
 
 #undef LOCTEXT_NAMESPACE
diff --git a/Engine/Source/Editor/Sequencer/Private/Widgets/Sidebar/SequencerSelectionDrawer.h b/Engine/Source/Editor/Sequencer/Private/Widgets/Sidebar/SequencerSelectionDrawer.h
--- a/Engine/Source/Editor/Sequencer/Private/Widgets/Sidebar/SequencerSelectionDrawer.h
+++ b/Engine/Source/Editor/Sequencer/Private/Widgets/Sidebar/SequencerSelectionDrawer.h
@@ -2,6 +2,7 @@
 
 #pragma once
 
+#include "MVVM/ViewModelPtr.h"
 #include "Sidebar/ISidebarDrawerContent.h"
 #include "Templates/SharedPointerFwd.h"
 
@@ -17,6 +18,7 @@ struct FKeyEditData;
 
 namespace UE::Sequencer
 {
+	class FChannelGroupOutlinerModel;
 	class FSequencerSelection;
 }
 
@@ -26,6 +28,7 @@ public:
 	static const FName UniqueId;
 
 	FSequencerSelectionDrawer(const TWeakPtr<FSequencer>& InWeakSequencer);
+	virtual ~FSequencerSelectionDrawer() override;
 
 	//~ Begin ISidebarDrawerContent
 	virtual FName GetUniqueId() const override;
@@ -42,6 +45,7 @@ protected:
 	void BuildTrackAreaDetails(const TSharedRef<UE::Sequencer::FSequencerSelection>& InSelection, FMenuBuilder& MenuBuilder);
 	void BuildOutlinerDetails(const TSharedRef<UE::Sequencer::FSequencerSelection>& InSelection, FMenuBuilder& MenuBuilder);
 	void BuildMarkedFrameDetails(const TSharedRef<UE::Sequencer::FSequencerSelection>& InSelection, FMenuBuilder& MenuBuilder);
+	void BuildExtensionDetails(const TSet<UE::Sequencer::TViewModelPtr<UE::Sequencer::FChannelGroupOutlinerModel>>& InChannelGroups, FMenuBuilder& MenuBuilder);
 
 	TSharedRef<SWidget> CreateHintText(const FText& InMessage);
 	TSharedRef<SWidget> CreateNoSelectionHintText();
@@ -49,10 +53,11 @@ protected:
 	FKeyEditData GetKeyEditData() const;
 	
 	TSharedPtr<SWidget> CreateKeyFrameDetails(const TSharedRef<UE::Sequencer::FSequencerSelection>& InSequencerSelection);
-	TSharedPtr<SWidget> CreateMarkedFrameDetails(const int32 InMarkedFrameIndex);
 
 	void ResetContent();
 
+	bool AreAllSameNames(const TArray<FName>& InNames) const;
+
 	TWeakPtr<FSequencer> WeakSequencer;
 
 	TSharedPtr<SVerticalBox> ContentBox;
