[Sequencer] Check for EMovieSceneDataChangeType::MovieSceneStructureItemsChanged instead to update for overridden curve extensions

diff --git a/Engine/Source/Editor/Sequencer/Private/Widgets/Sidebar/SequencerSelectionDrawer.cpp b/Engine/Source/Editor/Sequencer/Private/Widgets/Sidebar/SequencerSelectionDrawer.cpp
--- a/Engine/Source/Editor/Sequencer/Private/Widgets/Sidebar/SequencerSelectionDrawer.cpp
+++ b/Engine/Source/Editor/Sequencer/Private/Widgets/Sidebar/SequencerSelectionDrawer.cpp
@@ -137,7 +137,8 @@ TSharedRef<SWidget> FSequencerSelectionDrawer::CreateContentWidget()
 
 		Sequencer->OnMovieSceneDataChanged().AddLambda([this](const EMovieSceneDataChangeType InChangeType)
 			{
-				if (InChangeType == EMovieSceneDataChangeType::ActiveMovieSceneChanged)
+				// Update when a channel is overriden with a curve extension
+				if (InChangeType == EMovieSceneDataChangeType::MovieSceneStructureItemsChanged)
 				{
 					OnSequencerSelectionChanged();
 				}
