[Sequencer] Advanced filtering: Fix slow performance when using auto-key

diff --git a/Engine/Source/Editor/Sequencer/Private/Filters/SequencerFilterBar.cpp b/Engine/Source/Editor/Sequencer/Private/Filters/SequencerFilterBar.cpp
--- a/Engine/Source/Editor/Sequencer/Private/Filters/SequencerFilterBar.cpp
+++ b/Engine/Source/Editor/Sequencer/Private/Filters/SequencerFilterBar.cpp
@@ -1544,10 +1544,6 @@ TArray<TSharedRef<FSequencerTrackFilter>> FSequencerFilterBar::GetFilterList(con
 
 	AllFilters.Add(TextFilter);
 	AllFilters.Add(HideIsolateFilter);
-	AllFilters.Add(LevelFilter);
-	AllFilters.Add(GroupFilter);
-	AllFilters.Add(SelectedFilter);
-	AllFilters.Add(ModifiedFilter);
 
 	if (bInIncludeCustomTextFilters)
 	{
@@ -1562,11 +1558,16 @@ TArray<TSharedRef<FSequencerTrackFilter>> FSequencerFilterBar::GetFilterList(con
 
 bool FSequencerFilterBar::ShouldUpdateOnTrackValueChanged() const
 {
+	if (bFiltersMuted)
+	{
+		return false;
+	}
+
 	const TArray<TSharedRef<FSequencerTrackFilter>> AllFilters = GetFilterList();
 
 	for (const TSharedRef<FSequencerTrackFilter>& Filter : AllFilters)
 	{
-		if (Filter->ShouldUpdateOnTrackValueChanged())
+		if (Filter->ShouldUpdateOnTrackValueChanged() && IsFilterActive(Filter))
 		{
 			return true;
 		}
