Sequencer: Remove TextExportTransient on copyable bindings so that the spawnable object templates are actually copied/exported

diff --git a/Engine/Source/Editor/Sequencer/Private/MovieSceneCopyableBinding.h b/Engine/Source/Editor/Sequencer/Private/MovieSceneCopyableBinding.h
--- a/Engine/Source/Editor/Sequencer/Private/MovieSceneCopyableBinding.h
+++ b/Engine/Source/Editor/Sequencer/Private/MovieSceneCopyableBinding.h
@@ -21,10 +21,9 @@ public:
 	/** 
 	* Spawnables need to know about their Object Template but we cannot rely on automatic serialization due to the object
 	* template belonging to the Movie Scene (it gets serialized as a reference). Instead we manually serialize the object
-	* so that we can duplicate it into a new object (which is stored in this variable) but we don't want this exported with
-	* the rest of the text as it'll fall back to the same reference issue. Marking this as TextExportTransient solves this.
+	* so that we can duplicate it into a new object (which is stored in this variable).
 	*/
-	UPROPERTY(TextExportTransient)
+	UPROPERTY()
 	TArray<TObjectPtr<UObject>> SpawnableObjectTemplates;
 
 	UPROPERTY()
@@ -56,7 +55,7 @@ public:
 	UPROPERTY()
 	TArray<FName> Tags;
 
-	UPROPERTY(TextExportTransient)
+	UPROPERTY()
 	TArray<TObjectPtr<UMovieSceneCustomBinding>> CustomBindings;	
 	
 	UPROPERTY()
