[Sequencer]
 - Add icons for filter bar widget
 - Move Hide/Isolate/Show buttons back to the right of the search box
 - Remove folder button

diff --git a/Engine/Source/Editor/EditorStyle/Private/StarshipStyle.cpp b/Engine/Source/Editor/EditorStyle/Private/StarshipStyle.cpp
--- a/Engine/Source/Editor/EditorStyle/Private/StarshipStyle.cpp
+++ b/Engine/Source/Editor/EditorStyle/Private/StarshipStyle.cpp
@@ -2371,6 +2371,10 @@ void FStarshipEditorStyle::FStyle::SetupSequencerStyles()
 		Set("EMovieSceneBlendType::Relative", new IMAGE_BRUSH("Sequencer/EMovieSceneBlendType_Relative", FVector2D(32, 16)));
 		Set("EMovieSceneBlendType::Additive", new IMAGE_BRUSH("Sequencer/EMovieSceneBlendType_Additive", FVector2D(32, 16)));
 		Set("EMovieSceneBlendType::AdditiveFromBase", new IMAGE_BRUSH("Sequencer/EMovieSceneBlendType_AdditiveFromBase", FVector2D(32, 16)));
+
+		Set("Sequencer.TrackIsolate", new IMAGE_BRUSH_SVG("Sequencer/TrackIsolate", Icon16x16));
+		Set("Sequencer.TrackHide", new IMAGE_BRUSH_SVG("Sequencer/TrackHide", Icon16x16));
+		Set("Sequencer.TrackShow", new IMAGE_BRUSH_SVG("Sequencer/TrackShow", Icon16x16));
 	}
 
 
diff --git a/Engine/Source/Editor/Sequencer/Private/Filters/Widgets/SFilterBarIsolateHideShow.cpp b/Engine/Source/Editor/Sequencer/Private/Filters/Widgets/SFilterBarIsolateHideShow.cpp
--- a/Engine/Source/Editor/Sequencer/Private/Filters/Widgets/SFilterBarIsolateHideShow.cpp
+++ b/Engine/Source/Editor/Sequencer/Private/Filters/Widgets/SFilterBarIsolateHideShow.cpp
@@ -5,6 +5,8 @@
 #include "Filters/SequencerTrackFilterCommands.h"
 #include "Framework/Application/SlateApplication.h"
 #include "Styling/StyleColors.h"
+#include "Widgets/Images/SImage.h"
+#include "Widgets/Images/SLayeredImage.h"
 #include "Widgets/Input/SButton.h"
 #include "Widgets/SBoxPanel.h"
 
@@ -14,6 +16,9 @@ void SFilterBarIsolateHideShow::Construct(const FArguments& InArgs, const TShare
 {
 	WeakFilterBar = InFilterBar;
 
+	constexpr float ButtonContentPadding = 2.f;
+	constexpr float ButtonSpacing = 1.f;
+
 	ChildSlot
 	[
 		SNew(SHorizontalBox)
@@ -22,24 +27,19 @@ void SFilterBarIsolateHideShow::Construct(const FArguments& InArgs, const TShare
 		+ SHorizontalBox::Slot()
 		.AutoWidth()
 		.VAlign(VAlign_Center)
-		.Padding(2.f, 0.f, 0.f, 0.f)
+		.Padding(0.f, 0.f, ButtonSpacing, 0.f)
 		[
-			SNew(SBox)
-			.WidthOverride(20.f)
+			SNew(SButton)
+			.ContentPadding(ButtonContentPadding)
+			.HAlign(HAlign_Center)
+			.VAlign(VAlign_Center)
+			.ButtonStyle(FAppStyle::Get(), TEXT("SimpleButton"))
+			.ToolTipText(this, &SFilterBarIsolateHideShow::GetIsolateTracksButtonTooltipText)
+			.IsEnabled(this, &SFilterBarIsolateHideShow::AreFiltersMuted)
+			.OnClicked(this, &SFilterBarIsolateHideShow::HandleIsolateTracksClick)
 			[
-				SNew(SButton)
-				.ContentPadding(2.f)
-				.HAlign(HAlign_Center)
-				.VAlign(VAlign_Center)
-				.ButtonStyle(FAppStyle::Get(), TEXT("SimpleButton"))
-				.ToolTipText(this, &SFilterBarIsolateHideShow::GetIsolateTracksButtonTooltipText)
-				.IsEnabled(this, &SFilterBarIsolateHideShow::AreFiltersMuted)
-				.OnClicked(this, &SFilterBarIsolateHideShow::HandleIsolateTracksClick)
-				[
-					SNew(STextBlock)
-					.ColorAndOpacity(this, &SFilterBarIsolateHideShow::GetIsolateTracksButtonTextColor)
-					.Text(LOCTEXT("IsolateSelectedButtonText", "I"))
-				]
+				ConstructLayeredImage(TEXT("Sequencer.TrackIsolate")
+					, TAttribute<bool>::CreateSP(this, &SFilterBarIsolateHideShow::HasIsolatedTracks))
 			]
 		]
 
@@ -47,24 +47,19 @@ void SFilterBarIsolateHideShow::Construct(const FArguments& InArgs, const TShare
 		+ SHorizontalBox::Slot()
 		.AutoWidth()
 		.VAlign(VAlign_Center)
-		.Padding(2.f, 0.f, 0.f, 0.f)
+		.Padding(0.f, 0.f, ButtonSpacing, 0.f)
 		[
-			SNew(SBox)
-			.WidthOverride(20.f)
+			SNew(SButton)
+			.ContentPadding(ButtonContentPadding)
+			.HAlign(HAlign_Center)
+			.VAlign(VAlign_Center)
+			.ButtonStyle(FAppStyle::Get(), TEXT("SimpleButton"))
+			.ToolTipText(this, &SFilterBarIsolateHideShow::GetHideTracksButtonTooltipText)
+			.IsEnabled(this, &SFilterBarIsolateHideShow::AreFiltersMuted)
+			.OnClicked(this, &SFilterBarIsolateHideShow::HandleHideTracksClick)
 			[
-				SNew(SButton)
-				.ContentPadding(2.f)
-				.HAlign(HAlign_Center)
-				.VAlign(VAlign_Center)
-				.ButtonStyle(FAppStyle::Get(), TEXT("SimpleButton"))
-				.ToolTipText(this, &SFilterBarIsolateHideShow::GetHideTracksButtonTooltipText)
-				.IsEnabled(this, &SFilterBarIsolateHideShow::AreFiltersMuted)
-				.OnClicked(this, &SFilterBarIsolateHideShow::HandleHideTracksClick)
-				[
-					SNew(STextBlock)
-					.ColorAndOpacity(this, &SFilterBarIsolateHideShow::GetHideTracksButtonTextColor)
-					.Text(LOCTEXT("HideSelectedButtonText", "H"))
-				]
+				ConstructLayeredImage(TEXT("Sequencer.TrackHide")
+					, TAttribute<bool>::CreateSP(this, &SFilterBarIsolateHideShow::HasHiddenTracks))
 			]
 		]
 
@@ -72,29 +67,42 @@ void SFilterBarIsolateHideShow::Construct(const FArguments& InArgs, const TShare
 		+ SHorizontalBox::Slot()
 		.AutoWidth()
 		.VAlign(VAlign_Center)
-		.Padding(2.f, 0.f, 0.f, 0.f)
+		.Padding(0.f, 0.f, 0.f, 0.f)
 		[
-			SNew(SBox)
-			.WidthOverride(20.f)
+			SNew(SButton)
+			.ContentPadding(ButtonContentPadding)
+			.HAlign(HAlign_Center)
+			.VAlign(VAlign_Center)
+			.ButtonStyle(FAppStyle::Get(), TEXT("SimpleButton"))
+			.ToolTipText(this, &SFilterBarIsolateHideShow::GetShowAllTracksButtonTooltipText)
+			.IsEnabled(this, &SFilterBarIsolateHideShow::AreFiltersMuted)
+			.OnClicked(this, &SFilterBarIsolateHideShow::HandleShowAllTracksClick)
 			[
-				SNew(SButton)
-				.ContentPadding(2.f)
-				.HAlign(HAlign_Center)
-				.VAlign(VAlign_Center)
-				.ButtonStyle(FAppStyle::Get(), TEXT("SimpleButton"))
-				.ToolTipText(this, &SFilterBarIsolateHideShow::GetShowAllTracksButtonTooltipText)
-				.IsEnabled(this, &SFilterBarIsolateHideShow::AreFiltersMuted)
-				.OnClicked(this, &SFilterBarIsolateHideShow::HandleShowAllTracksClick)
-				[
-					SNew(STextBlock)
-					.ColorAndOpacity(this, &SFilterBarIsolateHideShow::GetShowAllTracksButtonTextColor)
-					.Text(LOCTEXT("ShowAllButtonText", "S"))
-				]
+				SNew(SImage)
+				.DesiredSizeOverride(FVector2D(16.f))
+				.ColorAndOpacity(this, &SFilterBarIsolateHideShow::GetShowAllTracksButtonTextColor)
+				.Image(FAppStyle::Get().GetBrush(TEXT("Sequencer.TrackShow")))
 			]
 		]
 	];
 }
 
+TSharedRef<SWidget> SFilterBarIsolateHideShow::ConstructLayeredImage(const FName InBaseImageName, const TAttribute<bool>& InShowBadge)
+{
+	const TSharedRef<SLayeredImage> LayeredImage = SNew(SLayeredImage)
+		.DesiredSizeOverride(FVector2D(16.f))
+		.Image(FAppStyle::Get().GetBrush(InBaseImageName));
+
+	LayeredImage->AddLayer(TAttribute<const FSlateBrush*>::CreateLambda([InShowBadge]() -> const FSlateBrush*
+		{
+			return InShowBadge.Get(false)
+				? FAppStyle::Get().GetBrush(TEXT("Icons.BadgeModified"))
+				: nullptr;
+		}));
+
+	return LayeredImage;
+}
+
 bool SFilterBarIsolateHideShow::AreFiltersMuted() const
 {
 	return WeakFilterBar.IsValid() ? !WeakFilterBar.Pin()->AreFiltersMuted() : true;
@@ -157,16 +165,14 @@ FReply SFilterBarIsolateHideShow::HandleShowAllTracksClick()
 	return FReply::Handled();
 }
 
-FSlateColor SFilterBarIsolateHideShow::GetHideTracksButtonTextColor() const
+bool SFilterBarIsolateHideShow::HasIsolatedTracks() const
 {
-	return WeakFilterBar.IsValid() && WeakFilterBar.Pin()->GetHiddenTracks().Num() == 0
-		? FStyleColors::Foreground : FStyleColors::AccentYellow;
+	return WeakFilterBar.IsValid() && WeakFilterBar.Pin()->GetIsolatedTracks().Num() > 0;
 }
 
-FSlateColor SFilterBarIsolateHideShow::GetIsolateTracksButtonTextColor() const
+bool SFilterBarIsolateHideShow::HasHiddenTracks() const
 {
-	return WeakFilterBar.IsValid() && WeakFilterBar.Pin()->GetIsolatedTracks().Num() == 0
-		? FStyleColors::Foreground : FStyleColors::AccentBlue;
+	return WeakFilterBar.IsValid() && WeakFilterBar.Pin()->GetHiddenTracks().Num() > 0;
 }
 
 FSlateColor SFilterBarIsolateHideShow::GetShowAllTracksButtonTextColor() const
@@ -174,7 +180,7 @@ FSlateColor SFilterBarIsolateHideShow::GetShowAllTracksButtonTextColor() const
 	if (const TSharedPtr<FSequencerFilterBar> FilterBar = WeakFilterBar.Pin())
 	{
 		return (FilterBar->GetHiddenTracks().Num() == 0 && FilterBar->GetIsolatedTracks().Num() == 0)
-			? FStyleColors::Foreground : FStyleColors::AccentRed;
+			? FStyleColors::Foreground : FStyleColors::Warning;
 	}
 	return FStyleColors::Foreground;
 }
diff --git a/Engine/Source/Editor/Sequencer/Private/Filters/Widgets/SFilterBarIsolateHideShow.h b/Engine/Source/Editor/Sequencer/Private/Filters/Widgets/SFilterBarIsolateHideShow.h
--- a/Engine/Source/Editor/Sequencer/Private/Filters/Widgets/SFilterBarIsolateHideShow.h
+++ b/Engine/Source/Editor/Sequencer/Private/Filters/Widgets/SFilterBarIsolateHideShow.h
@@ -21,14 +21,17 @@ public:
 	void Construct(const FArguments& InArgs, const TSharedRef<FSequencerFilterBar>& InFilterBar);
 
 protected:
+	TSharedRef<SWidget> ConstructLayeredImage(const FName InBaseImageName, const TAttribute<bool>& InShowBadge);
+	
 	bool AreFiltersMuted() const;
 
 	FReply HandleHideTracksClick();
 	FReply HandleIsolateTracksClick();
 	FReply HandleShowAllTracksClick();
 
-	FSlateColor GetHideTracksButtonTextColor() const;
-	FSlateColor GetIsolateTracksButtonTextColor() const;
+	bool HasIsolatedTracks() const;
+	bool HasHiddenTracks() const;
+
 	FSlateColor GetShowAllTracksButtonTextColor() const;
 
 	FText GetHideTracksButtonTooltipText() const;
diff --git a/Engine/Source/Editor/Sequencer/Private/SSequencer.cpp b/Engine/Source/Editor/Sequencer/Private/SSequencer.cpp
--- a/Engine/Source/Editor/Sequencer/Private/SSequencer.cpp
+++ b/Engine/Source/Editor/Sequencer/Private/SSequencer.cpp
@@ -1347,15 +1347,6 @@ TSharedRef<SWidget> SSequencer::ConstructSearchAndFilterRow()
 			.IsEnabled_Lambda([this]() { return !SequencerPtr.Pin()->IsReadOnly(); })
 		]
 
-		// Isolate / Hide / Show
-		+ SHorizontalBox::Slot()
-		.AutoWidth()
-		.VAlign(VAlign_Center)
-		.Padding(2.f, 0.f, 0.f, 0.f)
-		[
-			FilterBar->MakeIsolateHideShowPanel()
-		]
-
 		// Advanced Search Filter Combo Button
 		+ SHorizontalBox::Slot()
 		.AutoWidth()
@@ -1373,24 +1364,13 @@ TSharedRef<SWidget> SSequencer::ConstructSearchAndFilterRow()
 			SearchBox.ToSharedRef()
 		]
 
-		// Create Folder Button
+		// Isolate / Hide / Show
 		+ SHorizontalBox::Slot()
 		.AutoWidth()
 		.VAlign(VAlign_Center)
 		.Padding(2.f, 0.f, 0.f, 0.f)
 		[
-			SNew(SButton)
-			.ContentPadding(2.f)
-			.HAlign(HAlign_Center)
-			.VAlign(VAlign_Center)
-			.ButtonStyle(FAppStyle::Get(), TEXT("SimpleButton"))
-			.ToolTipText(LOCTEXT("CreateFolderToolTip", "Creates a new folder containing the current selection or adds a new root folder if there is no selection"))
-			.OnClicked(this, &SSequencer::HandleCreateFolderTrackClick)
-			[
-				SNew(SImage)
-				.ColorAndOpacity(FSlateColor::UseForeground())
-				.Image(FAppStyle::Get().GetBrush(TEXT("SceneOutliner.NewFolderIcon")))
-			]
+			FilterBar->MakeIsolateHideShowPanel()
 		]
 
 		// View Options Combo Button
@@ -4243,22 +4223,6 @@ void SSequencer::EnableCurveEditorPendingFocusOnHovering(const bool InEnabled) c
 	CurveEditorPanel->EnablePendingFocusOnHovering(InEnabled);
 }
 
-FReply SSequencer::HandleCreateFolderTrackClick()
-{
-	if (const TSharedPtr<FSequencer> Sequencer = SequencerPtr.Pin())
-	{
-		if (Sequencer->GetSelectedNodesToMove().Num() > 0)
-		{
-			Sequencer->MoveSelectedNodesToNewFolder();
-		}
-		else
-		{
-			Sequencer->AddFolder();
-		}
-	}
-	return FReply::Handled();
-}
-
 TSharedPtr<FSequencerFilterBar> SSequencer::GetFilterBar() const
 {
 	if (const TSharedPtr<FSequencer> Sequencer = SequencerPtr.Pin())
diff --git a/Engine/Source/Editor/Sequencer/Private/SSequencer.h b/Engine/Source/Editor/Sequencer/Private/SSequencer.h
--- a/Engine/Source/Editor/Sequencer/Private/SSequencer.h
+++ b/Engine/Source/Editor/Sequencer/Private/SSequencer.h
@@ -724,8 +724,6 @@ private:
 
 	void OnSidebarStateChanged(const FSidebarState& InNewState);
 
-	FReply HandleCreateFolderTrackClick();
-
 	void OnTrackFiltersChanged(const ESequencerFilterChange InChangeType, const TSharedRef<FSequencerTrackFilter>& InFilter);
 
 	/** This is the main container widget for the Sequencer (minus the toolbar). This is the parent of
