[Sequencer] Advanced filtering: Filter expression help dialog UX adjustments

diff --git a/Engine/Source/Editor/Sequencer/Private/Filters/TextExpressions/SequencerTextFilterExpression_ConditionFunc.cpp b/Engine/Source/Editor/Sequencer/Private/Filters/TextExpressions/SequencerTextFilterExpression_ConditionFunc.cpp
--- a/Engine/Source/Editor/Sequencer/Private/Filters/TextExpressions/SequencerTextFilterExpression_ConditionFunc.cpp
+++ b/Engine/Source/Editor/Sequencer/Private/Filters/TextExpressions/SequencerTextFilterExpression_ConditionFunc.cpp
@@ -28,7 +28,7 @@ ESequencerTextFilterValueType FSequencerTextFilterExpression_ConditionFunc::GetV
 
 FText FSequencerTextFilterExpression_ConditionFunc::GetDescription() const
 {
-	return LOCTEXT("ExpressionDescription_ConditionFunc", "Filters by the presence of a director blueprint condition with the given function/endponit name");
+	return LOCTEXT("ExpressionDescription_ConditionFunc", "Filters by the presence of a director blueprint condition with the given function/endpoint name");
 }
 
 bool FSequencerTextFilterExpression_ConditionFunc::TestComplexExpression(const FName& InKey
diff --git a/Engine/Source/Editor/Sequencer/Private/Filters/Widgets/SFilterExpressionHelpDialog.cpp b/Engine/Source/Editor/Sequencer/Private/Filters/Widgets/SFilterExpressionHelpDialog.cpp
--- a/Engine/Source/Editor/Sequencer/Private/Filters/Widgets/SFilterExpressionHelpDialog.cpp
+++ b/Engine/Source/Editor/Sequencer/Private/Filters/Widgets/SFilterExpressionHelpDialog.cpp
@@ -23,11 +23,10 @@ void SFilterExpressionHelpDialog::Construct(const FArguments& InArgs)
 	SWindow::Construct(SWindow::FArguments()
 		.Title(InArgs._DialogTitle)
 		.AutoCenter(EAutoCenter::PrimaryWorkArea)
-		.SizingRule(ESizingRule::FixedSize)
+		.SizingRule(ESizingRule::Autosized)
 		.HasCloseButton(true)
 		.SupportsMaximize(false)
-		.SupportsMinimize(false)
-		.ClientSize(InArgs._DialogSize));
+		.SupportsMinimize(false));
 
 	const TSharedRef<SVerticalBox> ContentWidget =
 		SNew(SVerticalBox)
@@ -40,10 +39,15 @@ void SFilterExpressionHelpDialog::Construct(const FArguments& InArgs)
 		+ SVerticalBox::Slot()
 		.Padding(5.f)
 		[
-			SNew(SBorder)
-			.BorderImage(FAppStyle::Get().GetBrush(TEXT("Brushes.Panel")))
+			SNew(SBox)
+			.MaxDesiredWidth(MaxDesiredWidth)
+			.MaxDesiredHeight(MaxDesiredHeight)
 			[
-				ConstructExpressionWidgetList()
+				SNew(SBorder)
+				.BorderImage(FAppStyle::Get().GetBrush(TEXT("Brushes.Panel")))
+				[
+					ConstructExpressionWidgetList()
+				]
 			]
 		];
 
@@ -150,6 +154,7 @@ TSharedRef<SWidget> SFilterExpressionHelpDialog::ConstructExpressionWidget(const
 		.Padding(15.f, 0.f, 5.f, 5.f)
 		[
 			SNew(STextBlock)
+			.AutoWrapText(true)
 			.TextStyle(FAppStyle::Get(), TEXT("NormalText"))
 			.Text(InExpressionContext->GetDescription())
 		];
diff --git a/Engine/Source/Editor/Sequencer/Private/Filters/Widgets/SFilterExpressionHelpDialog.h b/Engine/Source/Editor/Sequencer/Private/Filters/Widgets/SFilterExpressionHelpDialog.h
--- a/Engine/Source/Editor/Sequencer/Private/Filters/Widgets/SFilterExpressionHelpDialog.h
+++ b/Engine/Source/Editor/Sequencer/Private/Filters/Widgets/SFilterExpressionHelpDialog.h
@@ -10,14 +10,15 @@ enum class ESequencerTextFilterValueType : uint8;
 class SFilterExpressionHelpDialog : public SWindow
 {
 public:
+	static constexpr float MaxDesiredWidth = 460.f;
+	static constexpr float MaxDesiredHeight = 560.f;
+
 	SLATE_BEGIN_ARGS(SFilterExpressionHelpDialog)
 		: _DialogTitle(NSLOCTEXT("SFilterExpressionHelpDialog", "DialogTitle", "Text Filter Expression Help"))
-		, _DialogSize(FVector2D(320.f, 560.f))
 		, _HeaderText(NSLOCTEXT("SFilterExpressionHelpDialog", "HeaderText", "Text Filter Expressions"))
 		, _DocumentationLink(TEXT("https://dev.epicgames.com/documentation/en-us/unreal-engine/advanced-search-syntax-in-unreal-engine"))
 	{}
 		SLATE_ARGUMENT(FText, DialogTitle)
-		SLATE_ARGUMENT(FVector2D, DialogSize)
 		SLATE_ARGUMENT(FText, HeaderText)
 		SLATE_ARGUMENT(FString, DocumentationLink)
 		SLATE_ARGUMENT(TArray<TSharedRef<FSequencerTextFilterExpressionContext>>, TextFilterExpressionContexts)
