Level Sequences can now save a relevant thumbnail based on what the sequence looks like at a configurable frame.
- This feature is guarded behind the Sequencer.EnableRelevantThumbnails CVar
- Added configuration entry to USequencerSettings to let users configure whether to capture the first, middle, or last frame (default: middle)

diff --git a/Engine/Source/Editor/Sequencer/Private/Misc/ConsoleVariables.cpp b/Engine/Source/Editor/Sequencer/Private/Misc/ConsoleVariables.cpp
new file mode 100644
--- /dev/null
+++ b/Engine/Source/Editor/Sequencer/Private/Misc/ConsoleVariables.cpp
@@ -0,0 +1,12 @@
+﻿// Copyright Epic Games, Inc. All Rights Reserved.
+
+#include "ConsoleVariables.h"
+
+namespace UE::Sequencer
+{
+	TAutoConsoleVariable<bool> CVarEnableRelevantThumbnails(
+		TEXT("Sequencer.EnableRelevantThumbnails"),
+		false,
+		TEXT("If enabled, upon saving a sequence, tries to find a suitable frame, scrubs to it, renders a thumbnail, and saves it as thumbnail.")
+	);
+}
\ No newline at end of file
diff --git a/Engine/Source/Editor/Sequencer/Private/Misc/ConsoleVariables.h b/Engine/Source/Editor/Sequencer/Private/Misc/ConsoleVariables.h
new file mode 100644
--- /dev/null
+++ b/Engine/Source/Editor/Sequencer/Private/Misc/ConsoleVariables.h
@@ -0,0 +1,11 @@
+﻿// Copyright Epic Games, Inc. All Rights Reserved.
+
+#pragma once
+
+#include "HAL/IConsoleManager.h"
+
+namespace UE::Sequencer
+{
+	/** Whether to enable the feature that, upon saving a sequence, tries to find a suitable frame, scrubs to it, renders a thumbnail, and saves it as thumbnail. */
+	extern TAutoConsoleVariable<bool> CVarEnableRelevantThumbnails;
+}
diff --git a/Engine/Source/Editor/Sequencer/Private/Misc/SequencerThumbnailCaptureSettings.cpp b/Engine/Source/Editor/Sequencer/Private/Misc/SequencerThumbnailCaptureSettings.cpp
new file mode 100644
--- /dev/null
+++ b/Engine/Source/Editor/Sequencer/Private/Misc/SequencerThumbnailCaptureSettings.cpp
@@ -0,0 +1,34 @@
+﻿// Copyright Epic Games, Inc. All Rights Reserved.
+
+#include "Misc/SequencerThumbnailCaptureSettings.h"
+
+#include "ISequencer.h"
+#include "MovieScene.h"
+#include "Math/Range.h"
+
+namespace UE::Sequencer
+{
+	FFrameNumber GetFrameByRule(const ISequencer& Sequencer, ESequencerThumbnailCaptureTimeLocation Rule)
+	{
+		switch (Rule)
+		{
+		case ESequencerThumbnailCaptureTimeLocation::FirstFrame:
+			return Sequencer.GetRootMovieSceneSequence()->GetMovieScene()->GetPlaybackRange().GetLowerBoundValue();
+		case ESequencerThumbnailCaptureTimeLocation::MiddleFrame:
+			{
+				const FFrameNumber FirstFrame = GetFrameByRule(Sequencer, ESequencerThumbnailCaptureTimeLocation::FirstFrame);
+				const FFrameNumber LastFrame = GetFrameByRule(Sequencer, ESequencerThumbnailCaptureTimeLocation::LastFrame);
+				return FirstFrame + (LastFrame - FirstFrame) / 2;
+			}
+		case ESequencerThumbnailCaptureTimeLocation::LastFrame:
+			return Sequencer.GetRootMovieSceneSequence()->GetMovieScene()->GetPlaybackRange().GetUpperBoundValue();
+			
+		case ESequencerThumbnailCaptureTimeLocation::CurrentFrame:
+			return Sequencer.GetGlobalTime().Time.FrameNumber;
+			
+		default:
+			checkNoEntry();
+			return GetFrameByRule(Sequencer, ESequencerThumbnailCaptureTimeLocation::CurrentFrame);
+		}
+	}
+}
diff --git a/Engine/Source/Editor/Sequencer/Private/Misc/Thumbnail/ThumbnailCaptureUtils.cpp b/Engine/Source/Editor/Sequencer/Private/Misc/Thumbnail/ThumbnailCaptureUtils.cpp
new file mode 100644
--- /dev/null
+++ b/Engine/Source/Editor/Sequencer/Private/Misc/Thumbnail/ThumbnailCaptureUtils.cpp
@@ -0,0 +1,53 @@
+﻿// Copyright Epic Games, Inc. All Rights Reserved.
+
+#include "ThumbnailCaptureUtils.h"
+
+#include "AssetRegistry/AssetData.h"
+#include "Containers/Array.h"
+#include "Editor.h"
+#include "Misc/ConsoleVariables.h"
+#include "Misc/ScopeExit.h"
+#include "ObjectTools.h"
+#include "Templates/Function.h"
+#include "TextureResource.h"
+
+namespace UE::Sequencer
+{
+	void CaptureThumbnailForAssetBlocking(UObject& Asset, FSequencer& Sequencer, const FSequencerThumbnailCaptureSettings& Settings)
+	{
+		if (!CVarEnableRelevantThumbnails.GetValueOnGameThread()
+			|| !CaptureThumbnailFromCameraCutBlocking(Asset, Sequencer, Settings))
+		{
+			CaptureThumbnailFromViewportBlocking(Asset);
+		}
+	}
+	
+	void SetAssetThumbnail(UObject& Asset, const TArrayView64<FColor>& Bitmap)
+	{
+		const int32 ThumbnailSize  = ThumbnailTools::DefaultThumbnailSize;
+		const int32 NumPixels = ThumbnailSize * ThumbnailSize;
+		check(Bitmap.Num() == NumPixels);
+		
+		FObjectThumbnail TempThumbnail;
+		TempThumbnail.SetImageSize(ThumbnailSize, ThumbnailSize);
+		TArray<uint8>& ThumbnailByteArray = TempThumbnail.AccessImageData();
+		constexpr int32 MemorySize = NumPixels * sizeof(FColor);
+		ThumbnailByteArray.AddUninitialized(MemorySize);
+		FMemory::Memcpy(&ThumbnailByteArray[0], &Bitmap[0], MemorySize);
+		
+		const FString ObjectFullName = FAssetData(&Asset).GetFullName();
+		UPackage* Package = Asset.GetPackage();
+		FObjectThumbnail* NewThumbnail = ThumbnailTools::CacheThumbnail(ObjectFullName, &TempThumbnail, Package);
+		if (ensure(NewThumbnail))
+		{
+			// We need to indicate that the package needs to be resaved
+			Package->MarkPackageDirty();
+			// Let the content browser know that we've changed the thumbnail
+			NewThumbnail->MarkAsDirty();
+			// Signal that the asset was changed so thumbnail pools will update
+			Asset.PostEditChange();
+			// Set that thumbnail as a valid custom thumbnail so it'll be saved out
+			NewThumbnail->SetCreatedAfterCustomThumbsEnabled();
+		}
+	}
+}
\ No newline at end of file
diff --git a/Engine/Source/Editor/Sequencer/Private/Misc/Thumbnail/ThumbnailCaptureUtils.h b/Engine/Source/Editor/Sequencer/Private/Misc/Thumbnail/ThumbnailCaptureUtils.h
new file mode 100644
--- /dev/null
+++ b/Engine/Source/Editor/Sequencer/Private/Misc/Thumbnail/ThumbnailCaptureUtils.h
@@ -0,0 +1,65 @@
+﻿// Copyright Epic Games, Inc. All Rights Reserved.
+
+#pragma once
+
+#include "MovieScene.h"
+#include "Templates/FunctionFwd.h"
+
+class UCameraComponent;
+class FSequencer;
+class UMovieScene;
+class UObject;
+struct FSequencerThumbnailCaptureSettings;
+
+namespace UE::Sequencer
+{
+	using FCaptureFrame = TFunctionRef<void(const FFrameNumber& Frame, UCameraComponent& Camera)>;
+	
+	/**
+	 * Tries to capture a relevant thumbnail and saves it into Asset.
+	 * This function may block the game thread.
+	 *
+	 * The order of evaluation is as follows:
+	 * 1. If sequence has a camera cut, determine a suitable frame using Settings, and capture.
+	 * 2. If sequence has a camera track, determine a suitable frame using Settings, and capture
+	 * 3. Try steps 1-2 recursively on sub-sequences.
+	 * 4. Fall back to capturing the viewport content.
+	 *
+	 * Post-call: The sequencer's focus state is the same as before the call.
+	 * 
+	 * @param Asset The asset into which the thumbnail is supposed to be saved.
+	 * @param RootMovieScene The root movie scene of Asset.
+	 * @param Sequencer Used to animate all objects for capturing the thumbnail
+	 * @param Settings The settings to use for capturing the thumbnail
+	 */
+	void CaptureThumbnailForAssetBlocking(
+		UObject& Asset,
+		FSequencer& Sequencer,
+		const FSequencerThumbnailCaptureSettings& Settings
+		);
+	
+	/**
+	 * Attempts to locate a camera from any camera cut track within the specified movie scene and invokes 
+	 * CaptureFrameCallback with the identified camera and the corresponding frame based on the provided Settings.
+	 *
+	 * @param Asset The asset into which the thumbnail is supposed to be saved.
+	 * @param Sequencer The sequencer used to animate objects during the thumbnail capture process.
+	 * @param Settings The configuration settings to be used for capturing the thumbnail.
+	 *
+	 * Pre-call: The sequencer's focus state is focused on MovieScene.
+	 * Post-call: The sequencer's focus state is focused on MovieScene.
+	 * 
+	 * @return True if CaptureFrameCallback was invoked; otherwise, false.
+	 */
+	bool CaptureThumbnailFromCameraCutBlocking(
+		UObject& Asset,
+		FSequencer& Sequencer,
+		const FSequencerThumbnailCaptureSettings& Settings
+		);
+	
+	/** Captures the first active viewport as thumbnail for Asset. */
+	void CaptureThumbnailFromViewportBlocking(UObject& Asset);
+
+	/** Saves Bitmap as thumbnail of Asset. */
+	void SetAssetThumbnail(UObject& Asset, const TArrayView64<FColor>& Bitmap);
+}
diff --git a/Engine/Source/Editor/Sequencer/Private/Misc/Thumbnail/ThumbnailCaptureUtils_CameraCuts.cpp b/Engine/Source/Editor/Sequencer/Private/Misc/Thumbnail/ThumbnailCaptureUtils_CameraCuts.cpp
new file mode 100644
--- /dev/null
+++ b/Engine/Source/Editor/Sequencer/Private/Misc/Thumbnail/ThumbnailCaptureUtils_CameraCuts.cpp
@@ -0,0 +1,96 @@
+﻿// Copyright Epic Games, Inc. All Rights Reserved.
+
+#include "ThumbnailCaptureUtils.h"
+
+#include "AssetRegistry/AssetData.h"
+#include "Camera/CameraComponent.h"
+#include "Containers/Array.h"
+#include "Editor.h"
+#include "Engine/TextureRenderTarget2D.h"
+#include "Kismet/KismetRenderingLibrary.h"
+#include "Misc/ScopeExit.h"
+#include "MovieSceneToolsUserSettings.h"
+#include "ObjectTools.h"
+#include "Sequencer.h"
+#include "Templates/Function.h"
+#include "TextureResource.h"
+#include "TrackEditorThumbnail/TrackThumbnailUtils.h"
+
+namespace UE::Sequencer::Private
+{
+	static UTextureRenderTarget2D* CaptureThumbnail(
+		FSequencer& Sequencer,
+		const FFrameNumber& Frame
+		)
+	{
+		constexpr int32 ThumbnailSize  = ThumbnailTools::DefaultThumbnailSize;
+
+		// Before calling PreDrawThumbnailSetupSequencer, save where the user had scrubbed, so we can revert it later.
+		// This step differs from FTrackEditorThumbnailPool::DrawThumbnails, which does not need to do this... may want to investigate. 
+		const FQualifiedFrameTime RestoreTime = Sequencer.GetGlobalTime();
+		// Positions all animated objects by jumping to the right frame...
+		MoveSceneTools::PreDrawThumbnailSetupSequencer(Sequencer, Frame);
+		ON_SCOPE_EXIT
+		{
+			// It's important to reset the time after the GPU has captured the texture, or we'll simply be capturing the current frame!
+			if (RestoreTime.Time != Sequencer.GetGlobalTime().Time)
+			{
+				Sequencer.SetGlobalTime(RestoreTime.Time);
+			}
+			
+			MoveSceneTools::PostDrawThumbnailCleanupSequencer(Sequencer);
+		};
+		
+		UCameraComponent* Component = Sequencer.GetLastEvaluatedCameraCut().Get();
+		if (!Component)
+		{
+			return nullptr;
+		}
+		
+		// Important to get camera view after PreDrawThumbnailSetupSequencer since it may have set its transform!
+		FMinimalViewInfo ViewInfo;
+		Component->GetCameraView(FApp::GetDeltaTime(), ViewInfo);
+		// It's important to capture with sRGB so the thumbnail has the right brightness
+		UTextureRenderTarget2D* RenderTarget2D = UKismetRenderingLibrary::CreateRenderTarget2D(
+			Component, ThumbnailSize, ThumbnailSize, RTF_RGBA8_SRGB
+			);
+		FTextureRenderTargetResource* RenderTarget = RenderTarget2D->GameThread_GetRenderTargetResource();
+		
+		// ... enqueues rendering commands to fill RenderTarget
+		MoveSceneTools::DrawViewportThumbnail(
+			*RenderTarget, { ThumbnailSize, ThumbnailSize }, *GWorld->Scene, ViewInfo, EThumbnailQuality::Best
+			);
+		// ... cleans up any important state modified by PreDrawThumbnailSetupSequencer.
+
+		// Need to wait on the GPU to execute the commands above
+		FRenderCommandFence Fence;
+		Fence.BeginFence();
+		Fence.Wait();
+
+		return RenderTarget2D;
+	}
+}
+
+namespace UE::Sequencer
+{
+	bool CaptureThumbnailFromCameraCutBlocking(
+		UObject& Asset,
+		FSequencer& Sequencer,
+		const FSequencerThumbnailCaptureSettings& Settings
+		)
+	{
+		const FFrameNumber Frame = GetFrameByRule(Sequencer, Settings.CaptureFrameLocationRule);
+		if (UTextureRenderTarget2D* RenderTarget2D = Private::CaptureThumbnail(Sequencer, Frame))
+		{
+			ON_SCOPE_EXIT{ RenderTarget2D->ReleaseResource(); };
+			FTextureRenderTargetResource* RenderTarget = RenderTarget2D->GameThread_GetRenderTargetResource();
+		
+			TArray<FColor> Color;
+			RenderTarget->ReadPixels(Color);
+			SetAssetThumbnail(Asset, Color);
+			return true;
+		}
+		
+		return false;
+	}
+}
\ No newline at end of file
diff --git a/Engine/Source/Editor/Sequencer/Private/Misc/Thumbnail/ThumbnailCaptureUtils_Viewport.cpp b/Engine/Source/Editor/Sequencer/Private/Misc/Thumbnail/ThumbnailCaptureUtils_Viewport.cpp
new file mode 100644
--- /dev/null
+++ b/Engine/Source/Editor/Sequencer/Private/Misc/Thumbnail/ThumbnailCaptureUtils_Viewport.cpp
@@ -0,0 +1,79 @@
+﻿// Copyright Epic Games, Inc. All Rights Reserved.
+
+#include "ThumbnailCaptureUtils.h"
+
+#include "AssetRegistry/AssetData.h"
+#include "Containers/Array.h"
+#include "ContentBrowserModule.h"
+#include "Editor.h"
+#include "IContentBrowserSingleton.h"
+#include "LevelEditorViewport.h"
+#include "UnrealClient.h"
+
+namespace UE::Sequencer::PrivateThumbnailCapture
+{
+	/** @return The active viewport if valid and otherwise any viewport that allows cinematic preview. */
+	static FViewport* FindActiveViewportThenAnyWithCinematicPreview()
+	{
+		FViewport* Viewport = GEditor->GetActiveViewport();
+		if (Viewport)
+		{
+			return Viewport;
+		}
+
+		// If there's no active viewport, find any other viewport that allows cinematic preview.
+		for (FLevelEditorViewportClient* LevelVC : GEditor->GetLevelViewportClients())
+		{
+			if (LevelVC && LevelVC->AllowsCinematicControl() && LevelVC->Viewport)
+			{
+				return LevelVC->Viewport;
+			}
+		}
+		
+		return nullptr;
+	}
+}
+
+namespace UE::Sequencer
+{
+	void CaptureThumbnailFromViewportBlocking(UObject& Asset)
+	{
+		if (!GCurrentLevelEditingViewportClient)
+		{
+			return;
+		}
+		FViewport* Viewport = PrivateThumbnailCapture::FindActiveViewportThenAnyWithCinematicPreview();
+		if (!Viewport)
+		{
+			return;
+		}
+		
+		const bool bIsInGameView = GCurrentLevelEditingViewportClient->IsInGameView();
+		FLevelEditorViewportClient* const OldViewportClient = GCurrentLevelEditingViewportClient;
+		// Remove editor widgets from the render
+		GCurrentLevelEditingViewportClient->SetGameView(true);
+		// Remove selection box around client during render
+		GCurrentLevelEditingViewportClient = nullptr;
+		ON_SCOPE_EXIT
+		{
+			// Redraw viewport to have the yellow highlight again
+			GCurrentLevelEditingViewportClient = OldViewportClient;
+			GCurrentLevelEditingViewportClient->SetGameView(bIsInGameView);
+				
+			// If turn off game view now need to make sure widget/gizmo is on
+			if (!bIsInGameView )
+			{
+				GCurrentLevelEditingViewportClient->ShowWidget(true);
+			}
+				
+			Viewport->Draw();
+		};
+			
+		// Have to re-render the requested viewport before capturing a thumbnail
+		Viewport->Draw();
+		
+		IContentBrowserSingleton& ContentBrowser = FModuleManager::LoadModuleChecked<FContentBrowserModule>("ContentBrowser").Get();
+		TArray AssetDataList { FAssetData(&Asset) };
+		ContentBrowser.CaptureThumbnailFromViewport(Viewport, AssetDataList);
+	}
+}
\ No newline at end of file
diff --git a/Engine/Source/Editor/Sequencer/Private/Sequencer.cpp b/Engine/Source/Editor/Sequencer/Private/Sequencer.cpp
--- a/Engine/Source/Editor/Sequencer/Private/Sequencer.cpp
+++ b/Engine/Source/Editor/Sequencer/Private/Sequencer.cpp
@@ -164,6 +164,7 @@
 #include "Editor/TransBuffer.h"
 #include "Sidebar/SidebarDrawerConfig.h"
 #include "Widgets/Sidebar/SequencerSelectionDrawer.h"
+#include "Misc/Thumbnail/ThumbnailCaptureUtils.h"
 
 #include "EngineModule.h"
 #include "IViewportSelectableObject.h"
@@ -5990,52 +5991,11 @@ bool FSequencer::GetCurveEditorIsVisible() const
 
 void FSequencer::SaveCurrentMovieScene()
 {
-	// Capture thumbnail
-	// Convert UObject* array to FAssetData array
-	TArray<FAssetData> AssetDataList;
-	AssetDataList.Add(FAssetData(GetCurrentAsset()));
-
-	FViewport* Viewport = GEditor->GetActiveViewport();
-
-	// If there's no active viewport, find any other viewport that allows cinematic preview.
-	if (Viewport == nullptr)
-	{
-		for (FLevelEditorViewportClient* LevelVC : GEditor->GetLevelViewportClients())
-		{
-			if ((LevelVC == nullptr) || !LevelVC->AllowsCinematicControl())
-			{
-				continue;
-			}
-
-			Viewport = LevelVC->Viewport;
-		}
-	}
-
-	if (GCurrentLevelEditingViewportClient && Viewport)
-	{
-		bool bIsInGameView = GCurrentLevelEditingViewportClient->IsInGameView();
-		GCurrentLevelEditingViewportClient->SetGameView(true);
-
-		//have to re-render the requested viewport
-		FLevelEditorViewportClient* OldViewportClient = GCurrentLevelEditingViewportClient;
-		//remove selection box around client during render
-		GCurrentLevelEditingViewportClient = NULL;
-
-		Viewport->Draw();
-
-		IContentBrowserSingleton& ContentBrowser = FModuleManager::LoadModuleChecked<FContentBrowserModule>("ContentBrowser").Get();
-		ContentBrowser.CaptureThumbnailFromViewport(Viewport, AssetDataList);
-
-		//redraw viewport to have the yellow highlight again
-		GCurrentLevelEditingViewportClient = OldViewportClient;
-		GCurrentLevelEditingViewportClient->SetGameView(bIsInGameView);
-		//if turn off game view now need to make sure widget/gizmo is on
-		if (bIsInGameView == false)
-		{
-			GCurrentLevelEditingViewportClient->ShowWidget(true);
-		}
-		Viewport->Draw();
-	}
+	UE::Sequencer::CaptureThumbnailForAssetBlocking(
+		*GetCurrentAsset(),
+		*this,
+		Settings->GetThumbnailCaptureSettings()
+		);
 
 	OnPreSaveEvent.Broadcast(*this);
 
diff --git a/Engine/Source/Editor/Sequencer/Private/SequencerSettings.cpp b/Engine/Source/Editor/Sequencer/Private/SequencerSettings.cpp
--- a/Engine/Source/Editor/Sequencer/Private/SequencerSettings.cpp
+++ b/Engine/Source/Editor/Sequencer/Private/SequencerSettings.cpp
@@ -4,6 +4,7 @@
 #include "KeyParams.h"
 #include "ISequencer.h"
 #include "SSequencer.h"
+#include "Misc/ConsoleVariables.h"
 #include "MVVM/ViewModels/ViewDensity.h"
 
 USequencerSettings::USequencerSettings( const FObjectInitializer& ObjectInitializer )
@@ -1160,3 +1161,14 @@ void USequencerSettings::SetLastFilterBarSizeCoefficient(const float bInSizeCoef
 	LastFilterBarSizeCoefficient = bInSizeCoefficient;
 	SaveConfig();
 }
+
+void USequencerSettings::SetThumbnailCaptureSettings(const FSequencerThumbnailCaptureSettings& InNewValue)
+{
+	ThumbnailCaptureSettings = InNewValue;
+	SaveConfig();
+}
+
+bool USequencerSettings::ShouldShowThumbnailCaptureSettings()
+{
+	return UE::Sequencer::CVarEnableRelevantThumbnails.GetValueOnGameThread();
+}
diff --git a/Engine/Source/Editor/Sequencer/Public/Misc/SequencerThumbnailCaptureSettings.h b/Engine/Source/Editor/Sequencer/Public/Misc/SequencerThumbnailCaptureSettings.h
new file mode 100644
--- /dev/null
+++ b/Engine/Source/Editor/Sequencer/Public/Misc/SequencerThumbnailCaptureSettings.h
@@ -0,0 +1,41 @@
+﻿// Copyright Epic Games, Inc. All Rights Reserved.
+
+#pragma once
+
+#include "ISequencer.h"
+#include "Misc/FrameNumber.h"
+#include "SequencerThumbnailCaptureSettings.generated.h"
+
+class ISequencer;
+class UMovieSceneSection;
+
+/** Specifies which frame should be captured in a track. */
+UENUM()
+enum class ESequencerThumbnailCaptureTimeLocation : uint8
+{
+	/** Use the first frame */
+	FirstFrame,
+	/** Use the middle frame */
+	MiddleFrame,
+	/** Use the last frame */
+	LastFrame,
+	/** Use the frame at which the scrubber is currently positioned */
+	CurrentFrame
+};
+
+/** Configures how a thumbnail is supposed to be captured for a level sequence */
+USTRUCT()
+struct FSequencerThumbnailCaptureSettings
+{
+	GENERATED_BODY()
+
+	/** Specifies which frame should be captured in a track. */
+	UPROPERTY(EditAnywhere, Category = General)
+	ESequencerThumbnailCaptureTimeLocation CaptureFrameLocationRule = ESequencerThumbnailCaptureTimeLocation::CurrentFrame;
+};
+
+namespace UE::Sequencer
+{
+	/** @return The frame that Rule wants to get from Sequencer */
+	SEQUENCER_API FFrameNumber GetFrameByRule(const ISequencer& Sequencer, ESequencerThumbnailCaptureTimeLocation Rule);
+}
\ No newline at end of file
diff --git a/Engine/Source/Editor/Sequencer/Public/SequencerSettings.h b/Engine/Source/Editor/Sequencer/Public/SequencerSettings.h
--- a/Engine/Source/Editor/Sequencer/Public/SequencerSettings.h
+++ b/Engine/Source/Editor/Sequencer/Public/SequencerSettings.h
@@ -4,6 +4,7 @@
 
 #include "CoreMinimal.h"
 #include "Filters/SequencerFilterBarConfig.h"
+#include "Misc/SequencerThumbnailCaptureSettings.h"
 #include "UObject/ObjectMacros.h"
 #include "UObject/Object.h"
 #include "UObject/Package.h"
@@ -512,6 +513,11 @@ public:
 
 	float GetLastFilterBarSizeCoefficient() const;
 	void SetLastFilterBarSizeCoefficient(const float bInSizeCoefficient);
+	
+	/** Gets the settings that determine how the asset thumbnail is captured when the sequence is saved. */
+	const FSequencerThumbnailCaptureSettings& GetThumbnailCaptureSettings() const { return ThumbnailCaptureSettings; }
+	/** Sets how the asset thumbnail is captured when the sequence is saved.  */
+	void SetThumbnailCaptureSettings(const FSequencerThumbnailCaptureSettings& InNewValue);
 
 protected:
 
@@ -820,6 +826,13 @@ protected:
 	UPROPERTY(config)
 	float LastFilterBarSizeCoefficient;
 
+	/** Controls how the thumbnail is captured when the sequence is saved. */
+	UPROPERTY(config, EditAnywhere, Category = General, meta = (EditCondition="ShouldShowThumbnailCaptureSettings()", EditConditionHides))
+	FSequencerThumbnailCaptureSettings ThumbnailCaptureSettings;
+
+	UFUNCTION()
+	static bool ShouldShowThumbnailCaptureSettings();
+
 	FOnEvaluateSubSequencesInIsolationChanged OnEvaluateSubSequencesInIsolationChangedEvent;
 	FOnShowSelectedNodesOnlyChanged OnShowSelectedNodesOnlyChangedEvent;
 	FOnAllowEditsModeChanged OnAllowEditsModeChangedEvent;
diff --git a/Engine/Source/Editor/Sequencer/Sequencer.Build.cs b/Engine/Source/Editor/Sequencer/Sequencer.Build.cs
--- a/Engine/Source/Editor/Sequencer/Sequencer.Build.cs
+++ b/Engine/Source/Editor/Sequencer/Sequencer.Build.cs
@@ -51,6 +51,7 @@ public class Sequencer : ModuleRules
 				"TypedElementRuntime",
 				"UniversalObjectLocator",
 				"UniversalObjectLocatorEditor",
+				"RenderCore"
 				}
 			);
 
