[Sequencer]
 - Change "Show Rotation/Location/Scale" groups to not be additive and work with the selection
 - Scroll first selected item into view on "Show All"

diff --git a/Engine/Source/Editor/Sequencer/Private/Filters/Filters/SequencerTrackFilter_HideIsolate.cpp b/Engine/Source/Editor/Sequencer/Private/Filters/Filters/SequencerTrackFilter_HideIsolate.cpp
--- a/Engine/Source/Editor/Sequencer/Private/Filters/Filters/SequencerTrackFilter_HideIsolate.cpp
+++ b/Engine/Source/Editor/Sequencer/Private/Filters/Filters/SequencerTrackFilter_HideIsolate.cpp
@@ -2,6 +2,8 @@
 
 #include "SequencerTrackFilter_HideIsolate.h"
 #include "Filters/SequencerTrackFilterCommands.h"
+#include "MVVM/ViewModels/CategoryModel.h"
+#include "MVVM/ViewModels/TrackModel.h"
 #include "MVVM/ViewModels/ViewModelIterators.h"
 
 using namespace UE::Sequencer;
@@ -130,6 +132,71 @@ void FSequencerTrackFilter_HideIsolate::UnisolateTracks(const TSet<TWeakViewMode
 	BroadcastChangedEvent();
 }
 
+void FSequencerTrackFilter_HideIsolate::IsolateCategoryGroupTracks(const TSet<TWeakViewModelPtr<IOutlinerExtension>>& InTracks
+	, const TSet<FName>& InCategoryNames
+	, const bool bInAddToExisting)
+{
+	TSharedPtr<FSequencerEditorViewModel> SequencerViewModel = GetSequencer().GetViewModel();
+	if (!SequencerViewModel.IsValid())
+	{
+		return;
+	}
+
+	if (!bInAddToExisting)
+	{
+		EmptyIsolatedTracks(false);
+	}
+
+	TSet<TWeakViewModelPtr<IOutlinerExtension>> TracksToIsolate;
+	TSet<TViewModelPtr<IOutlinerExtension>> TracksToExpand;
+
+	auto IsolateChildCategoryGroups = [this, &InCategoryNames, &TracksToIsolate, &TracksToExpand](const TViewModelPtr<IOutlinerExtension>& InTrack)
+	{
+		const TParentFirstChildIterator<FCategoryGroupModel> ChildTracks = InTrack.AsModel()->GetDescendantsOfType<FCategoryGroupModel>(true);
+		for (const TViewModelPtr<FCategoryGroupModel>& ChildCategoryGroup : ChildTracks)
+		{
+			if (InCategoryNames.Contains(ChildCategoryGroup->GetCategoryName()))
+			{
+				TracksToIsolate.Add(ChildCategoryGroup);
+
+				const TParentModelIterator<IOutlinerExtension> Ancestors = ChildCategoryGroup->GetAncestorsOfType<IOutlinerExtension>();
+				TracksToExpand.Append(Ancestors.ToArray());
+				TracksToExpand.Add(ChildCategoryGroup);
+			}
+		}
+	};
+
+	for (const TWeakViewModelPtr<IOutlinerExtension>& WeakTrack : InTracks)
+	{
+		const TViewModelPtr<IOutlinerExtension> Track = WeakTrack.Pin();
+		if (!Track.IsValid())
+		{
+			continue;
+		}
+
+		const TViewModelPtr<FCategoryGroupModel> CategoryGroupModel = Track.AsModel()->FindAncestorOfType<FCategoryGroupModel>(true);
+		if (CategoryGroupModel.IsValid())
+		{
+			const TViewModelPtr<FTrackModel> ParentTrack = Track.AsModel()->FindAncestorOfType<FTrackModel>();
+			if (ParentTrack.IsValid())
+			{
+				IsolateChildCategoryGroups(ParentTrack);
+			}
+		}
+		else
+		{
+			IsolateChildCategoryGroups(Track);
+		}
+	}
+
+	IsolateTracks(TracksToIsolate, true);
+
+	for (const TViewModelPtr<IOutlinerExtension>& Track : TracksToExpand)
+	{
+		Track->SetExpansion(true);
+	}
+}
+
 void FSequencerTrackFilter_HideIsolate::ShowAllTracks()
 {
 	HiddenTracks.Empty();
diff --git a/Engine/Source/Editor/Sequencer/Private/Filters/Filters/SequencerTrackFilter_HideIsolate.h b/Engine/Source/Editor/Sequencer/Private/Filters/Filters/SequencerTrackFilter_HideIsolate.h
--- a/Engine/Source/Editor/Sequencer/Private/Filters/Filters/SequencerTrackFilter_HideIsolate.h
+++ b/Engine/Source/Editor/Sequencer/Private/Filters/Filters/SequencerTrackFilter_HideIsolate.h
@@ -18,6 +18,10 @@ public:
 	void IsolateTracks(const TSet<UE::Sequencer::TWeakViewModelPtr<UE::Sequencer::IOutlinerExtension>>& InTracks, const bool bInAddToExisting);
 	void UnisolateTracks(const TSet<UE::Sequencer::TWeakViewModelPtr<UE::Sequencer::IOutlinerExtension>>& InTracks);
 
+	void IsolateCategoryGroupTracks(const TSet<UE::Sequencer::TWeakViewModelPtr<UE::Sequencer::IOutlinerExtension>>& InTracks
+		, const TSet<FName>& InCategoryNames
+		, const bool bInAddToExisting);
+
 	void ShowAllTracks();
 
 	bool HasHiddenTracks() const;
diff --git a/Engine/Source/Editor/Sequencer/Private/Filters/SequencerFilterBar.cpp b/Engine/Source/Editor/Sequencer/Private/Filters/SequencerFilterBar.cpp
--- a/Engine/Source/Editor/Sequencer/Private/Filters/SequencerFilterBar.cpp
+++ b/Engine/Source/Editor/Sequencer/Private/Filters/SequencerFilterBar.cpp
@@ -25,6 +25,7 @@
 #include "MVVM/CurveEditorExtension.h"
 #include "MVVM/Selection/Selection.h"
 #include "MVVM/ViewModels/CategoryModel.h"
+#include "MVVM/Views/SOutlinerView.h"
 #include "Sequencer.h"
 #include "SequencerLog.h"
 #include "SSequencer.h"
@@ -351,6 +352,16 @@ void FSequencerFilterBar::ShowAllTracks()
 {
 	HideIsolateFilter->ShowAllTracks();
 
+	if (const TSharedPtr<SSequencer> SequencerWidget = StaticCastSharedRef<SSequencer>(Sequencer.GetSequencerWidget()))
+	{
+		const TSharedPtr<FSequencerSelection> Selection = Sequencer.GetViewModel()->GetSelection();
+		const TArray<TWeakViewModelPtr<IOutlinerExtension>> SelectedTracks = Selection->Outliner.GetSelected().Array();
+		if (SelectedTracks.Num() > 0)
+		{
+			SequencerWidget->GetTreeView()->RequestScrollIntoView(SelectedTracks[0]);
+		}
+	}
+
 	RequestFilterUpdate();
 }
 
@@ -1435,49 +1446,6 @@ bool FSequencerFilterBar::HasSelectedTracks() const
 	return !GetSelectedTracksOrAll().IsEmpty();
 }
 
-void FSequencerFilterBar::AddCategoryGroupIsolatedTracks(const TSet<TWeakViewModelPtr<IOutlinerExtension>>& InTracks
-	, const TSet<FName>& InCategoryNames
-	, const bool bInAddToExisting) const
-{
-	TSharedPtr<FSequencerEditorViewModel> SequencerViewModel = GetSequencer().GetViewModel();
-	if (!SequencerViewModel.IsValid())
-	{
-		return;
-	}
-
-	if (!bInAddToExisting)
-	{
-		HideIsolateFilter->EmptyIsolatedTracks(false);
-	}
-
-	TSet<TWeakViewModelPtr<IOutlinerExtension>> TracksToIsolate;
-	TSet<TViewModelPtr<IOutlinerExtension>> TracksToExpand;
-
-	const FViewModelPtr RootModel = SequencerViewModel->GetRootModel();
-	const TParentFirstChildIterator<IOutlinerExtension> AllNodes = RootModel->GetDescendantsOfType<IOutlinerExtension>();
-	for (const TViewModelPtr<IOutlinerExtension>& Node : AllNodes)
-	{
-		if (const TViewModelPtr<FCategoryGroupModel> CategoryGroupModel = Node.ImplicitCast())
-		{
-			if (InCategoryNames.Contains(CategoryGroupModel->GetCategoryName()))
-			{
-				TracksToIsolate.Add(CategoryGroupModel);
-
-				const TParentModelIterator<IOutlinerExtension> Ancestors = CategoryGroupModel->GetAncestorsOfType<IOutlinerExtension>();
-				TracksToExpand.Append(Ancestors.ToArray());
-				TracksToExpand.Add(CategoryGroupModel);
-			}
-		}
-	}
-
-	HideIsolateFilter->IsolateTracks(TracksToIsolate, true);
-
-	for (const TViewModelPtr<IOutlinerExtension>& Track : TracksToExpand)
-	{
-		Track->SetExpansion(true);
-	}
-}
-
 void FSequencerFilterBar::HideSelectedTracks()
 {
 	const bool bAddToExisting = !FSlateApplication::Get().GetModifierKeys().AreModifersDown(EModifierKey::Shift);
@@ -1494,17 +1462,17 @@ void FSequencerFilterBar::IsolateSelectedTracks()
 
 void FSequencerFilterBar::ShowOnlyLocationCategoryGroups()
 {
-	AddCategoryGroupIsolatedTracks(GetSelectedTracksOrAll(), { TEXT("Location") }, true);
+	HideIsolateFilter->IsolateCategoryGroupTracks(GetSelectedTracksOrAll(), { TEXT("Location") }, false);
 }
 
 void FSequencerFilterBar::ShowOnlyRotationCategoryGroups()
 {
-	AddCategoryGroupIsolatedTracks(GetSelectedTracksOrAll(), { TEXT("Rotation") }, true);
+	HideIsolateFilter->IsolateCategoryGroupTracks(GetSelectedTracksOrAll(), { TEXT("Rotation") }, false);
 }
 
 void FSequencerFilterBar::ShowOnlyScaleCategoryGroups()
 {
-	AddCategoryGroupIsolatedTracks(GetSelectedTracksOrAll(), { TEXT("Scale") }, true);
+	HideIsolateFilter->IsolateCategoryGroupTracks(GetSelectedTracksOrAll(), { TEXT("Scale") }, false);
 }
 
 void FSequencerFilterBar::SetTrackParentsExpanded(const TViewModelPtr<IOutlinerExtension>& InNode, const bool bInExpanded)
diff --git a/Engine/Source/Editor/Sequencer/Private/Filters/SequencerFilterBar.h b/Engine/Source/Editor/Sequencer/Private/Filters/SequencerFilterBar.h
--- a/Engine/Source/Editor/Sequencer/Private/Filters/SequencerFilterBar.h
+++ b/Engine/Source/Editor/Sequencer/Private/Filters/SequencerFilterBar.h
@@ -215,10 +215,6 @@ public:
 
 	TSet<UE::Sequencer::TWeakViewModelPtr<UE::Sequencer::IOutlinerExtension>> GetSelectedTracksOrAll() const;
 
-	void AddCategoryGroupIsolatedTracks(const TSet<UE::Sequencer::TWeakViewModelPtr<UE::Sequencer::IOutlinerExtension>>& InTracks
-		, const TSet<FName>& InCategoryNames
-		, const bool bInAddToExisting) const;
-
 	void SetTrackParentsExpanded(const UE::Sequencer::TViewModelPtr<UE::Sequencer::IOutlinerExtension>& InNode, const bool bInExpanded);
 
 	UWorld* GetWorld() const;
