[Sequencer] Fix SFilterBarClippingHorizontalBox not opening and crash

diff --git a/Engine/Source/Editor/Sequencer/Private/Filters/Widgets/SFilterBarClippingHorizontalBox.cpp b/Engine/Source/Editor/Sequencer/Private/Filters/Widgets/SFilterBarClippingHorizontalBox.cpp
--- a/Engine/Source/Editor/Sequencer/Private/Filters/Widgets/SFilterBarClippingHorizontalBox.cpp
+++ b/Engine/Source/Editor/Sequencer/Private/Filters/Widgets/SFilterBarClippingHorizontalBox.cpp
@@ -104,7 +104,6 @@ TSharedRef<SComboButton> SFilterBarClippingHorizontalBox::CreateWrapButton()
 		.ToolTipText(LOCTEXT("ExpandFilterBar", "Click to expand the filter bar"))
 		.OnGetMenuContent(OnWrapButtonClicked)
 		.Cursor(EMouseCursor::Default)
-		.OnMenuOpenChanged(this, &SFilterBarClippingHorizontalBox::OnWrapButtonOpenChanged)
 		.IsFocusable(true)
 		.ButtonContent()
 		[
@@ -118,38 +117,6 @@ TSharedRef<SComboButton> SFilterBarClippingHorizontalBox::CreateWrapButton()
 	return WrapButton;
 }
 
-void SFilterBarClippingHorizontalBox::OnWrapButtonOpenChanged(const bool bInIsOpen)
-{
-	if (bInIsOpen && !WrapButtonOpenTimer.IsValid())
-	{
-		WrapButtonOpenTimer = RegisterActiveTimer(0.0f, FWidgetActiveTimerDelegate::CreateSP(this, &SFilterBarClippingHorizontalBox::UpdateWrapButtonStatus));
-	}
-	else if(!bInIsOpen && WrapButtonOpenTimer.IsValid())
-	{
-		UnRegisterActiveTimer(WrapButtonOpenTimer.ToSharedRef());
-		WrapButtonOpenTimer.Reset();
-	}
-}
-
-EActiveTimerReturnType SFilterBarClippingHorizontalBox::UpdateWrapButtonStatus(const double InCurrentTime, const float InDeltaTime)
-{
-	const TSharedPtr<SComboButton> WrapButton = WeakWrapButton.Pin();
-	if (!WrapButton.IsValid())
-	{
-		return EActiveTimerReturnType::Stop;
-	}
-
-	if (LastClippedIndex != ClippedIndex || !WrapButton->IsOpen())
-	{
-		WrapButton->SetIsOpen(false);
-		WrapButtonOpenTimer.Reset();
-
-		return EActiveTimerReturnType::Stop;
-	}
-
-	return EActiveTimerReturnType::Continue;
-}
-
 TSharedRef<SWidget> SFilterBarClippingHorizontalBox::WrapVerticalListWithHeading(const TSharedRef<SWidget>& InWidget, const FPointerEventHandler InMouseButtonUpEvent)
 {
 	return SNew(SVerticalBox)
diff --git a/Engine/Source/Editor/Sequencer/Private/Filters/Widgets/SFilterBarClippingHorizontalBox.h b/Engine/Source/Editor/Sequencer/Private/Filters/Widgets/SFilterBarClippingHorizontalBox.h
--- a/Engine/Source/Editor/Sequencer/Private/Filters/Widgets/SFilterBarClippingHorizontalBox.h
+++ b/Engine/Source/Editor/Sequencer/Private/Filters/Widgets/SFilterBarClippingHorizontalBox.h
@@ -35,9 +35,6 @@ protected:
 	virtual void OnArrangeChildren(const FGeometry& InAllottedGeometry, FArrangedChildren& ArrangedChildren) const override;
 	//~ End SPanel
 
-	void OnWrapButtonOpenChanged(const bool bInIsOpen);
-	EActiveTimerReturnType UpdateWrapButtonStatus(const double InCurrentTime, const float InDeltaTime);
-
 	/** The button that is displayed when a toolbar or menubar is clipped */
 	TWeakPtr<SComboButton> WeakWrapButton;
 
