[Sequencer] Fix crash on opening multiple sequences

diff --git a/Engine/Source/Editor/Sequencer/Private/Filters/Filters/SequencerTrackFilter_Selected.cpp b/Engine/Source/Editor/Sequencer/Private/Filters/Filters/SequencerTrackFilter_Selected.cpp
--- a/Engine/Source/Editor/Sequencer/Private/Filters/Filters/SequencerTrackFilter_Selected.cpp
+++ b/Engine/Source/Editor/Sequencer/Private/Filters/Filters/SequencerTrackFilter_Selected.cpp
@@ -13,17 +13,27 @@
 FSequencerTrackFilter_Selected::FSequencerTrackFilter_Selected(ISequencerTrackFilters& InFilterInterface, TSharedPtr<FFilterCategory> InCategory)
 	: FSequencerTrackFilter(InFilterInterface, MoveTemp(InCategory))
 {
-	if (const FEditorModeTools* const EditorMode = GetEditorModeManager())
+}
+
+FSequencerTrackFilter_Selected::~FSequencerTrackFilter_Selected()
+{
+	UnbindSelectionChanged();
+}
+
+void FSequencerTrackFilter_Selected::BindSelectionChanged()
+{
+	if (!OnSelectionChangedHandle.IsValid())
 	{
-		EditorMode->GetSelectedObjects()->SelectionChangedEvent.AddRaw(this, &FSequencerTrackFilter_Selected::OnSelectionChanged);
+		OnSelectionChangedHandle = USelection::SelectionChangedEvent.AddRaw(this, &FSequencerTrackFilter_Selected::OnSelectionChanged);
 	}
 }
 
-FSequencerTrackFilter_Selected::~FSequencerTrackFilter_Selected()
+void FSequencerTrackFilter_Selected::UnbindSelectionChanged()
 {
-	if (const FEditorModeTools* const EditorMode = GetEditorModeManager())
+	if (OnSelectionChangedHandle.IsValid())
 	{
-		EditorMode->GetSelectedObjects()->SelectionChangedEvent.RemoveAll(this);
+		USelection::SelectionChangedEvent.Remove(OnSelectionChangedHandle);
+		OnSelectionChangedHandle.Reset();
 	}
 }
 
@@ -64,6 +74,20 @@ TSharedPtr<FUICommandInfo> FSequencerTrackFilter_Selected::GetToggleCommand() co
 	return FSequencerTrackFilterCommands::Get().ToggleFilter_Selected;
 }
 
+void FSequencerTrackFilter_Selected::ActiveStateChanged(const bool bInActive)
+{
+	FSequencerTrackFilter::ActiveStateChanged(bInActive);
+
+	if (bInActive)
+	{
+		BindSelectionChanged();
+	}
+	else
+	{
+		UnbindSelectionChanged();
+	}
+}
+
 FText FSequencerTrackFilter_Selected::GetDisplayName() const
 {
 	return LOCTEXT("SequencerTrackFilter_Selected", "Selected");
diff --git a/Engine/Source/Editor/Sequencer/Private/Filters/Filters/SequencerTrackFilter_Selected.h b/Engine/Source/Editor/Sequencer/Private/Filters/Filters/SequencerTrackFilter_Selected.h
--- a/Engine/Source/Editor/Sequencer/Private/Filters/Filters/SequencerTrackFilter_Selected.h
+++ b/Engine/Source/Editor/Sequencer/Private/Filters/Filters/SequencerTrackFilter_Selected.h
@@ -17,6 +17,7 @@ public:
 	virtual bool ShouldUpdateOnTrackValueChanged() const override;
 	virtual FText GetDefaultToolTipText() const override;
 	virtual TSharedPtr<FUICommandInfo> GetToggleCommand() const override;
+	virtual void ActiveStateChanged(const bool bInActive) override;
 	//~ End FSequencerTrackFilter
 
 	//~ Begin FFilterBase
@@ -32,8 +33,13 @@ public:
 	void ToggleShowOnlySelectedTracks();
 
 protected:
+	void BindSelectionChanged();
+	void UnbindSelectionChanged();
+
 	TSharedPtr<ILevelEditor> GetLevelEditor() const;
 	FEditorModeTools* GetEditorModeManager() const;
 
 	void OnSelectionChanged(UObject* const InObject);
+
+	FDelegateHandle OnSelectionChangedHandle;
 };
diff --git a/Engine/Source/Editor/Sequencer/Private/Filters/SequencerFilterBar.cpp b/Engine/Source/Editor/Sequencer/Private/Filters/SequencerFilterBar.cpp
--- a/Engine/Source/Editor/Sequencer/Private/Filters/SequencerFilterBar.cpp
+++ b/Engine/Source/Editor/Sequencer/Private/Filters/SequencerFilterBar.cpp
@@ -79,6 +79,9 @@ FSequencerFilterBar::~FSequencerFilterBar()
     LevelFilter->OnChanged().RemoveAll(this);
     HideIsolateFilter->OnChanged().RemoveAll(this);
     SelectedFilter->OnChanged().RemoveAll(this);
+
+	CommonFilters.Reset();
+	InternalFilters.Reset();
 }
 
 TSharedPtr<ICustomTextFilter<FSequencerTrackFilterType>> FSequencerFilterBar::CreateTextFilter()
diff --git a/Engine/Source/Editor/Sequencer/Private/Filters/SequencerFilterBar.h b/Engine/Source/Editor/Sequencer/Private/Filters/SequencerFilterBar.h
--- a/Engine/Source/Editor/Sequencer/Private/Filters/SequencerFilterBar.h
+++ b/Engine/Source/Editor/Sequencer/Private/Filters/SequencerFilterBar.h
@@ -253,8 +253,8 @@ protected:
 	TSharedRef<FFilterCategory> MiscCategory;
 	TSharedRef<FFilterCategory> TransientCategory;
 
-	TSharedRef<FSequencerTrackFilterCollection> CommonFilters;
-	TSharedRef<FSequencerTrackFilterCollection> InternalFilters;
+	TSharedPtr<FSequencerTrackFilterCollection> CommonFilters;
+	TSharedPtr<FSequencerTrackFilterCollection> InternalFilters;
 
 	TSharedRef<FSequencerTrackFilter_Text> TextFilter;
 	TSharedRef<FSequencerTrackFilter_HideIsolate> HideIsolateFilter;
