Sequencer: Don't allow keys on locked sections to be moved

diff --git a/Engine/Source/Editor/Sequencer/Private/Sequencer.cpp b/Engine/Source/Editor/Sequencer/Private/Sequencer.cpp
--- a/Engine/Source/Editor/Sequencer/Private/Sequencer.cpp
+++ b/Engine/Source/Editor/Sequencer/Private/Sequencer.cpp
@@ -2112,6 +2112,13 @@ void FSequencer::TransformSelectedKeysAndSections(FFrameTime InDeltaTime, float
 				{
 					continue;
 				}
+
+				// Skip any locked sections
+				if (ChannelInfo.OwningSection->IsLocked())
+				{
+					continue;
+				}
+
 				TPair<TArray<FFrameNumber>, TArray<FKeyHandle>>& KeyTimesScratch = ChannelsAndKeyTimes.FindOrAdd(Channel);
 				const int32 NumKeys = ChannelInfo.KeyHandles.Num();
 				KeyTimesScratch.Key.Reset(NumKeys);
@@ -2169,6 +2176,12 @@ void FSequencer::TransformSelectedKeysAndSections(FFrameTime InDeltaTime, float
 		// Dilate the sections
 		for (UMovieSceneSection* Section : SelectedSections)
 		{
+			// Skip any locked sections
+			if (Section->IsLocked())
+			{
+				continue;
+			}
+
 			TRangeBound<FFrameNumber> LowerBound = Section->GetRange().GetLowerBound();
 			TRangeBound<FFrameNumber> UpperBound = Section->GetRange().GetUpperBound();
 
