[Sequencer] Update filters on selected text expression selection changed

diff --git a/Engine/Source/Editor/Sequencer/Private/Filters/TextExpressions/SequencerTextFilterExpression_Selected.cpp b/Engine/Source/Editor/Sequencer/Private/Filters/TextExpressions/SequencerTextFilterExpression_Selected.cpp
--- a/Engine/Source/Editor/Sequencer/Private/Filters/TextExpressions/SequencerTextFilterExpression_Selected.cpp
+++ b/Engine/Source/Editor/Sequencer/Private/Filters/TextExpressions/SequencerTextFilterExpression_Selected.cpp
@@ -1,7 +1,6 @@
 ﻿// Copyright Epic Games, Inc. All Rights Reserved.
 
-#include "SequencerTextFilterExpression_Selected.h"
-#include "EditorModeManager.h"
+#include "Filters/TextExpressions/SequencerTextFilterExpression_Selected.h"
 #include "Filters/Filters/SequencerTrackFilter_Selected.h"
 #include "Selection.h"
 #include "Sequencer.h"
@@ -13,17 +12,28 @@ using namespace UE::Sequencer;
 FSequencerTextFilterExpression_Selected::FSequencerTextFilterExpression_Selected(ISequencerTrackFilters& InFilterInterface)
 	: FSequencerTextFilterExpressionContext(InFilterInterface)
 {
-	if (const TSharedPtr<IToolkitHost> ToolkitHost = FilterInterface.GetSequencer().GetToolkitHost())
+	BindSelectionChanged();
+}
+
+FSequencerTextFilterExpression_Selected::~FSequencerTextFilterExpression_Selected()
+{
+	UnbindSelectionChanged();
+}
+
+void FSequencerTextFilterExpression_Selected::BindSelectionChanged()
+{
+	if (!OnSelectionChangedHandle.IsValid())
 	{
-		ToolkitHost->GetEditorModeManager().GetSelectedObjects()->SelectionChangedEvent.AddRaw(this, &FSequencerTextFilterExpression_Selected::OnSelectionChanged);
+		OnSelectionChangedHandle = USelection::SelectionChangedEvent.AddRaw(this, &FSequencerTextFilterExpression_Selected::OnSelectionChanged);
 	}
 }
 
-FSequencerTextFilterExpression_Selected::~FSequencerTextFilterExpression_Selected()
+void FSequencerTextFilterExpression_Selected::UnbindSelectionChanged()
 {
-	if (const TSharedPtr<IToolkitHost> ToolkitHost = FilterInterface.GetSequencer().GetToolkitHost())
+	if (OnSelectionChangedHandle.IsValid())
 	{
-		ToolkitHost->GetEditorModeManager().GetSelectedObjects()->SelectionChangedEvent.RemoveAll(this);
+		USelection::SelectionChangedEvent.Remove(OnSelectionChangedHandle);
+		OnSelectionChangedHandle.Reset();
 	}
 }
 
@@ -59,7 +69,7 @@ bool FSequencerTextFilterExpression_Selected::TestComplexExpression(const FName&
 
 void FSequencerTextFilterExpression_Selected::OnSelectionChanged(UObject* const InObject)
 {
-	//FilterBar.RequestFilterUpdate();
+	FilterInterface.RequestFilterUpdate();
 }
 
 #undef LOCTEXT_NAMESPACE
diff --git a/Engine/Source/Editor/Sequencer/Private/Filters/TextExpressions/SequencerTextFilterExpression_Selected.h b/Engine/Source/Editor/Sequencer/Private/Filters/TextExpressions/SequencerTextFilterExpression_Selected.h
--- a/Engine/Source/Editor/Sequencer/Private/Filters/TextExpressions/SequencerTextFilterExpression_Selected.h
+++ b/Engine/Source/Editor/Sequencer/Private/Filters/TextExpressions/SequencerTextFilterExpression_Selected.h
@@ -24,5 +24,10 @@ public:
 	//~ End ITextFilterExpressionContext
 
 protected:
+	void BindSelectionChanged();
+	void UnbindSelectionChanged();
+
 	void OnSelectionChanged(UObject* const InObject);
+
+	FDelegateHandle OnSelectionChangedHandle;
 };
