[Sequencer] Reduce filter processing when not filtering pinned items

diff --git a/Engine/Source/Editor/Sequencer/Private/Filters/SequencerFilterBar.cpp b/Engine/Source/Editor/Sequencer/Private/Filters/SequencerFilterBar.cpp
--- a/Engine/Source/Editor/Sequencer/Private/Filters/SequencerFilterBar.cpp
+++ b/Engine/Source/Editor/Sequencer/Private/Filters/SequencerFilterBar.cpp
@@ -1378,14 +1378,15 @@ bool FSequencerFilterBar::FilterNodesRecursive(const bool bInHasActiveFilter, co
 	const USequencerSettings* const SequencerSettings = Sequencer.GetSequencerSettings();
     check(IsValid(SequencerSettings));
 
-	const TSharedPtr<IPinnableExtension> Pinnable = InStartNode.AsModel()->FindAncestorOfType<IPinnableExtension>(true);
-	const bool bIsPinned = Pinnable && Pinnable->IsPinned();
-
 	// Pinning overrides all other filters
-	if (!SequencerSettings->GetIncludePinnedInFilter() && bIsPinned)
+	if (!SequencerSettings->GetIncludePinnedInFilter())
 	{
-		FilterData.FilterInParentChildNodes(InStartNode, true, true, true);
-		return true;
+		const TSharedPtr<IPinnableExtension> Pinnable = InStartNode.AsModel()->FindAncestorOfType<IPinnableExtension>(true);
+		if (Pinnable.IsValid() && Pinnable->IsPinned())
+		{
+			FilterData.FilterInParentChildNodes(InStartNode, true, true, true);
+			return true;
+		}
 	}
 
 	const bool bPassedTextFilter = !TextFilter->IsActive() || TextFilter->PassesFilter(InStartNode);
