[Sequencer] Advanced filtering: Improve "NAME" and "CLASS" text expressions

diff --git a/Engine/Source/Editor/Sequencer/Private/Filters/TextExpressions/SequencerTextFilterExpression_Class.cpp b/Engine/Source/Editor/Sequencer/Private/Filters/TextExpressions/SequencerTextFilterExpression_Class.cpp
--- a/Engine/Source/Editor/Sequencer/Private/Filters/TextExpressions/SequencerTextFilterExpression_Class.cpp
+++ b/Engine/Source/Editor/Sequencer/Private/Filters/TextExpressions/SequencerTextFilterExpression_Class.cpp
@@ -38,22 +38,6 @@ bool FSequencerTextFilterExpression_Class::TestComplexExpression(const FName& In
 		return true;
 	}
 
-	if (const TViewModelPtr<ITrackExtension> TrackExtension = FilterItem->FindAncestorOfType<ITrackExtension>())
-	{
-		for (const TViewModelPtr<FSectionModel>& SectionModel : TrackExtension->GetSectionModels().IterateSubList<FSectionModel>())
-		{
-			if (const auto SectionInterface = SectionModel->GetSectionInterface())
-			{
-				const FString SectionTitle = SectionInterface->GetSectionTitle().ToString();
-
-				if (TextFilterUtils::TestComplexExpression(SectionTitle, InValue, InComparisonOperation, InTextComparisonMode))
-				{
-					return true;
-				}
-			}
-		}
-	}
-
 	if (WeakTrackObject.IsValid())
 	{
 		const FString TrackClassName = WeakTrackObject->GetClass()->GetName();
@@ -61,19 +45,21 @@ bool FSequencerTextFilterExpression_Class::TestComplexExpression(const FName& In
 		{
 			return true;
 		}
-		return false;
 	}
 
 	ISequencer& Sequencer = FilterInterface.GetSequencer();
 
 	UObject* const BoundObject = FSequencerTrackFilter::ResolveTrackBoundObject(Sequencer, FilterItem, FilterInterface.GetFilterData());
-	if (IsValid(BoundObject)
-		&& !TextFilterUtils::TestComplexExpression(BoundObject->GetClass()->GetName(), InValue, InComparisonOperation, InTextComparisonMode))
+	if (IsValid(BoundObject))
 	{
-		return false;
+		const FString BoundObjectClassName = BoundObject->GetClass()->GetName();
+		if (TextFilterUtils::TestComplexExpression(BoundObjectClassName, InValue, InComparisonOperation, InTextComparisonMode))
+		{
+			return true;
+		}
 	}
 
-	return true;
+	return false;
 }
 
 #undef LOCTEXT_NAMESPACE
diff --git a/Engine/Source/Editor/Sequencer/Private/Filters/TextExpressions/SequencerTextFilterExpression_Name.cpp b/Engine/Source/Editor/Sequencer/Private/Filters/TextExpressions/SequencerTextFilterExpression_Name.cpp
--- a/Engine/Source/Editor/Sequencer/Private/Filters/TextExpressions/SequencerTextFilterExpression_Name.cpp
+++ b/Engine/Source/Editor/Sequencer/Private/Filters/TextExpressions/SequencerTextFilterExpression_Name.cpp
@@ -1,6 +1,7 @@
 ﻿// Copyright Epic Games, Inc. All Rights Reserved.
 
 #include "SequencerTextFilterExpression_Name.h"
+#include "MovieSceneNameableTrack.h"
 #include "MVVM/ViewModelPtr.h"
 #include "MVVM/Extensions/IOutlinerExtension.h"
 
@@ -38,21 +39,33 @@ bool FSequencerTextFilterExpression_Name::TestComplexExpression(const FName& InK
 		return true;
 	}
 
-	if (WeakTrackObject.IsValid()
-		&& !TextFilterUtils::TestComplexExpression(WeakTrackObject->GetName(), InValue, InComparisonOperation, InTextComparisonMode))
+	if (WeakTrackObject.IsValid())
 	{
-		return false;
+		if (TextFilterUtils::TestComplexExpression(WeakTrackObject->GetName(), InValue, InComparisonOperation, InTextComparisonMode))
+		{
+			return true;
+		}
+
+		UMovieSceneNameableTrack* const NameableTrack = Cast<UMovieSceneNameableTrack>(WeakTrackObject);
+		if (IsValid(NameableTrack))
+		{
+			const FString DisplayName = NameableTrack->GetDisplayName().ToString();
+			if (TextFilterUtils::TestComplexExpression(DisplayName, InValue, InComparisonOperation, InTextComparisonMode))
+			{
+				return true;
+			}
+		}
 	}
 
 	if (const TViewModelPtr<IOutlinerExtension> OutlinerExtension = FilterItem.ImplicitCast())
 	{
-		if (!TextFilterUtils::TestComplexExpression(OutlinerExtension->GetLabel().ToString(), InValue, InComparisonOperation, InTextComparisonMode))
+		if (TextFilterUtils::TestComplexExpression(OutlinerExtension->GetLabel().ToString(), InValue, InComparisonOperation, InTextComparisonMode))
 		{
-			return false;
+			return true;
 		}
 	}
 
-	return true;
+	return false;
 }
 
 #undef LOCTEXT_NAMESPACE
