[Sequencer]
 - Fix bug with enabling all filters on double click of filter
 - Fix filter bar context menu crash
 - Add max height and scroll box to SFilterBarClippingHorizontalBox dropdown filter menu

diff --git a/Engine/Source/Editor/Sequencer/Private/Filters/Menus/SequencerFilterBarContextMenu.h b/Engine/Source/Editor/Sequencer/Private/Filters/Menus/SequencerFilterBarContextMenu.h
--- a/Engine/Source/Editor/Sequencer/Private/Filters/Menus/SequencerFilterBarContextMenu.h
+++ b/Engine/Source/Editor/Sequencer/Private/Filters/Menus/SequencerFilterBarContextMenu.h
@@ -8,7 +8,7 @@ class FSequencerFilterBar;
 class SWidget;
 class UToolMenu;
 
-class FSequencerFilterBarContextMenu
+class FSequencerFilterBarContextMenu : public TSharedFromThis<FSequencerFilterBarContextMenu>
 {
 public:
 	TSharedRef<SWidget> CreateMenu(const TSharedRef<FSequencerFilterBar>& InFilterBar);
diff --git a/Engine/Source/Editor/Sequencer/Private/Filters/Menus/SequencerFilterMenuContext.h b/Engine/Source/Editor/Sequencer/Private/Filters/Menus/SequencerFilterMenuContext.h
--- a/Engine/Source/Editor/Sequencer/Private/Filters/Menus/SequencerFilterMenuContext.h
+++ b/Engine/Source/Editor/Sequencer/Private/Filters/Menus/SequencerFilterMenuContext.h
@@ -4,6 +4,7 @@
 
 #include "Filters/SequencerTrackFilterBase.h"
 #include "Filters/Widgets/SSequencerFilter.h"
+#include "SequencerFilterBarContext.h"
 #include "Templates/SharedPointer.h"
 #include "SequencerFilterMenuContext.generated.h"
 
@@ -28,6 +29,8 @@ public:
 		return WeakFilterWidget.IsValid() ? WeakFilterWidget.Pin()->GetFilter() : nullptr;
 	}
 
+	FOnPopulateFilterBarMenu OnPopulateFilterBarMenu;
+
 protected:
 	TWeakPtr<SSequencerFilter> WeakFilterWidget;
 };
diff --git a/Engine/Source/Editor/Sequencer/Private/Filters/Menus/SequencerTrackFilterContextMenu.cpp b/Engine/Source/Editor/Sequencer/Private/Filters/Menus/SequencerTrackFilterContextMenu.cpp
--- a/Engine/Source/Editor/Sequencer/Private/Filters/Menus/SequencerTrackFilterContextMenu.cpp
+++ b/Engine/Source/Editor/Sequencer/Private/Filters/Menus/SequencerTrackFilterContextMenu.cpp
@@ -23,7 +23,13 @@ TSharedRef<SWidget> FSequencerTrackFilterContextMenu::CreateMenuWidget(const TSh
 	if (!ToolMenus->IsMenuRegistered(FilterMenuName))
 	{
 		UToolMenu* const Menu = ToolMenus->RegisterMenu(FilterMenuName);
-		Menu->AddDynamicSection(TEXT("PopulateMenu"), FNewToolMenuDelegate::CreateRaw(this, &FSequencerTrackFilterContextMenu::PopulateMenu));
+		Menu->AddDynamicSection(TEXT("PopulateMenu"), FNewToolMenuDelegate::CreateLambda([this](UToolMenu* const InMenu)
+			{
+				if (USequencerFilterMenuContext* const Context = InMenu->FindContext<USequencerFilterMenuContext>())
+				{
+					Context->OnPopulateFilterBarMenu.ExecuteIfBound(InMenu);
+				}
+			}));
 	}
 
 	const TSharedPtr<FSequencerTrackFilter> Filter = InFilterWidget->GetFilter();
@@ -31,9 +37,9 @@ TSharedRef<SWidget> FSequencerTrackFilterContextMenu::CreateMenuWidget(const TSh
 
 	USequencerFilterMenuContext* const ContextObject = NewObject<USequencerFilterMenuContext>();
 	ContextObject->Init(InFilterWidget);
+	ContextObject->OnPopulateFilterBarMenu = FOnPopulateFilterBarMenu::CreateRaw(this, &FSequencerTrackFilterContextMenu::PopulateMenu);
 
 	const FToolMenuContext MenuContext(CommandList, nullptr, ContextObject);
-
 	return UToolMenus::Get()->GenerateWidget(FilterMenuName, MenuContext);
 }
 
diff --git a/Engine/Source/Editor/Sequencer/Private/Filters/Menus/SequencerTrackFilterContextMenu.h b/Engine/Source/Editor/Sequencer/Private/Filters/Menus/SequencerTrackFilterContextMenu.h
--- a/Engine/Source/Editor/Sequencer/Private/Filters/Menus/SequencerTrackFilterContextMenu.h
+++ b/Engine/Source/Editor/Sequencer/Private/Filters/Menus/SequencerTrackFilterContextMenu.h
@@ -9,7 +9,7 @@ class SSequencerFilter;
 class SWidget;
 class UToolMenu;
 
-class FSequencerTrackFilterContextMenu
+class FSequencerTrackFilterContextMenu : public TSharedFromThis<FSequencerTrackFilterContextMenu>
 {
 public:
 	TSharedRef<SWidget> CreateMenuWidget(const TSharedRef<SSequencerFilter>& InFilterWidget);
diff --git a/Engine/Source/Editor/Sequencer/Private/Filters/SequencerFilterBar.cpp b/Engine/Source/Editor/Sequencer/Private/Filters/SequencerFilterBar.cpp
--- a/Engine/Source/Editor/Sequencer/Private/Filters/SequencerFilterBar.cpp
+++ b/Engine/Source/Editor/Sequencer/Private/Filters/SequencerFilterBar.cpp
@@ -624,29 +624,39 @@ void FSequencerFilterBar::ActivateCommonFilters(const bool bInActivate
 
 	FSequencerFilterBarConfig& Config = SequencerSettings->FindOrAddTrackFilterBar(GetIdentifier(), false);
 
-	CommonFilters->ForEachFilter([this, bInActivate, &InExceptions, &Config]
+	bool bNeedsSave = false;
+
+	CommonFilters->ForEachFilter([this, bInActivate, &InExceptions, &Config, &bNeedsSave]
 		(const TSharedRef<FSequencerTrackFilter>& InFilter)
 		{
-			if (InExceptions.IsEmpty() || !InExceptions.Contains(InFilter))
+			if (InExceptions.Contains(InFilter))
 			{
-				const FString FilterName = InFilter->GetDisplayName().ToString();
-				if (Config.SetFilterActive(FilterName, bInActivate))
-				{
-					const ESequencerFilterChange FilterChangeType = bInActivate
-						? ESequencerFilterChange::Activate
-						: ESequencerFilterChange::Deactivate;
-					FiltersChangedEvent.Broadcast(FilterChangeType, InFilter);
+				return true;
+			}
 
-					InFilter->SetActive(bInActivate);
-					InFilter->ActiveStateChanged(bInActivate);
-				}
+			const FString FilterName = InFilter->GetDisplayName().ToString();
+			if (Config.SetFilterActive(FilterName, bInActivate))
+			{
+				const ESequencerFilterChange FilterChangeType = bInActivate
+					? ESequencerFilterChange::Activate
+					: ESequencerFilterChange::Deactivate;
+				FiltersChangedEvent.Broadcast(FilterChangeType, InFilter);
+
+				InFilter->SetActive(bInActivate);
+				InFilter->ActiveStateChanged(bInActivate);
+
+				bNeedsSave = true;
 			}
+
 			return true;
 		}
 		, true
 		, InMatchCategories);
 
-	SequencerSettings->SaveConfig();
+	if (bNeedsSave)
+	{
+		SequencerSettings->SaveConfig();
+	}
 
 	RequestFilterUpdate();
 }
diff --git a/Engine/Source/Editor/Sequencer/Private/Filters/SequencerFilterBarConfig.cpp b/Engine/Source/Editor/Sequencer/Private/Filters/SequencerFilterBarConfig.cpp
--- a/Engine/Source/Editor/Sequencer/Private/Filters/SequencerFilterBarConfig.cpp
+++ b/Engine/Source/Editor/Sequencer/Private/Filters/SequencerFilterBarConfig.cpp
@@ -31,7 +31,20 @@ bool FSequencerFilterBarConfig::IsFilterActive(const FString& InFilterName) cons
 
 bool FSequencerFilterBarConfig::SetFilterActive(const FString& InFilterName, const bool bInActive)
 {
-	ActiveFilters.EnabledStates.FindOrAdd(InFilterName) = bInActive;
+	if (bool* EnabledStatePtr = ActiveFilters.EnabledStates.Find(InFilterName))
+	{
+		if (*EnabledStatePtr == bInActive)
+		{
+			return false;
+		}
+
+		*EnabledStatePtr = bInActive;
+
+		return true;
+	}
+
+	ActiveFilters.EnabledStates.Add(InFilterName) = bInActive;
+
 	return true;
 }
 
diff --git a/Engine/Source/Editor/Sequencer/Private/Filters/Widgets/SFilterBarClippingHorizontalBox.cpp b/Engine/Source/Editor/Sequencer/Private/Filters/Widgets/SFilterBarClippingHorizontalBox.cpp
--- a/Engine/Source/Editor/Sequencer/Private/Filters/Widgets/SFilterBarClippingHorizontalBox.cpp
+++ b/Engine/Source/Editor/Sequencer/Private/Filters/Widgets/SFilterBarClippingHorizontalBox.cpp
@@ -9,6 +9,8 @@
 #include "Widgets/Images/SImage.h"
 #include "Widgets/Input/SComboButton.h"
 #include "Widgets/Layout/SBorder.h"
+#include "Widgets/Layout/SBox.h"
+#include "Widgets/Layout/SScrollBox.h"
 
 #define LOCTEXT_NAMESPACE "SFilterBarClippingHorizontalBox"
 
@@ -140,7 +142,15 @@ TSharedRef<SWidget> SFilterBarClippingHorizontalBox::WrapVerticalListWithHeading
 		.FillHeight(1.f)
 		.VAlign(VAlign_Fill)
 		[
-			InWidget
+			SNew(SBox)
+			.MaxDesiredHeight(480.f)
+			[
+				SNew(SScrollBox)
+				+ SScrollBox::Slot()
+				[
+					InWidget
+				]
+			]
 		];
 }
 
diff --git a/Engine/Source/Editor/Sequencer/Private/Filters/Widgets/SSequencerFilter.cpp b/Engine/Source/Editor/Sequencer/Private/Filters/Widgets/SSequencerFilter.cpp
--- a/Engine/Source/Editor/Sequencer/Private/Filters/Widgets/SSequencerFilter.cpp
+++ b/Engine/Source/Editor/Sequencer/Private/Filters/Widgets/SSequencerFilter.cpp
@@ -42,7 +42,6 @@ void SSequencerFilter::Construct(const FArguments& InArgs
 		.Style(FAppStyle::Get(), BrushName)
 		.ToolTipText(TAttribute<FText>::Create(TAttribute<FText>::FGetter::CreateSP(InFilter, &FSequencerTrackFilter::GetToolTipText)))
 		.IsChecked(this, &SSequencerFilter::IsChecked)
-		.OnCheckStateChanged(this, &SSequencerFilter::OnFilterToggled)
 		.CheckBoxContentUsesAutoWidth(false)
 		.OnGetMenuContent(this, &SSequencerFilter::GetRightClickMenuContent)
 		[
@@ -50,8 +49,11 @@ void SSequencerFilter::Construct(const FArguments& InArgs
 		]
 	];
 
-	ToggleButtonPtr->SetOnMouseUp(SSequencerFilterCheckBox::FOnPointerEvent::CreateSP(this, &SSequencerFilter::OnFilterMouseUp));
-	ToggleButtonPtr->SetOnDoubleClick(SSequencerFilterCheckBox::FOnPointerEvent::CreateSP(this, &SSequencerFilter::OnFilterDoubleClicked));
+	ToggleButtonPtr->SetOnClick(FOnCheckStateChanged::CreateSP(this, &SSequencerFilter::OnFilterToggled));
+	ToggleButtonPtr->SetOnCtrlClick(FSimpleDelegate::CreateSP(this, &SSequencerFilter::OnFilterCtrlClick));
+	ToggleButtonPtr->SetOnAltClick(FSimpleDelegate::CreateSP(this, &SSequencerFilter::OnFilterAltClick));
+	ToggleButtonPtr->SetOnMiddleButtonClick(FSimpleDelegate::CreateSP(this, &SSequencerFilter::OnFilterMiddleButtonClick));
+	ToggleButtonPtr->SetOnDoubleClick(FSimpleDelegate::CreateSP(this, &SSequencerFilter::OnFilterDoubleClick));
 }
 
 TSharedRef<SWidget> SSequencerFilter::ConstructBasicFilterWidget()
@@ -138,54 +140,40 @@ void SSequencerFilter::OnFilterToggled(const ECheckBoxState NewState)
 		return;
 	}
 
-	const bool bNewActive = NewState == ECheckBoxState::Checked;
+	const bool bNewActive = NewState != ECheckBoxState::Checked;
 	FilterBar->SetFilterActive(Filter.ToSharedRef(), bNewActive, true);
 }
 
-FReply SSequencerFilter::OnFilterMouseUp(const FGeometry& InMyGeometry, const FPointerEvent& InMouseEvent)
+void SSequencerFilter::OnFilterCtrlClick()
 {
-	const TSharedPtr<FSequencerFilterBar> FilterBar = WeakFilterBar.Pin();
-	if (!FilterBar.IsValid())
-	{
-		return FReply::Handled();
-	}
-
-	if (InMouseEvent.IsControlDown())
-	{
-		FilterBar->ActivateCommonFilters(true, {}, {});
-	}
-	else if (InMouseEvent.IsAltDown())
-	{
-		FilterBar->ActivateCommonFilters(false, {}, {});
-	}
-	else if(InMouseEvent.GetEffectingButton() == EKeys::MiddleMouseButton)
-	{
-		if (const TSharedPtr<FSequencerTrackFilter> Filter = WeakFilter.Pin())
-		{
-			FilterBar->SetFilterEnabled(Filter.ToSharedRef(), false, true);
-		}
-	}
+	ActivateAllButThis(false);
+}
 
-	return FReply::Handled();
+void SSequencerFilter::OnFilterAltClick()
+{
+	ActivateAllButThis(true);
 }
 
-FReply SSequencerFilter::OnFilterDoubleClicked(const FGeometry& InMyGeometry, const FPointerEvent& InMouseEvent)
+void SSequencerFilter::OnFilterMiddleButtonClick()
 {
 	const TSharedPtr<FSequencerFilterBar> FilterBar = WeakFilterBar.Pin();
 	if (!FilterBar.IsValid())
 	{
-		return FReply::Handled();
+		return;
 	}
 
-	FilterBar->ActivateCommonFilters(false, {}, {});
-	FilterBar->ActivateCustomTextFilters(false, {});
-
-	if (const TSharedPtr<FSequencerTrackFilter> Filter = WeakFilter.Pin())
+	const TSharedPtr<FSequencerTrackFilter> Filter = WeakFilter.Pin();
+	if (!Filter.IsValid())
 	{
-		FilterBar->SetFilterActive(Filter.ToSharedRef(), true, true);
+		return;
 	}
 
-	return FReply::Handled();
+	FilterBar->SetFilterEnabled(Filter.ToSharedRef(), false, true);
+}
+
+void SSequencerFilter::OnFilterDoubleClick()
+{
+	ActivateAllButThis(false);
 }
 
 TSharedRef<SWidget> SSequencerFilter::GetRightClickMenuContent()
@@ -240,4 +228,22 @@ bool SSequencerFilter::IsButtonEnabled() const
 	return FilterBar->IsFilterActive(Filter.ToSharedRef());
 }
 
+void SSequencerFilter::ActivateAllButThis(const bool bInActive)
+{
+	const TSharedPtr<FSequencerFilterBar> FilterBar = WeakFilterBar.Pin();
+	if (!FilterBar.IsValid())
+	{
+		return;
+	}
+
+	const TSharedPtr<FSequencerTrackFilter> Filter = WeakFilter.Pin();
+	if (!Filter.IsValid())
+	{
+		return;
+	}
+
+	FilterBar->ActivateAllEnabledFilters(bInActive, {});
+	FilterBar->SetFilterActive(Filter.ToSharedRef(), !bInActive, true);
+}
+
 #undef LOCTEXT_NAMESPACE
diff --git a/Engine/Source/Editor/Sequencer/Private/Filters/Widgets/SSequencerFilter.h b/Engine/Source/Editor/Sequencer/Private/Filters/Widgets/SSequencerFilter.h
--- a/Engine/Source/Editor/Sequencer/Private/Filters/Widgets/SSequencerFilter.h
+++ b/Engine/Source/Editor/Sequencer/Private/Filters/Widgets/SSequencerFilter.h
@@ -32,9 +32,10 @@ protected:
 
 	void OnFilterToggled(const ECheckBoxState NewState);
 
-	FReply OnFilterMouseUp(const FGeometry& InMyGeometry, const FPointerEvent& InMouseEvent);
-
-	FReply OnFilterDoubleClicked(const FGeometry& InMyGeometry, const FPointerEvent& InMouseEvent);
+	void OnFilterCtrlClick();
+	void OnFilterAltClick();
+	void OnFilterMiddleButtonClick();
+	void OnFilterDoubleClick();
 
 	TSharedRef<SWidget> GetRightClickMenuContent();
 
@@ -50,6 +51,8 @@ protected:
 
 	bool IsButtonEnabled() const;
 
+	void ActivateAllButThis(const bool bInActive);
+
 	TWeakPtr<FSequencerFilterBar> WeakFilterBar;
 
 	TWeakPtr<FSequencerTrackFilter> WeakFilter;
diff --git a/Engine/Source/Editor/Sequencer/Private/Filters/Widgets/SSequencerFilterBar.cpp b/Engine/Source/Editor/Sequencer/Private/Filters/Widgets/SSequencerFilterBar.cpp
--- a/Engine/Source/Editor/Sequencer/Private/Filters/Widgets/SSequencerFilterBar.cpp
+++ b/Engine/Source/Editor/Sequencer/Private/Filters/Widgets/SSequencerFilterBar.cpp
@@ -542,7 +542,7 @@ TSharedRef<SWidget> SSequencerFilterBar::OnWrapButtonClicked()
 
 	const TSharedRef<SBorder> ContainerBorder = SNew(SBorder)
 		.BorderImage(FAppStyle::GetBrush(TEXT("Brushes.Panel")))
-		.Padding(2.f)
+		.Padding(0.f, 2.f, 2.f, 2.f)
 		[
 			VerticalContainer
 		];
diff --git a/Engine/Source/Editor/Sequencer/Private/Filters/Widgets/SSequencerFilterCheckBox.h b/Engine/Source/Editor/Sequencer/Private/Filters/Widgets/SSequencerFilterCheckBox.h
--- a/Engine/Source/Editor/Sequencer/Private/Filters/Widgets/SSequencerFilterCheckBox.h
+++ b/Engine/Source/Editor/Sequencer/Private/Filters/Widgets/SSequencerFilterCheckBox.h
@@ -2,47 +2,145 @@
 
 #pragma once
 
+#include "Framework/Application/SlateApplication.h"
+#include "GameFramework/InputSettings.h"
 #include "Widgets/Input/SCheckBox.h"
 
 class SSequencerFilterCheckBox : public SCheckBox
 {
 public:
-	DECLARE_DELEGATE_RetVal_TwoParams(FReply, FOnPointerEvent, const FGeometry& /*InMyGeometry*/, const FPointerEvent& /*InMouseEvent*/);
+	void SetOnClick(const FOnCheckStateChanged& InOnClick)
+	{
+		OnClick = InOnClick;
+	}
+
+	void SetOnCtrlClick(const FSimpleDelegate& InNewCtrlClick)
+	{
+		OnCtrlClick = InNewCtrlClick;
+	}
+
+	void SetOnAltClick(const FSimpleDelegate& InNewAltClick)
+	{
+		OnAltClick = InNewAltClick;
+	}
 
-	void SetOnMouseUp(const FOnPointerEvent& InNewOnMouseUp)
+	void SetOnMiddleButtonClick(const FSimpleDelegate& InNewMiddleButtonClick)
 	{
-		OnMouseUp = InNewOnMouseUp;
+		OnMiddleButtonClick = InNewMiddleButtonClick;
 	}
 
-	void SetOnDoubleClick(const FOnPointerEvent& InNewOnDoubleClick)
+	void SetOnDoubleClick(const FSimpleDelegate& InNewDoubleClick)
 	{
-		OnDoubleClick = InNewOnDoubleClick;
+		OnDoubleClick = InNewDoubleClick;
 	}
 
 protected:
 	//~ Begin SWidget
 
-	virtual FReply OnMouseButtonUp(const FGeometry& InMyGeometry, const FPointerEvent& InMouseEvent) override
+	virtual FReply OnMouseButtonDown(const FGeometry& InGeometry, const FPointerEvent& InPointerEvent) override
 	{
-		FReply Reply = SCheckBox::OnMouseButtonUp(InMyGeometry, InMouseEvent);
-		if (OnMouseUp.IsBound())
+		if (InPointerEvent.GetEffectingButton() == EKeys::LeftMouseButton)
+		{
+			bIsPressed = true;
+
+			const EButtonClickMethod::Type InputClickMethod = GetClickMethodFromInputType(InPointerEvent);
+
+			if (InputClickMethod == EButtonClickMethod::MouseDown)
+			{
+				Internal_Click(InPointerEvent);
+
+				return FReply::Handled().SetUserFocus(AsShared(), EFocusCause::Mouse);
+			}
+
+			return FReply::Handled().CaptureMouse(AsShared()).SetUserFocus(AsShared(), EFocusCause::Mouse);
+		}
+
+		if (InPointerEvent.GetEffectingButton() == EKeys::RightMouseButton && OnGetMenuContent.IsBound())
 		{
-			Reply = OnMouseUp.Execute(InMyGeometry, InMouseEvent).ReleaseMouseCapture();
+			FSlateApplication::Get().PushMenu(
+				AsShared(),
+				InPointerEvent.GetEventPath() ? *InPointerEvent.GetEventPath() : FWidgetPath(),
+				OnGetMenuContent.Execute(),
+				InPointerEvent.GetScreenSpacePosition(),
+				FPopupTransitionEffect(FPopupTransitionEffect::ContextMenu));
+
+			return FReply::Handled();
 		}
-		return Reply;
+
+		return FReply::Unhandled();
 	}
 
-	virtual FReply OnMouseButtonDoubleClick(const FGeometry& InMyGeometry, const FPointerEvent& InMouseEvent) override
+	virtual FReply OnMouseButtonUp(const FGeometry& InGeometry, const FPointerEvent& InPointerEvent) override
 	{
-		if (OnDoubleClick.IsBound())
+		const EButtonClickMethod::Type InputClickMethod = GetClickMethodFromInputType(InPointerEvent);
+		const bool bMustBePressed = InputClickMethod == EButtonClickMethod::DownAndUp || InputClickMethod == EButtonClickMethod::PreciseClick;
+		const bool bMeetsPressedRequirements = !bMustBePressed || (bIsPressed && bMustBePressed);
+
+		if (bMeetsPressedRequirements && ((InPointerEvent.GetEffectingButton() == EKeys::LeftMouseButton || InPointerEvent.IsTouchEvent())))
 		{
-			return OnDoubleClick.Execute(InMyGeometry, InMouseEvent);
+			bIsPressed = false;
+
+			if (InputClickMethod != EButtonClickMethod::MouseDown)
+			{
+				const bool IsUnderMouse = InGeometry.IsUnderLocation(InPointerEvent.GetScreenSpacePosition());
+				if (IsUnderMouse)
+				{
+					// If we were asked to allow the button to be clicked on mouse up, regardless of whether the user
+					// pressed the button down first, then we'll allow the click to proceed without an active capture
+					if (InputClickMethod == EButtonClickMethod::MouseUp || HasMouseCapture())
+					{
+						Internal_Click(InPointerEvent);
+					}
+				}
+			}
+
+			return FReply::Handled().ReleaseMouseCapture();
 		}
-		return SCheckBox::OnMouseButtonDoubleClick(InMyGeometry, InMouseEvent);
+
+		return FReply::Unhandled();
+	}
+
+	virtual FReply OnMouseButtonDoubleClick(const FGeometry& InGeometry, const FPointerEvent& InPointerEvent) override
+	{
+		OnDoubleClick.ExecuteIfBound();
+
+		return FReply::Handled();
 	}
 
 	//~ End SWidget
 
-	FOnPointerEvent OnDoubleClick;
-	FOnPointerEvent OnMouseUp;
+	void Internal_Click(const FPointerEvent& InPointerEvent)
+	{
+		// Use a timer for click to give double click a chance to register
+		TimerHandle = RegisterActiveTimer(GetDefault<UInputSettings>()->DoubleClickTime,
+			FWidgetActiveTimerDelegate::CreateLambda([this, InPointerEvent](double InCurrentTime, float InDeltaTime) -> EActiveTimerReturnType
+			{
+				if (InPointerEvent.IsControlDown())
+				{
+					OnCtrlClick.ExecuteIfBound();
+				}
+				else if (InPointerEvent.IsAltDown())
+				{
+					OnAltClick.ExecuteIfBound();
+				}
+				else if (InPointerEvent.GetEffectingButton() == EKeys::MiddleMouseButton)
+				{
+					OnMiddleButtonClick.ExecuteIfBound();
+				}
+				else
+				{
+					OnClick.ExecuteIfBound(GetCheckedState());
+				}
+
+				return EActiveTimerReturnType::Stop;
+			}));
+	}
+
+	FOnCheckStateChanged OnClick;
+	FSimpleDelegate OnCtrlClick;
+	FSimpleDelegate OnAltClick;
+	FSimpleDelegate OnDoubleClick;
+	FSimpleDelegate OnMiddleButtonClick;
+
+	TSharedPtr<FActiveTimerHandle> TimerHandle;
 };
