Sequencer: Follow up to 36131623. Default should be 500x300

diff --git a/Engine/Source/Editor/Sequencer/Private/SequencerSettings.cpp b/Engine/Source/Editor/Sequencer/Private/SequencerSettings.cpp
--- a/Engine/Source/Editor/Sequencer/Private/SequencerSettings.cpp
+++ b/Engine/Source/Editor/Sequencer/Private/SequencerSettings.cpp
@@ -73,7 +73,7 @@ USequencerSettings::USequencerSettings( const FObjectInitializer& ObjectInitiali
 	bShowMarkedFrames = true;
 	ViewDensity = "Relaxed";
 	AssetBrowserWidth = 500.f;
-	AssetBrowserHeight = 400.f;
+	AssetBrowserHeight = 300.f;
 
 	SectionColorTints.Add(FColor(88, 102, 142, 255)); // blue
 	SectionColorTints.Add(FColor(99, 137, 132, 255)); // blue-green
