Sequencer: Fix iteration so that possessables aren't being updated in place while iterating through them

diff --git a/Engine/Source/Editor/Sequencer/Private/SequencerUtilities.cpp b/Engine/Source/Editor/Sequencer/Private/SequencerUtilities.cpp
--- a/Engine/Source/Editor/Sequencer/Private/SequencerUtilities.cpp
+++ b/Engine/Source/Editor/Sequencer/Private/SequencerUtilities.cpp
@@ -3880,6 +3880,8 @@ FGuid FSequencerUtilities::AssignActor(TSharedRef<ISequencer> Sequencer, AActor*
 		{
 			ComponentNameToComponent.Add(Component->GetName(), Component);
 		}
+
+		TMap<FGuid, UActorComponent*> ComponentsToUpdate;
 		for (int32 i = 0; i < OwnerMovieScene->GetPossessableCount(); i++)
 		{
 			FMovieScenePossessable& OldPossessable = OwnerMovieScene->GetPossessable(i);
@@ -3888,10 +3890,15 @@ FGuid FSequencerUtilities::AssignActor(TSharedRef<ISequencer> Sequencer, AActor*
 				UActorComponent** ComponentPtr = ComponentNameToComponent.Find(OldPossessable.GetName());
 				if (ComponentPtr != nullptr)
 				{
-					UpdateComponent(OldPossessable.GetGuid(), *ComponentPtr, NewComponentGuids);
+					ComponentsToUpdate.Add(OldPossessable.GetGuid(), *ComponentPtr);
 				}
 			}
 		}
+
+		for (TPair<FGuid, UActorComponent*> ComponentToUpdate : ComponentsToUpdate)
+		{
+			UpdateComponent(ComponentToUpdate.Key, ComponentToUpdate.Value, NewComponentGuids);
+		}
 	}
 
 	// Replace the actor itself after components have been updated
