Sequencer: Make the condition filter also for conditions on sections

#rnx

diff --git a/Engine/Source/Editor/Sequencer/Private/Filters/Filters/SequencerTrackFilter_Condition.cpp b/Engine/Source/Editor/Sequencer/Private/Filters/Filters/SequencerTrackFilter_Condition.cpp
--- a/Engine/Source/Editor/Sequencer/Private/Filters/Filters/SequencerTrackFilter_Condition.cpp
+++ b/Engine/Source/Editor/Sequencer/Private/Filters/Filters/SequencerTrackFilter_Condition.cpp
@@ -47,15 +47,14 @@ FString FSequencerTrackFilter_Condition::GetName() const
 
 bool FSequencerTrackFilter_Condition::PassesFilter(FSequencerTrackFilterType InItem) const
 {
-	const TViewModelPtr<IConditionableExtension> Conditionable = InItem.ImplicitCast();
-	if (!Conditionable)
+	for (const TViewModelPtr<IConditionableExtension>& Conditionable : InItem->GetDescendantsOfType<IConditionableExtension>(true))
 	{
-		return false;
-	}
-	else
-	{
-		return Conditionable->GetConditionState() != EConditionableConditionState::None;
+		if (Conditionable->GetConditionState() != EConditionableConditionState::None)
+		{
+			return true;
+		}
 	}
+	return false;
 }
 
 #undef LOCTEXT_NAMESPACE
