Sequencer: Restore track row metadata details that were removed by CL36175781.

diff --git a/Engine/Source/Editor/Sequencer/Private/Widgets/Sidebar/SequencerSelectionDrawer.cpp b/Engine/Source/Editor/Sequencer/Private/Widgets/Sidebar/SequencerSelectionDrawer.cpp
--- a/Engine/Source/Editor/Sequencer/Private/Widgets/Sidebar/SequencerSelectionDrawer.cpp
+++ b/Engine/Source/Editor/Sequencer/Private/Widgets/Sidebar/SequencerSelectionDrawer.cpp
@@ -327,18 +327,20 @@ void FSequencerSelectionDrawer::BuildOutlinerDetails(const TSharedRef<FSequencer
 	TArray<TWeakObjectPtr<>> WeakTrackObjects;
 	TSet<TViewModelPtr<FObjectBindingModel>> ObjectBindings;
 	TSet<TViewModelPtr<FChannelGroupOutlinerModel>> ChannelGroups;
+	TArray<TPair<TWeakObjectPtr<UMovieSceneTrack>, int32>> SelectedTrackRows;
+
 
 	for (const FViewModelPtr OutlinerItem : InSelection->Outliner)
 	{
-		if (const TViewModelPtr<FTrackModel> TrackModel = OutlinerItem.ImplicitCast())
-		{
-			WeakSectionObjects.Append(TrackModel->GetSections());
-			WeakTrackObjects.Add(TrackModel->GetTrack());
-		}
-		else if (const TViewModelPtr<FTrackRowModel> TrackRowModel = OutlinerItem.ImplicitCast())
+		if (const TViewModelPtr<ITrackExtension> TrackExtension = OutlinerItem.ImplicitCast())
 		{
-			WeakSectionObjects.Append(TrackRowModel->GetSections());
-			WeakTrackObjects.Add(TrackRowModel->GetTrack());
+			WeakSectionObjects.Append(TrackExtension->GetSections());
+			UMovieSceneTrack* Track = TrackExtension->GetTrack();
+			if (IsValid(Track))
+			{
+				WeakTrackObjects.Add(Track);
+				SelectedTrackRows.Add(TPair<TWeakObjectPtr<UMovieSceneTrack>, int32>(Track, TrackExtension->GetRowIndex()));
+			}
 		}
 		else if (const TViewModelPtr<FObjectBindingModel> ObjectBindingModel = OutlinerItem.ImplicitCast())
 		{
@@ -370,9 +372,6 @@ void FSequencerSelectionDrawer::BuildOutlinerDetails(const TSharedRef<FSequencer
 		// Shows duplicate information as above?
 		//MenuBuilder.BeginSection(TEXT("CustomBinding"));
 		//MenuBuilder.EndSection();
-
-		MenuBuilder.BeginSection(TEXT("TrackRowMetadata"));
-		MenuBuilder.EndSection();
 	}
 
 	if (!ChannelGroups.IsEmpty())
@@ -380,6 +379,18 @@ void FSequencerSelectionDrawer::BuildOutlinerDetails(const TSharedRef<FSequencer
 		BuildExtensionDetails(ChannelGroups, MenuBuilder);
 	}
 
+	if (!SelectedTrackRows.IsEmpty() && !Algo::AnyOf(SelectedTrackRows, [](const TPair<TWeakObjectPtr<UMovieSceneTrack>, int32> TrackRow) {
+		return TrackRow.Key.IsValid() && !TrackRow.Key->SupportsMultipleRows();
+		}))
+	{
+		MenuBuilder.BeginSection(TEXT("TrackRowMetadata"), LOCTEXT("TrackRowMetadata", "Track Row Metadata"));
+		{
+			// Empty here, will be implemented by extension.
+		}
+		MenuBuilder.EndSection();
+	}
+
+
 	if (!WeakTrackObjects.IsEmpty())
 	{
 		SequencerHelpers::BuildEditTrackMenu(Sequencer, WeakTrackObjects, MenuBuilder, false);
