[Sequencer] Advanced filtering:
 - Split "Class" text expression into "ObjectClass" and "TrackClass"
 - Fix NotEqual operator not working with string text expressions

diff --git a/Engine/Source/Editor/Sequencer/Private/Filters/Filters/SequencerTrackFilter_Text.cpp b/Engine/Source/Editor/Sequencer/Private/Filters/Filters/SequencerTrackFilter_Text.cpp
--- a/Engine/Source/Editor/Sequencer/Private/Filters/Filters/SequencerTrackFilter_Text.cpp
+++ b/Engine/Source/Editor/Sequencer/Private/Filters/Filters/SequencerTrackFilter_Text.cpp
@@ -11,17 +11,18 @@
 #include "Filters/TextExpressions/SequencerTextFilterExpression_ConditionClass.h"
 #include "Filters/TextExpressions/SequencerTextFilterExpression_ConditionFunc.h"
 #include "Filters/TextExpressions/SequencerTextFilterExpression_ConditionPasses.h"
-#include "Filters/TextExpressions/SequencerTextFilterExpression_Class.h"
 #include "Filters/TextExpressions/SequencerTextFilterExpression_Group.h"
 #include "Filters/TextExpressions/SequencerTextFilterExpression_Level.h"
 #include "Filters/TextExpressions/SequencerTextFilterExpression_Locked.h"
 #include "Filters/TextExpressions/SequencerTextFilterExpression_Modified.h"
 #include "Filters/TextExpressions/SequencerTextFilterExpression_Muted.h"
 #include "Filters/TextExpressions/SequencerTextFilterExpression_Name.h"
+#include "Filters/TextExpressions/SequencerTextFilterExpression_ObjectClass.h"
 #include "Filters/TextExpressions/SequencerTextFilterExpression_Selected.h"
 #include "Filters/TextExpressions/SequencerTextFilterExpression_Soloed.h"
 #include "Filters/TextExpressions/SequencerTextFilterExpression_Tag.h"
 #include "Filters/TextExpressions/SequencerTextFilterExpression_Time.h"
+#include "Filters/TextExpressions/SequencerTextFilterExpression_TrackClass.h"
 #include "Filters/TextExpressions/SequencerTextFilterExpression_Unbound.h"
 #include "UObject/UObjectIterator.h"
 
@@ -40,7 +41,8 @@ FSequencerTrackFilter_Text::FSequencerTrackFilter_Text(ISequencerTrackFilters& I
 
 	// Ordered by importance and most often used. This will dictate the order of display in the text expressions help dialog.
 	TextFilterExpressionContexts.Add(MakeShared<FSequencerTextFilterExpression_Name>(InFilterInterface));
-	TextFilterExpressionContexts.Add(MakeShared<FSequencerTextFilterExpression_Class>(InFilterInterface));
+	TextFilterExpressionContexts.Add(MakeShared<FSequencerTextFilterExpression_TrackClass>(InFilterInterface));
+	TextFilterExpressionContexts.Add(MakeShared<FSequencerTextFilterExpression_ObjectClass>(InFilterInterface));
 	TextFilterExpressionContexts.Add(MakeShared<FSequencerTextFilterExpression_Condition>(InFilterInterface));
 	TextFilterExpressionContexts.Add(MakeShared<FSequencerTextFilterExpression_ConditionClass>(InFilterInterface));
 	TextFilterExpressionContexts.Add(MakeShared<FSequencerTextFilterExpression_ConditionFunc>(InFilterInterface));
diff --git a/Engine/Source/Editor/Sequencer/Private/Filters/TextExpressions/SequencerTextFilterExpression_Class.cpp b/Engine/Source/Editor/Sequencer/Private/Filters/TextExpressions/SequencerTextFilterExpression_Class.cpp
deleted file mode 100644
diff --git a/Engine/Source/Editor/Sequencer/Private/Filters/TextExpressions/SequencerTextFilterExpression_Class.h b/Engine/Source/Editor/Sequencer/Private/Filters/TextExpressions/SequencerTextFilterExpression_Class.h
deleted file mode 100644
diff --git a/Engine/Source/Editor/Sequencer/Private/Filters/TextExpressions/SequencerTextFilterExpression_ObjectClass.cpp b/Engine/Source/Editor/Sequencer/Private/Filters/TextExpressions/SequencerTextFilterExpression_ObjectClass.cpp
new file mode 100644
--- /dev/null
+++ b/Engine/Source/Editor/Sequencer/Private/Filters/TextExpressions/SequencerTextFilterExpression_ObjectClass.cpp
@@ -0,0 +1,55 @@
+﻿// Copyright Epic Games, Inc. All Rights Reserved.
+
+#include "SequencerTextFilterExpression_ObjectClass.h"
+#include "Sequencer.h"
+
+using namespace UE::Sequencer;
+
+#define LOCTEXT_NAMESPACE "SequencerTextFilterExpression_ObjectClass"
+
+FSequencerTextFilterExpression_ObjectClass::FSequencerTextFilterExpression_ObjectClass(ISequencerTrackFilters& InFilterInterface)
+	: FSequencerTextFilterExpressionContext(InFilterInterface)
+{
+}
+
+TSet<FName> FSequencerTextFilterExpression_ObjectClass::GetKeys() const
+{
+	return { TEXT("OBJECTCLASS"), TEXT("OBJECTTYPE") };
+}
+
+ESequencerTextFilterValueType FSequencerTextFilterExpression_ObjectClass::GetValueType() const
+{
+	return ESequencerTextFilterValueType::String;
+}
+
+FText FSequencerTextFilterExpression_ObjectClass::GetDescription() const
+{
+	return LOCTEXT("ExpressionDescription_Class", "Filter by bound object class name");
+}
+
+bool FSequencerTextFilterExpression_ObjectClass::TestComplexExpression(const FName& InKey
+	, const FTextFilterString& InValue
+	, const ETextFilterComparisonOperation InComparisonOperation
+	, const ETextFilterTextComparisonMode InTextComparisonMode) const
+{
+	if (!FSequencerTextFilterExpressionContext::TestComplexExpression(InKey, InValue, InComparisonOperation, InTextComparisonMode))
+	{
+		return true;
+	}
+
+	ISequencer& Sequencer = FilterInterface.GetSequencer();
+
+	UObject* const BoundObject = FSequencerTrackFilter::ResolveTrackBoundObject(Sequencer, FilterItem, FilterInterface.GetFilterData());
+	if (IsValid(BoundObject))
+	{
+		const FString BoundObjectClassName = BoundObject->GetClass()->GetName();
+		if (TextFilterUtils::TestComplexExpression(BoundObjectClassName, InValue, InComparisonOperation, InTextComparisonMode))
+		{
+			return true;
+		}
+	}
+
+	return false;
+}
+
+#undef LOCTEXT_NAMESPACE
diff --git a/Engine/Source/Editor/Sequencer/Private/Filters/TextExpressions/SequencerTextFilterExpression_ObjectClass.h b/Engine/Source/Editor/Sequencer/Private/Filters/TextExpressions/SequencerTextFilterExpression_ObjectClass.h
new file mode 100644
--- /dev/null
+++ b/Engine/Source/Editor/Sequencer/Private/Filters/TextExpressions/SequencerTextFilterExpression_ObjectClass.h
@@ -0,0 +1,24 @@
+﻿// Copyright Epic Games, Inc. All Rights Reserved.
+
+#pragma once
+
+#include "Filters/SequencerTextFilterExpressionContext.h"
+
+class FSequencerTextFilterExpression_ObjectClass : public FSequencerTextFilterExpressionContext
+{
+public:
+	FSequencerTextFilterExpression_ObjectClass(ISequencerTrackFilters& InFilterInterface);
+
+	//~ Begin FSequencerTextFilterExpressionContext
+	virtual TSet<FName> GetKeys() const override;
+	virtual ESequencerTextFilterValueType GetValueType() const override;
+	virtual FText GetDescription() const override;
+	//~ End FSequencerTextFilterExpressionContext
+
+	//~ Begin ITextFilterExpressionContext
+	virtual bool TestComplexExpression(const FName& InKey
+		, const FTextFilterString& InValue
+		, const ETextFilterComparisonOperation InComparisonOperation
+		, const ETextFilterTextComparisonMode InTextComparisonMode) const override;
+	//~ End ITextFilterExpressionContext
+};
diff --git a/Engine/Source/Editor/Sequencer/Private/Filters/TextExpressions/SequencerTextFilterExpression_TrackClass.cpp b/Engine/Source/Editor/Sequencer/Private/Filters/TextExpressions/SequencerTextFilterExpression_TrackClass.cpp
new file mode 100644
--- /dev/null
+++ b/Engine/Source/Editor/Sequencer/Private/Filters/TextExpressions/SequencerTextFilterExpression_TrackClass.cpp
@@ -0,0 +1,53 @@
+﻿// Copyright Epic Games, Inc. All Rights Reserved.
+
+#include "SequencerTextFilterExpression_TrackClass.h"
+#include "Filters/SequencerFilterData.h"
+#include "Sequencer.h"
+
+using namespace UE::Sequencer;
+
+#define LOCTEXT_NAMESPACE "SequencerTextFilterExpression_TrackClass"
+
+FSequencerTextFilterExpression_TrackClass::FSequencerTextFilterExpression_TrackClass(ISequencerTrackFilters& InFilterInterface)
+	: FSequencerTextFilterExpressionContext(InFilterInterface)
+{
+}
+
+TSet<FName> FSequencerTextFilterExpression_TrackClass::GetKeys() const
+{
+	return { TEXT("TRACKCLASS"), TEXT("TRACKTYPE") };
+}
+
+ESequencerTextFilterValueType FSequencerTextFilterExpression_TrackClass::GetValueType() const
+{
+	return ESequencerTextFilterValueType::String;
+}
+
+FText FSequencerTextFilterExpression_TrackClass::GetDescription() const
+{
+	return LOCTEXT("ExpressionDescription_Class", "Filter by track class name");
+}
+
+bool FSequencerTextFilterExpression_TrackClass::TestComplexExpression(const FName& InKey
+	, const FTextFilterString& InValue
+	, const ETextFilterComparisonOperation InComparisonOperation
+	, const ETextFilterTextComparisonMode InTextComparisonMode) const
+{
+	if (!FSequencerTextFilterExpressionContext::TestComplexExpression(InKey, InValue, InComparisonOperation, InTextComparisonMode))
+	{
+		return true;
+	}
+
+	if (WeakTrackObject.IsValid())
+	{
+		const FString TrackClassName = WeakTrackObject->GetClass()->GetName();
+		if (TextFilterUtils::TestComplexExpression(TrackClassName, InValue, InComparisonOperation, InTextComparisonMode))
+		{
+			return true;
+		}
+	}
+
+	return false;
+}
+
+#undef LOCTEXT_NAMESPACE
diff --git a/Engine/Source/Editor/Sequencer/Private/Filters/TextExpressions/SequencerTextFilterExpression_TrackClass.h b/Engine/Source/Editor/Sequencer/Private/Filters/TextExpressions/SequencerTextFilterExpression_TrackClass.h
new file mode 100644
--- /dev/null
+++ b/Engine/Source/Editor/Sequencer/Private/Filters/TextExpressions/SequencerTextFilterExpression_TrackClass.h
@@ -0,0 +1,24 @@
+﻿// Copyright Epic Games, Inc. All Rights Reserved.
+
+#pragma once
+
+#include "Filters/SequencerTextFilterExpressionContext.h"
+
+class FSequencerTextFilterExpression_TrackClass : public FSequencerTextFilterExpressionContext
+{
+public:
+	FSequencerTextFilterExpression_TrackClass(ISequencerTrackFilters& InFilterInterface);
+
+	//~ Begin FSequencerTextFilterExpressionContext
+	virtual TSet<FName> GetKeys() const override;
+	virtual ESequencerTextFilterValueType GetValueType() const override;
+	virtual FText GetDescription() const override;
+	//~ End FSequencerTextFilterExpressionContext
+
+	//~ Begin ITextFilterExpressionContext
+	virtual bool TestComplexExpression(const FName& InKey
+		, const FTextFilterString& InValue
+		, const ETextFilterComparisonOperation InComparisonOperation
+		, const ETextFilterTextComparisonMode InTextComparisonMode) const override;
+	//~ End ITextFilterExpressionContext
+};
