[Sequencer] Add commands for shared Curve Editor

diff --git a/Engine/Source/Editor/Sequencer/Private/Filters/SequencerFilterBar.cpp b/Engine/Source/Editor/Sequencer/Private/Filters/SequencerFilterBar.cpp
--- a/Engine/Source/Editor/Sequencer/Private/Filters/SequencerFilterBar.cpp
+++ b/Engine/Source/Editor/Sequencer/Private/Filters/SequencerFilterBar.cpp
@@ -201,21 +201,27 @@ void FSequencerFilterBar::BindCommands()
 	}
 
 	// Add bindings for curve editor if supported
-	FCurveEditorExtension* CurveEditorExtension = Sequencer.GetViewModel()->CastDynamic<FCurveEditorExtension>();
-	if (CurveEditorExtension && ensure(CurveEditorExtension->GetCurveEditor()))
+	FCurveEditorExtension* const CurveEditorExtension = Sequencer.GetViewModel()->CastDynamic<FCurveEditorExtension>();
+	if (CurveEditorExtension)
 	{
-		const TSharedPtr<FUICommandList> SequencerBindings = Sequencer.GetCommandBindings(ESequencerCommandBindings::Sequencer);
-		SequencerBindings->MapAction(TrackFilterCommands.ResetFilters, *CommandList->GetActionForCommand(TrackFilterCommands.ResetFilters));
-		if (const TSharedPtr<FCurveEditor> CurveEditor = CurveEditorExtension->GetCurveEditor())
+		const TSharedPtr<FCurveEditor> CurveEditor = CurveEditorExtension->GetCurveEditor();
+		if (ensure(CurveEditor.IsValid()))
 		{
-			CurveEditor->GetCommands()->Append(SequencerBindings.ToSharedRef());
-		}
-		
-		const TSharedPtr<FUICommandList> CurveEditorSharedBindings = Sequencer.GetCommandBindings(ESequencerCommandBindings::CurveEditor);
-		CurveEditorSharedBindings->MapAction(TrackFilterCommands.ResetFilters, *CommandList->GetActionForCommand(TrackFilterCommands.ResetFilters));
-		if (const TSharedPtr<FCurveEditor> CurveEditor = CurveEditorExtension->GetCurveEditor())
-		{
-			CurveEditor->GetCommands()->Append(CurveEditorSharedBindings.ToSharedRef());
+			const TSharedPtr<FUICommandList> CurveEditorCommands = CurveEditor->GetCommands();
+			if (ensure(CurveEditorCommands.IsValid()))
+			{
+				const TSharedPtr<FUICommandList> CurveEditorSharedBindings = Sequencer.GetCommandBindings(ESequencerCommandBindings::CurveEditor);
+
+				for (const TSharedPtr<FUICommandInfo>& Command : TrackFilterCommands.GetAllCommands())
+				{
+					if (Command.IsValid() && CommandList->IsActionMapped(Command))
+					{
+						CurveEditorSharedBindings->MapAction(Command, *CommandList->GetActionForCommand(Command));
+					}
+				}
+
+				CurveEditorCommands->Append(CurveEditorSharedBindings.ToSharedRef());
+			}
 		}
 	}
 }
diff --git a/Engine/Source/Editor/Sequencer/Private/Filters/SequencerTrackFilterCommands.cpp b/Engine/Source/Editor/Sequencer/Private/Filters/SequencerTrackFilterCommands.cpp
--- a/Engine/Source/Editor/Sequencer/Private/Filters/SequencerTrackFilterCommands.cpp
+++ b/Engine/Source/Editor/Sequencer/Private/Filters/SequencerTrackFilterCommands.cpp
@@ -64,4 +64,59 @@ void FSequencerTrackFilterCommands::RegisterCommands()
 	UI_COMMAND(ToggleFilter_Levels, "Toggle Levels Filter", "Toggle the filter for Level tracks", EUserInterfaceActionType::Button, FInputChord());
 }
 
+TArray<TSharedPtr<FUICommandInfo>> FSequencerTrackFilterCommands::GetAllCommands() const
+{
+	TArray<TSharedPtr<FUICommandInfo>> OutCommands;
+
+	OutCommands.Add(ToggleFilterBarVisibility);
+
+	OutCommands.Add(ResetFilters);
+
+	OutCommands.Add(ToggleMuteFilters);
+
+	OutCommands.Add(DisableAllFilters);
+
+	OutCommands.Add(ToggleActivateEnabledFilters);
+	
+	OutCommands.Add(HideSelectedTracks);
+	OutCommands.Add(IsolateSelectedTracks);
+
+	OutCommands.Add(ClearHiddenTracks);
+	OutCommands.Add(ClearIsolatedTracks);
+
+	OutCommands.Add(ShowAllTracks);
+
+	OutCommands.Add(ShowLocationCategoryGroups);
+	OutCommands.Add(ShowRotationCategoryGroups);
+	OutCommands.Add(ShowScaleCategoryGroups);
+
+	OutCommands.Add(ToggleFilter_Audio);
+	OutCommands.Add(ToggleFilter_DataLayer);
+	OutCommands.Add(ToggleFilter_Event);
+	OutCommands.Add(ToggleFilter_Fade);
+	OutCommands.Add(ToggleFilter_Folder);
+	OutCommands.Add(ToggleFilter_LevelVisibility);
+	OutCommands.Add(ToggleFilter_Particle);
+	OutCommands.Add(ToggleFilter_CinematicShot);
+	OutCommands.Add(ToggleFilter_Subsequence);
+	OutCommands.Add(ToggleFilter_TimeDilation);
+	OutCommands.Add(ToggleFilter_TimeWarp);
+	
+	OutCommands.Add(ToggleFilter_Camera);
+	OutCommands.Add(ToggleFilter_CameraCut);
+	OutCommands.Add(ToggleFilter_Light);
+	OutCommands.Add(ToggleFilter_SkeletalMesh);
+
+	OutCommands.Add(ToggleFilter_Condition);
+	OutCommands.Add(ToggleFilter_Keyed);
+	OutCommands.Add(ToggleFilter_Modified);
+	OutCommands.Add(ToggleFilter_Selected);
+	OutCommands.Add(ToggleFilter_Unbound);
+
+	OutCommands.Add(ToggleFilter_Groups);
+	OutCommands.Add(ToggleFilter_Levels);
+
+	return OutCommands;
+}
+
 #undef LOCTEXT_NAMESPACE
diff --git a/Engine/Source/Editor/Sequencer/Private/Filters/SequencerTrackFilterCommands.h b/Engine/Source/Editor/Sequencer/Private/Filters/SequencerTrackFilterCommands.h
--- a/Engine/Source/Editor/Sequencer/Private/Filters/SequencerTrackFilterCommands.h
+++ b/Engine/Source/Editor/Sequencer/Private/Filters/SequencerTrackFilterCommands.h
@@ -15,6 +15,8 @@ public:
 	virtual void RegisterCommands() override;
 	//~ End TCommands
 
+	TArray<TSharedPtr<FUICommandInfo>> GetAllCommands() const;
+
 	/** FilterBar Commands */
 
 	TSharedPtr<FUICommandInfo> ToggleFilterBarVisibility;
diff --git a/Engine/Source/Editor/Sequencer/Private/SSequencer.cpp b/Engine/Source/Editor/Sequencer/Private/SSequencer.cpp
--- a/Engine/Source/Editor/Sequencer/Private/SSequencer.cpp
+++ b/Engine/Source/Editor/Sequencer/Private/SSequencer.cpp
@@ -1107,11 +1107,6 @@ void SSequencer::BindCommands(TSharedRef<FUICommandList> SequencerCommandBinding
 		FSequencerCommands::Get().OpenNodeGroupsManager,
 		FExecuteAction::CreateSP(this, &SSequencer::OpenNodeGroupsManager)
 	);
-
-	if (const TSharedPtr<FSequencerFilterBar> FilterBar = GetFilterBar())
-	{
-		FilterBar->BindCommands();
-	}
 }
 
 void SSequencer::OpenTickResolutionOptions()
