[Sequencer] Better filter widget double click handling (removes delay)

diff --git a/Engine/Source/Editor/Sequencer/Private/Filters/Widgets/SSequencerFilter.cpp b/Engine/Source/Editor/Sequencer/Private/Filters/Widgets/SSequencerFilter.cpp
--- a/Engine/Source/Editor/Sequencer/Private/Filters/Widgets/SSequencerFilter.cpp
+++ b/Engine/Source/Editor/Sequencer/Private/Filters/Widgets/SSequencerFilter.cpp
@@ -42,6 +42,7 @@ void SSequencerFilter::Construct(const FArguments& InArgs
 		.Style(FAppStyle::Get(), BrushName)
 		.ToolTipText(TAttribute<FText>::Create(TAttribute<FText>::FGetter::CreateSP(InFilter, &FSequencerTrackFilter::GetToolTipText)))
 		.IsChecked(this, &SSequencerFilter::IsChecked)
+		.OnCheckStateChanged(this, &SSequencerFilter::OnFilterToggled)
 		.CheckBoxContentUsesAutoWidth(false)
 		.OnGetMenuContent(this, &SSequencerFilter::GetRightClickMenuContent)
 		[
@@ -49,11 +50,10 @@ void SSequencerFilter::Construct(const FArguments& InArgs
 		]
 	];
 
-	ToggleButtonPtr->SetOnClick(FOnCheckStateChanged::CreateSP(this, &SSequencerFilter::OnFilterToggled));
-	ToggleButtonPtr->SetOnCtrlClick(FSimpleDelegate::CreateSP(this, &SSequencerFilter::OnFilterCtrlClick));
-	ToggleButtonPtr->SetOnAltClick(FSimpleDelegate::CreateSP(this, &SSequencerFilter::OnFilterAltClick));
-	ToggleButtonPtr->SetOnMiddleButtonClick(FSimpleDelegate::CreateSP(this, &SSequencerFilter::OnFilterMiddleButtonClick));
-	ToggleButtonPtr->SetOnDoubleClick(FSimpleDelegate::CreateSP(this, &SSequencerFilter::OnFilterDoubleClick));
+	ToggleButtonPtr->SetOnCtrlClick(FOnClicked::CreateSP(this, &SSequencerFilter::OnFilterCtrlClick));
+	ToggleButtonPtr->SetOnAltClick(FOnClicked::CreateSP(this, &SSequencerFilter::OnFilterAltClick));
+	ToggleButtonPtr->SetOnMiddleButtonClick(FOnClicked::CreateSP(this, &SSequencerFilter::OnFilterMiddleButtonClick));
+	ToggleButtonPtr->SetOnDoubleClick(FOnClicked::CreateSP(this, &SSequencerFilter::OnFilterDoubleClick));
 }
 
 TSharedRef<SWidget> SSequencerFilter::ConstructBasicFilterWidget()
@@ -140,40 +140,48 @@ void SSequencerFilter::OnFilterToggled(const ECheckBoxState NewState)
 		return;
 	}
 
-	const bool bNewActive = NewState != ECheckBoxState::Checked;
+	const bool bNewActive = NewState == ECheckBoxState::Checked;
 	FilterBar->SetFilterActive(Filter.ToSharedRef(), bNewActive, true);
 }
 
-void SSequencerFilter::OnFilterCtrlClick()
+FReply SSequencerFilter::OnFilterCtrlClick()
 {
 	ActivateAllButThis(false);
+
+	return FReply::Handled();
 }
 
-void SSequencerFilter::OnFilterAltClick()
+FReply SSequencerFilter::OnFilterAltClick()
 {
 	ActivateAllButThis(true);
+
+	return FReply::Handled();
 }
 
-void SSequencerFilter::OnFilterMiddleButtonClick()
+FReply SSequencerFilter::OnFilterMiddleButtonClick()
 {
 	const TSharedPtr<FSequencerFilterBar> FilterBar = WeakFilterBar.Pin();
 	if (!FilterBar.IsValid())
 	{
-		return;
+		return FReply::Handled();
 	}
 
 	const TSharedPtr<FSequencerTrackFilter> Filter = WeakFilter.Pin();
 	if (!Filter.IsValid())
 	{
-		return;
+		return FReply::Handled();
 	}
 
 	FilterBar->SetFilterEnabled(Filter.ToSharedRef(), false, true);
+
+	return FReply::Handled();
 }
 
-void SSequencerFilter::OnFilterDoubleClick()
+FReply SSequencerFilter::OnFilterDoubleClick()
 {
 	ActivateAllButThis(false);
+
+	return FReply::Handled();
 }
 
 TSharedRef<SWidget> SSequencerFilter::GetRightClickMenuContent()
diff --git a/Engine/Source/Editor/Sequencer/Private/Filters/Widgets/SSequencerFilter.h b/Engine/Source/Editor/Sequencer/Private/Filters/Widgets/SSequencerFilter.h
--- a/Engine/Source/Editor/Sequencer/Private/Filters/Widgets/SSequencerFilter.h
+++ b/Engine/Source/Editor/Sequencer/Private/Filters/Widgets/SSequencerFilter.h
@@ -32,10 +32,10 @@ protected:
 
 	void OnFilterToggled(const ECheckBoxState NewState);
 
-	void OnFilterCtrlClick();
-	void OnFilterAltClick();
-	void OnFilterMiddleButtonClick();
-	void OnFilterDoubleClick();
+	FReply OnFilterCtrlClick();
+	FReply OnFilterAltClick();
+	FReply OnFilterMiddleButtonClick();
+	FReply OnFilterDoubleClick();
 
 	TSharedRef<SWidget> GetRightClickMenuContent();
 
diff --git a/Engine/Source/Editor/Sequencer/Private/Filters/Widgets/SSequencerFilterCheckBox.h b/Engine/Source/Editor/Sequencer/Private/Filters/Widgets/SSequencerFilterCheckBox.h
--- a/Engine/Source/Editor/Sequencer/Private/Filters/Widgets/SSequencerFilterCheckBox.h
+++ b/Engine/Source/Editor/Sequencer/Private/Filters/Widgets/SSequencerFilterCheckBox.h
@@ -3,33 +3,27 @@
 #pragma once
 
 #include "Framework/Application/SlateApplication.h"
-#include "GameFramework/InputSettings.h"
 #include "Widgets/Input/SCheckBox.h"
 
 class SSequencerFilterCheckBox : public SCheckBox
 {
 public:
-	void SetOnClick(const FOnCheckStateChanged& InOnClick)
-	{
-		OnClick = InOnClick;
-	}
-
-	void SetOnCtrlClick(const FSimpleDelegate& InNewCtrlClick)
+	void SetOnCtrlClick(const FOnClicked& InNewCtrlClick)
 	{
 		OnCtrlClick = InNewCtrlClick;
 	}
 
-	void SetOnAltClick(const FSimpleDelegate& InNewAltClick)
+	void SetOnAltClick(const FOnClicked& InNewAltClick)
 	{
 		OnAltClick = InNewAltClick;
 	}
 
-	void SetOnMiddleButtonClick(const FSimpleDelegate& InNewMiddleButtonClick)
+	void SetOnMiddleButtonClick(const FOnClicked& InNewMiddleButtonClick)
 	{
 		OnMiddleButtonClick = InNewMiddleButtonClick;
 	}
 
-	void SetOnDoubleClick(const FSimpleDelegate& InNewDoubleClick)
+	void SetOnDoubleClick(const FOnClicked& InNewDoubleClick)
 	{
 		OnDoubleClick = InNewDoubleClick;
 	}
@@ -37,110 +31,42 @@ public:
 protected:
 	//~ Begin SWidget
 
-	virtual FReply OnMouseButtonDown(const FGeometry& InGeometry, const FPointerEvent& InPointerEvent) override
+	virtual FReply OnMouseButtonUp(const FGeometry& InGeometry, const FPointerEvent& InPointerEvent) override
 	{
-		if (InPointerEvent.GetEffectingButton() == EKeys::LeftMouseButton)
+		if (InPointerEvent.IsControlDown() && OnCtrlClick.IsBound())
 		{
-			bIsPressed = true;
-
-			const EButtonClickMethod::Type InputClickMethod = GetClickMethodFromInputType(InPointerEvent);
-
-			if (InputClickMethod == EButtonClickMethod::MouseDown)
-			{
-				Internal_Click(InPointerEvent);
-
-				return FReply::Handled().SetUserFocus(AsShared(), EFocusCause::Mouse);
-			}
-
-			return FReply::Handled().CaptureMouse(AsShared()).SetUserFocus(AsShared(), EFocusCause::Mouse);
+			return OnCtrlClick.Execute();
 		}
 
-		if (InPointerEvent.GetEffectingButton() == EKeys::RightMouseButton && OnGetMenuContent.IsBound())
+		if (InPointerEvent.IsAltDown() && OnAltClick.IsBound())
 		{
-			FSlateApplication::Get().PushMenu(
-				AsShared(),
-				InPointerEvent.GetEventPath() ? *InPointerEvent.GetEventPath() : FWidgetPath(),
-				OnGetMenuContent.Execute(),
-				InPointerEvent.GetScreenSpacePosition(),
-				FPopupTransitionEffect(FPopupTransitionEffect::ContextMenu));
-
-			return FReply::Handled();
+			return OnAltClick.Execute();
 		}
 
-		return FReply::Unhandled();
-	}
-
-	virtual FReply OnMouseButtonUp(const FGeometry& InGeometry, const FPointerEvent& InPointerEvent) override
-	{
-		const EButtonClickMethod::Type InputClickMethod = GetClickMethodFromInputType(InPointerEvent);
-		const bool bMustBePressed = InputClickMethod == EButtonClickMethod::DownAndUp || InputClickMethod == EButtonClickMethod::PreciseClick;
-		const bool bMeetsPressedRequirements = !bMustBePressed || (bIsPressed && bMustBePressed);
-
-		if (bMeetsPressedRequirements && ((InPointerEvent.GetEffectingButton() == EKeys::LeftMouseButton || InPointerEvent.IsTouchEvent())))
+		if (InPointerEvent.GetEffectingButton() == EKeys::MiddleMouseButton && OnMiddleButtonClick.IsBound() )
 		{
-			bIsPressed = false;
-
-			if (InputClickMethod != EButtonClickMethod::MouseDown)
-			{
-				const bool IsUnderMouse = InGeometry.IsUnderLocation(InPointerEvent.GetScreenSpacePosition());
-				if (IsUnderMouse)
-				{
-					// If we were asked to allow the button to be clicked on mouse up, regardless of whether the user
-					// pressed the button down first, then we'll allow the click to proceed without an active capture
-					if (InputClickMethod == EButtonClickMethod::MouseUp || HasMouseCapture())
-					{
-						Internal_Click(InPointerEvent);
-					}
-				}
-			}
-
-			return FReply::Handled().ReleaseMouseCapture();
+			return OnMiddleButtonClick.Execute();
 		}
 
-		return FReply::Unhandled();
+		SCheckBox::OnMouseButtonUp(InGeometry, InPointerEvent);
+
+		return FReply::Handled().ReleaseMouseCapture();
 	}
 
 	virtual FReply OnMouseButtonDoubleClick(const FGeometry& InGeometry, const FPointerEvent& InPointerEvent) override
 	{
-		OnDoubleClick.ExecuteIfBound();
+		if (InPointerEvent.GetEffectingButton() == EKeys::LeftMouseButton && OnDoubleClick.IsBound())
+		{
+			return OnDoubleClick.Execute();
+		}
 
-		return FReply::Handled();
+		return SCheckBox::OnMouseButtonDoubleClick(InGeometry, InPointerEvent);
 	}
 
 	//~ End SWidget
 
-	void Internal_Click(const FPointerEvent& InPointerEvent)
-	{
-		// Use a timer for click to give double click a chance to register
-		TimerHandle = RegisterActiveTimer(GetDefault<UInputSettings>()->DoubleClickTime,
-			FWidgetActiveTimerDelegate::CreateLambda([this, InPointerEvent](double InCurrentTime, float InDeltaTime) -> EActiveTimerReturnType
-			{
-				if (InPointerEvent.IsControlDown())
-				{
-					OnCtrlClick.ExecuteIfBound();
-				}
-				else if (InPointerEvent.IsAltDown())
-				{
-					OnAltClick.ExecuteIfBound();
-				}
-				else if (InPointerEvent.GetEffectingButton() == EKeys::MiddleMouseButton)
-				{
-					OnMiddleButtonClick.ExecuteIfBound();
-				}
-				else
-				{
-					OnClick.ExecuteIfBound(GetCheckedState());
-				}
-
-				return EActiveTimerReturnType::Stop;
-			}));
-	}
-
-	FOnCheckStateChanged OnClick;
-	FSimpleDelegate OnCtrlClick;
-	FSimpleDelegate OnAltClick;
-	FSimpleDelegate OnDoubleClick;
-	FSimpleDelegate OnMiddleButtonClick;
-
-	TSharedPtr<FActiveTimerHandle> TimerHandle;
+	FOnClicked OnCtrlClick;
+	FOnClicked OnAltClick;
+	FOnClicked OnDoubleClick;
+	FOnClicked OnMiddleButtonClick;
 };
