Sequencer: Fix 'Rebind Posssessable References' menu item to skip custom bindings, as this can create broken bindings.

diff --git a/Engine/Source/Editor/Sequencer/Private/Sequencer.cpp b/Engine/Source/Editor/Sequencer/Private/Sequencer.cpp
--- a/Engine/Source/Editor/Sequencer/Private/Sequencer.cpp
+++ b/Engine/Source/Editor/Sequencer/Private/Sequencer.cpp
@@ -10577,12 +10577,19 @@ void FSequencer::RebindPossessableReferences()
 
 	TMap<FGuid, TArray<UObject*, TInlineAllocator<1>>> AllObjects;
 
-	UObject* PlaybackContext = PlaybackContextAttribute.Get(nullptr);
+	UObject* PlaybackContext = PlaybackContextAttribute.Get(nullptr); 
+	const FMovieSceneBindingReferences* Refs = FocusedSequence->GetBindingReferences();
 
 	for (int32 Index = 0; Index < FocusedMovieScene->GetPossessableCount(); Index++)
 	{
 		const FMovieScenePossessable& Possessable = FocusedMovieScene->GetPossessable(Index);
 
+		// Skip custom bindings here
+		if (Refs && Refs->GetCustomBinding(Possessable.GetGuid(), 0))
+		{
+			continue;
+		}
+
 		TArray<UObject*, TInlineAllocator<1>>& References = AllObjects.FindOrAdd(Possessable.GetGuid());
 		TArrayView<TWeakObjectPtr<>> BoundObjects = State.FindBoundObjects(Possessable.GetGuid(), GetFocusedTemplateID(), GetSharedPlaybackState());
 		for (TWeakObjectPtr<> WeakObj : BoundObjects)
