Sequencer: Outliner Decorators

IOutlinerDecorators can now be registered with Sequencer. The Outliner Decorator Column will gather all compatible IOutlinerDecorators for the row and display them. If more than one, the icons will not appear, just the colored backgrounds.

There's currently two IOutlinerDecorators - TimeWarp and Condition

diff --git a/Engine/Source/Editor/EditorStyle/Private/StarshipStyle.cpp b/Engine/Source/Editor/EditorStyle/Private/StarshipStyle.cpp
--- a/Engine/Source/Editor/EditorStyle/Private/StarshipStyle.cpp
+++ b/Engine/Source/Editor/EditorStyle/Private/StarshipStyle.cpp
@@ -2401,13 +2401,14 @@ void FStarshipEditorStyle::FStyle::SetupSequencerStyles()
 	// Sequencer outliner column UI
 	if (IncludeEditorSpecificStyles())
 	{
-		Set("Sequencer.Column.Condition", new IMAGE_BRUSH_SVG("Sequencer/Column_Widgets/SequencerCondition", Icon14x14));
 		Set("Sequencer.Column.Mute", new IMAGE_BRUSH_SVG("Sequencer/Column_Widgets/SequencerMute", Icon14x14));
 		Set("Sequencer.Column.Locked", new IMAGE_BRUSH_SVG("Sequencer/Column_Widgets/SequencerLocked", Icon14x14));
 		Set("Sequencer.Column.Solo", new IMAGE_BRUSH_SVG("Sequencer/Column_Widgets/SequencerSolo", Icon14x14));
 		Set("Sequencer.Column.Unpinned", new IMAGE_BRUSH_SVG("Sequencer/Column_Widgets/SequencerUnpinned", Icon14x14));
 		Set("Sequencer.Column.CheckBoxIndeterminate", new CORE_IMAGE_BRUSH_SVG("Starship/CoreWidgets/CheckBox/CheckBoxIndeterminate_14", Icon14x14));
 		Set("Sequencer.Column.OutlinerColumnBox", new FSlateRoundedBoxBrush(FStyleColors::Header, 2.0f));
+		Set("Sequencer.Decorator.Condition", new IMAGE_BRUSH_SVG("Sequencer/Column_Widgets/SequencerCondition", Icon12x12));
+		Set("Sequencer.Decorator.TimeWarp", new IMAGE_BRUSH_SVG("Sequencer/Column_Widgets/SequencerTimeWarp", Icon12x12));
 	}
 }
 
diff --git a/Engine/Source/Editor/Sequencer/Private/MVVM/ViewModels/OutlinerColumns/IndicatorOutlinerColumn.cpp b/Engine/Source/Editor/Sequencer/Private/MVVM/ViewModels/OutlinerColumns/IndicatorOutlinerColumn.cpp
deleted file mode 100644
diff --git a/Engine/Source/Editor/Sequencer/Private/MVVM/ViewModels/OutlinerColumns/IndicatorOutlinerColumn.h b/Engine/Source/Editor/Sequencer/Private/MVVM/ViewModels/OutlinerColumns/IndicatorOutlinerColumn.h
deleted file mode 100644
diff --git a/Engine/Source/Editor/Sequencer/Private/MVVM/ViewModels/OutlinerColumns/OutlinerDecoratorColumn.cpp b/Engine/Source/Editor/Sequencer/Private/MVVM/ViewModels/OutlinerColumns/OutlinerDecoratorColumn.cpp
new file mode 100644
--- /dev/null
+++ b/Engine/Source/Editor/Sequencer/Private/MVVM/ViewModels/OutlinerColumns/OutlinerDecoratorColumn.cpp
@@ -0,0 +1,95 @@
+// Copyright Epic Games, Inc. All Rights Reserved.
+
+#include "MVVM/ViewModels/OutlinerColumns/OutlinerDecoratorColumn.h"
+
+#include "Sequencer.h"
+#include "MVVM/Extensions/IConditionableExtension.h"
+#include "MVVM/SharedViewModelData.h"
+#include "MVVM/ViewModels/OutlinerDecorators/IOutlinerDecorator.h"
+#include "MVVM/ViewModels/SequencerEditorViewModel.h"
+#include "Widgets/Layout/SBorder.h"
+#include "Widgets/OutlinerDecorators/SConditionDecoratorWidget.h"
+#include "Widgets/SBoxPanel.h"
+
+#define LOCTEXT_NAMESPACE "FOutlinerDecoratorColumn"
+
+namespace UE::Sequencer
+{
+
+FOutlinerDecoratorColumn::FOutlinerDecoratorColumn()
+{
+	Name = FCommonOutlinerNames::Decorator;
+	Label = LOCTEXT("DecoratorColumnLabel", "Decorators");
+	Position = FOutlinerColumnPosition{ 0, EOutlinerColumnGroup::FarLeftGutter };
+	Layout = FOutlinerColumnLayout{ 14, FMargin(0.f, 0.f), HAlign_Fill, VAlign_Fill, EOutlinerColumnSizeMode::Fixed, EOutlinerColumnFlags::None };
+}
+
+bool FOutlinerDecoratorColumn::IsItemCompatibleWithColumn(const FCreateOutlinerColumnParams& InParams) const
+{
+	TViewModelPtr<FSequencerEditorViewModel> Editor = InParams.Editor->CastThisShared<FSequencerEditorViewModel>();
+	if (!Editor)
+	{
+		return false;
+	}
+
+	TSharedPtr<FSequencer> Sequencer = Editor->GetSequencerImpl();
+	if (!Sequencer)
+	{
+		return false;
+	}
+
+	for (const TTuple< FName, TSharedPtr<IOutlinerDecorator> >& OutlinerDecorator : Sequencer->GetOutlinerDecorators())
+	{
+		if (OutlinerDecorator.Value.Get()->IsItemCompatibleWithDecorator(InParams))
+		{
+			return true;
+		}
+	}
+
+	return false;
+}
+
+TSharedPtr<SWidget> FOutlinerDecoratorColumn::CreateColumnWidget(const FCreateOutlinerColumnParams& InParams, const TSharedRef<ISequencerTreeViewRow>& TreeViewRow)
+{
+	TViewModelPtr<FSequencerEditorViewModel> Editor = InParams.Editor->CastThisShared<FSequencerEditorViewModel>();
+	if (!Editor)
+	{
+		return nullptr;
+	}
+
+	TSharedPtr<FSequencer> Sequencer = Editor->GetSequencerImpl();
+	if (!Sequencer)
+	{
+		return nullptr;
+	}
+
+	ColumnWidget = SNew(SHorizontalBox);
+
+	TArray<TSharedPtr<IOutlinerDecorator> > CompatibleDecorators;
+
+	for (const TTuple< FName, TSharedPtr<IOutlinerDecorator> >& OutlinerDecorator : Sequencer->GetOutlinerDecorators())
+	{
+		TSharedPtr<IOutlinerDecorator> Decorator = OutlinerDecorator.Value;
+
+		if (Decorator->IsItemCompatibleWithDecorator(InParams))
+		{
+			CompatibleDecorators.Add(Decorator);
+		}
+	}
+
+	const int32 NumCompatibleDecorators = CompatibleDecorators.Num();
+
+	for (TSharedPtr<IOutlinerDecorator>& Decorator : CompatibleDecorators)
+	{
+		ColumnWidget->AddSlot()
+		[
+			Decorator->CreateDecoratorWidget(InParams, TreeViewRow, SharedThis(this), NumCompatibleDecorators).ToSharedRef()
+		];
+	}
+
+	return ColumnWidget;
+}
+
+} // namespace UE::Sequencer
+
+#undef LOCTEXT_NAMESPACE
\ No newline at end of file
diff --git a/Engine/Source/Editor/Sequencer/Private/MVVM/ViewModels/OutlinerColumns/OutlinerDecoratorColumn.h b/Engine/Source/Editor/Sequencer/Private/MVVM/ViewModels/OutlinerColumns/OutlinerDecoratorColumn.h
new file mode 100644
--- /dev/null
+++ b/Engine/Source/Editor/Sequencer/Private/MVVM/ViewModels/OutlinerColumns/OutlinerDecoratorColumn.h
@@ -0,0 +1,29 @@
+// Copyright Epic Games, Inc. All Rights Reserved.
+
+#pragma once
+
+#include "MVVM/ViewModels/OutlinerColumns/OutlinerColumnBase.h"
+
+class SHorizontalBox;
+
+namespace UE::Sequencer
+{
+
+/**
+ * A column for showing various decorators on the presence of features (e.g. conditions, time warp) on that row.
+ */
+class FOutlinerDecoratorColumn
+	: public FOutlinerColumnBase
+{
+public:
+
+	FOutlinerDecoratorColumn();
+
+	bool IsItemCompatibleWithColumn(const FCreateOutlinerColumnParams& InParams) const override;
+	TSharedPtr<SWidget> CreateColumnWidget(const FCreateOutlinerColumnParams& InParams, const TSharedRef<ISequencerTreeViewRow>& TreeViewRow) override;
+
+private:
+	TSharedPtr<SHorizontalBox> ColumnWidget;
+};
+
+} // namespace UE::Sequencer
\ No newline at end of file
diff --git a/Engine/Source/Editor/Sequencer/Private/MVVM/ViewModels/OutlinerDecorators/ConditionOutlinerDecorator.cpp b/Engine/Source/Editor/Sequencer/Private/MVVM/ViewModels/OutlinerDecorators/ConditionOutlinerDecorator.cpp
new file mode 100644
--- /dev/null
+++ b/Engine/Source/Editor/Sequencer/Private/MVVM/ViewModels/OutlinerDecorators/ConditionOutlinerDecorator.cpp
@@ -0,0 +1,77 @@
+// Copyright Epic Games, Inc. All Rights Reserved.
+
+#include "MVVM/ViewModels/OutlinerDecorators/ConditionOutlinerDecorator.h"
+
+#include "MVVM/ViewModels/OutlinerColumns/IOutlinerColumn.h"
+#include "MVVM/ViewModels/OutlinerColumns/OutlinerColumnTypes.h"
+#include "Widgets/OutlinerDecorators/SConditionDecoratorWidget.h"
+#include "MVVM/Extensions/IConditionableExtension.h"
+#include "MVVM/SharedViewModelData.h"
+#include "Widgets/Layout/SBorder.h"
+
+#define LOCTEXT_NAMESPACE "FConditionOutlinerDecorator"
+
+namespace UE::Sequencer
+{
+
+FConditionOutlinerDecorator::FConditionOutlinerDecorator()
+{
+	Opacity = 1.f;
+}
+
+FName FConditionOutlinerDecorator::GetDecoratorName() const
+{
+	return FCommonOutlinerNames::Condition;
+}
+
+bool FConditionOutlinerDecorator::IsItemCompatibleWithDecorator(const FCreateOutlinerColumnParams& InParams) const
+{
+	if (FConditionStateCacheExtension* ConditionStateCache = InParams.OutlinerExtension.AsModel()->GetSharedData()->CastThis<FConditionStateCacheExtension>())
+	{
+		return EnumHasAnyFlags(ConditionStateCache->GetCachedFlags(InParams.OutlinerExtension), ECachedConditionState::HasCondition | ECachedConditionState::ParentHasCondition | ECachedConditionState::ChildHasCondition | ECachedConditionState::SectionHasCondition);
+	}
+
+	return false;
+}
+
+TSharedPtr<SWidget> FConditionOutlinerDecorator::CreateDecoratorWidget(const FCreateOutlinerColumnParams& InParams, const TSharedRef<ISequencerTreeViewRow>& TreeViewRow, const TSharedRef<IOutlinerColumn>& OutlinerColumn, const int32 NumCompatibleDecorators)
+{
+	static const FLinearColor ConditionColor = FLinearColor::FromSRGBColor(FColor(92, 220, 205));
+
+	if (NumCompatibleDecorators > 1)
+	{
+		return SNew(SBorder)
+			.VAlign(VAlign_Fill)
+			.HAlign(HAlign_Fill)
+			.Padding(0.0f)
+			.BorderImage(FAppStyle::GetBrush("WhiteBrush"))
+			.BorderBackgroundColor_Lambda([this]() -> FLinearColor { FLinearColor BackgroundColor = ConditionColor; BackgroundColor.A = Opacity; return BackgroundColor; });
+	}
+
+	return SNew(SOverlay)
+		+ SOverlay::Slot()
+		[
+			SNew(SBorder)
+				.VAlign(VAlign_Fill)
+				.HAlign(HAlign_Fill)
+				.Padding(0.0f)
+				.BorderImage(FAppStyle::GetBrush("WhiteBrush"))
+				.BorderBackgroundColor_Lambda([this]() -> FLinearColor { FLinearColor BackgroundColor = ConditionColor; BackgroundColor.A = Opacity; return BackgroundColor; })
+		]
+
+		+ SOverlay::Slot()
+		[
+			SNew(SBox)
+				.WidthOverride(12.f)
+				.HeightOverride(12.f)
+				.HAlign(HAlign_Center)
+				.VAlign(VAlign_Center)
+				[
+					SAssignNew(DecoratorWidget, SConditionDecoratorWidget, OutlinerColumn, SharedThis(this), InParams)
+				]
+		];
+}
+
+} // namespace UE::Sequencer
+
+#undef LOCTEXT_NAMESPACE
\ No newline at end of file
diff --git a/Engine/Source/Editor/Sequencer/Private/MVVM/ViewModels/OutlinerDecorators/ConditionOutlinerDecorator.h b/Engine/Source/Editor/Sequencer/Private/MVVM/ViewModels/OutlinerDecorators/ConditionOutlinerDecorator.h
new file mode 100644
--- /dev/null
+++ b/Engine/Source/Editor/Sequencer/Private/MVVM/ViewModels/OutlinerDecorators/ConditionOutlinerDecorator.h
@@ -0,0 +1,33 @@
+// Copyright Epic Games, Inc. All Rights Reserved.
+
+#pragma once
+
+#include "MVVM/ViewModels/OutlinerDecorators/OutlinerDecoratorBase.h"
+
+namespace UE::Sequencer
+{
+
+struct FCreateOutlinerColumnParams;
+class IOutlinerColumn;
+class SConditionDecoratorWidget;
+
+class FConditionOutlinerDecorator
+	: public FOutlinerDecoratorBase
+{
+public:
+	FConditionOutlinerDecorator();
+	
+	virtual FName GetDecoratorName() const override;
+	virtual bool IsItemCompatibleWithDecorator(const FCreateOutlinerColumnParams& InParams) const override;
+	virtual TSharedPtr<SWidget> CreateDecoratorWidget(const FCreateOutlinerColumnParams& InParams, const TSharedRef<ISequencerTreeViewRow>& TreeViewRow, const TSharedRef<IOutlinerColumn>& OutlinerColumn, const int32 NumCompatibleDecorators) override;
+
+protected:
+	friend class SConditionDecoratorWidget;
+	mutable float Opacity;
+
+private:
+	TSharedPtr<SConditionDecoratorWidget> DecoratorWidget;
+};
+
+} // namespace UE::Sequencer
+
diff --git a/Engine/Source/Editor/Sequencer/Private/MVVM/ViewModels/OutlinerDecorators/TimeWarpOutlinerDecorator.cpp b/Engine/Source/Editor/Sequencer/Private/MVVM/ViewModels/OutlinerDecorators/TimeWarpOutlinerDecorator.cpp
new file mode 100644
--- /dev/null
+++ b/Engine/Source/Editor/Sequencer/Private/MVVM/ViewModels/OutlinerDecorators/TimeWarpOutlinerDecorator.cpp
@@ -0,0 +1,97 @@
+// Copyright Epic Games, Inc. All Rights Reserved.
+
+#include "MVVM/ViewModels/OutlinerDecorators/TimeWarpOutlinerDecorator.h"
+
+#include "MovieSceneSection.h"
+#include "MVVM/Extensions/ITrackExtension.h"
+#include "MVVM/SharedViewModelData.h"
+#include "MVVM/ViewModels/SectionModel.h"
+#include "MVVM/ViewModels/OutlinerColumns/IOutlinerColumn.h"
+#include "MVVM/ViewModels/OutlinerColumns/OutlinerColumnTypes.h"
+#include "Variants/MovieSceneTimeWarpVariant.h"
+#include "Variants/MovieSceneTimeWarpGetter.h"
+#include "Widgets/Layout/SBorder.h"
+#include "Widgets/OutlinerDecorators/STimeWarpDecoratorWidget.h"
+
+#define LOCTEXT_NAMESPACE "FTimeWarpOutlinerDecorator"
+
+namespace UE::Sequencer
+{
+
+FTimeWarpOutlinerDecorator::FTimeWarpOutlinerDecorator()
+{
+}
+
+FName FTimeWarpOutlinerDecorator::GetDecoratorName() const
+{
+	return FCommonOutlinerNames::TimeWarp;
+}
+
+bool FTimeWarpOutlinerDecorator::IsItemCompatibleWithDecorator(const FCreateOutlinerColumnParams& InParams) const
+{
+	if (InParams.OutlinerExtension.AsModel()->IsA<ITrackExtension>())
+	{
+		ITrackExtension* TrackModel = InParams.OutlinerExtension.AsModel()->CastThis<ITrackExtension>();
+
+		for (const TViewModelPtr<FSectionModel>& SectionModel : TrackModel->GetSectionModels().IterateSubList<FSectionModel>())
+		{
+			UMovieSceneSection* Section = SectionModel->GetSection();
+			if (!Section)
+			{
+				continue;
+			}
+
+			if (FMovieSceneTimeWarpVariant* Variant = Section->GetTimeWarp())
+			{
+				if (Variant->GetType() == EMovieSceneTimeWarpType::Custom)
+				{
+					return true;
+				}
+			}
+		}
+	}
+
+	return false;
+}
+
+TSharedPtr<SWidget> FTimeWarpOutlinerDecorator::CreateDecoratorWidget(const FCreateOutlinerColumnParams& InParams, const TSharedRef<ISequencerTreeViewRow>& TreeViewRow, const TSharedRef<IOutlinerColumn>& OutlinerColumn, const int32 NumCompatibleDecorators)
+{
+	static const FLinearColor ConditionColor = FLinearColor::FromSRGBColor(FColor(212, 147, 20));
+
+	if (NumCompatibleDecorators > 1)
+	{
+		return SNew(SBorder)
+			.VAlign(VAlign_Fill)
+			.HAlign(HAlign_Fill)
+			.Padding(0.0f)
+			.BorderImage(FAppStyle::GetBrush("WhiteBrush"))
+			.BorderBackgroundColor_Lambda([this]() -> FLinearColor { FLinearColor BackgroundColor = ConditionColor; BackgroundColor.A = Opacity; return BackgroundColor; });
+	}
+
+	return SNew(SOverlay)
+		+ SOverlay::Slot()
+		[
+			SNew(SBorder)
+				.VAlign(VAlign_Fill)
+				.HAlign(HAlign_Fill)
+				.Padding(0.0f)
+				.BorderImage(FAppStyle::GetBrush("WhiteBrush"))
+				.BorderBackgroundColor_Lambda([this]() -> FLinearColor { FLinearColor BackgroundColor = ConditionColor; BackgroundColor.A = Opacity; return BackgroundColor; })
+		]
+
+		+ SOverlay::Slot()
+		[
+			SNew(SBox)
+				.WidthOverride(12.f)
+				.HeightOverride(12.f)
+				.HAlign(HAlign_Center)
+				.VAlign(VAlign_Center)
+				[
+					SAssignNew(DecoratorWidget, STimeWarpDecoratorWidget, OutlinerColumn, SharedThis(this), InParams)
+				]
+		];
+}
+
+} // namespace UE::Sequencer
+
+#undef LOCTEXT_NAMESPACE
\ No newline at end of file
diff --git a/Engine/Source/Editor/Sequencer/Private/MVVM/ViewModels/OutlinerDecorators/TimeWarpOutlinerDecorator.h b/Engine/Source/Editor/Sequencer/Private/MVVM/ViewModels/OutlinerDecorators/TimeWarpOutlinerDecorator.h
new file mode 100644
--- /dev/null
+++ b/Engine/Source/Editor/Sequencer/Private/MVVM/ViewModels/OutlinerDecorators/TimeWarpOutlinerDecorator.h
@@ -0,0 +1,33 @@
+// Copyright Epic Games, Inc. All Rights Reserved.
+
+#pragma once
+
+#include "MVVM/ViewModels/OutlinerDecorators/OutlinerDecoratorBase.h"
+
+namespace UE::Sequencer
+{
+
+struct FCreateOutlinerColumnParams;
+class IOutlinerColumn;
+class STimeWarpDecoratorWidget;
+
+class FTimeWarpOutlinerDecorator
+	: public FOutlinerDecoratorBase
+{
+public:
+	FTimeWarpOutlinerDecorator();
+	
+	virtual FName GetDecoratorName() const override;
+	virtual bool IsItemCompatibleWithDecorator(const FCreateOutlinerColumnParams& InParams) const override;
+	virtual TSharedPtr<SWidget> CreateDecoratorWidget(const FCreateOutlinerColumnParams& InParams, const TSharedRef<ISequencerTreeViewRow>& TreeViewRow, const TSharedRef<IOutlinerColumn>& OutlinerColumn, const int32 NumCompatibleDecorators) override;
+
+protected:
+	friend class STimeWarpDecoratorWidget;
+	mutable float Opacity;
+
+private:
+	TSharedPtr<STimeWarpDecoratorWidget> DecoratorWidget;
+};
+
+} // namespace UE::Sequencer
+
diff --git a/Engine/Source/Editor/Sequencer/Private/Sequencer.cpp b/Engine/Source/Editor/Sequencer/Private/Sequencer.cpp
--- a/Engine/Source/Editor/Sequencer/Private/Sequencer.cpp
+++ b/Engine/Source/Editor/Sequencer/Private/Sequencer.cpp
@@ -131,6 +131,7 @@
 #include "ISequencerChannelInterface.h"
 #include "IMovieRendererInterface.h"
 #include "MVVM/ViewModels/OutlinerColumns/IOutlinerColumn.h"
+#include "MVVM/ViewModels/OutlinerDecorators/IOutlinerDecorator.h"
 #include "SequencerKeyCollection.h"
 #include "CurveEditor.h"
 #include "CurveEditorScreenSpace.h"
@@ -298,7 +299,12 @@ namespace UE
 const FName FSequencer::SelectionDrawerId = TEXT("SelectionDetails");
 bool FSequencer::bSelectionLimited = false;
 
-void FSequencer::InitSequencer(const FSequencerInitParams& InitParams, const TSharedRef<ISequencerObjectChangeListener>& InObjectChangeListener, const TArray<FOnCreateTrackEditor>& TrackEditorDelegates, const TArray<FOnCreateEditorObjectBinding>& EditorObjectBindingDelegates, const TArray<FOnCreateOutlinerColumn>& OutlinerColumnDelegates)
+void FSequencer::InitSequencer(const FSequencerInitParams& InitParams
+	, const TSharedRef<ISequencerObjectChangeListener>& InObjectChangeListener
+	, const TArray<FOnCreateTrackEditor>& TrackEditorDelegates
+	, const TArray<FOnCreateEditorObjectBinding>& EditorObjectBindingDelegates
+	, const TArray<FOnCreateOutlinerColumn>& OutlinerColumnDelegates
+	, const TArray<FOnCreateOutlinerDecorator>& OutlinerDecoratorDelegates)
 {
 	using namespace UE::MovieScene;
 	using namespace UE::Sequencer;
@@ -524,6 +530,19 @@ void FSequencer::InitSequencer(const FSequencerInitParams& InitParams, const TSh
 		}
 	}
 
+	// OutlinerDecorators are registered to be provided to SOutlinerView by SSequencer
+	for (int32 DelegateIndex = 0; DelegateIndex < OutlinerDecoratorDelegates.Num(); ++DelegateIndex)
+	{
+		check(OutlinerDecoratorDelegates[DelegateIndex].IsBound());
+		TSharedRef<IOutlinerDecorator> OutlinerDecorator = OutlinerDecoratorDelegates[DelegateIndex].Execute();
+
+		if (OutlinerDecorator->SupportsSequence(InitParams.RootSequence))
+		{
+			check(!OutlinerDecorators.Contains(OutlinerDecorator->GetDecoratorName()));
+			OutlinerDecorators.Add(OutlinerDecorator->GetDecoratorName(), OutlinerDecorator);
+		}
+	}
+
 	TWeakPtr<UE::MovieScene::IDeferredSignedObjectChangeHandler> WeakDeferredSignedObjectChangeHandler = UMovieSceneSignedObject::GetDeferredHandler();
 	if (WeakDeferredSignedObjectChangeHandler.IsValid())
 	{
diff --git a/Engine/Source/Editor/Sequencer/Private/Sequencer.h b/Engine/Source/Editor/Sequencer/Private/Sequencer.h
--- a/Engine/Source/Editor/Sequencer/Private/Sequencer.h
+++ b/Engine/Source/Editor/Sequencer/Private/Sequencer.h
@@ -140,7 +140,7 @@ public:
 	 * @param TrackEditorDelegates Delegates to call to create auto-key handlers for this sequencer.
 	 * @param EditorObjectBindingDelegates Delegates to call to create object bindings for this sequencer.
 	 */
-	void InitSequencer(const FSequencerInitParams& InitParams, const TSharedRef<ISequencerObjectChangeListener>& InObjectChangeListener, const TArray<FOnCreateTrackEditor>& TrackEditorDelegates, const TArray<FOnCreateEditorObjectBinding>& EditorObjectBindingDelegatess, const TArray<FOnCreateOutlinerColumn>& OutlinerColumnDelegates);
+	void InitSequencer(const FSequencerInitParams& InitParams, const TSharedRef<ISequencerObjectChangeListener>& InObjectChangeListener, const TArray<FOnCreateTrackEditor>& TrackEditorDelegates, const TArray<FOnCreateEditorObjectBinding>& EditorObjectBindingDelegatess, const TArray<FOnCreateOutlinerColumn>& OutlinerColumnDelegates, const TArray<FOnCreateOutlinerDecorator>& OutlinerDecoratorDelegates);
 
 	/**
 	 * Reinitializes sequencer after the playback context has changed
@@ -314,6 +314,14 @@ public:
 		return OutlinerColumns;
 	}
 
+	/**
+	* @return Outliner Decorator Items registered to the sequencer by column name
+	*/
+	const TMap<FName, TSharedPtr<UE::Sequencer::IOutlinerDecorator>>& GetOutlinerDecorators() const
+	{
+		return OutlinerDecorators;
+	}
+
 public:
 
 	/** @return The set of vertical frames */
@@ -1269,6 +1277,9 @@ private:
 	/** List of Outliner column creators that are supported by the Sequencer. */
 	TMap<FName, TSharedPtr<UE::Sequencer::IOutlinerColumn>> OutlinerColumns;
 
+	/** List of Outliner decorator item creators that are supported by the Sequencer. */
+	TMap<FName, TSharedPtr<UE::Sequencer::IOutlinerDecorator>> OutlinerDecorators;
+
 	/** List of object bindings we can use */
 	TArray<TSharedPtr<ISequencerEditorObjectBinding>> ObjectBindings;
 
diff --git a/Engine/Source/Editor/Sequencer/Private/SequencerModule.cpp b/Engine/Source/Editor/Sequencer/Private/SequencerModule.cpp
--- a/Engine/Source/Editor/Sequencer/Private/SequencerModule.cpp
+++ b/Engine/Source/Editor/Sequencer/Private/SequencerModule.cpp
@@ -33,7 +33,7 @@
 #include "MVVM/ViewModels/SequenceModel.h"
 #include "MVVM/ViewModels/SequencerEditorViewModel.h"
 
-#include "MVVM/ViewModels/OutlinerColumns/IndicatorOutlinerColumn.h"
+#include "MVVM/ViewModels/OutlinerColumns/OutlinerDecoratorColumn.h"
 #include "MVVM/ViewModels/OutlinerColumns/LockOutlinerColumn.h"
 #include "MVVM/ViewModels/OutlinerColumns/MuteOutlinerColumn.h"
 #include "MVVM/ViewModels/OutlinerColumns/PinOutlinerColumn.h"
@@ -45,6 +45,9 @@
 #include "MVVM/ViewModels/OutlinerColumns/KeyFrameOutlinerColumn.h"
 #include "MVVM/ViewModels/OutlinerColumns/ColorPickerOutlinerColumn.h"
 
+#include "MVVM/ViewModels/OutlinerDecorators/ConditionOutlinerDecorator.h"
+#include "MVVM/ViewModels/OutlinerDecorators/TimeWarpOutlinerDecorator.h"
+
 #include "ToolMenus.h"
 #include "ContentBrowserMenuContexts.h"
 #include "SequencerUtilities.h"
@@ -294,7 +297,7 @@ public:
 
 		OnPreSequencerInit.Broadcast(Sequencer, ObjectChangeListener, InitParams);
 
-		Sequencer->InitSequencer(InitParams, ObjectChangeListener, TrackEditorDelegates, EditorObjectBindingDelegates, OutlinerColumnDelegates);
+		Sequencer->InitSequencer(InitParams, ObjectChangeListener, TrackEditorDelegates, EditorObjectBindingDelegates, OutlinerColumnDelegates, OutlinerDecoratorDelegates);
 
 		OnSequencerCreated.Broadcast(Sequencer);
 
@@ -357,6 +360,15 @@ public:
 		OutlinerColumnDelegates.RemoveAll([=](const FOnCreateOutlinerColumn& Delegate) { return Delegate.GetHandle() == InHandle; });
 	}
 
+	virtual FDelegateHandle RegisterOutlinerDecorator(FOnCreateOutlinerDecorator InCreator) override {
+		OutlinerDecoratorDelegates.Add(InCreator);
+		return OutlinerDecoratorDelegates.Last().GetHandle();
+	}
+
+	virtual void UnregisterOutlinerDecorator(FDelegateHandle InHandle) override {
+		OutlinerDecoratorDelegates.RemoveAll([=](const FOnCreateOutlinerDecorator& Delegate) { return Delegate.GetHandle() == InHandle; });
+	}
+
 	virtual FDelegateHandle RegisterOnSequencerCreated(FOnSequencerCreated::FDelegate InOnSequencerCreated) override
 	{
 		return OnSequencerCreated.Add(InOnSequencerCreated);
@@ -481,7 +493,7 @@ public:
 			OnGetGlobalRowExtensionHandle = EditModule.GetGlobalRowExtensionDelegate().AddStatic(&RegisterKeyframeExtensionHandler);
 
 			// Register far left gutter columns
-			IndicatorOutlinerColumnHandle = RegisterOutlinerColumn(FOnCreateOutlinerColumn::CreateStatic([] { return TSharedRef<IOutlinerColumn>(MakeShared<FIndicatorOutlinerColumn>()); }));
+			OutlinerDecoratorColumnHandle = RegisterOutlinerColumn(FOnCreateOutlinerColumn::CreateStatic([] { return TSharedRef<IOutlinerColumn>(MakeShared<FOutlinerDecoratorColumn>()); }));
 
 			// Register left gutter columns
 			PinOutlinerColumnHandle  = RegisterOutlinerColumn(FOnCreateOutlinerColumn::CreateStatic([]{ return TSharedRef<IOutlinerColumn>(MakeShared<FPinOutlinerColumn>()); }));
@@ -499,6 +511,10 @@ public:
 			NavOutlinerColumnHandle          = RegisterOutlinerColumn(FOnCreateOutlinerColumn::CreateStatic([]{ return TSharedRef<IOutlinerColumn>(MakeShared<FNavOutlinerColumn>()); }));
 			ColorPickerOutlinerColumnHandle  = RegisterOutlinerColumn(FOnCreateOutlinerColumn::CreateStatic([]{ return TSharedRef<IOutlinerColumn>(MakeShared<FColorPickerOutlinerColumn>()); }));
 
+			// Register outliner decorator items
+			ConditionOutlinerDecoratorHandle = RegisterOutlinerDecorator(FOnCreateOutlinerDecorator::CreateStatic([] { return TSharedRef<IOutlinerDecorator>(MakeShared<FConditionOutlinerDecorator>()); }));
+			TimeWarpOutlinerDecoratorHandle = RegisterOutlinerDecorator(FOnCreateOutlinerDecorator::CreateStatic([] { return TSharedRef<IOutlinerDecorator>(MakeShared<FTimeWarpOutlinerDecorator>()); }));
+
 			RegisterObjectSchemas();
 		}
 
@@ -546,7 +562,7 @@ public:
 			FEditorModeRegistry::Get().UnregisterMode(FSequencerEdMode::EM_SequencerMode);
 
 			// unregister outliner columns
-			UnregisterOutlinerColumn(IndicatorOutlinerColumnHandle);
+			UnregisterOutlinerColumn(OutlinerDecoratorColumnHandle);
 			UnregisterOutlinerColumn(PinOutlinerColumnHandle);
 			UnregisterOutlinerColumn(MuteOutlinerColumnHandle);
 			UnregisterOutlinerColumn(LockOutlinerColumnHandle);
@@ -557,6 +573,10 @@ public:
 			UnregisterOutlinerColumn(KeyFrameOutlinerColumnHandle);
 			UnregisterOutlinerColumn(NavOutlinerColumnHandle);
 			UnregisterOutlinerColumn(ColorPickerOutlinerColumnHandle);
+
+			// unregister outliner decorator items
+			UnregisterOutlinerDecorator(ConditionOutlinerDecoratorHandle);
+			UnregisterOutlinerDecorator(TimeWarpOutlinerDecoratorHandle);
 		}
 	}
 
@@ -704,6 +724,9 @@ private:
 	/** List of outliner column creators */
 	TArray<FOnCreateOutlinerColumn> OutlinerColumnDelegates;
 
+	/** List of outliner decorator item creators */
+	TArray<FOnCreateOutlinerDecorator> OutlinerDecoratorDelegates;
+
 	TArray<TSharedPtr<UE::Sequencer::IObjectSchema>> ObjectSchemas;
 
 	/** Global details row extension delegate; */
@@ -743,7 +766,7 @@ private:
 	TArray<FMovieRendererEntry> MovieRenderers;
 
 	// Outliner Column Delegate Handles
-	FDelegateHandle IndicatorOutlinerColumnHandle;
+	FDelegateHandle OutlinerDecoratorColumnHandle;
 	FDelegateHandle PinOutlinerColumnHandle;
 	FDelegateHandle MuteOutlinerColumnHandle;
 	FDelegateHandle LockOutlinerColumnHandle;
@@ -754,6 +777,10 @@ private:
 	FDelegateHandle KeyFrameOutlinerColumnHandle;
 	FDelegateHandle NavOutlinerColumnHandle;
 	FDelegateHandle ColorPickerOutlinerColumnHandle;
+
+	// Outliner Decorator Item Delegate Handles
+	FDelegateHandle ConditionOutlinerDecoratorHandle;
+	FDelegateHandle TimeWarpOutlinerDecoratorHandle;
 };
 
 IMPLEMENT_MODULE(FSequencerModule, Sequencer);
diff --git a/Engine/Source/Editor/Sequencer/Private/Widgets/OutlinerColumns/SConditionColumnWidget.cpp b/Engine/Source/Editor/Sequencer/Private/Widgets/OutlinerColumns/SConditionColumnWidget.cpp
deleted file mode 100644
diff --git a/Engine/Source/Editor/Sequencer/Private/Widgets/OutlinerColumns/SConditionColumnWidget.h b/Engine/Source/Editor/Sequencer/Private/Widgets/OutlinerColumns/SConditionColumnWidget.h
deleted file mode 100644
diff --git a/Engine/Source/Editor/Sequencer/Private/Widgets/OutlinerDecorators/SConditionDecoratorWidget.cpp b/Engine/Source/Editor/Sequencer/Private/Widgets/OutlinerDecorators/SConditionDecoratorWidget.cpp
new file mode 100644
--- /dev/null
+++ b/Engine/Source/Editor/Sequencer/Private/Widgets/OutlinerDecorators/SConditionDecoratorWidget.cpp
@@ -0,0 +1,101 @@
+// Copyright Epic Games, Inc. All Rights Reserved.
+
+#include "Widgets/OutlinerDecorators/SConditionDecoratorWidget.h"
+
+#include "MVVM/SharedViewModelData.h"
+#include "MVVM/ViewModels/ViewModelIterators.h"
+#include "MVVM/ViewModels/EditorViewModel.h"
+#include "MVVM/ViewModels/OutlinerColumns/IOutlinerColumn.h"
+#include "MVVM/ViewModels/OutlinerDecorators/IOutlinerDecorator.h"
+#include "MVVM/ViewModels/OutlinerDecorators/ConditionOutlinerDecorator.h"
+#include "MVVM/Extensions/IConditionableExtension.h"
+#include "MVVM/Selection/SequencerCoreSelection.h"
+#include "MVVM/Selection/SequencerOutlinerSelection.h"
+
+namespace UE::Sequencer
+{
+
+void SConditionDecoratorWidget::Construct(const FArguments& InArgs, const TWeakPtr<IOutlinerColumn> InWeakOutlinerColumn, const TWeakPtr<FConditionOutlinerDecorator>& InWeakOutlinerDecorator, const FCreateOutlinerColumnParams& InParams)
+{
+	SColumnToggleWidget::Construct(
+		SColumnToggleWidget::FArguments(),
+		InWeakOutlinerColumn,
+		InParams
+	);
+
+	WeakConditionStateCacheExtension = CastViewModel<FConditionStateCacheExtension>(InParams.OutlinerExtension.AsModel()->GetSharedData());
+
+	WeakOutlinerDecorator = InWeakOutlinerDecorator;
+}
+
+bool SConditionDecoratorWidget::IsActive() const
+{
+	if (TViewModelPtr<FConditionStateCacheExtension> StateCache = WeakConditionStateCacheExtension.Pin())
+	{
+		return EnumHasAnyFlags(StateCache->GetCachedFlags(ModelID), ECachedConditionState::HasCondition);
+	}
+
+	return false;
+}
+
+const FSlateBrush* SConditionDecoratorWidget::GetActiveBrush() const
+{
+	static const FName NAME_ConditionBrush = TEXT("Sequencer.Decorator.Condition");
+	return FAppStyle::Get().GetBrush(NAME_ConditionBrush);
+}
+
+FSlateColor SConditionDecoratorWidget::GetImageColorAndOpacity() const
+{
+	TSharedPtr<FEditorViewModel> Editor = WeakEditor.Pin();
+	TViewModelPtr<IOutlinerExtension> OutlinerItem = WeakOutlinerExtension.Pin();
+
+	bool bActive = IsActive();
+	bool bEvaluatingTrue = false;
+	bool bChildHasCondition = false;
+	bool bSectionHasCondition = false;
+	if (TViewModelPtr<FConditionStateCacheExtension> StateCache = WeakConditionStateCacheExtension.Pin())
+	{
+		bEvaluatingTrue = EnumHasAnyFlags(StateCache->GetCachedFlags(ModelID), ECachedConditionState::ConditionEvaluatingTrue);
+		bChildHasCondition = EnumHasAnyFlags(StateCache->GetCachedFlags(ModelID), ECachedConditionState::ChildHasCondition);
+		bSectionHasCondition = EnumHasAnyFlags(StateCache->GetCachedFlags(ModelID), ECachedConditionState::SectionHasCondition);
+	}
+	FLinearColor OutColor = FLinearColor::Black;
+
+	if (!Editor || !OutlinerItem)
+	{
+		return OutColor;
+	}
+
+	float Opacity = 1.0f;
+
+	if (IsActive())
+	{
+		if (bEvaluatingTrue)
+		{
+			Opacity = 1.0f;
+		}
+		else
+		{
+			Opacity = 0.5f;
+		}
+	}
+	else if ((bChildHasCondition && !OutlinerItem->IsExpanded()) || bSectionHasCondition)
+	{
+		Opacity = 0.25f;
+	}
+	else
+	{
+		// Not active, invisible
+		Opacity = 0.0f;
+	}
+	////OutColor.A = Opacity; // The background should be at the specified opacity and the icon should stay 100% black.
+
+	if (WeakOutlinerDecorator.IsValid())
+	{
+		WeakOutlinerDecorator.Pin()->Opacity = Opacity;
+	}
+
+	return OutColor;
+}
+
+} // namespace UE::Sequencer
diff --git a/Engine/Source/Editor/Sequencer/Private/Widgets/OutlinerDecorators/SConditionDecoratorWidget.h b/Engine/Source/Editor/Sequencer/Private/Widgets/OutlinerDecorators/SConditionDecoratorWidget.h
new file mode 100644
--- /dev/null
+++ b/Engine/Source/Editor/Sequencer/Private/Widgets/OutlinerDecorators/SConditionDecoratorWidget.h
@@ -0,0 +1,65 @@
+// Copyright Epic Games, Inc. All Rights Reserved.
+
+#pragma once
+
+#include "MVVM/Views/OutlinerColumns/SColumnToggleWidget.h"
+
+namespace UE::Sequencer
+{
+
+class IOutlinerColumn;
+class IOutlinerDecorator;
+class FConditionOutlinerDecorator;
+class FConditionStateCacheExtension;
+
+/**
+ * A widget that shows an decorator based on the presence and state of conditions in a row
+ */
+class SConditionDecoratorWidget
+	: public SColumnToggleWidget
+{
+public:
+	SLATE_BEGIN_ARGS(SConditionDecoratorWidget) {}
+	SLATE_END_ARGS()
+
+	void Construct(const FArguments& InArgs, const TWeakPtr<IOutlinerColumn> InWeakOutlinerColumn, const TWeakPtr<FConditionOutlinerDecorator>& OutlinerDecorator, const FCreateOutlinerColumnParams& InParams);
+
+	/** Get the color and opacity of the column toggle widget. */
+	virtual FSlateColor GetImageColorAndOpacity() const override;
+
+protected:
+
+	/** Returns whether or not the item is directly active or not. */
+	virtual bool IsActive() const override;
+
+	/** Sets the active state of this item to the input value. */
+	virtual void SetIsActive(const bool bInIsActive) override { /* no-op for now, no toggle*/ }
+
+	/** Returns true if a child of this item is active. */
+	virtual bool IsChildActive() const override { return false; }
+
+	/** Returns true if this item is implicitly active, but not directly active */
+	virtual bool IsImplicitlyActive() const override { return false; }
+
+	/** Returns the brush to be used to represent a widget is active. */
+	virtual const FSlateBrush* GetActiveBrush() const override;
+
+protected:
+
+	/** Get the image this widget displays. */
+	virtual const FSlateBrush* GetBrush() const override { return GetActiveBrush(); }
+
+protected:
+
+	/** Handles left click-type inputs and potentially begins drag and drop operation. */
+	FReply HandleClick() { return FReply::Handled(); }
+
+private:
+
+	/** Weak cache extension ptr (can be null). */
+	TWeakViewModelPtr<FConditionStateCacheExtension> WeakConditionStateCacheExtension;
+
+	TWeakPtr<FConditionOutlinerDecorator> WeakOutlinerDecorator;
+};
+
+} // namespace UE::Sequencer
\ No newline at end of file
diff --git a/Engine/Source/Editor/Sequencer/Private/Widgets/OutlinerDecorators/STimeWarpDecoratorWidget.cpp b/Engine/Source/Editor/Sequencer/Private/Widgets/OutlinerDecorators/STimeWarpDecoratorWidget.cpp
new file mode 100644
--- /dev/null
+++ b/Engine/Source/Editor/Sequencer/Private/Widgets/OutlinerDecorators/STimeWarpDecoratorWidget.cpp
@@ -0,0 +1,53 @@
+// Copyright Epic Games, Inc. All Rights Reserved.
+
+#include "Widgets/OutlinerDecorators/STimeWarpDecoratorWidget.h"
+
+#include "MVVM/SharedViewModelData.h"
+#include "MVVM/ViewModels/ViewModelIterators.h"
+#include "MVVM/ViewModels/EditorViewModel.h"
+#include "MVVM/ViewModels/OutlinerColumns/IOutlinerColumn.h"
+#include "MVVM/ViewModels/OutlinerDecorators/IOutlinerDecorator.h"
+#include "MVVM/ViewModels/OutlinerDecorators/TimeWarpOutlinerDecorator.h"
+#include "MVVM/Selection/SequencerCoreSelection.h"
+#include "MVVM/Selection/SequencerOutlinerSelection.h"
+
+namespace UE::Sequencer
+{
+
+void STimeWarpDecoratorWidget::Construct(const FArguments& InArgs, const TWeakPtr<IOutlinerColumn> InWeakOutlinerColumn, const TWeakPtr<FTimeWarpOutlinerDecorator>& InWeakOutlinerDecorator, const FCreateOutlinerColumnParams& InParams)
+{
+	SColumnToggleWidget::Construct(
+		SColumnToggleWidget::FArguments(),
+		InWeakOutlinerColumn,
+		InParams
+	);
+
+	WeakOutlinerDecorator = InWeakOutlinerDecorator;
+}
+
+bool STimeWarpDecoratorWidget::IsActive() const
+{
+	return true;
+}
+
+const FSlateBrush* STimeWarpDecoratorWidget::GetActiveBrush() const
+{
+	static const FName NAME_TimeWarpBrush = TEXT("Sequencer.Decorator.TimeWarp");
+	return FAppStyle::Get().GetBrush(NAME_TimeWarpBrush);
+}
+
+FSlateColor STimeWarpDecoratorWidget::GetImageColorAndOpacity() const
+{
+	FLinearColor OutColor = FLinearColor::Black;
+	float Opacity = 1.0f;
+	////OutColor.A = Opacity; // The background should be at the specified opacity and the icon should stay 100% black.
+
+	if (WeakOutlinerDecorator.IsValid())
+	{
+		WeakOutlinerDecorator.Pin()->Opacity = Opacity;
+	}
+
+	return OutColor;
+}
+
+} // namespace UE::Sequencer
diff --git a/Engine/Source/Editor/Sequencer/Private/Widgets/OutlinerDecorators/STimeWarpDecoratorWidget.h b/Engine/Source/Editor/Sequencer/Private/Widgets/OutlinerDecorators/STimeWarpDecoratorWidget.h
new file mode 100644
--- /dev/null
+++ b/Engine/Source/Editor/Sequencer/Private/Widgets/OutlinerDecorators/STimeWarpDecoratorWidget.h
@@ -0,0 +1,61 @@
+// Copyright Epic Games, Inc. All Rights Reserved.
+
+#pragma once
+
+#include "MVVM/Views/OutlinerColumns/SColumnToggleWidget.h"
+
+namespace UE::Sequencer
+{	
+
+class IOutlinerColumn;
+class IOutlinerDecorator;
+class FTimeWarpOutlinerDecorator;
+
+/**
+ * A widget that shows an decorator based on the presence of time warp on a row
+ */
+class STimeWarpDecoratorWidget
+	: public SColumnToggleWidget
+{
+public:
+	SLATE_BEGIN_ARGS(STimeWarpDecoratorWidget) {}
+	SLATE_END_ARGS()
+
+	void Construct(const FArguments& InArgs, const TWeakPtr<IOutlinerColumn> InWeakOutlinerColumn, const TWeakPtr<FTimeWarpOutlinerDecorator>& OutlinerDecorator, const FCreateOutlinerColumnParams& InParams);
+
+	/** Get the color and opacity of the column toggle widget. */
+	virtual FSlateColor GetImageColorAndOpacity() const override;
+
+protected:
+
+	/** Returns whether or not the item is directly active or not. */
+	virtual bool IsActive() const override;
+
+	/** Sets the active state of this item to the input value. */
+	virtual void SetIsActive(const bool bInIsActive) override { /* no-op for now, no toggle*/ }
+
+	/** Returns true if a child of this item is active. */
+	virtual bool IsChildActive() const override { return false; }
+
+	/** Returns true if this item is implicitly active, but not directly active */
+	virtual bool IsImplicitlyActive() const override { return false; }
+
+	/** Returns the brush to be used to represent a widget is active. */
+	virtual const FSlateBrush* GetActiveBrush() const override;
+
+protected:
+
+	/** Get the image this widget displays. */
+	virtual const FSlateBrush* GetBrush() const override { return GetActiveBrush(); }
+
+protected:
+
+	/** Handles left click-type inputs and potentially begins drag and drop operation. */
+	FReply HandleClick() { return FReply::Handled(); }
+
+private:
+
+	TWeakPtr<FTimeWarpOutlinerDecorator> WeakOutlinerDecorator;
+};
+
+} // namespace UE::Sequencer
\ No newline at end of file
diff --git a/Engine/Source/Editor/Sequencer/Public/ISequencerModule.h b/Engine/Source/Editor/Sequencer/Public/ISequencerModule.h
--- a/Engine/Source/Editor/Sequencer/Public/ISequencerModule.h
+++ b/Engine/Source/Editor/Sequencer/Public/ISequencerModule.h
@@ -29,6 +29,7 @@ namespace UE::Sequencer
 	class FTrackModel;
 	class IObjectSchema;
 	class IOutlinerColumn;
+	class IOutlinerDecorator;
 } // namespace UE::Sequencer
 
 enum class ECurveEditorTreeFilterType : uint32;
@@ -63,6 +64,9 @@ DECLARE_DELEGATE_RetVal_OneParam(TSharedPtr<UE::Sequencer::FTrackModel>, FOnCrea
 /** A delegate which will create an outliner column */
 DECLARE_DELEGATE_RetVal(TSharedRef<UE::Sequencer::IOutlinerColumn>, FOnCreateOutlinerColumn);
 
+/** A delegate which will create an outliner decorator item */
+DECLARE_DELEGATE_RetVal(TSharedRef<UE::Sequencer::IOutlinerDecorator>, FOnCreateOutlinerDecorator);
+
 /** A delegate that is executed when adding menu content. */
 DECLARE_DELEGATE_OneParam(FOnGetContextMenuContent, FMenuBuilder& /*MenuBuilder*/);
 
@@ -270,8 +274,22 @@ public:
 	 */
 	virtual void UnregisterOutlinerColumn(FDelegateHandle InHandle) = 0;
 
+	/**
+	 * Registers a delegate that will create an outliner decorator item
+	 *
+	 * @param InCreator Delegate to register
+	 * @return A handle to the newly added delegate
+	 */
+	virtual FDelegateHandle RegisterOutlinerDecorator(FOnCreateOutlinerDecorator InCreator) = 0;
 
-	/** 
+	/**
+	 * Unregisters a previously registered delegate for creating an outliner decorator item
+	 *
+	 * @param InHandle Handle to the delegate to unregister
+	 */
+	virtual void UnregisterOutlinerDecorator(FDelegateHandle InHandle) = 0;
+
+	/**
 	 * Registers a delegate that will be called when a sequencer is created
 	 *
 	 * @param InOnSequencerCreated	Delegate to register.
diff --git a/Engine/Source/Editor/SequencerCore/Private/MVVM/ViewModels/OutlinerColumns/OutlinerColumnTypes.cpp b/Engine/Source/Editor/SequencerCore/Private/MVVM/ViewModels/OutlinerColumns/OutlinerColumnTypes.cpp
--- a/Engine/Source/Editor/SequencerCore/Private/MVVM/ViewModels/OutlinerColumns/OutlinerColumnTypes.cpp
+++ b/Engine/Source/Editor/SequencerCore/Private/MVVM/ViewModels/OutlinerColumns/OutlinerColumnTypes.cpp
@@ -5,7 +5,7 @@
 namespace UE::Sequencer
 {
 
-FName FCommonOutlinerNames::Indicator("Indicator");
+FName FCommonOutlinerNames::Decorator("Decorator");
 FName FCommonOutlinerNames::Pin("Pin");
 FName FCommonOutlinerNames::Lock("Lock");
 FName FCommonOutlinerNames::Mute("Mute");
@@ -16,5 +16,7 @@ FName FCommonOutlinerNames::Add("Add");
 FName FCommonOutlinerNames::Nav("Nav");
 FName FCommonOutlinerNames::KeyFrame("KeyFrame");
 FName FCommonOutlinerNames::ColorPicker("ColorPicker");
+FName FCommonOutlinerNames::TimeWarp("TimeWarp");
+FName FCommonOutlinerNames::Condition("Condition");
 
 } // namespace UE::Sequencer
\ No newline at end of file
diff --git a/Engine/Source/Editor/SequencerCore/Private/MVVM/Views/SOutlinerView.cpp b/Engine/Source/Editor/SequencerCore/Private/MVVM/Views/SOutlinerView.cpp
--- a/Engine/Source/Editor/SequencerCore/Private/MVVM/Views/SOutlinerView.cpp
+++ b/Engine/Source/Editor/SequencerCore/Private/MVVM/Views/SOutlinerView.cpp
@@ -388,11 +388,6 @@ void SOutlinerView::UpdateOutlinerColumns()
 	//    This is implemented this way because columns can be turned on and off dynamically, but we must
 	//    always have a consistent padding within the group. Separators dynamically appear based on the presence of
 	//    each group so we can't put padding on those
-	if (NumFarLeftGutter > 0)
-	{
-		ColumnMetaData->Columns[0].CellPadding.Left += 4.f;
-		ColumnMetaData->Columns[NumFarLeftGutter - 1].CellPadding.Right += 4.f;
-	}
 	if (NumLeftGutter > 0)
 	{
 		ColumnMetaData->Columns[NumFarLeftGutter].CellPadding.Left  += 4.f;
diff --git a/Engine/Source/Editor/SequencerCore/Public/MVVM/ViewModels/OutlinerColumns/OutlinerColumnTypes.h b/Engine/Source/Editor/SequencerCore/Public/MVVM/ViewModels/OutlinerColumns/OutlinerColumnTypes.h
--- a/Engine/Source/Editor/SequencerCore/Public/MVVM/ViewModels/OutlinerColumns/OutlinerColumnTypes.h
+++ b/Engine/Source/Editor/SequencerCore/Public/MVVM/ViewModels/OutlinerColumns/OutlinerColumnTypes.h
@@ -65,7 +65,7 @@ struct FOutlinerColumnLayout
 
 struct FCommonOutlinerNames
 {
-	static SEQUENCERCORE_API FName Indicator;
+	static SEQUENCERCORE_API FName Decorator;
 	static SEQUENCERCORE_API FName Pin;
 	static SEQUENCERCORE_API FName Lock;
 	static SEQUENCERCORE_API FName Mute;
@@ -76,6 +76,8 @@ struct FCommonOutlinerNames
 	static SEQUENCERCORE_API FName Nav;
 	static SEQUENCERCORE_API FName KeyFrame;
 	static SEQUENCERCORE_API FName ColorPicker;
+	static SEQUENCERCORE_API FName TimeWarp;
+	static SEQUENCERCORE_API FName Condition;
 };
 
 
diff --git a/Engine/Source/Editor/SequencerCore/Public/MVVM/ViewModels/OutlinerDecorators/IOutlinerDecorator.h b/Engine/Source/Editor/SequencerCore/Public/MVVM/ViewModels/OutlinerDecorators/IOutlinerDecorator.h
new file mode 100644
--- /dev/null
+++ b/Engine/Source/Editor/SequencerCore/Public/MVVM/ViewModels/OutlinerDecorators/IOutlinerDecorator.h
@@ -0,0 +1,42 @@
+// Copyright Epic Games, Inc. All Rights Reserved.
+
+#pragma once
+
+#include "Templates/SharedPointer.h"
+
+class SWidget;
+class UMovieSceneSequence;
+
+namespace UE::Sequencer
+{
+
+struct FCreateOutlinerColumnParams;
+class IOutlinerColumn;
+class ISequencerTreeViewRow;
+
+/**
+* Interface for sequencer decorator outliner item.
+*/
+class IOutlinerDecorator : public TSharedFromThis<IOutlinerDecorator>
+{
+
+public:
+
+	/** Returns the name of the Decorator outliner item. */
+	virtual FName GetDecoratorName() const = 0;
+
+	/* Gets whether or not this Decorator outliner item is supported by a given Sequencer. */
+	virtual bool SupportsSequence(UMovieSceneSequence* InSequence) const { return true; }
+
+	virtual bool IsItemCompatibleWithDecorator(const FCreateOutlinerColumnParams& InParams) const = 0;
+	virtual TSharedPtr<SWidget> CreateDecoratorWidget(const FCreateOutlinerColumnParams& InParams, const TSharedRef<ISequencerTreeViewRow>& TreeViewRow, const TSharedRef<IOutlinerColumn>& OutlinerColumn, const int32 NumCompatibleDecorators) = 0;
+
+public:
+
+	/** Virtual destructor. */
+	virtual ~IOutlinerDecorator() { }
+
+};
+
+} // namespace UE::Sequencer
+
diff --git a/Engine/Source/Editor/SequencerCore/Public/MVVM/ViewModels/OutlinerDecorators/OutlinerDecoratorBase.h b/Engine/Source/Editor/SequencerCore/Public/MVVM/ViewModels/OutlinerDecorators/OutlinerDecoratorBase.h
new file mode 100644
--- /dev/null
+++ b/Engine/Source/Editor/SequencerCore/Public/MVVM/ViewModels/OutlinerDecorators/OutlinerDecoratorBase.h
@@ -0,0 +1,25 @@
+// Copyright Epic Games, Inc. All Rights Reserved.
+
+#pragma once
+
+#include "MVVM/ViewModels/OutlinerDecorators/IOutlinerDecorator.h"
+
+namespace UE::Sequencer
+{
+
+class FOutlinerDecoratorBase
+	: public IOutlinerDecorator
+{
+public:
+
+	FOutlinerDecoratorBase()
+	{
+	}
+	
+public:
+
+	bool IsItemCompatibleWithDecorator(const FCreateOutlinerColumnParams& InParams) const override { return false; }
+	TSharedPtr<SWidget> CreateDecoratorWidget(const FCreateOutlinerColumnParams& InParams, const TSharedRef<ISequencerTreeViewRow>& TreeViewRow, const TSharedRef<IOutlinerColumn>& OutlinerColumn, const int32 NumCompatibleDecorators) override { return nullptr; }
+};
+
+} // namespace UE::Sequencer
\ No newline at end of file
