Sequencer: Initialize the transform box's delta time to 1 frame

#rnx

diff --git a/Engine/Source/Editor/Sequencer/Private/SSequencerTransformBox.cpp b/Engine/Source/Editor/Sequencer/Private/SSequencerTransformBox.cpp
--- a/Engine/Source/Editor/Sequencer/Private/SSequencerTransformBox.cpp
+++ b/Engine/Source/Editor/Sequencer/Private/SSequencerTransformBox.cpp
@@ -1,6 +1,8 @@
 // Copyright Epic Games, Inc. All Rights Reserved.
 
 #include "SSequencerTransformBox.h"
+
+#include "Misc/FrameRate.h"
 #include "Sequencer.h"
 #include "SequencerSettings.h"
 #include "SequencerCommonHelpers.h"
@@ -16,7 +18,6 @@ void SSequencerTransformBox::Construct(const FArguments& InArgs, const TSharedRe
 	Settings = &InSettings;
 	NumericTypeInterface = InNumericTypeInterface;
 	
-	DeltaTime = FFrameNumber(0);
 	ScaleFactor = 1.f;
 
 	const FDockTabStyle* GenericTabStyle = &FCoreStyle::Get().GetWidgetStyle<FDockTabStyle>("Docking.Tab");
@@ -59,7 +60,7 @@ void SSequencerTransformBox::Construct(const FArguments& InArgs, const TSharedRe
 						.ToolTipText(LOCTEXT("TransformDelta_Tooltip", "The amount to offset the selected keys/sections by"))
 						.OnValueCommitted(this, &SSequencerTransformBox::OnDeltaCommitted)
 						.OnValueChanged(this, &SSequencerTransformBox::OnDeltaChanged)
-						.Value_Lambda([this](){ return DeltaTime.Value; })
+						.Value_Lambda([this](){ return DeltaTime.IsSet() ? DeltaTime.GetValue().Value : 0.0; })
 					]
 
 				+ SHorizontalBox::Slot()
@@ -128,6 +129,11 @@ void SSequencerTransformBox::ToggleVisibility()
 	}
 	else
 	{
+		if (!DeltaTime.IsSet())
+		{
+			InitializeDeltaTime();
+		}
+
 		Border->SetVisibility(EVisibility::Visible);
 		LastFocusedWidget = SlateApplication.GetUserFocusedWidget(0);
 		SlateApplication.SetAllUserFocus(OffsetEntryBox, EFocusCause::Navigation);
@@ -156,11 +162,16 @@ void SSequencerTransformBox::OnScaleChanged(float Value)
 
 FReply SSequencerTransformBox::OnPlusButtonClicked()
 {
-	if (DeltaTime != 0)
+	if (!DeltaTime.IsSet())
+	{
+		InitializeDeltaTime();
+	}
+
+	if (DeltaTime.GetValue() != 0)
 	{
 		TSharedPtr<FSequencer> Sequencer = SequencerPtr.Pin();
 
-		Sequencer->TransformSelectedKeysAndSections(DeltaTime, 1.f);
+		Sequencer->TransformSelectedKeysAndSections(DeltaTime.GetValue(), 1.f);
 	}
 
 	return FReply::Handled();
@@ -168,11 +179,16 @@ FReply SSequencerTransformBox::OnPlusButtonClicked()
 
 FReply SSequencerTransformBox::OnMinusButtonClicked()
 {
-	if (DeltaTime != 0)
+	if (!DeltaTime.IsSet())
+	{
+		InitializeDeltaTime();
+	}
+
+	if (DeltaTime.GetValue() != 0)
 	{
 		TSharedPtr<FSequencer> Sequencer = SequencerPtr.Pin();
 	
-		Sequencer->TransformSelectedKeysAndSections(-DeltaTime, 1.f);
+		Sequencer->TransformSelectedKeysAndSections(-DeltaTime.GetValue(), 1.f);
 	}
 
 	return FReply::Handled();
@@ -208,4 +224,14 @@ FReply SSequencerTransformBox::OnCloseButtonClicked()
 	return FReply::Handled();
 }
 
+void SSequencerTransformBox::InitializeDeltaTime()
+{
+	TSharedPtr<FSequencer> Sequencer = SequencerPtr.Pin();
+
+	// Default to 1 frame
+	const FFrameRate DisplayRate = Sequencer->GetFocusedDisplayRate();
+	const FFrameRate TickResolution = Sequencer->GetFocusedTickResolution();
+	DeltaTime = FFrameRate::TransformTime(FFrameTime(1), DisplayRate, TickResolution).FrameNumber;
+}
+
 #undef LOCTEXT_NAMESPACE
diff --git a/Engine/Source/Editor/Sequencer/Private/SSequencerTransformBox.h b/Engine/Source/Editor/Sequencer/Private/SSequencerTransformBox.h
--- a/Engine/Source/Editor/Sequencer/Private/SSequencerTransformBox.h
+++ b/Engine/Source/Editor/Sequencer/Private/SSequencerTransformBox.h
@@ -42,6 +42,8 @@ private:
 	/** Callback for when the close button is clicked. */
 	FReply OnCloseButtonClicked();
 
+	void InitializeDeltaTime();
+
 private:
 
 	/** The border widget. */
@@ -66,7 +68,7 @@ private:
 	USequencerSettings* Settings;
 
 	/** Cached delta time. */
-	FFrameNumber DeltaTime;
+	TOptional<FFrameNumber> DeltaTime;
 
 	/** Cached scale factor. */
 	float ScaleFactor;
