[Sequencer] FilterBar: Use OnPopulateMenu for view options menu

diff --git a/Engine/Source/Editor/Sequencer/Private/Filters/Menus/SequencerMenuContext.h b/Engine/Source/Editor/Sequencer/Private/Filters/Menus/SequencerMenuContext.h
--- a/Engine/Source/Editor/Sequencer/Private/Filters/Menus/SequencerMenuContext.h
+++ b/Engine/Source/Editor/Sequencer/Private/Filters/Menus/SequencerMenuContext.h
@@ -2,6 +2,7 @@
 
 #pragma once
 
+#include "SequencerFilterBarContext.h"
 #include "Templates/SharedPointer.h"
 #include "SequencerMenuContext.generated.h"
 
@@ -23,6 +24,8 @@ public:
 		return WeakSequencer.Pin();
 	}
 
+	FOnPopulateFilterBarMenu OnPopulateFilterBarMenu;
+
 protected:
 	TWeakPtr<FSequencer> WeakSequencer;
 };
diff --git a/Engine/Source/Editor/Sequencer/Private/Filters/Menus/SequencerViewOptionsMenu.cpp b/Engine/Source/Editor/Sequencer/Private/Filters/Menus/SequencerViewOptionsMenu.cpp
--- a/Engine/Source/Editor/Sequencer/Private/Filters/Menus/SequencerViewOptionsMenu.cpp
+++ b/Engine/Source/Editor/Sequencer/Private/Filters/Menus/SequencerViewOptionsMenu.cpp
@@ -5,6 +5,7 @@
 #include "Filters/SequencerTrackFilterCommands.h"
 #include "Sequencer.h"
 #include "SequencerCommands.h"
+#include "SequencerFilterBarContext.h"
 #include "SSequencer.h"
 #include "ToolMenu.h"
 #include "ToolMenus.h"
@@ -22,11 +23,18 @@ TSharedRef<SWidget> FSequencerViewOptionsMenu::CreateMenu(const TWeakPtr<FSequen
 	if (!UToolMenus::Get()->IsMenuRegistered(FilterMenuName))
 	{
 		UToolMenu* const Menu = UToolMenus::Get()->RegisterMenu(FilterMenuName);
-		Menu->AddDynamicSection(NAME_None, FNewToolMenuDelegate::CreateRaw(this, &FSequencerViewOptionsMenu::PopulateMenu));
+		Menu->AddDynamicSection(NAME_None, FNewToolMenuDelegate::CreateLambda([this](UToolMenu* const InMenu)
+			{
+				if (USequencerMenuContext* const Context = InMenu->FindContext<USequencerMenuContext>())
+				{
+					Context->OnPopulateFilterBarMenu.ExecuteIfBound(InMenu);
+				}
+			}));
 	}
 
 	USequencerMenuContext* const ContextObject = NewObject<USequencerMenuContext>();
 	ContextObject->Init(InSequencerWeak);
+	ContextObject->OnPopulateFilterBarMenu = FOnPopulateFilterBarMenu::CreateSP(this, &FSequencerViewOptionsMenu::PopulateMenu);
 
 	const FToolMenuContext MenuContext(InSequencerWeak.Pin()->GetFilterInterface()->GetCommandList(), nullptr, ContextObject);
 	return UToolMenus::Get()->GenerateWidget(FilterMenuName, MenuContext);
