[Sequencer] Fix crash when having the Custom Text Filter help dialog up then closing Sequencer and reopening it

diff --git a/Engine/Source/Editor/Sequencer/Private/Filters/Widgets/SSequencerCustomTextFilterDialog.cpp b/Engine/Source/Editor/Sequencer/Private/Filters/Widgets/SSequencerCustomTextFilterDialog.cpp
--- a/Engine/Source/Editor/Sequencer/Private/Filters/Widgets/SSequencerCustomTextFilterDialog.cpp
+++ b/Engine/Source/Editor/Sequencer/Private/Filters/Widgets/SSequencerCustomTextFilterDialog.cpp
@@ -121,6 +121,7 @@ void SSequencerCustomTextFilterDialog::CloseWindow()
 	if (DialogInstance.IsValid())
 	{
 		DialogInstance->RequestDestroyWindow();
+		DialogInstance.Reset();
 	}
 }
 
diff --git a/Engine/Source/Editor/Sequencer/Private/Filters/Widgets/SSequencerFilterBar.cpp b/Engine/Source/Editor/Sequencer/Private/Filters/Widgets/SSequencerFilterBar.cpp
--- a/Engine/Source/Editor/Sequencer/Private/Filters/Widgets/SSequencerFilterBar.cpp
+++ b/Engine/Source/Editor/Sequencer/Private/Filters/Widgets/SSequencerFilterBar.cpp
@@ -21,7 +21,23 @@ using namespace UE::Sequencer;
 
 SSequencerFilterBar::~SSequencerFilterBar()
 {
+	if (const TSharedPtr<FSequencerFilterBar> FilterBar = WeakFilterBar.Pin())
+	{
+		FilterBar->GetOnFiltersChanged().RemoveAll(this);
+	}
+
 	ContextMenu.Reset();
+
+	if (SSequencerCustomTextFilterDialog::IsOpen())
+	{
+		SSequencerCustomTextFilterDialog::CloseWindow();
+	}
+
+	if (TextExpressionHelpDialog.IsValid())
+	{
+		TextExpressionHelpDialog->RequestDestroyWindow();
+		TextExpressionHelpDialog.Reset();
+	}
 }
 
 void SSequencerFilterBar::Construct(const FArguments& InArgs, const TSharedRef<FSequencerFilterBar>& InFilterBar)
@@ -417,19 +433,19 @@ void SSequencerFilterBar::OnOpenTextExpressionHelp()
 		return;
 	}
 
-	if (TextExpressionHelpWindow.IsValid())
+	if (TextExpressionHelpDialog.IsValid())
 	{
-		TextExpressionHelpWindow->BringToFront();
+		TextExpressionHelpDialog->BringToFront();
 	}
 	else
 	{
-		const TSharedRef<SFilterExpressionHelpDialog> HelpDialog = SNew(SFilterExpressionHelpDialog)
+		TextExpressionHelpDialog = SNew(SFilterExpressionHelpDialog)
 			.DialogTitle(LOCTEXT("SequencerCustomTextFilterHelp", "Sequencer Custom Text Filter Help"))
 			.TextFilterExpressionContexts(FilterBar->GetTextFilter()->GetTextFilterExpressionContexts());
 
-		HelpDialog->GetOnWindowClosedEvent().AddLambda([this](const TSharedRef<SWindow>& InWindow)
+		TextExpressionHelpDialog->GetOnWindowClosedEvent().AddLambda([this](const TSharedRef<SWindow>& InWindow)
 			{
-				TextExpressionHelpWindow.Reset();
+				TextExpressionHelpDialog.Reset();
 			});
 
 		TSharedPtr<SWindow> ParentWindow;
@@ -441,11 +457,11 @@ void SSequencerFilterBar::OnOpenTextExpressionHelp()
 
 		if (ParentWindow.IsValid())
 		{
-			TextExpressionHelpWindow = FSlateApplication::Get().AddWindowAsNativeChild(HelpDialog, ParentWindow.ToSharedRef());
+			FSlateApplication::Get().AddWindowAsNativeChild(TextExpressionHelpDialog.ToSharedRef(), ParentWindow.ToSharedRef());
 		}
 		else
 		{
-			TextExpressionHelpWindow = FSlateApplication::Get().AddWindow(HelpDialog);
+			FSlateApplication::Get().AddWindow(TextExpressionHelpDialog.ToSharedRef());
 		}
 	}
 }
diff --git a/Engine/Source/Editor/Sequencer/Private/Filters/Widgets/SSequencerFilterBar.h b/Engine/Source/Editor/Sequencer/Private/Filters/Widgets/SSequencerFilterBar.h
--- a/Engine/Source/Editor/Sequencer/Private/Filters/Widgets/SSequencerFilterBar.h
+++ b/Engine/Source/Editor/Sequencer/Private/Filters/Widgets/SSequencerFilterBar.h
@@ -10,6 +10,7 @@ class FSequencerFilterBarContextMenu;
 class FSequencerFilterBar;
 class FSequencerTrackFilter_CustomText;
 class SFilterBarClippingHorizontalBox;
+class SFilterExpressionHelpDialog;
 class SSequencerFilter;
 class UMovieSceneNodeGroup;
 enum class ESequencerFilterChange : uint8;
@@ -125,7 +126,7 @@ protected:
 
 	TArray<TSharedRef<SSequencerFilter>> FilterWidgets;
 
-	TSharedPtr<SWindow> TextExpressionHelpWindow;
+	TSharedPtr<SFilterExpressionHelpDialog> TextExpressionHelpDialog;
 
 	TSharedPtr<FSequencerFilterBarContextMenu> ContextMenu;
 };
diff --git a/Engine/Source/Editor/Sequencer/Private/SSequencer.cpp b/Engine/Source/Editor/Sequencer/Private/SSequencer.cpp
--- a/Engine/Source/Editor/Sequencer/Private/SSequencer.cpp
+++ b/Engine/Source/Editor/Sequencer/Private/SSequencer.cpp
@@ -2796,8 +2796,10 @@ SSequencer::~SSequencer()
 		Window->DestroyWindowImmediately();
 		NodeGroupManager.Reset();
 	}
-}
 
+	// Ensure the FilterBarWidget destructor is called
+	FilterBarSplitterContainer->SetContent(SNullWidget::NullWidget);
+}
 
 void SSequencer::RegisterActiveTimerForPlayback()
 {
@@ -4269,12 +4271,6 @@ void SSequencer::ToggleFilterBarVisibility()
 		SequencerSettings->SetFilterBarVisible(bNewFilterBarVisible);
 	}
 
-	if (FilterBarWidget.IsValid())
-	{
-		const EVisibility NewVisibility = bNewFilterBarVisible ? EVisibility::Visible : EVisibility::Collapsed;
-		FilterBarWidget->SetVisibility(NewVisibility);
-	}
-
 	RebuildFilterBarContent();
 }
 
