Sequencer: Minor changes to display of track, track row, and section properties.
* Reorder track properties to be above track row properties
* Remove the track properties customization as it's no longer necessary, and was hiding options we want to show
* Hide the edit condition toggle on the eval nearest section option, as it shouldn't be editable by user.
* Consolidate the 'Track' and 'General' sections in Track properties to be inside 'Track' for less confusion since we're showing properties for Track, Track Row, and Section all in one place, and 'General' isn't specific enough.

diff --git a/Engine/Source/Editor/DetailCustomizations/Private/DetailCustomizations.cpp b/Engine/Source/Editor/DetailCustomizations/Private/DetailCustomizations.cpp
--- a/Engine/Source/Editor/DetailCustomizations/Private/DetailCustomizations.cpp
+++ b/Engine/Source/Editor/DetailCustomizations/Private/DetailCustomizations.cpp
@@ -362,7 +362,6 @@ void FDetailCustomizationsModule::RegisterPropertyTypeCustomizations()
 	RegisterCustomPropertyTypeLayout("CameraFocusSettings", FOnGetPropertyTypeCustomizationInstance::CreateStatic(&FCameraFocusSettingsCustomization::MakeInstance));
 	RegisterCustomPropertyTypeLayout("MovieSceneSequenceLoopCount", FOnGetPropertyTypeCustomizationInstance::CreateStatic(&FMovieSceneSequenceLoopCountCustomization::MakeInstance));
 	RegisterCustomPropertyTypeLayout("MovieSceneBindingOverrideData", FOnGetPropertyTypeCustomizationInstance::CreateStatic(&FMovieSceneBindingOverrideDataCustomization::MakeInstance));
-	RegisterCustomPropertyTypeLayout("MovieSceneTrackEvalOptions", FOnGetPropertyTypeCustomizationInstance::CreateStatic(&FMovieSceneTrackEvalOptionsCustomization::MakeInstance));
 	RegisterCustomPropertyTypeLayout("MovieSceneSectionEvalOptions", FOnGetPropertyTypeCustomizationInstance::CreateStatic(&FMovieSceneSectionEvalOptionsCustomization::MakeInstance));
 	RegisterCustomPropertyTypeLayout("MovieSceneEventParameters", FOnGetPropertyTypeCustomizationInstance::CreateStatic(&FMovieSceneEventParametersCustomization::MakeInstance));
 	RegisterCustomPropertyTypeLayout("FrameRate", FOnGetPropertyTypeCustomizationInstance::CreateStatic(&FFrameRateCustomization::MakeInstance));
diff --git a/Engine/Source/Editor/DetailCustomizations/Private/MovieSceneEvalOptionsCustomization.cpp b/Engine/Source/Editor/DetailCustomizations/Private/MovieSceneEvalOptionsCustomization.cpp
--- a/Engine/Source/Editor/DetailCustomizations/Private/MovieSceneEvalOptionsCustomization.cpp
+++ b/Engine/Source/Editor/DetailCustomizations/Private/MovieSceneEvalOptionsCustomization.cpp
@@ -9,33 +9,6 @@
 #include "MovieSceneTrack.h"
 #include "PropertyHandle.h"
 
-TSharedRef<IPropertyTypeCustomization> FMovieSceneTrackEvalOptionsCustomization::MakeInstance()
-{
-	return MakeShared<FMovieSceneTrackEvalOptionsCustomization>();
-}
-
-void FMovieSceneTrackEvalOptionsCustomization::CustomizeHeader(TSharedRef<IPropertyHandle> PropertyHandle, FDetailWidgetRow& HeaderRow, IPropertyTypeCustomizationUtils& CustomizationUtils)
-{
-}
-
-void FMovieSceneTrackEvalOptionsCustomization::CustomizeChildren(TSharedRef<IPropertyHandle> PropertyHandle, IDetailChildrenBuilder& ChildBuilder, IPropertyTypeCustomizationUtils& CustomizationUtils)
-{
-	TArray<void*> RawData;
-	PropertyHandle->AccessRawData(RawData);
-
-	bool bCanEvaluateNearestSection = !RawData.ContainsByPredicate(
-		[](void* Ptr){
-			return !static_cast<FMovieSceneTrackEvalOptions*>(Ptr)->bCanEvaluateNearestSection;
-		}
-	);
-
-	TSharedPtr<IPropertyHandle> bEvalNearestSectionHandle = PropertyHandle->GetChildHandle(GET_MEMBER_NAME_CHECKED(FMovieSceneTrackEvalOptions, bEvalNearestSection));
-	if (bCanEvaluateNearestSection && bEvalNearestSectionHandle.IsValid())
-	{
-		ChildBuilder.AddProperty(bEvalNearestSectionHandle.ToSharedRef());
-	}
-}
-
 TSharedRef<IPropertyTypeCustomization> FMovieSceneSectionEvalOptionsCustomization::MakeInstance()
 {
 	return MakeShared<FMovieSceneSectionEvalOptionsCustomization>();
diff --git a/Engine/Source/Editor/DetailCustomizations/Private/MovieSceneEvalOptionsCustomization.h b/Engine/Source/Editor/DetailCustomizations/Private/MovieSceneEvalOptionsCustomization.h
--- a/Engine/Source/Editor/DetailCustomizations/Private/MovieSceneEvalOptionsCustomization.h
+++ b/Engine/Source/Editor/DetailCustomizations/Private/MovieSceneEvalOptionsCustomization.h
@@ -9,15 +9,6 @@ class FDetailWidgetRow;
 class IDetailChildrenBuilder;
 class IPropertyHandle;
 
-class FMovieSceneTrackEvalOptionsCustomization : public IPropertyTypeCustomization
-{
-public:
-	static TSharedRef<IPropertyTypeCustomization> MakeInstance();
-
-	virtual void CustomizeHeader( TSharedRef<IPropertyHandle> PropertyHandle, FDetailWidgetRow& HeaderRow, IPropertyTypeCustomizationUtils& CustomizationUtils ) override;
-	virtual void CustomizeChildren( TSharedRef<IPropertyHandle> PropertyHandle, IDetailChildrenBuilder& ChildBuilder, IPropertyTypeCustomizationUtils& CustomizationUtils ) override;
-};
-
 class FMovieSceneSectionEvalOptionsCustomization : public IPropertyTypeCustomization
 {
 public:
diff --git a/Engine/Source/Editor/Sequencer/Private/MVVM/Selection.cpp b/Engine/Source/Editor/Sequencer/Private/MVVM/Selection.cpp
--- a/Engine/Source/Editor/Sequencer/Private/MVVM/Selection.cpp
+++ b/Engine/Source/Editor/Sequencer/Private/MVVM/Selection.cpp
@@ -268,6 +268,18 @@ TSet<TPair<UMovieSceneTrack*, int32>> FSequencerSelection::GetSelectedTrackRows(
 		}
 	}
 
+	TSet<UMovieSceneSection*> SelectedSections = GetSelectedSections();
+	for (UMovieSceneSection* Section : SelectedSections)
+	{
+		if (UMovieSceneTrack* Track = Section->GetTypedOuter<UMovieSceneTrack>())
+		{
+			if (Track->SupportsMultipleRows())
+			{
+				SelectedTrackRows.Add(TPair<UMovieSceneTrack*, int32>(Track, Section->GetRowIndex()));
+			}
+		}
+	}
+
 	return SelectedTrackRows;
 }
 
diff --git a/Engine/Source/Editor/Sequencer/Private/Widgets/Sidebar/SequencerSelectionDrawer.cpp b/Engine/Source/Editor/Sequencer/Private/Widgets/Sidebar/SequencerSelectionDrawer.cpp
--- a/Engine/Source/Editor/Sequencer/Private/Widgets/Sidebar/SequencerSelectionDrawer.cpp
+++ b/Engine/Source/Editor/Sequencer/Private/Widgets/Sidebar/SequencerSelectionDrawer.cpp
@@ -293,13 +293,27 @@ void FSequencerSelectionDrawer::BuildKeySelectionDetails(const TSharedRef<FSeque
 
 void FSequencerSelectionDrawer::BuildTrackAreaDetails(const TSharedRef<FSequencerSelection>& InSelection, FMenuBuilder& MenuBuilder)
 {
-	TArray<TWeakObjectPtr<>> AllSectionObjects;
+	TArray<TWeakObjectPtr<>> AllSectionObjects; 
+	TArray<TWeakObjectPtr<>> WeakTrackObjects;
+	TArray<TPair<TWeakObjectPtr<UMovieSceneTrack>, int32>> SelectedTrackRows;
 
 	for (const FViewModelPtr TrackAreaItem : InSelection->TrackArea)
 	{
 		if (const TViewModelPtr<FSectionModel> SectionModel = TrackAreaItem.ImplicitCast())
 		{
-			AllSectionObjects.Add(SectionModel->GetSection());
+			if (UMovieSceneSection* Section = SectionModel->GetSection())
+			{
+				AllSectionObjects.Add(SectionModel->GetSection());
+
+				if (UMovieSceneTrack* Track = Section->GetTypedOuter<UMovieSceneTrack>())
+				{
+					WeakTrackObjects.Add(Track);
+					if (Track->SupportsMultipleRows())
+					{
+						SelectedTrackRows.Add(TPair<TWeakObjectPtr<UMovieSceneTrack>, int32>(Track, Section->GetRowIndex()));
+					}
+				}
+			}
 		}
 	}
 
@@ -307,6 +321,21 @@ void FSequencerSelectionDrawer::BuildTrackAreaDetails(const TSharedRef<FSequence
 	{
 		SequencerHelpers::BuildEditSectionMenu(WeakSequencer, AllSectionObjects, MenuBuilder, false);
 	}
+
+	if (!SelectedTrackRows.IsEmpty() && !Algo::AnyOf(SelectedTrackRows, [](const TPair<TWeakObjectPtr<UMovieSceneTrack>, int32> TrackRow) {
+		return TrackRow.Key.IsValid() && !TrackRow.Key->SupportsMultipleRows();
+		}))
+	{
+		MenuBuilder.BeginSection(TEXT("TrackRowMetadata"), LOCTEXT("TrackRowMetadata", "Track Row Metadata"));
+		{
+			// Empty here, will be implemented by extension.
+		}
+		MenuBuilder.EndSection();
+	}
+	if (!WeakTrackObjects.IsEmpty())
+	{
+		SequencerHelpers::BuildEditTrackMenu(WeakSequencer, WeakTrackObjects, MenuBuilder, false);
+	}
 }
 
 void FSequencerSelectionDrawer::BuildOutlinerDetails(const TSharedRef<FSequencerSelection>& InSelection, FMenuBuilder& MenuBuilder)
@@ -339,7 +368,16 @@ void FSequencerSelectionDrawer::BuildOutlinerDetails(const TSharedRef<FSequencer
 			if (IsValid(Track))
 			{
 				WeakTrackObjects.Add(Track);
-				SelectedTrackRows.Add(TPair<TWeakObjectPtr<UMovieSceneTrack>, int32>(Track, TrackExtension->GetRowIndex()));
+				// Only add a 'track row' as selected if either we have an actual 'track row' selected, or else we have a track selected and there's only a single
+				// track row, and the track allows multiple rows.
+				if (TViewModelPtr<FTrackRowModel> TrackRowModel = TrackExtension.ImplicitCast())
+				{
+					SelectedTrackRows.Add(TPair<TWeakObjectPtr<UMovieSceneTrack>, int32>(Track, TrackExtension->GetRowIndex()));
+				}
+				else if (Track->SupportsMultipleRows() && Track->GetMaxRowIndex() == 0)
+				{
+					SelectedTrackRows.Add(TPair<TWeakObjectPtr<UMovieSceneTrack>, int32>(Track, TrackExtension->GetRowIndex()));
+				}
 			}
 		}
 		else if (const TViewModelPtr<FObjectBindingModel> ObjectBindingModel = OutlinerItem.ImplicitCast())
@@ -378,6 +416,7 @@ void FSequencerSelectionDrawer::BuildOutlinerDetails(const TSharedRef<FSequencer
 	{
 		BuildExtensionDetails(ChannelGroups, MenuBuilder);
 	}
+	
 
 	if (!SelectedTrackRows.IsEmpty() && !Algo::AnyOf(SelectedTrackRows, [](const TPair<TWeakObjectPtr<UMovieSceneTrack>, int32> TrackRow) {
 		return TrackRow.Key.IsValid() && !TrackRow.Key->SupportsMultipleRows();
@@ -390,15 +429,21 @@ void FSequencerSelectionDrawer::BuildOutlinerDetails(const TSharedRef<FSequencer
 		MenuBuilder.EndSection();
 	}
 
-
 	if (!WeakTrackObjects.IsEmpty())
 	{
 		SequencerHelpers::BuildEditTrackMenu(Sequencer, WeakTrackObjects, MenuBuilder, false);
 	}
 
-	if (!WeakSectionObjects.IsEmpty())
+	// Only show section details for a selected track/trackrow if a single infinite section on the row
+	if (WeakSectionObjects.Num() == 1 && WeakSectionObjects[0].IsValid())
 	{
-		SequencerHelpers::BuildEditSectionMenu(WeakSequencer, WeakSectionObjects, MenuBuilder, false);
+		if (UMovieSceneSection* Section = Cast<UMovieSceneSection>(WeakSectionObjects[0].Get()))
+		{
+			if (!Section->GetRange().HasLowerBound() && !Section->GetRange().HasUpperBound())
+			{
+				SequencerHelpers::BuildEditSectionMenu(WeakSequencer, WeakSectionObjects, MenuBuilder, false);
+			}
+		}
 	}
 }
 
diff --git a/Engine/Source/Runtime/MovieScene/Public/MovieSceneTrack.h b/Engine/Source/Runtime/MovieScene/Public/MovieSceneTrack.h
--- a/Engine/Source/Runtime/MovieScene/Public/MovieSceneTrack.h
+++ b/Engine/Source/Runtime/MovieScene/Public/MovieSceneTrack.h
@@ -74,15 +74,15 @@ struct FMovieSceneTrackEvalOptions
 	uint32 bCanEvaluateNearestSection : 1;
 
 	/** When evaluating empty space on a track, will evaluate the last position of the previous section (if possible), or the first position of the next section, in that order of preference. */
-	UPROPERTY(EditAnywhere, Category="General", DisplayName="Evaluate Nearest Section", meta=(EditCondition=bCanEvaluateNearestSection))
+	UPROPERTY(EditAnywhere, Category="Track", DisplayName="Evaluate Nearest Section", meta=(EditCondition=bCanEvaluateNearestSection, HideEditConditionToggle))
 	uint32 bEvalNearestSection : 1;
 
 	/** Evaluate this track as part of its parent sub-section's pre-roll, if applicable */
-	UPROPERTY(EditAnywhere, Category="General")
+	UPROPERTY(EditAnywhere, Category="Track")
 	uint32 bEvaluateInPreroll : 1;
 
 	/** Evaluate this track as part of its parent sub-section's post-roll, if applicable */
-	UPROPERTY(EditAnywhere, Category="General")
+	UPROPERTY(EditAnywhere, Category="Track")
 	uint32 bEvaluateInPostroll : 1;
 
 	UPROPERTY()
@@ -101,7 +101,7 @@ struct FMovieSceneTrackDisplayOptions
 	{}
 
 	/** Show bounds as vertical frames */
-	UPROPERTY(EditAnywhere, Category = "General")
+	UPROPERTY(EditAnywhere, Category = "Track")
 	uint32 bShowVerticalFrames : 1;
 };
 
@@ -197,17 +197,17 @@ public:
 public:
 
 	/** General evaluation options for a given track */
-	UPROPERTY(EditAnywhere, Category = "General", meta = (ShowOnlyInnerProperties))
+	UPROPERTY(EditAnywhere, Category = "Track", meta = (ShowOnlyInnerProperties))
 	FMovieSceneTrackEvalOptions EvalOptions;
 
 #if WITH_EDITORONLY_DATA
 	/** General display options for a given track */
-	UPROPERTY(EditAnywhere, Category = "General", meta = (ShowOnlyInnerProperties))
+	UPROPERTY(EditAnywhere, Category = "Track", meta = (ShowOnlyInnerProperties))
 	FMovieSceneTrackDisplayOptions DisplayOptions;
 #endif
 
 	/** Optional dynamic condition for whether this track/any of the sections on this track evaluates at runtime. */
-	UPROPERTY(EditAnywhere, Category = "General")
+	UPROPERTY(EditAnywhere, Category = "Track")
 	FMovieSceneConditionContainer ConditionContainer;
 
 	/**
@@ -576,7 +576,7 @@ protected:
 	FGuid ObjectBindingID;
 
 	/** This track's tint color */
-	UPROPERTY(EditAnywhere, Category=General, DisplayName=Color)
+	UPROPERTY(EditAnywhere, Category="Track", DisplayName = Color)
 	FColor TrackTint;
 
 	/** This folder's desired sorting order */
