Sequencer: Follow up fix to 33696183 which sets newly created sections to infinite bounds.

This now only set to infinite bounds if the section supports it

#rnx

diff --git a/Engine/Source/Editor/Sequencer/Private/SequencerUtilities.cpp b/Engine/Source/Editor/Sequencer/Private/SequencerUtilities.cpp
--- a/Engine/Source/Editor/Sequencer/Private/SequencerUtilities.cpp
+++ b/Engine/Source/Editor/Sequencer/Private/SequencerUtilities.cpp
@@ -315,6 +315,9 @@ void FSequencerUtilities::PopulateMenu_CreateNewSection(FMenuBuilder& MenuBuilde
 			return;
 		}
 
+		FQualifiedFrameTime CurrentTime = Sequencer->GetLocalTime();
+		FFrameNumber PlaybackEnd = UE::MovieScene::DiscreteExclusiveUpper(Sequencer->GetFocusedMovieSceneSequence()->GetMovieScene()->GetPlaybackRange());
+
 		FScopedTransaction Transaction(LOCTEXT("AddSectionTransactionText", "Add Section"));
 		if (UMovieSceneSection* NewSection = Track->CreateNewSection())
 		{
@@ -339,10 +342,22 @@ void FSequencerUtilities::PopulateMenu_CreateNewSection(FMenuBuilder& MenuBuilde
 
 			Track->OnRowIndicesChanged(NewToOldRowIndices);
 
-			if (Sequencer->GetInfiniteKeyAreas())
+			if (Sequencer->GetInfiniteKeyAreas() && NewSection->GetSupportsInfiniteRange())
 			{
 				NewSection->SetRange(TRange<FFrameNumber>::All());
 			}
+			else
+			{
+				FFrameNumber NewSectionRangeEnd = PlaybackEnd;
+				if (PlaybackEnd <= CurrentTime.Time.FrameNumber)
+				{
+					const FAnimatedRange ViewRange = Sequencer->GetViewRange();
+					const FFrameRate TickResolution = Sequencer->GetFocusedTickResolution();
+					NewSectionRangeEnd = (ViewRange.GetUpperBoundValue() * TickResolution).FloorToFrame();
+				}
+
+				NewSection->SetRange(TRange<FFrameNumber>(CurrentTime.Time.FrameNumber, NewSectionRangeEnd));
+			}
 
 			NewSection->SetOverlapPriority(OverlapPriority);
 			NewSection->SetRowIndex(RowIndex);			
