[Sequencer]
 - Change text expression keys to camel case instead of upper case
 - Change some descriptions to be as short and concise as possible (mainly for the suggestion box)
 - Export text expression functions for outside module extensions

diff --git a/Engine/Source/Editor/Sequencer/Private/Filters/TextExpressions/SequencerTextFilterExpression_Condition.cpp b/Engine/Source/Editor/Sequencer/Private/Filters/TextExpressions/SequencerTextFilterExpression_Condition.cpp
--- a/Engine/Source/Editor/Sequencer/Private/Filters/TextExpressions/SequencerTextFilterExpression_Condition.cpp
+++ b/Engine/Source/Editor/Sequencer/Private/Filters/TextExpressions/SequencerTextFilterExpression_Condition.cpp
@@ -17,7 +17,7 @@ FSequencerTextFilterExpression_Condition::FSequencerTextFilterExpression_Conditi
 
 TSet<FName> FSequencerTextFilterExpression_Condition::GetKeys() const
 {
-	return { TEXT("CONDITION"), TEXT("HASCONDITION") };
+	return { TEXT("Condition"), TEXT("HasCondition") };
 }
 
 ESequencerTextFilterValueType FSequencerTextFilterExpression_Condition::GetValueType() const
@@ -27,7 +27,7 @@ ESequencerTextFilterValueType FSequencerTextFilterExpression_Condition::GetValue
 
 FText FSequencerTextFilterExpression_Condition::GetDescription() const
 {
-	return LOCTEXT("ExpressionDescription_Condition", "Filters by the presence of a condition");
+	return LOCTEXT("ExpressionDescription_Condition", "Filter by presence of a condition");
 }
 
 bool FSequencerTextFilterExpression_Condition::TestComplexExpression(const FName& InKey
diff --git a/Engine/Source/Editor/Sequencer/Private/Filters/TextExpressions/SequencerTextFilterExpression_ConditionClass.cpp b/Engine/Source/Editor/Sequencer/Private/Filters/TextExpressions/SequencerTextFilterExpression_ConditionClass.cpp
--- a/Engine/Source/Editor/Sequencer/Private/Filters/TextExpressions/SequencerTextFilterExpression_ConditionClass.cpp
+++ b/Engine/Source/Editor/Sequencer/Private/Filters/TextExpressions/SequencerTextFilterExpression_ConditionClass.cpp
@@ -17,7 +17,7 @@ FSequencerTextFilterExpression_ConditionClass::FSequencerTextFilterExpression_Co
 
 TSet<FName> FSequencerTextFilterExpression_ConditionClass::GetKeys() const
 {
-	return { TEXT("CONDITIONCLASS"), TEXT("CONDITIONTYPE") };
+	return { TEXT("ConditionClass"), TEXT("ConditionType") };
 }
 
 ESequencerTextFilterValueType FSequencerTextFilterExpression_ConditionClass::GetValueType() const
@@ -27,7 +27,7 @@ ESequencerTextFilterValueType FSequencerTextFilterExpression_ConditionClass::Get
 
 FText FSequencerTextFilterExpression_ConditionClass::GetDescription() const
 {
-	return LOCTEXT("ExpressionDescription_ConditionClass", "Filters based on the presence of a condition with the given type");
+	return LOCTEXT("ExpressionDescription_ConditionClass", "Filter by presence of a condition with the given type");
 }
 
 bool FSequencerTextFilterExpression_ConditionClass::TestComplexExpression(const FName& InKey
diff --git a/Engine/Source/Editor/Sequencer/Private/Filters/TextExpressions/SequencerTextFilterExpression_ConditionFunc.cpp b/Engine/Source/Editor/Sequencer/Private/Filters/TextExpressions/SequencerTextFilterExpression_ConditionFunc.cpp
--- a/Engine/Source/Editor/Sequencer/Private/Filters/TextExpressions/SequencerTextFilterExpression_ConditionFunc.cpp
+++ b/Engine/Source/Editor/Sequencer/Private/Filters/TextExpressions/SequencerTextFilterExpression_ConditionFunc.cpp
@@ -18,7 +18,7 @@ FSequencerTextFilterExpression_ConditionFunc::FSequencerTextFilterExpression_Con
 
 TSet<FName> FSequencerTextFilterExpression_ConditionFunc::GetKeys() const
 {
-	return { TEXT("CONDITIONFUNC"), TEXT("CONDITIONENDPOINT") };
+	return { TEXT("ConditionFunc"), TEXT("ConditionEndpoint") };
 }
 
 ESequencerTextFilterValueType FSequencerTextFilterExpression_ConditionFunc::GetValueType() const
@@ -28,7 +28,7 @@ ESequencerTextFilterValueType FSequencerTextFilterExpression_ConditionFunc::GetV
 
 FText FSequencerTextFilterExpression_ConditionFunc::GetDescription() const
 {
-	return LOCTEXT("ExpressionDescription_ConditionFunc", "Filters by the presence of a director blueprint condition with the given function/endpoint name");
+	return LOCTEXT("ExpressionDescription_ConditionFunc", "Filter by presence of director blueprint condition with the given function/endpoint name");
 }
 
 bool FSequencerTextFilterExpression_ConditionFunc::TestComplexExpression(const FName& InKey
diff --git a/Engine/Source/Editor/Sequencer/Private/Filters/TextExpressions/SequencerTextFilterExpression_ConditionPasses.cpp b/Engine/Source/Editor/Sequencer/Private/Filters/TextExpressions/SequencerTextFilterExpression_ConditionPasses.cpp
--- a/Engine/Source/Editor/Sequencer/Private/Filters/TextExpressions/SequencerTextFilterExpression_ConditionPasses.cpp
+++ b/Engine/Source/Editor/Sequencer/Private/Filters/TextExpressions/SequencerTextFilterExpression_ConditionPasses.cpp
@@ -17,7 +17,7 @@ FSequencerTextFilterExpression_ConditionPasses::FSequencerTextFilterExpression_C
 
 TSet<FName> FSequencerTextFilterExpression_ConditionPasses::GetKeys() const
 {
-	return { TEXT("CONDITIONPASSES"), TEXT("CONDITIONEVALUATES") };
+	return { TEXT("ConditionPasses"), TEXT("ConditionEvaluates") };
 }
 
 ESequencerTextFilterValueType FSequencerTextFilterExpression_ConditionPasses::GetValueType() const
diff --git a/Engine/Source/Editor/Sequencer/Private/Filters/TextExpressions/SequencerTextFilterExpression_Group.cpp b/Engine/Source/Editor/Sequencer/Private/Filters/TextExpressions/SequencerTextFilterExpression_Group.cpp
--- a/Engine/Source/Editor/Sequencer/Private/Filters/TextExpressions/SequencerTextFilterExpression_Group.cpp
+++ b/Engine/Source/Editor/Sequencer/Private/Filters/TextExpressions/SequencerTextFilterExpression_Group.cpp
@@ -16,7 +16,7 @@ FSequencerTextFilterExpression_Group::FSequencerTextFilterExpression_Group(ISequ
 
 TSet<FName> FSequencerTextFilterExpression_Group::GetKeys() const
 {
-	return { TEXT("GROUP") };
+	return { TEXT("Group") };
 }
 
 ESequencerTextFilterValueType FSequencerTextFilterExpression_Group::GetValueType() const
diff --git a/Engine/Source/Editor/Sequencer/Private/Filters/TextExpressions/SequencerTextFilterExpression_Keyed.cpp b/Engine/Source/Editor/Sequencer/Private/Filters/TextExpressions/SequencerTextFilterExpression_Keyed.cpp
--- a/Engine/Source/Editor/Sequencer/Private/Filters/TextExpressions/SequencerTextFilterExpression_Keyed.cpp
+++ b/Engine/Source/Editor/Sequencer/Private/Filters/TextExpressions/SequencerTextFilterExpression_Keyed.cpp
@@ -14,7 +14,7 @@ FSequencerTextFilterExpression_Keyed::FSequencerTextFilterExpression_Keyed(ISequ
 
 TSet<FName> FSequencerTextFilterExpression_Keyed::GetKeys() const
 {
-	return { TEXT("KEYED"), TEXT("KEYS"), TEXT("ANIMATED") };
+	return { TEXT("Keyed"), TEXT("Keys"), TEXT("Animated") };
 }
 
 ESequencerTextFilterValueType FSequencerTextFilterExpression_Keyed::GetValueType() const
@@ -24,7 +24,7 @@ ESequencerTextFilterValueType FSequencerTextFilterExpression_Keyed::GetValueType
 
 FText FSequencerTextFilterExpression_Keyed::GetDescription() const
 {
-	return LOCTEXT("ExpressionDescription_Keys", "Filter by keys");
+	return LOCTEXT("ExpressionDescription_Keys", "Filter by presence of keys");
 }
 
 bool FSequencerTextFilterExpression_Keyed::TestComplexExpression(const FName& InKey
diff --git a/Engine/Source/Editor/Sequencer/Private/Filters/TextExpressions/SequencerTextFilterExpression_Level.cpp b/Engine/Source/Editor/Sequencer/Private/Filters/TextExpressions/SequencerTextFilterExpression_Level.cpp
--- a/Engine/Source/Editor/Sequencer/Private/Filters/TextExpressions/SequencerTextFilterExpression_Level.cpp
+++ b/Engine/Source/Editor/Sequencer/Private/Filters/TextExpressions/SequencerTextFilterExpression_Level.cpp
@@ -16,7 +16,7 @@ FSequencerTextFilterExpression_Level::FSequencerTextFilterExpression_Level(ISequ
 
 TSet<FName> FSequencerTextFilterExpression_Level::GetKeys() const
 {
-	return { TEXT("LEVEL"), TEXT("MAP") };
+	return { TEXT("Level"), TEXT("Map") };
 }
 
 ESequencerTextFilterValueType FSequencerTextFilterExpression_Level::GetValueType() const
diff --git a/Engine/Source/Editor/Sequencer/Private/Filters/TextExpressions/SequencerTextFilterExpression_Locked.cpp b/Engine/Source/Editor/Sequencer/Private/Filters/TextExpressions/SequencerTextFilterExpression_Locked.cpp
--- a/Engine/Source/Editor/Sequencer/Private/Filters/TextExpressions/SequencerTextFilterExpression_Locked.cpp
+++ b/Engine/Source/Editor/Sequencer/Private/Filters/TextExpressions/SequencerTextFilterExpression_Locked.cpp
@@ -15,7 +15,7 @@ FSequencerTextFilterExpression_Locked::FSequencerTextFilterExpression_Locked(ISe
 
 TSet<FName> FSequencerTextFilterExpression_Locked::GetKeys() const
 {
-	return { TEXT("LOCK"), TEXT("LOCKED") };
+	return { TEXT("Lock"), TEXT("Locked") };
 }
 
 ESequencerTextFilterValueType FSequencerTextFilterExpression_Locked::GetValueType() const
diff --git a/Engine/Source/Editor/Sequencer/Private/Filters/TextExpressions/SequencerTextFilterExpression_Modified.cpp b/Engine/Source/Editor/Sequencer/Private/Filters/TextExpressions/SequencerTextFilterExpression_Modified.cpp
--- a/Engine/Source/Editor/Sequencer/Private/Filters/TextExpressions/SequencerTextFilterExpression_Modified.cpp
+++ b/Engine/Source/Editor/Sequencer/Private/Filters/TextExpressions/SequencerTextFilterExpression_Modified.cpp
@@ -14,7 +14,7 @@ FSequencerTextFilterExpression_Modified::FSequencerTextFilterExpression_Modified
 
 TSet<FName> FSequencerTextFilterExpression_Modified::GetKeys() const
 {
-	return { TEXT("MODIFIED"), TEXT("CHANGED"), TEXT("DIRTY") };
+	return { TEXT("Modified"), TEXT("Changed"), TEXT("Dirty") };
 }
 
 ESequencerTextFilterValueType FSequencerTextFilterExpression_Modified::GetValueType() const
diff --git a/Engine/Source/Editor/Sequencer/Private/Filters/TextExpressions/SequencerTextFilterExpression_Muted.cpp b/Engine/Source/Editor/Sequencer/Private/Filters/TextExpressions/SequencerTextFilterExpression_Muted.cpp
--- a/Engine/Source/Editor/Sequencer/Private/Filters/TextExpressions/SequencerTextFilterExpression_Muted.cpp
+++ b/Engine/Source/Editor/Sequencer/Private/Filters/TextExpressions/SequencerTextFilterExpression_Muted.cpp
@@ -15,7 +15,7 @@ FSequencerTextFilterExpression_Muted::FSequencerTextFilterExpression_Muted(ISequ
 
 TSet<FName> FSequencerTextFilterExpression_Muted::GetKeys() const
 {
-	return { TEXT("MUTE"), TEXT("MUTED") };
+	return { TEXT("Mute"), TEXT("Muted") };
 }
 
 ESequencerTextFilterValueType FSequencerTextFilterExpression_Muted::GetValueType() const
diff --git a/Engine/Source/Editor/Sequencer/Private/Filters/TextExpressions/SequencerTextFilterExpression_Name.cpp b/Engine/Source/Editor/Sequencer/Private/Filters/TextExpressions/SequencerTextFilterExpression_Name.cpp
--- a/Engine/Source/Editor/Sequencer/Private/Filters/TextExpressions/SequencerTextFilterExpression_Name.cpp
+++ b/Engine/Source/Editor/Sequencer/Private/Filters/TextExpressions/SequencerTextFilterExpression_Name.cpp
@@ -16,7 +16,7 @@ FSequencerTextFilterExpression_Name::FSequencerTextFilterExpression_Name(ISequen
 
 TSet<FName> FSequencerTextFilterExpression_Name::GetKeys() const
 {
-	return { TEXT("NAME") };
+	return { TEXT("Name") };
 }
 
 ESequencerTextFilterValueType FSequencerTextFilterExpression_Name::GetValueType() const
diff --git a/Engine/Source/Editor/Sequencer/Private/Filters/TextExpressions/SequencerTextFilterExpression_ObjectClass.cpp b/Engine/Source/Editor/Sequencer/Private/Filters/TextExpressions/SequencerTextFilterExpression_ObjectClass.cpp
--- a/Engine/Source/Editor/Sequencer/Private/Filters/TextExpressions/SequencerTextFilterExpression_ObjectClass.cpp
+++ b/Engine/Source/Editor/Sequencer/Private/Filters/TextExpressions/SequencerTextFilterExpression_ObjectClass.cpp
@@ -14,7 +14,7 @@ FSequencerTextFilterExpression_ObjectClass::FSequencerTextFilterExpression_Objec
 
 TSet<FName> FSequencerTextFilterExpression_ObjectClass::GetKeys() const
 {
-	return { TEXT("OBJECTCLASS"), TEXT("OBJECTTYPE") };
+	return { TEXT("ObjectClass"), TEXT("ObjectType") };
 }
 
 ESequencerTextFilterValueType FSequencerTextFilterExpression_ObjectClass::GetValueType() const
@@ -24,7 +24,7 @@ ESequencerTextFilterValueType FSequencerTextFilterExpression_ObjectClass::GetVal
 
 FText FSequencerTextFilterExpression_ObjectClass::GetDescription() const
 {
-	return LOCTEXT("ExpressionDescription_Class", "Filter by bound object class name");
+	return LOCTEXT("ExpressionDescription_ObjectClass", "Filter by bound object class name");
 }
 
 bool FSequencerTextFilterExpression_ObjectClass::TestComplexExpression(const FName& InKey
diff --git a/Engine/Source/Editor/Sequencer/Private/Filters/TextExpressions/SequencerTextFilterExpression_Selected.cpp b/Engine/Source/Editor/Sequencer/Private/Filters/TextExpressions/SequencerTextFilterExpression_Selected.cpp
--- a/Engine/Source/Editor/Sequencer/Private/Filters/TextExpressions/SequencerTextFilterExpression_Selected.cpp
+++ b/Engine/Source/Editor/Sequencer/Private/Filters/TextExpressions/SequencerTextFilterExpression_Selected.cpp
@@ -39,7 +39,7 @@ void FSequencerTextFilterExpression_Selected::UnbindSelectionChanged()
 
 TSet<FName> FSequencerTextFilterExpression_Selected::GetKeys() const
 {
-	return { TEXT("SELECTED"), TEXT("VIEWPORT") };
+	return { TEXT("Selected"), TEXT("Viewport") };
 }
 
 ESequencerTextFilterValueType FSequencerTextFilterExpression_Selected::GetValueType() const
diff --git a/Engine/Source/Editor/Sequencer/Private/Filters/TextExpressions/SequencerTextFilterExpression_Soloed.cpp b/Engine/Source/Editor/Sequencer/Private/Filters/TextExpressions/SequencerTextFilterExpression_Soloed.cpp
--- a/Engine/Source/Editor/Sequencer/Private/Filters/TextExpressions/SequencerTextFilterExpression_Soloed.cpp
+++ b/Engine/Source/Editor/Sequencer/Private/Filters/TextExpressions/SequencerTextFilterExpression_Soloed.cpp
@@ -15,7 +15,7 @@ FSequencerTextFilterExpression_Soloed::FSequencerTextFilterExpression_Soloed(ISe
 
 TSet<FName> FSequencerTextFilterExpression_Soloed::GetKeys() const
 {
-	return { TEXT("SOLO"), TEXT("SOLOED") };
+	return { TEXT("Solo"), TEXT("Soloed") };
 }
 
 ESequencerTextFilterValueType FSequencerTextFilterExpression_Soloed::GetValueType() const
diff --git a/Engine/Source/Editor/Sequencer/Private/Filters/TextExpressions/SequencerTextFilterExpression_Tag.cpp b/Engine/Source/Editor/Sequencer/Private/Filters/TextExpressions/SequencerTextFilterExpression_Tag.cpp
--- a/Engine/Source/Editor/Sequencer/Private/Filters/TextExpressions/SequencerTextFilterExpression_Tag.cpp
+++ b/Engine/Source/Editor/Sequencer/Private/Filters/TextExpressions/SequencerTextFilterExpression_Tag.cpp
@@ -16,7 +16,7 @@ FSequencerTextFilterExpression_Tag::FSequencerTextFilterExpression_Tag(ISequence
 
 TSet<FName> FSequencerTextFilterExpression_Tag::GetKeys() const
 {
-	return { TEXT("TAG") };
+	return { TEXT("Tag") };
 }
 
 ESequencerTextFilterValueType FSequencerTextFilterExpression_Tag::GetValueType() const
diff --git a/Engine/Source/Editor/Sequencer/Private/Filters/TextExpressions/SequencerTextFilterExpression_Time.cpp b/Engine/Source/Editor/Sequencer/Private/Filters/TextExpressions/SequencerTextFilterExpression_Time.cpp
--- a/Engine/Source/Editor/Sequencer/Private/Filters/TextExpressions/SequencerTextFilterExpression_Time.cpp
+++ b/Engine/Source/Editor/Sequencer/Private/Filters/TextExpressions/SequencerTextFilterExpression_Time.cpp
@@ -18,7 +18,7 @@ FSequencerTextFilterExpression_Time::FSequencerTextFilterExpression_Time(ISequen
 
 TSet<FName> FSequencerTextFilterExpression_Time::GetKeys() const
 {
-	return { TEXT("TIME") };
+	return { TEXT("Time") };
 }
 
 ESequencerTextFilterValueType FSequencerTextFilterExpression_Time::GetValueType() const
@@ -31,6 +31,15 @@ FText FSequencerTextFilterExpression_Time::GetDescription() const
 	return LOCTEXT("ExpressionDescription_Time", "Filter by time range (frame number)");
 }
 
+TArray<FSequencerTextFilterKeyword> FSequencerTextFilterExpression_Time::GetValueKeywords() const
+{
+	return {
+		{ TEXT("Now"), LOCTEXT("NowKeywordDescription", "Use the current playhead time value") },
+		{ TEXT("Start"), LOCTEXT("StartKeywordDescription", "Use the start time of the sequence") },
+		{ TEXT("End"), LOCTEXT("EndKeywordDescription", "Use the end time of the sequence") }
+	};
+}
+
 bool FSequencerTextFilterExpression_Time::TestComplexExpression(const FName& InKey
 	, const FTextFilterString& InValue
 	, const ETextFilterComparisonOperation InComparisonOperation
diff --git a/Engine/Source/Editor/Sequencer/Private/Filters/TextExpressions/SequencerTextFilterExpression_Time.h b/Engine/Source/Editor/Sequencer/Private/Filters/TextExpressions/SequencerTextFilterExpression_Time.h
--- a/Engine/Source/Editor/Sequencer/Private/Filters/TextExpressions/SequencerTextFilterExpression_Time.h
+++ b/Engine/Source/Editor/Sequencer/Private/Filters/TextExpressions/SequencerTextFilterExpression_Time.h
@@ -20,6 +20,7 @@ public:
 	virtual TSet<FName> GetKeys() const override;
 	virtual ESequencerTextFilterValueType GetValueType() const override;
 	virtual FText GetDescription() const override;
+	virtual TArray<FSequencerTextFilterKeyword> GetValueKeywords() const override;
 	//~ End FSequencerTextFilterExpressionContext
 
 	//~ Begin ITextFilterExpressionContext
diff --git a/Engine/Source/Editor/Sequencer/Private/Filters/TextExpressions/SequencerTextFilterExpression_TrackClass.cpp b/Engine/Source/Editor/Sequencer/Private/Filters/TextExpressions/SequencerTextFilterExpression_TrackClass.cpp
--- a/Engine/Source/Editor/Sequencer/Private/Filters/TextExpressions/SequencerTextFilterExpression_TrackClass.cpp
+++ b/Engine/Source/Editor/Sequencer/Private/Filters/TextExpressions/SequencerTextFilterExpression_TrackClass.cpp
@@ -15,7 +15,7 @@ FSequencerTextFilterExpression_TrackClass::FSequencerTextFilterExpression_TrackC
 
 TSet<FName> FSequencerTextFilterExpression_TrackClass::GetKeys() const
 {
-	return { TEXT("TRACKCLASS"), TEXT("TRACKTYPE") };
+	return { TEXT("TrackClass"), TEXT("TrackType") };
 }
 
 ESequencerTextFilterValueType FSequencerTextFilterExpression_TrackClass::GetValueType() const
@@ -25,7 +25,7 @@ ESequencerTextFilterValueType FSequencerTextFilterExpression_TrackClass::GetValu
 
 FText FSequencerTextFilterExpression_TrackClass::GetDescription() const
 {
-	return LOCTEXT("ExpressionDescription_Class", "Filter by track class name");
+	return LOCTEXT("ExpressionDescription_TrackClass", "Filter by track class name");
 }
 
 bool FSequencerTextFilterExpression_TrackClass::TestComplexExpression(const FName& InKey
diff --git a/Engine/Source/Editor/Sequencer/Private/Filters/TextExpressions/SequencerTextFilterExpression_Unbound.cpp b/Engine/Source/Editor/Sequencer/Private/Filters/TextExpressions/SequencerTextFilterExpression_Unbound.cpp
--- a/Engine/Source/Editor/Sequencer/Private/Filters/TextExpressions/SequencerTextFilterExpression_Unbound.cpp
+++ b/Engine/Source/Editor/Sequencer/Private/Filters/TextExpressions/SequencerTextFilterExpression_Unbound.cpp
@@ -18,7 +18,7 @@ FSequencerTextFilterExpression_Unbound::FSequencerTextFilterExpression_Unbound(I
 
 TSet<FName> FSequencerTextFilterExpression_Unbound::GetKeys() const
 {
-	return { TEXT("UNBOUND"), TEXT("MISSING") };
+	return { TEXT("Unbound"), TEXT("Missing") };
 }
 
 ESequencerTextFilterValueType FSequencerTextFilterExpression_Unbound::GetValueType() const
@@ -28,7 +28,7 @@ ESequencerTextFilterValueType FSequencerTextFilterExpression_Unbound::GetValueTy
 
 FText FSequencerTextFilterExpression_Unbound::GetDescription() const
 {
-	return LOCTEXT("ExpressionDescription_Locked", "Filter by track unbound or missing objects");
+	return LOCTEXT("ExpressionDescription_Unbound", "Filter by track unbound or missing objects");
 }
 
 bool FSequencerTextFilterExpression_Unbound::TestComplexExpression(const FName& InKey
diff --git a/Engine/Source/Editor/Sequencer/Private/Filters/Widgets/SFilterExpressionHelpDialog.cpp b/Engine/Source/Editor/Sequencer/Private/Filters/Widgets/SFilterExpressionHelpDialog.cpp
--- a/Engine/Source/Editor/Sequencer/Private/Filters/Widgets/SFilterExpressionHelpDialog.cpp
+++ b/Engine/Source/Editor/Sequencer/Private/Filters/Widgets/SFilterExpressionHelpDialog.cpp
@@ -174,7 +174,7 @@ TSharedRef<SWidget> SFilterExpressionHelpDialog::ConstructKeysWidget(const TSet<
 				.AutoWidth()
 				[
 					SNew(STextBlock)
-					.Text(LOCTEXT("KeySeperator", " | "))
+					.Text(LOCTEXT("KeySeparator", " | "))
 				];
 		}
 
@@ -218,14 +218,14 @@ TSharedRef<SWidget> SFilterExpressionHelpDialog::ConstructValueWidget(const ESeq
 	{
 	case ESequencerTextFilterValueType::String:
 		AddValueTypeSlot(LessThanText);
-		AddValueTypeSlot(LOCTEXT("StringValue", "STRING"), ValueColor);
+		AddValueTypeSlot(LOCTEXT("StringValue", "String"), ValueColor);
 		AddValueTypeSlot(GreaterThanText);
 		break;
 	case ESequencerTextFilterValueType::Boolean:
 		AddValueTypeSlot(LessThanText);
-		AddValueTypeSlot(LOCTEXT("TrueValue", "TRUE"), ValueColor);
+		AddValueTypeSlot(LOCTEXT("TrueValue", "True"), ValueColor);
 		AddValueTypeSlot(LOCTEXT("Slash", "/"));
-		AddValueTypeSlot(LOCTEXT("FalseValue", "FALSE"), ValueColor);
+		AddValueTypeSlot(LOCTEXT("FalseValue", "False"), ValueColor);
 		AddValueTypeSlot(GreaterThanText);
 		break;
 	case ESequencerTextFilterValueType::Integer:
diff --git a/Engine/Source/Editor/Sequencer/Public/Filters/SequencerTextFilterExpressionContext.h b/Engine/Source/Editor/Sequencer/Public/Filters/SequencerTextFilterExpressionContext.h
--- a/Engine/Source/Editor/Sequencer/Public/Filters/SequencerTextFilterExpressionContext.h
+++ b/Engine/Source/Editor/Sequencer/Public/Filters/SequencerTextFilterExpressionContext.h
@@ -16,6 +16,12 @@ enum class SEQUENCER_API ESequencerTextFilterValueType : uint8
 	Integer
 };
 
+struct SEQUENCER_API FSequencerTextFilterKeyword
+{
+	FString Keyword;
+	FText Description;
+};
+
 /** Text expression context to test the given asset data against the current text filter */
 class FSequencerTextFilterExpressionContext : public ITextFilterExpressionContext
 { 
@@ -25,9 +31,11 @@ public:
 	void SetFilterItem(FSequencerTrackFilterType InFilterItem, UMovieSceneTrack* const InTrackObject);
 
 	//~ Begin FSequencerTextFilterExpressionContext
-	virtual TSet<FName> GetKeys() const = 0;
-	virtual ESequencerTextFilterValueType GetValueType() const = 0;
-	virtual FText GetDescription() const = 0;
+	SEQUENCER_API virtual TSet<FName> GetKeys() const = 0;
+	SEQUENCER_API virtual ESequencerTextFilterValueType GetValueType() const = 0;
+	SEQUENCER_API virtual FText GetDescription() const = 0;
+	SEQUENCER_API virtual FText GetCategory() const { return FText::GetEmpty(); }
+	SEQUENCER_API virtual TArray<FSequencerTextFilterKeyword> GetValueKeywords() const { return {}; }
 	//~ End FSequencerTextFilterExpressionContext
 
 	//~ Begin ITextFilterExpressionContext
@@ -43,12 +51,12 @@ public:
 	//~ End ITextFilterExpressionContext
 
 protected:
-	UMovieSceneSequence* GetFocusedMovieSceneSequence() const;
-	UMovieScene* GetFocusedGetMovieScene() const;
+	SEQUENCER_API UMovieSceneSequence* GetFocusedMovieSceneSequence() const;
+	SEQUENCER_API UMovieScene* GetFocusedGetMovieScene() const;
 
-	bool CompareFStringForExactBool(const FTextFilterString& InValue, const bool bInPassedFilter) const;
+	SEQUENCER_API bool CompareFStringForExactBool(const FTextFilterString& InValue, const bool bInPassedFilter) const;
 
-	bool CompareFStringForExactBool(const FTextFilterString& InValue
+	SEQUENCER_API bool CompareFStringForExactBool(const FTextFilterString& InValue
 		, const ETextFilterComparisonOperation InComparisonOperation
 		, const bool bInPassedFilter) const;
 
