Sequencer: Minor renaming/UI touchup of track row properties.
* Removed 'track row metadata' header
* Renamed 'Track Row Metadata' category to 'Track Row'

diff --git a/Engine/Source/Editor/Sequencer/Private/MVVM/ViewModels/OutlinerItemModel.cpp b/Engine/Source/Editor/Sequencer/Private/MVVM/ViewModels/OutlinerItemModel.cpp
--- a/Engine/Source/Editor/Sequencer/Private/MVVM/ViewModels/OutlinerItemModel.cpp
+++ b/Engine/Source/Editor/Sequencer/Private/MVVM/ViewModels/OutlinerItemModel.cpp
@@ -762,7 +762,7 @@ void FOutlinerItemModelMixin::BuildTrackRowOptionsMenu(FMenuBuilder& MenuBuilder
 		return;
 	}
 
-	MenuBuilder.BeginSection(TEXT("TrackRowMetadata"), LOCTEXT("TrackRowMetadata", "Track Row Metadata"));
+	MenuBuilder.BeginSection(TEXT("TrackRowMetadata"));
 	{
 		// Empty here, will be implemented by extension.
 	}
diff --git a/Engine/Source/Editor/Sequencer/Private/Widgets/Sidebar/SequencerSelectionDrawer.cpp b/Engine/Source/Editor/Sequencer/Private/Widgets/Sidebar/SequencerSelectionDrawer.cpp
--- a/Engine/Source/Editor/Sequencer/Private/Widgets/Sidebar/SequencerSelectionDrawer.cpp
+++ b/Engine/Source/Editor/Sequencer/Private/Widgets/Sidebar/SequencerSelectionDrawer.cpp
@@ -326,7 +326,7 @@ void FSequencerSelectionDrawer::BuildTrackAreaDetails(const TSharedRef<FSequence
 		return TrackRow.Key.IsValid() && !TrackRow.Key->SupportsMultipleRows();
 		}))
 	{
-		MenuBuilder.BeginSection(TEXT("TrackRowMetadata"), LOCTEXT("TrackRowMetadata", "Track Row Metadata"));
+		MenuBuilder.BeginSection(TEXT("TrackRowMetadata"));
 		{
 			// Empty here, will be implemented by extension.
 		}
@@ -422,7 +422,7 @@ void FSequencerSelectionDrawer::BuildOutlinerDetails(const TSharedRef<FSequencer
 		return TrackRow.Key.IsValid() && !TrackRow.Key->SupportsMultipleRows();
 		}))
 	{
-		MenuBuilder.BeginSection(TEXT("TrackRowMetadata"), LOCTEXT("TrackRowMetadata", "Track Row Metadata"));
+		MenuBuilder.BeginSection(TEXT("TrackRowMetadata"));
 		{
 			// Empty here, will be implemented by extension.
 		}
diff --git a/Engine/Source/Runtime/MovieScene/Public/MovieSceneTrack.h b/Engine/Source/Runtime/MovieScene/Public/MovieSceneTrack.h
--- a/Engine/Source/Runtime/MovieScene/Public/MovieSceneTrack.h
+++ b/Engine/Source/Runtime/MovieScene/Public/MovieSceneTrack.h
@@ -176,7 +176,7 @@ struct MOVIESCENE_API FMovieSceneTrackRowMetadata
 	GENERATED_BODY()
 	
 	/* Optional dynamic conditions tied to specific track rows. */
-	UPROPERTY(EditAnywhere, Category="Track Row Metadata", meta=(ShowOnlyInnerProperties))
+	UPROPERTY(EditAnywhere, Category="Track Row", meta=(ShowOnlyInnerProperties))
 	FMovieSceneConditionContainer ConditionContainer;
 };
 
