Sequencer: Fixes to decorator column/icon behavior.
* Renamed the singleton decorator classes and interface to 'DecoratorBuilder' to be more explicit that they are singletons and create individual indicator widgets/column widgets
* Fixed opacity for condition decorator to reflect state of condition evaluation as well as children having conditions.

diff --git a/Engine/Source/Editor/Sequencer/Private/MVVM/ViewModels/OutlinerColumns/OutlinerDecoratorColumn.cpp b/Engine/Source/Editor/Sequencer/Private/MVVM/ViewModels/OutlinerColumns/OutlinerDecoratorColumn.cpp
--- a/Engine/Source/Editor/Sequencer/Private/MVVM/ViewModels/OutlinerColumns/OutlinerDecoratorColumn.cpp
+++ b/Engine/Source/Editor/Sequencer/Private/MVVM/ViewModels/OutlinerColumns/OutlinerDecoratorColumn.cpp
@@ -5,7 +5,7 @@
 #include "Sequencer.h"
 #include "MVVM/Extensions/IConditionableExtension.h"
 #include "MVVM/SharedViewModelData.h"
-#include "MVVM/ViewModels/OutlinerDecorators/IOutlinerDecorator.h"
+#include "MVVM/ViewModels/OutlinerDecorators/IOutlinerDecoratorBuilder.h"
 #include "MVVM/ViewModels/SequencerEditorViewModel.h"
 #include "Widgets/Layout/SBorder.h"
 #include "Widgets/OutlinerDecorators/SConditionDecoratorWidget.h"
@@ -38,7 +38,7 @@ bool FOutlinerDecoratorColumn::IsItemCompatibleWithColumn(const FCreateOutlinerC
 		return false;
 	}
 
-	for (const TTuple< FName, TSharedPtr<IOutlinerDecorator> >& OutlinerDecorator : Sequencer->GetOutlinerDecorators())
+	for (const TTuple< FName, TSharedPtr<IOutlinerDecoratorBuilder> >& OutlinerDecorator : Sequencer->GetOutlinerDecorators())
 	{
 		if (OutlinerDecorator.Value.Get()->IsItemCompatibleWithDecorator(InParams))
 		{
@@ -65,11 +65,11 @@ TSharedPtr<SWidget> FOutlinerDecoratorColumn::CreateColumnWidget(const FCreateOu
 
 	ColumnWidget = SNew(SHorizontalBox);
 
-	TArray<TSharedPtr<IOutlinerDecorator> > CompatibleDecorators;
+	TArray<TSharedPtr<IOutlinerDecoratorBuilder> > CompatibleDecorators;
 
-	for (const TTuple< FName, TSharedPtr<IOutlinerDecorator> >& OutlinerDecorator : Sequencer->GetOutlinerDecorators())
+	for (const TTuple< FName, TSharedPtr<IOutlinerDecoratorBuilder> >& OutlinerDecorator : Sequencer->GetOutlinerDecorators())
 	{
-		TSharedPtr<IOutlinerDecorator> Decorator = OutlinerDecorator.Value;
+		TSharedPtr<IOutlinerDecoratorBuilder> Decorator = OutlinerDecorator.Value;
 
 		if (Decorator->IsItemCompatibleWithDecorator(InParams))
 		{
@@ -79,7 +79,7 @@ TSharedPtr<SWidget> FOutlinerDecoratorColumn::CreateColumnWidget(const FCreateOu
 
 	const int32 NumCompatibleDecorators = CompatibleDecorators.Num();
 
-	for (TSharedPtr<IOutlinerDecorator>& Decorator : CompatibleDecorators)
+	for (TSharedPtr<IOutlinerDecoratorBuilder>& Decorator : CompatibleDecorators)
 	{
 		ColumnWidget->AddSlot()
 		[
diff --git a/Engine/Source/Editor/Sequencer/Private/MVVM/ViewModels/OutlinerDecorators/ConditionOutlinerDecorator.cpp b/Engine/Source/Editor/Sequencer/Private/MVVM/ViewModels/OutlinerDecorators/ConditionOutlinerDecorator.cpp
deleted file mode 100644
diff --git a/Engine/Source/Editor/Sequencer/Private/MVVM/ViewModels/OutlinerDecorators/ConditionOutlinerDecorator.h b/Engine/Source/Editor/Sequencer/Private/MVVM/ViewModels/OutlinerDecorators/ConditionOutlinerDecorator.h
deleted file mode 100644
diff --git a/Engine/Source/Editor/Sequencer/Private/MVVM/ViewModels/OutlinerDecorators/ConditionOutlinerDecoratorBuilder.cpp b/Engine/Source/Editor/Sequencer/Private/MVVM/ViewModels/OutlinerDecorators/ConditionOutlinerDecoratorBuilder.cpp
new file mode 100644
--- /dev/null
+++ b/Engine/Source/Editor/Sequencer/Private/MVVM/ViewModels/OutlinerDecorators/ConditionOutlinerDecoratorBuilder.cpp
@@ -0,0 +1,84 @@
+// Copyright Epic Games, Inc. All Rights Reserved.
+
+#include "MVVM/ViewModels/OutlinerDecorators/ConditionOutlinerDecoratorBuilder.h"
+
+#include "MVVM/ViewModels/OutlinerColumns/IOutlinerColumn.h"
+#include "MVVM/ViewModels/OutlinerColumns/OutlinerColumnTypes.h"
+#include "Widgets/OutlinerDecorators/SConditionDecoratorWidget.h"
+#include "MVVM/Extensions/IConditionableExtension.h"
+#include "MVVM/SharedViewModelData.h"
+#include "Widgets/Layout/SBorder.h"
+
+#define LOCTEXT_NAMESPACE "FConditionOutlinerDecoratorBuilder"
+
+namespace UE::Sequencer
+{
+
+FConditionOutlinerDecoratorBuilder::FConditionOutlinerDecoratorBuilder()
+{
+}
+
+FName FConditionOutlinerDecoratorBuilder::GetDecoratorName() const
+{
+	return FCommonOutlinerNames::Condition;
+}
+
+bool FConditionOutlinerDecoratorBuilder::IsItemCompatibleWithDecorator(const FCreateOutlinerColumnParams& InParams) const
+{
+	if (FConditionStateCacheExtension* ConditionStateCache = InParams.OutlinerExtension.AsModel()->GetSharedData()->CastThis<FConditionStateCacheExtension>())
+	{
+		return EnumHasAnyFlags(ConditionStateCache->GetCachedFlags(InParams.OutlinerExtension), ECachedConditionState::HasCondition | ECachedConditionState::ParentHasCondition | ECachedConditionState::ChildHasCondition | ECachedConditionState::SectionHasCondition);
+	}
+
+	return false;
+}
+
+TSharedPtr<SWidget> FConditionOutlinerDecoratorBuilder::CreateDecoratorWidget(const FCreateOutlinerColumnParams& InParams, const TSharedRef<ISequencerTreeViewRow>& TreeViewRow, const TSharedRef<IOutlinerColumn>& OutlinerColumn, const int32 NumCompatibleDecorators)
+{
+	static const FLinearColor ConditionColor = FLinearColor::FromSRGBColor(FColor(92, 220, 205));
+
+	TSharedRef<SConditionDecoratorWidget> ConditionDecoratorWidget = SNew(SConditionDecoratorWidget, OutlinerColumn, InParams);
+	if (NumCompatibleDecorators > 1)
+	{
+		return SNew(SBorder)
+			.VAlign(VAlign_Fill)
+			.HAlign(HAlign_Fill)
+			.Padding(0.0f)
+			.BorderImage(FAppStyle::GetBrush("WhiteBrush"))
+			.BorderBackgroundColor(ConditionColor);
+	}
+
+	return SNew(SOverlay)
+		+ SOverlay::Slot()
+		[
+			SNew(SBorder)
+				.VAlign(VAlign_Fill)
+				.HAlign(HAlign_Fill)
+				.Padding(0.0f)
+				.BorderImage(FAppStyle::GetBrush("WhiteBrush"))
+				.BorderBackgroundColor_Lambda([WeakWidget = ConditionDecoratorWidget.ToWeakPtr()]() -> FLinearColor
+					{
+						if (WeakWidget.IsValid())
+						{
+							return WeakWidget.Pin()->GetDecoratorBackgroundColorAndOpacity().GetSpecifiedColor();
+						}
+						return FLinearColor(0, 0, 0, 0);
+					})
+		]
+
+		+ SOverlay::Slot()
+		[
+			SNew(SBox)
+				.WidthOverride(12.f)
+				.HeightOverride(12.f)
+				.HAlign(HAlign_Center)
+				.VAlign(VAlign_Center)
+				[
+					ConditionDecoratorWidget
+				]
+		];
+}
+
+} // namespace UE::Sequencer
+
+#undef LOCTEXT_NAMESPACE
\ No newline at end of file
diff --git a/Engine/Source/Editor/Sequencer/Private/MVVM/ViewModels/OutlinerDecorators/ConditionOutlinerDecoratorBuilder.h b/Engine/Source/Editor/Sequencer/Private/MVVM/ViewModels/OutlinerDecorators/ConditionOutlinerDecoratorBuilder.h
new file mode 100644
--- /dev/null
+++ b/Engine/Source/Editor/Sequencer/Private/MVVM/ViewModels/OutlinerDecorators/ConditionOutlinerDecoratorBuilder.h
@@ -0,0 +1,26 @@
+// Copyright Epic Games, Inc. All Rights Reserved.
+
+#pragma once
+
+#include "MVVM/ViewModels/OutlinerDecorators/OutlinerDecoratorBuilderBase.h"
+
+namespace UE::Sequencer
+{
+
+struct FCreateOutlinerColumnParams;
+class IOutlinerColumn;
+class SConditionDecoratorWidget;
+
+class FConditionOutlinerDecoratorBuilder
+	: public FOutlinerDecoratorBuilderBase
+{
+public:
+	FConditionOutlinerDecoratorBuilder();
+	
+	virtual FName GetDecoratorName() const override;
+	virtual bool IsItemCompatibleWithDecorator(const FCreateOutlinerColumnParams& InParams) const override;
+	virtual TSharedPtr<SWidget> CreateDecoratorWidget(const FCreateOutlinerColumnParams& InParams, const TSharedRef<ISequencerTreeViewRow>& TreeViewRow, const TSharedRef<IOutlinerColumn>& OutlinerColumn, const int32 NumCompatibleDecorators) override;
+};
+
+} // namespace UE::Sequencer
+
diff --git a/Engine/Source/Editor/Sequencer/Private/MVVM/ViewModels/OutlinerDecorators/TimeWarpOutlinerDecorator.cpp b/Engine/Source/Editor/Sequencer/Private/MVVM/ViewModels/OutlinerDecorators/TimeWarpOutlinerDecorator.cpp
deleted file mode 100644
diff --git a/Engine/Source/Editor/Sequencer/Private/MVVM/ViewModels/OutlinerDecorators/TimeWarpOutlinerDecorator.h b/Engine/Source/Editor/Sequencer/Private/MVVM/ViewModels/OutlinerDecorators/TimeWarpOutlinerDecorator.h
deleted file mode 100644
diff --git a/Engine/Source/Editor/Sequencer/Private/MVVM/ViewModels/OutlinerDecorators/TimeWarpOutlinerDecoratorBuilder.cpp b/Engine/Source/Editor/Sequencer/Private/MVVM/ViewModels/OutlinerDecorators/TimeWarpOutlinerDecoratorBuilder.cpp
new file mode 100644
--- /dev/null
+++ b/Engine/Source/Editor/Sequencer/Private/MVVM/ViewModels/OutlinerDecorators/TimeWarpOutlinerDecoratorBuilder.cpp
@@ -0,0 +1,97 @@
+// Copyright Epic Games, Inc. All Rights Reserved.
+
+#include "MVVM/ViewModels/OutlinerDecorators/TimeWarpOutlinerDecoratorBuilder.h"
+
+#include "MovieSceneSection.h"
+#include "MVVM/Extensions/ITrackExtension.h"
+#include "MVVM/SharedViewModelData.h"
+#include "MVVM/ViewModels/SectionModel.h"
+#include "MVVM/ViewModels/OutlinerColumns/IOutlinerColumn.h"
+#include "MVVM/ViewModels/OutlinerColumns/OutlinerColumnTypes.h"
+#include "Variants/MovieSceneTimeWarpVariant.h"
+#include "Variants/MovieSceneTimeWarpGetter.h"
+#include "Widgets/Layout/SBorder.h"
+#include "Widgets/OutlinerDecorators/STimeWarpDecoratorWidget.h"
+
+#define LOCTEXT_NAMESPACE "FTimeWarpOutlinerDecoratorBuilder"
+
+namespace UE::Sequencer
+{
+
+FTimeWarpOutlinerDecoratorBuilder::FTimeWarpOutlinerDecoratorBuilder()
+{
+}
+
+FName FTimeWarpOutlinerDecoratorBuilder::GetDecoratorName() const
+{
+	return FCommonOutlinerNames::TimeWarp;
+}
+
+bool FTimeWarpOutlinerDecoratorBuilder::IsItemCompatibleWithDecorator(const FCreateOutlinerColumnParams& InParams) const
+{
+	if (InParams.OutlinerExtension.AsModel()->IsA<ITrackExtension>())
+	{
+		ITrackExtension* TrackModel = InParams.OutlinerExtension.AsModel()->CastThis<ITrackExtension>();
+
+		for (const TViewModelPtr<FSectionModel>& SectionModel : TrackModel->GetSectionModels().IterateSubList<FSectionModel>())
+		{
+			UMovieSceneSection* Section = SectionModel->GetSection();
+			if (!Section)
+			{
+				continue;
+			}
+
+			if (FMovieSceneTimeWarpVariant* Variant = Section->GetTimeWarp())
+			{
+				if (Variant->GetType() == EMovieSceneTimeWarpType::Custom)
+				{
+					return true;
+				}
+			}
+		}
+	}
+
+	return false;
+}
+
+TSharedPtr<SWidget> FTimeWarpOutlinerDecoratorBuilder::CreateDecoratorWidget(const FCreateOutlinerColumnParams& InParams, const TSharedRef<ISequencerTreeViewRow>& TreeViewRow, const TSharedRef<IOutlinerColumn>& OutlinerColumn, const int32 NumCompatibleDecorators)
+{
+	static const FLinearColor ConditionColor = FLinearColor::FromSRGBColor(FColor(212, 147, 20));
+
+	if (NumCompatibleDecorators > 1)
+	{
+		return SNew(SBorder)
+			.VAlign(VAlign_Fill)
+			.HAlign(HAlign_Fill)
+			.Padding(0.0f)
+			.BorderImage(FAppStyle::GetBrush("WhiteBrush"))
+			.BorderBackgroundColor_Lambda([this]() -> FLinearColor {  return ConditionColor; });
+	}
+
+	return SNew(SOverlay)
+		+ SOverlay::Slot()
+		[
+			SNew(SBorder)
+				.VAlign(VAlign_Fill)
+				.HAlign(HAlign_Fill)
+				.Padding(0.0f)
+				.BorderImage(FAppStyle::GetBrush("WhiteBrush"))
+				.BorderBackgroundColor_Lambda([this]() -> FLinearColor { return ConditionColor; })
+		]
+
+		+ SOverlay::Slot()
+		[
+			SNew(SBox)
+				.WidthOverride(12.f)
+				.HeightOverride(12.f)
+				.HAlign(HAlign_Center)
+				.VAlign(VAlign_Center)
+				[
+					SNew(STimeWarpDecoratorWidget, OutlinerColumn, InParams)
+				]
+		];
+}
+
+} // namespace UE::Sequencer
+
+#undef LOCTEXT_NAMESPACE
\ No newline at end of file
diff --git a/Engine/Source/Editor/Sequencer/Private/MVVM/ViewModels/OutlinerDecorators/TimeWarpOutlinerDecoratorBuilder.h b/Engine/Source/Editor/Sequencer/Private/MVVM/ViewModels/OutlinerDecorators/TimeWarpOutlinerDecoratorBuilder.h
new file mode 100644
--- /dev/null
+++ b/Engine/Source/Editor/Sequencer/Private/MVVM/ViewModels/OutlinerDecorators/TimeWarpOutlinerDecoratorBuilder.h
@@ -0,0 +1,26 @@
+// Copyright Epic Games, Inc. All Rights Reserved.
+
+#pragma once
+
+#include "MVVM/ViewModels/OutlinerDecorators/OutlinerDecoratorBuilderBase.h"
+
+namespace UE::Sequencer
+{
+
+struct FCreateOutlinerColumnParams;
+class IOutlinerColumn;
+class STimeWarpDecoratorWidget;
+
+class FTimeWarpOutlinerDecoratorBuilder
+	: public FOutlinerDecoratorBuilderBase
+{
+public:
+	FTimeWarpOutlinerDecoratorBuilder();
+	
+	virtual FName GetDecoratorName() const override;
+	virtual bool IsItemCompatibleWithDecorator(const FCreateOutlinerColumnParams& InParams) const override;
+	virtual TSharedPtr<SWidget> CreateDecoratorWidget(const FCreateOutlinerColumnParams& InParams, const TSharedRef<ISequencerTreeViewRow>& TreeViewRow, const TSharedRef<IOutlinerColumn>& OutlinerColumn, const int32 NumCompatibleDecorators) override;
+};
+
+} // namespace UE::Sequencer
+
diff --git a/Engine/Source/Editor/Sequencer/Private/MVVM/ViewModels/TrackModel.cpp b/Engine/Source/Editor/Sequencer/Private/MVVM/ViewModels/TrackModel.cpp
--- a/Engine/Source/Editor/Sequencer/Private/MVVM/ViewModels/TrackModel.cpp
+++ b/Engine/Source/Editor/Sequencer/Private/MVVM/ViewModels/TrackModel.cpp
@@ -678,27 +678,26 @@ const UMovieSceneCondition* FTrackModel::GetCondition() const
 
 EConditionableConditionState FTrackModel::GetConditionState() const
 {
-	UMovieSceneTrack* const Track = GetTrack();
-	if (IsValid(Track))
+	TSharedPtr<FSequenceModel> SequenceModel = FindAncestorOfType<FSequenceModel>();
+	TSharedPtr<ISequencer> Sequencer = SequenceModel ? SequenceModel->GetSequencer() : nullptr;
+	if (Sequencer)
 	{
-		if (Track->ConditionContainer.Condition)
-		{
-			if (Track->ConditionContainer.Condition->bEditorForceTrue)
-			{
-				return EConditionableConditionState::HasConditionEditorForceTrue;
-			}
+		FGuid BindingID;
 
-			TSharedPtr<FSequenceModel> SequenceModel = FindAncestorOfType<FSequenceModel>();
-			TSharedPtr<ISequencer> Sequencer = SequenceModel ? SequenceModel->GetSequencer() : nullptr;
-			if (Sequencer)
+		if (TSharedPtr<IObjectBindingExtension> ParentBinding = FindAncestorOfType<IObjectBindingExtension>())
+		{
+			BindingID = ParentBinding->GetObjectGuid();
+		}
+		UMovieSceneTrack* const Track = GetTrack();
+		if (IsValid(Track))
+		{
+			if (Track->ConditionContainer.Condition)
 			{
-				FGuid BindingID;
-
-				if (TSharedPtr<IObjectBindingExtension> ParentBinding = FindAncestorOfType<IObjectBindingExtension>())
+				if (Track->ConditionContainer.Condition->bEditorForceTrue)
 				{
-					BindingID = ParentBinding->GetObjectGuid();
+					return EConditionableConditionState::HasConditionEditorForceTrue;
 				}
-				
+
 				if (MovieSceneHelpers::EvaluateSequenceCondition(BindingID, Sequencer->GetFocusedTemplateID(), Track->ConditionContainer.Condition, Track, Sequencer->GetSharedPlaybackState()))
 				{
 					return EConditionableConditionState::HasConditionEvaluatingTrue;
@@ -708,6 +707,29 @@ EConditionableConditionState FTrackModel::GetConditionState() const
 					return EConditionableConditionState::HasConditionEvaluatingFalse;
 				}
 			}
+		
+			// Special case. If we support multiple rows, and there is only a single row, then we must also check track row metadata for a condition here, as there will be no track row model.
+			if (Track->SupportsMultipleRows() && Track->GetMaxRowIndex() == 0)
+			{
+				if (const FMovieSceneTrackRowMetadata* TrackRowMetadata = Track->FindTrackRowMetadata(GetRowIndex()))
+				{
+					if (TrackRowMetadata->ConditionContainer.Condition)
+					{
+						if (TrackRowMetadata->ConditionContainer.Condition->bEditorForceTrue)
+						{
+							return EConditionableConditionState::HasConditionEditorForceTrue;
+						}
+						else if (MovieSceneHelpers::EvaluateSequenceCondition(BindingID, Sequencer->GetFocusedTemplateID(), TrackRowMetadata->ConditionContainer.Condition, Track, Sequencer->GetSharedPlaybackState()))
+						{
+							return EConditionableConditionState::HasConditionEvaluatingTrue;
+						}
+						else
+						{
+							return EConditionableConditionState::HasConditionEvaluatingFalse;
+						}
+					}
+				}
+			}
 		}
 	}
 	return EConditionableConditionState::None;
diff --git a/Engine/Source/Editor/Sequencer/Private/Sequencer.cpp b/Engine/Source/Editor/Sequencer/Private/Sequencer.cpp
--- a/Engine/Source/Editor/Sequencer/Private/Sequencer.cpp
+++ b/Engine/Source/Editor/Sequencer/Private/Sequencer.cpp
@@ -131,7 +131,7 @@
 #include "ISequencerChannelInterface.h"
 #include "IMovieRendererInterface.h"
 #include "MVVM/ViewModels/OutlinerColumns/IOutlinerColumn.h"
-#include "MVVM/ViewModels/OutlinerDecorators/IOutlinerDecorator.h"
+#include "MVVM/ViewModels/OutlinerDecorators/IOutlinerDecoratorBuilder.h"
 #include "SequencerKeyCollection.h"
 #include "CurveEditor.h"
 #include "CurveEditorScreenSpace.h"
@@ -534,7 +534,7 @@ void FSequencer::InitSequencer(const FSequencerInitParams& InitParams
 	for (int32 DelegateIndex = 0; DelegateIndex < OutlinerDecoratorDelegates.Num(); ++DelegateIndex)
 	{
 		check(OutlinerDecoratorDelegates[DelegateIndex].IsBound());
-		TSharedRef<IOutlinerDecorator> OutlinerDecorator = OutlinerDecoratorDelegates[DelegateIndex].Execute();
+		TSharedRef<IOutlinerDecoratorBuilder> OutlinerDecorator = OutlinerDecoratorDelegates[DelegateIndex].Execute();
 
 		if (OutlinerDecorator->SupportsSequence(InitParams.RootSequence))
 		{
diff --git a/Engine/Source/Editor/Sequencer/Private/Sequencer.h b/Engine/Source/Editor/Sequencer/Private/Sequencer.h
--- a/Engine/Source/Editor/Sequencer/Private/Sequencer.h
+++ b/Engine/Source/Editor/Sequencer/Private/Sequencer.h
@@ -317,7 +317,7 @@ public:
 	/**
 	* @return Outliner Decorator Items registered to the sequencer by column name
 	*/
-	const TMap<FName, TSharedPtr<UE::Sequencer::IOutlinerDecorator>>& GetOutlinerDecorators() const
+	const TMap<FName, TSharedPtr<UE::Sequencer::IOutlinerDecoratorBuilder>>& GetOutlinerDecorators() const
 	{
 		return OutlinerDecorators;
 	}
@@ -1278,7 +1278,7 @@ private:
 	TMap<FName, TSharedPtr<UE::Sequencer::IOutlinerColumn>> OutlinerColumns;
 
 	/** List of Outliner decorator item creators that are supported by the Sequencer. */
-	TMap<FName, TSharedPtr<UE::Sequencer::IOutlinerDecorator>> OutlinerDecorators;
+	TMap<FName, TSharedPtr<UE::Sequencer::IOutlinerDecoratorBuilder>> OutlinerDecorators;
 
 	/** List of object bindings we can use */
 	TArray<TSharedPtr<ISequencerEditorObjectBinding>> ObjectBindings;
diff --git a/Engine/Source/Editor/Sequencer/Private/SequencerCommonHelpers.cpp b/Engine/Source/Editor/Sequencer/Private/SequencerCommonHelpers.cpp
--- a/Engine/Source/Editor/Sequencer/Private/SequencerCommonHelpers.cpp
+++ b/Engine/Source/Editor/Sequencer/Private/SequencerCommonHelpers.cpp
@@ -305,7 +305,15 @@ public:
 	{
 		if (const TSharedPtr<ISequencer> Sequencer = WeakSequencer.Pin())
 		{
-			Sequencer->NotifyMovieSceneDataChanged(EMovieSceneDataChangeType::TrackValueChanged);
+			if (PropertyThatChanged && PropertyThatChanged->GetName() == TEXT("Condition"))
+			{
+				// Rebuild hierarchy on changing a condition so the indicators have a chance to refresh
+				Sequencer->NotifyMovieSceneDataChanged(EMovieSceneDataChangeType::MovieSceneStructureItemsChanged);
+			}
+			else
+			{
+				Sequencer->NotifyMovieSceneDataChanged(EMovieSceneDataChangeType::TrackValueChanged);
+			}
 		}
 	}
 
diff --git a/Engine/Source/Editor/Sequencer/Private/SequencerModule.cpp b/Engine/Source/Editor/Sequencer/Private/SequencerModule.cpp
--- a/Engine/Source/Editor/Sequencer/Private/SequencerModule.cpp
+++ b/Engine/Source/Editor/Sequencer/Private/SequencerModule.cpp
@@ -45,8 +45,8 @@
 #include "MVVM/ViewModels/OutlinerColumns/KeyFrameOutlinerColumn.h"
 #include "MVVM/ViewModels/OutlinerColumns/ColorPickerOutlinerColumn.h"
 
-#include "MVVM/ViewModels/OutlinerDecorators/ConditionOutlinerDecorator.h"
-#include "MVVM/ViewModels/OutlinerDecorators/TimeWarpOutlinerDecorator.h"
+#include "MVVM/ViewModels/OutlinerDecorators/ConditionOutlinerDecoratorBuilder.h"
+#include "MVVM/ViewModels/OutlinerDecorators/TimeWarpOutlinerDecoratorBuilder.h"
 
 #include "ToolMenus.h"
 #include "ContentBrowserMenuContexts.h"
@@ -512,8 +512,8 @@ public:
 			ColorPickerOutlinerColumnHandle  = RegisterOutlinerColumn(FOnCreateOutlinerColumn::CreateStatic([]{ return TSharedRef<IOutlinerColumn>(MakeShared<FColorPickerOutlinerColumn>()); }));
 
 			// Register outliner decorator items
-			ConditionOutlinerDecoratorHandle = RegisterOutlinerDecorator(FOnCreateOutlinerDecorator::CreateStatic([] { return TSharedRef<IOutlinerDecorator>(MakeShared<FConditionOutlinerDecorator>()); }));
-			TimeWarpOutlinerDecoratorHandle = RegisterOutlinerDecorator(FOnCreateOutlinerDecorator::CreateStatic([] { return TSharedRef<IOutlinerDecorator>(MakeShared<FTimeWarpOutlinerDecorator>()); }));
+			ConditionOutlinerDecoratorHandle = RegisterOutlinerDecorator(FOnCreateOutlinerDecorator::CreateStatic([] { return TSharedRef<IOutlinerDecoratorBuilder>(MakeShared<FConditionOutlinerDecoratorBuilder>()); }));
+			TimeWarpOutlinerDecoratorHandle = RegisterOutlinerDecorator(FOnCreateOutlinerDecorator::CreateStatic([] { return TSharedRef<IOutlinerDecoratorBuilder>(MakeShared<FTimeWarpOutlinerDecoratorBuilder>()); }));
 
 			RegisterObjectSchemas();
 		}
diff --git a/Engine/Source/Editor/Sequencer/Private/Widgets/OutlinerDecorators/SConditionDecoratorWidget.cpp b/Engine/Source/Editor/Sequencer/Private/Widgets/OutlinerDecorators/SConditionDecoratorWidget.cpp
--- a/Engine/Source/Editor/Sequencer/Private/Widgets/OutlinerDecorators/SConditionDecoratorWidget.cpp
+++ b/Engine/Source/Editor/Sequencer/Private/Widgets/OutlinerDecorators/SConditionDecoratorWidget.cpp
@@ -6,8 +6,8 @@
 #include "MVVM/ViewModels/ViewModelIterators.h"
 #include "MVVM/ViewModels/EditorViewModel.h"
 #include "MVVM/ViewModels/OutlinerColumns/IOutlinerColumn.h"
-#include "MVVM/ViewModels/OutlinerDecorators/IOutlinerDecorator.h"
-#include "MVVM/ViewModels/OutlinerDecorators/ConditionOutlinerDecorator.h"
+#include "MVVM/ViewModels/OutlinerDecorators/IOutlinerDecoratorBuilder.h"
+#include "MVVM/ViewModels/OutlinerDecorators/ConditionOutlinerDecoratorBuilder.h"
 #include "MVVM/Extensions/IConditionableExtension.h"
 #include "MVVM/Selection/SequencerCoreSelection.h"
 #include "MVVM/Selection/SequencerOutlinerSelection.h"
@@ -15,7 +15,7 @@
 namespace UE::Sequencer
 {
 
-void SConditionDecoratorWidget::Construct(const FArguments& InArgs, const TWeakPtr<IOutlinerColumn> InWeakOutlinerColumn, const TWeakPtr<FConditionOutlinerDecorator>& InWeakOutlinerDecorator, const FCreateOutlinerColumnParams& InParams)
+void SConditionDecoratorWidget::Construct(const FArguments& InArgs, const TWeakPtr<IOutlinerColumn> InWeakOutlinerColumn, const FCreateOutlinerColumnParams& InParams)
 {
 	SColumnToggleWidget::Construct(
 		SColumnToggleWidget::FArguments(),
@@ -24,8 +24,6 @@ void SConditionDecoratorWidget::Construct(const FArguments& InArgs, const TWeakP
 	);
 
 	WeakConditionStateCacheExtension = CastViewModel<FConditionStateCacheExtension>(InParams.OutlinerExtension.AsModel()->GetSharedData());
-
-	WeakOutlinerDecorator = InWeakOutlinerDecorator;
 }
 
 bool SConditionDecoratorWidget::IsActive() const
@@ -46,6 +44,7 @@ const FSlateBrush* SConditionDecoratorWidget::GetActiveBrush() const
 
 FSlateColor SConditionDecoratorWidget::GetImageColorAndOpacity() const
 {
+	FLinearColor OutColor = FLinearColor::Black;
 	TSharedPtr<FEditorViewModel> Editor = WeakEditor.Pin();
 	TViewModelPtr<IOutlinerExtension> OutlinerItem = WeakOutlinerExtension.Pin();
 
@@ -59,7 +58,46 @@ FSlateColor SConditionDecoratorWidget::GetImageColorAndOpacity() const
 		bChildHasCondition = EnumHasAnyFlags(StateCache->GetCachedFlags(ModelID), ECachedConditionState::ChildHasCondition);
 		bSectionHasCondition = EnumHasAnyFlags(StateCache->GetCachedFlags(ModelID), ECachedConditionState::SectionHasCondition);
 	}
-	FLinearColor OutColor = FLinearColor::Black;
+
+	if (!Editor || !OutlinerItem)
+	{
+		return OutColor;
+	}
+
+	float Opacity = 1.0f;
+
+	if (IsActive() || (bChildHasCondition && !OutlinerItem->IsExpanded()) || bSectionHasCondition)
+	{
+		Opacity = 1.0f;
+	}
+	else
+	{
+		// Not active, invisible
+		Opacity = 0.0f;
+	}
+	OutColor.A = Opacity; // The background should be at the specified opacity and the icon should stay 100% black.
+	
+	return OutColor;
+}
+
+FSlateColor SConditionDecoratorWidget::GetDecoratorBackgroundColorAndOpacity() const
+{
+	static const FLinearColor ConditionColor = FLinearColor::FromSRGBColor(FColor(92, 220, 205));
+	TSharedPtr<FEditorViewModel> Editor = WeakEditor.Pin();
+	TViewModelPtr<IOutlinerExtension> OutlinerItem = WeakOutlinerExtension.Pin();
+
+	bool bActive = IsActive();
+	bool bEvaluatingTrue = false;
+	bool bChildHasCondition = false;
+	bool bSectionHasCondition = false;
+	if (TViewModelPtr<FConditionStateCacheExtension> StateCache = WeakConditionStateCacheExtension.Pin())
+	{
+		bEvaluatingTrue = EnumHasAnyFlags(StateCache->GetCachedFlags(ModelID), ECachedConditionState::ConditionEvaluatingTrue);
+		bChildHasCondition = EnumHasAnyFlags(StateCache->GetCachedFlags(ModelID), ECachedConditionState::ChildHasCondition);
+		bSectionHasCondition = EnumHasAnyFlags(StateCache->GetCachedFlags(ModelID), ECachedConditionState::SectionHasCondition);
+	}
+
+	FLinearColor OutColor = ConditionColor;
 
 	if (!Editor || !OutlinerItem)
 	{
@@ -88,12 +126,7 @@ FSlateColor SConditionDecoratorWidget::GetImageColorAndOpacity() const
 		// Not active, invisible
 		Opacity = 0.0f;
 	}
-	////OutColor.A = Opacity; // The background should be at the specified opacity and the icon should stay 100% black.
-
-	if (WeakOutlinerDecorator.IsValid())
-	{
-		WeakOutlinerDecorator.Pin()->Opacity = Opacity;
-	}
+	OutColor.A = Opacity; // The background should be at the specified opacity and the icon should stay 100% black.
 
 	return OutColor;
 }
diff --git a/Engine/Source/Editor/Sequencer/Private/Widgets/OutlinerDecorators/SConditionDecoratorWidget.h b/Engine/Source/Editor/Sequencer/Private/Widgets/OutlinerDecorators/SConditionDecoratorWidget.h
--- a/Engine/Source/Editor/Sequencer/Private/Widgets/OutlinerDecorators/SConditionDecoratorWidget.h
+++ b/Engine/Source/Editor/Sequencer/Private/Widgets/OutlinerDecorators/SConditionDecoratorWidget.h
@@ -9,7 +9,7 @@ namespace UE::Sequencer
 
 class IOutlinerColumn;
 class IOutlinerDecorator;
-class FConditionOutlinerDecorator;
+class FConditionOutlinerDecoratorBuilder;
 class FConditionStateCacheExtension;
 
 /**
@@ -22,11 +22,14 @@ public:
 	SLATE_BEGIN_ARGS(SConditionDecoratorWidget) {}
 	SLATE_END_ARGS()
 
-	void Construct(const FArguments& InArgs, const TWeakPtr<IOutlinerColumn> InWeakOutlinerColumn, const TWeakPtr<FConditionOutlinerDecorator>& OutlinerDecorator, const FCreateOutlinerColumnParams& InParams);
+	void Construct(const FArguments& InArgs, const TWeakPtr<IOutlinerColumn> InWeakOutlinerColumn, const FCreateOutlinerColumnParams& InParams);
 
 	/** Get the color and opacity of the column toggle widget. */
 	virtual FSlateColor GetImageColorAndOpacity() const override;
 
+	/* Gets the color and opacity of the decorator background. */
+	FSlateColor GetDecoratorBackgroundColorAndOpacity() const;
+
 protected:
 
 	/** Returns whether or not the item is directly active or not. */
@@ -58,8 +61,6 @@ private:
 
 	/** Weak cache extension ptr (can be null). */
 	TWeakViewModelPtr<FConditionStateCacheExtension> WeakConditionStateCacheExtension;
-
-	TWeakPtr<FConditionOutlinerDecorator> WeakOutlinerDecorator;
 };
 
 } // namespace UE::Sequencer
\ No newline at end of file
diff --git a/Engine/Source/Editor/Sequencer/Private/Widgets/OutlinerDecorators/STimeWarpDecoratorWidget.cpp b/Engine/Source/Editor/Sequencer/Private/Widgets/OutlinerDecorators/STimeWarpDecoratorWidget.cpp
--- a/Engine/Source/Editor/Sequencer/Private/Widgets/OutlinerDecorators/STimeWarpDecoratorWidget.cpp
+++ b/Engine/Source/Editor/Sequencer/Private/Widgets/OutlinerDecorators/STimeWarpDecoratorWidget.cpp
@@ -6,23 +6,21 @@
 #include "MVVM/ViewModels/ViewModelIterators.h"
 #include "MVVM/ViewModels/EditorViewModel.h"
 #include "MVVM/ViewModels/OutlinerColumns/IOutlinerColumn.h"
-#include "MVVM/ViewModels/OutlinerDecorators/IOutlinerDecorator.h"
-#include "MVVM/ViewModels/OutlinerDecorators/TimeWarpOutlinerDecorator.h"
+#include "MVVM/ViewModels/OutlinerDecorators/IOutlinerDecoratorBuilder.h"
+#include "MVVM/ViewModels/OutlinerDecorators/TimeWarpOutlinerDecoratorBuilder.h"
 #include "MVVM/Selection/SequencerCoreSelection.h"
 #include "MVVM/Selection/SequencerOutlinerSelection.h"
 
 namespace UE::Sequencer
 {
 
-void STimeWarpDecoratorWidget::Construct(const FArguments& InArgs, const TWeakPtr<IOutlinerColumn> InWeakOutlinerColumn, const TWeakPtr<FTimeWarpOutlinerDecorator>& InWeakOutlinerDecorator, const FCreateOutlinerColumnParams& InParams)
+void STimeWarpDecoratorWidget::Construct(const FArguments& InArgs, const TWeakPtr<IOutlinerColumn> InWeakOutlinerColumn, const FCreateOutlinerColumnParams& InParams)
 {
 	SColumnToggleWidget::Construct(
 		SColumnToggleWidget::FArguments(),
 		InWeakOutlinerColumn,
 		InParams
 	);
-
-	WeakOutlinerDecorator = InWeakOutlinerDecorator;
 }
 
 bool STimeWarpDecoratorWidget::IsActive() const
@@ -38,16 +36,7 @@ const FSlateBrush* STimeWarpDecoratorWidget::GetActiveBrush() const
 
 FSlateColor STimeWarpDecoratorWidget::GetImageColorAndOpacity() const
 {
-	FLinearColor OutColor = FLinearColor::Black;
-	float Opacity = 1.0f;
-	////OutColor.A = Opacity; // The background should be at the specified opacity and the icon should stay 100% black.
-
-	if (WeakOutlinerDecorator.IsValid())
-	{
-		WeakOutlinerDecorator.Pin()->Opacity = Opacity;
-	}
-
-	return OutColor;
+	return FLinearColor::Black;
 }
 
 } // namespace UE::Sequencer
diff --git a/Engine/Source/Editor/Sequencer/Private/Widgets/OutlinerDecorators/STimeWarpDecoratorWidget.h b/Engine/Source/Editor/Sequencer/Private/Widgets/OutlinerDecorators/STimeWarpDecoratorWidget.h
--- a/Engine/Source/Editor/Sequencer/Private/Widgets/OutlinerDecorators/STimeWarpDecoratorWidget.h
+++ b/Engine/Source/Editor/Sequencer/Private/Widgets/OutlinerDecorators/STimeWarpDecoratorWidget.h
@@ -9,7 +9,7 @@ namespace UE::Sequencer
 
 class IOutlinerColumn;
 class IOutlinerDecorator;
-class FTimeWarpOutlinerDecorator;
+class FTimeWarpOutlinerDecoratorBuilder;
 
 /**
  * A widget that shows an decorator based on the presence of time warp on a row
@@ -21,7 +21,7 @@ public:
 	SLATE_BEGIN_ARGS(STimeWarpDecoratorWidget) {}
 	SLATE_END_ARGS()
 
-	void Construct(const FArguments& InArgs, const TWeakPtr<IOutlinerColumn> InWeakOutlinerColumn, const TWeakPtr<FTimeWarpOutlinerDecorator>& OutlinerDecorator, const FCreateOutlinerColumnParams& InParams);
+	void Construct(const FArguments& InArgs, const TWeakPtr<IOutlinerColumn> InWeakOutlinerColumn, const FCreateOutlinerColumnParams& InParams);
 
 	/** Get the color and opacity of the column toggle widget. */
 	virtual FSlateColor GetImageColorAndOpacity() const override;
@@ -55,7 +55,7 @@ protected:
 
 private:
 
-	TWeakPtr<FTimeWarpOutlinerDecorator> WeakOutlinerDecorator;
+	TWeakPtr<FTimeWarpOutlinerDecoratorBuilder> WeakOutlinerDecorator;
 };
 
 } // namespace UE::Sequencer
\ No newline at end of file
diff --git a/Engine/Source/Editor/Sequencer/Public/ISequencerModule.h b/Engine/Source/Editor/Sequencer/Public/ISequencerModule.h
--- a/Engine/Source/Editor/Sequencer/Public/ISequencerModule.h
+++ b/Engine/Source/Editor/Sequencer/Public/ISequencerModule.h
@@ -29,7 +29,7 @@ namespace UE::Sequencer
 	class FTrackModel;
 	class IObjectSchema;
 	class IOutlinerColumn;
-	class IOutlinerDecorator;
+	class IOutlinerDecoratorBuilder;
 } // namespace UE::Sequencer
 
 enum class ECurveEditorTreeFilterType : uint32;
@@ -65,7 +65,7 @@ DECLARE_DELEGATE_RetVal_OneParam(TSharedPtr<UE::Sequencer::FTrackModel>, FOnCrea
 DECLARE_DELEGATE_RetVal(TSharedRef<UE::Sequencer::IOutlinerColumn>, FOnCreateOutlinerColumn);
 
 /** A delegate which will create an outliner decorator item */
-DECLARE_DELEGATE_RetVal(TSharedRef<UE::Sequencer::IOutlinerDecorator>, FOnCreateOutlinerDecorator);
+DECLARE_DELEGATE_RetVal(TSharedRef<UE::Sequencer::IOutlinerDecoratorBuilder>, FOnCreateOutlinerDecorator);
 
 /** A delegate that is executed when adding menu content. */
 DECLARE_DELEGATE_OneParam(FOnGetContextMenuContent, FMenuBuilder& /*MenuBuilder*/);
diff --git a/Engine/Source/Editor/SequencerCore/Public/MVVM/ViewModels/OutlinerDecorators/IOutlinerDecorator.h b/Engine/Source/Editor/SequencerCore/Public/MVVM/ViewModels/OutlinerDecorators/IOutlinerDecorator.h
deleted file mode 100644
diff --git a/Engine/Source/Editor/SequencerCore/Public/MVVM/ViewModels/OutlinerDecorators/IOutlinerDecoratorBuilder.h b/Engine/Source/Editor/SequencerCore/Public/MVVM/ViewModels/OutlinerDecorators/IOutlinerDecoratorBuilder.h
new file mode 100644
--- /dev/null
+++ b/Engine/Source/Editor/SequencerCore/Public/MVVM/ViewModels/OutlinerDecorators/IOutlinerDecoratorBuilder.h
@@ -0,0 +1,42 @@
+// Copyright Epic Games, Inc. All Rights Reserved.
+
+#pragma once
+
+#include "Templates/SharedPointer.h"
+
+class SWidget;
+class UMovieSceneSequence;
+
+namespace UE::Sequencer
+{
+
+struct FCreateOutlinerColumnParams;
+class IOutlinerColumn;
+class ISequencerTreeViewRow;
+
+/**
+* Interface for building sequencer decorator outliner items.
+*/
+class IOutlinerDecoratorBuilder : public TSharedFromThis<IOutlinerDecoratorBuilder>
+{
+
+public:
+
+	/** Returns the name of the Decorator outliner item. */
+	virtual FName GetDecoratorName() const = 0;
+
+	/* Gets whether or not this Decorator outliner item is supported by a given Sequencer. */
+	virtual bool SupportsSequence(UMovieSceneSequence* InSequence) const { return true; }
+
+	virtual bool IsItemCompatibleWithDecorator(const FCreateOutlinerColumnParams& InParams) const = 0;
+	virtual TSharedPtr<SWidget> CreateDecoratorWidget(const FCreateOutlinerColumnParams& InParams, const TSharedRef<ISequencerTreeViewRow>& TreeViewRow, const TSharedRef<IOutlinerColumn>& OutlinerColumn, const int32 NumCompatibleDecorators) = 0;
+
+public:
+
+	/** Virtual destructor. */
+	virtual ~IOutlinerDecoratorBuilder() { }
+
+};
+
+} // namespace UE::Sequencer
+
diff --git a/Engine/Source/Editor/SequencerCore/Public/MVVM/ViewModels/OutlinerDecorators/OutlinerDecoratorBase.h b/Engine/Source/Editor/SequencerCore/Public/MVVM/ViewModels/OutlinerDecorators/OutlinerDecoratorBase.h
deleted file mode 100644
diff --git a/Engine/Source/Editor/SequencerCore/Public/MVVM/ViewModels/OutlinerDecorators/OutlinerDecoratorBuilderBase.h b/Engine/Source/Editor/SequencerCore/Public/MVVM/ViewModels/OutlinerDecorators/OutlinerDecoratorBuilderBase.h
new file mode 100644
--- /dev/null
+++ b/Engine/Source/Editor/SequencerCore/Public/MVVM/ViewModels/OutlinerDecorators/OutlinerDecoratorBuilderBase.h
@@ -0,0 +1,25 @@
+// Copyright Epic Games, Inc. All Rights Reserved.
+
+#pragma once
+
+#include "MVVM/ViewModels/OutlinerDecorators/IOutlinerDecoratorBuilder.h"
+
+namespace UE::Sequencer
+{
+
+class FOutlinerDecoratorBuilderBase
+	: public IOutlinerDecoratorBuilder
+{
+public:
+
+	FOutlinerDecoratorBuilderBase()
+	{
+	}
+	
+public:
+
+	bool IsItemCompatibleWithDecorator(const FCreateOutlinerColumnParams& InParams) const override { return false; }
+	TSharedPtr<SWidget> CreateDecoratorWidget(const FCreateOutlinerColumnParams& InParams, const TSharedRef<ISequencerTreeViewRow>& TreeViewRow, const TSharedRef<IOutlinerColumn>& OutlinerColumn, const int32 NumCompatibleDecorators) override { return nullptr; }
+};
+
+} // namespace UE::Sequencer
\ No newline at end of file
