[Sequencer] Advanced Filtering: Fix track row additive tracks being displayed when they shouldn't

diff --git a/Engine/Plugins/Animation/ControlRig/Source/ControlRigEditor/Private/Sequencer/ControlRigSequencerFilter.cpp b/Engine/Plugins/Animation/ControlRig/Source/ControlRigEditor/Private/Sequencer/ControlRigSequencerFilter.cpp
--- a/Engine/Plugins/Animation/ControlRig/Source/ControlRigEditor/Private/Sequencer/ControlRigSequencerFilter.cpp
+++ b/Engine/Plugins/Animation/ControlRig/Source/ControlRigEditor/Private/Sequencer/ControlRigSequencerFilter.cpp
@@ -73,7 +73,7 @@ public:
 	virtual bool PassesFilter(FSequencerTrackFilterType InItem) const override
 	{
 		FSequencerFilterData& FilterData = FilterInterface.GetFilterData();
-		const UMovieSceneTrack* const TrackObject = ResolveMovieSceneTrackObject(InItem, FilterData);
+		const UMovieSceneTrack* const TrackObject = FilterData.ResolveMovieSceneTrackObject(InItem);
 		const UMovieSceneControlRigParameterTrack* const Track = Cast<UMovieSceneControlRigParameterTrack>(TrackObject);
 		return IsValid(Track);
 	}
@@ -158,7 +158,7 @@ public:
 	{
 		FSequencerFilterData& FilterData = FilterInterface.GetFilterData();
 
-		UMovieSceneTrack* const TrackObject = ResolveMovieSceneTrackObject(InItem, FilterData);
+		UMovieSceneTrack* const TrackObject = FilterData.ResolveMovieSceneTrackObject(InItem);
 		URigHierarchy* const ControlRigHierarchy = GetControlRigHierarchyFromTrackObject(TrackObject);
 
 		if (!IsValid(ControlRigHierarchy))
diff --git a/Engine/Plugins/FX/Niagara/Source/NiagaraEditor/Private/Sequencer/NiagaraSequence/NiagaraSequencerFilters.cpp b/Engine/Plugins/FX/Niagara/Source/NiagaraEditor/Private/Sequencer/NiagaraSequence/NiagaraSequencerFilters.cpp
--- a/Engine/Plugins/FX/Niagara/Source/NiagaraEditor/Private/Sequencer/NiagaraSequence/NiagaraSequencerFilters.cpp
+++ b/Engine/Plugins/FX/Niagara/Source/NiagaraEditor/Private/Sequencer/NiagaraSequence/NiagaraSequencerFilters.cpp
@@ -68,7 +68,7 @@ public:
 	{
 		FSequencerFilterData& FilterData = FilterInterface.GetFilterData();
 
-		const UObject* const BoundObject = ResolveTrackBoundObject(InItem, FilterData);
+		const UObject* const BoundObject = FilterData.ResolveTrackBoundObject(GetSequencer(), InItem);
 		if (IsValid(BoundObject))
 		{
 			if (BoundObject->IsA(UMovieSceneNiagaraTrack::StaticClass())
diff --git a/Engine/Source/Editor/Sequencer/Private/Filters/Filters/SequencerTrackFilterBase.cpp b/Engine/Source/Editor/Sequencer/Private/Filters/Filters/SequencerTrackFilterBase.cpp
--- a/Engine/Source/Editor/Sequencer/Private/Filters/Filters/SequencerTrackFilterBase.cpp
+++ b/Engine/Source/Editor/Sequencer/Private/Filters/Filters/SequencerTrackFilterBase.cpp
@@ -1,7 +1,6 @@
 ﻿// Copyright Epic Games, Inc. All Rights Reserved.
 
 #include "Filters/SequencerTrackFilterBase.h"
-#include "Animation/WidgetAnimation.h"
 #include "Filters/SequencerFilterBar.h"
 #include "ISequencer.h"
 #include "LevelSequence.h"
@@ -76,81 +75,6 @@ bool FSequencerTrackFilter::SupportsUMGSequence(UMovieSceneSequence* const InSeq
 		&& InSequence->GetClass()->IsChildOf(WidgetAnimationClass);
 }
 
-UMovieSceneTrack* FSequencerTrackFilter::ResolveMovieSceneTrackObject(FSequencerTrackFilterType InNode, FSequencerFilterData& FilterData)
-{
-	if (!InNode.IsValid())
-	{
-		return nullptr;
-	}
-
-	const TWeakViewModelPtr<IOutlinerExtension> WeakOutlinerNode = InNode.ImplicitCast();
-
-	// Use cache version if it exists, otherwise resolve below
-	if (FilterData.ResolvedTrackObjects.Contains(WeakOutlinerNode))
-	{
-		if (FilterData.ResolvedTrackObjects[WeakOutlinerNode].IsValid())
-		{
-			return FilterData.ResolvedTrackObjects[WeakOutlinerNode].Get();
-		}
-
-		FilterData.ResolvedTrackObjects.Remove(WeakOutlinerNode);
-	}
-
-	UMovieSceneTrack* TrackObject = nullptr;
-
-	if (const TViewModelPtr<ITrackExtension> AncestorTrackModel = InNode->FindAncestorOfType<ITrackExtension>(true))
-	{
-		TrackObject = AncestorTrackModel->GetTrack();
-	}
-
-	if (IsValid(TrackObject))
-	{
-		FilterData.ResolvedTrackObjects.Add(WeakOutlinerNode, TrackObject);
-	}
-
-	return TrackObject;
-}
-
-UObject* FSequencerTrackFilter::ResolveTrackBoundObject(ISequencer& InSequencer, FSequencerTrackFilterType InNode, FSequencerFilterData& FilterData)
-{
-	if (!InNode.IsValid())
-	{
-		return nullptr;
-	}
-
-	const TWeakViewModelPtr<IOutlinerExtension> WeakOutlinerNode = InNode.ImplicitCast();
-
-	// Use cache version if it exists, otherwise resolve below
-	if (FilterData.ResolvedBoundObjects.Contains(WeakOutlinerNode))
-	{
-		if (FilterData.ResolvedBoundObjects[WeakOutlinerNode].IsValid())
-		{
-			return FilterData.ResolvedBoundObjects[WeakOutlinerNode].Get();
-		}
-
-		FilterData.ResolvedBoundObjects.Remove(WeakOutlinerNode);
-	}
-
-	UObject* BoundObject = nullptr;
-
-	if (const TViewModelPtr<IObjectBindingExtension> ObjectBindingModel = InNode->FindAncestorOfType<IObjectBindingExtension>(true))
-	{
-		BoundObject = InSequencer.FindSpawnedObjectOrTemplate(ObjectBindingModel->GetObjectGuid());
-	}
-
-	if (IsValid(BoundObject))
-	{
-		FilterData.ResolvedBoundObjects.Add(WeakOutlinerNode, BoundObject);
-	}
-
-	return BoundObject;
-}
-
-UObject* FSequencerTrackFilter::ResolveTrackBoundObject(const TViewModelPtr<FViewModel> InNode, FSequencerFilterData& FilterData) const
-{
-	return ResolveTrackBoundObject(GetSequencer(), InNode, FilterData);
-}
-
 FText FSequencerTrackFilter::BuildTooltipTextForCommand(const FText& InBaseText, const TSharedPtr<FUICommandInfo>& InCommand)
 {
 	const TSharedRef<const FInputChord> FirstValidChord = InCommand->GetFirstValidChord();
diff --git a/Engine/Source/Editor/Sequencer/Private/Filters/Filters/SequencerTrackFilter_Keyed.cpp b/Engine/Source/Editor/Sequencer/Private/Filters/Filters/SequencerTrackFilter_Keyed.cpp
--- a/Engine/Source/Editor/Sequencer/Private/Filters/Filters/SequencerTrackFilter_Keyed.cpp
+++ b/Engine/Source/Editor/Sequencer/Private/Filters/Filters/SequencerTrackFilter_Keyed.cpp
@@ -40,18 +40,8 @@ bool FSequencerTrackFilter_Keyed::SupportsSequence(UMovieSceneSequence* const In
 
 bool FSequencerTrackFilter_Keyed::PassesFilter(FSequencerTrackFilterType InItem) const
 {
-	FSequencerFilterData& FilterData = FilterInterface.GetFilterData();
-
-	for (const TViewModelPtr<IOutlinerExtension>& ChildNode : InItem->GetDescendantsOfType<IOutlinerExtension>(true))
-	{
-		const UMovieSceneTrack* const TrackObject = ResolveMovieSceneTrackObject(InItem, FilterData);
-		if (DoesMovieSceneTrackHaveKeys(TrackObject))
-		{
-			return true;
-		}
-	}
-
-	return false;
+	const TWeakViewModelPtr<ITrackExtension> Track = FilterInterface.GetFilterData().ResolveTrack(InItem);
+	return DoesTrackExtensionHaveKeys(Track);
 }
 
 FText FSequencerTrackFilter_Keyed::GetDisplayName() const
@@ -69,24 +59,36 @@ FString FSequencerTrackFilter_Keyed::GetName() const
 	return TEXT("Keyed");
 }
 
-bool FSequencerTrackFilter_Keyed::DoesMovieSceneTrackHaveKeys(const UMovieSceneTrack* const InTrackObject)
+bool FSequencerTrackFilter_Keyed::DoesTrackExtensionHaveKeys(const TWeakViewModelPtr<ITrackExtension>& InTrack)
 {
-	if (!IsValid(InTrackObject))
+	const TViewModelPtr<ITrackExtension> Track = InTrack.Pin();
+	if (!Track.IsValid())
+	{
+		return false;
+	}
+
+	UMovieSceneTrack* const TrackObject = Track->GetTrack();
+	if (!IsValid(TrackObject))
 	{
 		return false;
 	}
 
-	for (const UMovieSceneSection* const Section : InTrackObject->GetAllSections())
+	const int32 RowIndex = Track->GetRowIndex();
+
+	for (const UMovieSceneSection* const Section : TrackObject->GetAllSections())
 	{
-		for (const FMovieSceneChannelEntry& ChannelEntry : Section->GetChannelProxy().GetAllEntries())
+		if (Section->GetRowIndex() == RowIndex)
 		{
-			const TConstArrayView<FMovieSceneChannel*> Channels = ChannelEntry.GetChannels();
-			for (int32 ChannelIndex = 0; ChannelIndex < Channels.Num(); ++ChannelIndex)
+			for (const FMovieSceneChannelEntry& ChannelEntry : Section->GetChannelProxy().GetAllEntries())
 			{
-				const FMovieSceneChannel* const Channel = Channels[ChannelIndex];
-				if (Channel && Channel->GetNumKeys() > 0)
+				const TConstArrayView<FMovieSceneChannel*> Channels = ChannelEntry.GetChannels();
+				for (int32 ChannelIndex = 0; ChannelIndex < Channels.Num(); ++ChannelIndex)
 				{
-					return true;
+					const FMovieSceneChannel* const Channel = Channels[ChannelIndex];
+					if (Channel && Channel->GetNumKeys() > 0)
+					{
+						return true;
+					}
 				}
 			}
 		}
diff --git a/Engine/Source/Editor/Sequencer/Private/Filters/Filters/SequencerTrackFilter_Keyed.h b/Engine/Source/Editor/Sequencer/Private/Filters/Filters/SequencerTrackFilter_Keyed.h
--- a/Engine/Source/Editor/Sequencer/Private/Filters/Filters/SequencerTrackFilter_Keyed.h
+++ b/Engine/Source/Editor/Sequencer/Private/Filters/Filters/SequencerTrackFilter_Keyed.h
@@ -7,7 +7,7 @@
 class FSequencerTrackFilter_Keyed : public FSequencerTrackFilter
 {
 public:
-	static bool DoesMovieSceneTrackHaveKeys(const UMovieSceneTrack* const InTrackObject);
+	static bool DoesTrackExtensionHaveKeys(const UE::Sequencer::TWeakViewModelPtr<UE::Sequencer::ITrackExtension>& InTrack);
 
 	FSequencerTrackFilter_Keyed(ISequencerTrackFilters& InFilterInterface, TSharedPtr<FFilterCategory> InCategory = nullptr);
 
diff --git a/Engine/Source/Editor/Sequencer/Private/Filters/Filters/SequencerTrackFilter_Modified.cpp b/Engine/Source/Editor/Sequencer/Private/Filters/Filters/SequencerTrackFilter_Modified.cpp
--- a/Engine/Source/Editor/Sequencer/Private/Filters/Filters/SequencerTrackFilter_Modified.cpp
+++ b/Engine/Source/Editor/Sequencer/Private/Filters/Filters/SequencerTrackFilter_Modified.cpp
@@ -46,7 +46,7 @@ bool FSequencerTrackFilter_Modified::PassesFilter(FSequencerTrackFilterType InIt
 {
 	FSequencerFilterData& FilterData = FilterInterface.GetFilterData();
 
-	UMovieSceneTrack* const TrackObject = ResolveMovieSceneTrackObject(InItem, FilterData);
+	const UMovieSceneTrack* const TrackObject = FilterData.ResolveMovieSceneTrackObject(InItem);
 	if (!IsValid(TrackObject))
 	{
 		return true;
diff --git a/Engine/Source/Editor/Sequencer/Private/Filters/Filters/SequencerTrackFilter_Selected.cpp b/Engine/Source/Editor/Sequencer/Private/Filters/Filters/SequencerTrackFilter_Selected.cpp
--- a/Engine/Source/Editor/Sequencer/Private/Filters/Filters/SequencerTrackFilter_Selected.cpp
+++ b/Engine/Source/Editor/Sequencer/Private/Filters/Filters/SequencerTrackFilter_Selected.cpp
@@ -119,7 +119,7 @@ bool FSequencerTrackFilter_Selected::PassesFilter(FSequencerTrackFilterType InIt
 
 	FSequencerFilterData& FilterData = FilterInterface.GetFilterData();
 
-	const UObject* const TrackObject = ResolveTrackBoundObject(InItem, FilterData);
+	const UObject* const TrackObject = FilterData.ResolveTrackBoundObject(GetSequencer(), InItem);
 	if (!IsValid(TrackObject))
 	{
 		return false;
diff --git a/Engine/Source/Editor/Sequencer/Private/Filters/Filters/SequencerTrackFilter_Text.cpp b/Engine/Source/Editor/Sequencer/Private/Filters/Filters/SequencerTrackFilter_Text.cpp
--- a/Engine/Source/Editor/Sequencer/Private/Filters/Filters/SequencerTrackFilter_Text.cpp
+++ b/Engine/Source/Editor/Sequencer/Private/Filters/Filters/SequencerTrackFilter_Text.cpp
@@ -99,7 +99,7 @@ bool FSequencerTrackFilter_Text::PassesFilter(FSequencerTrackFilterType InItem)
 
 	for (const TSharedRef<FSequencerTextFilterExpressionContext>& TextFilterExpressionContext : TextFilterExpressionContexts)
 	{
-		UMovieSceneTrack* const TrackObject = ResolveMovieSceneTrackObject(InItem, FilterData);
+		UMovieSceneTrack* const TrackObject = FilterData.ResolveMovieSceneTrackObject(InItem);
 		TextFilterExpressionContext->SetFilterItem(InItem, TrackObject);
 
 		if (!TextFilterExpressionEvaluator.TestTextFilter(*TextFilterExpressionContext))
diff --git a/Engine/Source/Editor/Sequencer/Private/Filters/Filters/SequencerTrackFilters.h b/Engine/Source/Editor/Sequencer/Private/Filters/Filters/SequencerTrackFilters.h
--- a/Engine/Source/Editor/Sequencer/Private/Filters/Filters/SequencerTrackFilters.h
+++ b/Engine/Source/Editor/Sequencer/Private/Filters/Filters/SequencerTrackFilters.h
@@ -223,7 +223,7 @@ public:
 	virtual bool PassesFilter(FSequencerTrackFilterType InItem) const override
 	{
 		FSequencerFilterData& FilterData = FilterInterface.GetFilterData();
-		const UMovieSceneTrack* const TrackObject = ResolveMovieSceneTrackObject(InItem, FilterData);
+		const UMovieSceneTrack* const TrackObject = FilterData.ResolveMovieSceneTrackObject(InItem);
 		return IsValid(TrackObject)
 			&& TrackObject->IsA(UMovieSceneSubTrack::StaticClass())
 			&& !TrackObject->IsA(UMovieSceneCinematicShotTrack::StaticClass());
diff --git a/Engine/Source/Editor/Sequencer/Private/Filters/SequencerFilterData.cpp b/Engine/Source/Editor/Sequencer/Private/Filters/SequencerFilterData.cpp
--- a/Engine/Source/Editor/Sequencer/Private/Filters/SequencerFilterData.cpp
+++ b/Engine/Source/Editor/Sequencer/Private/Filters/SequencerFilterData.cpp
@@ -1,7 +1,10 @@
 ﻿// Copyright Epic Games, Inc. All Rights Reserved.
 
 #include "Filters/SequencerFilterData.h"
+#include "ISequencer.h"
+#include "MVVM/Extensions/IObjectBindingExtension.h"
 #include "MVVM/Extensions/IOutlinerExtension.h"
+#include "MVVM/Extensions/ITrackExtension.h"
 #include "MVVM/ViewModels/ViewModelIterators.h"
 
 using namespace UE::Sequencer;
@@ -126,3 +129,149 @@ bool FSequencerFilterData::IsFilteredOut(const TViewModelPtr<IOutlinerExtension>
 {
 	return !FilterInNodes.Contains(InNode);
 }
+
+TWeakViewModelPtr<ITrackExtension> FSequencerFilterData::ResolveTrack(FSequencerTrackFilterType InNode)
+{
+	if (!InNode.IsValid())
+	{
+		return nullptr;
+	}
+
+	const TWeakViewModelPtr<IOutlinerExtension> WeakOutlinerNode = InNode.ImplicitCast();
+
+	// Use cache version if it exists, otherwise resolve below
+	if (ResolvedTracks.Contains(WeakOutlinerNode))
+	{
+		if (const TViewModelPtr<ITrackExtension> ResolvedTrack = ResolvedTracks[WeakOutlinerNode].Pin())
+		{
+			return ResolvedTrack;
+		}
+
+		ResolvedTracks.Remove(WeakOutlinerNode);
+	}
+
+	// Resolve and cache
+	if (const TViewModelPtr<ITrackExtension> AncestorTrack = InNode->FindAncestorOfType<ITrackExtension>(true))
+	{
+		ResolvedTracks.Add(WeakOutlinerNode, AncestorTrack);
+		return AncestorTrack;
+	}
+
+	return nullptr;
+}
+
+TWeakViewModelPtr<IObjectBindingExtension> FSequencerFilterData::ResolveTrackObjectBinding(FSequencerTrackFilterType InNode)
+{
+	if (!InNode.IsValid())
+	{
+		return nullptr;
+	}
+
+	const TWeakViewModelPtr<IOutlinerExtension> WeakOutlinerNode = InNode.ImplicitCast();
+
+	// Use cache version if it exists, otherwise resolve below
+	if (ResolvedObjectBindings.Contains(WeakOutlinerNode))
+	{
+		if (const TViewModelPtr<IObjectBindingExtension> ResolvedObjectBinding = ResolvedObjectBindings[WeakOutlinerNode].Pin())
+		{
+			return ResolvedObjectBinding;
+		}
+
+		ResolvedObjectBindings.Remove(WeakOutlinerNode);
+	}
+
+	// Resolve and cache
+	if (const TViewModelPtr<IObjectBindingExtension> ObjectBindingModel = InNode->FindAncestorOfType<IObjectBindingExtension>(true))
+	{
+		ResolvedObjectBindings.Add(WeakOutlinerNode, ObjectBindingModel);
+		return ObjectBindingModel;
+	}
+
+	return nullptr;
+}
+
+UMovieSceneTrack* FSequencerFilterData::ResolveMovieSceneTrackObject(FSequencerTrackFilterType InNode)
+{
+	if (!InNode.IsValid())
+	{
+		return nullptr;
+	}
+
+	const TWeakViewModelPtr<IOutlinerExtension> WeakOutlinerNode = InNode.ImplicitCast();
+
+	// Use cache version if it exists, otherwise resolve below
+	if (ResolvedTracks.Contains(WeakOutlinerNode))
+	{
+		if (const TViewModelPtr<ITrackExtension> ResolvedTrack = ResolvedTracks[WeakOutlinerNode].Pin())
+		{
+			return ResolvedTrack->GetTrack();
+		}
+
+		ResolvedTracks.Remove(WeakOutlinerNode);
+	}
+
+	// Resolve and cache
+	if (const TViewModelPtr<ITrackExtension> AncestorTrackModel = InNode->FindAncestorOfType<ITrackExtension>(true))
+	{
+		UMovieSceneTrack* const TrackObject = AncestorTrackModel->GetTrack();
+		if (IsValid(TrackObject))
+		{
+			ResolvedObjects.Add(WeakOutlinerNode, TrackObject);
+			return TrackObject;
+		}
+	}
+
+	return nullptr;
+}
+
+UObject* FSequencerFilterData::ResolveTrackBoundObject(ISequencer& InSequencer, const TViewModelPtr<FViewModel> InNode)
+{
+	if (!InNode.IsValid())
+	{
+		return nullptr;
+	}
+
+	const TWeakViewModelPtr<IOutlinerExtension> WeakOutlinerNode = InNode.ImplicitCast();
+
+	// Use cache version if it exists, otherwise resolve below
+	if (ResolvedObjects.Contains(WeakOutlinerNode))
+	{
+		UObject* const ResolvedObject = ResolvedObjects[WeakOutlinerNode].Get();
+		if (IsValid(ResolvedObject))
+		{
+			return ResolvedObject;
+		}
+
+		ResolvedObjects.Remove(WeakOutlinerNode);
+	}
+
+	if (ResolvedObjectBindings.Contains(WeakOutlinerNode))
+	{
+		if (const TViewModelPtr<IObjectBindingExtension> ResolvedObjectBinding = ResolvedObjectBindings[WeakOutlinerNode].Pin())
+		{
+			UObject* const BoundObject = InSequencer.FindSpawnedObjectOrTemplate(ResolvedObjectBinding->GetObjectGuid());
+			if (IsValid(BoundObject))
+			{
+				ResolvedObjects.Add(WeakOutlinerNode, BoundObject);
+				return BoundObject;
+			}
+		}
+
+		ResolvedObjectBindings.Remove(WeakOutlinerNode);
+	}
+
+	// Resolve and cache
+	if (const TViewModelPtr<IObjectBindingExtension> ObjectBindingModel = InNode->FindAncestorOfType<IObjectBindingExtension>(true))
+	{
+		ResolvedObjectBindings.Add(WeakOutlinerNode, ObjectBindingModel);
+
+		UObject* const BoundObject = InSequencer.FindSpawnedObjectOrTemplate(ObjectBindingModel->GetObjectGuid());
+		if (IsValid(BoundObject))
+		{
+			ResolvedObjects.Add(WeakOutlinerNode, BoundObject);
+			return BoundObject;
+		}
+	}
+
+	return nullptr;
+}
diff --git a/Engine/Source/Editor/Sequencer/Private/Filters/TextExpressions/SequencerTextFilterExpression_ObjectClass.cpp b/Engine/Source/Editor/Sequencer/Private/Filters/TextExpressions/SequencerTextFilterExpression_ObjectClass.cpp
--- a/Engine/Source/Editor/Sequencer/Private/Filters/TextExpressions/SequencerTextFilterExpression_ObjectClass.cpp
+++ b/Engine/Source/Editor/Sequencer/Private/Filters/TextExpressions/SequencerTextFilterExpression_ObjectClass.cpp
@@ -39,7 +39,7 @@ bool FSequencerTextFilterExpression_ObjectClass::TestComplexExpression(const FNa
 
 	ISequencer& Sequencer = FilterInterface.GetSequencer();
 
-	UObject* const BoundObject = FSequencerTrackFilter::ResolveTrackBoundObject(Sequencer, FilterItem, FilterInterface.GetFilterData());
+	UObject* const BoundObject = FilterInterface.GetFilterData().ResolveTrackBoundObject(Sequencer, FilterItem);
 	if (IsValid(BoundObject))
 	{
 		const FString BoundObjectClassName = BoundObject->GetClass()->GetName();
diff --git a/Engine/Source/Editor/Sequencer/Public/Filters/SequencerFilterData.h b/Engine/Source/Editor/Sequencer/Public/Filters/SequencerFilterData.h
--- a/Engine/Source/Editor/Sequencer/Public/Filters/SequencerFilterData.h
+++ b/Engine/Source/Editor/Sequencer/Public/Filters/SequencerFilterData.h
@@ -5,11 +5,17 @@
 #include "MovieSceneTrack.h"
 #include "MVVM/ViewModelPtr.h"
 
+class ISequencer;
+
 namespace UE::Sequencer
 {
+	class IObjectBindingExtension;
 	class IOutlinerExtension;
+	class ITrackExtension;
 }
 
+using FSequencerTrackFilterType = UE::Sequencer::FViewModelPtr;
+
 /** Represents a cache between nodes for a filter operation. */
 struct SEQUENCER_API FSequencerFilterData
 {
@@ -42,8 +48,15 @@ struct SEQUENCER_API FSequencerFilterData
 
 	bool IsFilteredOut(const UE::Sequencer::TViewModelPtr<UE::Sequencer::IOutlinerExtension>& InNode) const;
 
-	TMap<UE::Sequencer::TWeakViewModelPtr<UE::Sequencer::IOutlinerExtension>, TWeakObjectPtr<UMovieSceneTrack>> ResolvedTrackObjects;
-	TMap<UE::Sequencer::TWeakViewModelPtr<UE::Sequencer::IOutlinerExtension>, TWeakObjectPtr<>> ResolvedBoundObjects;
+	UMovieSceneTrack* ResolveMovieSceneTrackObject(FSequencerTrackFilterType InNode);
+	UObject* ResolveTrackBoundObject(ISequencer& InSequencer, FSequencerTrackFilterType InNode);
+
+	UE::Sequencer::TWeakViewModelPtr<UE::Sequencer::ITrackExtension> ResolveTrack(FSequencerTrackFilterType InNode);
+	UE::Sequencer::TWeakViewModelPtr<UE::Sequencer::IObjectBindingExtension> ResolveTrackObjectBinding(FSequencerTrackFilterType InNode);
+
+	TMap<UE::Sequencer::TWeakViewModelPtr<UE::Sequencer::IOutlinerExtension>, UE::Sequencer::TWeakViewModelPtr<UE::Sequencer::ITrackExtension>> ResolvedTracks;
+	TMap<UE::Sequencer::TWeakViewModelPtr<UE::Sequencer::IOutlinerExtension>, UE::Sequencer::TWeakViewModelPtr<UE::Sequencer::IObjectBindingExtension>> ResolvedObjectBindings;
+	TMap<UE::Sequencer::TWeakViewModelPtr<UE::Sequencer::IOutlinerExtension>, TWeakObjectPtr<>> ResolvedObjects;
 
 protected:
 	FString RawFilterText;
diff --git a/Engine/Source/Editor/Sequencer/Public/Filters/SequencerTrackFilterBase.h b/Engine/Source/Editor/Sequencer/Public/Filters/SequencerTrackFilterBase.h
--- a/Engine/Source/Editor/Sequencer/Public/Filters/SequencerTrackFilterBase.h
+++ b/Engine/Source/Editor/Sequencer/Public/Filters/SequencerTrackFilterBase.h
@@ -8,6 +8,7 @@
 #include "Misc/IFilter.h"
 #include "MovieSceneTrack.h"
 #include "MovieSceneTrackEditor.h"
+#include "SequencerFilterData.h"
 #include "MVVM/ViewModelPtr.h"
 #include "Textures/SlateIcon.h"
 
@@ -49,9 +50,6 @@ public:
 		return Support == ETrackSupport::Supported;
 	}
 
-	SEQUENCER_API static UMovieSceneTrack* ResolveMovieSceneTrackObject(FSequencerTrackFilterType InNode, FSequencerFilterData& FilterData);
-	SEQUENCER_API static UObject* ResolveTrackBoundObject(ISequencer& InSequencer, FSequencerTrackFilterType InNode, FSequencerFilterData& FilterData);
-
 	SEQUENCER_API static FText BuildTooltipTextForCommand(const FText& InBaseText, const TSharedPtr<FUICommandInfo>& InCommand);
 
 	SEQUENCER_API FSequencerTrackFilter(ISequencerTrackFilters& InOutFilterInterface, TSharedPtr<FFilterCategory>&& InCategory = nullptr);
@@ -109,8 +107,6 @@ protected:
 	UMovieSceneSequence* GetFocusedMovieSceneSequence() const;
 	UMovieScene* GetFocusedGetMovieScene() const;
 
-	SEQUENCER_API UObject* ResolveTrackBoundObject(FSequencerTrackFilterType InNode, FSequencerFilterData& FilterData) const;
-
 	ISequencerTrackFilters& FilterInterface;
 
 private:
@@ -137,7 +133,6 @@ public:
 	{
 		const UE::Sequencer::TViewModelPtr<InModelType> Model = InItem->FindAncestorOfType<InModelType>();
 		return Model.IsValid(); // show child tracks
-		//return InItem->IsA<InModelType>(); // don't show child tracks
 	}
 	//~ End IFilter
 };
@@ -158,7 +153,7 @@ public:
 	virtual bool PassesFilter(FSequencerTrackFilterType InItem) const override
 	{
 		FSequencerFilterData& FilterData = FilterInterface.GetFilterData();
-		const UMovieSceneTrack* const TrackObject = ResolveMovieSceneTrackObject(InItem, FilterData);
+		const UMovieSceneTrack* const TrackObject = FilterData.ResolveMovieSceneTrackObject(InItem);
 		return IsValid(TrackObject) && TrackObject->IsA(InClassType::StaticClass());
 	}
 	//~ End IFilter
@@ -188,13 +183,13 @@ public:
 	{
 		FSequencerFilterData& FilterData = FilterInterface.GetFilterData();
 
-		const UMovieSceneTrack* const TrackObject = ResolveMovieSceneTrackObject(InItem, FilterData);
+		const UMovieSceneTrack* const TrackObject = FilterData.ResolveMovieSceneTrackObject(InItem);
 		if (IsValid(TrackObject) && TrackObject->IsA(InComponentType::StaticClass()))
 		{
 			return true;
 		}
 
-		const TWeakObjectPtr<> BoundObject = ResolveTrackBoundObject(GetSequencer(), InItem, FilterData);
+		const TWeakObjectPtr<> BoundObject = FilterData.ResolveTrackBoundObject(GetSequencer(), InItem);
 		if (BoundObject.IsValid() && BoundObject->IsA(InComponentType::StaticClass()))
 		{
 			return true;
