Sequencer: Updated timewarp UI and behavior

  - PlayRate timewarp curves now affect the sequence's duration at runtime. This means a constant playrate of 0.5 will cause the sequence to take twice as long to complete.
  - Warped time is now shown as an orange scrubber below the active time-warp track
  - Introduced new MVVM extension for controlling a model's time domain
  - TimeWarp tracks are now always pinned and sorted to the top of the tree
  - Added a new function to ISequencerSection that allows richer creation of custom section widgets including controlling ZOrder
  - Added various functionality to ISequencer for retrieving and mutating warped and unwarped time.

diff --git a/Engine/Source/Editor/EditorStyle/Private/StarshipStyle.cpp b/Engine/Source/Editor/EditorStyle/Private/StarshipStyle.cpp
--- a/Engine/Source/Editor/EditorStyle/Private/StarshipStyle.cpp
+++ b/Engine/Source/Editor/EditorStyle/Private/StarshipStyle.cpp
@@ -2140,13 +2140,9 @@ void FStarshipEditorStyle::FStyle::SetupSequencerStyles()
 		Set("Sequencer.KeyBar.Dashed", new BORDER_BRUSH(TEXT("Sequencer/Keys/KeyBar_Dashed"), FMargin(0.f, 4.f, 0.f, 0.f)));
 		Set("Sequencer.KeyBar.Solid", new BORDER_BRUSH(TEXT("Sequencer/Keys/KeyBar_Solid"), FMargin(0.f, 4.f, 0.f, 0.f)));
 
-		Set("Sequencer.Timeline.ScrubHandleDown", new BOX_BRUSH("Sequencer/ScrubHandleDown", FMargin(6.f / 13.f, 5 / 12.f, 6 / 13.f, 8 / 12.f)));
-		Set("Sequencer.Timeline.ScrubHandleUp", new BOX_BRUSH("Sequencer/ScrubHandleUp", FMargin(6.f / 13.f, 8 / 12.f, 6 / 13.f, 5 / 12.f)));
+		Set("Sequencer.Timeline.ScrubHandle", new BOX_BRUSH_SVG("Starship/Sequencer/ScrubHandle", FMargin(0.2f)));
 		Set("Sequencer.Timeline.ScrubFill", new BOX_BRUSH("Sequencer/ScrubFill", FMargin(2.f / 4.f, 0.f)));
-		Set("Sequencer.Timeline.FrameBlockScrubHandleDown", new BOX_BRUSH("Sequencer/ScrubHandleDown", FMargin(6.f / 13.f, 5 / 12.f, 6 / 13.f, 8 / 12.f)));
-		Set("Sequencer.Timeline.FrameBlockScrubHandleUp", new BOX_BRUSH("Sequencer/ScrubHandleUp", FMargin(6.f / 13.f, 8 / 12.f, 6 / 13.f, 5 / 12.f)));
-		Set("Sequencer.Timeline.VanillaScrubHandleDown", new BOX_BRUSH("Sequencer/ScrubHandleDown_Clamped", FMargin(6.f / 13.f, 3.f / 12.f, 6.f / 13.f, 7.f / 12.f)));
-		Set("Sequencer.Timeline.VanillaScrubHandleUp", new BOX_BRUSH("Sequencer/ScrubHandleUp_Clamped", FMargin(6.f / 13.f, 8 / 12.f, 6 / 13.f, 5 / 12.f)));
+		Set("Sequencer.Timeline.FrameBlockScrubHandle", new BOX_BRUSH("Sequencer/ScrubHandleDown", FMargin(6.f / 13.f, 5 / 12.f, 6 / 13.f, 8 / 12.f)));
 		Set("Sequencer.Timeline.ScrubHandleWhole", new BOX_BRUSH("Sequencer/ScrubHandleWhole", FMargin(6.f / 13.f, 10 / 24.f, 6 / 13.f, 10 / 24.f)));
 		Set("Sequencer.Timeline.RangeHandleLeft", new BOX_BRUSH("Sequencer/GenericGripLeft", FMargin(5.f / 16.f)));
 		Set("Sequencer.Timeline.RangeHandleRight", new BOX_BRUSH("Sequencer/GenericGripRight", FMargin(5.f / 16.f)));
@@ -2239,6 +2235,8 @@ void FStarshipEditorStyle::FStyle::SetupSequencerStyles()
 		Set("Sequencer.AnimationOutliner.RegularFont", DEFAULT_FONT("Regular", 9));
 		Set("Sequencer.AnimationOutliner.ItalicFont", DEFAULT_FONT("Italic", 10));
 
+		Set("Sequencer.Outliner.Decorators.TimeWarp", new IMAGE_BRUSH_SVG("Starship/Sequencer/TimeWarp_12", Icon12x12));
+
 		Set("Sequencer.ShotFilter", new IMAGE_BRUSH("Sequencer/FilteredArea", FVector2D(74, 74), FLinearColor::White, ESlateBrushTileType::Both));
 		Set("Sequencer.KeyMark", new IMAGE_BRUSH("Sequencer/KeyMark", FVector2D(3, 21), FLinearColor::White, ESlateBrushTileType::NoTile));
 		Set("Sequencer.ToggleAutoKeyEnabled", new IMAGE_BRUSH_SVG("Starship/Sequencer/AutoKey", Icon20x20));
@@ -2284,6 +2282,7 @@ void FStarshipEditorStyle::FStyle::SetupSequencerStyles()
 		Set("Sequencer.Tracks.CameraCut", new IMAGE_BRUSH_SVG("Starship/Sequencer/CameraCutTrack", Icon16x16));
 		Set("Sequencer.Tracks.CinematicShot", new IMAGE_BRUSH_SVG("Starship/Sequencer/ShotTrack", Icon16x16));
 		Set("Sequencer.Tracks.Slomo", new IMAGE_BRUSH_SVG("Starship/Sequencer/SlomoTrack", Icon16x16));
+		Set("Sequencer.Tracks.TimeWarp", new IMAGE_BRUSH_SVG("Starship/Sequencer/TimeWarp_16", Icon16x16));
 		Set("Sequencer.Tracks.Sub", new IMAGE_BRUSH_SVG("Starship/Sequencer/SubTrack", Icon16x16));
 		Set("Sequencer.Tracks.LevelVisibility", new IMAGE_BRUSH_SVG("Starship/Sequencer/LevelVisibilityTrack", Icon16x16));
 		Set("Sequencer.Tracks.DataLayer", new IMAGE_BRUSH_SVG("Starship/Common/DataLayers", Icon16x16));
diff --git a/Engine/Source/Editor/MovieSceneTools/Private/Channels/BuiltInChannelEditors.cpp b/Engine/Source/Editor/MovieSceneTools/Private/Channels/BuiltInChannelEditors.cpp
--- a/Engine/Source/Editor/MovieSceneTools/Private/Channels/BuiltInChannelEditors.cpp
+++ b/Engine/Source/Editor/MovieSceneTools/Private/Channels/BuiltInChannelEditors.cpp
@@ -521,11 +521,12 @@ TSharedRef<SWidget> CreateKeyEditor(const TMovieSceneChannelHandle<FMovieSceneTi
 		Params.ObjectBindingID, Channel,
 		Params.OwningSection, Params.Sequencer, Params.PropertyBindings, ExternalValue
 		);
-
+	
+	KeyEditor.SetApplyInUnwarpedLocalSpace(true);
 	KeyEditor.SetOwningObject(Cast<UMovieSceneSignedObject>(Params.OwningObject));
 
 	const FMovieSceneTimeWarpChannel* ChannelPtr = Channel.Get();
-	if (ChannelPtr && ChannelPtr->Domain == ETimeWarpChannelDomain::Time)
+	if (ChannelPtr && ChannelPtr->Domain == UE::MovieScene::ETimeWarpChannelDomain::Time)
 	{
 		// Set the numeric type interface for frame numbers if the channel is in the time domain
 		KeyEditor.SetNumericTypeInterface(Params.Sequencer->GetNumericTypeInterface());
diff --git a/Engine/Source/Editor/MovieSceneTools/Private/Channels/TimeWarpChannelCurveModel.cpp b/Engine/Source/Editor/MovieSceneTools/Private/Channels/TimeWarpChannelCurveModel.cpp
--- a/Engine/Source/Editor/MovieSceneTools/Private/Channels/TimeWarpChannelCurveModel.cpp
+++ b/Engine/Source/Editor/MovieSceneTools/Private/Channels/TimeWarpChannelCurveModel.cpp
@@ -19,7 +19,7 @@ void FTimeWarpChannelCurveModel::SetCurveAttributes(const FCurveAttributes& InCu
 	const FMovieSceneTimeWarpChannel* Channel = static_cast<const FMovieSceneTimeWarpChannel*>(GetChannelHandle().Get());
 
 	// Disallow certain extrapolation modes for play rate curves.
-	if (Channel && Channel->Domain == ETimeWarpChannelDomain::PlayRate)
+	if (Channel && Channel->Domain == UE::MovieScene::ETimeWarpChannelDomain::PlayRate)
 	{
 		FCurveAttributes ExistingAttributes;
 		GetCurveAttributes(ExistingAttributes);
@@ -82,7 +82,7 @@ void FTimeWarpChannelCurveModel::AllocateAxes(FCurveEditor* InCurveEditor, TShar
 {
 	FMovieSceneTimeWarpChannel* Channel = static_cast<FMovieSceneTimeWarpChannel*>(GetChannelHandle().Get());
 
-	if (Channel && Channel->Domain == ETimeWarpChannelDomain::Time)
+	if (Channel && Channel->Domain == UE::MovieScene::ETimeWarpChannelDomain::Time)
 	{
 		OutVerticalAxis = InCurveEditor->FindAxis("FocusedSequenceTime");
 	}
diff --git a/Engine/Source/Editor/MovieSceneTools/Private/CurveKeyEditors/SequencerKeyEditor.h b/Engine/Source/Editor/MovieSceneTools/Private/CurveKeyEditors/SequencerKeyEditor.h
--- a/Engine/Source/Editor/MovieSceneTools/Private/CurveKeyEditors/SequencerKeyEditor.h
+++ b/Engine/Source/Editor/MovieSceneTools/Private/CurveKeyEditors/SequencerKeyEditor.h
@@ -112,7 +112,8 @@ struct TSequencerKeyEditor
 
 		if (Channel && ChannelMetaData && Sequencer && OwningSection)
 		{
-			const FFrameTime CurrentTime = UE::MovieScene::ClampToDiscreteRange(Sequencer->GetLocalTime().Time, OwningSection->GetRange()) - ChannelMetaData->GetOffsetTime(OwningSection);
+			FFrameTime LocalTime = GetCurrentTime();
+			const FFrameTime CurrentTime = UE::MovieScene::ClampToDiscreteRange(LocalTime, OwningSection->GetRange()) - ChannelMetaData->GetOffsetTime(OwningSection);
 
 			//If we have no keys and no default, key with the external value if it exists
 			if (!EvaluateChannel(OwningSection, Channel, CurrentTime, Result))
@@ -214,7 +215,8 @@ struct TSequencerKeyEditor
 		}
 		else
 		{
-			const FFrameNumber CurrentTime = Sequencer->GetLocalTime().Time.RoundToFrame() - ChannelMetaData->GetOffsetTime(OwningSection);
+			FFrameTime LocalTime = GetCurrentTime();
+			const FFrameNumber CurrentTime = LocalTime.RoundToFrame() - ChannelMetaData->GetOffsetTime(OwningSection);
 
 			EMovieSceneKeyInterpolation Interpolation = GetInterpolationMode(Channel, CurrentTime, Sequencer->GetKeyInterpolation());
 
@@ -274,6 +276,11 @@ struct TSequencerKeyEditor
 		}
 	}
 
+	void SetApplyInUnwarpedLocalSpace(bool bInApplyInUnwarpedLocalSpace)
+	{
+		bApplyInUnwarpedLocalSpace = bInApplyInUnwarpedLocalSpace;
+	}
+
 	const FGuid& GetObjectBindingID() const
 	{
 		return ObjectBindingID;
@@ -315,6 +322,18 @@ struct TSequencerKeyEditor
 		return FString();
 	}
 
+	FFrameTime GetCurrentTime() const
+	{
+		ISequencer* Sequencer = GetSequencer();
+		if (Sequencer)
+		{
+			// @todo: Really bApplyInUnwarpedLocalSpace should be looking for an ITimeDomainExtension on a view model, but we don't
+			//        have that information here because all these mechanisms pre-date the MVVM framework.
+			return bApplyInUnwarpedLocalSpace ? Sequencer->GetUnwarpedLocalTime().Time : Sequencer->GetLocalTime().Time;
+		}
+		return 0;
+	}
+
 private:
 
 	FGuid ObjectBindingID;
@@ -325,6 +344,7 @@ private:
 	TWeakPtr<FTrackInstancePropertyBindings> WeakPropertyBindings;
 	TFunction<TOptional<ValueType>(UObject&, FTrackInstancePropertyBindings*)> OnGetExternalValue;
 	TSharedPtr<INumericTypeInterface<ValueType>> NumericTypeInterface;
+	bool bApplyInUnwarpedLocalSpace = false;
 };
 
 
diff --git a/Engine/Source/Editor/MovieSceneTools/Private/MVVM/ViewModels/TimeWarpChannelModel.cpp b/Engine/Source/Editor/MovieSceneTools/Private/MVVM/ViewModels/TimeWarpChannelModel.cpp
--- a/Engine/Source/Editor/MovieSceneTools/Private/MVVM/ViewModels/TimeWarpChannelModel.cpp
+++ b/Engine/Source/Editor/MovieSceneTools/Private/MVVM/ViewModels/TimeWarpChannelModel.cpp
@@ -28,7 +28,7 @@ TSharedPtr<SWidget> FTimeWarpChannelModel::CreateOutlinerViewForColumn(const FCr
 	{
 		return MakeButton(
 			LOCTEXT("ChangeTimeWarpToolTip", "Change Time Warp to utilize a different curve type"),
-			FAppStyle::GetBrush("Sequencer.Tracks.Slomo"),
+			FAppStyle::GetBrush("Sequencer.Outliner.Decorators.TimeWarp"),
 			FOnGetContent::CreateSP(this, &FTimeWarpChannelModel::BuildReplaceTimeWarpSubMenu),
 			OutlinerItem.AsModel());
 	}
diff --git a/Engine/Source/Editor/MovieSceneTools/Private/MovieSceneToolsModule.cpp b/Engine/Source/Editor/MovieSceneTools/Private/MovieSceneToolsModule.cpp
--- a/Engine/Source/Editor/MovieSceneTools/Private/MovieSceneToolsModule.cpp
+++ b/Engine/Source/Editor/MovieSceneTools/Private/MovieSceneToolsModule.cpp
@@ -207,6 +207,7 @@ void FMovieSceneToolsModule::StartupModule()
 		CameraCutTrackModelHandle = SequencerModule.RegisterTrackModel(FOnCreateTrackModel::CreateStatic(&FCameraCutTrackModel::CreateTrackModel));
 		CinematicShotTrackModelHandle = SequencerModule.RegisterTrackModel(FOnCreateTrackModel::CreateStatic(&FCinematicShotTrackModel::CreateTrackModel));
 		BindingLifetimeTrackModelHandle = SequencerModule.RegisterTrackModel(FOnCreateTrackModel::CreateStatic(&FBindingLifetimeTrackModel::CreateTrackModel));
+		TimeWarpTrackModelHandle = SequencerModule.RegisterTrackModel(FOnCreateTrackModel::CreateStatic(&FTimeWarpTrackEditor::CreateTrackModel));
 
 		RegisterClipboardConversions();
 
@@ -395,6 +396,8 @@ void FMovieSceneToolsModule::ShutdownModule()
 	// unregister track models
 	SequencerModule.UnregisterTrackModel( CameraCutTrackModelHandle );
 	SequencerModule.UnregisterTrackModel( CinematicShotTrackModelHandle );
+	SequencerModule.UnregisterTrackModel( BindingLifetimeTrackModelHandle );
+	SequencerModule.UnregisterTrackModel( TimeWarpTrackModelHandle );
 
 	if (FModuleManager::Get().IsModuleLoaded("PropertyEditor"))
 	{	
diff --git a/Engine/Source/Editor/MovieSceneTools/Private/TrackEditors/TimeWarpTrackEditor.cpp b/Engine/Source/Editor/MovieSceneTools/Private/TrackEditors/TimeWarpTrackEditor.cpp
--- a/Engine/Source/Editor/MovieSceneTools/Private/TrackEditors/TimeWarpTrackEditor.cpp
+++ b/Engine/Source/Editor/MovieSceneTools/Private/TrackEditors/TimeWarpTrackEditor.cpp
@@ -3,92 +3,350 @@
 #include "TrackEditors/TimeWarpTrackEditor.h"
 #include "ISequencer.h"
 #include "ISequencerSection.h"
+#include "ISequencerEditTool.h"
 #include "SequencerSectionPainter.h"
 #include "SequencerUtilities.h"
+#include "SequencerSettings.h"
+#include "TimeSliderArgs.h"
 #include "Variants/MovieSceneTimeWarpGetter.h"
 #include "Framework/MultiBox/MultiBoxBuilder.h"
 
 #include "Sections/MovieSceneTimeWarpSection.h"
+#include "Channels/MovieSceneTimeWarpChannel.h"
 
+#include "MVVM/ViewModels/TrackModel.h"
 #include "MVVM/ViewModels/SectionModel.h"
 #include "MVVM/ViewModels/EditorSharedViewModelData.h"
 #include "MVVM/ViewModels/SequencerEditorViewModel.h"
+#include "MVVM/ViewModels/TrackAreaViewModel.h"
+#include "MVVM/ViewModels/VirtualTrackArea.h"
+#include "MVVM/Views/STrackAreaView.h"
+#include "MVVM/Views/ITrackAreaHotspot.h"
+
+#include "Widgets/SOverlay.h"
 
 #define LOCTEXT_NAMESPACE "TimeWarpTrackEditor"
 
 namespace UE::Sequencer
 {
 
-struct FTimeWarpSection : FSequencerSection
+struct FScrubberHotspot : ITrackAreaHotspot
 {
-	FTimeWarpSection(UMovieSceneSection& InSection)
-		: FSequencerSection(InSection)
+	TWeakPtr<ISequencer> WeakSequencer;
+
+	FScrubberHotspot(TWeakPtr<ISequencer> InWeakSequencer)
+		: WeakSequencer(InWeakSequencer)
 	{}
 
-	int32 OnPaintSection(FSequencerSectionPainter& InPainter) const override
+	virtual void UpdateOnHover(FTrackAreaViewModel& InTrackArea) const
 	{
-		TSharedPtr<FEditorSharedViewModelData> SharedEditorData = CastViewModel<FEditorSharedViewModelData>(InPainter.SectionModel->GetSharedData());
-		TSharedPtr<FSequencerEditorViewModel>  SequencerEditor  = SharedEditorData ? CastViewModel<FSequencerEditorViewModel>(SharedEditorData->GetEditor()) : nullptr;
-		if (!SequencerEditor)
+		InTrackArea.AttemptToActivateTool("Movement");
+	}
+
+	virtual TOptional<FFrameNumber> GetTime() const
+	{
+		TSharedPtr<ISequencer> Sequencer = WeakSequencer.Pin();
+		if (Sequencer)
 		{
-			return InPainter.LayerId;
+			return Sequencer->GetLocalTime().Time.FrameNumber;
 		}
+		return TOptional<FFrameNumber>();
+	}
 
-		InPainter.LayerId = InPainter.PaintSectionBackground();
-
-		TSharedPtr<ISequencer> Sequencer = SequencerEditor->GetSequencer();
-
-		UMovieSceneTimeWarpSection* TimeWarpSection = Cast<UMovieSceneTimeWarpSection>(InPainter.SectionModel->GetSection());
-		if (TimeWarpSection)
+	virtual TSharedPtr<ISequencerEditToolDragOperation> InitiateDrag(const FPointerEvent& MouseEvent)
+	{
+		struct FScrubLocalTime : ISequencerEditToolDragOperation
 		{
-			// Paint the unwarped current time
-			FFrameTime LocalTime = Sequencer->GetLocalTime().Time;
-			FFrameTime NewLocalTime = LocalTime;
-
-			FMovieSceneInverseNestedSequenceTransform Inverse = TimeWarpSection->GenerateTransform().Inverse();
+			TWeakPtr<ISequencer> WeakSequencer;
+			FScrubLocalTime(TWeakPtr<ISequencer> InWeakSequencer)
+				: WeakSequencer(InWeakSequencer)
+			{}
 
-			if (Inverse.IsLinear())
+			void OnBeginDrag(const FPointerEvent& MouseEvent, FVector2D LocalMousePos, const FVirtualTrackArea& VirtualTrackArea) override
 			{
-				NewLocalTime = LocalTime * Inverse.AsLinear();
+				if (TSharedPtr<ISequencer> Sequencer = WeakSequencer.Pin())
+				{
+					Sequencer->OnBeginScrubbing();
+				}
 			}
-			else
+			void OnDrag(const FPointerEvent& MouseEvent, FVector2D LocalMousePos, const FVirtualTrackArea& VirtualTrackArea) override
 			{
-				FMovieSceneSequenceTransform Transform = Sequencer->GetFocusedMovieSceneSequenceTransform();
-
-				// Time warp track transforms are always added last
-				if (Transform.NestedTransforms.Num() > 0)
+				FFrameTime ScrubTime = VirtualTrackArea.PixelToFrame(LocalMousePos.X);
+				if (TSharedPtr<ISequencer> Sequencer = WeakSequencer.Pin())
 				{
-					Transform.NestedTransforms.RemoveAt(Transform.NestedTransforms.Num()-1);
+					USequencerSettings* Settings = Sequencer->GetSequencerSettings();
+					if (Settings->GetIsSnapEnabled() || MouseEvent.IsShiftDown())
+					{
+						if (Settings->GetSnapPlayTimeToInterval())
+						{
+							// Set the style of the scrub handle
+							ScrubTime = FFrameRate::Snap(ScrubTime, Sequencer->GetFocusedTickResolution(), Sequencer->GetFocusedDisplayRate());
+						}
+
+						ENearestKeyOption NearestKeyOption = MouseEvent.IsShiftDown()
+							? ENearestKeyOption::NKO_SearchKeys | ENearestKeyOption::NKO_SearchSections | ENearestKeyOption::NKO_SearchMarkers
+							: ENearestKeyOption::NKO_None;
+
+						if (Settings->GetSnapPlayTimeToKeys())
+						{
+							EnumAddFlags(NearestKeyOption, ENearestKeyOption::NKO_SearchKeys);
+						}
+						if (Settings->GetSnapPlayTimeToSections())
+						{
+							EnumAddFlags(NearestKeyOption, ENearestKeyOption::NKO_SearchSections);
+						}
+						if (Settings->GetSnapPlayTimeToMarkers())
+						{
+							EnumAddFlags(NearestKeyOption, ENearestKeyOption::NKO_SearchMarkers);
+						}
+
+						FFrameNumber NearestKey = Sequencer->OnGetNearestKey(ScrubTime, NearestKeyOption);
+
+						static float MouseTolerance = 20.f;
+						if (FMath::IsNearlyEqual(VirtualTrackArea.FrameToPixel(NearestKey), LocalMousePos.X, MouseTolerance))
+						{
+							ScrubTime = NearestKey;
+						}
+					}
+
+					// @todo: Autoscroll goes wild when scrubbing warped time.
+					//        That is an intricate system that needs updating to handle warped times, but for now
+					//        we just hack it off when scrubbing.
+					if (Settings && Settings->GetAutoScrollEnabled())
+					{
+						Settings->SetAutoScrollEnabled(false);
+						Sequencer->OnScrubPositionChanged(ScrubTime, true, true);
+						Settings->SetAutoScrollEnabled(true);
+					}
+					else
+					{
+						Sequencer->OnScrubPositionChanged(ScrubTime, true, true);
+					}
 				}
-
-				NewLocalTime = Sequencer->GetGlobalTime().Time * Transform;
 			}
+			void OnEndDrag( const FPointerEvent& MouseEvent, FVector2D LocalMousePos, const FVirtualTrackArea& VirtualTrackArea) override
+			{
+				if (TSharedPtr<ISequencer> Sequencer = WeakSequencer.Pin())
+				{
+					Sequencer->OnEndScrubbing();
+				}
+			}
+			FCursorReply GetCursor() const override
+			{
+				return FCursorReply::Cursor(EMouseCursor::Default);
+			}
+			int32 OnPaint(const FGeometry& AllottedGeometry, const FSlateRect& MyCullingRect, FSlateWindowElementList& OutDrawElements, int32 LayerId) const override
+			{
+				return LayerId;
+			}
+		};
+
+		return MakeShared<FScrubLocalTime>(WeakSequencer);
+	}
+
+	virtual FCursorReply GetCursor() const
+	{
+		return FCursorReply::Cursor(EMouseCursor::Default);
+	}
 
-			if (LocalTime != NewLocalTime)
+	virtual int32 Priority() const
+	{
+		return 10000;
+	}
+};
+
+struct STimeWarpScrubber : public SLeafWidget
+{
+	SLATE_BEGIN_ARGS(STimeWarpScrubber){}
+	SLATE_END_ARGS()
+
+	static constexpr float HalfScrubberWidthPx = 7.f;
+	static constexpr float ScrubberWidthPx = HalfScrubberWidthPx * 2.f;
+
+	void Construct(const FArguments& InArgs, UMovieSceneTimeWarpSection* Section, TSharedPtr<FTimeToPixel> InTimeToPixel, TWeakPtr<STrackAreaView> InWeakTrackAreaView, TWeakPtr<ISequencer> InWeakSequencer)
+	{
+		WeakSection = Section;
+		TimeToPixel = InTimeToPixel;
+		WeakSequencer = InWeakSequencer;
+		WeakTrackAreaView = InWeakTrackAreaView;
+
+		SetVisibility(MakeAttributeSP(this, &STimeWarpScrubber::GetVisibility));
+	}
+
+	int32 OnPaint(const FPaintArgs& Args, const FGeometry& AllottedGeometry, const FSlateRect& MyCullingRect, FSlateWindowElementList& OutDrawElements, int32 LayerId, const FWidgetStyle& InWidgetStyle, bool bParentEnabled) const
+	{
+		using namespace UE::MovieScene;
+
+		UMovieSceneTimeWarpSection* TimeWarpSection = WeakSection.Get();
+		UMovieSceneTimeWarpTrack*   TimeWarpTrack   = TimeWarpSection ? TimeWarpSection->GetTypedOuter<UMovieSceneTimeWarpTrack>() : nullptr;
+
+		TSharedPtr<ISequencer> Sequencer = WeakSequencer.Pin();
+
+		if (Sequencer && TimeWarpTrack && TimeWarpTrack->bIsActiveTimeWarp)
+		{
+			FLinearColor TimeWarpColor = FStyleColors::AccentOrange.GetSpecifiedColor();
+			if (IsDirectlyHovered())
 			{
-				const ESlateDrawEffect DrawEffects = InPainter.bParentEnabled ? ESlateDrawEffect::None : ESlateDrawEffect::DisabledEffect;
-				float PixelPosition = InPainter.GetTimeConverter().SecondsToPixel(NewLocalTime / Sequencer->GetFocusedTickResolution());
-
-				FSlateDrawElement::MakeBox(
-					InPainter.DrawElements,
-					InPainter.LayerId++,
-					InPainter.SectionGeometry.ToPaintGeometry(
-						FVector2f(1.0f, InPainter.SectionGeometry.Size.Y),
-						FSlateLayoutTransform(FVector2f(PixelPosition, 0.f))
-					),
-					FAppStyle::GetBrush("WhiteBrush"),
-					DrawEffects,
-					FColor(255, 255, 255, 128)	// 0, 75, 50 (HSV)
-				);
+				FLinearColor HSV = TimeWarpColor.LinearRGBToHSV();
+				HSV.B = .6f;
+				HSV.G = .6f;
+				TimeWarpColor = HSV.HSVToLinearRGB();
 			}
+
+			const FSlateBrush* Brush = FAppStyle::GetBrush("Sequencer.Timeline.ScrubHandle");
+			FSlateDrawElement::MakeBox(
+				OutDrawElements,
+				LayerId++,
+				AllottedGeometry.ToPaintGeometry(),
+				Brush,
+				ESlateDrawEffect::None,
+				TimeWarpColor
+			);
+		}
+
+		return LayerId;
+	}
+
+	void OnMouseEnter(const FGeometry& MyGeometry, const FPointerEvent& MouseEvent) override
+	{
+		TSharedPtr<STrackAreaView> TrackAreaView = WeakTrackAreaView.Pin();
+		if (TrackAreaView)
+		{
+			TrackAreaView->GetViewModel()->SetHotspot(MakeShared<FScrubberHotspot>(WeakSequencer));
+		}
+	}
+
+	void OnMouseLeave(const FPointerEvent& MouseEvent) override
+	{
+		TSharedPtr<STrackAreaView> TrackAreaView = WeakTrackAreaView.Pin();
+		if (TrackAreaView)
+		{
+			TrackAreaView->GetViewModel()->SetHotspot(nullptr);
 		}
+	}
+
+	FReply OnMouseMove(const FGeometry& MyGeometry, const FPointerEvent& MouseEvent) override
+	{
+		// Hack to prevent the section from being able to handle this
+		return FReply::Handled();
+	}
 
-		return InPainter.LayerId;
+	FVector2D ComputeDesiredSize(float) const
+	{
+		return FVector2D(ScrubberWidthPx, 100.f);
 	}
+
+	EVisibility GetVisibility() const
+	{
+		TSharedPtr<ISequencer> Sequencer = WeakSequencer.Pin();
+		UMovieSceneTimeWarpSection* TimeWarpSection = WeakSection.Get();
+		UMovieSceneTimeWarpTrack*   TimeWarpTrack   = TimeWarpSection ? TimeWarpSection->GetTypedOuter<UMovieSceneTimeWarpTrack>() : nullptr;
+
+		const bool bIsVisible = (Sequencer && Sequencer->GetSequencerSettings()->GetTimeWarpDisplayMode() == ESequencerTimeWarpDisplay::Both)
+			&& TimeWarpTrack && TimeWarpTrack->bIsActiveTimeWarp;
+
+		return bIsVisible ? EVisibility::Visible : EVisibility::Collapsed;
+	}
+
+private:
+	TWeakObjectPtr<UMovieSceneTimeWarpSection> WeakSection;
+	TSharedPtr<FTimeToPixel> TimeToPixel;
+	TWeakPtr<STrackAreaView> WeakTrackAreaView;
+	TWeakPtr<ISequencer> WeakSequencer;
 };
 
+
+struct FTimeWarpSection : FSequencerSection
+{
+	TWeakPtr<ISequencer> WeakSequencer;
+
+	FTimeWarpSection(UMovieSceneSection& InSection, TSharedPtr<ISequencer> InSequencer)
+		: FSequencerSection(InSection)
+		, WeakSequencer(InSequencer)
+	{
+	}
+
+	void CreateViewWidgets(const FCreateSectionViewWidgetParams& Params) override
+	{
+		if (UMovieSceneTimeWarpSection* TimeWarpSection = Cast<UMovieSceneTimeWarpSection>(GetSectionObject()))
+		{
+			auto ScrubPosition = [WeakSequencer = this->WeakSequencer, TimeToPixel = Params.SectionView->GetTimeToPixel()]
+			{
+				FMargin Margin;
+				if (TSharedPtr<ISequencer> Sequencer = WeakSequencer.Pin())
+				{
+					Margin.Left = TimeToPixel->SecondsToPixel(Sequencer->GetLocalTime().AsSeconds()) - STimeWarpScrubber::HalfScrubberWidthPx;
+				}
+				return Margin;
+			};
+
+			// Add our widget above everything
+			Params.Overlay->AddSlot(FCreateSectionViewWidgetParams::ChannelViewOrder + 10)
+			.HAlign(HAlign_Left)
+			.Padding(MakeAttributeLambda(ScrubPosition))
+			[
+				SNew(STimeWarpScrubber, TimeWarpSection, Params.SectionView->GetTimeToPixel(), Params.TrackAreaView, WeakSequencer)
+			];
+		}
+	}
+};
+
+
+UE_SEQUENCER_DEFINE_CASTABLE(FTimeWarpTrackModel)
+UE_SEQUENCER_DEFINE_VIEW_MODEL_TYPE_ID(FTimeWarpTrackExtension)
+
+
+bool FTimeWarpTrackModel::IsActiveTimeWarp() const
+{
+	UMovieSceneTimeWarpTrack* Track = Cast<UMovieSceneTimeWarpTrack>(GetTrack());
+	return Track && Track->bIsActiveTimeWarp;
+}
+
+void FTimeWarpTrackModel::OnConstruct()
+{
+	FTrackModel::OnConstruct();
+
+	TViewModelPtr<FEditorSharedViewModelData> Shared = GetSharedData()->CastThisShared<FEditorSharedViewModelData>();
+	if (Shared)
+	{
+		FTimeWarpTrackExtension& TrackExtension = Shared->AddDynamicExtension<FTimeWarpTrackExtension>();
+		TrackExtension.WeakTimeWarpModels.Add(SharedThis(this));
+	}
+}
+
+const FTimeWarpTrackModel* FTimeWarpTrackExtension::GetActiveTimeWarpTrack() const
+{
+	for (TWeakViewModelPtr<FTimeWarpTrackModel> WeakTimeWarpTrack : WeakTimeWarpModels)
+	{
+		TViewModelPtr<FTimeWarpTrackModel> TimeWarpTrack = WeakTimeWarpTrack.Pin();
+		if (TimeWarpTrack && TimeWarpTrack->IsActiveTimeWarp())
+		{
+			return TimeWarpTrack.Get();
+		}
+	}
+	return nullptr;
+}
+
+
 } // namespace UE::Sequencer
 
+TSharedPtr<UE::Sequencer::FTrackModel> FTimeWarpTrackEditor::CreateTrackModel(UMovieSceneTrack* Track)
+{
+	using namespace UE::Sequencer;
+
+	if (UMovieSceneTimeWarpTrack* TimeWarpTrack = Cast<UMovieSceneTimeWarpTrack>(Track))
+	{
+		return MakeShared<FTimeWarpTrackModel>(TimeWarpTrack);
+	}
+	return nullptr;
+}
+
+void FTimeWarpTrackEditor::ProcessKeyOperation(FFrameNumber InKeyTime, const UE::Sequencer::FKeyOperation& Operation, ISequencer& InSequencer)
+{
+	//InKeyTime = InSequencer.GetLocalTimeWarpTransform().Inverse().TryTransformTime(InKeyTime).Get(InKeyTime).RoundToFrame();
+	Operation.ApplyDefault(InKeyTime, InSequencer);
+}
+
 void FTimeWarpTrackEditor::BuildAddTrackMenu(FMenuBuilder& MenuBuilder)
 {
 	auto HandleAddTimeWarp = [this](TSubclassOf<UMovieSceneTimeWarpGetter> InClass)
@@ -101,13 +359,13 @@ void FTimeWarpTrackEditor::BuildAddTrackMenu(FMenuBuilder& MenuBuilder)
 		LOCTEXT("AddTimeWarpTrackTooltip", "Adds a new track that manipulates the time of the current sequence."),
 		FNewMenuDelegate::CreateStatic(FSequencerUtilities::PopulateTimeWarpSubMenu, TFunction<void(TSubclassOf<UMovieSceneTimeWarpGetter>)>(HandleAddTimeWarp)),
 		false,
-		FSlateIcon(FAppStyle::GetAppStyleSetName(), "Sequencer.Tracks.Slomo")
+		FSlateIcon(FAppStyle::GetAppStyleSetName(), "Sequencer.Tracks.TimeWarp")
 	);
 }
 
 TSharedRef<ISequencerSection> FTimeWarpTrackEditor::MakeSectionInterface(UMovieSceneSection& SectionObject, UMovieSceneTrack& Track, FGuid ObjectBinding)
 {
-	return MakeShared<UE::Sequencer::FTimeWarpSection>(SectionObject);
+	return MakeShared<UE::Sequencer::FTimeWarpSection>(SectionObject, GetSequencer());
 }
 
 void FTimeWarpTrackEditor::HandleAddTimeWarpTrack(TSubclassOf<UMovieSceneTimeWarpGetter> ClassType)
diff --git a/Engine/Source/Editor/MovieSceneTools/Private/TrackEditors/TimeWarpTrackEditor.h b/Engine/Source/Editor/MovieSceneTools/Private/TrackEditors/TimeWarpTrackEditor.h
deleted file mode 100644
diff --git a/Engine/Source/Editor/MovieSceneTools/Public/MovieSceneToolsModule.h b/Engine/Source/Editor/MovieSceneTools/Public/MovieSceneToolsModule.h
--- a/Engine/Source/Editor/MovieSceneTools/Public/MovieSceneToolsModule.h
+++ b/Engine/Source/Editor/MovieSceneTools/Public/MovieSceneToolsModule.h
@@ -149,6 +149,7 @@ private:
 	FDelegateHandle CameraCutTrackModelHandle;
 	FDelegateHandle CinematicShotTrackModelHandle;
 	FDelegateHandle BindingLifetimeTrackModelHandle;
+	FDelegateHandle TimeWarpTrackModelHandle;
 
 	FDelegateHandle GenerateEventEntryPointsHandle;
 	FDelegateHandle FixupDynamicBindingPayloadParameterNameHandle;
diff --git a/Engine/Source/Editor/MovieSceneTools/Public/TrackEditors/TimeWarpTrackEditor.h b/Engine/Source/Editor/MovieSceneTools/Public/TrackEditors/TimeWarpTrackEditor.h
new file mode 100644
--- /dev/null
+++ b/Engine/Source/Editor/MovieSceneTools/Public/TrackEditors/TimeWarpTrackEditor.h
@@ -0,0 +1,79 @@
+// Copyright Epic Games, Inc. All Rights Reserved.
+
+#pragma once
+
+#include "KeyframeTrackEditor.h"
+#include "Tracks/MovieSceneTimeWarpTrack.h"
+#include "MVVM/ViewModels/TrackModel.h"
+#include "MVVM/Extensions/DynamicExtensionContainer.h"
+#include "MVVM/Extensions/ITimeDomainExtension.h"
+
+namespace UE::Sequencer
+{
+
+
+	struct FTimeWarpTrackModel : FTrackModel, ITimeDomainExtension
+	{
+		UE_SEQUENCER_DECLARE_CASTABLE_API(MOVIESCENETOOLS_API, FTimeWarpTrackModel, FTrackModel, ITimeDomainExtension)
+
+		FTimeWarpTrackModel(UMovieSceneTimeWarpTrack* TimeWarpTrack)
+			: FTrackModel(TimeWarpTrack)
+		{}
+
+		bool ShouldAnchorToTop() const override
+		{
+			return true;
+		}
+
+		MOVIESCENETOOLS_API bool IsActiveTimeWarp() const;
+
+		ETimeDomain GetDomain() const override
+		{
+			return ETimeDomain::Unwarped;
+		}
+
+		void OnConstruct() override;
+	};
+
+
+	class FTimeWarpTrackExtension : public IDynamicExtension
+	{
+	public:
+
+		UE_SEQUENCER_DECLARE_VIEW_MODEL_TYPE_ID_API(MOVIESCENETOOLS_API, FTimeWarpTrackExtension)
+
+		MOVIESCENETOOLS_API const FTimeWarpTrackModel* GetActiveTimeWarpTrack() const;
+
+		TArray<TWeakViewModelPtr<FTimeWarpTrackModel>> WeakTimeWarpModels;
+	};
+
+
+} // namespace UE::Sequencer
+
+class UMovieSceneTimeWarpGetter;
+
+template <typename T> class TSubclassOf;
+
+class FTimeWarpTrackEditor
+	: public FKeyframeTrackEditor<UMovieSceneTimeWarpTrack>
+{
+public:
+
+	FTimeWarpTrackEditor(TSharedRef<ISequencer> InSequencer)
+		: FKeyframeTrackEditor<UMovieSceneTimeWarpTrack>(InSequencer)
+	{}
+
+	static TSharedRef<ISequencerTrackEditor> CreateTrackEditor(TSharedRef<ISequencer> InSequencer)
+	{
+		return MakeShared<FTimeWarpTrackEditor>(InSequencer);
+	}
+
+	static TSharedPtr<UE::Sequencer::FTrackModel> CreateTrackModel(UMovieSceneTrack* Track);
+
+private:
+	virtual void BuildAddTrackMenu(FMenuBuilder& MenuBuilder) override;
+	virtual void ProcessKeyOperation(FFrameNumber InKeyTime, const UE::Sequencer::FKeyOperation& Operation, ISequencer& InSequencer) override;
+	virtual TSharedRef<ISequencerSection> MakeSectionInterface(UMovieSceneSection& SectionObject, UMovieSceneTrack& Track, FGuid ObjectBinding) override;
+
+	void HandleAddTimeWarpTrack(TSubclassOf<UMovieSceneTimeWarpGetter> ClassType);
+};
diff --git a/Engine/Source/Editor/Sequencer/Private/Filters/Filters/SequencerTrackFilters.h b/Engine/Source/Editor/Sequencer/Private/Filters/Filters/SequencerTrackFilters.h
--- a/Engine/Source/Editor/Sequencer/Private/Filters/Filters/SequencerTrackFilters.h
+++ b/Engine/Source/Editor/Sequencer/Private/Filters/Filters/SequencerTrackFilters.h
@@ -505,7 +505,7 @@ public:
 
 	//~ Begin FFilterBase
 	virtual FText GetDisplayName() const override { return LOCTEXT("SequencerTrackFilter_TimeWarp", "Time Warp"); }
-	virtual FSlateIcon GetIcon() const override { return FSlateIcon(FAppStyle::GetAppStyleSetName(), TEXT("Sequencer.Tracks.Slomo")); }
+	virtual FSlateIcon GetIcon() const override { return FSlateIcon(FAppStyle::GetAppStyleSetName(), TEXT("Sequencer.Tracks.TimeWarp")); }
 	//~ End FFilterBase
 
 	//~ Begin FSequencerTrackFilter
diff --git a/Engine/Source/Editor/Sequencer/Private/ISequencerSection.cpp b/Engine/Source/Editor/Sequencer/Private/ISequencerSection.cpp
--- a/Engine/Source/Editor/Sequencer/Private/ISequencerSection.cpp
+++ b/Engine/Source/Editor/Sequencer/Private/ISequencerSection.cpp
@@ -24,6 +24,7 @@
 #include "ISequencerModule.h"
 #include "Modules/ModuleManager.h"
 #include "MVVM/ViewModels/ViewDensity.h"
+#include "Widgets/SOverlay.h"
 
 struct FMovieSceneChannel;
 
@@ -201,6 +202,18 @@ void ISequencerSection::GenerateSectionLayout( ISectionLayoutBuilder& LayoutBuil
 	}
 }
 
+void ISequencerSection::CreateViewWidgets(const UE::Sequencer::FCreateSectionViewWidgetParams& Params)
+{
+	TSharedRef<SWidget> LegacyGeneratedWidget = GenerateSectionWidget();
+	if (LegacyGeneratedWidget != SNullWidget::NullWidget)
+	{
+		Params.Overlay->AddSlot(UE::Sequencer::FCreateSectionViewWidgetParams::DefaultWidgetOrder)
+		[
+			LegacyGeneratedWidget
+		];
+	}
+}
+
 float ISequencerSection::GetSectionHeight() const
 {
 	return SequencerSectionConstants::DefaultSectionHeight;
diff --git a/Engine/Source/Editor/Sequencer/Private/MVVM/SequencerViewModelTypeIDs.cpp b/Engine/Source/Editor/Sequencer/Private/MVVM/SequencerViewModelTypeIDs.cpp
--- a/Engine/Source/Editor/Sequencer/Private/MVVM/SequencerViewModelTypeIDs.cpp
+++ b/Engine/Source/Editor/Sequencer/Private/MVVM/SequencerViewModelTypeIDs.cpp
@@ -19,6 +19,7 @@
 #include "MVVM/Extensions/ISnappableExtension.h"
 #include "MVVM/Extensions/IBindingLifetimeExtension.h"
 #include "MVVM/Extensions/IConditionableExtension.h"
+#include "MVVM/Extensions/ITimeDomainExtension.h"
 
 #include "MVVM/ViewModels/SequencerOutlinerViewModel.h"
 #include "MVVM/ViewModels/SequencerEditorViewModel.h"
@@ -79,6 +80,7 @@ UE_SEQUENCER_DEFINE_VIEW_MODEL_TYPE_ID(FTrackRowModelStorageExtension);
 UE_SEQUENCER_DEFINE_VIEW_MODEL_TYPE_ID(ICurveEditorTreeItemExtension);
 UE_SEQUENCER_DEFINE_VIEW_MODEL_TYPE_ID(IObjectBindingExtension);
 UE_SEQUENCER_DEFINE_VIEW_MODEL_TYPE_ID(ISnappableExtension);
+UE_SEQUENCER_DEFINE_VIEW_MODEL_TYPE_ID(ITimeDomainExtension);
 UE_SEQUENCER_DEFINE_VIEW_MODEL_TYPE_ID(IBindingLifetimeExtension);
 UE_SEQUENCER_DEFINE_VIEW_MODEL_TYPE_ID(FOutlinerCacheExtension);
 UE_SEQUENCER_DEFINE_VIEW_MODEL_TYPE_ID(FConditionStateCacheExtension);
diff --git a/Engine/Source/Editor/Sequencer/Private/MVVM/Views/SSequencerKeyNavigationButtons.h b/Engine/Source/Editor/Sequencer/Private/MVVM/Views/SSequencerKeyNavigationButtons.h
--- a/Engine/Source/Editor/Sequencer/Private/MVVM/Views/SSequencerKeyNavigationButtons.h
+++ b/Engine/Source/Editor/Sequencer/Private/MVVM/Views/SSequencerKeyNavigationButtons.h
@@ -11,6 +11,7 @@
 #include "SequencerCommands.h"
 #include "SequencerCommonHelpers.h"
 #include "ScopedTransaction.h"
+#include "MVVM/Extensions/ITimeDomainExtension.h"
 
 #define LOCTEXT_NAMESPACE "SSequencerKeyNavigationButtons"
 
@@ -38,6 +39,9 @@ public:
 		WeakModel = InModel;
 		WeakSequencer = InSequencer;
 
+		ITimeDomainExtension* TimeDomain = InModel->CastThis<ITimeDomainExtension>();
+		Domain = TimeDomain ? TimeDomain->GetDomain() : ETimeDomain::Warped;
+
 		ChildSlot
 		[
 			SNew(SKeyNavigationButtons, InModel)
@@ -82,11 +86,17 @@ public:
 
 	FFrameTime GetTime() const
 	{
-		return WeakSequencer.Pin()->GetLocalTime().Time;
+		return Domain == ETimeDomain::Warped
+			? WeakSequencer.Pin()->GetLocalTime().Time
+			: WeakSequencer.Pin()->GetUnwarpedLocalTime().Time;
 	}
 
 	void SetTime(FFrameTime InTime)
 	{
+		if (Domain == ETimeDomain::Unwarped)
+		{
+			InTime = WeakSequencer.Pin()->GetLocalTimeWarpTransform().TransformTime(InTime);
+		}
 		WeakSequencer.Pin()->SetLocalTime(InTime);
 	}
 
@@ -102,6 +112,7 @@ public:
 
 	TWeakPtr<FViewModel> WeakModel;
 	TWeakPtr<ISequencer> WeakSequencer;
+	UE::Sequencer::ETimeDomain Domain;
 };
 
 } // namespace Sequencer
diff --git a/Engine/Source/Editor/Sequencer/Private/SSequencerSection.cpp b/Engine/Source/Editor/Sequencer/Private/SSequencerSection.cpp
--- a/Engine/Source/Editor/Sequencer/Private/SSequencerSection.cpp
+++ b/Engine/Source/Editor/Sequencer/Private/SSequencerSection.cpp
@@ -90,7 +90,7 @@ struct FSequencerSectionPainterImpl : FSequencerSectionPainter
 		: FSequencerSectionPainter(_OutDrawElements, InSectionGeometry, InSection)
 		, Sequencer(InSequencer)
 		, SectionWidget(InSectionWidget)
-		, TimeToPixelConverter(ConstructTimeConverterForSection(SectionGeometry, *InSection->GetSection(), Sequencer))
+		, TimeToPixelConverter(*InSectionWidget.GetTimeToPixel())
 		, TrackAreaViewModel(InTrackAreaViewModel)
 		, bClipRectEnabled(false)
 	{
@@ -854,6 +854,7 @@ void SSequencerSection::Construct( const FArguments& InArgs, TSharedPtr<FSequenc
 	WeakOwningTrackLane = OwningTrackLane;
 	SectionInterface = InSectionModel->GetSectionInterface();
 	HandleOffsetPx = 0.f;
+	TimeToPixel = MakeShared<FTimeToPixel>(100.f, TRange<double>(0.0, 1.0), FFrameRate());
 
 	SetEnabled(MakeAttributeSP(this, &SSequencerSection::IsEnabled));
 	SetToolTipText(MakeAttributeSP(this, &SSequencerSection::GetToolTipText));
@@ -867,33 +868,38 @@ void SSequencerSection::Construct( const FArguments& InArgs, TSharedPtr<FSequenc
 
 	UpdateUnderlappingSegments();
 
-	ChildSlot
-	.Padding(MakeAttributeSP(this, &SSequencerSection::GetHandleOffsetPadding))
-	[
+	TSharedRef<SOverlay> Overlay =
+
 		SNew(SOverlay)
 
-		+ SOverlay::Slot()
+		+ SOverlay::Slot(FCreateSectionViewWidgetParams::CompoundTrackLaneViewOrder)
 		[
 			SAssignNew(ChildLaneWidgets, SCompoundTrackLaneView)
 			.TimeToPixel(FGetTimeToPixel::CreateLambda([this](const FGeometry& AllottedGeometry){
-				// The given allotted geometry already has the handle padding removed, since it's inside
-				// our child slot (see padding above). We can therefore build the time converter directly.
-				UMovieSceneSection& Section = *SectionInterface->GetSectionObject();
-				return ConstructTimeConverterForSection(AllottedGeometry, Section, GetSequencer());
+				return *this->TimeToPixel;
 			}))
 		]
 
-		+ SOverlay::Slot()
-		[
-			SectionInterface->GenerateSectionWidget()
-		]
-
-		+ SOverlay::Slot()
+		+ SOverlay::Slot(FCreateSectionViewWidgetParams::ChannelViewOrder)
 		[
 			SNew(SChannelView, InSectionModel, OwningTrackLane->GetTrackAreaView())
 			.KeyBarColor(this, &SSequencerSection::GetTopLevelKeyBarColor)
 			.Visibility(this, &SSequencerSection::GetTopLevelChannelGroupVisibility)
-		]
+		];
+
+	FCreateSectionViewWidgetParams Params{
+		Overlay,
+		SharedThis(this),
+		OwningTrackLane.ToSharedRef(),
+		OwningTrackLane->GetTrackAreaView().ToSharedRef(),
+		InSectionModel.ToSharedRef()
+	};
+	SectionInterface->CreateViewWidgets(Params);
+
+	ChildSlot
+	.Padding(MakeAttributeSP(this, &SSequencerSection::GetHandleOffsetPadding))
+	[
+		Overlay
 	];
 }
 
@@ -1090,7 +1096,6 @@ bool SSequencerSection::CheckForEasingHandleInteraction( const FPointerEvent& Mo
 	}
 
 	FGeometry SectionGeometry = MakeSectionGeometryWithoutHandles(AllottedGeometry);
-	FTimeToPixel TimeToPixelConverter = ConstructTimeConverterForSection(SectionGeometry, *ThisSection, GetSequencer());
 
 	const double MousePositionX = SectionGeometry.AbsoluteToLocal(MouseEvent.GetScreenSpacePosition()).X;
 
@@ -1129,7 +1134,7 @@ bool SSequencerSection::CheckForEasingHandleInteraction( const FPointerEvent& Mo
 			// Compute the ease-in handle screen position.
 			const TRange<FFrameNumber> EaseInRange = EasingSectionObj->GetEaseInRange();
 			const FFrameNumber HandleTime = EaseInRange.IsEmpty() ? EasingSectionObj->GetInclusiveStartFrame() : EaseInRange.GetUpperBoundValue();
-			const double HandlePosition = TimeToPixelConverter.FrameToPixel(HandleTime) - HandleOffsetPx;
+			const double HandlePosition = TimeToPixel->FrameToPixel(HandleTime) - HandleOffsetPx;
 
 			if (FMath::IsNearlyEqual(MousePositionX, HandlePosition + GripHalfSizePx, GripHalfSizePx))
 			{
@@ -1143,7 +1148,7 @@ bool SSequencerSection::CheckForEasingHandleInteraction( const FPointerEvent& Mo
 			// Compute the ease-out handle screen position.
 			TRange<FFrameNumber> EaseOutRange = EasingSectionObj->GetEaseOutRange();
 			const FFrameNumber HandleTime = EaseOutRange.IsEmpty() ? EasingSectionObj->GetExclusiveEndFrame() : EaseOutRange.GetLowerBoundValue();
-			const double HandlePosition = TimeToPixelConverter.FrameToPixel(HandleTime) + HandleOffsetPx;
+			const double HandlePosition = TimeToPixel->FrameToPixel(HandleTime) + HandleOffsetPx;
 
 			if (FMath::IsNearlyEqual(MousePositionX, HandlePosition - GripHalfSizePx, GripHalfSizePx))
 			{
@@ -1192,7 +1197,6 @@ bool SSequencerSection::CheckForEdgeInteraction( const FPointerEvent& MouseEvent
 	}
 
 	FGeometry SectionGeometry = MakeSectionGeometryWithoutHandles(AllottedGeometry);
-	FTimeToPixel TimeToPixelConverter = ConstructTimeConverterForSection(SectionGeometry, *ThisSection, GetSequencer());
 
 	for (const TSharedPtr<FSectionModel>& UnderlappingSection : AllUnderlappingSections)
 	{
@@ -1211,7 +1215,7 @@ bool SSequencerSection::CheckForEdgeInteraction( const FPointerEvent& MouseEvent
 			// Make areas to the left and right of the geometry.  We will use these areas to determine if someone dragged the left or right edge of a section
 			FGeometry SectionRectLeft = SectionGeometry.MakeChild(
 				GripSize,
-				FSlateLayoutTransform(FVector2D( TimeToPixelConverter.FrameToPixel(UnderlappingSectionObj->GetInclusiveStartFrame()) - ThisHandleOffset, 0.f ))
+				FSlateLayoutTransform(FVector2D( TimeToPixel->FrameToPixel(UnderlappingSectionObj->GetInclusiveStartFrame()) - ThisHandleOffset, 0.f ))
 			);
 
 			if( SectionRectLeft.IsUnderLocation( MouseEvent.GetScreenSpacePosition() ) )
@@ -1225,7 +1229,7 @@ bool SSequencerSection::CheckForEdgeInteraction( const FPointerEvent& MouseEvent
 		{
 			FGeometry SectionRectRight = SectionGeometry.MakeChild(
 				GripSize,
-				FSlateLayoutTransform(FVector2D( TimeToPixelConverter.FrameToPixel(UnderlappingSectionObj->GetExclusiveEndFrame()) - UnderlappingSectionInterface->GetSectionGripSize() + ThisHandleOffset, 0 ))
+				FSlateLayoutTransform(FVector2D( TimeToPixel->FrameToPixel(UnderlappingSectionObj->GetExclusiveEndFrame()) - UnderlappingSectionInterface->GetSectionGripSize() + ThisHandleOffset, 0 ))
 			);
 
 			if( SectionRectRight.IsUnderLocation( MouseEvent.GetScreenSpacePosition() ) )
@@ -1257,9 +1261,7 @@ bool SSequencerSection::CheckForEasingAreaInteraction( const FPointerEvent& Mous
 	}
 
 	UMovieSceneSection* ThisSection = SectionInterface->GetSectionObject();
-	FGeometry SectionGeometry = MakeSectionGeometryWithoutHandles(AllottedGeometry);
-	FTimeToPixel TimeToPixelConverter = ConstructTimeConverterForSection(SectionGeometry, *ThisSection, GetSequencer());
-	const FFrameNumber MouseTime = TimeToPixelConverter.PixelToFrame(LocalMousePos.X).FrameNumber;
+	const FFrameNumber MouseTime = TimeToPixel->PixelToFrame(LocalMousePos.X).FrameNumber;
 
 	// First off, set the hotspot to an easing area if necessary
 	TSharedPtr<FTrackAreaViewModel> TrackAreaViewModel = GetTrackAreaViewModel();
@@ -1414,6 +1416,7 @@ int32 SSequencerSection::OnPaint( const FPaintArgs& Args, const FGeometry& Allot
 	const ESlateDrawEffect DrawEffects = bEnabled ? ESlateDrawEffect::None : ESlateDrawEffect::DisabledEffect;
 
 	FGeometry SectionGeometry = MakeSectionGeometryWithoutHandles(AllottedGeometry);
+	*TimeToPixel = ConstructTimeConverterForSection(SectionGeometry, *SectionObject, GetSequencer()); 
 
 	FSequencerSectionPainterImpl Painter(GetSequencer(), TrackAreaViewModel, SectionModel, OutDrawElements, SectionGeometry, *this);
 
@@ -1705,7 +1708,7 @@ void SSequencerSection::PaintEasingHandles( FSequencerSectionPainter& InPainter,
 			TRange<FFrameNumber> EaseInRange = UnderlappingSectionObj->GetEaseInRange();
 			const bool bHasEaseIn = !EaseInRange.IsEmpty();
 			FFrameNumber HandleFrame = bHasEaseIn ? UE::MovieScene::DiscreteExclusiveUpper(EaseInRange) : UnderlappingSectionObj->GetInclusiveStartFrame();
-			const float HandlePos(TimeToPixelConverter.FrameToPixel(HandleFrame) - HandleOffsetPx);
+			const float HandlePos(TimeToPixel->FrameToPixel(HandleFrame) - HandleOffsetPx);
 			FColor HandleColor = (bLeftHandleActive ? SelectionColor : InactiveHandleColor).ToFColorSRGB();
 
 			TArray<FSlateVertex> Verts;
@@ -1744,7 +1747,7 @@ void SSequencerSection::PaintEasingHandles( FSequencerSectionPainter& InPainter,
 			TRange<FFrameNumber> EaseOutRange = UnderlappingSectionObj->GetEaseOutRange();
 			const bool bHasEaseOut = !EaseOutRange.IsEmpty();
 			FFrameNumber HandleFrame = bHasEaseOut ? UE::MovieScene::DiscreteInclusiveLower(EaseOutRange) : UnderlappingSectionObj->GetExclusiveEndFrame();
-			const float HandlePos(TimeToPixelConverter.FrameToPixel(HandleFrame) + HandleOffsetPx);
+			const float HandlePos(TimeToPixel->FrameToPixel(HandleFrame) + HandleOffsetPx);
 			const FColor HandleColor = (bRightHandleActive ? SelectionColor : InactiveHandleColor).ToFColorSRGB();
 
 			TArray<FSlateVertex> Verts;
@@ -1824,7 +1827,6 @@ void SSequencerSection::DrawSectionHandles( const FGeometry& AllottedGeometry, F
 	}
 
 	FGeometry SectionGeometry = MakeSectionGeometryWithoutHandles(AllottedGeometry);
-	FTimeToPixel TimeToPixelConverter = ConstructTimeConverterForSection(SectionGeometry, *ThisSection, GetSequencer());
 
 	for (TSharedPtr<FSectionModel> SectionModel : AllUnderlappingSections)
 	{
@@ -1876,7 +1878,7 @@ void SSequencerSection::DrawSectionHandles( const FGeometry& AllottedGeometry, F
 		{
 			FGeometry SectionRectLeft = SectionGeometry.MakeChild(
 				GripSize,
-				FSlateLayoutTransform(FVector2D( TimeToPixelConverter.FrameToPixel(UnderlappingSectionObj->GetInclusiveStartFrame()) - ThisHandleOffset, 0.f ))
+				FSlateLayoutTransform(FVector2D( TimeToPixel->FrameToPixel(UnderlappingSectionObj->GetInclusiveStartFrame()) - ThisHandleOffset, 0.f ))
 			);
 			FSlateDrawElement::MakeBox
 			(
@@ -1894,7 +1896,7 @@ void SSequencerSection::DrawSectionHandles( const FGeometry& AllottedGeometry, F
 		{
 			FGeometry SectionRectRight = SectionGeometry.MakeChild(
 				GripSize,
-				FSlateLayoutTransform(FVector2D( TimeToPixelConverter.FrameToPixel(UnderlappingSectionObj->GetExclusiveEndFrame()) - UnderlappingSection->GetSectionGripSize() + ThisHandleOffset, 0 ))
+				FSlateLayoutTransform(FVector2D( TimeToPixel->FrameToPixel(UnderlappingSectionObj->GetExclusiveEndFrame()) - UnderlappingSection->GetSectionGripSize() + ThisHandleOffset, 0 ))
 			);
 			FSlateDrawElement::MakeBox
 			(
@@ -1917,10 +1919,17 @@ void SSequencerSection::DrawSectionHandles( const FGeometry& AllottedGeometry, F
 
 void SSequencerSection::Tick( const FGeometry& AllottedGeometry, const double InCurrentTime, const float InDeltaTime )
 {
+	UMovieSceneSection*    Section      = SectionInterface->GetSectionObject();
+	TSharedPtr<FSequencer> SequencerPtr = Sequencer.Pin();
+
+	if (!Section || !SequencerPtr)
+	{
+		return;
+	}
+
 	if (GetVisibility() == EVisibility::Visible)
 	{
-		UMovieSceneSection* Section = SectionInterface->GetSectionObject();
-		if (Section && Section->HasStartFrame() && Section->HasEndFrame())
+		if (Section->HasStartFrame() && Section->HasEndFrame())
 		{
 			constexpr float MinSectionWidth = 1.f;
 
@@ -1938,8 +1947,13 @@ void SSequencerSection::Tick( const FGeometry& AllottedGeometry, const double In
 		}
 
 		FGeometry SectionGeometry = MakeSectionGeometryWithoutHandles(AllottedGeometry);
+		*TimeToPixel = ConstructTimeConverterForSection(SectionGeometry, *Section, *SequencerPtr);
 		SectionInterface->Tick(SectionGeometry, ParentGeometry, InCurrentTime, InDeltaTime);
 	}
+	else
+	{
+		*TimeToPixel = ConstructTimeConverterForSection(AllottedGeometry, *Section, *SequencerPtr);
+	}
 }
 
 void SSequencerSection::AddChildView(TSharedPtr<ITrackLaneWidget> ChildWidget, TWeakPtr<STrackLane> InWeakOwningLane)
@@ -1955,6 +1969,11 @@ void SSequencerSection::OnModifiedIndirectly(UMovieSceneSignedObject* Object)
 	}
 }
 
+TSharedRef<FTimeToPixel> SSequencerSection::GetTimeToPixel() const
+{
+	return TimeToPixel.ToSharedRef();
+}
+
 FReply SSequencerSection::OnMouseButtonDown(const FGeometry& MyGeometry, const FPointerEvent& MouseEvent)
 {
 	return FReply::Unhandled();
diff --git a/Engine/Source/Editor/Sequencer/Private/SSequencerSection.h b/Engine/Source/Editor/Sequencer/Private/SSequencerSection.h
--- a/Engine/Source/Editor/Sequencer/Private/SSequencerSection.h
+++ b/Engine/Source/Editor/Sequencer/Private/SSequencerSection.h
@@ -9,6 +9,7 @@
 #include "Widgets/SCompoundWidget.h"
 #include "Rendering/RenderingCommon.h"
 
+#include "ISequencerSection.h"
 #include "EventHandlers/ISignedObjectEventHandler.h"
 #include "MVVM/ViewModels/SectionModel.h"
 #include "MVVM/Extensions/ITrackLaneExtension.h"
@@ -29,7 +30,7 @@ class FSectionModel;
 class FTrackAreaViewModel;
 struct ITrackAreaHotspot;
 
-class SSequencerSection : public SCompoundWidget, public ITrackLaneWidget, public UE::MovieScene::ISignedObjectEventHandler
+class SSequencerSection : public SCompoundWidget, public ITrackLaneWidget, public UE::Sequencer::ISectionView, public UE::MovieScene::ISignedObjectEventHandler
 {
 public:
 	SLATE_BEGIN_ARGS( SSequencerSection )
@@ -99,6 +100,9 @@ private:
 	/*~ ISignedObjectEventHandler interface */
 	virtual void OnModifiedIndirectly(UMovieSceneSignedObject* Object) override;
 
+	/*~ ISectionView interface */
+	virtual TSharedRef<FTimeToPixel> GetTimeToPixel() const override;
+
 	/**
 	 * Paint the easing handles for this section
 	 */
@@ -182,6 +186,8 @@ private:
 	TSharedPtr<SCompoundTrackLaneView> ChildLaneWidgets;
 	/** The track lane that this widget is on */
 	TWeakPtr<STrackLane> WeakOwningTrackLane;
+	/** Shared time to local pixel space for this section */
+	TSharedPtr<FTimeToPixel> TimeToPixel;
 	/** Cached parent geometry to pass down to any section interfaces that need it during tick */
 	FGeometry ParentGeometry;
 	/** The end time for a throbbing animation for selected sections */
diff --git a/Engine/Source/Editor/Sequencer/Private/Sequencer.cpp b/Engine/Source/Editor/Sequencer/Private/Sequencer.cpp
--- a/Engine/Source/Editor/Sequencer/Private/Sequencer.cpp
+++ b/Engine/Source/Editor/Sequencer/Private/Sequencer.cpp
@@ -45,6 +45,7 @@
 #include "Editor.h"
 #include "BlueprintActionDatabase.h"
 #include "Channels/MovieSceneChannelProxy.h"
+#include "Channels/MovieSceneTimeWarpChannel.h"
 #include "MovieScenePossessable.h"
 #include "MovieScene.h"
 #include "Compilation/MovieSceneCompiledDataManager.h"
@@ -604,7 +605,7 @@ void FSequencer::InitSequencer(const FSequencerInitParams& InitParams
 		.OnMarkEndDrag(this, &FSequencer::OnMarkEndDrag)
 		.IsPlaybackRangeLocked( this, &FSequencer::IsPlaybackRangeLocked )
 		.OnTogglePlaybackRangeLocked( this, &FSequencer::TogglePlaybackRangeLocked )
-		.ScrubPosition( this, &FSequencer::GetLocalFrameTime )
+		.ScrubPosition( this, &FSequencer::GetScrubPosition )
 		.ScrubPositionText( this, &FSequencer::GetFrameTimeText )
 		.ScrubPositionParent( this, &FSequencer::GetScrubPositionParent)
 		.ScrubPositionParentChain( this, &FSequencer::GetScrubPositionParentChain)
@@ -767,6 +768,8 @@ FSequencer::FSequencer()
 	// Exposes the sequencer and curve editor command lists to subscribers from other systems
 	FInputBindingManager::Get().RegisterCommandList(FSequencerCommands::Get().GetContextName(), SequencerCommandBindings);
 	FInputBindingManager::Get().RegisterCommandList(FCurveEditorCommands::Get().GetContextName(), CurveEditorSharedBindings);
+
+	TimeOperationDomain = UE::Sequencer::ETimeDomain::Warped;
 }
 
 
@@ -881,6 +884,7 @@ TSharedPtr<ISequencerTrackEditor> FSequencer::GetTrackEditor(UMovieSceneTrack* I
 
 void FSequencer::Tick(float InDeltaTime)
 {
+	using namespace UE::MovieScene;
 	using namespace UE::Sequencer;
 
 	static bool bEnableRefCountCheck = true;
@@ -919,6 +923,11 @@ void FSequencer::Tick(float InDeltaTime)
 	{
 		if (CompiledDataManager->IsDirty(RootSequencePtr))
 		{
+			// Try and preserve the current local time across compilations.
+			//     When a modification results in a change of transform or time-warp,
+			//     this helps to stop the play head from jumping around
+			FFrameTime OldLocalTime = GetUnwarpedLocalTime().Time;
+
 			CompiledDataManager->Compile(RootSequencePtr);
 
 			// Reset to the root sequence if the focused sequence no longer exists. This can happen if either the subsequence has been deleted or the hierarchy has changed.
@@ -943,6 +952,9 @@ void FSequencer::Tick(float InDeltaTime)
 				}
 			}
 
+			FTimeDomainOverride TimeDomain = OverrideTimeDomain(ETimeDomain::Unwarped);
+			SetLocalTime(OldLocalTime, ESnapTimeMode::STM_None, false /* bEvaluate */);
+
 			SuppressAutoEvalSignature.Reset();
 
 			if (GetNumericTypeInterface()->GetOnSettingChanged() != nullptr)
@@ -1010,10 +1022,10 @@ void FSequencer::Tick(float InDeltaTime)
 		FFrameTime NewGlobalTime = TimeController->RequestCurrentTime(GlobalTime, PlaybackSpeed * Dilation, GetFocusedDisplayRate());
 
 		// Put the time into clamped local space
-		FFrameTime LocalTime;
 		FMovieSceneTransformBreadcrumbs Breadcrumbs;
-		CalculateLocalTimeClamped(NewGlobalTime, RootToLocalTransform, LocalTime, Breadcrumbs);
+		FFrameTime LocalTime = RootToUnwarpedLocalTransform.TransformTime(NewGlobalTime, FTransformTimeParams().HarvestBreadcrumbs(Breadcrumbs).IgnoreClamps());
 
+		FTimeDomainOverride DomainOverride = OverrideTimeDomain(ETimeDomain::Unwarped);
 		SetLocalTimeLooped(LocalTime, Breadcrumbs);
 
 		if (GetPlaybackStatus() == EMovieScenePlayerStatus::Playing)
@@ -1024,7 +1036,7 @@ void FSequencer::Tick(float InDeltaTime)
 	}
 	else
 	{
-		PlayPosition.Reset(GlobalTime.ConvertTo(PlayPosition.GetInputRate()));
+		PlayPosition.Reset(GetGlobalTime().ConvertTo(PlayPosition.GetInputRate()));
 	}
 
 	if (AutoScrubTarget.IsSet())
@@ -1059,7 +1071,8 @@ void FSequencer::Tick(float InDeltaTime)
 	if (PendingScrubPosition.IsSet())
 	{
 		const bool bEvaluate = false;
-		SetLocalTimeDirectly(PendingScrubPosition.GetValue(), bEvaluate);
+		FTimeDomainOverride DomainOverride = OverrideTimeDomain(PendingScrubPosition->Get<1>());
+		SetLocalTimeDirectly(PendingScrubPosition->Get<0>(), bEvaluate);
 		PendingScrubPosition.Reset();
 
 		bNeedsEvaluate = true;
@@ -1119,9 +1132,19 @@ UMovieSceneSequence* FSequencer::GetRootMovieSceneSequence() const
 	return RootSequence.Get();
 }
 
+FMovieSceneSequenceTransform FSequencer::GetLocalTimeWarpTransform() const
+{
+	return LocalToWarpedLocalTransform;
+}
+
+FMovieSceneSequenceTransform FSequencer::GetGlobalPlaybackWarpTransform() const
+{
+	return GlobalPlaybackWarpTransform;
+}
+
 FMovieSceneSequenceTransform FSequencer::GetFocusedMovieSceneSequenceTransform() const
 {
-	return RootToLocalTransform;
+	return RootToWarpedLocalTransform;
 }
 
 UMovieSceneSequence* FSequencer::GetFocusedMovieSceneSequence() const
@@ -1204,15 +1227,7 @@ void FSequencer::ResetToNewRootSequence(UMovieSceneSequence& NewSequence)
 	RootTemplateInstance.Initialize(NewSequence, *this, CompiledDataManager);
 	RootTemplateInstance.EnableGlobalPreAnimatedStateCapture();
 
-	const FMovieSceneSequenceHierarchy* Hierarchy = CompiledDataManager->FindHierarchy(RootTemplateInstance.GetCompiledDataID());
-	if (Hierarchy)
-	{
-		RootToLocalTransform = Hierarchy->GetRootTransform();
-	}
-	else
-	{
-		RootToLocalTransform = FMovieSceneSequenceTransform();
-	}
+	UpdateSubSequenceData();
 
 	CurrentTimeBreadcrumbs.Reset();
 
@@ -1220,7 +1235,7 @@ void FSequencer::ResetToNewRootSequence(UMovieSceneSequence& NewSequence)
 	SequencerWidget->ResetBreadcrumbs();
 
 	PlayPosition.Reset(ConvertFrameTime(GetPlaybackRange().GetLowerBoundValue(), GetRootTickResolution(), PlayPosition.GetInputRate()));
-	TimeController->Reset(FQualifiedFrameTime(PlayPosition.GetCurrentPosition(), GetRootTickResolution()));
+	TimeController->Reset(GetGlobalTime());
 
 	UpdateSequencerCustomizations(PreviousRootSequence);
 
@@ -1517,13 +1532,18 @@ void FSequencer::UpdateSubSequenceData()
 	SubSequenceRange = TRange<FFrameNumber>::Empty();
 
 	const FMovieSceneSequenceHierarchy* Hierarchy = CompiledDataManager->FindHierarchy(RootTemplateInstance.GetCompiledDataID());
-	if (Hierarchy)
+
+	// If the root transform for the whole sequence is in the Play Rate domain, we assign a GlobalPlaybackWarpTransform
+	//     that mutates the evaluation ranges as we evaluate. This is not necessary for Time domain transforms since that is handled
+	//     by the sequence updaters themselves
+	FMovieSceneSequenceTransform RootTransform = Hierarchy ? Hierarchy->GetRootTransform() : FMovieSceneSequenceTransform();;
+	if (Hierarchy && Hierarchy->GetRootTransform().FindFirstWarpDomain() == ETimeWarpChannelDomain::PlayRate)
 	{
-		RootToLocalTransform = RootTransform = Hierarchy->GetRootTransform();
+		GlobalPlaybackWarpTransform = Hierarchy->GetRootTransform();
 	}
 	else
 	{
-		RootToLocalTransform = RootTransform = FMovieSceneSequenceTransform();
+		GlobalPlaybackWarpTransform = FMovieSceneSequenceTransform();
 	}
 
 	// Find the parent sub section and set up the sub sequence range, if necessary
@@ -1531,35 +1551,47 @@ void FSequencer::UpdateSubSequenceData()
 	{
 		// Ensure in this case we also reset the CurrentTimeBreadcrumbs.
 		CurrentTimeBreadcrumbs.Reset();
+
+		// Reset everything first
+		RootToUnwarpedLocalTransform = RootToWarpedLocalTransform = LocalToWarpedLocalTransform = FMovieSceneSequenceTransform();
+
+		if (Hierarchy)
+		{
+			LocalToWarpedLocalTransform = Hierarchy->GetRootTransform();
+		}
+		RootToWarpedLocalTransform = LocalToWarpedLocalTransform;
 		return;
 	}
 
 	check(Hierarchy);
 
-	const FMovieSceneSubSequenceData*   SubSequenceData = Hierarchy->FindSubData(ActiveTemplateIDs.Top());
+	const FMovieSceneSubSequenceData* SubSequenceData = Hierarchy->FindSubData(ActiveTemplateIDs.Top());
 
 	if (SubSequenceData)
 	{
 		SubSequenceRange = SubSequenceData->PlayRange.Value;
-		RootToLocalTransform = SubSequenceData->RootToSequenceTransform;
+
+		RootToUnwarpedLocalTransform = SubSequenceData->RootToUnwarpedLocalTransform * RootTransform;
+		RootToWarpedLocalTransform   = SubSequenceData->RootToSequenceTransform * RootTransform;
+		LocalToWarpedLocalTransform  = SubSequenceData->LocalToWarpedLocalTransform;
 
 		const bool bIsScrubbing = GetPlaybackStatus() == EMovieScenePlayerStatus::Scrubbing;
 
+		FFrameTime CurrentTime = GetGlobalTime().Time;
+
 		CurrentTimeBreadcrumbs.Reset();
-		RootToLocalTransform.TransformTime(GetGlobalTime().Time, FTransformTimeParams().HarvestBreadcrumbs(CurrentTimeBreadcrumbs).IgnoreClamps());
+		RootToWarpedLocalTransform.TransformTime(CurrentTime, FTransformTimeParams().HarvestBreadcrumbs(CurrentTimeBreadcrumbs).IgnoreClamps());
 
-		// Inner play range
+		// Inner play range in unwarped space
 		UMovieSceneSequence* SubSequence   = SubSequenceData->GetSequence();
 		TRange<FFrameNumber> PlaybackRange = SubSequence->GetMovieScene()->GetPlaybackRange();
 
-		FMovieSceneInverseSequenceTransform LocalToRootTransform = RootToLocalTransform.Inverse();
+		FMovieSceneInverseSequenceTransform LocalToRootTransform = RootToUnwarpedLocalTransform.Inverse();
 		const int32 PlaybackSize = UE::MovieScene::DiscreteSize(PlaybackRange);
 
 		TOptional<FFrameTime> PlayStart = LocalToRootTransform.TryTransformTime(PlaybackRange.GetLowerBoundValue(), CurrentTimeBreadcrumbs, EInverseEvaluateFlags::Backwards | EInverseEvaluateFlags::Cycle);
 		TOptional<FFrameTime> PlayEnd   = LocalToRootTransform.TryTransformTime(PlaybackRange.GetUpperBoundValue(), CurrentTimeBreadcrumbs, EInverseEvaluateFlags::Forwards  | EInverseEvaluateFlags::Cycle);
 
-		FFrameTime CurrentTime = GetGlobalTime().Time;
-
 		TRange<FFrameTime> ValidRange = TRange<FFrameTime>::All();
 		if (PlayStart)
 		{
@@ -3242,13 +3274,12 @@ FQualifiedFrameTime FSequencer::GetLocalTime() const
 {
 	using namespace UE::MovieScene;
 
-	const FFrameRate FocusedResolution = GetFocusedTickResolution();
-	const FFrameTime CurrentPosition   = PlayPosition.GetCurrentPosition();
+	const FFrameRate    FocusedResolution = GetFocusedTickResolution();
+	FQualifiedFrameTime CurrentPosition   = GetGlobalTime();
 
-	FFrameTime RootTime  = ConvertFrameTime(CurrentPosition, PlayPosition.GetInputRate(), PlayPosition.GetOutputRate());
+	CurrentPosition.Time = RootToWarpedLocalTransform.TransformTime(CurrentPosition.Time, FTransformTimeParams().IgnoreClamps());
 
-	const FFrameTime LocalTime = RootToLocalTransform.TransformTime(RootTime, FTransformTimeParams().IgnoreClamps());
-	return FQualifiedFrameTime(LocalTime, FocusedResolution);
+	return CurrentPosition;
 }
 
 
@@ -3257,7 +3288,7 @@ TOptional<int32> FSequencer::GetLocalLoopIndex() const
 	using namespace UE::MovieScene;
 
 	TOptional<int32> LoopIndex;
-	RootToLocalTransform.TransformTime(GetGlobalTime().Time, FTransformTimeParams().TrackCycleCounts(&LoopIndex));
+	RootToUnwarpedLocalTransform.TransformTime(GetGlobalTime().Time, FTransformTimeParams().TrackCycleCounts(&LoopIndex));
 	return LoopIndex;
 }
 
@@ -3268,11 +3299,36 @@ FQualifiedFrameTime FSequencer::GetGlobalTime() const
 	return FQualifiedFrameTime(RootTime, PlayPosition.GetOutputRate());
 }
 
+FQualifiedFrameTime FSequencer::GetUnwarpedLocalTime() const
+{
+	using namespace UE::MovieScene;
+
+	const FFrameRate    FocusedResolution = GetFocusedTickResolution();
+	FQualifiedFrameTime CurrentPosition   = GetGlobalTime();
+
+	CurrentPosition.Time = RootToUnwarpedLocalTransform.TransformTime(CurrentPosition.Time, FTransformTimeParams().IgnoreClamps());
+	return CurrentPosition;
+}
+
+FFrameTime FSequencer::GetScrubPosition() const
+{
+	if (Settings->GetTimeWarpDisplayMode() == ESequencerTimeWarpDisplay::WarpedTime)
+	{
+		return GetLocalTime().Time;
+	}
+	return GetUnwarpedLocalTime().Time;
+}
+
 FFrameTime FSequencer::GetLastEvaluatedLocalTime() const 
 {
 	return LastEvaluatedLocalTime;
 }
 
+UE::Sequencer::FTimeDomainOverride FSequencer::OverrideTimeDomain(UE::Sequencer::ETimeDomain NewDomain)
+{
+	return UE::Sequencer::FTimeDomainOverride(&TimeOperationDomain, NewDomain);
+}
+
 void FSequencer::SetLocalTime( FFrameTime NewTime, ESnapTimeMode SnapTimeMode, bool bEvaluate)
 {
 	FFrameRate LocalResolution = GetFocusedTickResolution();
@@ -3319,24 +3375,16 @@ void FSequencer::SetLocalTime( FFrameTime NewTime, ESnapTimeMode SnapTimeMode, b
 void FSequencer::SetLocalTimeDirectly(FFrameTime NewTime, bool bEvaluate)
 {
 	using namespace UE::MovieScene;
-
-	if (ActiveTemplateIDs.Num() <= 1)
-	{
-		if (!RootToLocalTransform.IsLinear())
-		{
-			NewTime = RootToLocalTransform.Inverse().TryTransformTime(NewTime).Get(NewTime);
-		}
-		SetGlobalTime(NewTime, bEvaluate);
-		return;
-	}
-
-
-	FMovieSceneInverseSequenceTransform LocalToRootTransform = RootToLocalTransform.Inverse();
+	using namespace UE::Sequencer;
 
 	const FMovieSceneTransformBreadcrumbs& Breadcrumbs = GetPlaybackStatus() == EMovieScenePlayerStatus::Scrubbing
 		? ScrubStartBreadcrumbs
 		: CurrentTimeBreadcrumbs;
 
+	FMovieSceneInverseSequenceTransform LocalToRootTransform = TimeOperationDomain == ETimeDomain::Warped
+		? RootToWarpedLocalTransform.Inverse()
+		: RootToUnwarpedLocalTransform.Inverse();
+
 	// Transform the time to the root time-space
 	TOptional<FFrameTime> NewGlobalTime = LocalToRootTransform.TryTransformTime(NewTime + ScrubLinearOffset, Breadcrumbs,
 		EInverseEvaluateFlags::AnyDirection | EInverseEvaluateFlags::Cycle | EInverseEvaluateFlags::IgnoreClamps);
@@ -3344,6 +3392,7 @@ void FSequencer::SetLocalTimeDirectly(FFrameTime NewTime, bool bEvaluate)
 	// If we still didn't find a time there's nothing we can do
 	if (NewGlobalTime.IsSet())
 	{
+		FTimeDomainOverride DomainOverride = OverrideTimeDomain(ETimeDomain::Unwarped);
 		SetGlobalTime(NewGlobalTime.GetValue(), bEvaluate);
 	}
 }
@@ -3352,6 +3401,7 @@ void FSequencer::SetLocalTimeDirectly(FFrameTime NewTime, bool bEvaluate)
 void FSequencer::SetGlobalTime(FFrameTime NewTime, bool bEvaluate)
 {
 	using namespace UE::MovieScene;
+	using namespace UE::Sequencer;
 
 	NewTime = ConvertFrameTime(NewTime, GetRootTickResolution(), PlayPosition.GetInputRate());
 	if (PlayPosition.GetEvaluationType() == EMovieSceneEvaluationType::FrameLocked)
@@ -3364,7 +3414,7 @@ void FSequencer::SetGlobalTime(FFrameTime NewTime, bool bEvaluate)
 	if (PlayPosition.GetCurrentPosition() != NewTime)
 	{
 		// Make sure breadcrumbs are up to date
-		RootToLocalTransform.TransformTime(PlayPosition.GetCurrentPosition(), FTransformTimeParams().HarvestBreadcrumbs(CurrentTimeBreadcrumbs).IgnoreClamps());
+		RootToWarpedLocalTransform.TransformTime(PlayPosition.GetCurrentPosition(), FTransformTimeParams().HarvestBreadcrumbs(CurrentTimeBreadcrumbs).IgnoreClamps());
 
 		FMovieSceneEvaluationRange EvalRange = PlayPosition.JumpTo(NewTime);
 		if (bEvaluate)
@@ -3461,7 +3511,13 @@ void FSequencer::EvaluateInternal(FMovieSceneEvaluationRange InRange, bool bHasJ
 	using namespace UE::MovieScene;
 
 	// Ensure breadcrumbs are up to date
-	RootToLocalTransform.TransformTime(InRange.GetTime(), FTransformTimeParams().HarvestBreadcrumbs(CurrentTimeBreadcrumbs).IgnoreClamps());
+	RootToUnwarpedLocalTransform.TransformTime(InRange.GetTime(), FTransformTimeParams().HarvestBreadcrumbs(CurrentTimeBreadcrumbs).IgnoreClamps());
+
+	if (!GlobalPlaybackWarpTransform.IsIdentity())
+	{
+		TRange<FFrameTime> NewRange = GlobalPlaybackWarpTransform.ComputeTraversedHull(InRange.GetRange());
+		InRange = FMovieSceneEvaluationRange(NewRange, InRange.GetFrameRate(), InRange.GetDirection());
+	}
 
 	LastEvaluatedLocalTime = GetLocalTime().Time;
 
@@ -4798,8 +4854,10 @@ FReply FSequencer::OnPlay(bool bTogglePlay)
 
 		FFrameNumber MinInclusiveTime = UE::MovieScene::DiscreteInclusiveLower(RootTimeBounds);
 		FFrameNumber MaxInclusiveTime = UE::MovieScene::DiscreteExclusiveUpper(RootTimeBounds) - 1;
+		
+		FFrameTime GlobalTime = GlobalPlaybackWarpTransform.TransformTime(GetGlobalTime().Time);
 
-		if (GetGlobalTime().Time <= MinInclusiveTime || GetGlobalTime().Time >= MaxInclusiveTime)
+		if (GlobalTime <= MinInclusiveTime || GlobalTime >= MaxInclusiveTime)
 		{
 			SetGlobalTime(PlaybackSpeed > 0 ? MinInclusiveTime : MaxInclusiveTime);
 		}
@@ -4860,12 +4918,16 @@ FReply FSequencer::OnPlayBackward(bool bTogglePlay)
 
 FReply FSequencer::OnStepForward(FFrameNumber Increment)
 {
+	using namespace UE::Sequencer;
+
 	SetPlaybackStatus(EMovieScenePlayerStatus::Stepping);
 
 	FFrameRate          DisplayRate = GetFocusedDisplayRate();
-	FQualifiedFrameTime CurrentTime = GetLocalTime();
+	FQualifiedFrameTime CurrentTime = GetUnwarpedLocalTime();
+
+	FFrameTime NewPosition = FFrameRate::TransformTime((CurrentTime.ConvertTo(DisplayRate) + Increment).FloorToFrame(), DisplayRate, CurrentTime.Rate);
 
-	FFrameTime NewPosition = FFrameRate::TransformTime(CurrentTime.ConvertTo(DisplayRate).FloorToFrame() + Increment, DisplayRate, CurrentTime.Rate);
+	FTimeDomainOverride TimeDomain = OverrideTimeDomain(ETimeDomain::Unwarped);
 	SetLocalTime(NewPosition, ESnapTimeMode::STM_Interval);
 	return FReply::Handled();
 }
@@ -4873,13 +4935,16 @@ FReply FSequencer::OnStepForward(FFrameNumber Increment)
 
 FReply FSequencer::OnStepBackward(FFrameNumber Increment)
 {
+	using namespace UE::Sequencer;
+
 	SetPlaybackStatus(EMovieScenePlayerStatus::Stepping);
 
 	FFrameRate          DisplayRate = GetFocusedDisplayRate();
-	FQualifiedFrameTime CurrentTime = GetLocalTime();
+	FQualifiedFrameTime CurrentTime = GetUnwarpedLocalTime();
 
-	FFrameTime NewPosition = FFrameRate::TransformTime(CurrentTime.ConvertTo(DisplayRate).FloorToFrame() - Increment, DisplayRate, CurrentTime.Rate);
+	FFrameTime NewPosition = FFrameRate::TransformTime((CurrentTime.ConvertTo(DisplayRate) - Increment).FloorToFrame(), DisplayRate, CurrentTime.Rate);
 
+	FTimeDomainOverride TimeDomain = OverrideTimeDomain(ETimeDomain::Unwarped);
 	SetLocalTime(NewPosition, ESnapTimeMode::STM_Interval);
 	return FReply::Handled();
 }
@@ -5031,10 +5096,13 @@ ESequencerLoopMode FSequencer::GetLoopMode() const
 void FSequencer::SetLocalTimeLooped(FFrameTime NewLocalTime, const FMovieSceneTransformBreadcrumbs& Breadcrumbs)
 {
 	using namespace UE::MovieScene;
+	using namespace UE::Sequencer;
 
 	TOptional<EMovieScenePlayerStatus::Type> NewPlaybackStatus;
 
-	FMovieSceneInverseSequenceTransform LocalToRootTransform = ActiveTemplateIDs.Num() > 1 ? RootToLocalTransform.Inverse() : FMovieSceneInverseSequenceTransform();
+	FMovieSceneInverseSequenceTransform LocalToRootTransform = TimeOperationDomain == ETimeDomain::Unwarped
+		? RootToUnwarpedLocalTransform.Inverse()
+		: RootToWarpedLocalTransform.Inverse();
 
 	// Default to the CurrentTimeBreadcrumbs
 	const FMovieSceneTransformBreadcrumbs& BreadcrumbsToUse = Breadcrumbs.Num() == 0 ? CurrentTimeBreadcrumbs : Breadcrumbs;
@@ -5063,7 +5131,13 @@ void FSequencer::SetLocalTimeLooped(FFrameTime NewLocalTime, const FMovieSceneTr
 		FFrameNumber MinInclusiveTime = UE::MovieScene::DiscreteInclusiveLower(TimeBounds);
 		FFrameNumber MaxInclusiveTime = UE::MovieScene::DiscreteExclusiveUpper(TimeBounds) - 1;
 
-		if (NewLocalTime < MinInclusiveTime || NewLocalTime > MaxInclusiveTime)
+		FFrameTime LoopTime = NewLocalTime;
+		if (TimeOperationDomain == ETimeDomain::Unwarped && LocalToWarpedLocalTransform.FindFirstWarpDomain() == ETimeWarpChannelDomain::PlayRate)
+		{
+			LoopTime = LocalToWarpedLocalTransform.TransformTime(LoopTime);
+		}
+
+		if (LoopTime < MinInclusiveTime || LoopTime > MaxInclusiveTime)
 		{
 			NewGlobalTime = LocalToRootTransform.TryTransformTime((PlaybackSpeed > 0 ? MinInclusiveTime : MaxInclusiveTime), BreadcrumbsToUse);
 			if (NewGlobalTime)
@@ -5083,23 +5157,26 @@ void FSequencer::SetLocalTimeLooped(FFrameTime NewLocalTime, const FMovieSceneTr
 
 		if (NewGlobalTime)
 		{
+			FFrameTime NewWarpedGlobalTime = GlobalPlaybackWarpTransform.TransformTime(NewGlobalTime.GetValue());
+			FFrameTime OldWarpedGlobalTime = GlobalPlaybackWarpTransform.TransformTime(GetGlobalTime().Time);
+
 			TRange<FFrameNumber> RootTimeBounds = GetRootTimeBounds();
 			FFrameNumber MinInclusiveTime = UE::MovieScene::DiscreteInclusiveLower(RootTimeBounds);
 			FFrameNumber MaxInclusiveTime = UE::MovieScene::DiscreteExclusiveUpper(RootTimeBounds) - 1;
 			bool bReachedEnd = false;
 			if (PlaybackSpeed > 0)
 			{
-				bReachedEnd = GetGlobalTime().Time <= MaxInclusiveTime && NewGlobalTime.GetValue() >= MaxInclusiveTime;
+				bReachedEnd = OldWarpedGlobalTime <= MaxInclusiveTime && NewWarpedGlobalTime >= MaxInclusiveTime;
 			}
 			else
 			{
-				bReachedEnd = GetGlobalTime().Time >= MinInclusiveTime && NewGlobalTime.GetValue() <= MinInclusiveTime;
+				bReachedEnd = OldWarpedGlobalTime >= MinInclusiveTime && NewWarpedGlobalTime <= MinInclusiveTime;
 			}
 
 			// Stop if we hit the playback range end
 			if (bReachedEnd)
 			{
-				NewGlobalTime = PlaybackSpeed > 0 ? MaxInclusiveTime : MinInclusiveTime;
+				NewGlobalTime = GlobalPlaybackWarpTransform.Inverse().TryTransformTime(PlaybackSpeed > 0 ? MaxInclusiveTime : MinInclusiveTime);
 				NewPlaybackStatus = EMovieScenePlayerStatus::Stopped;
 			}
 		}
@@ -5110,15 +5187,15 @@ void FSequencer::SetLocalTimeLooped(FFrameTime NewLocalTime, const FMovieSceneTr
 		NewGlobalTime = LocalToRootTransform.TryTransformTime(NewLocalTime, BreadcrumbsToUse);
 	}
 
-	if (!NewGlobalTime && RootToLocalTransform.NestedTransforms.Num() > 0)
+	if (!NewGlobalTime && RootToUnwarpedLocalTransform.NestedTransforms.Num() > 0)
 	{
 		// If we couldn't transform the time, we try one last-ditch attempt to guess a cycling or looping
 		//     transformation based on the linear transformation between each of the breadcrumbs.
 		FMovieSceneTransformBreadcrumbs DenseBreadcrumbs;
-		FFrameTime CurrentLocalTime = RootToLocalTransform.TransformTime(GetGlobalTime().Time, FTransformTimeParams().HarvestBreadcrumbs(DenseBreadcrumbs));
+		FFrameTime CurrentLocalTime = RootToUnwarpedLocalTransform.TransformTime(GetGlobalTime().Time, FTransformTimeParams().HarvestBreadcrumbs(DenseBreadcrumbs));
 
 		FFrameTime LocalDeltaGuess   = NewLocalTime - CurrentLocalTime;
-		FFrameTime NewLocalTimeGuess = RootToLocalTransform.NestedTransforms.Last().TransformTime(LocalDeltaGuess);
+		FFrameTime NewLocalTimeGuess = RootToUnwarpedLocalTransform.NestedTransforms.Last().TransformTime(LocalDeltaGuess);
 
 		NewGlobalTime = LocalToRootTransform.TryTransformTime(NewLocalTimeGuess, DenseBreadcrumbs);
 	}
@@ -5133,7 +5210,7 @@ void FSequencer::SetLocalTimeLooped(FFrameTime NewLocalTime, const FMovieSceneTr
 	UMovieScene* MovieScene = GetFocusedMovieSceneSequence()->GetMovieScene();
 	if (IsAutoScrollEnabled())
 	{
-		ScrollIntoView((NewGlobalTime.GetValue() * RootToLocalTransform) / RootTickResolution);
+		ScrollIntoView((NewGlobalTime.GetValue() * RootToUnwarpedLocalTransform) / RootTickResolution);
 	}
 
 	FFrameTime NewPlayPosition = ConvertFrameTime(NewGlobalTime.GetValue(), RootTickResolution, PlayPosition.GetInputRate());
@@ -5142,11 +5219,11 @@ void FSequencer::SetLocalTimeLooped(FFrameTime NewLocalTime, const FMovieSceneTr
 	if (bResetPosition)
 	{
 		PlayPosition.Reset(NewPlayPosition);
-		TimeController->Reset(FQualifiedFrameTime(NewGlobalTime.GetValue(), RootTickResolution));
+		TimeController->Reset(FQualifiedFrameTime(NewPlayPosition, RootTickResolution));
 	}
 
 	// Ensure breadcrumbs are up to date
-	RootToLocalTransform.TransformTime(NewPlayPosition, FTransformTimeParams().HarvestBreadcrumbs(CurrentTimeBreadcrumbs).IgnoreClamps());
+	RootToUnwarpedLocalTransform.TransformTime(NewGlobalTime.GetValue(), FTransformTimeParams().HarvestBreadcrumbs(CurrentTimeBreadcrumbs).IgnoreClamps());
 
 	// Evaluate the sequence
 	FMovieSceneEvaluationRange EvalRange = PlayPosition.PlayTo(NewPlayPosition);
@@ -5184,20 +5261,6 @@ void FSequencer::RestorePlaybackSpeedAfterPlay()
 	PlaybackSpeed = PlaybackSpeedBeforePlay;
 }
 
-void FSequencer::CalculateLocalTimeClamped(FFrameTime RootTime, const FMovieSceneSequenceTransform& RootToParentChainTransform, FFrameTime& OutLocalTime, FMovieSceneTransformBreadcrumbs& OutBreadcrumbs) const
-{
-	using namespace UE::MovieScene;
-
-	if (ActiveTemplateIDs.Num() > 1)
-	{
-		OutLocalTime = RootToLocalTransform.TransformTime(RootTime, FTransformTimeParams().HarvestBreadcrumbs(OutBreadcrumbs).IgnoreClamps());
-	}
-	else
-	{
-		OutLocalTime = RootTime;
-	}
-}
-
 EPlaybackMode::Type FSequencer::GetPlaybackMode() const
 {
 	if (PlaybackState == EMovieScenePlayerStatus::Playing)
@@ -5286,17 +5349,7 @@ TRange<FFrameNumber> FSequencer::GetRootTimeBounds() const
 		return TRange<FFrameNumber>(-100000, 100000);
 	}
 
-	TRange<FFrameNumber> RootTimeBounds = RootMovieSceneSequence->GetMovieScene()->GetPlaybackRange();
-
-	TOptional<FFrameTime> Lower = RootTransform.Inverse().TryTransformTime(RootTimeBounds.GetLowerBoundValue());
-	TOptional<FFrameTime> Upper = RootTransform.Inverse().TryTransformTime(RootTimeBounds.GetUpperBoundValue());
-
-	if (Lower && Upper)
-	{
-		return TRange<FFrameNumber>(Lower->FrameNumber, Upper->FrameNumber);
-	}
-
-	return RootTimeBounds;
+	return RootMovieSceneSequence->GetMovieScene()->GetPlaybackRange();
 }
 
 void FSequencer::RefreshSupportedCustomBindingTypes()
@@ -5528,7 +5581,7 @@ void FSequencer::OnScrubPositionChanged( FFrameTime NewScrubPosition, bool bScru
 	{
 		// Evaluation can be expensive and we may receive multiple OnScrubPositionChanged events in
 		// a frame, so defer the evaluation until the next tick
-		PendingScrubPosition = NewScrubPosition;
+		PendingScrubPosition = MakeTuple(NewScrubPosition, TimeOperationDomain);
 	}
 	else
 	{
@@ -5607,37 +5660,32 @@ void FSequencer::OnMarkEndDrag()
 
 FString FSequencer::GetFrameTimeText() const
 {
-	FMovieSceneSequenceTransform RootToParentChainTransform = RootToLocalTransform;
+	using namespace UE::MovieScene;
 
-	if (ScrubPositionParent.IsSet())
+	const bool bWarpedTime = Settings->GetTimeWarpDisplayMode() == ESequencerTimeWarpDisplay::WarpedTime;
+
+	const FMovieSceneSequenceTransform* RootToParentChainTransform = bWarpedTime
+		? &RootToWarpedLocalTransform
+		: &RootToUnwarpedLocalTransform;
+
+	if (ScrubPositionParent.IsSet() && ScrubPositionParent.GetValue() != MovieSceneSequenceID::Root)
 	{
-		if (ScrubPositionParent.GetValue() == MovieSceneSequenceID::Root)
-		{
-			RootToParentChainTransform = FMovieSceneSequenceTransform();
-		}
-		else if (const FMovieSceneSequenceHierarchy* Hierarchy = CompiledDataManager->FindHierarchy(RootTemplateInstance.GetCompiledDataID()))
+		if (const FMovieSceneSequenceHierarchy* Hierarchy = CompiledDataManager->FindHierarchy(RootTemplateInstance.GetCompiledDataID()))
 		{
 			for (const TTuple<FMovieSceneSequenceID, FMovieSceneSubSequenceData>& Pair : Hierarchy->AllSubSequenceData())
 			{
 				if (Pair.Key == ScrubPositionParent.GetValue())
 				{
-					RootToParentChainTransform = Pair.Value.RootToSequenceTransform;
+					RootToParentChainTransform = bWarpedTime 
+						? &Pair.Value.RootToSequenceTransform
+						: &Pair.Value.RootToUnwarpedLocalTransform;
 					break;
 				}
 			}
 		}
 	}
 
-	const FFrameRate FocusedResolution = GetFocusedTickResolution();
-	const FFrameTime CurrentPosition   = PlayPosition.GetCurrentPosition();
-
-	const FFrameTime RootTime = ConvertFrameTime(CurrentPosition, PlayPosition.GetInputRate(), PlayPosition.GetOutputRate());
-	
-	FFrameTime LocalTime;
-	FMovieSceneTransformBreadcrumbs LoopCounter;
-
-	CalculateLocalTimeClamped(RootTime, RootToParentChainTransform, LocalTime, LoopCounter);
-
+	FFrameTime LocalTime = RootToParentChainTransform->TransformTime(GetGlobalTime().Time, FTransformTimeParams().IgnoreClamps());
 	return GetNumericTypeInterface()->ToString(LocalTime.GetFrame().Value);
 }
 
@@ -9406,11 +9454,11 @@ void FSequencer::Pause()
 	// When stopping a sequence, we always evaluate a non-empty range if possible. This ensures accurate paused motion blur effects.
 	if (Settings->GetIsSnapEnabled())
 	{
-		FQualifiedFrameTime LocalTime          = GetLocalTime();
+		FQualifiedFrameTime LocalTime          = GetUnwarpedLocalTime();
 		FFrameRate          FocusedDisplayRate = GetFocusedDisplayRate();
 
 		// Snap to the focused play rate
-		TOptional<FFrameTime> RootPosition  = RootToLocalTransform.Inverse().TryTransformTime(
+		TOptional<FFrameTime> RootPosition  = RootToUnwarpedLocalTransform.Inverse().TryTransformTime(
 			FFrameRate::Snap(LocalTime.Time, LocalTime.Rate, FocusedDisplayRate), CurrentTimeBreadcrumbs);
 
 		if (RootPosition)
diff --git a/Engine/Source/Editor/Sequencer/Private/Sequencer.h b/Engine/Source/Editor/Sequencer/Private/Sequencer.h
--- a/Engine/Source/Editor/Sequencer/Private/Sequencer.h
+++ b/Engine/Source/Editor/Sequencer/Private/Sequencer.h
@@ -11,6 +11,7 @@
 #include "UObject/WeakInterfacePtr.h"
 #include "Widgets/SWidget.h"
 #include "SequencerNodeTree.h"
+#include "SequencerTimeDomainOverride.h"
 #include "UObject/GCObject.h"
 #include "MovieSceneMarkedFrame.h"
 #include "MovieSceneSequenceID.h"
@@ -713,6 +714,8 @@ public:
 	virtual UMovieSceneSequence* GetRootMovieSceneSequence() const override;
 	virtual UMovieSceneSequence* GetFocusedMovieSceneSequence() const override;
 	virtual FMovieSceneSequenceTransform GetFocusedMovieSceneSequenceTransform() const override;
+	virtual FMovieSceneSequenceTransform GetLocalTimeWarpTransform() const override;
+	virtual FMovieSceneSequenceTransform GetGlobalPlaybackWarpTransform() const override;
 	virtual FMovieSceneRootEvaluationTemplateInstance& GetEvaluationTemplate() override { return RootTemplateInstance; }
 	virtual void ResetToNewRootSequence(UMovieSceneSequence& NewSequence) override;
 	virtual void FocusSequenceInstance(UMovieSceneSubSection& InSubSection) override;
@@ -732,7 +735,9 @@ public:
 	virtual bool GetAutoSetTrackDefaults() const override;
 	virtual FQualifiedFrameTime GetLocalTime() const override;
 	virtual FQualifiedFrameTime GetGlobalTime() const override;
+	virtual FQualifiedFrameTime GetUnwarpedLocalTime() const override;
 	virtual TOptional<int32> GetLocalLoopIndex() const override;
+	virtual UE::Sequencer::FTimeDomainOverride OverrideTimeDomain(UE::Sequencer::ETimeDomain NewDomain) override;
 	virtual void SetLocalTime(FFrameTime Time, ESnapTimeMode SnapTimeMode = ESnapTimeMode::STM_None, bool bEvaluate = true) override;
 	virtual void SetLocalTimeDirectly(FFrameTime NewTime, bool bEvaluate = true) override;
 	virtual FFrameTime GetLastEvaluatedLocalTime() const override;
@@ -913,7 +918,7 @@ protected:
 	 * @param NearestKeyOption (ie. search keys/markers/all tracks)
 	 * @return NearestKey
 	 */
-	FFrameNumber OnGetNearestKey(FFrameTime InTime, ENearestKeyOption NearestKeyOption);
+	virtual FFrameNumber OnGetNearestKey(FFrameTime InTime, ENearestKeyOption NearestKeyOption) override;
 
 	/**
 	 * Called when the scrub position is changed by the user
@@ -923,13 +928,12 @@ protected:
 	 * @param bScrubbing If scrubbing
 	 * @param bEvaluate  Do evaluate sequencer after changing time
 	 */
-	void OnScrubPositionChanged( FFrameTime NewScrubPosition, bool bScrubbing, bool bEvaluate);
+	void OnScrubPositionChanged( FFrameTime NewScrubPosition, bool bScrubbing, bool bEvaluate) override;
 
-	/** Called when the user has begun scrubbing */
-	void OnBeginScrubbing();
+	void OnBeginScrubbing() override;
 
 	/** Called when the user has finished scrubbing */
-	void OnEndScrubbing();
+	void OnEndScrubbing() override;
 
 	/** Called when the user has begun dragging the playback range */
 	void OnPlaybackRangeBeginDrag();
@@ -949,8 +953,8 @@ protected:
 	/** Called when the user has finished dragging a mark */
 	void OnMarkEndDrag();
 
-	/** Get the unqualified local time */
-	FFrameTime GetLocalFrameTime() const { return GetLocalTime().Time; }
+	/** Get the current scrub position based on settings */
+	FFrameTime GetScrubPosition() const;
 
 	/** Get the frame time text */
 	FString GetFrameTimeText() const;
@@ -1162,8 +1166,6 @@ private:
 	void RestorePlaybackSpeedAfterPlay();
 
 	FGuid FindUnspawnedObjectGuid(UObject& InObject);
-	// Given the root sequence time, returns the local time and loop counter clamped to the maximum number of loops
-	void CalculateLocalTimeClamped(FFrameTime RootTime, const FMovieSceneSequenceTransform& RootToParentChainTransform, FFrameTime& OutTime, FMovieSceneTransformBreadcrumbs& OutBreadcrumbs) const;
 
 public:
 
@@ -1322,8 +1324,10 @@ private:
 	TArray<bool> ActiveTemplateStates;
 
 	/** Time transformation from the root sequence to the currently edited sequence. */
-	FMovieSceneSequenceTransform RootToLocalTransform;
-	FMovieSceneSequenceTransform RootTransform;
+	FMovieSceneSequenceTransform GlobalPlaybackWarpTransform;
+	FMovieSceneSequenceTransform RootToUnwarpedLocalTransform;
+	FMovieSceneSequenceTransform RootToWarpedLocalTransform;
+	FMovieSceneSequenceTransform LocalToWarpedLocalTransform;
 
 	/** Breadcrumbs to the current local time from the root sequence. */
 	FMovieSceneTransformBreadcrumbs CurrentTimeBreadcrumbs;
@@ -1367,7 +1371,7 @@ private:
 	TOptional<FAutoScrubTarget> AutoScrubTarget;
 
 	/** The latest scrub position request, to be applied in the next tick */
-	TOptional<FFrameTime> PendingScrubPosition;
+	TOptional<TTuple<FFrameTime, UE::Sequencer::ETimeDomain>> PendingScrubPosition;
 
 	/** Zoom smoothing curves */
 	FCurveSequence ZoomAnimation;
@@ -1520,6 +1524,8 @@ private:
 	/** Event contexts retrieved from the above attribute once per frame */
 	TArray<TWeakObjectPtr<UObject>> CachedEventContexts;
 
+	UE::Sequencer::ETimeDomain TimeOperationDomain;
+
 	/** When true, sequence will be forcefully evaluated on the next tick */
 	bool bNeedsEvaluate;
 
diff --git a/Engine/Source/Editor/Sequencer/Private/SequencerHotspots.cpp b/Engine/Source/Editor/Sequencer/Private/SequencerHotspots.cpp
--- a/Engine/Source/Editor/Sequencer/Private/SequencerHotspots.cpp
+++ b/Engine/Source/Editor/Sequencer/Private/SequencerHotspots.cpp
@@ -2,6 +2,7 @@
 
 #include "SequencerHotspots.h"
 #include "MVVM/Extensions/IObjectBindingExtension.h"
+#include "MVVM/Extensions/ITimeDomainExtension.h"
 #include "MVVM/ViewModels/SequencerEditorViewModel.h"
 #include "MVVM/Selection/Selection.h"
 #include "SequencerCommonHelpers.h"
@@ -217,6 +218,40 @@ void FKeyHotspot::UpdateOnHover(FTrackAreaViewModel& InTrackArea) const
 	InTrackArea.AttemptToActivateTool(FSequencerEditTool_Movement::Identifier);
 }
 
+TOptional<ETimeDomain> FKeyHotspot::GetDomain() const
+{
+	TOptional<ETimeDomain> Domain;
+
+	TSet<const FChannelModel*> VisitedChannels;
+	for (const FSequencerSelectedKey& Key : Keys)
+	{
+		TSharedPtr<FChannelModel> Channel = Key.WeakChannel.Pin();
+		if (!Channel || VisitedChannels.Contains(Channel.Get()))
+		{
+			continue;
+		}
+
+		VisitedChannels.Add(Channel.Get());
+
+		TViewModelPtr<ITimeDomainExtension> TimeDomain = Channel->FindAncestorOfType<ITimeDomainExtension>(true);
+		if (TimeDomain)
+		{
+			if (Domain && Domain != TimeDomain->GetDomain())
+			{
+				return TOptional<ETimeDomain>();
+			}
+
+			Domain = TimeDomain->GetDomain();
+		}
+	}
+
+	if (!Domain)
+	{
+		Domain = ETimeDomain::Warped;
+	}
+	return Domain;
+}
+
 TOptional<FFrameNumber> FKeyHotspot::GetTime() const
 {
 	FFrameNumber Time = 0;
diff --git a/Engine/Source/Editor/Sequencer/Private/SequencerHotspots.h b/Engine/Source/Editor/Sequencer/Private/SequencerHotspots.h
--- a/Engine/Source/Editor/Sequencer/Private/SequencerHotspots.h
+++ b/Engine/Source/Editor/Sequencer/Private/SequencerHotspots.h
@@ -70,6 +70,7 @@ struct FKeyHotspot
 	SEQUENCER_API FKeyHotspot(const TArray<FSequencerSelectedKey>& InKeys, TWeakPtr<FSequencer> InWeakSequencer);
 
 	virtual void UpdateOnHover(FTrackAreaViewModel& InTrackArea) const override;
+	virtual TOptional<ETimeDomain> GetDomain() const;
 	virtual TOptional<FFrameNumber> GetTime() const override;
 	virtual bool PopulateContextMenu(FMenuBuilder& MenuBuilder, TSharedPtr<FExtender> MenuExtender, FFrameTime MouseDownTime) override;
 
diff --git a/Engine/Source/Editor/Sequencer/Private/SequencerSettings.cpp b/Engine/Source/Editor/Sequencer/Private/SequencerSettings.cpp
--- a/Engine/Source/Editor/Sequencer/Private/SequencerSettings.cpp
+++ b/Engine/Source/Editor/Sequencer/Private/SequencerSettings.cpp
@@ -48,6 +48,7 @@ USequencerSettings::USequencerSettings( const FObjectInitializer& ObjectInitiali
 	bLeftMouseDragDoesMarquee = false;
 	ZeroPadFrames = 0;
 	JumpFrameIncrement = FFrameNumber(5);
+	TimeWarpDisplay = ESequencerTimeWarpDisplay::Both;
 	bShowLayerBars = true;
 	bShowKeyBars = true;
 	bInfiniteKeyAreas = false;
@@ -584,6 +585,20 @@ void USequencerSettings::SetJumpFrameIncrement(FFrameNumber InJumpFrameIncrement
 	}
 }
 
+ESequencerTimeWarpDisplay USequencerSettings::GetTimeWarpDisplayMode() const
+{
+	return TimeWarpDisplay;
+}
+
+void USequencerSettings::SetTimeWarpDisplayMode(ESequencerTimeWarpDisplay InTimeWarpDisplay)
+{
+	if (TimeWarpDisplay != InTimeWarpDisplay)
+	{
+		TimeWarpDisplay = InTimeWarpDisplay;
+		SaveConfig();
+	}
+}
+
 bool USequencerSettings::GetShowLayerBars() const
 {
 	return bShowLayerBars;
diff --git a/Engine/Source/Editor/Sequencer/Private/SequencerTimeSliderController.cpp b/Engine/Source/Editor/Sequencer/Private/SequencerTimeSliderController.cpp
--- a/Engine/Source/Editor/Sequencer/Private/SequencerTimeSliderController.cpp
+++ b/Engine/Source/Editor/Sequencer/Private/SequencerTimeSliderController.cpp
@@ -3,6 +3,12 @@
 #include "SequencerTimeSliderController.h"
 #include "MVVM/ViewModels/SequencerEditorViewModel.h"
 #include "MVVM/Selection/Selection.h"
+#include "MVVM/SharedViewModelData.h"
+#include "MVVM/Views/SOutlinerView.h"
+#include "TrackEditors/TimeWarpTrackEditor.h"
+#include "Channels/MovieSceneTimeWarpChannel.h"
+#include "SequencerTimeDomainOverride.h"
+#include "SSequencer.h"
 #include "Fonts/SlateFontInfo.h"
 #include "Rendering/DrawElements.h"
 #include "Misc/Paths.h"
@@ -56,6 +62,15 @@ namespace ScrubConstants
 	const int32 MarkLabelBoxWideMargin = 4;
 }
 
+bool GSequencerShowTimeWarpScrubberLink = false;
+FAutoConsoleVariableRef CVarSequencerShowTimeWarpScrubberLink(
+	TEXT("Sequencer.ShowTimeWarpScrubberLink"),
+	GSequencerShowTimeWarpScrubberLink,
+	TEXT("(Default: false) When enabled, shows a dashed link to link warped and unwarped time scrubbers."),
+	ECVF_Default
+);
+
+
 FSequencerTimeSliderController::FSequencerTimeSliderController( const FTimeSliderArgs& InArgs, TWeakPtr<FSequencer> InWeakSequencer )
 	: WeakSequencer(InWeakSequencer)
 	, TimeSliderArgs( InArgs )
@@ -66,10 +81,8 @@ FSequencerTimeSliderController::FSequencerTimeSliderController( const FTimeSlide
 	, HoverMarkIndex( INDEX_NONE )
 {
 	ScrubFillBrush              = FAppStyle::GetBrush( TEXT( "Sequencer.Timeline.ScrubFill" ) );
-	FrameBlockScrubHandleUpBrush   = FAppStyle::GetBrush( TEXT( "Sequencer.Timeline.FrameBlockScrubHandleUp" ) ); 
-	FrameBlockScrubHandleDownBrush = FAppStyle::GetBrush( TEXT( "Sequencer.Timeline.FrameBlockScrubHandleDown" ) );
-	VanillaScrubHandleUpBrush      = FAppStyle::GetBrush( TEXT( "Sequencer.Timeline.VanillaScrubHandleUp" ) ); 
-	VanillaScrubHandleDownBrush    = FAppStyle::GetBrush( TEXT( "Sequencer.Timeline.VanillaScrubHandleDown" ) );
+	FrameBlockScrubHandleDownBrush = FAppStyle::GetBrush( TEXT( "Sequencer.Timeline.FrameBlockScrubHandle" ) );
+	VanillaScrubHandleDownBrush    = FAppStyle::GetBrush( TEXT( "Sequencer.Timeline.ScrubHandle" ) );
 
 	FontMeasureService = FSlateApplication::Get().GetRenderer()->GetFontMeasureService().ToSharedPtr();
 	SmallLayoutFont = FCoreStyle::GetDefaultFontStyle("Regular", 10);
@@ -464,6 +477,8 @@ int32 FSequencerTimeSliderController::DrawVerticalFrames(const FGeometry& Allott
 
 int32 FSequencerTimeSliderController::OnPaintTimeSlider( bool bMirrorLabels, const FGeometry& AllottedGeometry, const FSlateRect& MyCullingRect, FSlateWindowElementList& OutDrawElements, int32 LayerId, const FWidgetStyle& InWidgetStyle, bool bParentEnabled ) const
 {
+	using namespace UE::Sequencer;
+
 	TSharedPtr<FSequencer> Sequencer = WeakSequencer.Pin();
 	if (!Sequencer.IsValid())
 	{
@@ -522,26 +537,29 @@ int32 FSequencerTimeSliderController::OnPaintTimeSlider( bool bMirrorLabels, con
 
 		const int32 ArrowLayer = LayerId + 2;
 		FPaintGeometry MyGeometry =	AllottedGeometry.ToPaintGeometry( FVector2f( HandleEnd - HandleStart, AllottedGeometry.Size.Y ), FSlateLayoutTransform(FVector2f( HandleStart, 0.f )) );
-		FLinearColor ScrubColor = InWidgetStyle.GetColorAndOpacityTint();
-		if(bIsEvaluating)
-		{
-			// @todo Sequencer this color should be specified in the style
-			ScrubColor.A = ScrubColor.A * 0.75f;
-			ScrubColor.B *= 0.1f;
-			ScrubColor.G *= 0.2f;
-		}
-		else
-		{
-			ScrubColor.A = ScrubColor.A * 0.75f;
-			ScrubColor.R = 0.7f;
-			ScrubColor.B = 0.1f;
-			ScrubColor.G = 0.7f;
-		}
+
+		FTimeWarpTrackExtension*   TimeWarpExtension   = Sequencer->GetViewModel()->GetRootSequenceModel()->GetSharedData()->CastDynamic<FTimeWarpTrackExtension>();
+		const FTimeWarpTrackModel* ActiveTimeWarpTrack = TimeWarpExtension ? TimeWarpExtension->GetActiveTimeWarpTrack() : nullptr;
+
+		const bool bIsWarped = (Sequencer->GetSequencerSettings()->GetTimeWarpDisplayMode() == ESequencerTimeWarpDisplay::WarpedTime && ActiveTimeWarpTrack != nullptr);
+		FLinearColor ScrubColor = bIsWarped
+			? FStyleColors::AccentOrange.GetSpecifiedColor()
+			: FStyleColors::AccentRed.GetSpecifiedColor();
+
 		const FSlateBrush* Brush = ScrubMetrics.Style == ESequencerScrubberStyle::Vanilla
-			? ( bMirrorLabels ? VanillaScrubHandleUpBrush    : VanillaScrubHandleDownBrush )
-			: ( bMirrorLabels ? FrameBlockScrubHandleUpBrush : FrameBlockScrubHandleDownBrush );
+			? VanillaScrubHandleDownBrush
+			: FrameBlockScrubHandleDownBrush;
 
-		
+		if (bMirrorLabels)
+		{
+			FSlateRenderTransform FlipTransform(FScale2f(1.f, -1.f), FVector2f(0.f, AllottedGeometry.Size.Y));
+			MyGeometry.SetRenderTransform(
+				Concatenate(
+					FlipTransform,
+					MyGeometry.GetAccumulatedRenderTransform()
+				)
+			);
+		}
 		FSlateDrawElement::MakeBox(
 			OutDrawElements,
 			ArrowLayer,
@@ -550,7 +568,7 @@ int32 FSequencerTimeSliderController::OnPaintTimeSlider( bool bMirrorLabels, con
 			DrawEffects,
 			ScrubColor
 		);
-		
+
 		LayerId = DrawMarkedFrames(AllottedGeometry, RangeToScreen, OutDrawElements, LayerId, DrawEffects, InWidgetStyle, true);
 
 		{
@@ -563,7 +581,7 @@ int32 FSequencerTimeSliderController::OnPaintTimeSlider( bool bMirrorLabels, con
 			}
 			else
 			{
-				FrameString = TimeSliderArgs.NumericTypeInterface->ToString(TimeSliderArgs.ScrubPosition.Get().GetFrame().Value);
+				FrameString = TimeSliderArgs.NumericTypeInterface->ToString(ScrubPosition.Time.GetFrame().Value);
 			}
 
 			if (TimeSliderArgs.ScrubPositionParent.Get() != MovieSceneSequenceID::Invalid)
@@ -606,11 +624,24 @@ int32 FSequencerTimeSliderController::OnPaintTimeSlider( bool bMirrorLabels, con
 			float RangePosX = MouseStartPosX < MouseEndPosX ? MouseStartPosX : MouseEndPosX;
 			float RangeSizeX = FMath::Abs(MouseStartPosX - MouseEndPosX);
 
+			FPaintGeometry PaintGeom = AllottedGeometry.ToPaintGeometry( FVector2f(RangeSizeX, AllottedGeometry.Size.Y), FSlateLayoutTransform(FVector2f(RangePosX, 0.f)) );
+
+			if (bMirrorLabels)
+			{
+				FSlateRenderTransform FlipTransform(FScale2f(1.f, -1.f), FVector2f(0.f, AllottedGeometry.Size.Y));
+				PaintGeom.SetRenderTransform(
+					Concatenate(
+						FlipTransform,
+						PaintGeom.GetAccumulatedRenderTransform()
+					)
+				);
+			}
+
 			FSlateDrawElement::MakeBox(
 				OutDrawElements,
 				LayerId+1,
-				AllottedGeometry.ToPaintGeometry( FVector2f(RangeSizeX, AllottedGeometry.Size.Y), FSlateLayoutTransform(FVector2f(RangePosX, 0.f)) ),
-				bMirrorLabels ? VanillaScrubHandleDownBrush : VanillaScrubHandleUpBrush,
+				PaintGeom,
+				VanillaScrubHandleDownBrush,
 				DrawEffects,
 				MouseStartPosX < MouseEndPosX ? FLinearColor(0.5f, 0.5f, 0.5f) : FLinearColor(0.25f, 0.3f, 0.3f)
 			);
@@ -1253,31 +1284,39 @@ FReply FSequencerTimeSliderController::OnMouseMoveImpl( SWidget& WidgetOwner, co
 
 void FSequencerTimeSliderController::CommitScrubPosition( FFrameTime NewValue, bool bIsScrubbing, bool bEvaluate)
 {
+	using namespace UE::Sequencer;
+
 	bIsEvaluating = bEvaluate;
+
 	// The user can scrub past the viewing range of the time slider controller, so we clamp it to the view range.
 	TSharedPtr<FSequencer> Sequencer = WeakSequencer.Pin();
-	if(Sequencer.IsValid())
+	if(!Sequencer.IsValid())
 	{
-		Sequencer->SnapSequencerTime(NewValue);
+		return;
+	}
 
-		if (bIsScrubbing)
-		{
-			FAnimatedRange ViewRange = GetViewRange();
-		
-			FFrameRate DisplayRate = Sequencer->GetFocusedDisplayRate();
-			FFrameRate TickResolution = Sequencer->GetFocusedTickResolution();
+	const bool bIsWarped = Sequencer->GetSequencerSettings()->GetTimeWarpDisplayMode() == ESequencerTimeWarpDisplay::WarpedTime;
+	FTimeDomainOverride TimeDomainOverride = Sequencer->OverrideTimeDomain(bIsWarped ? ETimeDomain::Warped : ETimeDomain::Unwarped);
 
-			FFrameTime LowerBound = (ViewRange.GetLowerBoundValue() * TickResolution).CeilToFrame();
-			FFrameTime UpperBound = (ViewRange.GetUpperBoundValue() * TickResolution).FloorToFrame();
+	Sequencer->SnapSequencerTime(NewValue);
 
-			if (Sequencer->GetSequencerSettings()->GetIsSnapEnabled() && Sequencer->GetSequencerSettings()->GetSnapPlayTimeToInterval())
-			{
-				LowerBound = FFrameRate::Snap(LowerBound, TickResolution, DisplayRate);
-				UpperBound = FFrameRate::Snap(UpperBound, TickResolution, DisplayRate);
-			}
+	if (bIsScrubbing)
+	{
+		FAnimatedRange ViewRange = GetViewRange();
+	
+		FFrameRate DisplayRate = Sequencer->GetFocusedDisplayRate();
+		FFrameRate TickResolution = Sequencer->GetFocusedTickResolution();
+
+		FFrameTime LowerBound = (ViewRange.GetLowerBoundValue() * TickResolution).CeilToFrame();
+		FFrameTime UpperBound = (ViewRange.GetUpperBoundValue() * TickResolution).FloorToFrame();
 
-			NewValue = FMath::Clamp(NewValue, LowerBound, UpperBound);
+		if (Sequencer->GetSequencerSettings()->GetIsSnapEnabled() && Sequencer->GetSequencerSettings()->GetSnapPlayTimeToInterval())
+		{
+			LowerBound = FFrameRate::Snap(LowerBound, TickResolution, DisplayRate);
+			UpperBound = FFrameRate::Snap(UpperBound, TickResolution, DisplayRate);
 		}
+
+		NewValue = FMath::Clamp(NewValue, LowerBound, UpperBound);
 	}
 
 	// Manage the scrub position ourselves if its not bound to a delegate
@@ -1375,6 +1414,9 @@ FCursorReply FSequencerTimeSliderController::OnCursorQuery( TSharedRef<const SWi
 
 int32 FSequencerTimeSliderController::OnPaintViewArea( const FGeometry& AllottedGeometry, const FSlateRect& MyCullingRect, FSlateWindowElementList& OutDrawElements, int32 LayerId, bool bEnabled, const FPaintViewAreaArgs& Args ) const
 {
+	using namespace UE::MovieScene;
+	using namespace UE::Sequencer;
+
 	TSharedPtr<FSequencer> Sequencer = WeakSequencer.Pin();
 	if (!Sequencer.IsValid())
 	{
@@ -1425,44 +1467,184 @@ int32 FSequencerTimeSliderController::OnPaintViewArea( const FGeometry& Allotted
 
 	LayerId = DrawVerticalFrames(AllottedGeometry, RangeToScreen, OutDrawElements, LayerId, DrawEffects);
 
-	if( Args.bDisplayScrubPosition )
+	FTimeWarpTrackExtension*   TimeWarpExtension   = Sequencer->GetViewModel()->GetRootSequenceModel()->GetSharedData()->CastDynamic<FTimeWarpTrackExtension>();
+	const FTimeWarpTrackModel* ActiveTimeWarpTrack = TimeWarpExtension ? TimeWarpExtension->GetActiveTimeWarpTrack() : nullptr;
+
+	FLinearColor TimeWarpColor = FStyleColors::AccentOrange.GetSpecifiedColor();
+
+	// If we have no active timewarp track, just draw the single scrub position
+	if (ActiveTimeWarpTrack == nullptr || Sequencer->GetSequencerSettings()->GetTimeWarpDisplayMode() != ESequencerTimeWarpDisplay::Both)
 	{
+		const bool bIsWarped = (Sequencer->GetSequencerSettings()->GetTimeWarpDisplayMode() == ESequencerTimeWarpDisplay::WarpedTime && ActiveTimeWarpTrack != nullptr);
+		FLinearColor ScrubColor = bIsWarped ? TimeWarpColor : FLinearColor::White;
+
 		FQualifiedFrameTime ScrubPosition = FQualifiedFrameTime(TimeSliderArgs.ScrubPosition.Get(), GetTickResolution());
-		FScrubberMetrics    ScrubMetrics = GetScrubPixelMetrics(ScrubPosition, RangeToScreen);
+		FScrubberMetrics    ScrubMetrics  = GetScrubPixelMetrics(ScrubPosition, RangeToScreen);
 
 		if (ScrubMetrics.bDrawExtents)
 		{
-			// Draw a box for the scrub position
 			FSlateDrawElement::MakeBox(
 				OutDrawElements,
 				LayerId + 1,
-				AllottedGeometry.ToPaintGeometry(FVector2f(ScrubMetrics.FrameExtentsPx.Size<float>(), AllottedGeometry.Size.Y), FSlateLayoutTransform(FVector2f(ScrubMetrics.FrameExtentsPx.GetLowerBoundValue(), 0.0f))),
+				AllottedGeometry.ToPaintGeometry(
+					FVector2f(ScrubMetrics.FrameExtentsPx.Size<float>(), AllottedGeometry.Size.Y),
+					FSlateLayoutTransform(FVector2f(ScrubMetrics.FrameExtentsPx.GetLowerBoundValue(), 0.f))),
+				ScrubFillBrush,
+				DrawEffects,
+				ScrubColor.CopyWithNewOpacity(0.5f)
+			);
+		}
+
+		if (Args.bDisplayScrubPosition)
+		{
+			const float ScrubPixelX = RangeToScreen.InputToLocalX(ScrubPosition.AsSeconds());
+
+			TArray<FVector2f> LinePoints;
+			LinePoints.AddUninitialized(2);
+			LinePoints[0] = FVector2f( ScrubPixelX, 0.f );
+			LinePoints[1] = FVector2f( ScrubPixelX, AllottedGeometry.Size.Y );
+
+			// Draw a white line for the unwarped scrub position
+			FSlateDrawElement::MakeLines(
+				OutDrawElements,
+				LayerId+1,
+				AllottedGeometry.ToPaintGeometry(),
+				LinePoints,
+				DrawEffects,
+				ScrubColor,
+				false
+			);
+		}
+	}
+	else
+	{
+		// Handle timewarp by drawing the white (unwarped) position up to the bottom of the active timewarp track,
+		//     then draw the orange (warped) time from there
+		FVirtualGeometry Geometry = ActiveTimeWarpTrack->GetVirtualGeometry();
+		const float UnwarpedScrubVerticalClip = StaticCastSharedRef<SSequencer>(Sequencer->GetSequencerWidget())->GetPinnedTreeView()->VirtualToPhysical(Geometry.GetTop() + Geometry.GetNestedHeight());
+
+		FQualifiedFrameTime UnwarpedScrubPosition = Sequencer->GetUnwarpedLocalTime();
+
+		FQualifiedFrameTime WarpedScrubPosition   = Sequencer->GetLocalTime();
+		FScrubberMetrics    UnwarpedScrubMetrics  = GetScrubPixelMetrics(UnwarpedScrubPosition, RangeToScreen);
+		FScrubberMetrics    WarpedScrubMetrics    = GetScrubPixelMetrics(WarpedScrubPosition, RangeToScreen);
+
+		if (UnwarpedScrubMetrics.bDrawExtents)
+		{
+			// Draw a box for the unwarped scrub position
+			FSlateDrawElement::MakeBox(
+				OutDrawElements,
+				LayerId+1,
+				AllottedGeometry.ToPaintGeometry(
+					FVector2f(UnwarpedScrubMetrics.FrameExtentsPx.Size<float>(), UnwarpedScrubVerticalClip - 1.f),
+					FSlateLayoutTransform(FVector2f(UnwarpedScrubMetrics.FrameExtentsPx.GetLowerBoundValue(), 0.f))),
 				ScrubFillBrush,
 				DrawEffects,
 				FLinearColor::White.CopyWithNewOpacity(0.5f)
 			);
+
+			// Draw a box for the warped scrub position
+			FSlateDrawElement::MakeBox(
+				OutDrawElements,
+				LayerId + 1,
+				AllottedGeometry.ToPaintGeometry(
+					FVector2f(WarpedScrubMetrics.FrameExtentsPx.Size<float>(), AllottedGeometry.Size.Y - UnwarpedScrubVerticalClip),
+					FSlateLayoutTransform(FVector2f(WarpedScrubMetrics.FrameExtentsPx.GetLowerBoundValue(), UnwarpedScrubVerticalClip))),
+				ScrubFillBrush,
+				DrawEffects,
+				TimeWarpColor.CopyWithNewOpacity(0.75f)
+			);
 		}
 
-			// Draw a line for the scrub position
-			TArray<FVector2D> LinePoints;
-			{
-			float LinePos = RangeToScreen.InputToLocalX(ScrubPosition.AsSeconds());
+		if (Args.bDisplayScrubPosition)
+		{
+			const float UnwarpedScrubPixelX = RangeToScreen.InputToLocalX(UnwarpedScrubPosition.AsSeconds());
+			const float WarpedScrubPixelX   = RangeToScreen.InputToLocalX(WarpedScrubPosition.AsSeconds());
+			const float OriginPixelX        = RangeToScreen.InputToLocalX(0.f);
 
-				LinePoints.AddUninitialized(2);
-			LinePoints[0] = FVector2D( LinePos, 0.0f );
-			LinePoints[1] = FVector2D( LinePos, FMath::FloorToFloat( AllottedGeometry.Size.Y ) );
-			}
+			TArray<FVector2f> LinePoints;
+			LinePoints.AddUninitialized(2);
+			LinePoints[0] = FVector2f( UnwarpedScrubPixelX, 0.f );
+			LinePoints[1] = FVector2f( UnwarpedScrubPixelX, UnwarpedScrubVerticalClip-1.f );
 
+			// Draw a white line for the unwarped scrub position
 			FSlateDrawElement::MakeLines(
 				OutDrawElements,
 				LayerId+1,
-			AllottedGeometry.ToPaintGeometry(),
+				AllottedGeometry.ToPaintGeometry(),
 				LinePoints,
 				DrawEffects,
 				FLinearColor(1.f, 1.f, 1.f, .5f),
 				false
 			);
+
+			LinePoints[0] = FVector2f( WarpedScrubPixelX, UnwarpedScrubVerticalClip );
+			LinePoints[1] = FVector2f( WarpedScrubPixelX, AllottedGeometry.Size.Y );
+
+			// Draw a line for the warped scrub position
+			FSlateDrawElement::MakeLines(
+				OutDrawElements,
+				LayerId+1,
+				AllottedGeometry.ToPaintGeometry(),
+				LinePoints,
+				DrawEffects,
+				TimeWarpColor,
+				false
+			);
+
+			// Optional cvar for drawing the dashed line
+			if (GSequencerShowTimeWarpScrubberLink)
+			{
+				const float StartLine = WarpedScrubPixelX < UnwarpedScrubPixelX ? WarpedScrubPixelX : UnwarpedScrubPixelX;
+				const float EndLine   = WarpedScrubPixelX < UnwarpedScrubPixelX ? UnwarpedScrubPixelX : WarpedScrubPixelX;
+
+				const float HalfPos = (EndLine + StartLine)*.5f;
+
+				// Draw a dashed line to connect them
+				LinePoints[0] = FVector2f( StartLine,   UnwarpedScrubVerticalClip-1.f );
+				LinePoints[1] = FVector2f( FMath::Max(HalfPos - 10.f, StartLine), UnwarpedScrubVerticalClip-1.f );
+
+				FSlateDrawElement::MakeDashedLines(
+					OutDrawElements,
+					LayerId++,
+					AllottedGeometry.ToPaintGeometry(),
+					CopyTemp(LinePoints),
+					DrawEffects,
+					FColor(212, 147, 20),
+					2.f,
+					5.f,
+					StartLine - OriginPixelX
+				);
+
+				LinePoints[0] = FVector2f( FMath::Min(HalfPos + 10.f, EndLine),   UnwarpedScrubVerticalClip-1.f );
+				LinePoints[1] = FVector2f( EndLine, UnwarpedScrubVerticalClip-1.f );
+
+				FSlateDrawElement::MakeDashedLines(
+					OutDrawElements,
+					LayerId++,
+					AllottedGeometry.ToPaintGeometry(),
+					CopyTemp(LinePoints),
+					DrawEffects,
+					FColor(212, 147, 20),
+					2.f,
+					5.f,
+					StartLine - OriginPixelX
+				);
+
+				FSlateDrawElement::MakeBox(
+					OutDrawElements,
+					LayerId + 1,
+					AllottedGeometry.ToPaintGeometry(
+						FVector2f(12.f, 12.f),
+						FSlateLayoutTransform(FVector2f( HalfPos-6.f, UnwarpedScrubVerticalClip-6.f ))
+					),
+					FAppStyle::GetBrush( TEXT( "Sequencer.Tracks.TimeWarp" ) ),
+					DrawEffects,
+					TimeWarpColor.CopyWithNewOpacity(0.75f)
+				);
+			}
 		}
+	}
 
 	return LayerId;
 }
diff --git a/Engine/Source/Editor/Sequencer/Private/Tools/SequencerEditTool_Movement.cpp b/Engine/Source/Editor/Sequencer/Private/Tools/SequencerEditTool_Movement.cpp
--- a/Engine/Source/Editor/Sequencer/Private/Tools/SequencerEditTool_Movement.cpp
+++ b/Engine/Source/Editor/Sequencer/Private/Tools/SequencerEditTool_Movement.cpp
@@ -3,6 +3,7 @@
 #include "Tools/SequencerEditTool_Movement.h"
 #include "MVVM/ViewModels/SectionModel.h"
 #include "MVVM/ViewModels/SequencerEditorViewModel.h"
+#include "SequencerTimeDomainOverride.h"
 #include "MVVM/Selection/Selection.h"
 #include "Editor.h"
 #include "Fonts/FontMeasure.h"
@@ -50,12 +51,15 @@ FReply FSequencerEditTool_Movement::OnMouseButtonDown(SWidget& OwnerWidget, cons
 
 		DelayedDrag = FDelayedDrag_Hotspot(VirtualTrackArea.CachedTrackAreaGeometry().AbsoluteToLocal(MouseEvent.GetScreenSpacePosition()), MouseEvent.GetEffectingButton(), Hotspot);
 
- 		if (Sequencer.GetSequencerSettings()->GetSnapPlayTimeToPressedKey() || (MouseEvent.IsShiftDown() && MouseEvent.GetEffectingButton() == EKeys::LeftMouseButton) )
+		if (Sequencer.GetSequencerSettings()->GetSnapPlayTimeToPressedKey() || (MouseEvent.IsShiftDown() && MouseEvent.GetEffectingButton() == EKeys::LeftMouseButton) )
 		{
 			if (TSharedPtr<FKeyHotspot> KeyHotspot = HotspotCast<FKeyHotspot>(DelayedDrag->Hotspot))
 			{
-				if (TOptional<FFrameNumber> Time = KeyHotspot->GetTime())
+				TOptional<FFrameNumber> Time   = KeyHotspot->GetTime();
+				TOptional<ETimeDomain>  Domain = KeyHotspot->GetDomain();
+				if (Time && Domain)
 				{
+					FTimeDomainOverride DomainOverride = Sequencer.OverrideTimeDomain(Domain.GetValue());
 					Sequencer.SetLocalTime(Time.GetValue());
 				}
 			}
diff --git a/Engine/Source/Editor/Sequencer/Public/ISequencer.h b/Engine/Source/Editor/Sequencer/Public/ISequencer.h
--- a/Engine/Source/Editor/Sequencer/Public/ISequencer.h
+++ b/Engine/Source/Editor/Sequencer/Public/ISequencer.h
@@ -73,6 +73,7 @@ template <typename NumericType> struct INumericTypeInterface;
 
 enum class EMapChangeType : uint8;
 enum class EPropertyKeyedStatus : uint8;
+enum class ENearestKeyOption : uint8;
 class FCurveEditor;
 class FCurveModel;
 class IToolkitHost;
@@ -85,8 +86,10 @@ namespace UE
 namespace Sequencer
 {
 
+enum class ETimeDomain : uint8;
 class FSequencerEditorViewModel;
 struct FCreateBindingParams;
+struct FTimeDomainOverride;
 
 } // namespace Sequencer
 } // namespace UE
@@ -262,6 +265,12 @@ public:
 	/**@return Returns the time transform from the focused sequence back to the root*/
 	virtual FMovieSceneSequenceTransform GetFocusedMovieSceneSequenceTransform() const = 0;
 
+	/**@return Returns the timewrap transform that local time to warped-local-time */
+	virtual FMovieSceneSequenceTransform GetLocalTimeWarpTransform() const = 0;
+
+	/**@return Returns the timewrap transform that applies to global playback */
+	virtual FMovieSceneSequenceTransform GetGlobalPlaybackWarpTransform() const = 0;
+
 	/** @return The root movie scene being used */
 	virtual FMovieSceneSequenceIDRef GetRootTemplateID() const = 0;
 	virtual FMovieSceneSequenceIDRef GetFocusedTemplateID() const = 0;
@@ -436,6 +445,16 @@ public:
 	 */
 	virtual FQualifiedFrameTime GetGlobalTime() const = 0;
 
+	/**
+	 * Retrieve the current local time in unwarped space. This is usually only required for time-warp operations. Prefer GetLocalTime.
+	 */
+	virtual FQualifiedFrameTime GetUnwarpedLocalTime() const = 0;
+
+	/**
+	 * Temporarily override the behavior of a call to one of the SetLocalTime or OnScrubPositionChanged functions to operate in a specific time-domain
+	 */
+	[[nodiscard]] virtual UE::Sequencer::FTimeDomainOverride OverrideTimeDomain(UE::Sequencer::ETimeDomain NewDomain) = 0;
+
 	virtual TOptional<int32> GetLocalLoopIndex() const = 0;
 
 	/**
@@ -645,6 +664,8 @@ public:
 	virtual void GetKeysFromSelection(TUniquePtr<FSequencerKeyCollection>& KeyCollection, float DuplicateThresoldTime) = 0;
 	virtual FSequencerKeyCollection* GetKeyCollection() = 0;
 
+	virtual FFrameNumber OnGetNearestKey(FFrameTime InTime, ENearestKeyOption NearestKeyOption) = 0;
+
 	virtual TArray<FMovieSceneMarkedFrame> GetMarkedFrames() const = 0;
 
 	/** Gets the currently selected tracks. */
@@ -696,6 +717,10 @@ public:
 	virtual void ThrobKeySelection() = 0;
 	virtual void ThrobSectionSelection() = 0;
 
+	virtual void OnScrubPositionChanged(FFrameTime NewScrubPosition, bool bScrubbing, bool bEvaluate) = 0;
+	virtual void OnBeginScrubbing() = 0;
+	virtual void OnEndScrubbing() = 0;
+
 	/** Gets a multicast delegate which is executed whenever the global time changes. */
 	virtual FOnGlobalTimeChanged& OnGlobalTimeChanged() = 0;
 
diff --git a/Engine/Source/Editor/Sequencer/Public/ISequencerSection.h b/Engine/Source/Editor/Sequencer/Public/ISequencerSection.h
--- a/Engine/Source/Editor/Sequencer/Public/ISequencerSection.h
+++ b/Engine/Source/Editor/Sequencer/Public/ISequencerSection.h
@@ -31,10 +31,12 @@ class ISequencer;
 class ISequencerSection;
 class ISequencerTrackEditor;
 class SWidget;
+class SOverlay;
 struct FGeometry;
 struct FKeyHandle;
 struct FPointerEvent;
 struct FSlateBrush;
+struct FTimeToPixel;
 template <typename ElementType> class TRange;
 struct FMovieSceneChannelMetaData;
 
@@ -42,7 +44,31 @@ namespace UE::Sequencer
 {
 	class FCategoryModel;
 	class FChannelModel;
+	class FSectionModel;
+	class STrackLane;
+	class STrackAreaView;
 	struct FViewDensityInfo;
+
+	struct ISectionView
+	{
+		virtual ~ISectionView()
+		{}
+
+		virtual TSharedRef<FTimeToPixel> GetTimeToPixel() const = 0;
+	};
+
+	struct FCreateSectionViewWidgetParams
+	{
+		static constexpr int32 CompoundTrackLaneViewOrder = 0;
+		static constexpr int32 DefaultWidgetOrder = 10;
+		static constexpr int32 ChannelViewOrder = 20;
+
+		TSharedRef<SOverlay> Overlay;
+		TSharedRef<ISectionView> SectionView;
+		TSharedRef<STrackLane> TrackLane;
+		TSharedRef<STrackAreaView> TrackAreaView;
+		TSharedRef<FSectionModel> SectionModel;
+	};
 }
 
 /** Enumerates which edge is being resized */
@@ -125,6 +151,11 @@ public:
 	 */
 	virtual TSharedRef<SWidget> GenerateSectionWidget() { return SNullWidget::NullWidget; }
 
+	/**
+	 * Create view widgets for the section by adding the necessary widgets to the overlay widget
+	 */
+	SEQUENCER_API virtual void CreateViewWidgets(const UE::Sequencer::FCreateSectionViewWidgetParams& Params);
+
 	/**
 	 * Called when the section is double clicked
 	 *
diff --git a/Engine/Source/Editor/Sequencer/Public/MVVM/Extensions/ITimeDomainExtension.h b/Engine/Source/Editor/Sequencer/Public/MVVM/Extensions/ITimeDomainExtension.h
new file mode 100644
--- /dev/null
+++ b/Engine/Source/Editor/Sequencer/Public/MVVM/Extensions/ITimeDomainExtension.h
@@ -0,0 +1,27 @@
+// Copyright Epic Games, Inc. All Rights Reserved.
+
+#pragma once
+
+#include "MVVM/ViewModelTypeID.h"
+#include "SequencerTimeDomainOverride.h"
+
+namespace UE::Sequencer
+{
+
+
+/**
+ * Extension class that can be added to sections and tracks in order to define which time domain their data is defined in
+ */
+class ITimeDomainExtension
+{
+public:
+
+	UE_SEQUENCER_DECLARE_VIEW_MODEL_TYPE_ID_API(SEQUENCER_API, ITimeDomainExtension)
+
+	virtual ~ITimeDomainExtension(){}
+
+	virtual ETimeDomain GetDomain() const = 0;
+};
+
+} // namespace UE::Sequencer
+
diff --git a/Engine/Source/Editor/Sequencer/Public/SequencerSettings.h b/Engine/Source/Editor/Sequencer/Public/SequencerSettings.h
--- a/Engine/Source/Editor/Sequencer/Public/SequencerSettings.h
+++ b/Engine/Source/Editor/Sequencer/Public/SequencerSettings.h
@@ -56,6 +56,15 @@ enum ESequencerLoopMode : int
 	SLM_LoopSelectionRange UMETA(DisplayName="Loop Selection Range"),
 };
 
+UENUM()
+enum class ESequencerTimeWarpDisplay : uint8
+{
+	UnwarpedTime = 1 UMETA(DisplayName="Unwarped Time"),
+	WarpedTime = 2 UMETA(DisplayName="Warped Time"),
+	Both = UnwarpedTime | WarpedTime UMETA(DisplayName="Both"),
+};
+ENUM_CLASS_FLAGS(ESequencerTimeWarpDisplay)
+
 /** Empty class used to house multiple named USequencerSettings */
 UCLASS()
 class SEQUENCER_API USequencerSettingsContainer
@@ -315,6 +324,11 @@ public:
 	/** Set the number of frames to increment when jumping forwards/backwards */
 	void SetJumpFrameIncrement(FFrameNumber InJumpFrameIncrement);
 
+	/** Get the time-warp display mode */
+	ESequencerTimeWarpDisplay GetTimeWarpDisplayMode() const;
+	/** Set the time-warp display mode */
+	void SetTimeWarpDisplayMode(ESequencerTimeWarpDisplay InTimeWarpDisplay);
+
 	/** @return true if showing layer bars */
 	bool GetShowLayerBars() const;
 	/** Set whether to show layer bars */ 
@@ -686,6 +700,10 @@ protected:
 	UPROPERTY( config, EditAnywhere, Category=Timeline )
 	FFrameNumber JumpFrameIncrement;
 
+	/** Controls how time-warped time is displayed on the timeline. */
+	UPROPERTY( config, EditAnywhere, Category=Timeline )
+	ESequencerTimeWarpDisplay TimeWarpDisplay;
+
 	/** Enable or disable the layer bars to edit keyframes in bulk. */
 	UPROPERTY( config, EditAnywhere, Category=Timeline )
 	bool bShowLayerBars;
diff --git a/Engine/Source/Editor/Sequencer/Public/SequencerTimeDomainOverride.h b/Engine/Source/Editor/Sequencer/Public/SequencerTimeDomainOverride.h
new file mode 100644
--- /dev/null
+++ b/Engine/Source/Editor/Sequencer/Public/SequencerTimeDomainOverride.h
@@ -0,0 +1,61 @@
+// Copyright Epic Games, Inc. All Rights Reserved.
+
+#pragma once
+
+#include "CoreTypes.h"
+
+namespace UE::Sequencer
+{
+
+
+enum class ETimeDomain : uint8
+{
+	Warped,
+	Unwarped
+};
+
+/**
+ * RAII struct that temporarily overrides time operations within Sequencer to operate in warped or unwarped space.
+ * By default all SetLocalTime behavior is in warped space, since this comprises the majority of tracks.
+ * This class is very similar to TGuardValue<ETimeDomain>, except that it is moveable.
+ */
+struct FTimeDomainOverride
+{
+public:
+
+	FTimeDomainOverride(ETimeDomain* Ptr, ETimeDomain NewValue)
+		: ValuePtr(Ptr)
+	{
+		if (Ptr)
+		{
+			OriginalValue = *Ptr;
+			*Ptr = NewValue;
+		}
+	}
+
+	FTimeDomainOverride(const FTimeDomainOverride& Other) = delete;
+	FTimeDomainOverride& operator=(const FTimeDomainOverride& Other) = delete;
+
+	FTimeDomainOverride(FTimeDomainOverride&& Other)
+		: ValuePtr(Other.ValuePtr)
+		, OriginalValue(Other.OriginalValue)
+	{
+		Other.ValuePtr = nullptr;
+	}
+
+	~FTimeDomainOverride()
+	{
+		if (ValuePtr)
+		{
+			*ValuePtr = OriginalValue;
+		}
+	}
+
+private:
+
+	ETimeDomain* ValuePtr;
+	ETimeDomain OriginalValue;
+};
+
+} // namespace UE::Sequencer
+
diff --git a/Engine/Source/Editor/Sequencer/Public/SequencerTimeSliderController.h b/Engine/Source/Editor/Sequencer/Public/SequencerTimeSliderController.h
--- a/Engine/Source/Editor/Sequencer/Public/SequencerTimeSliderController.h
+++ b/Engine/Source/Editor/Sequencer/Public/SequencerTimeSliderController.h
@@ -324,9 +324,6 @@ private:
 
 	/** Brush for drawingthe fill area on the scrubber */
 	const FSlateBrush* ScrubFillBrush;
-
-	/** Brush for drawing an upwards facing scrub handles */
-	const FSlateBrush* FrameBlockScrubHandleUpBrush, *VanillaScrubHandleUpBrush;
 	
 	/** Brush for drawing a downwards facing scrub handle */
 	const FSlateBrush* FrameBlockScrubHandleDownBrush, *VanillaScrubHandleDownBrush;
diff --git a/Engine/Source/Editor/SequencerCore/Private/MVVM/Views/SOutlinerView.cpp b/Engine/Source/Editor/SequencerCore/Private/MVVM/Views/SOutlinerView.cpp
--- a/Engine/Source/Editor/SequencerCore/Private/MVVM/Views/SOutlinerView.cpp
+++ b/Engine/Source/Editor/SequencerCore/Private/MVVM/Views/SOutlinerView.cpp
@@ -838,15 +838,18 @@ void SOutlinerView::Refresh()
 		}
 
 		// Only add pinned nodes if this is showing pinned only
-		bool bIsPinned = false;
+		bool bIsPinned = Extension->ShouldAnchorToTop();
 
 		constexpr bool bIncludeThis = true;
-		for (const TViewModelPtr<IPinnableExtension>& Pinnable : Extension.AsModel()->GetAncestorsOfType<IPinnableExtension>(bIncludeThis))
+		if (!bIsPinned)
 		{
-			if (Pinnable->IsPinned())
+			for (const TViewModelPtr<IPinnableExtension>& Pinnable : Extension.AsModel()->GetAncestorsOfType<IPinnableExtension>(bIncludeThis))
 			{
-				bIsPinned = true;
-				break;
+				if (Pinnable->IsPinned())
+				{
+					bIsPinned = true;
+					break;
+				}
 			}
 		}
 
@@ -875,7 +878,6 @@ void SOutlinerView::Refresh()
 	}
 
 	RebuildList();
-	//RequestTreeRefresh();
 
 	for (TSharedPtr<SOutlinerView> PinnedTreeView : PinnedTreeViews)
 	{
diff --git a/Engine/Source/Editor/SequencerCore/Private/MVVM/Views/STrackLane.cpp b/Engine/Source/Editor/SequencerCore/Private/MVVM/Views/STrackLane.cpp
--- a/Engine/Source/Editor/SequencerCore/Private/MVVM/Views/STrackLane.cpp
+++ b/Engine/Source/Editor/SequencerCore/Private/MVVM/Views/STrackLane.cpp
@@ -125,12 +125,10 @@ TArrayView<const TWeakPtr<STrackLane>> STrackLane::GetChildLanes() const
 
 bool STrackLane::IsPinned() const
 {
+	TSharedPtr<IOutlinerExtension> Outliner = WeakOutlinerItem.Pin();
 	TSharedPtr<IPinnableExtension> Pinnable = WeakOutlinerItem.ImplicitPin();
-	if (Pinnable)
-	{
-		return Pinnable->IsPinned();
-	}
-	return false;
+
+	return (Pinnable && Pinnable->IsPinned()) || (Outliner && Outliner->ShouldAnchorToTop());
 }
 
 void STrackLane::OnHierarchyUpdated()
diff --git a/Engine/Source/Editor/SequencerCore/Public/MVVM/Extensions/IOutlinerExtension.h b/Engine/Source/Editor/SequencerCore/Public/MVVM/Extensions/IOutlinerExtension.h
--- a/Engine/Source/Editor/SequencerCore/Public/MVVM/Extensions/IOutlinerExtension.h
+++ b/Engine/Source/Editor/SequencerCore/Public/MVVM/Extensions/IOutlinerExtension.h
@@ -171,6 +171,8 @@ public:
 	/** Sets whether this outliner item is being filtered out */
 	virtual void SetFilteredOut(bool bIsFilteredOut) = 0;
 
+	virtual bool ShouldAnchorToTop() const = 0;
+
 	/** Gets the outliner item's selection state */
 	virtual EOutlinerSelectionState GetSelectionState() const = 0;
 
@@ -241,6 +243,7 @@ public:
 	virtual void SetFilteredOut(bool bInIsFilteredOut) override { bIsFilteredOut = bInIsFilteredOut; }
 	virtual EOutlinerSelectionState GetSelectionState() const override { return SelectionState; }
 	virtual void SetSelectionState(EOutlinerSelectionState InState) override { SelectionState = InState; }
+	virtual bool ShouldAnchorToTop() const override { return false; }
 
 protected:
 
diff --git a/Engine/Source/Editor/SequencerCore/Public/MVVM/ICastable.h b/Engine/Source/Editor/SequencerCore/Public/MVVM/ICastable.h
--- a/Engine/Source/Editor/SequencerCore/Public/MVVM/ICastable.h
+++ b/Engine/Source/Editor/SequencerCore/Public/MVVM/ICastable.h
@@ -131,7 +131,10 @@ struct TImplements
 };
 
 #define UE_SEQUENCER_DECLARE_CASTABLE(ThisType, ...)																																\
-	UE_SEQUENCER_DECLARE_VIEW_MODEL_TYPE_ID(ThisType);																																\
+	UE_SEQUENCER_DECLARE_CASTABLE_API(, ThisType, __VA_ARGS__)
+
+#define UE_SEQUENCER_DECLARE_CASTABLE_API(MODULE_API, ThisType, ...)																													\
+	UE_SEQUENCER_DECLARE_VIEW_MODEL_TYPE_ID_API(MODULE_API, ThisType);																													\
 	using Implements = ::UE::Sequencer::TImplements<__VA_ARGS__>;																														\
 	UE_NO_UNIQUE_ADDRESS ::UE::Sequencer::ICastable::FAutoRegisterTypeTable AutoRegisterTypeTable = ::UE::Sequencer::ICastable::FAutoRegisterTypeTable(static_cast<ThisType*>(this));	\
 
diff --git a/Engine/Source/Runtime/MovieScene/Private/Channels/MovieSceneTimeWarpChannel.cpp b/Engine/Source/Runtime/MovieScene/Private/Channels/MovieSceneTimeWarpChannel.cpp
--- a/Engine/Source/Runtime/MovieScene/Private/Channels/MovieSceneTimeWarpChannel.cpp
+++ b/Engine/Source/Runtime/MovieScene/Private/Channels/MovieSceneTimeWarpChannel.cpp
@@ -4,7 +4,7 @@
 
 void Dilate(FMovieSceneTimeWarpChannel* InChannel, FFrameNumber Origin, double DilationFactor)
 {
-	if (InChannel->Domain == ETimeWarpChannelDomain::PlayRate)
+	if (InChannel->Domain == UE::MovieScene::ETimeWarpChannelDomain::PlayRate)
 	{
 		// Inverse dilate the values if we are in the play-rate domain
 		for (FMovieSceneDoubleValue& Value : InChannel->GetData().GetValues())
diff --git a/Engine/Source/Runtime/MovieScene/Private/Compilation/MovieSceneCompiledDataManager.cpp b/Engine/Source/Runtime/MovieScene/Private/Compilation/MovieSceneCompiledDataManager.cpp
--- a/Engine/Source/Runtime/MovieScene/Private/Compilation/MovieSceneCompiledDataManager.cpp
+++ b/Engine/Source/Runtime/MovieScene/Private/Compilation/MovieSceneCompiledDataManager.cpp
@@ -12,6 +12,7 @@
 #include "Sections/MovieSceneSubSection.h"
 #include "Tracks/MovieSceneSubTrack.h"
 #include "Tracks/MovieSceneTimeWarpTrack.h"
+#include "Channels/MovieSceneTimeWarpChannel.h"
 #include "IMovieSceneModule.h"
 #include "MovieSceneTimeHelpers.h"
 #include "MovieSceneTransformTypes.h"
@@ -1595,7 +1596,14 @@ bool UMovieSceneCompiledDataManager::CompileHierarchy(UMovieSceneSequence* Seque
 		for (UMovieSceneTrack* Track : MovieScene->GetTracks())
 		{
 			UMovieSceneTimeWarpTrack* TimeWarpTrack = Cast<UMovieSceneTimeWarpTrack>(Track);
-			if (TimeWarpTrack && !TimeWarpTrack->IsEvalDisabled())
+			if (!TimeWarpTrack)
+			{
+				continue;
+			}
+
+			TimeWarpTrack->bIsActiveTimeWarp = false;
+
+			if (!bContainsTimeWarp && !TimeWarpTrack->IsEvalDisabled())
 			{
 				FMovieSceneNestedSequenceTransform TimeWarpTransform = TimeWarpTrack->GenerateTransform();
 
@@ -1604,10 +1612,9 @@ bool UMovieSceneCompiledDataManager::CompileHierarchy(UMovieSceneSequence* Seque
 				{
 					InOutHierarchy->SetRootTransform(FMovieSceneSequenceTransform(MoveTemp(TimeWarpTransform)));
 					bContainsTimeWarp = true;
-				}
 
-				// Only 1 time-warp track supported
-				break;
+					TimeWarpTrack->bIsActiveTimeWarp = true;
+				}
 			}
 		}
 	}
diff --git a/Engine/Source/Runtime/MovieScene/Private/EntitySystem/MovieSceneSequenceUpdaters.cpp b/Engine/Source/Runtime/MovieScene/Private/EntitySystem/MovieSceneSequenceUpdaters.cpp
--- a/Engine/Source/Runtime/MovieScene/Private/EntitySystem/MovieSceneSequenceUpdaters.cpp
+++ b/Engine/Source/Runtime/MovieScene/Private/EntitySystem/MovieSceneSequenceUpdaters.cpp
@@ -21,6 +21,7 @@
 #include "Evaluation/MovieSceneRootOverridePath.h"
 
 #include "MovieSceneTimeHelpers.h"
+#include "Channels/MovieSceneTimeWarpChannel.h"
 
 #include "Algo/IndexOf.h"
 #include "Algo/Transform.h"
@@ -294,7 +295,7 @@ void FSequenceUpdater_Flat::PopulateUpdateFlags(TSharedRef<const FSharedPlayback
 	}
 
 	const FMovieSceneSequenceHierarchy* Hierarchy = SharedPlaybackState->GetCompiledDataManager()->FindHierarchy(CompiledDataID);
-	if (Hierarchy && !Hierarchy->GetRootTransform().IsIdentity())
+	if (Hierarchy && Hierarchy->GetRootTransform().FindFirstWarpDomain() == ETimeWarpChannelDomain::Time)
 	{
 		// Time-warped root transforms require dissection to manipulate the evaluation range
 		OutUpdateFlags |= ESequenceInstanceUpdateFlags::NeedsDissection;
@@ -500,7 +501,7 @@ void FSequenceUpdater_Hierarchical::PopulateUpdateFlags(TSharedRef<const FShared
 
 	if (const FMovieSceneSequenceHierarchy* Hierarchy = CompiledDataManager->FindHierarchy(CompiledDataID))
 	{
-		if (!Hierarchy->GetRootTransform().IsLinear())
+		if (Hierarchy->GetRootTransform().FindFirstWarpDomain() == ETimeWarpChannelDomain::Time)
 		{
 			OutUpdateFlags |= ESequenceInstanceUpdateFlags::NeedsDissection;
 		}
@@ -541,7 +542,7 @@ void FSequenceUpdater_Hierarchical::DissectContext(TSharedRef<const FSharedPlayb
 			RootContext = Context.Transform(SubData->RootToSequenceTransform, SubData->TickResolution);
 		}
 	}
-	else if (!RootHierarchy->GetRootTransform().IsIdentity())
+	else if (RootHierarchy->GetRootTransform().FindFirstWarpDomain() == ETimeWarpChannelDomain::Time)
 	{
 		RootContext = Context.Transform(RootHierarchy->GetRootTransform(), Context.GetFrameRate());
 	}
@@ -634,7 +635,7 @@ void FSequenceUpdater_Hierarchical::DissectContext(TSharedRef<const FSharedPlayb
 		}
 		UE::MovieScene::DissectRange(RootDissectionTimes, RootContext.GetRange(), OutDissections);
 	}
-	else if (!RootHierarchy->GetRootTransform().IsIdentity())
+	else if (RootHierarchy->GetRootTransform().FindFirstWarpDomain() == ETimeWarpChannelDomain::Time)
 	{
 		OutDissections.Add(RootContext.GetRange());
 	}
diff --git a/Engine/Source/Runtime/MovieScene/Private/Evaluation/MovieSceneSequenceHierarchy.cpp b/Engine/Source/Runtime/MovieScene/Private/Evaluation/MovieSceneSequenceHierarchy.cpp
--- a/Engine/Source/Runtime/MovieScene/Private/Evaluation/MovieSceneSequenceHierarchy.cpp
+++ b/Engine/Source/Runtime/MovieScene/Private/Evaluation/MovieSceneSequenceHierarchy.cpp
@@ -67,7 +67,14 @@ FMovieSceneSubSequenceData::FMovieSceneSubSequenceData(const UMovieSceneSubSecti
 	// being truly the full transform.
 	OuterToInnerTransform = InSubSection.OuterToInnerTransform_NoInnerTimeWarp();
 
+#if WITH_EDITORONLY_DATA
+	RootToUnwarpedLocalTransform = OuterToInnerTransform;
+	InSubSection.AppendInnerTimeWarpTransform(LocalToWarpedLocalTransform);
+	OuterToInnerTransform.Append(LocalToWarpedLocalTransform);
+#else
 	InSubSection.AppendInnerTimeWarpTransform(OuterToInnerTransform);
+#endif
+
 	RootToSequenceTransform = OuterToInnerTransform;
 
 	PlayRange.Value = UE::MovieScene::ConvertToDiscreteRange(OuterToInnerTransform.ComputeTraversedHull(ParentPlayRange.Value));
diff --git a/Engine/Source/Runtime/MovieScene/Private/Evaluation/MovieSceneSequenceTransform.cpp b/Engine/Source/Runtime/MovieScene/Private/Evaluation/MovieSceneSequenceTransform.cpp
--- a/Engine/Source/Runtime/MovieScene/Private/Evaluation/MovieSceneSequenceTransform.cpp
+++ b/Engine/Source/Runtime/MovieScene/Private/Evaluation/MovieSceneSequenceTransform.cpp
@@ -397,6 +397,20 @@ bool FMovieSceneNestedSequenceTransform::SupportsBoundaries() const
 	}
 }
 
+TOptional<UE::MovieScene::ETimeWarpChannelDomain> FMovieSceneNestedSequenceTransform::GetWarpDomain() const
+{
+	UMovieSceneTimeWarpGetter* Getter = TimeScale.GetType() == EMovieSceneTimeWarpType::Custom
+		? TimeScale.AsCustom()
+		: nullptr;
+
+	if (Getter)
+	{
+		return Getter->GetDomain();
+	}
+
+	return TOptional<UE::MovieScene::ETimeWarpChannelDomain>();
+}
+
 FMovieSceneInverseNestedSequenceTransform FMovieSceneNestedSequenceTransform::Inverse() const
 {
 	if (TimeScale.GetType() == EMovieSceneTimeWarpType::FixedPlayRate)
@@ -974,6 +988,20 @@ bool FMovieSceneSequenceTransform::ExtractBoundariesWithinRange(FFrameTime Start
 	return false;
 }
 
+TOptional<UE::MovieScene::ETimeWarpChannelDomain> FMovieSceneSequenceTransform::FindFirstWarpDomain() const
+{
+	for (const FMovieSceneNestedSequenceTransform& NestedTransform : NestedTransforms)
+	{
+		TOptional<UE::MovieScene::ETimeWarpChannelDomain> Domain = NestedTransform.GetWarpDomain();
+		if (Domain)
+		{
+			return Domain;
+		}
+	}
+
+	return TOptional<UE::MovieScene::ETimeWarpChannelDomain>();
+}
+
 void FMovieSceneSequenceTransform::AddLoop(FFrameNumber InStart, FFrameNumber InEnd)
 {
 	check(InStart < InEnd);
@@ -1029,6 +1057,23 @@ FMovieSceneInverseSequenceTransform FMovieSceneSequenceTransform::Inverse() cons
 }
 
 
+void FMovieSceneSequenceTransform::Append(const FMovieSceneSequenceTransform& Tail)
+{
+	if (IsLinear())
+	{
+		if (!Tail.LinearTransform.IsIdentity())
+		{
+			LinearTransform = Tail.LinearTransform * LinearTransform;
+		}
+	}
+	else if (!Tail.LinearTransform.IsIdentity())
+	{
+		NestedTransforms.Add(Tail.LinearTransform);
+	}
+
+	NestedTransforms.Append(Tail.NestedTransforms);
+}
+
 FMovieSceneSequenceTransform FMovieSceneSequenceTransform::operator*(const FMovieSceneSequenceTransform& RHS) const
 {
 	if (IsLinear() && RHS.IsLinear())
diff --git a/Engine/Source/Runtime/MovieScene/Private/MovieSceneSequencePlayer.cpp b/Engine/Source/Runtime/MovieScene/Private/MovieSceneSequencePlayer.cpp
--- a/Engine/Source/Runtime/MovieScene/Private/MovieSceneSequencePlayer.cpp
+++ b/Engine/Source/Runtime/MovieScene/Private/MovieSceneSequencePlayer.cpp
@@ -6,6 +6,7 @@
 #include "MovieSceneTimeHelpers.h"
 #include "MovieSceneSequence.h"
 #include "MovieSceneSequenceTickManager.h"
+#include "Channels/MovieSceneTimeWarpChannel.h"
 #include "Engine/Engine.h"
 #include "UObject/Stack.h"
 #include "Internationalization/Text.h"
@@ -1196,6 +1197,18 @@ void UMovieSceneSequencePlayer::UpdateTimeCursorPosition_Internal(FFrameTime New
 		bPendingOnStartedPlaying = false;
 	}
 
+	if (Method == EUpdatePositionMethod::Play)
+	{
+		const FMovieSceneSequenceHierarchy* Hierarchy = RootTemplateInstance.GetHierarchy();
+
+		if (Hierarchy && Hierarchy->GetRootTransform().FindFirstWarpDomain() == UE::MovieScene::ETimeWarpChannelDomain::PlayRate)
+		{
+			NewPosition = ConvertFrameTime(NewPosition, PlayPosition.GetInputRate(), PlayPosition.GetOutputRate());
+			NewPosition = Hierarchy->GetRootTransform().TransformTime(NewPosition);
+			NewPosition = ConvertFrameTime(NewPosition, PlayPosition.GetOutputRate(), PlayPosition.GetInputRate());
+		}
+	}
+
 	// If we should pause during this evaluation, we'll handle that below.
 	const TOptional<TRange<FFrameTime>> PauseRange = GetPauseRange(NewPosition);
 	if (!PauseRange.IsSet() && Method == EUpdatePositionMethod::Play && ShouldStopOrLoop(NewPosition))
diff --git a/Engine/Source/Runtime/MovieScene/Private/Tracks/MovieSceneTimeWarpTrack.cpp b/Engine/Source/Runtime/MovieScene/Private/Tracks/MovieSceneTimeWarpTrack.cpp
--- a/Engine/Source/Runtime/MovieScene/Private/Tracks/MovieSceneTimeWarpTrack.cpp
+++ b/Engine/Source/Runtime/MovieScene/Private/Tracks/MovieSceneTimeWarpTrack.cpp
@@ -11,6 +11,9 @@ UMovieSceneTimeWarpTrack::UMovieSceneTimeWarpTrack(const FObjectInitializer& Obj
 {
 #if WITH_EDITORONLY_DATA
 	bSupportsConditions = false;
+
+	// Timewarp should always exist at the top
+	SortingOrder = -10000;
 #endif
 }
 
diff --git a/Engine/Source/Runtime/MovieScene/Private/Variants/MovieScenePlayRateCurve.cpp b/Engine/Source/Runtime/MovieScene/Private/Variants/MovieScenePlayRateCurve.cpp
--- a/Engine/Source/Runtime/MovieScene/Private/Variants/MovieScenePlayRateCurve.cpp
+++ b/Engine/Source/Runtime/MovieScene/Private/Variants/MovieScenePlayRateCurve.cpp
@@ -13,7 +13,7 @@
 UMovieScenePlayRateCurve::UMovieScenePlayRateCurve()
 {
 	PlayRate.Owner = nullptr;
-	PlayRate.Domain = ETimeWarpChannelDomain::PlayRate;
+	PlayRate.Domain = UE::MovieScene::ETimeWarpChannelDomain::PlayRate;
 
 	OnSignatureChanged().AddUObject(this, &UMovieScenePlayRateCurve::InvalidateTimeWarp);
 }
@@ -172,3 +172,8 @@ void UMovieScenePlayRateCurve::ScaleBy(double UnwarpedScaleFactor)
 	bUpToDate = false;
 }
 
+UE::MovieScene::ETimeWarpChannelDomain UMovieScenePlayRateCurve::GetDomain() const
+{
+	return UE::MovieScene::ETimeWarpChannelDomain::PlayRate;
+}
+
diff --git a/Engine/Source/Runtime/MovieScene/Private/Variants/MovieSceneTimeWarpCurve.cpp b/Engine/Source/Runtime/MovieScene/Private/Variants/MovieSceneTimeWarpCurve.cpp
--- a/Engine/Source/Runtime/MovieScene/Private/Variants/MovieSceneTimeWarpCurve.cpp
+++ b/Engine/Source/Runtime/MovieScene/Private/Variants/MovieSceneTimeWarpCurve.cpp
@@ -12,7 +12,7 @@
 UMovieSceneTimeWarpCurve::UMovieSceneTimeWarpCurve()
 {
 	Channel.Owner = nullptr;
-	Channel.Domain = ETimeWarpChannelDomain::Time;
+	Channel.Domain = UE::MovieScene::ETimeWarpChannelDomain::Time;
 }
 
 void UMovieSceneTimeWarpCurve::InitializeDefaults()
@@ -200,4 +200,9 @@ void UMovieSceneTimeWarpCurve::ScaleBy(double UnwarpedScaleFactor)
 {
 	Modify();
 	Dilate(&Channel, 0, UnwarpedScaleFactor);
+}
+
+UE::MovieScene::ETimeWarpChannelDomain UMovieSceneTimeWarpCurve::GetDomain() const
+{
+	return UE::MovieScene::ETimeWarpChannelDomain::Time;
 }
\ No newline at end of file
diff --git a/Engine/Source/Runtime/MovieScene/Public/Channels/MovieSceneTimeWarpChannel.h b/Engine/Source/Runtime/MovieScene/Public/Channels/MovieSceneTimeWarpChannel.h
--- a/Engine/Source/Runtime/MovieScene/Public/Channels/MovieSceneTimeWarpChannel.h
+++ b/Engine/Source/Runtime/MovieScene/Public/Channels/MovieSceneTimeWarpChannel.h
@@ -8,12 +8,14 @@
 
 class UMovieScene;
 
-
-enum class ETimeWarpChannelDomain : uint8
+namespace UE::MovieScene
 {
-	Time,
-	PlayRate,
-};
+	enum class ETimeWarpChannelDomain : uint8
+	{
+		Time,
+		PlayRate,
+	};
+}
 
 
 USTRUCT()
@@ -24,7 +26,7 @@ struct FMovieSceneTimeWarpChannel : public FMovieSceneDoubleChannel
 	UPROPERTY(transient)
 	TObjectPtr<UMovieScene> Owner;
 
-	ETimeWarpChannelDomain Domain;
+	UE::MovieScene::ETimeWarpChannelDomain Domain;
 };
 
 MOVIESCENE_API void Dilate(FMovieSceneTimeWarpChannel* InChannel, FFrameNumber Origin, double DilationFactor);
diff --git a/Engine/Source/Runtime/MovieScene/Public/Evaluation/MovieSceneSequenceHierarchy.h b/Engine/Source/Runtime/MovieScene/Public/Evaluation/MovieSceneSequenceHierarchy.h
--- a/Engine/Source/Runtime/MovieScene/Public/Evaluation/MovieSceneSequenceHierarchy.h
+++ b/Engine/Source/Runtime/MovieScene/Public/Evaluation/MovieSceneSequenceHierarchy.h
@@ -100,6 +100,14 @@ struct FMovieSceneSubSequenceData
 
 #if WITH_EDITORONLY_DATA
 
+	/** Transform that turns local space into warped local space based on a timewarp track that exists in this sequence. */
+	UPROPERTY()
+	FMovieSceneSequenceTransform LocalToWarpedLocalTransform;
+
+	/** The transform from root space to this sub-sequence's unwarped local space. */
+	UPROPERTY()
+	FMovieSceneSequenceTransform RootToUnwarpedLocalTransform;
+
 	UPROPERTY()
 	FMovieSceneTransformBreadcrumbs StartTimeBreadcrumbs;
 
diff --git a/Engine/Source/Runtime/MovieScene/Public/Evaluation/MovieSceneSequenceTransform.h b/Engine/Source/Runtime/MovieScene/Public/Evaluation/MovieSceneSequenceTransform.h
--- a/Engine/Source/Runtime/MovieScene/Public/Evaluation/MovieSceneSequenceTransform.h
+++ b/Engine/Source/Runtime/MovieScene/Public/Evaluation/MovieSceneSequenceTransform.h
@@ -24,6 +24,8 @@ namespace UE::MovieScene
 {
 	struct FInverseTransformTimeParams;
 	struct FTransformTimeParams;
+
+	enum class ETimeWarpChannelDomain : uint8;
 }
 
 
@@ -502,6 +504,12 @@ public:
 	MOVIESCENE_API bool SupportsBoundaries() const;
 
 
+	/**
+	 * Retrieve this nested transform's time-warp domain, if it has one
+	 */
+	MOVIESCENE_API TOptional<UE::MovieScene::ETimeWarpChannelDomain> GetWarpDomain() const;
+
+
 	/**
 	 * Extract all the boundaries for this transform within the specified untransformed start and end time,
 	 * invoking Visitor for every boundary that is present within the range. Iteration will continue until
@@ -739,6 +747,12 @@ struct FMovieSceneSequenceTransform
 	MOVIESCENE_API TRange<FFrameTime> ComputeTraversedHull(const TRange<FFrameNumber>& Range) const;
 
 
+	/**
+	 * Retrieve the first active timewarp domain that is present in this transform, if any is present at all
+	 */
+	MOVIESCENE_API TOptional<UE::MovieScene::ETimeWarpChannelDomain> FindFirstWarpDomain() const;
+
+
 public:
 
 
@@ -764,6 +778,14 @@ public:
 	MOVIESCENE_API void AddLoop(FFrameNumber Start, FFrameNumber End);
 
 
+	/**
+	 * Append another transform to this one, resulting in a transform that effectively goes from this -> Tail
+	 * 
+	 * @param Tail      The transform to append
+	 */
+	MOVIESCENE_API void Append(const FMovieSceneSequenceTransform& Tail);
+
+
 	/**
 	 * Multiply this tranmsform with another transform, resulting in a single transform that gets from RHS parent to LHS space
 	 * @note Transforms apply from right to left
diff --git a/Engine/Source/Runtime/MovieScene/Public/Tracks/MovieSceneTimeWarpTrack.h b/Engine/Source/Runtime/MovieScene/Public/Tracks/MovieSceneTimeWarpTrack.h
--- a/Engine/Source/Runtime/MovieScene/Public/Tracks/MovieSceneTimeWarpTrack.h
+++ b/Engine/Source/Runtime/MovieScene/Public/Tracks/MovieSceneTimeWarpTrack.h
@@ -39,6 +39,10 @@ public:
 
 public:
 
+	/** True if this track is the active timewarp in its sequence. Assigned as part of compilation */
+	UPROPERTY()
+	bool bIsActiveTimeWarp = true;
+
 	// UMovieSceneTrack interface
 
 	virtual void RemoveAllAnimationData() override;
diff --git a/Engine/Source/Runtime/MovieScene/Public/Variants/MovieScenePlayRateCurve.h b/Engine/Source/Runtime/MovieScene/Public/Variants/MovieScenePlayRateCurve.h
--- a/Engine/Source/Runtime/MovieScene/Public/Variants/MovieScenePlayRateCurve.h
+++ b/Engine/Source/Runtime/MovieScene/Public/Variants/MovieScenePlayRateCurve.h
@@ -30,6 +30,7 @@ public:
 	EMovieSceneChannelProxyType PopulateChannelProxy(FMovieSceneChannelProxyData& OutProxyData, EAllowTopLevelChannels AllowTopLevel) override;
 	bool DeleteChannel(FMovieSceneTimeWarpVariant& OutVariant, FName ChannelName) override;
 	void ScaleBy(double UnwarpedScaleFactor) override;
+	UE::MovieScene::ETimeWarpChannelDomain GetDomain() const override;
 	/* End UMovieSceneTimeWarpGetter Implementation */
 
 #if WITH_EDITOR
diff --git a/Engine/Source/Runtime/MovieScene/Public/Variants/MovieSceneTimeWarpCurve.h b/Engine/Source/Runtime/MovieScene/Public/Variants/MovieSceneTimeWarpCurve.h
--- a/Engine/Source/Runtime/MovieScene/Public/Variants/MovieSceneTimeWarpCurve.h
+++ b/Engine/Source/Runtime/MovieScene/Public/Variants/MovieSceneTimeWarpCurve.h
@@ -28,6 +28,7 @@ public:
 	EMovieSceneChannelProxyType PopulateChannelProxy(FMovieSceneChannelProxyData& OutProxyData, EAllowTopLevelChannels AllowTopLevel) override;
 	bool DeleteChannel(FMovieSceneTimeWarpVariant& OutVariant, FName ChannelName) override;
 	void ScaleBy(double UnwarpedScaleFactor) override;
+	UE::MovieScene::ETimeWarpChannelDomain GetDomain() const override;
 	/* End UMovieSceneTimeWarpGetter Implementation */
 
 
diff --git a/Engine/Source/Runtime/MovieScene/Public/Variants/MovieSceneTimeWarpGetter.h b/Engine/Source/Runtime/MovieScene/Public/Variants/MovieSceneTimeWarpGetter.h
--- a/Engine/Source/Runtime/MovieScene/Public/Variants/MovieSceneTimeWarpGetter.h
+++ b/Engine/Source/Runtime/MovieScene/Public/Variants/MovieSceneTimeWarpGetter.h
@@ -8,6 +8,7 @@
 #include "Misc/FrameTime.h"
 #include "Misc/Optional.h"
 #include "Channels/IMovieSceneChannelOwner.h"
+#include "Channels/MovieSceneTimeWarpChannel.h"
 #include "Variants/MovieSceneNumericVariantGetter.h"
 #include "MovieSceneTimeWarpGetter.generated.h"
 
@@ -115,6 +116,12 @@ public:
 	virtual bool InverseRemapTimeWithinRange(FFrameTime InTime, FFrameTime RangeStart, FFrameTime RangeEnd, const TFunctionRef<bool(FFrameTime)>& VisitorCallback) const
 		PURE_VIRTUAL(UMovieSceneTimeWarpGetter::InverseRemapTimeWithinRange, return true; )
 
+	/**
+	 * Retrieve the time domain that this time warp getter operates within
+	 */
+	virtual UE::MovieScene::ETimeWarpChannelDomain GetDomain() const
+		PURE_VIRTUAL(UMovieSceneTimeWarpGetter::GetDomain, return UE::MovieScene::ETimeWarpChannelDomain::Time; )
+
 	/*~~~~~~~~~~~~~~~~~~~~~~~~~~~~~~~~~~~~~~~~~~~
 	 *            End abstract API
 	  ~~~~~~~~~~~~~~~~~~~~~~~~~~~~~~~~~~~~~~~~~~~*/
diff --git a/Engine/Source/Runtime/UMG/Private/Animation/UMGSequencePlayer.cpp b/Engine/Source/Runtime/UMG/Private/Animation/UMGSequencePlayer.cpp
--- a/Engine/Source/Runtime/UMG/Private/Animation/UMGSequencePlayer.cpp
+++ b/Engine/Source/Runtime/UMG/Private/Animation/UMGSequencePlayer.cpp
@@ -5,7 +5,9 @@
 #include "UMGPrivate.h"
 #include "Animation/WidgetAnimation.h"
 #include "MovieSceneTimeHelpers.h"
+#include "Evaluation/MovieSceneSequenceHierarchy.h"
 #include "Evaluation/MovieScenePlayback.h"
+#include "Channels/MovieSceneTimeWarpChannel.h"
 #include "EntitySystem/MovieSceneEntitySystemRunner.h"
 #include "EntitySystem/MovieSceneEntitySystemLinker.h"
 #include "Animation/UMGSequenceTickManager.h"
@@ -38,7 +40,7 @@ UUMGSequencePlayer::UUMGSequencePlayer(const FObjectInitializer& ObjectInitializ
 	: Super(ObjectInitializer)
 {
 	PlayerStatus = EMovieScenePlayerStatus::Stopped;
-	TimeCursorPosition = FFrameTime(0);
+	TimeCursorPosition = UnwarpedPosition = FFrameTime(0);
 	PlaybackSpeed = 1;
 	bRestoreState = false;
 	Animation = nullptr;
@@ -101,8 +103,18 @@ void UUMGSequencePlayer::Tick(float DeltaTime)
 	{
 		FFrameTime DeltaFrameTime = (bIsPlayingForward ? DeltaTime * PlaybackSpeed : -DeltaTime * PlaybackSpeed) * AnimationResolution;
 
+		UnwarpedPosition += DeltaFrameTime;
+
 		FFrameTime LastTimePosition = TimeCursorPosition;
-		TimeCursorPosition += DeltaFrameTime;
+
+		const FMovieSceneSequenceHierarchy* Hierarchy = RootTemplateInstance.GetHierarchy();
+		FMovieSceneSequenceTransform RootTransform;
+		if (Hierarchy && Hierarchy->GetRootTransform().FindFirstWarpDomain() == UE::MovieScene::ETimeWarpChannelDomain::PlayRate)
+		{
+			RootTransform = Hierarchy->GetRootTransform();
+		}
+
+		TimeCursorPosition = RootTransform.TransformTime(UnwarpedPosition);
 
 		// Check if we crossed over bounds
 		const bool bCrossedLowerBound = TimeCursorPosition < FFrameTime(0);
@@ -127,7 +139,7 @@ void UUMGSequencePlayer::Tick(float DeltaTime)
 		{
 			if (bCompleted)
 			{
-				TimeCursorPosition = FFrameTime(0);
+				TimeCursorPosition = UnwarpedPosition = FFrameTime(0);
 			}
 			else
 			{
@@ -177,6 +189,11 @@ void UUMGSequencePlayer::Tick(float DeltaTime)
 			}
 		}
 
+		if (bCrossedLowerBound || bCrossedUpperBound || bCrossedEndTime)
+		{
+			UnwarpedPosition = RootTransform.Inverse().TryTransformTime(TimeCursorPosition).Get(TimeCursorPosition);
+		}
+
 		bCompleteOnPostEvaluation = bCompleted;
 
 		const bool bHasJumped = (bCrossedLowerBound || bCrossedUpperBound || bCrossedEndTime);
@@ -260,6 +277,17 @@ void UUMGSequencePlayer::PlayInternal(double StartAtTime, double EndAtTime, int3
 
 	// Clamp the start time and end time to be within the bounds
 	TimeCursorPosition = FMath::Clamp(TimeCursorPosition, FFrameTime(0), LastValidFrame);
+
+	const FMovieSceneSequenceHierarchy* Hierarchy = RootTemplateInstance.GetHierarchy();
+	if (Hierarchy && Hierarchy->GetRootTransform().FindFirstWarpDomain() == UE::MovieScene::ETimeWarpChannelDomain::PlayRate)
+	{
+		UnwarpedPosition = Hierarchy->GetRootTransform().Inverse().TryTransformTime(TimeCursorPosition).Get(TimeCursorPosition);
+	}
+	else
+	{
+		UnwarpedPosition = TimeCursorPosition;
+	}
+
 	EndTime = FMath::Clamp(EndAtTime * AnimationResolution, FFrameTime(0), LastValidFrame);
 
 	if ( PlayMode == EUMGSequencePlayMode::PingPong )
@@ -381,7 +409,7 @@ void UUMGSequencePlayer::Stop()
 	UUserWidget* Widget = UserWidget.Get();
 	UUMGSequenceTickManager* TickManager = Widget ? ToRawPtr(Widget->AnimationTickManager) : nullptr;
 
-	TimeCursorPosition = FFrameTime(0);
+	TimeCursorPosition = UnwarpedPosition = FFrameTime(0);
 
 	if (!TickManager || !RootTemplateInstance.IsValid())
 	{
diff --git a/Engine/Source/Runtime/UMG/Public/Animation/UMGSequencePlayer.h b/Engine/Source/Runtime/UMG/Public/Animation/UMGSequencePlayer.h
--- a/Engine/Source/Runtime/UMG/Public/Animation/UMGSequencePlayer.h
+++ b/Engine/Source/Runtime/UMG/Public/Animation/UMGSequencePlayer.h
@@ -125,6 +125,9 @@ private:
 	/** The current time cursor position within the sequence, between 0 and Duration */
 	FFrameTime TimeCursorPosition;
 
+	/** Unwarped position. Only different from TimeCursorPosition if there is a global playback transform. */
+	FFrameTime UnwarpedPosition;
+
 	/** The duration of the sequence */
 	int32 Duration;
 
