Sequencer: Fix icon for condition filter

#rnx

diff --git a/Engine/Source/Editor/Sequencer/Private/Filters/Filters/SequencerTrackFilter_Condition.cpp b/Engine/Source/Editor/Sequencer/Private/Filters/Filters/SequencerTrackFilter_Condition.cpp
--- a/Engine/Source/Editor/Sequencer/Private/Filters/Filters/SequencerTrackFilter_Condition.cpp
+++ b/Engine/Source/Editor/Sequencer/Private/Filters/Filters/SequencerTrackFilter_Condition.cpp
@@ -37,7 +37,7 @@ FText FSequencerTrackFilter_Condition::GetDisplayName() const
 
 FSlateIcon FSequencerTrackFilter_Condition::GetIcon() const
 {
-	return FSlateIcon(FAppStyle::GetAppStyleSetName(), TEXT("Sequencer.Column.Condition"));
+	return FSlateIcon(FAppStyle::GetAppStyleSetName(), TEXT("Sequencer.Decorator.Condition"));
 }
 
 FString FSequencerTrackFilter_Condition::GetName() const
