Sequencer: TimeWarp polish

 - Reinstate yellow scrubber for middle-mouse click
 - Fixed slate brush rendering artifacts
 - Tweaked padding and layout

diff --git a/Engine/Source/Editor/EditorStyle/Private/StarshipStyle.cpp b/Engine/Source/Editor/EditorStyle/Private/StarshipStyle.cpp
--- a/Engine/Source/Editor/EditorStyle/Private/StarshipStyle.cpp
+++ b/Engine/Source/Editor/EditorStyle/Private/StarshipStyle.cpp
@@ -2140,7 +2140,7 @@ void FStarshipEditorStyle::FStyle::SetupSequencerStyles()
 		Set("Sequencer.KeyBar.Dashed", new BORDER_BRUSH(TEXT("Sequencer/Keys/KeyBar_Dashed"), FMargin(0.f, 4.f, 0.f, 0.f)));
 		Set("Sequencer.KeyBar.Solid", new BORDER_BRUSH(TEXT("Sequencer/Keys/KeyBar_Solid"), FMargin(0.f, 4.f, 0.f, 0.f)));
 
-		Set("Sequencer.Timeline.ScrubHandle", new BOX_BRUSH_SVG("Starship/Sequencer/ScrubHandle", FMargin(0.2f)));
+		Set("Sequencer.Timeline.ScrubHandle", new IMAGE_BRUSH_SVG("Starship/Sequencer/ScrubHandle", FVector2f(14.f, 32.f)));
 		Set("Sequencer.Timeline.ScrubFill", new BOX_BRUSH("Sequencer/ScrubFill", FMargin(2.f / 4.f, 0.f)));
 		Set("Sequencer.Timeline.FrameBlockScrubHandle", new BOX_BRUSH("Sequencer/ScrubHandleDown", FMargin(6.f / 13.f, 5 / 12.f, 6 / 13.f, 8 / 12.f)));
 		Set("Sequencer.Timeline.ScrubHandleWhole", new BOX_BRUSH("Sequencer/ScrubHandleWhole", FMargin(6.f / 13.f, 10 / 24.f, 6 / 13.f, 10 / 24.f)));
diff --git a/Engine/Source/Editor/MovieSceneTools/Private/TrackEditors/TimeWarpTrackEditor.cpp b/Engine/Source/Editor/MovieSceneTools/Private/TrackEditors/TimeWarpTrackEditor.cpp
--- a/Engine/Source/Editor/MovieSceneTools/Private/TrackEditors/TimeWarpTrackEditor.cpp
+++ b/Engine/Source/Editor/MovieSceneTools/Private/TrackEditors/TimeWarpTrackEditor.cpp
@@ -197,7 +197,7 @@ struct STimeWarpScrubber : public SLeafWidget
 			FSlateDrawElement::MakeBox(
 				OutDrawElements,
 				LayerId++,
-				AllottedGeometry.ToPaintGeometry(),
+				AllottedGeometry.ToPaintGeometry(AllottedGeometry.Size - FVector2f(0.f, 1.f), FSlateLayoutTransform(FVector2f(0.f, 1.f))),
 				Brush,
 				ESlateDrawEffect::None,
 				TimeWarpColor
diff --git a/Engine/Source/Editor/Sequencer/Private/SequencerTimeSliderController.cpp b/Engine/Source/Editor/Sequencer/Private/SequencerTimeSliderController.cpp
--- a/Engine/Source/Editor/Sequencer/Private/SequencerTimeSliderController.cpp
+++ b/Engine/Source/Editor/Sequencer/Private/SequencerTimeSliderController.cpp
@@ -536,7 +536,9 @@ int32 FSequencerTimeSliderController::OnPaintTimeSlider( bool bMirrorLabels, con
 		const float         HandleEnd     = ScrubMetrics.HandleRangePx.GetUpperBoundValue();
 
 		const int32 ArrowLayer = LayerId + 2;
-		FPaintGeometry MyGeometry =	AllottedGeometry.ToPaintGeometry( FVector2f( HandleEnd - HandleStart, AllottedGeometry.Size.Y ), FSlateLayoutTransform(FVector2f( HandleStart, 0.f )) );
+		FPaintGeometry MyGeometry =	AllottedGeometry.ToPaintGeometry(
+			FVector2f( HandleEnd - HandleStart, AllottedGeometry.Size.Y-1.f ),
+			FSlateLayoutTransform(FVector2f( HandleStart, 1.f )) );
 
 		FTimeWarpTrackExtension*   TimeWarpExtension   = Sequencer->GetViewModel()->GetRootSequenceModel()->GetSharedData()->CastDynamic<FTimeWarpTrackExtension>();
 		const FTimeWarpTrackModel* ActiveTimeWarpTrack = TimeWarpExtension ? TimeWarpExtension->GetActiveTimeWarpTrack() : nullptr;
@@ -546,6 +548,11 @@ int32 FSequencerTimeSliderController::OnPaintTimeSlider( bool bMirrorLabels, con
 			? FStyleColors::AccentOrange.GetSpecifiedColor()
 			: FStyleColors::AccentRed.GetSpecifiedColor();
 
+		if (!bIsEvaluating)
+		{
+			ScrubColor = FStyleColors::AccentYellow.GetSpecifiedColor().CopyWithNewOpacity(.7f);
+		}
+
 		const FSlateBrush* Brush = ScrubMetrics.Style == ESequencerScrubberStyle::Vanilla
 			? VanillaScrubHandleDownBrush
 			: FrameBlockScrubHandleDownBrush;
