[Sequencer] Only unregister filter bar commands when there are no more filter bar instances alive.

diff --git a/Engine/Source/Editor/Sequencer/Private/Filters/SequencerFilterBar.cpp b/Engine/Source/Editor/Sequencer/Private/Filters/SequencerFilterBar.cpp
--- a/Engine/Source/Editor/Sequencer/Private/Filters/SequencerFilterBar.cpp
+++ b/Engine/Source/Editor/Sequencer/Private/Filters/SequencerFilterBar.cpp
@@ -39,6 +39,8 @@ using namespace UE::Sequencer;
 
 const FName FSequencerFilterBar::SharedIdentifier = TEXT("SharedSequencerTrackFilter");
 
+int32 FSequencerFilterBar::InstanceCount = 0;
+
 FSequencerFilterBar::FSequencerFilterBar(FSequencer& InSequencer)
 	: Sequencer(InSequencer)
 	, CommandList(MakeShared<FUICommandList>())
@@ -57,6 +59,8 @@ FSequencerFilterBar::FSequencerFilterBar(FSequencer& InSequencer)
 	, FilterMenu(MakeShared<FSequencerTrackFilterMenu>())
 	, FilterData(FString())
 {
+	InstanceCount++;
+
 	FSequencerTrackFilterCommands::Register();
 
 	CommonFilters->OnChanged().AddRaw(this, &FSequencerFilterBar::RequestFilterUpdate);
@@ -71,8 +75,13 @@ FSequencerFilterBar::FSequencerFilterBar(FSequencer& InSequencer)
 
 FSequencerFilterBar::~FSequencerFilterBar()
 {
-	FSequencerTrackFilterCommands::Unregister();
-	
+	InstanceCount--;
+
+	if (InstanceCount == 0)
+	{
+		FSequencerTrackFilterCommands::Unregister();
+	}
+
 	CommonFilters->OnChanged().RemoveAll(this);
     InternalFilters->OnChanged().RemoveAll(this);
     TextFilter->OnChanged().RemoveAll(this);
diff --git a/Engine/Source/Editor/Sequencer/Private/Filters/SequencerFilterBar.h b/Engine/Source/Editor/Sequencer/Private/Filters/SequencerFilterBar.h
--- a/Engine/Source/Editor/Sequencer/Private/Filters/SequencerFilterBar.h
+++ b/Engine/Source/Editor/Sequencer/Private/Filters/SequencerFilterBar.h
@@ -272,6 +272,8 @@ protected:
 	FSequencerFilterData FilterData;
 
 private:
+	static int32 InstanceCount;
+
 	/** Do not call directly! Should only be called by FilterNodes(). */
 	bool FilterNodesRecursive(const bool bInHasActiveFilter, const UE::Sequencer::TViewModelPtr<UE::Sequencer::IOutlinerExtension>& InStartNode);
 };
