[Sequencer] Selectability: Fix indices that are checked for zero planes on frustum select

diff --git a/Engine/Source/Editor/Sequencer/Private/SequencerSelectabilityTool.cpp b/Engine/Source/Editor/Sequencer/Private/SequencerSelectabilityTool.cpp
--- a/Engine/Source/Editor/Sequencer/Private/SequencerSelectabilityTool.cpp
+++ b/Engine/Source/Editor/Sequencer/Private/SequencerSelectabilityTool.cpp
@@ -35,9 +35,9 @@ bool FSequencerSelectabilityTool::FrustumSelect(const FConvexVolume& InFrustum,
 	}
 
 	// Need to check for a zero frustum since ComponentIsTouchingSelectionFrustum will return true, selecting everything, when this is the case
-	const bool bMalformedFrustum = (InFrustum.Planes[0].IsNearlyZero() && InFrustum.Planes[2].IsNearlyZero())
-		|| (InFrustum.Planes[3].IsNearlyZero() && InFrustum.Planes[4].IsNearlyZero());
-	if (bMalformedFrustum)
+	const bool bAreTopBottomMalformed = InFrustum.Planes[0].IsNearlyZero() && InFrustum.Planes[2].IsNearlyZero();
+	const bool bAreRightLeftMalformed = InFrustum.Planes[1].IsNearlyZero() && InFrustum.Planes[3].IsNearlyZero();
+	if (bAreTopBottomMalformed || bAreRightLeftMalformed)
 	{
 		return false;
 	}
