Sequencer: Rename OutlinerDecorator back to OutlinerIndicator

#rnx

diff --git a/Engine/Source/Editor/EditorStyle/Private/StarshipStyle.cpp b/Engine/Source/Editor/EditorStyle/Private/StarshipStyle.cpp
--- a/Engine/Source/Editor/EditorStyle/Private/StarshipStyle.cpp
+++ b/Engine/Source/Editor/EditorStyle/Private/StarshipStyle.cpp
@@ -2245,7 +2245,7 @@ void FStarshipEditorStyle::FStyle::SetupSequencerStyles()
 		Set("Sequencer.AnimationOutliner.RegularFont", DEFAULT_FONT("Regular", 9));
 		Set("Sequencer.AnimationOutliner.ItalicFont", DEFAULT_FONT("Italic", 10));
 
-		Set("Sequencer.Outliner.Decorators.TimeWarp", new IMAGE_BRUSH_SVG("Starship/Sequencer/TimeWarp_12", Icon12x12));
+		Set("Sequencer.Outliner.Indicators.TimeWarp", new IMAGE_BRUSH_SVG("Starship/Sequencer/TimeWarp_12", Icon12x12));
 
 		Set("Sequencer.ShotFilter", new IMAGE_BRUSH("Sequencer/FilteredArea", FVector2D(74, 74), FLinearColor::White, ESlateBrushTileType::Both));
 		Set("Sequencer.KeyMark", new IMAGE_BRUSH("Sequencer/KeyMark", FVector2D(3, 21), FLinearColor::White, ESlateBrushTileType::NoTile));
@@ -2419,8 +2419,8 @@ void FStarshipEditorStyle::FStyle::SetupSequencerStyles()
 		Set("Sequencer.Column.Unpinned", new IMAGE_BRUSH_SVG("Sequencer/Column_Widgets/SequencerUnpinned", Icon14x14));
 		Set("Sequencer.Column.CheckBoxIndeterminate", new CORE_IMAGE_BRUSH_SVG("Starship/CoreWidgets/CheckBox/CheckBoxIndeterminate_14", Icon14x14));
 		Set("Sequencer.Column.OutlinerColumnBox", new FSlateRoundedBoxBrush(FStyleColors::Header, 2.0f));
-		Set("Sequencer.Decorator.Condition", new IMAGE_BRUSH_SVG("Sequencer/Column_Widgets/SequencerCondition", Icon12x12));
-		Set("Sequencer.Decorator.TimeWarp", new IMAGE_BRUSH_SVG("Sequencer/Column_Widgets/SequencerTimeWarp", Icon12x12));
+		Set("Sequencer.Indicator.Condition", new IMAGE_BRUSH_SVG("Sequencer/Column_Widgets/SequencerCondition", Icon12x12));
+		Set("Sequencer.Indicator.TimeWarp", new IMAGE_BRUSH_SVG("Sequencer/Column_Widgets/SequencerTimeWarp", Icon12x12));
 	}
 }
 
diff --git a/Engine/Source/Editor/MovieSceneTools/Private/MVVM/ViewModels/TimeWarpChannelModel.cpp b/Engine/Source/Editor/MovieSceneTools/Private/MVVM/ViewModels/TimeWarpChannelModel.cpp
--- a/Engine/Source/Editor/MovieSceneTools/Private/MVVM/ViewModels/TimeWarpChannelModel.cpp
+++ b/Engine/Source/Editor/MovieSceneTools/Private/MVVM/ViewModels/TimeWarpChannelModel.cpp
@@ -28,7 +28,7 @@ TSharedPtr<SWidget> FTimeWarpChannelModel::CreateOutlinerViewForColumn(const FCr
 	{
 		return MakeButton(
 			LOCTEXT("ChangeTimeWarpToolTip", "Change Time Warp to utilize a different curve type"),
-			FAppStyle::GetBrush("Sequencer.Outliner.Decorators.TimeWarp"),
+			FAppStyle::GetBrush("Sequencer.Outliner.Indicator.TimeWarp"),
 			FOnGetContent::CreateSP(this, &FTimeWarpChannelModel::BuildReplaceTimeWarpSubMenu),
 			OutlinerItem.AsModel());
 	}
diff --git a/Engine/Source/Editor/Sequencer/Private/Filters/Filters/SequencerTrackFilter_Condition.cpp b/Engine/Source/Editor/Sequencer/Private/Filters/Filters/SequencerTrackFilter_Condition.cpp
--- a/Engine/Source/Editor/Sequencer/Private/Filters/Filters/SequencerTrackFilter_Condition.cpp
+++ b/Engine/Source/Editor/Sequencer/Private/Filters/Filters/SequencerTrackFilter_Condition.cpp
@@ -37,7 +37,7 @@ FText FSequencerTrackFilter_Condition::GetDisplayName() const
 
 FSlateIcon FSequencerTrackFilter_Condition::GetIcon() const
 {
-	return FSlateIcon(FAppStyle::GetAppStyleSetName(), TEXT("Sequencer.Decorator.Condition"));
+	return FSlateIcon(FAppStyle::GetAppStyleSetName(), TEXT("Sequencer.Indicator.Condition"));
 }
 
 FString FSequencerTrackFilter_Condition::GetName() const
diff --git a/Engine/Source/Editor/Sequencer/Private/MVVM/ViewModels/OutlinerColumns/OutlinerDecoratorColumn.cpp b/Engine/Source/Editor/Sequencer/Private/MVVM/ViewModels/OutlinerColumns/OutlinerDecoratorColumn.cpp
deleted file mode 100644
diff --git a/Engine/Source/Editor/Sequencer/Private/MVVM/ViewModels/OutlinerColumns/OutlinerDecoratorColumn.h b/Engine/Source/Editor/Sequencer/Private/MVVM/ViewModels/OutlinerColumns/OutlinerDecoratorColumn.h
deleted file mode 100644
diff --git a/Engine/Source/Editor/Sequencer/Private/MVVM/ViewModels/OutlinerColumns/OutlinerIndicatorColumn.cpp b/Engine/Source/Editor/Sequencer/Private/MVVM/ViewModels/OutlinerColumns/OutlinerIndicatorColumn.cpp
new file mode 100644
--- /dev/null
+++ b/Engine/Source/Editor/Sequencer/Private/MVVM/ViewModels/OutlinerColumns/OutlinerIndicatorColumn.cpp
@@ -0,0 +1,95 @@
+// Copyright Epic Games, Inc. All Rights Reserved.
+
+#include "MVVM/ViewModels/OutlinerColumns/OutlinerIndicatorColumn.h"
+
+#include "Sequencer.h"
+#include "MVVM/Extensions/IConditionableExtension.h"
+#include "MVVM/SharedViewModelData.h"
+#include "MVVM/ViewModels/OutlinerIndicators/IOutlinerIndicatorBuilder.h"
+#include "MVVM/ViewModels/SequencerEditorViewModel.h"
+#include "Widgets/Layout/SBorder.h"
+#include "Widgets/OutlinerIndicators/SConditionIndicatorWidget.h"
+#include "Widgets/SBoxPanel.h"
+
+#define LOCTEXT_NAMESPACE "FOutlinerIndicatorColumn"
+
+namespace UE::Sequencer
+{
+
+FOutlinerIndicatorColumn::FOutlinerIndicatorColumn()
+{
+	Name = FCommonOutlinerNames::Indicator;
+	Label = LOCTEXT("IndicatorColumnLabel", "Indicators");
+	Position = FOutlinerColumnPosition{ 0, EOutlinerColumnGroup::FarLeftGutter };
+	Layout = FOutlinerColumnLayout{ 14, FMargin(0.f, 0.f), HAlign_Fill, VAlign_Fill, EOutlinerColumnSizeMode::Fixed, EOutlinerColumnFlags::None };
+}
+
+bool FOutlinerIndicatorColumn::IsItemCompatibleWithColumn(const FCreateOutlinerColumnParams& InParams) const
+{
+	TViewModelPtr<FSequencerEditorViewModel> Editor = InParams.Editor->CastThisShared<FSequencerEditorViewModel>();
+	if (!Editor)
+	{
+		return false;
+	}
+
+	TSharedPtr<FSequencer> Sequencer = Editor->GetSequencerImpl();
+	if (!Sequencer)
+	{
+		return false;
+	}
+
+	for (const TTuple< FName, TSharedPtr<IOutlinerIndicatorBuilder> >& OutlinerIndicator : Sequencer->GetOutlinerIndicators())
+	{
+		if (OutlinerIndicator.Value.Get()->IsItemCompatibleWithIndicator(InParams))
+		{
+			return true;
+		}
+	}
+
+	return false;
+}
+
+TSharedPtr<SWidget> FOutlinerIndicatorColumn::CreateColumnWidget(const FCreateOutlinerColumnParams& InParams, const TSharedRef<ISequencerTreeViewRow>& TreeViewRow)
+{
+	TViewModelPtr<FSequencerEditorViewModel> Editor = InParams.Editor->CastThisShared<FSequencerEditorViewModel>();
+	if (!Editor)
+	{
+		return nullptr;
+	}
+
+	TSharedPtr<FSequencer> Sequencer = Editor->GetSequencerImpl();
+	if (!Sequencer)
+	{
+		return nullptr;
+	}
+
+	ColumnWidget = SNew(SHorizontalBox);
+
+	TArray<TSharedPtr<IOutlinerIndicatorBuilder> > CompatibleIndicators;
+
+	for (const TTuple< FName, TSharedPtr<IOutlinerIndicatorBuilder> >& OutlinerIndicator : Sequencer->GetOutlinerIndicators())
+	{
+		TSharedPtr<IOutlinerIndicatorBuilder> Indicator = OutlinerIndicator.Value;
+
+		if (Indicator->IsItemCompatibleWithIndicator(InParams))
+		{
+			CompatibleIndicators.Add(Indicator);
+		}
+	}
+
+	const int32 NumCompatibleIndicators = CompatibleIndicators.Num();
+
+	for (TSharedPtr<IOutlinerIndicatorBuilder>& Indicator : CompatibleIndicators)
+	{
+		ColumnWidget->AddSlot()
+		[
+			Indicator->CreateIndicatorWidget(InParams, TreeViewRow, SharedThis(this), NumCompatibleIndicators).ToSharedRef()
+		];
+	}
+
+	return ColumnWidget;
+}
+
+} // namespace UE::Sequencer
+
+#undef LOCTEXT_NAMESPACE
\ No newline at end of file
diff --git a/Engine/Source/Editor/Sequencer/Private/MVVM/ViewModels/OutlinerColumns/OutlinerIndicatorColumn.h b/Engine/Source/Editor/Sequencer/Private/MVVM/ViewModels/OutlinerColumns/OutlinerIndicatorColumn.h
new file mode 100644
--- /dev/null
+++ b/Engine/Source/Editor/Sequencer/Private/MVVM/ViewModels/OutlinerColumns/OutlinerIndicatorColumn.h
@@ -0,0 +1,29 @@
+// Copyright Epic Games, Inc. All Rights Reserved.
+
+#pragma once
+
+#include "MVVM/ViewModels/OutlinerColumns/OutlinerColumnBase.h"
+
+class SHorizontalBox;
+
+namespace UE::Sequencer
+{
+
+/**
+ * A column for showing various decorators on the presence of features (e.g. conditions, time warp) on that row.
+ */
+class FOutlinerIndicatorColumn
+	: public FOutlinerColumnBase
+{
+public:
+
+	FOutlinerIndicatorColumn();
+
+	bool IsItemCompatibleWithColumn(const FCreateOutlinerColumnParams& InParams) const override;
+	TSharedPtr<SWidget> CreateColumnWidget(const FCreateOutlinerColumnParams& InParams, const TSharedRef<ISequencerTreeViewRow>& TreeViewRow) override;
+
+private:
+	TSharedPtr<SHorizontalBox> ColumnWidget;
+};
+
+} // namespace UE::Sequencer
\ No newline at end of file
diff --git a/Engine/Source/Editor/Sequencer/Private/MVVM/ViewModels/OutlinerDecorators/ConditionOutlinerDecoratorBuilder.cpp b/Engine/Source/Editor/Sequencer/Private/MVVM/ViewModels/OutlinerDecorators/ConditionOutlinerDecoratorBuilder.cpp
deleted file mode 100644
diff --git a/Engine/Source/Editor/Sequencer/Private/MVVM/ViewModels/OutlinerDecorators/ConditionOutlinerDecoratorBuilder.h b/Engine/Source/Editor/Sequencer/Private/MVVM/ViewModels/OutlinerDecorators/ConditionOutlinerDecoratorBuilder.h
deleted file mode 100644
diff --git a/Engine/Source/Editor/Sequencer/Private/MVVM/ViewModels/OutlinerDecorators/TimeWarpOutlinerDecoratorBuilder.cpp b/Engine/Source/Editor/Sequencer/Private/MVVM/ViewModels/OutlinerDecorators/TimeWarpOutlinerDecoratorBuilder.cpp
deleted file mode 100644
diff --git a/Engine/Source/Editor/Sequencer/Private/MVVM/ViewModels/OutlinerDecorators/TimeWarpOutlinerDecoratorBuilder.h b/Engine/Source/Editor/Sequencer/Private/MVVM/ViewModels/OutlinerDecorators/TimeWarpOutlinerDecoratorBuilder.h
deleted file mode 100644
diff --git a/Engine/Source/Editor/Sequencer/Private/MVVM/ViewModels/OutlinerIndicators/ConditionOutlinerIndicatorBuilder.cpp b/Engine/Source/Editor/Sequencer/Private/MVVM/ViewModels/OutlinerIndicators/ConditionOutlinerIndicatorBuilder.cpp
new file mode 100644
--- /dev/null
+++ b/Engine/Source/Editor/Sequencer/Private/MVVM/ViewModels/OutlinerIndicators/ConditionOutlinerIndicatorBuilder.cpp
@@ -0,0 +1,84 @@
+// Copyright Epic Games, Inc. All Rights Reserved.
+
+#include "MVVM/ViewModels/OutlinerIndicators/ConditionOutlinerIndicatorBuilder.h"
+
+#include "MVVM/ViewModels/OutlinerColumns/IOutlinerColumn.h"
+#include "MVVM/ViewModels/OutlinerColumns/OutlinerColumnTypes.h"
+#include "Widgets/OutlinerIndicators/SConditionIndicatorWidget.h"
+#include "MVVM/Extensions/IConditionableExtension.h"
+#include "MVVM/SharedViewModelData.h"
+#include "Widgets/Layout/SBorder.h"
+
+#define LOCTEXT_NAMESPACE "FConditionOutlinerIndicatorBuilder"
+
+namespace UE::Sequencer
+{
+
+FConditionOutlinerIndicatorBuilder::FConditionOutlinerIndicatorBuilder()
+{
+}
+
+FName FConditionOutlinerIndicatorBuilder::GetIndicatorName() const
+{
+	return FCommonOutlinerNames::Condition;
+}
+
+bool FConditionOutlinerIndicatorBuilder::IsItemCompatibleWithIndicator(const FCreateOutlinerColumnParams& InParams) const
+{
+	if (FConditionStateCacheExtension* ConditionStateCache = InParams.OutlinerExtension.AsModel()->GetSharedData()->CastThis<FConditionStateCacheExtension>())
+	{
+		return EnumHasAnyFlags(ConditionStateCache->GetCachedFlags(InParams.OutlinerExtension), ECachedConditionState::HasCondition | ECachedConditionState::ParentHasCondition | ECachedConditionState::ChildHasCondition | ECachedConditionState::SectionHasCondition);
+	}
+
+	return false;
+}
+
+TSharedPtr<SWidget> FConditionOutlinerIndicatorBuilder::CreateIndicatorWidget(const FCreateOutlinerColumnParams& InParams, const TSharedRef<ISequencerTreeViewRow>& TreeViewRow, const TSharedRef<IOutlinerColumn>& OutlinerColumn, const int32 NumCompatibleIndicators)
+{
+	static const FLinearColor ConditionColor = FLinearColor::FromSRGBColor(FColor(92, 220, 205));
+
+	TSharedRef<SConditionIndicatorWidget> ConditionIndicatorWidget = SNew(SConditionIndicatorWidget, OutlinerColumn, InParams);
+	if (NumCompatibleIndicators > 1)
+	{
+		return SNew(SBorder)
+			.VAlign(VAlign_Fill)
+			.HAlign(HAlign_Fill)
+			.Padding(0.0f)
+			.BorderImage(FAppStyle::GetBrush("WhiteBrush"))
+			.BorderBackgroundColor(ConditionColor);
+	}
+
+	return SNew(SOverlay)
+		+ SOverlay::Slot()
+		[
+			SNew(SBorder)
+				.VAlign(VAlign_Fill)
+				.HAlign(HAlign_Fill)
+				.Padding(0.0f)
+				.BorderImage(FAppStyle::GetBrush("WhiteBrush"))
+				.BorderBackgroundColor_Lambda([WeakWidget = ConditionIndicatorWidget.ToWeakPtr()]() -> FLinearColor
+					{
+						if (WeakWidget.IsValid())
+						{
+							return WeakWidget.Pin()->GetIndicatorBackgroundColorAndOpacity().GetSpecifiedColor();
+						}
+						return FLinearColor(0, 0, 0, 0);
+					})
+		]
+
+		+ SOverlay::Slot()
+		[
+			SNew(SBox)
+				.WidthOverride(12.f)
+				.HeightOverride(12.f)
+				.HAlign(HAlign_Center)
+				.VAlign(VAlign_Center)
+				[
+					ConditionIndicatorWidget
+				]
+		];
+}
+
+} // namespace UE::Sequencer
+
+#undef LOCTEXT_NAMESPACE
\ No newline at end of file
diff --git a/Engine/Source/Editor/Sequencer/Private/MVVM/ViewModels/OutlinerIndicators/ConditionOutlinerIndicatorBuilder.h b/Engine/Source/Editor/Sequencer/Private/MVVM/ViewModels/OutlinerIndicators/ConditionOutlinerIndicatorBuilder.h
new file mode 100644
--- /dev/null
+++ b/Engine/Source/Editor/Sequencer/Private/MVVM/ViewModels/OutlinerIndicators/ConditionOutlinerIndicatorBuilder.h
@@ -0,0 +1,26 @@
+// Copyright Epic Games, Inc. All Rights Reserved.
+
+#pragma once
+
+#include "MVVM/ViewModels/OutlinerIndicators/OutlinerIndicatorBuilderBase.h"
+
+namespace UE::Sequencer
+{
+
+struct FCreateOutlinerColumnParams;
+class IOutlinerColumn;
+class SConditionIndicatorWidget;
+
+class FConditionOutlinerIndicatorBuilder
+	: public FOutlinerIndicatorBuilderBase
+{
+public:
+	FConditionOutlinerIndicatorBuilder();
+	
+	virtual FName GetIndicatorName() const override;
+	virtual bool IsItemCompatibleWithIndicator(const FCreateOutlinerColumnParams& InParams) const override;
+	virtual TSharedPtr<SWidget> CreateIndicatorWidget(const FCreateOutlinerColumnParams& InParams, const TSharedRef<ISequencerTreeViewRow>& TreeViewRow, const TSharedRef<IOutlinerColumn>& OutlinerColumn, const int32 NumCompatibleIndicators) override;
+};
+
+} // namespace UE::Sequencer
+
diff --git a/Engine/Source/Editor/Sequencer/Private/MVVM/ViewModels/OutlinerIndicators/TimeWarpOutlinerIndicatorBuilder.cpp b/Engine/Source/Editor/Sequencer/Private/MVVM/ViewModels/OutlinerIndicators/TimeWarpOutlinerIndicatorBuilder.cpp
new file mode 100644
--- /dev/null
+++ b/Engine/Source/Editor/Sequencer/Private/MVVM/ViewModels/OutlinerIndicators/TimeWarpOutlinerIndicatorBuilder.cpp
@@ -0,0 +1,97 @@
+// Copyright Epic Games, Inc. All Rights Reserved.
+
+#include "MVVM/ViewModels/OutlinerIndicators/TimeWarpOutlinerIndicatorBuilder.h"
+
+#include "MovieSceneSection.h"
+#include "MVVM/Extensions/ITrackExtension.h"
+#include "MVVM/SharedViewModelData.h"
+#include "MVVM/ViewModels/SectionModel.h"
+#include "MVVM/ViewModels/OutlinerColumns/IOutlinerColumn.h"
+#include "MVVM/ViewModels/OutlinerColumns/OutlinerColumnTypes.h"
+#include "Variants/MovieSceneTimeWarpVariant.h"
+#include "Variants/MovieSceneTimeWarpGetter.h"
+#include "Widgets/Layout/SBorder.h"
+#include "Widgets/OutlinerIndicators/STimeWarpIndicatorWidget.h"
+
+#define LOCTEXT_NAMESPACE "FTimeWarpOutlinerIndicatorBuilder"
+
+namespace UE::Sequencer
+{
+
+FTimeWarpOutlinerIndicatorBuilder::FTimeWarpOutlinerIndicatorBuilder()
+{
+}
+
+FName FTimeWarpOutlinerIndicatorBuilder::GetIndicatorName() const
+{
+	return FCommonOutlinerNames::TimeWarp;
+}
+
+bool FTimeWarpOutlinerIndicatorBuilder::IsItemCompatibleWithIndicator(const FCreateOutlinerColumnParams& InParams) const
+{
+	if (InParams.OutlinerExtension.AsModel()->IsA<ITrackExtension>())
+	{
+		ITrackExtension* TrackModel = InParams.OutlinerExtension.AsModel()->CastThis<ITrackExtension>();
+
+		for (const TViewModelPtr<FSectionModel>& SectionModel : TrackModel->GetSectionModels().IterateSubList<FSectionModel>())
+		{
+			UMovieSceneSection* Section = SectionModel->GetSection();
+			if (!Section)
+			{
+				continue;
+			}
+
+			if (FMovieSceneTimeWarpVariant* Variant = Section->GetTimeWarp())
+			{
+				if (Variant->GetType() == EMovieSceneTimeWarpType::Custom)
+				{
+					return true;
+				}
+			}
+		}
+	}
+
+	return false;
+}
+
+TSharedPtr<SWidget> FTimeWarpOutlinerIndicatorBuilder::CreateIndicatorWidget(const FCreateOutlinerColumnParams& InParams, const TSharedRef<ISequencerTreeViewRow>& TreeViewRow, const TSharedRef<IOutlinerColumn>& OutlinerColumn, const int32 NumCompatibleIndicators)
+{
+	static const FLinearColor ConditionColor = FLinearColor::FromSRGBColor(FColor(212, 147, 20));
+
+	if (NumCompatibleIndicators > 1)
+	{
+		return SNew(SBorder)
+			.VAlign(VAlign_Fill)
+			.HAlign(HAlign_Fill)
+			.Padding(0.0f)
+			.BorderImage(FAppStyle::GetBrush("WhiteBrush"))
+			.BorderBackgroundColor_Lambda([this]() -> FLinearColor {  return ConditionColor; });
+	}
+
+	return SNew(SOverlay)
+		+ SOverlay::Slot()
+		[
+			SNew(SBorder)
+				.VAlign(VAlign_Fill)
+				.HAlign(HAlign_Fill)
+				.Padding(0.0f)
+				.BorderImage(FAppStyle::GetBrush("WhiteBrush"))
+				.BorderBackgroundColor_Lambda([this]() -> FLinearColor { return ConditionColor; })
+		]
+
+		+ SOverlay::Slot()
+		[
+			SNew(SBox)
+				.WidthOverride(12.f)
+				.HeightOverride(12.f)
+				.HAlign(HAlign_Center)
+				.VAlign(VAlign_Center)
+				[
+					SNew(STimeWarpIndicatorWidget, OutlinerColumn, InParams)
+				]
+		];
+}
+
+} // namespace UE::Sequencer
+
+#undef LOCTEXT_NAMESPACE
\ No newline at end of file
diff --git a/Engine/Source/Editor/Sequencer/Private/MVVM/ViewModels/OutlinerIndicators/TimeWarpOutlinerIndicatorBuilder.h b/Engine/Source/Editor/Sequencer/Private/MVVM/ViewModels/OutlinerIndicators/TimeWarpOutlinerIndicatorBuilder.h
new file mode 100644
--- /dev/null
+++ b/Engine/Source/Editor/Sequencer/Private/MVVM/ViewModels/OutlinerIndicators/TimeWarpOutlinerIndicatorBuilder.h
@@ -0,0 +1,26 @@
+// Copyright Epic Games, Inc. All Rights Reserved.
+
+#pragma once
+
+#include "MVVM/ViewModels/OutlinerIndicators/OutlinerIndicatorBuilderBase.h"
+
+namespace UE::Sequencer
+{
+
+struct FCreateOutlinerColumnParams;
+class IOutlinerColumn;
+class STimeWarpIndicatorWidget;
+
+class FTimeWarpOutlinerIndicatorBuilder
+	: public FOutlinerIndicatorBuilderBase
+{
+public:
+	FTimeWarpOutlinerIndicatorBuilder();
+	
+	virtual FName GetIndicatorName() const override;
+	virtual bool IsItemCompatibleWithIndicator(const FCreateOutlinerColumnParams& InParams) const override;
+	virtual TSharedPtr<SWidget> CreateIndicatorWidget(const FCreateOutlinerColumnParams& InParams, const TSharedRef<ISequencerTreeViewRow>& TreeViewRow, const TSharedRef<IOutlinerColumn>& OutlinerColumn, const int32 NumCompatibleIndicators) override;
+};
+
+} // namespace UE::Sequencer
+
diff --git a/Engine/Source/Editor/Sequencer/Private/Sequencer.cpp b/Engine/Source/Editor/Sequencer/Private/Sequencer.cpp
--- a/Engine/Source/Editor/Sequencer/Private/Sequencer.cpp
+++ b/Engine/Source/Editor/Sequencer/Private/Sequencer.cpp
@@ -132,7 +132,7 @@
 #include "ISequencerChannelInterface.h"
 #include "IMovieRendererInterface.h"
 #include "MVVM/ViewModels/OutlinerColumns/IOutlinerColumn.h"
-#include "MVVM/ViewModels/OutlinerDecorators/IOutlinerDecoratorBuilder.h"
+#include "MVVM/ViewModels/OutlinerIndicators/IOutlinerIndicatorBuilder.h"
 #include "SequencerKeyCollection.h"
 #include "CurveEditor.h"
 #include "CurveEditorScreenSpace.h"
@@ -305,7 +305,7 @@ void FSequencer::InitSequencer(const FSequencerInitParams& InitParams
 	, const TArray<FOnCreateTrackEditor>& TrackEditorDelegates
 	, const TArray<FOnCreateEditorObjectBinding>& EditorObjectBindingDelegates
 	, const TArray<FOnCreateOutlinerColumn>& OutlinerColumnDelegates
-	, const TArray<FOnCreateOutlinerDecorator>& OutlinerDecoratorDelegates)
+	, const TArray<FOnCreateOutlinerIndicator>& OutlinerIndicatorDelegates)
 {
 	using namespace UE::MovieScene;
 	using namespace UE::Sequencer;
@@ -531,16 +531,16 @@ void FSequencer::InitSequencer(const FSequencerInitParams& InitParams
 		}
 	}
 
-	// OutlinerDecorators are registered to be provided to SOutlinerView by SSequencer
-	for (int32 DelegateIndex = 0; DelegateIndex < OutlinerDecoratorDelegates.Num(); ++DelegateIndex)
+	// OutlinerIndicators are registered to be provided to SOutlinerView by SSequencer
+	for (int32 DelegateIndex = 0; DelegateIndex < OutlinerIndicatorDelegates.Num(); ++DelegateIndex)
 	{
-		check(OutlinerDecoratorDelegates[DelegateIndex].IsBound());
-		TSharedRef<IOutlinerDecoratorBuilder> OutlinerDecorator = OutlinerDecoratorDelegates[DelegateIndex].Execute();
+		check(OutlinerIndicatorDelegates[DelegateIndex].IsBound());
+		TSharedRef<IOutlinerIndicatorBuilder> OutlinerIndicator = OutlinerIndicatorDelegates[DelegateIndex].Execute();
 
-		if (OutlinerDecorator->SupportsSequence(InitParams.RootSequence))
+		if (OutlinerIndicator->SupportsSequence(InitParams.RootSequence))
 		{
-			check(!OutlinerDecorators.Contains(OutlinerDecorator->GetDecoratorName()));
-			OutlinerDecorators.Add(OutlinerDecorator->GetDecoratorName(), OutlinerDecorator);
+			check(!OutlinerIndicators.Contains(OutlinerIndicator->GetIndicatorName()));
+			OutlinerIndicators.Add(OutlinerIndicator->GetIndicatorName(), OutlinerIndicator);
 		}
 	}
 
diff --git a/Engine/Source/Editor/Sequencer/Private/Sequencer.h b/Engine/Source/Editor/Sequencer/Private/Sequencer.h
--- a/Engine/Source/Editor/Sequencer/Private/Sequencer.h
+++ b/Engine/Source/Editor/Sequencer/Private/Sequencer.h
@@ -141,7 +141,7 @@ public:
 	 * @param TrackEditorDelegates Delegates to call to create auto-key handlers for this sequencer.
 	 * @param EditorObjectBindingDelegates Delegates to call to create object bindings for this sequencer.
 	 */
-	void InitSequencer(const FSequencerInitParams& InitParams, const TSharedRef<ISequencerObjectChangeListener>& InObjectChangeListener, const TArray<FOnCreateTrackEditor>& TrackEditorDelegates, const TArray<FOnCreateEditorObjectBinding>& EditorObjectBindingDelegatess, const TArray<FOnCreateOutlinerColumn>& OutlinerColumnDelegates, const TArray<FOnCreateOutlinerDecorator>& OutlinerDecoratorDelegates);
+	void InitSequencer(const FSequencerInitParams& InitParams, const TSharedRef<ISequencerObjectChangeListener>& InObjectChangeListener, const TArray<FOnCreateTrackEditor>& TrackEditorDelegates, const TArray<FOnCreateEditorObjectBinding>& EditorObjectBindingDelegatess, const TArray<FOnCreateOutlinerColumn>& OutlinerColumnDelegates, const TArray<FOnCreateOutlinerIndicator>& OutlinerIndicatorDelegates);
 
 	/**
 	 * Reinitializes sequencer after the playback context has changed
@@ -316,11 +316,11 @@ public:
 	}
 
 	/**
-	* @return Outliner Decorator Items registered to the sequencer by column name
+	* @return Outliner Indicator Items registered to the sequencer by column name
 	*/
-	const TMap<FName, TSharedPtr<UE::Sequencer::IOutlinerDecoratorBuilder>>& GetOutlinerDecorators() const
+	const TMap<FName, TSharedPtr<UE::Sequencer::IOutlinerIndicatorBuilder>>& GetOutlinerIndicators() const
 	{
-		return OutlinerDecorators;
+		return OutlinerIndicators;
 	}
 
 public:
@@ -1279,8 +1279,8 @@ private:
 	/** List of Outliner column creators that are supported by the Sequencer. */
 	TMap<FName, TSharedPtr<UE::Sequencer::IOutlinerColumn>> OutlinerColumns;
 
-	/** List of Outliner decorator item creators that are supported by the Sequencer. */
-	TMap<FName, TSharedPtr<UE::Sequencer::IOutlinerDecoratorBuilder>> OutlinerDecorators;
+	/** List of Outliner indicator item creators that are supported by the Sequencer. */
+	TMap<FName, TSharedPtr<UE::Sequencer::IOutlinerIndicatorBuilder>> OutlinerIndicators;
 
 	/** List of object bindings we can use */
 	TArray<TSharedPtr<ISequencerEditorObjectBinding>> ObjectBindings;
diff --git a/Engine/Source/Editor/Sequencer/Private/SequencerModule.cpp b/Engine/Source/Editor/Sequencer/Private/SequencerModule.cpp
--- a/Engine/Source/Editor/Sequencer/Private/SequencerModule.cpp
+++ b/Engine/Source/Editor/Sequencer/Private/SequencerModule.cpp
@@ -33,7 +33,7 @@
 #include "MVVM/ViewModels/SequenceModel.h"
 #include "MVVM/ViewModels/SequencerEditorViewModel.h"
 
-#include "MVVM/ViewModels/OutlinerColumns/OutlinerDecoratorColumn.h"
+#include "MVVM/ViewModels/OutlinerColumns/OutlinerIndicatorColumn.h"
 #include "MVVM/ViewModels/OutlinerColumns/LockOutlinerColumn.h"
 #include "MVVM/ViewModels/OutlinerColumns/MuteOutlinerColumn.h"
 #include "MVVM/ViewModels/OutlinerColumns/PinOutlinerColumn.h"
@@ -45,8 +45,8 @@
 #include "MVVM/ViewModels/OutlinerColumns/KeyFrameOutlinerColumn.h"
 #include "MVVM/ViewModels/OutlinerColumns/ColorPickerOutlinerColumn.h"
 
-#include "MVVM/ViewModels/OutlinerDecorators/ConditionOutlinerDecoratorBuilder.h"
-#include "MVVM/ViewModels/OutlinerDecorators/TimeWarpOutlinerDecoratorBuilder.h"
+#include "MVVM/ViewModels/OutlinerIndicators/ConditionOutlinerIndicatorBuilder.h"
+#include "MVVM/ViewModels/OutlinerIndicators/TimeWarpOutlinerIndicatorBuilder.h"
 
 #include "ToolMenus.h"
 #include "ContentBrowserMenuContexts.h"
@@ -297,7 +297,7 @@ public:
 
 		OnPreSequencerInit.Broadcast(Sequencer, ObjectChangeListener, InitParams);
 
-		Sequencer->InitSequencer(InitParams, ObjectChangeListener, TrackEditorDelegates, EditorObjectBindingDelegates, OutlinerColumnDelegates, OutlinerDecoratorDelegates);
+		Sequencer->InitSequencer(InitParams, ObjectChangeListener, TrackEditorDelegates, EditorObjectBindingDelegates, OutlinerColumnDelegates, OutlinerIndicatorDelegates);
 
 		OnSequencerCreated.Broadcast(Sequencer);
 
@@ -360,13 +360,13 @@ public:
 		OutlinerColumnDelegates.RemoveAll([=](const FOnCreateOutlinerColumn& Delegate) { return Delegate.GetHandle() == InHandle; });
 	}
 
-	virtual FDelegateHandle RegisterOutlinerDecorator(FOnCreateOutlinerDecorator InCreator) override {
-		OutlinerDecoratorDelegates.Add(InCreator);
-		return OutlinerDecoratorDelegates.Last().GetHandle();
+	virtual FDelegateHandle RegisterOutlinerIndicator(FOnCreateOutlinerIndicator InCreator) override {
+		OutlinerIndicatorDelegates.Add(InCreator);
+		return OutlinerIndicatorDelegates.Last().GetHandle();
 	}
 
-	virtual void UnregisterOutlinerDecorator(FDelegateHandle InHandle) override {
-		OutlinerDecoratorDelegates.RemoveAll([=](const FOnCreateOutlinerDecorator& Delegate) { return Delegate.GetHandle() == InHandle; });
+	virtual void UnregisterOutlinerIndicator(FDelegateHandle InHandle) override {
+		OutlinerIndicatorDelegates.RemoveAll([=](const FOnCreateOutlinerIndicator& Delegate) { return Delegate.GetHandle() == InHandle; });
 	}
 
 	virtual FDelegateHandle RegisterOnSequencerCreated(FOnSequencerCreated::FDelegate InOnSequencerCreated) override
@@ -493,7 +493,7 @@ public:
 			OnGetGlobalRowExtensionHandle = EditModule.GetGlobalRowExtensionDelegate().AddStatic(&RegisterKeyframeExtensionHandler);
 
 			// Register far left gutter columns
-			OutlinerDecoratorColumnHandle = RegisterOutlinerColumn(FOnCreateOutlinerColumn::CreateStatic([] { return TSharedRef<IOutlinerColumn>(MakeShared<FOutlinerDecoratorColumn>()); }));
+			OutlinerIndicatorColumnHandle = RegisterOutlinerColumn(FOnCreateOutlinerColumn::CreateStatic([] { return TSharedRef<IOutlinerColumn>(MakeShared<FOutlinerIndicatorColumn>()); }));
 
 			// Register left gutter columns
 			PinOutlinerColumnHandle  = RegisterOutlinerColumn(FOnCreateOutlinerColumn::CreateStatic([]{ return TSharedRef<IOutlinerColumn>(MakeShared<FPinOutlinerColumn>()); }));
@@ -511,9 +511,9 @@ public:
 			NavOutlinerColumnHandle          = RegisterOutlinerColumn(FOnCreateOutlinerColumn::CreateStatic([]{ return TSharedRef<IOutlinerColumn>(MakeShared<FNavOutlinerColumn>()); }));
 			ColorPickerOutlinerColumnHandle  = RegisterOutlinerColumn(FOnCreateOutlinerColumn::CreateStatic([]{ return TSharedRef<IOutlinerColumn>(MakeShared<FColorPickerOutlinerColumn>()); }));
 
-			// Register outliner decorator items
-			ConditionOutlinerDecoratorHandle = RegisterOutlinerDecorator(FOnCreateOutlinerDecorator::CreateStatic([] { return TSharedRef<IOutlinerDecoratorBuilder>(MakeShared<FConditionOutlinerDecoratorBuilder>()); }));
-			TimeWarpOutlinerDecoratorHandle = RegisterOutlinerDecorator(FOnCreateOutlinerDecorator::CreateStatic([] { return TSharedRef<IOutlinerDecoratorBuilder>(MakeShared<FTimeWarpOutlinerDecoratorBuilder>()); }));
+			// Register outliner indicator items
+			ConditionOutlinerIndicatorHandle = RegisterOutlinerIndicator(FOnCreateOutlinerIndicator::CreateStatic([] { return TSharedRef<IOutlinerIndicatorBuilder>(MakeShared<FConditionOutlinerIndicatorBuilder>()); }));
+			TimeWarpOutlinerIndicatorHandle = RegisterOutlinerIndicator(FOnCreateOutlinerIndicator::CreateStatic([] { return TSharedRef<IOutlinerIndicatorBuilder>(MakeShared<FTimeWarpOutlinerIndicatorBuilder>()); }));
 
 			RegisterObjectSchemas();
 		}
@@ -562,7 +562,7 @@ public:
 			FEditorModeRegistry::Get().UnregisterMode(FSequencerEdMode::EM_SequencerMode);
 
 			// unregister outliner columns
-			UnregisterOutlinerColumn(OutlinerDecoratorColumnHandle);
+			UnregisterOutlinerColumn(OutlinerIndicatorColumnHandle);
 			UnregisterOutlinerColumn(PinOutlinerColumnHandle);
 			UnregisterOutlinerColumn(MuteOutlinerColumnHandle);
 			UnregisterOutlinerColumn(LockOutlinerColumnHandle);
@@ -574,9 +574,9 @@ public:
 			UnregisterOutlinerColumn(NavOutlinerColumnHandle);
 			UnregisterOutlinerColumn(ColorPickerOutlinerColumnHandle);
 
-			// unregister outliner decorator items
-			UnregisterOutlinerDecorator(ConditionOutlinerDecoratorHandle);
-			UnregisterOutlinerDecorator(TimeWarpOutlinerDecoratorHandle);
+			// unregister outliner indicator items
+			UnregisterOutlinerIndicator(ConditionOutlinerIndicatorHandle);
+			UnregisterOutlinerIndicator(TimeWarpOutlinerIndicatorHandle);
 		}
 	}
 
@@ -724,8 +724,8 @@ private:
 	/** List of outliner column creators */
 	TArray<FOnCreateOutlinerColumn> OutlinerColumnDelegates;
 
-	/** List of outliner decorator item creators */
-	TArray<FOnCreateOutlinerDecorator> OutlinerDecoratorDelegates;
+	/** List of outliner indicator item creators */
+	TArray<FOnCreateOutlinerIndicator> OutlinerIndicatorDelegates;
 
 	TArray<TSharedPtr<UE::Sequencer::IObjectSchema>> ObjectSchemas;
 
@@ -766,7 +766,7 @@ private:
 	TArray<FMovieRendererEntry> MovieRenderers;
 
 	// Outliner Column Delegate Handles
-	FDelegateHandle OutlinerDecoratorColumnHandle;
+	FDelegateHandle OutlinerIndicatorColumnHandle;
 	FDelegateHandle PinOutlinerColumnHandle;
 	FDelegateHandle MuteOutlinerColumnHandle;
 	FDelegateHandle LockOutlinerColumnHandle;
@@ -778,9 +778,9 @@ private:
 	FDelegateHandle NavOutlinerColumnHandle;
 	FDelegateHandle ColorPickerOutlinerColumnHandle;
 
-	// Outliner Decorator Item Delegate Handles
-	FDelegateHandle ConditionOutlinerDecoratorHandle;
-	FDelegateHandle TimeWarpOutlinerDecoratorHandle;
+	// Outliner Indicator Item Delegate Handles
+	FDelegateHandle ConditionOutlinerIndicatorHandle;
+	FDelegateHandle TimeWarpOutlinerIndicatorHandle;
 };
 
 IMPLEMENT_MODULE(FSequencerModule, Sequencer);
diff --git a/Engine/Source/Editor/Sequencer/Private/Widgets/OutlinerDecorators/SConditionDecoratorWidget.cpp b/Engine/Source/Editor/Sequencer/Private/Widgets/OutlinerDecorators/SConditionDecoratorWidget.cpp
deleted file mode 100644
diff --git a/Engine/Source/Editor/Sequencer/Private/Widgets/OutlinerDecorators/SConditionDecoratorWidget.h b/Engine/Source/Editor/Sequencer/Private/Widgets/OutlinerDecorators/SConditionDecoratorWidget.h
deleted file mode 100644
diff --git a/Engine/Source/Editor/Sequencer/Private/Widgets/OutlinerDecorators/STimeWarpDecoratorWidget.cpp b/Engine/Source/Editor/Sequencer/Private/Widgets/OutlinerDecorators/STimeWarpDecoratorWidget.cpp
deleted file mode 100644
diff --git a/Engine/Source/Editor/Sequencer/Private/Widgets/OutlinerDecorators/STimeWarpDecoratorWidget.h b/Engine/Source/Editor/Sequencer/Private/Widgets/OutlinerDecorators/STimeWarpDecoratorWidget.h
deleted file mode 100644
diff --git a/Engine/Source/Editor/Sequencer/Private/Widgets/OutlinerIndicators/SConditionIndicatorWidget.cpp b/Engine/Source/Editor/Sequencer/Private/Widgets/OutlinerIndicators/SConditionIndicatorWidget.cpp
new file mode 100644
--- /dev/null
+++ b/Engine/Source/Editor/Sequencer/Private/Widgets/OutlinerIndicators/SConditionIndicatorWidget.cpp
@@ -0,0 +1,134 @@
+// Copyright Epic Games, Inc. All Rights Reserved.
+
+#include "Widgets/OutlinerIndicators/SConditionIndicatorWidget.h"
+
+#include "MVVM/SharedViewModelData.h"
+#include "MVVM/ViewModels/ViewModelIterators.h"
+#include "MVVM/ViewModels/EditorViewModel.h"
+#include "MVVM/ViewModels/OutlinerColumns/IOutlinerColumn.h"
+#include "MVVM/ViewModels/OutlinerIndicators/IOutlinerIndicatorBuilder.h"
+#include "MVVM/ViewModels/OutlinerIndicators/ConditionOutlinerIndicatorBuilder.h"
+#include "MVVM/Extensions/IConditionableExtension.h"
+#include "MVVM/Selection/SequencerCoreSelection.h"
+#include "MVVM/Selection/SequencerOutlinerSelection.h"
+
+namespace UE::Sequencer
+{
+
+void SConditionIndicatorWidget::Construct(const FArguments& InArgs, const TWeakPtr<IOutlinerColumn> InWeakOutlinerColumn, const FCreateOutlinerColumnParams& InParams)
+{
+	SColumnToggleWidget::Construct(
+		SColumnToggleWidget::FArguments(),
+		InWeakOutlinerColumn,
+		InParams
+	);
+
+	WeakConditionStateCacheExtension = CastViewModel<FConditionStateCacheExtension>(InParams.OutlinerExtension.AsModel()->GetSharedData());
+}
+
+bool SConditionIndicatorWidget::IsActive() const
+{
+	if (TViewModelPtr<FConditionStateCacheExtension> StateCache = WeakConditionStateCacheExtension.Pin())
+	{
+		return EnumHasAnyFlags(StateCache->GetCachedFlags(ModelID), ECachedConditionState::HasCondition);
+	}
+
+	return false;
+}
+
+const FSlateBrush* SConditionIndicatorWidget::GetActiveBrush() const
+{
+	static const FName NAME_ConditionBrush = TEXT("Sequencer.Indicator.Condition");
+	return FAppStyle::Get().GetBrush(NAME_ConditionBrush);
+}
+
+FSlateColor SConditionIndicatorWidget::GetImageColorAndOpacity() const
+{
+	FLinearColor OutColor = FLinearColor::Black;
+	TSharedPtr<FEditorViewModel> Editor = WeakEditor.Pin();
+	TViewModelPtr<IOutlinerExtension> OutlinerItem = WeakOutlinerExtension.Pin();
+
+	bool bActive = IsActive();
+	bool bEvaluatingTrue = false;
+	bool bChildHasCondition = false;
+	bool bSectionHasCondition = false;
+	if (TViewModelPtr<FConditionStateCacheExtension> StateCache = WeakConditionStateCacheExtension.Pin())
+	{
+		bEvaluatingTrue = EnumHasAnyFlags(StateCache->GetCachedFlags(ModelID), ECachedConditionState::ConditionEvaluatingTrue);
+		bChildHasCondition = EnumHasAnyFlags(StateCache->GetCachedFlags(ModelID), ECachedConditionState::ChildHasCondition);
+		bSectionHasCondition = EnumHasAnyFlags(StateCache->GetCachedFlags(ModelID), ECachedConditionState::SectionHasCondition);
+	}
+
+	if (!Editor || !OutlinerItem)
+	{
+		return OutColor;
+	}
+
+	float Opacity = 1.0f;
+
+	if (IsActive() || (bChildHasCondition && !OutlinerItem->IsExpanded()) || bSectionHasCondition)
+	{
+		Opacity = 1.0f;
+	}
+	else
+	{
+		// Not active, invisible
+		Opacity = 0.0f;
+	}
+	OutColor.A = Opacity; // The background should be at the specified opacity and the icon should stay 100% black.
+	
+	return OutColor;
+}
+
+FSlateColor SConditionIndicatorWidget::GetIndicatorBackgroundColorAndOpacity() const
+{
+	static const FLinearColor ConditionColor = FLinearColor::FromSRGBColor(FColor(92, 220, 205));
+	TSharedPtr<FEditorViewModel> Editor = WeakEditor.Pin();
+	TViewModelPtr<IOutlinerExtension> OutlinerItem = WeakOutlinerExtension.Pin();
+
+	bool bActive = IsActive();
+	bool bEvaluatingTrue = false;
+	bool bChildHasCondition = false;
+	bool bSectionHasCondition = false;
+	if (TViewModelPtr<FConditionStateCacheExtension> StateCache = WeakConditionStateCacheExtension.Pin())
+	{
+		bEvaluatingTrue = EnumHasAnyFlags(StateCache->GetCachedFlags(ModelID), ECachedConditionState::ConditionEvaluatingTrue);
+		bChildHasCondition = EnumHasAnyFlags(StateCache->GetCachedFlags(ModelID), ECachedConditionState::ChildHasCondition);
+		bSectionHasCondition = EnumHasAnyFlags(StateCache->GetCachedFlags(ModelID), ECachedConditionState::SectionHasCondition);
+	}
+
+	FLinearColor OutColor = ConditionColor;
+
+	if (!Editor || !OutlinerItem)
+	{
+		return OutColor;
+	}
+
+	float Opacity = 1.0f;
+
+	if (IsActive())
+	{
+		if (bEvaluatingTrue)
+		{
+			Opacity = 1.0f;
+		}
+		else
+		{
+			Opacity = 0.5f;
+		}
+	}
+	else if ((bChildHasCondition && !OutlinerItem->IsExpanded()) || bSectionHasCondition)
+	{
+		Opacity = 0.25f;
+	}
+	else
+	{
+		// Not active, invisible
+		Opacity = 0.0f;
+	}
+	OutColor.A = Opacity; // The background should be at the specified opacity and the icon should stay 100% black.
+
+	return OutColor;
+}
+
+} // namespace UE::Sequencer
diff --git a/Engine/Source/Editor/Sequencer/Private/Widgets/OutlinerIndicators/SConditionIndicatorWidget.h b/Engine/Source/Editor/Sequencer/Private/Widgets/OutlinerIndicators/SConditionIndicatorWidget.h
new file mode 100644
--- /dev/null
+++ b/Engine/Source/Editor/Sequencer/Private/Widgets/OutlinerIndicators/SConditionIndicatorWidget.h
@@ -0,0 +1,66 @@
+// Copyright Epic Games, Inc. All Rights Reserved.
+
+#pragma once
+
+#include "MVVM/Views/OutlinerColumns/SColumnToggleWidget.h"
+
+namespace UE::Sequencer
+{
+
+class IOutlinerColumn;
+class IOutlinerIndicator;
+class FConditionOutlinerIndicatorBuilder;
+class FConditionStateCacheExtension;
+
+/**
+ * A widget that shows an indicator based on the presence and state of conditions in a row
+ */
+class SConditionIndicatorWidget
+	: public SColumnToggleWidget
+{
+public:
+	SLATE_BEGIN_ARGS(SConditionIndicatorWidget) {}
+	SLATE_END_ARGS()
+
+	void Construct(const FArguments& InArgs, const TWeakPtr<IOutlinerColumn> InWeakOutlinerColumn, const FCreateOutlinerColumnParams& InParams);
+
+	/** Get the color and opacity of the column toggle widget. */
+	virtual FSlateColor GetImageColorAndOpacity() const override;
+
+	/* Gets the color and opacity of the indicator background. */
+	FSlateColor GetIndicatorBackgroundColorAndOpacity() const;
+
+protected:
+
+	/** Returns whether or not the item is directly active or not. */
+	virtual bool IsActive() const override;
+
+	/** Sets the active state of this item to the input value. */
+	virtual void SetIsActive(const bool bInIsActive) override { /* no-op for now, no toggle*/ }
+
+	/** Returns true if a child of this item is active. */
+	virtual bool IsChildActive() const override { return false; }
+
+	/** Returns true if this item is implicitly active, but not directly active */
+	virtual bool IsImplicitlyActive() const override { return false; }
+
+	/** Returns the brush to be used to represent a widget is active. */
+	virtual const FSlateBrush* GetActiveBrush() const override;
+
+protected:
+
+	/** Get the image this widget displays. */
+	virtual const FSlateBrush* GetBrush() const override { return GetActiveBrush(); }
+
+protected:
+
+	/** Handles left click-type inputs and potentially begins drag and drop operation. */
+	FReply HandleClick() { return FReply::Handled(); }
+
+private:
+
+	/** Weak cache extension ptr (can be null). */
+	TWeakViewModelPtr<FConditionStateCacheExtension> WeakConditionStateCacheExtension;
+};
+
+} // namespace UE::Sequencer
\ No newline at end of file
diff --git a/Engine/Source/Editor/Sequencer/Private/Widgets/OutlinerIndicators/STimeWarpIndicatorWidget.cpp b/Engine/Source/Editor/Sequencer/Private/Widgets/OutlinerIndicators/STimeWarpIndicatorWidget.cpp
new file mode 100644
--- /dev/null
+++ b/Engine/Source/Editor/Sequencer/Private/Widgets/OutlinerIndicators/STimeWarpIndicatorWidget.cpp
@@ -0,0 +1,42 @@
+// Copyright Epic Games, Inc. All Rights Reserved.
+
+#include "Widgets/OutlinerIndicators/STimeWarpIndicatorWidget.h"
+
+#include "MVVM/SharedViewModelData.h"
+#include "MVVM/ViewModels/ViewModelIterators.h"
+#include "MVVM/ViewModels/EditorViewModel.h"
+#include "MVVM/ViewModels/OutlinerColumns/IOutlinerColumn.h"
+#include "MVVM/ViewModels/OutlinerIndicators/IOutlinerIndicatorBuilder.h"
+#include "MVVM/ViewModels/OutlinerIndicators/TimeWarpOutlinerIndicatorBuilder.h"
+#include "MVVM/Selection/SequencerCoreSelection.h"
+#include "MVVM/Selection/SequencerOutlinerSelection.h"
+
+namespace UE::Sequencer
+{
+
+void STimeWarpIndicatorWidget::Construct(const FArguments& InArgs, const TWeakPtr<IOutlinerColumn> InWeakOutlinerColumn, const FCreateOutlinerColumnParams& InParams)
+{
+	SColumnToggleWidget::Construct(
+		SColumnToggleWidget::FArguments(),
+		InWeakOutlinerColumn,
+		InParams
+	);
+}
+
+bool STimeWarpIndicatorWidget::IsActive() const
+{
+	return true;
+}
+
+const FSlateBrush* STimeWarpIndicatorWidget::GetActiveBrush() const
+{
+	static const FName NAME_TimeWarpBrush = TEXT("Sequencer.Indicator.TimeWarp");
+	return FAppStyle::Get().GetBrush(NAME_TimeWarpBrush);
+}
+
+FSlateColor STimeWarpIndicatorWidget::GetImageColorAndOpacity() const
+{
+	return FLinearColor::Black;
+}
+
+} // namespace UE::Sequencer
diff --git a/Engine/Source/Editor/Sequencer/Private/Widgets/OutlinerIndicators/STimeWarpIndicatorWidget.h b/Engine/Source/Editor/Sequencer/Private/Widgets/OutlinerIndicators/STimeWarpIndicatorWidget.h
new file mode 100644
--- /dev/null
+++ b/Engine/Source/Editor/Sequencer/Private/Widgets/OutlinerIndicators/STimeWarpIndicatorWidget.h
@@ -0,0 +1,61 @@
+// Copyright Epic Games, Inc. All Rights Reserved.
+
+#pragma once
+
+#include "MVVM/Views/OutlinerColumns/SColumnToggleWidget.h"
+
+namespace UE::Sequencer
+{	
+
+class IOutlinerColumn;
+class IOutlinerIndicator;
+class FTimeWarpOutlinerIndicatorBuilder;
+
+/**
+ * A widget that shows an indicator based on the presence of time warp on a row
+ */
+class STimeWarpIndicatorWidget
+	: public SColumnToggleWidget
+{
+public:
+	SLATE_BEGIN_ARGS(STimeWarpIndicatorWidget) {}
+	SLATE_END_ARGS()
+
+	void Construct(const FArguments& InArgs, const TWeakPtr<IOutlinerColumn> InWeakOutlinerColumn, const FCreateOutlinerColumnParams& InParams);
+
+	/** Get the color and opacity of the column toggle widget. */
+	virtual FSlateColor GetImageColorAndOpacity() const override;
+
+protected:
+
+	/** Returns whether or not the item is directly active or not. */
+	virtual bool IsActive() const override;
+
+	/** Sets the active state of this item to the input value. */
+	virtual void SetIsActive(const bool bInIsActive) override { /* no-op for now, no toggle*/ }
+
+	/** Returns true if a child of this item is active. */
+	virtual bool IsChildActive() const override { return false; }
+
+	/** Returns true if this item is implicitly active, but not directly active */
+	virtual bool IsImplicitlyActive() const override { return false; }
+
+	/** Returns the brush to be used to represent a widget is active. */
+	virtual const FSlateBrush* GetActiveBrush() const override;
+
+protected:
+
+	/** Get the image this widget displays. */
+	virtual const FSlateBrush* GetBrush() const override { return GetActiveBrush(); }
+
+protected:
+
+	/** Handles left click-type inputs and potentially begins drag and drop operation. */
+	FReply HandleClick() { return FReply::Handled(); }
+
+private:
+
+	TWeakPtr<FTimeWarpOutlinerIndicatorBuilder> WeakOutlinerIndicator;
+};
+
+} // namespace UE::Sequencer
\ No newline at end of file
diff --git a/Engine/Source/Editor/Sequencer/Public/ISequencerModule.h b/Engine/Source/Editor/Sequencer/Public/ISequencerModule.h
--- a/Engine/Source/Editor/Sequencer/Public/ISequencerModule.h
+++ b/Engine/Source/Editor/Sequencer/Public/ISequencerModule.h
@@ -29,7 +29,7 @@ namespace UE::Sequencer
 	class FTrackModel;
 	class IObjectSchema;
 	class IOutlinerColumn;
-	class IOutlinerDecoratorBuilder;
+	class IOutlinerIndicatorBuilder;
 } // namespace UE::Sequencer
 
 enum class ECurveEditorTreeFilterType : uint32;
@@ -64,8 +64,8 @@ DECLARE_DELEGATE_RetVal_OneParam(TSharedPtr<UE::Sequencer::FTrackModel>, FOnCrea
 /** A delegate which will create an outliner column */
 DECLARE_DELEGATE_RetVal(TSharedRef<UE::Sequencer::IOutlinerColumn>, FOnCreateOutlinerColumn);
 
-/** A delegate which will create an outliner decorator item */
-DECLARE_DELEGATE_RetVal(TSharedRef<UE::Sequencer::IOutlinerDecoratorBuilder>, FOnCreateOutlinerDecorator);
+/** A delegate which will create an outliner indicator item */
+DECLARE_DELEGATE_RetVal(TSharedRef<UE::Sequencer::IOutlinerIndicatorBuilder>, FOnCreateOutlinerIndicator);
 
 /** A delegate that is executed when adding menu content. */
 DECLARE_DELEGATE_OneParam(FOnGetContextMenuContent, FMenuBuilder& /*MenuBuilder*/);
@@ -275,19 +275,19 @@ public:
 	virtual void UnregisterOutlinerColumn(FDelegateHandle InHandle) = 0;
 
 	/**
-	 * Registers a delegate that will create an outliner decorator item
+	 * Registers a delegate that will create an outliner indicator item
 	 *
 	 * @param InCreator Delegate to register
 	 * @return A handle to the newly added delegate
 	 */
-	virtual FDelegateHandle RegisterOutlinerDecorator(FOnCreateOutlinerDecorator InCreator) = 0;
+	virtual FDelegateHandle RegisterOutlinerIndicator(FOnCreateOutlinerIndicator InCreator) = 0;
 
 	/**
-	 * Unregisters a previously registered delegate for creating an outliner decorator item
+	 * Unregisters a previously registered delegate for creating an outliner indicator item
 	 *
 	 * @param InHandle Handle to the delegate to unregister
 	 */
-	virtual void UnregisterOutlinerDecorator(FDelegateHandle InHandle) = 0;
+	virtual void UnregisterOutlinerIndicator(FDelegateHandle InHandle) = 0;
 
 	/**
 	 * Registers a delegate that will be called when a sequencer is created
diff --git a/Engine/Source/Editor/SequencerCore/Private/MVVM/ViewModels/OutlinerColumns/OutlinerColumnTypes.cpp b/Engine/Source/Editor/SequencerCore/Private/MVVM/ViewModels/OutlinerColumns/OutlinerColumnTypes.cpp
--- a/Engine/Source/Editor/SequencerCore/Private/MVVM/ViewModels/OutlinerColumns/OutlinerColumnTypes.cpp
+++ b/Engine/Source/Editor/SequencerCore/Private/MVVM/ViewModels/OutlinerColumns/OutlinerColumnTypes.cpp
@@ -5,7 +5,7 @@
 namespace UE::Sequencer
 {
 
-FName FCommonOutlinerNames::Decorator("Decorator");
+FName FCommonOutlinerNames::Indicator("Indicator");
 FName FCommonOutlinerNames::Pin("Pin");
 FName FCommonOutlinerNames::Lock("Lock");
 FName FCommonOutlinerNames::Mute("Mute");
diff --git a/Engine/Source/Editor/SequencerCore/Public/MVVM/ViewModels/OutlinerColumns/OutlinerColumnTypes.h b/Engine/Source/Editor/SequencerCore/Public/MVVM/ViewModels/OutlinerColumns/OutlinerColumnTypes.h
--- a/Engine/Source/Editor/SequencerCore/Public/MVVM/ViewModels/OutlinerColumns/OutlinerColumnTypes.h
+++ b/Engine/Source/Editor/SequencerCore/Public/MVVM/ViewModels/OutlinerColumns/OutlinerColumnTypes.h
@@ -65,7 +65,7 @@ struct FOutlinerColumnLayout
 
 struct FCommonOutlinerNames
 {
-	static SEQUENCERCORE_API FName Decorator;
+	static SEQUENCERCORE_API FName Indicator;
 	static SEQUENCERCORE_API FName Pin;
 	static SEQUENCERCORE_API FName Lock;
 	static SEQUENCERCORE_API FName Mute;
diff --git a/Engine/Source/Editor/SequencerCore/Public/MVVM/ViewModels/OutlinerDecorators/IOutlinerDecoratorBuilder.h b/Engine/Source/Editor/SequencerCore/Public/MVVM/ViewModels/OutlinerDecorators/IOutlinerDecoratorBuilder.h
deleted file mode 100644
diff --git a/Engine/Source/Editor/SequencerCore/Public/MVVM/ViewModels/OutlinerDecorators/OutlinerDecoratorBuilderBase.h b/Engine/Source/Editor/SequencerCore/Public/MVVM/ViewModels/OutlinerDecorators/OutlinerDecoratorBuilderBase.h
deleted file mode 100644
diff --git a/Engine/Source/Editor/SequencerCore/Public/MVVM/ViewModels/OutlinerIndicators/IOutlinerIndicatorBuilder.h b/Engine/Source/Editor/SequencerCore/Public/MVVM/ViewModels/OutlinerIndicators/IOutlinerIndicatorBuilder.h
new file mode 100644
--- /dev/null
+++ b/Engine/Source/Editor/SequencerCore/Public/MVVM/ViewModels/OutlinerIndicators/IOutlinerIndicatorBuilder.h
@@ -0,0 +1,42 @@
+// Copyright Epic Games, Inc. All Rights Reserved.
+
+#pragma once
+
+#include "Templates/SharedPointer.h"
+
+class SWidget;
+class UMovieSceneSequence;
+
+namespace UE::Sequencer
+{
+
+struct FCreateOutlinerColumnParams;
+class IOutlinerColumn;
+class ISequencerTreeViewRow;
+
+/**
+* Interface for building sequencer indicator outliner items.
+*/
+class IOutlinerIndicatorBuilder : public TSharedFromThis<IOutlinerIndicatorBuilder>
+{
+
+public:
+
+	/** Returns the name of the indicator outliner item. */
+	virtual FName GetIndicatorName() const = 0;
+
+	/* Gets whether or not this indicator outliner item is supported by a given Sequencer. */
+	virtual bool SupportsSequence(UMovieSceneSequence* InSequence) const { return true; }
+
+	virtual bool IsItemCompatibleWithIndicator(const FCreateOutlinerColumnParams& InParams) const = 0;
+	virtual TSharedPtr<SWidget> CreateIndicatorWidget(const FCreateOutlinerColumnParams& InParams, const TSharedRef<ISequencerTreeViewRow>& TreeViewRow, const TSharedRef<IOutlinerColumn>& OutlinerColumn, const int32 NumCompatibleIndicators) = 0;
+
+public:
+
+	/** Virtual destructor. */
+	virtual ~IOutlinerIndicatorBuilder() { }
+
+};
+
+} // namespace UE::Sequencer
+
diff --git a/Engine/Source/Editor/SequencerCore/Public/MVVM/ViewModels/OutlinerIndicators/OutlinerIndicatorBuilderBase.h b/Engine/Source/Editor/SequencerCore/Public/MVVM/ViewModels/OutlinerIndicators/OutlinerIndicatorBuilderBase.h
new file mode 100644
--- /dev/null
+++ b/Engine/Source/Editor/SequencerCore/Public/MVVM/ViewModels/OutlinerIndicators/OutlinerIndicatorBuilderBase.h
@@ -0,0 +1,25 @@
+// Copyright Epic Games, Inc. All Rights Reserved.
+
+#pragma once
+
+#include "MVVM/ViewModels/OutlinerIndicators/IOutlinerIndicatorBuilder.h"
+
+namespace UE::Sequencer
+{
+
+class FOutlinerIndicatorBuilderBase
+	: public IOutlinerIndicatorBuilder
+{
+public:
+
+	FOutlinerIndicatorBuilderBase()
+	{
+	}
+	
+public:
+
+	bool IsItemCompatibleWithIndicator(const FCreateOutlinerColumnParams& InParams) const override { return false; }
+	TSharedPtr<SWidget> CreateIndicatorWidget(const FCreateOutlinerColumnParams& InParams, const TSharedRef<ISequencerTreeViewRow>& TreeViewRow, const TSharedRef<IOutlinerColumn>& OutlinerColumn, const int32 NumCompatibleIndicators) override { return nullptr; }
+};
+
+} // namespace UE::Sequencer
\ No newline at end of file
