Sequencer: Fixed marquee selection not working for Time Warp channels that are forcibly pinned

  - The timewarp work introduced a new function, ShouldAnchorToTop which forcibly pins the track in the UI, but marquee selection was not considering this function.
  - I have audited remaining uses of IPinnableExtension and there are no more that need to be be fixed up to include this

diff --git a/Engine/Source/Editor/Sequencer/Private/Tools/SequencerEditTool_Selection.cpp b/Engine/Source/Editor/Sequencer/Private/Tools/SequencerEditTool_Selection.cpp
--- a/Engine/Source/Editor/Sequencer/Private/Tools/SequencerEditTool_Selection.cpp
+++ b/Engine/Source/Editor/Sequencer/Private/Tools/SequencerEditTool_Selection.cpp
@@ -41,8 +41,11 @@ struct FSelectionPreviewVisitor final
 	{
 		using namespace UE::Sequencer;
 
-		TSharedPtr<IPinnableExtension> PinnableItem = Channel->GetLinkedOutlinerItem().ImplicitCast();
-		if (PinnableItem && PinnableItem->IsPinned() != bPinned)
+		TViewModelPtr<IOutlinerExtension> OutlinerItem = Channel->GetLinkedOutlinerItem();
+		TSharedPtr<IPinnableExtension> PinnableItem = OutlinerItem.ImplicitCast();
+
+		const bool bIsItemPinned = (PinnableItem && PinnableItem->IsPinned()) || (OutlinerItem && OutlinerItem->ShouldAnchorToTop());
+		if (bIsItemPinned != bPinned)
 		{
 			return;
 		}
