[Sequencer] Change Isolate/Hide/Show icons to use Foreground color instead of White

diff --git a/Engine/Source/Editor/Sequencer/Private/Filters/Widgets/SFilterBarIsolateHideShow.cpp b/Engine/Source/Editor/Sequencer/Private/Filters/Widgets/SFilterBarIsolateHideShow.cpp
--- a/Engine/Source/Editor/Sequencer/Private/Filters/Widgets/SFilterBarIsolateHideShow.cpp
+++ b/Engine/Source/Editor/Sequencer/Private/Filters/Widgets/SFilterBarIsolateHideShow.cpp
@@ -91,6 +91,7 @@ TSharedRef<SWidget> SFilterBarIsolateHideShow::ConstructLayeredImage(const FName
 {
 	const TSharedRef<SLayeredImage> LayeredImage = SNew(SLayeredImage)
 		.DesiredSizeOverride(FVector2D(16.f))
+		.ColorAndOpacity(FStyleColors::Foreground)
 		.Image(FAppStyle::Get().GetBrush(InBaseImageName));
 
 	LayeredImage->AddLayer(TAttribute<const FSlateBrush*>::CreateLambda([InShowBadge]() -> const FSlateBrush*
