Curve Editor: Fixed curve editor filtering for filter terms separated by periods (.)

  - Introduced a new match result utility that keeps track of remaining matched tokens.
  - Filter terms that include periods (.) are now also eligible to be matched non-recursively using a literal match (ie, an entry with the literal label Location.X will be matched by Location.X, but LocationABC.X will not match)
  - Previously, sequencer curve editor filter terms were not being matched recursively. This has been corrected using the new match result utility.

diff --git a/Engine/Source/Editor/CurveEditor/Private/Tree/CurveEditorTreeFilter.cpp b/Engine/Source/Editor/CurveEditor/Private/Tree/CurveEditorTreeFilter.cpp
--- a/Engine/Source/Editor/CurveEditor/Private/Tree/CurveEditorTreeFilter.cpp
+++ b/Engine/Source/Editor/CurveEditor/Private/Tree/CurveEditorTreeFilter.cpp
@@ -7,6 +7,54 @@
 #include "Templates/UnrealTemplate.h"
 
 
+FCurveEditorTreeTextFilterTerm::FMatchResult FCurveEditorTreeTextFilterTerm::FMatchResult::Match(FStringView CandidateString) const
+{
+	// No tokens left - already a total match
+	if (RemainingTokens.Num() == 0)
+	{
+		return FMatchResult(RemainingTokens);
+	}
+
+	FStringView MatchString = RemainingTokens[0].Token;
+	const int32 MatchStartIndex = CandidateString.Find(MatchString, 0, ESearchCase::IgnoreCase);
+
+	if (MatchStartIndex == INDEX_NONE)
+	{
+		return FMatchResult();
+	}
+
+	FMatchResult Result(RemainingTokens.RightChop(1));
+
+	// The token matched! Continue to match chains of tokens separated by a period (.) within the same string
+
+	CandidateString = CandidateString.Left(MatchStartIndex);
+
+	while (Result.RemainingTokens.Num() > 0 && CandidateString.EndsWith('.'))
+	{
+		MatchString = Result.RemainingTokens[0].Token;
+
+		// Remove the period from the end of the candidate string
+		CandidateString = CandidateString.LeftChop(1);
+
+		// Compare the tail
+		if (CandidateString.Right(MatchString.Len()).Compare(MatchString, ESearchCase::IgnoreCase) != 0)
+		{
+			break;
+		}
+
+		// This token matched as well - remove it and keep matching...
+		Result.RemainingTokens = Result.RemainingTokens.RightChop(1);
+	}
+
+	// If there are no more tokens we must have matched them all
+	return Result;
+}
+
+FCurveEditorTreeTextFilterTerm::FMatchResult FCurveEditorTreeTextFilterTerm::Match(FStringView InString) const
+{
+	return FMatchResult(ChildToParentTokens).Match(InString);
+}
+
 
 ECurveEditorTreeFilterType FCurveEditorTreeFilter::RegisterFilterType()
 {
diff --git a/Engine/Source/Editor/CurveEditor/Public/Tree/CurveEditorTreeFilter.h b/Engine/Source/Editor/CurveEditor/Public/Tree/CurveEditorTreeFilter.h
--- a/Engine/Source/Editor/CurveEditor/Public/Tree/CurveEditorTreeFilter.h
+++ b/Engine/Source/Editor/CurveEditor/Public/Tree/CurveEditorTreeFilter.h
@@ -108,6 +108,53 @@ struct FCurveEditorTreeTextFilterToken
 /** A text filter term containing >= 1 sparate tokens ordered from child to parent */
 struct FCurveEditorTreeTextFilterTerm
 {
+	/**
+	 * Struct representing the result of a match operation,
+	 *     including any remaining tokens that need to be matched recursively upwards
+	 */
+	struct FMatchResult
+	{
+		/** Default constructor to an empty result that did not match */
+		FMatchResult()
+			: bMatched(false)
+		{}
+
+		/** Constructor that creates a successful match with optional remaining tokens */
+		FMatchResult(TArrayView<const FCurveEditorTreeTextFilterToken> InRemainingTokens)
+			: RemainingTokens(InRemainingTokens)
+			, bMatched(true)
+		{}
+
+		/** Match another string against the remaining tokens for this result */
+		CURVEEDITOR_API FMatchResult Match(FStringView InString) const;
+
+		/** Check whether this result is a partial or total match. Returns false if any match has failed. */
+		bool IsAnyMatch() const
+		{
+			return bMatched;
+		}
+		/** Check whether this result matched, but still has tokens left to match */
+		bool IsPartialMatch() const
+		{
+			return bMatched && RemainingTokens.Num() > 0;
+		}
+		/** Check whether this result has completely matched all tokens */
+		bool IsTotalMatch() const
+		{
+			return bMatched && RemainingTokens.Num() == 0;
+		}
+
+	private:
+
+		TArrayView<const FCurveEditorTreeTextFilterToken> RemainingTokens;
+		bool bMatched;
+	};
+
+	/**
+	 * Match a string against these tokens
+	 */
+	CURVEEDITOR_API FMatchResult Match(FStringView InString) const;
+
 	TArray<FCurveEditorTreeTextFilterToken, TInlineAllocator<1>> ChildToParentTokens;
 };
 
diff --git a/Engine/Source/Editor/CurveTableEditor/Private/CurveTableEditor.cpp b/Engine/Source/Editor/CurveTableEditor/Private/CurveTableEditor.cpp
--- a/Engine/Source/Editor/CurveTableEditor/Private/CurveTableEditor.cpp
+++ b/Engine/Source/Editor/CurveTableEditor/Private/CurveTableEditor.cpp
@@ -249,21 +249,14 @@ class FCurveTableEditorItem : public ICurveEditorTreeItem,  public TSharedFromTh
 	{
 		if (InFilter->GetType() == ECurveEditorTreeFilterType::Text)
 		{
-			FString     DisplayNameAsString = DisplayName.ToString();
-			FStringView RemainingStringView = DisplayNameAsString;
+			FString DisplayNameAsString = DisplayName.ToString();
 
 			const FCurveEditorTreeTextFilter* Filter = static_cast<const FCurveEditorTreeTextFilter*>(InFilter);
 			for (const FCurveEditorTreeTextFilterTerm& Term : Filter->GetTerms())
 			{
-				for (const FCurveEditorTreeTextFilterToken& Token : Term.ChildToParentTokens)
+				if (!Term.Match(DisplayNameAsString).IsTotalMatch())
 				{
-					FStringView TokenView(Token.Token);
-					const int32 MatchPos = RemainingStringView.Find(TokenView, 0, ESearchCase::IgnoreCase);
-					if (MatchPos == INDEX_NONE)
-					{
-						return false;
-					}
-					RemainingStringView.LeftInline(MatchPos);
+					return false;
 				}
 			}
 
diff --git a/Engine/Source/Editor/Sequencer/Private/MVVM/ViewModels/OutlinerItemModel.cpp b/Engine/Source/Editor/Sequencer/Private/MVVM/ViewModels/OutlinerItemModel.cpp
--- a/Engine/Source/Editor/Sequencer/Private/MVVM/ViewModels/OutlinerItemModel.cpp
+++ b/Engine/Source/Editor/Sequencer/Private/MVVM/ViewModels/OutlinerItemModel.cpp
@@ -43,30 +43,24 @@ namespace UE
 namespace Sequencer
 {
 
-static bool NodeMatchesTextFilterTerm(TSharedPtr<const UE::Sequencer::IOutlinerExtension> Node, const FCurveEditorTreeTextFilterTerm& Term)
+static bool NodeMatchesTextFilterTerm(TViewModelPtr<const UE::Sequencer::IOutlinerExtension> Node, const FCurveEditorTreeTextFilterTerm& Term)
 {
 	using namespace UE::Sequencer;
 
-	bool bMatched = false;
+	FCurveEditorTreeTextFilterTerm::FMatchResult Match(Term.ChildToParentTokens);
 
-	for (const FCurveEditorTreeTextFilterToken& Token : Term.ChildToParentTokens)
+	while (Node && Match.IsPartialMatch())
 	{
-		if (!Node)
+		FCurveEditorTreeTextFilterTerm::FMatchResult NewMatch = Match.Match(Node->GetLabel().ToString());
+		if (NewMatch.IsAnyMatch())
 		{
-			// No match - ran out of parents
-			return false;
+			// If we matched, keep searching parents using the remaining match result
+			Match = NewMatch;
 		}
-		else if (!Token.Match(*Node->GetIdentifier().ToString()))
-		{
-			return false;
-		}
-
-		bMatched = true;
-		//Node = Node->GetParent();
-		break;
+		Node = Node.AsModel()->FindAncestorOfType<const IOutlinerExtension>();
 	}
 
-	return bMatched;
+	return Match.IsTotalMatch();
 }
 
 void FOutlinerItemModelMixin::AddEvalOptionsPropertyMenuItem(FMenuBuilder& InMenuBuilder, const FBoolProperty* InProperty, TFunction<bool(UMovieSceneTrack*)> InValidator)
@@ -1101,16 +1095,18 @@ bool FOutlinerItemModelMixin::PassesFilter(const FCurveEditorTreeFilter* InFilte
 	{
 		const FCurveEditorTreeTextFilter* Filter = static_cast<const FCurveEditorTreeTextFilter*>(InFilter);
 
-		TSharedPtr<const IOutlinerExtension> This = AsViewModel()->CastThisShared<IOutlinerExtension>();
+		TViewModelPtr<const IOutlinerExtension> This = AsViewModel()->CastThisShared<IOutlinerExtension>();
+
+		// Must match all text tokens
 		for (const FCurveEditorTreeTextFilterTerm& Term : Filter->GetTerms())
 		{
-			if (NodeMatchesTextFilterTerm(This, Term))
+			if (!NodeMatchesTextFilterTerm(This, Term))
 			{
-				return true;
+				return false;
 			}
 		}
 
-		return false;
+		return true;
 	}
 	else if (InFilter->GetType() == ISequencerModule::GetSequencerSelectionFilterType())
 	{
