[Sequencer] Advanced Filtering:
 - Don't request filter update if "Selected" text expression is not actually active
 - Create function to check for text expressions in text filter string

diff --git a/Engine/Source/Editor/Sequencer/Private/Filters/Filters/SequencerTrackFilter_Text.cpp b/Engine/Source/Editor/Sequencer/Private/Filters/Filters/SequencerTrackFilter_Text.cpp
--- a/Engine/Source/Editor/Sequencer/Private/Filters/Filters/SequencerTrackFilter_Text.cpp
+++ b/Engine/Source/Editor/Sequencer/Private/Filters/Filters/SequencerTrackFilter_Text.cpp
@@ -101,7 +101,6 @@ bool FSequencerTrackFilter_Text::PassesFilter(FSequencerTrackFilterType InItem)
 	{
 		UMovieSceneTrack* const TrackObject = FilterData.ResolveMovieSceneTrackObject(InItem);
 		TextFilterExpressionContext->SetFilterItem(InItem, TrackObject);
-
 		if (!TextFilterExpressionEvaluator.TestTextFilter(*TextFilterExpressionContext))
 		{
 			return false;
@@ -139,4 +138,93 @@ const TArray<TSharedRef<FSequencerTextFilterExpressionContext>>& FSequencerTrack
 	return TextFilterExpressionContexts;
 }
 
+bool FSequencerTrackFilter_Text::DoesTextFilterStringContainExpressionPair(const FSequencerTextFilterExpressionContext& InExpression) const
+{
+	const TArray<FExpressionToken>& ExpressionTokens = TextFilterExpressionEvaluator.GetFilterExpressionTokens();
+	const int32 ExpressionCount = ExpressionTokens.Num();
+
+	// Need atleast three tokens: key, operator, and value
+	if (ExpressionTokens.Num() < 3)
+	{
+		return false;
+	}
+
+	const TSet<FName> Keys = InExpression.GetKeys();
+	const ESequencerTextFilterValueType ValueType = InExpression.GetValueType();
+	const TArray<FSequencerTextFilterKeyword> ValueKeywords = InExpression.GetValueKeywords();
+
+	for (int32 Index = 0; Index < ExpressionCount - 2; ++Index)
+	{
+		// Match key
+		const FExpressionToken& KeyToken = ExpressionTokens[Index];
+		const FString KeyTokenString = KeyToken.Context.GetString();
+
+		if (IsTokenKey(KeyToken, Keys))
+		{
+			// Match operator
+			const int32 OperatorIndex = Index + 1;
+			const FExpressionToken& OperatorToken = ExpressionTokens[OperatorIndex];
+			const FString OperatorTokenString = OperatorToken.Context.GetString();
+
+			if (IsTokenOperator(OperatorToken, ValueType))
+			{
+				// Match value
+				const int32 ValueIndex = Index + 2;
+				const FExpressionToken& ValueToken = ExpressionTokens[ValueIndex];
+				const FString ValueTokenString = ValueToken.Context.GetString();
+
+				if (IsTokenValueValid(ValueToken, ValueType))
+				{
+					return true;
+				}
+			}
+		}
+	}
+
+	return false;
+}
+
+bool FSequencerTrackFilter_Text::IsTokenKey(const FExpressionToken& InToken, const TSet<FName>& InKeys)
+{
+	const FString KeyTokenString = InToken.Context.GetString();
+	
+	for (const FName& Key : InKeys)
+	{
+		if (KeyTokenString.Equals(Key.ToString(), ESearchCase::IgnoreCase))
+		{
+			return true;
+		}
+	}
+
+	return false;
+};
+
+bool FSequencerTrackFilter_Text::IsTokenOperator(const FExpressionToken& InToken, const ESequencerTextFilterValueType InValueType)
+{
+	if (InValueType == ESequencerTextFilterValueType::String)
+	{
+		return InToken.Node.Cast<TextFilterExpressionParser::FEqual>()
+			|| InToken.Node.Cast<TextFilterExpressionParser::FNotEqual>();
+	}
+
+	return InToken.Node.Cast<TextFilterExpressionParser::FEqual>()
+		|| InToken.Node.Cast<TextFilterExpressionParser::FNotEqual>()
+		|| InToken.Node.Cast<TextFilterExpressionParser::FLess>()
+		|| InToken.Node.Cast<TextFilterExpressionParser::FLessOrEqual>()
+		|| InToken.Node.Cast<TextFilterExpressionParser::FGreater>()
+		|| InToken.Node.Cast<TextFilterExpressionParser::FGreaterOrEqual>();
+}
+
+bool FSequencerTrackFilter_Text::IsTokenValueValid(const FExpressionToken& InToken, const ESequencerTextFilterValueType InValueType)
+{
+	if (!InToken.Node.Cast<TextFilterExpressionParser::FTextToken>())
+	{
+		return false;
+	}
+
+	// @TODO: better value checking?
+
+	return true;
+}
+
 #undef LOCTEXT_NAMESPACE
diff --git a/Engine/Source/Editor/Sequencer/Private/Filters/Filters/SequencerTrackFilter_Text.h b/Engine/Source/Editor/Sequencer/Private/Filters/Filters/SequencerTrackFilter_Text.h
--- a/Engine/Source/Editor/Sequencer/Private/Filters/Filters/SequencerTrackFilter_Text.h
+++ b/Engine/Source/Editor/Sequencer/Private/Filters/Filters/SequencerTrackFilter_Text.h
@@ -29,7 +29,13 @@ public:
 
 	const TArray<TSharedRef<FSequencerTextFilterExpressionContext>>& GetTextFilterExpressionContexts() const;
 
+	bool DoesTextFilterStringContainExpressionPair(const FSequencerTextFilterExpressionContext& InExpression) const;
+
 protected:
+	static bool IsTokenKey(const FExpressionToken& InToken, const TSet<FName>& InKeys);
+	static bool IsTokenOperator(const FExpressionToken& InToken, const ESequencerTextFilterValueType InValueType);
+	static bool IsTokenValueValid(const FExpressionToken& InToken, const ESequencerTextFilterValueType InValueType);
+
 	/** Expression evaluator that can be used to perform complex text filter queries */
 	FTextFilterExpressionEvaluator TextFilterExpressionEvaluator;
 
diff --git a/Engine/Source/Editor/Sequencer/Private/Filters/SequencerFilterBar.cpp b/Engine/Source/Editor/Sequencer/Private/Filters/SequencerFilterBar.cpp
--- a/Engine/Source/Editor/Sequencer/Private/Filters/SequencerFilterBar.cpp
+++ b/Engine/Source/Editor/Sequencer/Private/Filters/SequencerFilterBar.cpp
@@ -336,6 +336,11 @@ void FSequencerFilterBar::SetTextFilterString(const FString& InText)
 	}
 }
 
+bool FSequencerFilterBar::DoesTextFilterStringContainExpressionPair(const FSequencerTextFilterExpressionContext& InExpression) const
+{
+	return TextFilter->DoesTextFilterStringContainExpressionPair(InExpression);
+}
+
 TSharedRef<FSequencerTrackFilter_Text> FSequencerFilterBar::GetTextFilter() const
 {
 	return TextFilter;
diff --git a/Engine/Source/Editor/Sequencer/Private/Filters/SequencerFilterBar.h b/Engine/Source/Editor/Sequencer/Private/Filters/SequencerFilterBar.h
--- a/Engine/Source/Editor/Sequencer/Private/Filters/SequencerFilterBar.h
+++ b/Engine/Source/Editor/Sequencer/Private/Filters/SequencerFilterBar.h
@@ -57,8 +57,6 @@ public:
 
 	TSharedPtr<ICustomTextFilter<FSequencerTrackFilterType>> CreateTextFilter();
 
-	virtual TSharedPtr<FUICommandList> GetCommandList() const override;
-
 	TSharedRef<SSequencerFilterBar> GenerateWidget(const TSharedPtr<SFilterSearchBox>& InSearchBox, const EFilterBarLayout InLayout);
 
 	bool AreFiltersMuted() const;
@@ -70,9 +68,6 @@ public:
 
 	FSequencerFiltersChanged& GetOnFiltersChanged() { return FiltersChangedEvent; }
 
-	FString GetTextFilterString() const;
-	void SetTextFilterString(const FString& InText);
-
 	TSharedRef<FSequencerTrackFilter_Text> GetTextFilter() const;
 	FText GetFilterErrorText() const;
 
@@ -102,6 +97,13 @@ public:
 
 	virtual FSequencer& GetSequencer() const override;
 
+	virtual TSharedPtr<FUICommandList> GetCommandList() const override;
+
+	virtual FString GetTextFilterString() const override;
+    virtual void SetTextFilterString(const FString& InText) override;
+
+	virtual bool DoesTextFilterStringContainExpressionPair(const FSequencerTextFilterExpressionContext& InExpression) const override;
+
 	virtual void RequestFilterUpdate() override;
 
 	virtual void EnableAllFilters(const bool bInEnable, const TArray<FString> InExceptionFilterNames) override;
diff --git a/Engine/Source/Editor/Sequencer/Private/Filters/TextExpressions/SequencerTextFilterExpression_Selected.cpp b/Engine/Source/Editor/Sequencer/Private/Filters/TextExpressions/SequencerTextFilterExpression_Selected.cpp
--- a/Engine/Source/Editor/Sequencer/Private/Filters/TextExpressions/SequencerTextFilterExpression_Selected.cpp
+++ b/Engine/Source/Editor/Sequencer/Private/Filters/TextExpressions/SequencerTextFilterExpression_Selected.cpp
@@ -69,7 +69,10 @@ bool FSequencerTextFilterExpression_Selected::TestComplexExpression(const FName&
 
 void FSequencerTextFilterExpression_Selected::OnSelectionChanged(UObject* const InObject)
 {
-	FilterInterface.RequestFilterUpdate();
+	if (FilterInterface.DoesTextFilterStringContainExpressionPair(*this))
+	{
+		FilterInterface.RequestFilterUpdate();
+	}
 }
 
 #undef LOCTEXT_NAMESPACE
diff --git a/Engine/Source/Editor/Sequencer/Private/Sequencer.cpp b/Engine/Source/Editor/Sequencer/Private/Sequencer.cpp
--- a/Engine/Source/Editor/Sequencer/Private/Sequencer.cpp
+++ b/Engine/Source/Editor/Sequencer/Private/Sequencer.cpp
@@ -6804,8 +6804,6 @@ void FSequencer::SynchronizeSequencerSelectionWithExternalSelection()
 			}
 		}
 	}
-
-	FilterBar->RequestFilterUpdate();
 }
 
 void FSequencer::SelectNodesByPath(const TSet<FString>& NodePaths)
diff --git a/Engine/Source/Editor/Sequencer/Public/Filters/ISequencerTrackFilters.h b/Engine/Source/Editor/Sequencer/Public/Filters/ISequencerTrackFilters.h
--- a/Engine/Source/Editor/Sequencer/Public/Filters/ISequencerTrackFilters.h
+++ b/Engine/Source/Editor/Sequencer/Public/Filters/ISequencerTrackFilters.h
@@ -4,6 +4,7 @@
 
 #include "Framework/Commands/UICommandList.h"
 
+class FSequencerTextFilterExpressionContext;
 class FSequencerTrackFilter_CustomText;
 class FSequencer;
 class FString;
@@ -21,6 +22,13 @@ public:
 
 	SEQUENCER_API virtual TSharedPtr<FUICommandList> GetCommandList() const = 0;
 
+	SEQUENCER_API virtual FString GetTextFilterString() const = 0;
+	SEQUENCER_API virtual void SetTextFilterString(const FString& InText) = 0;
+
+	/** Returns true if the current filter bar text filter string contains the specified text expression.
+	 * The text expression must have key, operator, and value tokens. */
+	SEQUENCER_API virtual bool DoesTextFilterStringContainExpressionPair(const FSequencerTextFilterExpressionContext& InExpression) const = 0;
+
 	SEQUENCER_API virtual bool HasAnyFilterActive(const bool bCheckTextFilter = true
 		, const bool bInCheckHideIsolateFilter = true
 		, const bool bInCheckCommonFilters = true
