[Sequencer] Advanced filtering:
 - Fix regression of child tracks not showing
 - Only show category and channel groups that are animated when keys filter is active

diff --git a/Engine/Source/Editor/Sequencer/Private/Filters/Filters/SequencerTrackFilter_Keyed.cpp b/Engine/Source/Editor/Sequencer/Private/Filters/Filters/SequencerTrackFilter_Keyed.cpp
--- a/Engine/Source/Editor/Sequencer/Private/Filters/Filters/SequencerTrackFilter_Keyed.cpp
+++ b/Engine/Source/Editor/Sequencer/Private/Filters/Filters/SequencerTrackFilter_Keyed.cpp
@@ -40,8 +40,23 @@ bool FSequencerTrackFilter_Keyed::SupportsSequence(UMovieSceneSequence* const In
 
 bool FSequencerTrackFilter_Keyed::PassesFilter(FSequencerTrackFilterType InItem) const
 {
-	const TWeakViewModelPtr<ITrackExtension> Track = FilterInterface.GetFilterData().ResolveTrack(InItem);
-	return DoesTrackExtensionHaveKeys(Track);
+	if (const TViewModelPtr<FCategoryGroupModel> CategoryGroupModel = InItem.ImplicitCast())
+	{
+		for (const TWeakViewModelPtr<FCategoryModel>& WeakCategory : CategoryGroupModel->GetCategories())
+		{
+			if (const TViewModelPtr<FCategoryModel> Category = WeakCategory.Pin())
+			{
+				return Category->IsAnimated();
+			}
+		}
+	}
+	else if (const TViewModelPtr<FChannelGroupOutlinerModel> ChannelGroupOutlinerModel = InItem.ImplicitCast())
+	{
+		return ChannelGroupOutlinerModel->IsAnimated();
+	}
+
+	const TWeakViewModelPtr<ITrackExtension> WeakTrack = FilterInterface.GetFilterData().ResolveTrack(InItem);
+	return DoesTrackExtensionHaveKeys(WeakTrack);
 }
 
 FText FSequencerTrackFilter_Keyed::GetDisplayName() const
diff --git a/Engine/Source/Editor/Sequencer/Private/Filters/SequencerFilterData.cpp b/Engine/Source/Editor/Sequencer/Private/Filters/SequencerFilterData.cpp
--- a/Engine/Source/Editor/Sequencer/Private/Filters/SequencerFilterData.cpp
+++ b/Engine/Source/Editor/Sequencer/Private/Filters/SequencerFilterData.cpp
@@ -160,36 +160,6 @@ TWeakViewModelPtr<ITrackExtension> FSequencerFilterData::ResolveTrack(FSequencer
 	return nullptr;
 }
 
-TWeakViewModelPtr<IObjectBindingExtension> FSequencerFilterData::ResolveTrackObjectBinding(FSequencerTrackFilterType InNode)
-{
-	if (!InNode.IsValid())
-	{
-		return nullptr;
-	}
-
-	const TWeakViewModelPtr<IOutlinerExtension> WeakOutlinerNode = InNode.ImplicitCast();
-
-	// Use cache version if it exists, otherwise resolve below
-	if (ResolvedObjectBindings.Contains(WeakOutlinerNode))
-	{
-		if (const TViewModelPtr<IObjectBindingExtension> ResolvedObjectBinding = ResolvedObjectBindings[WeakOutlinerNode].Pin())
-		{
-			return ResolvedObjectBinding;
-		}
-
-		ResolvedObjectBindings.Remove(WeakOutlinerNode);
-	}
-
-	// Resolve and cache
-	if (const TViewModelPtr<IObjectBindingExtension> ObjectBindingModel = InNode->FindAncestorOfType<IObjectBindingExtension>(true))
-	{
-		ResolvedObjectBindings.Add(WeakOutlinerNode, ObjectBindingModel);
-		return ObjectBindingModel;
-	}
-
-	return nullptr;
-}
-
 UMovieSceneTrack* FSequencerFilterData::ResolveMovieSceneTrackObject(FSequencerTrackFilterType InNode)
 {
 	if (!InNode.IsValid())
@@ -200,23 +170,22 @@ UMovieSceneTrack* FSequencerFilterData::ResolveMovieSceneTrackObject(FSequencerT
 	const TWeakViewModelPtr<IOutlinerExtension> WeakOutlinerNode = InNode.ImplicitCast();
 
 	// Use cache version if it exists, otherwise resolve below
-	if (ResolvedTracks.Contains(WeakOutlinerNode))
+	if (ResolvedTrackObjects.Contains(WeakOutlinerNode))
 	{
-		if (const TViewModelPtr<ITrackExtension> ResolvedTrack = ResolvedTracks[WeakOutlinerNode].Pin())
+		if (ResolvedTrackObjects[WeakOutlinerNode].IsValid())
 		{
-			return ResolvedTrack->GetTrack();
+			return ResolvedTrackObjects[WeakOutlinerNode].Get();
 		}
 
-		ResolvedTracks.Remove(WeakOutlinerNode);
+		ResolvedTrackObjects.Remove(WeakOutlinerNode);
 	}
 
-	// Resolve and cache
 	if (const TViewModelPtr<ITrackExtension> AncestorTrackModel = InNode->FindAncestorOfType<ITrackExtension>(true))
 	{
 		UMovieSceneTrack* const TrackObject = AncestorTrackModel->GetTrack();
 		if (IsValid(TrackObject))
 		{
-			ResolvedObjects.Add(WeakOutlinerNode, TrackObject);
+			ResolvedTrackObjects.Add(WeakOutlinerNode, TrackObject);
 			return TrackObject;
 		}
 	}
@@ -224,7 +193,7 @@ UMovieSceneTrack* FSequencerFilterData::ResolveMovieSceneTrackObject(FSequencerT
 	return nullptr;
 }
 
-UObject* FSequencerFilterData::ResolveTrackBoundObject(ISequencer& InSequencer, const TViewModelPtr<FViewModel> InNode)
+UObject* FSequencerFilterData::ResolveTrackBoundObject(ISequencer& InSequencer, FSequencerTrackFilterType InNode)
 {
 	if (!InNode.IsValid())
 	{
@@ -236,35 +205,16 @@ UObject* FSequencerFilterData::ResolveTrackBoundObject(ISequencer& InSequencer,
 	// Use cache version if it exists, otherwise resolve below
 	if (ResolvedObjects.Contains(WeakOutlinerNode))
 	{
-		UObject* const ResolvedObject = ResolvedObjects[WeakOutlinerNode].Get();
-		if (IsValid(ResolvedObject))
+		if (ResolvedObjects[WeakOutlinerNode].IsValid())
 		{
-			return ResolvedObject;
+			return ResolvedObjects[WeakOutlinerNode].Get();
 		}
 
 		ResolvedObjects.Remove(WeakOutlinerNode);
 	}
 
-	if (ResolvedObjectBindings.Contains(WeakOutlinerNode))
-	{
-		if (const TViewModelPtr<IObjectBindingExtension> ResolvedObjectBinding = ResolvedObjectBindings[WeakOutlinerNode].Pin())
-		{
-			UObject* const BoundObject = InSequencer.FindSpawnedObjectOrTemplate(ResolvedObjectBinding->GetObjectGuid());
-			if (IsValid(BoundObject))
-			{
-				ResolvedObjects.Add(WeakOutlinerNode, BoundObject);
-				return BoundObject;
-			}
-		}
-
-		ResolvedObjectBindings.Remove(WeakOutlinerNode);
-	}
-
-	// Resolve and cache
 	if (const TViewModelPtr<IObjectBindingExtension> ObjectBindingModel = InNode->FindAncestorOfType<IObjectBindingExtension>(true))
 	{
-		ResolvedObjectBindings.Add(WeakOutlinerNode, ObjectBindingModel);
-
 		UObject* const BoundObject = InSequencer.FindSpawnedObjectOrTemplate(ObjectBindingModel->GetObjectGuid());
 		if (IsValid(BoundObject))
 		{
diff --git a/Engine/Source/Editor/Sequencer/Public/Filters/SequencerFilterData.h b/Engine/Source/Editor/Sequencer/Public/Filters/SequencerFilterData.h
--- a/Engine/Source/Editor/Sequencer/Public/Filters/SequencerFilterData.h
+++ b/Engine/Source/Editor/Sequencer/Public/Filters/SequencerFilterData.h
@@ -48,14 +48,12 @@ struct SEQUENCER_API FSequencerFilterData
 
 	bool IsFilteredOut(const UE::Sequencer::TViewModelPtr<UE::Sequencer::IOutlinerExtension>& InNode) const;
 
+	UE::Sequencer::TWeakViewModelPtr<UE::Sequencer::ITrackExtension> ResolveTrack(FSequencerTrackFilterType InNode);
 	UMovieSceneTrack* ResolveMovieSceneTrackObject(FSequencerTrackFilterType InNode);
 	UObject* ResolveTrackBoundObject(ISequencer& InSequencer, FSequencerTrackFilterType InNode);
 
-	UE::Sequencer::TWeakViewModelPtr<UE::Sequencer::ITrackExtension> ResolveTrack(FSequencerTrackFilterType InNode);
-	UE::Sequencer::TWeakViewModelPtr<UE::Sequencer::IObjectBindingExtension> ResolveTrackObjectBinding(FSequencerTrackFilterType InNode);
-
 	TMap<UE::Sequencer::TWeakViewModelPtr<UE::Sequencer::IOutlinerExtension>, UE::Sequencer::TWeakViewModelPtr<UE::Sequencer::ITrackExtension>> ResolvedTracks;
-	TMap<UE::Sequencer::TWeakViewModelPtr<UE::Sequencer::IOutlinerExtension>, UE::Sequencer::TWeakViewModelPtr<UE::Sequencer::IObjectBindingExtension>> ResolvedObjectBindings;
+	TMap<UE::Sequencer::TWeakViewModelPtr<UE::Sequencer::IOutlinerExtension>, TWeakObjectPtr<UMovieSceneTrack>> ResolvedTrackObjects;
 	TMap<UE::Sequencer::TWeakViewModelPtr<UE::Sequencer::IOutlinerExtension>, TWeakObjectPtr<>> ResolvedObjects;
 
 protected:
