Sequencer Added 'Fix Possessable Object Class' to the Advanced menu

diff --git a/Engine/Source/Editor/Sequencer/Private/SSequencer.cpp b/Engine/Source/Editor/Sequencer/Private/SSequencer.cpp
--- a/Engine/Source/Editor/Sequencer/Private/SSequencer.cpp
+++ b/Engine/Source/Editor/Sequencer/Private/SSequencer.cpp
@@ -1846,6 +1846,7 @@ void SSequencer::FillAdvancedMenu(FMenuBuilder& MenuBuilder)
 		MenuBuilder.BeginSection("Bindings", LOCTEXT("BindingsMenuHeader", "Bindings"));
 
 		MenuBuilder.AddMenuEntry(FSequencerCommands::Get().RebindPossessableReferences);
+		MenuBuilder.AddMenuEntry(FSequencerCommands::Get().FixPossessableObjectClass);
 
 		MenuBuilder.EndSection();
 	}
