Sequencer: Dynamically add/remove the Outliner columns when required instead of closing and re-opening the Outliner

diff --git a/Engine/Plugins/Experimental/DaySequence/Source/DaySequenceEditor/Private/DaySequenceEditorToolkit.cpp b/Engine/Plugins/Experimental/DaySequence/Source/DaySequenceEditor/Private/DaySequenceEditorToolkit.cpp
--- a/Engine/Plugins/Experimental/DaySequence/Source/DaySequenceEditor/Private/DaySequenceEditorToolkit.cpp
+++ b/Engine/Plugins/Experimental/DaySequence/Source/DaySequenceEditor/Private/DaySequenceEditorToolkit.cpp
@@ -323,14 +323,6 @@ void FDaySequenceEditorToolkit::InitializeInternal(
 
 	if (bIsSequenceEditor)
 	{
-		// Reopen the scene outliner so that is refreshed with the sequencer columns
-		TSharedPtr<FTabManager> LevelEditorTabManager = LevelEditorModule.GetLevelEditorTabManager();
-		if (LevelEditorTabManager->FindExistingLiveTab(FName("LevelEditorSceneOutliner")).IsValid())
-		{
-			LevelEditorTabManager->TryInvokeTab(FName("LevelEditorSceneOutliner"))->RequestCloseTab();
-			LevelEditorTabManager->TryInvokeTab(FName("LevelEditorSceneOutliner"));
-		}
-
 		// Now Attach so this window will appear in the correct front first order
 		TSharedPtr<SDockTab> DockTab = LevelEditorModule.AttachSequencer(Sequencer->GetSequencerWidget(), SharedThis(this));
 		if (DockTab.IsValid())
diff --git a/Engine/Plugins/MovieScene/LevelSequenceEditor/Source/LevelSequenceEditor/Private/LevelSequenceEditorToolkit.cpp b/Engine/Plugins/MovieScene/LevelSequenceEditor/Source/LevelSequenceEditor/Private/LevelSequenceEditorToolkit.cpp
--- a/Engine/Plugins/MovieScene/LevelSequenceEditor/Source/LevelSequenceEditor/Private/LevelSequenceEditorToolkit.cpp
+++ b/Engine/Plugins/MovieScene/LevelSequenceEditor/Source/LevelSequenceEditor/Private/LevelSequenceEditorToolkit.cpp
@@ -193,16 +193,6 @@ void FLevelSequenceEditorToolkit::Initialize(const EToolkitMode::Type Mode, cons
 	FLevelEditorSequencerIntegration::Get().AddSequencer(Sequencer.ToSharedRef(), Options);
 	ULevelSequenceEditorBlueprintLibrary::SetSequencer(Sequencer.ToSharedRef());
 
-	// Reopen the scene outliner so that is refreshed with the sequencer columns
-	{
-		TSharedPtr<FTabManager> LevelEditorTabManager = LevelEditorModule.GetLevelEditorTabManager();
-		if (LevelEditorTabManager->FindExistingLiveTab(FName("LevelEditorSceneOutliner")).IsValid())
-		{
-			LevelEditorTabManager->TryInvokeTab(FName("LevelEditorSceneOutliner"))->RequestCloseTab();
-			LevelEditorTabManager->TryInvokeTab(FName("LevelEditorSceneOutliner"));
-		}
-	}
-	
 	// Now Attach so this window will apear in the correct front first order
 	TSharedPtr<SDockTab> DockTab = LevelEditorModule.AttachSequencer(Sequencer->GetSequencerWidget(), SharedThis(this));
 	if (DockTab.IsValid())
diff --git a/Engine/Source/Editor/Sequencer/Private/LevelEditorSequencerIntegration.cpp b/Engine/Source/Editor/Sequencer/Private/LevelEditorSequencerIntegration.cpp
--- a/Engine/Source/Editor/Sequencer/Private/LevelEditorSequencerIntegration.cpp
+++ b/Engine/Source/Editor/Sequencer/Private/LevelEditorSequencerIntegration.cpp
@@ -1106,52 +1106,62 @@ TSharedRef< ISceneOutlinerColumn > FLevelEditorSequencerIntegration::CreateSeque
 
 void FLevelEditorSequencerIntegration::AttachOutlinerColumn()
 {
-	// Register Spawnable Column 
-	FSceneOutlinerModule& SceneOutlinerModule = FModuleManager::LoadModuleChecked< FSceneOutlinerModule >("SceneOutliner");
-
+	FLevelEditorModule& LevelEditorModule = FModuleManager::LoadModuleChecked<FLevelEditorModule>("LevelEditor");
+	
 	FSceneOutlinerColumnInfo SpawnColumnInfo(ESceneOutlinerColumnVisibility::Visible, 11, 
-		FCreateSceneOutlinerColumn::CreateRaw( this, &FLevelEditorSequencerIntegration::CreateSequencerSpawnableColumn),
-		true, TOptional<float>(), LOCTEXT("SpawnableColumnName", "Spawnable"));
+	FCreateSceneOutlinerColumn::CreateRaw( this, &FLevelEditorSequencerIntegration::CreateSequencerSpawnableColumn),
+	true, TOptional<float>(), LOCTEXT("SpawnableColumnName", "Spawnable"));
 
-	SceneOutlinerModule.RegisterDefaultColumnType< Sequencer::FSequencerSpawnableColumn >(SpawnColumnInfo);
-
-	FSceneOutlinerColumnInfo ColumnInfo(ESceneOutlinerColumnVisibility::Visible, 15, 
+	FSceneOutlinerColumnInfo SequencerColumnInfo(ESceneOutlinerColumnVisibility::Visible, 15, 
 		FCreateSceneOutlinerColumn::CreateRaw( this, &FLevelEditorSequencerIntegration::CreateSequencerInfoColumn), 
 		true, TOptional<float>(), LOCTEXT("SequencerColumnName", "Sequencer"));
 
-	SceneOutlinerModule.RegisterDefaultColumnType< Sequencer::FSequencerInfoColumn >(ColumnInfo);
+	FSceneOutlinerModule& SceneOutlinerModule = FModuleManager::LoadModuleChecked<FSceneOutlinerModule>("SceneOutliner");
+
+	// First we register the columns as default columns so they show up in any new Outliners that are opened
+	SceneOutlinerModule.RegisterDefaultColumnType< Sequencer::FSequencerSpawnableColumn >(SpawnColumnInfo);
+	SceneOutlinerModule.RegisterDefaultColumnType< Sequencer::FSequencerInfoColumn >(SequencerColumnInfo);
+
+	// Then we go through all currently open Outliners and add the column manually
+	if(TSharedPtr<ILevelEditor> LevelEditor = LevelEditorModule.GetFirstLevelEditor())
+	{
+		TArray<TWeakPtr<ISceneOutliner>> Outliners = LevelEditor->GetAllSceneOutliners();
 
+		for(TWeakPtr<ISceneOutliner> Outliner : Outliners)
+		{
+			if(TSharedPtr<ISceneOutliner> OutlinerPinned = Outliner.Pin())
+			{
+				OutlinerPinned->AddColumn(Sequencer::FSequencerSpawnableColumn::GetID(), SpawnColumnInfo);
+				OutlinerPinned->AddColumn(Sequencer::FSequencerInfoColumn::GetID(), SequencerColumnInfo);
+			}
+		}
+	}
 }
 
 void FLevelEditorSequencerIntegration::DetachOutlinerColumn()
 {
-	FSceneOutlinerModule& SceneOutlinerModule = FModuleManager::LoadModuleChecked< FSceneOutlinerModule >("SceneOutliner");
-
-	SceneOutlinerModule.UnRegisterColumnType< Sequencer::FSequencerSpawnableColumn >();
-	SceneOutlinerModule.UnRegisterColumnType< Sequencer::FSequencerInfoColumn >();
-
 	FLevelEditorModule& LevelEditorModule = FModuleManager::LoadModuleChecked<FLevelEditorModule>("LevelEditor");
 
-	// @todo reopen the scene outliner so that is refreshed without the sequencer info column
-	TSharedPtr<FTabManager> LevelEditorTabManager = LevelEditorModule.GetLevelEditorTabManager();
-	if (LevelEditorTabManager->FindExistingLiveTab(FName("LevelEditorSceneOutliner")).IsValid())
+	// First we remove the column from any active Outliners
+	if(TSharedPtr<ILevelEditor> LevelEditor = LevelEditorModule.GetFirstLevelEditor())
 	{
-		if (LevelEditorTabManager.IsValid() && LevelEditorTabManager.Get())
-		{
-			if (LevelEditorTabManager->GetOwnerTab().IsValid())
-			{
-				LevelEditorTabManager->TryInvokeTab(FName("LevelEditorSceneOutliner"))->RequestCloseTab();			
-			}
-		}
-		
-		if (LevelEditorTabManager.IsValid() && LevelEditorTabManager.Get())
+		TArray<TWeakPtr<ISceneOutliner>> Outliners = LevelEditor->GetAllSceneOutliners();
+
+		for(TWeakPtr<ISceneOutliner> Outliner : Outliners)
 		{
-			if (LevelEditorTabManager->GetOwnerTab().IsValid())
+			if(TSharedPtr<ISceneOutliner> OutlinerPinned = Outliner.Pin())
 			{
-				LevelEditorTabManager->TryInvokeTab(FName("LevelEditorSceneOutliner"));
+				OutlinerPinned->RemoveColumn(Sequencer::FSequencerSpawnableColumn::GetID());
+				OutlinerPinned->RemoveColumn(Sequencer::FSequencerInfoColumn::GetID());
 			}
 		}
 	}
+
+	FSceneOutlinerModule& SceneOutlinerModule = FModuleManager::LoadModuleChecked<FSceneOutlinerModule>("SceneOutliner");
+
+	// Then we unregister the column type so it isn't added to any Outliners that are opened in the future	
+	SceneOutlinerModule.UnRegisterColumnType< Sequencer::FSequencerSpawnableColumn >();
+	SceneOutlinerModule.UnRegisterColumnType< Sequencer::FSequencerInfoColumn >();
 }
 
 void FLevelEditorSequencerIntegration::ActivateRealtimeViewports()
