[Sequencer] Add text filter search box suggestions

diff --git a/Engine/Source/Editor/Sequencer/Private/Filters/Filters/SequencerTrackFilter_Text.cpp b/Engine/Source/Editor/Sequencer/Private/Filters/Filters/SequencerTrackFilter_Text.cpp
--- a/Engine/Source/Editor/Sequencer/Private/Filters/Filters/SequencerTrackFilter_Text.cpp
+++ b/Engine/Source/Editor/Sequencer/Private/Filters/Filters/SequencerTrackFilter_Text.cpp
@@ -133,6 +133,11 @@ void FSequencerTrackFilter_Text::SetRawFilterText(const FText& InFilterText)
 	}
 }
 
+const FTextFilterExpressionEvaluator& FSequencerTrackFilter_Text::GetTextFilterExpressionEvaluator() const
+{
+	return TextFilterExpressionEvaluator;
+}
+
 const TArray<TSharedRef<FSequencerTextFilterExpressionContext>>& FSequencerTrackFilter_Text::GetTextFilterExpressionContexts() const
 {
 	return TextFilterExpressionContexts;
diff --git a/Engine/Source/Editor/Sequencer/Private/Filters/Filters/SequencerTrackFilter_Text.h b/Engine/Source/Editor/Sequencer/Private/Filters/Filters/SequencerTrackFilter_Text.h
--- a/Engine/Source/Editor/Sequencer/Private/Filters/Filters/SequencerTrackFilter_Text.h
+++ b/Engine/Source/Editor/Sequencer/Private/Filters/Filters/SequencerTrackFilter_Text.h
@@ -1,4 +1,4 @@
-﻿// Copyright Epic Games, Inc. All Rights Reserved.
+// Copyright Epic Games, Inc. All Rights Reserved.
 
 #pragma once
 
@@ -27,6 +27,7 @@ public:
 
 	FText GetFilterErrorText() const;
 
+	const FTextFilterExpressionEvaluator& GetTextFilterExpressionEvaluator() const;
 	const TArray<TSharedRef<FSequencerTextFilterExpressionContext>>& GetTextFilterExpressionContexts() const;
 
 	bool DoesTextFilterStringContainExpressionPair(const FSequencerTextFilterExpressionContext& InExpression) const;
diff --git a/Engine/Source/Editor/Sequencer/Private/Filters/SequencerFilterBar.cpp b/Engine/Source/Editor/Sequencer/Private/Filters/SequencerFilterBar.cpp
--- a/Engine/Source/Editor/Sequencer/Private/Filters/SequencerFilterBar.cpp
+++ b/Engine/Source/Editor/Sequencer/Private/Filters/SequencerFilterBar.cpp
@@ -32,6 +32,7 @@
 #include "Stats/StatsMisc.h"
 #include "UObject/UObjectIterator.h"
 #include "Widgets/SFilterBarIsolateHideShow.h"
+#include "Widgets/SSequencerSearchBox.h"
 
 using namespace UE::Sequencer;
 
@@ -272,7 +273,7 @@ FName FSequencerFilterBar::GetIdentifier() const
 	return TEXT("SequencerMain");
 }
 
-TSharedRef<SSequencerFilterBar> FSequencerFilterBar::GenerateWidget(const TSharedPtr<SFilterSearchBox>& InSearchBox, const EFilterBarLayout InLayout)
+TSharedRef<SSequencerFilterBar> FSequencerFilterBar::GenerateWidget(const TSharedPtr<SSequencerSearchBox>& InSearchBox, const EFilterBarLayout InLayout)
 {
 	return SNew(SSequencerFilterBar, SharedThis(this))
 		.FilterBarLayout(InLayout)
@@ -293,7 +294,7 @@ void FSequencerFilterBar::MuteFilters(const bool bInMute)
 	{
 		FilterBarWidget->SetEnabled(!bFiltersMuted);
 
-		const TWeakPtr<SFilterSearchBox> WeakSearchBox = FilterBarWidget->GetSearchBox();
+		const TWeakPtr<SSequencerSearchBox> WeakSearchBox = FilterBarWidget->GetSearchBox();
 		if (WeakSearchBox.IsValid())
 		{
 			WeakSearchBox.Pin()->SetEnabled(!bFiltersMuted);
diff --git a/Engine/Source/Editor/Sequencer/Private/Filters/SequencerFilterBar.h b/Engine/Source/Editor/Sequencer/Private/Filters/SequencerFilterBar.h
--- a/Engine/Source/Editor/Sequencer/Private/Filters/SequencerFilterBar.h
+++ b/Engine/Source/Editor/Sequencer/Private/Filters/SequencerFilterBar.h
@@ -1,4 +1,4 @@
-﻿// Copyright Epic Games, Inc. All Rights Reserved.
+// Copyright Epic Games, Inc. All Rights Reserved.
 
 #pragma once
 
@@ -21,8 +21,8 @@ class FSequencerTrackFilterMenu;
 class FUICommandList;
 class SComboButton;
 class SFilterBarIsolateHideShow;
-class SFilterSearchBox;
 class SSequencerFilterBar;
+class SSequencerSearchBox;
 enum class EFilterBarLayout : uint8;
 
 namespace UE::Sequencer
@@ -57,7 +57,7 @@ public:
 
 	TSharedPtr<ICustomTextFilter<FSequencerTrackFilterType>> CreateTextFilter();
 
-	TSharedRef<SSequencerFilterBar> GenerateWidget(const TSharedPtr<SFilterSearchBox>& InSearchBox, const EFilterBarLayout InLayout);
+	TSharedRef<SSequencerFilterBar> GenerateWidget(const TSharedPtr<SSequencerSearchBox>& InSearchBox, const EFilterBarLayout InLayout);
 
 	bool AreFiltersMuted() const;
 	void MuteFilters(const bool bInMute);
diff --git a/Engine/Source/Editor/Sequencer/Private/Filters/TextExpressions/SequencerTextFilterExpression_Time.cpp b/Engine/Source/Editor/Sequencer/Private/Filters/TextExpressions/SequencerTextFilterExpression_Time.cpp
--- a/Engine/Source/Editor/Sequencer/Private/Filters/TextExpressions/SequencerTextFilterExpression_Time.cpp
+++ b/Engine/Source/Editor/Sequencer/Private/Filters/TextExpressions/SequencerTextFilterExpression_Time.cpp
@@ -50,16 +50,16 @@ bool FSequencerTextFilterExpression_Time::TestComplexExpression(const FName& InK
 		return true;
 	}
 
-	const TViewModelPtr<FChannelGroupOutlinerModel> ChannelGroupOutlinerModel = CastViewModel<FChannelGroupOutlinerModel>(FilterItem);
-	if (!ChannelGroupOutlinerModel)
+	UMovieScene* const FocusedMovieScene = GetFocusedGetMovieScene();
+	if (!IsValid(FocusedMovieScene))
 	{
 		return true;
 	}
 
-	UMovieScene* const FocusedMovieScene = GetFocusedGetMovieScene();
-	if (!FocusedMovieScene)
+	const TViewModelPtr<FChannelGroupOutlinerModel> ChannelGroupOutlinerModel = FilterItem.ImplicitCast();
+	if (!ChannelGroupOutlinerModel.IsValid())
 	{
-		return true;
+		return false;
 	}
 
 	FSequencer& Sequencer = FilterInterface.GetSequencer();
@@ -79,26 +79,26 @@ bool FSequencerTextFilterExpression_Time::TestComplexExpression(const FName& InK
 		ValueToCheck = FString::FromInt(FrameTime.Time.FloorToFrame().Value);
 	}
 
-	const TArray<TSharedRef<IKeyArea>> KeyAreas = ChannelGroupOutlinerModel->GetAllKeyAreas();
+	const TArray<TSharedRef<IKeyArea>> CachedKeyAreas = ChannelGroupOutlinerModel->GetAllKeyAreas();
 
-	if (ValueToCheck.CompareFString(TEXT("NOW"), ETextFilterTextComparisonMode::Exact))
+	if (ValueToCheck.CompareFString(TEXT("Now"), ETextFilterTextComparisonMode::Exact))
 	{
 		const FTextFilterString NowTimeValue = FString::FromInt(Sequencer.GetGlobalTime().Time.FrameNumber.Value);
-		return CompareTime(Sequencer, NowTimeValue, KeyAreas, InComparisonOperation);
+		return CompareTime(Sequencer, NowTimeValue, CachedKeyAreas, InComparisonOperation);
 	}
-	if (ValueToCheck.CompareFString(TEXT("START"), ETextFilterTextComparisonMode::Exact))
+	if (ValueToCheck.CompareFString(TEXT("Start"), ETextFilterTextComparisonMode::Exact))
 	{
 		const FTextFilterString NowTimeValue = FString::FromInt(0);
-		return CompareTime(Sequencer, NowTimeValue, KeyAreas, InComparisonOperation);
+		return CompareTime(Sequencer, NowTimeValue, CachedKeyAreas, InComparisonOperation);
 	}
-	if (ValueToCheck.CompareFString(TEXT("END"), ETextFilterTextComparisonMode::Exact))
+	if (ValueToCheck.CompareFString(TEXT("End"), ETextFilterTextComparisonMode::Exact))
 	{
 		const FFrameNumber PlaybackRange = FocusedMovieScene->GetPlaybackRange().Size<FFrameNumber>();
 		const FTextFilterString NowTimeValue = FString::FromInt(PlaybackRange.Value);
-		return CompareTime(Sequencer, NowTimeValue, KeyAreas, InComparisonOperation);
+		return CompareTime(Sequencer, NowTimeValue, CachedKeyAreas, InComparisonOperation);
 	}
 
-	return CompareTime(Sequencer, ValueToCheck, KeyAreas, InComparisonOperation);
+	return CompareTime(Sequencer, ValueToCheck, CachedKeyAreas, InComparisonOperation);
 }
 
 bool FSequencerTextFilterExpression_Time::CompareTime(ISequencer& InSequencer
diff --git a/Engine/Source/Editor/Sequencer/Private/Filters/Widgets/SSequencerFilterBar.cpp b/Engine/Source/Editor/Sequencer/Private/Filters/Widgets/SSequencerFilterBar.cpp
--- a/Engine/Source/Editor/Sequencer/Private/Filters/Widgets/SSequencerFilterBar.cpp
+++ b/Engine/Source/Editor/Sequencer/Private/Filters/Widgets/SSequencerFilterBar.cpp
@@ -13,6 +13,7 @@
 #include "SequencerLog.h"
 #include "SSequencerFilter.h"
 #include "SSequencerCustomTextFilterDialog.h"
+#include "Widgets/SSequencerSearchBox.h"
 #include "Widgets/SWindow.h"
 
 using namespace UE::Sequencer;
@@ -130,7 +131,7 @@ void SSequencerFilterBar::SetTextFilterString(const FString& InText)
 		return;
 	}
 
-	if (const TSharedPtr<SFilterSearchBox> SearchBox = WeakSearchBox.Pin())
+	if (const TSharedPtr<SSequencerSearchBox> SearchBox = WeakSearchBox.Pin())
 	{
 		if (!InText.Equals(SearchBox->GetText().ToString()))
 		{
@@ -169,7 +170,7 @@ void SSequencerFilterBar::SetLayout(const EFilterBarLayout InFilterBarLayout)
 	}
 }
 
-void SSequencerFilterBar::AttachFilterSearchBox(const TSharedPtr<SFilterSearchBox>& InFilterSearchBox)
+void SSequencerFilterBar::AttachFilterSearchBox(const TSharedPtr<SSequencerSearchBox>& InFilterSearchBox)
 {
 	if (InFilterSearchBox)
 	{
@@ -379,7 +380,7 @@ UWorld* SSequencerFilterBar::GetWorld() const
 	return PlaybackContext->GetWorld();
 }
 
-TWeakPtr<SFilterSearchBox> SSequencerFilterBar::GetSearchBox() const
+TWeakPtr<SSequencerSearchBox> SSequencerFilterBar::GetSearchBox() const
 {
 	return WeakSearchBox;
 }
diff --git a/Engine/Source/Editor/Sequencer/Private/Filters/Widgets/SSequencerFilterBar.h b/Engine/Source/Editor/Sequencer/Private/Filters/Widgets/SSequencerFilterBar.h
--- a/Engine/Source/Editor/Sequencer/Private/Filters/Widgets/SSequencerFilterBar.h
+++ b/Engine/Source/Editor/Sequencer/Private/Filters/Widgets/SSequencerFilterBar.h
@@ -12,6 +12,7 @@ class FSequencerTrackFilter_CustomText;
 class SFilterBarClippingHorizontalBox;
 class SFilterExpressionHelpDialog;
 class SSequencerFilter;
+class SSequencerSearchBox;
 class UMovieSceneNodeGroup;
 enum class ESequencerFilterChange : uint8;
 
@@ -39,7 +40,7 @@ public:
 		/** An SFilterSearchBox that can be attached to this filter bar. When provided along with a CreateTextFilter
 		 *  delegate, allows the user to save searches from the Search Box as text filters for the filter bar.
 		 *	NOTE: Will bind a delegate to SFilterSearchBox::OnClickedAddSearchHistoryButton */
-		SLATE_ARGUMENT(TSharedPtr<SFilterSearchBox>, FilterSearchBox)
+		SLATE_ARGUMENT(TSharedPtr<SSequencerSearchBox>, FilterSearchBox)
 
 		/** The layout that determines how the filters are laid out */
 		SLATE_ARGUMENT(EFilterBarLayout, FilterBarLayout)
@@ -73,7 +74,7 @@ public:
 	EFilterBarLayout GetLayout() const;
 	void SetLayout(const EFilterBarLayout InFilterBarLayout);
 
-	void AttachFilterSearchBox(const TSharedPtr<SFilterSearchBox>& InFilterSearchBox);
+	void AttachFilterSearchBox(const TSharedPtr<SSequencerSearchBox>& InFilterSearchBox);
 
 	bool HasAnyFilterWidgets() const;
 
@@ -83,7 +84,7 @@ public:
 
 	void SaveCurrentFilterSetAsCustomTextFilter();
 
-	TWeakPtr<SFilterSearchBox> GetSearchBox() const;
+	TWeakPtr<SSequencerSearchBox> GetSearchBox() const;
 
 protected:
 	void AddWidgetToLayout(const TSharedRef<SWidget>& InWidget);
@@ -114,7 +115,7 @@ protected:
 
 	TWeakPtr<FSequencerFilterBar> WeakFilterBar;
 
-	TWeakPtr<SFilterSearchBox> WeakSearchBox;
+	TWeakPtr<SSequencerSearchBox> WeakSearchBox;
 
 	TSharedPtr<SWidgetSwitcher> FilterBoxWidget;
 	TSharedPtr<SFilterBarClippingHorizontalBox> HorizontalContainerWidget;
diff --git a/Engine/Source/Editor/Sequencer/Private/Filters/Widgets/SSequencerFilterSuggestionListRow.cpp b/Engine/Source/Editor/Sequencer/Private/Filters/Widgets/SSequencerFilterSuggestionListRow.cpp
new file mode 100644
--- /dev/null
+++ b/Engine/Source/Editor/Sequencer/Private/Filters/Widgets/SSequencerFilterSuggestionListRow.cpp
@@ -0,0 +1,101 @@
+﻿// Copyright Epic Games, Inc. All Rights Reserved.
+
+#include "SSequencerFilterSuggestionListRow.h"
+#include "Widgets/Layout/SBorder.h"
+#include "Widgets/Layout/SBox.h"
+#include "Widgets/Text/STextBlock.h"
+
+#define LOCTEXT_NAMESPACE "SSequencerFilterSuggestionListRow"
+
+void SSequencerFilterSuggestionListRow::Construct(const FArguments& InArgs, const TSharedRef<STableViewBase>& InOwnerTable)
+{
+	OwnerTable = InOwnerTable;
+
+	check(InArgs._ListItem.IsValid());
+	ListEntry = InArgs._ListItem;
+	HighlightText = InArgs._HighlightText;
+
+	SMultiColumnTableRow<TSharedPtr<FSequencerFilterSuggestionListEntryBase>>::Construct(FSuperRowType::FArguments(), InOwnerTable);
+}
+
+TSharedRef<SWidget> SSequencerFilterSuggestionListRow::GenerateWidgetForColumn(const FName& InColumnName)
+{
+	switch (ListEntry->GetItemType())
+	{
+	case ESequencerFilterSuggestionListEntryType::Header:
+		{
+			const TSharedPtr<FSequencerFilterSuggestionListHeaderEntry> HeaderListItem = StaticCastSharedPtr<FSequencerFilterSuggestionListHeaderEntry>(ListEntry);
+			check(HeaderListItem.IsValid());
+			return CreateHeaderItem(HeaderListItem, InColumnName);
+		}
+	case ESequencerFilterSuggestionListEntryType::Suggestion:
+		{
+			const TSharedPtr<FSequencerFilterSuggestionListEntry> SuggestionListItem = StaticCastSharedPtr<FSequencerFilterSuggestionListEntry>(ListEntry);
+			check(SuggestionListItem.IsValid());
+			return CreateSuggestionItem(SuggestionListItem, InColumnName);
+		}
+	}
+
+	return SNullWidget::NullWidget;
+}
+
+TSharedRef<SWidget> SSequencerFilterSuggestionListRow::CreateHeaderItem(const TSharedPtr<FSequencerFilterSuggestionListHeaderEntry>& InHeaderListItem, const FName& InColumnName)
+{
+	const TSharedRef<SVerticalBox> HeaderBox = SNew(SVerticalBox);
+
+	const bool bIsFirstItem = OwnerTable->GetNumGeneratedChildren() == 0;
+	if (!bIsFirstItem)
+	{
+		HeaderBox->AddSlot()
+			.AutoHeight()
+			.Padding(0.f, 4.f, 0.f, 2.f) // Add some empty space before the line, and a tiny bit after it
+			[
+				SNew(SBorder)
+				.Padding(FAppStyle::GetMargin(TEXT("Menu.Separator.Padding"))) // We'll use the border's padding to actually create the horizontal line
+				.BorderImage(FAppStyle::GetBrush(TEXT("Menu.Separator")))
+			];
+	}
+
+	HeaderBox->AddSlot()
+		.AutoHeight()
+		.Padding(3.f, 0.f)
+		[
+			SNew(STextBlock)
+			.Text(InHeaderListItem->Title)
+			.TextStyle(FAppStyle::Get(), TEXT("Menu.Heading"))
+		];
+
+	return HeaderBox;
+}
+
+TSharedRef<SWidget> SSequencerFilterSuggestionListRow::CreateSuggestionItem(const TSharedPtr<FSequencerFilterSuggestionListEntry>& InSuggestionListItem, const FName& InColumnName)
+{
+	if (InColumnName == TEXT("Suggestion"))
+	{
+		return SNew(SBox)
+			.VAlign(VAlign_Center)
+			.Padding(12.f, 1.f, 9.f, 1.f)
+			[
+				SNew(STextBlock)
+				.Text(InSuggestionListItem->DisplayName)
+				.ToolTipText(InSuggestionListItem->Description)
+				.HighlightText(HighlightText)
+			];
+	}
+
+	if (InColumnName == TEXT("Description"))
+	{
+		return SNew(SBox)
+			.VAlign(VAlign_Bottom)
+			.Padding(3.f, 1.f)
+			[
+				SNew(STextBlock)
+				.Text(InSuggestionListItem->Description)
+				.TextStyle(FAppStyle::Get(), TEXT("HintText"))
+			];
+	}
+
+	return SNullWidget::NullWidget;
+}
+
+#undef LOCTEXT_NAMESPACE
diff --git a/Engine/Source/Editor/Sequencer/Private/Filters/Widgets/SSequencerFilterSuggestionListRow.h b/Engine/Source/Editor/Sequencer/Private/Filters/Widgets/SSequencerFilterSuggestionListRow.h
new file mode 100644
--- /dev/null
+++ b/Engine/Source/Editor/Sequencer/Private/Filters/Widgets/SSequencerFilterSuggestionListRow.h
@@ -0,0 +1,33 @@
+﻿// Copyright Epic Games, Inc. All Rights Reserved.
+
+#pragma once
+
+#include "Filters/Widgets/SequencerFilterSuggestionListEntry.h"
+#include "Templates/SharedPointer.h"
+#include "Widgets/Views/SListView.h"
+
+class SSequencerFilterSuggestionListRow : public SMultiColumnTableRow<TSharedPtr<FSequencerFilterSuggestionListEntryBase>>
+{
+public:
+	SLATE_BEGIN_ARGS(SSequencerFilterSuggestionListRow)
+	{}
+		SLATE_ARGUMENT(TSharedPtr<FSequencerFilterSuggestionListEntryBase>, ListItem)
+		SLATE_ATTRIBUTE(FText, HighlightText)
+	SLATE_END_ARGS()
+
+	void Construct(const FArguments& InArgs, const TSharedRef<STableViewBase>& InOwnerTable);
+
+protected:
+	//~ Begin SMultiColumnTableRow
+	virtual TSharedRef<SWidget> GenerateWidgetForColumn(const FName& InColumnName) override;
+	//~ End SMultiColumnTableRow
+
+	TSharedRef<SWidget> CreateHeaderItem(const TSharedPtr<FSequencerFilterSuggestionListHeaderEntry>& InHeaderListItem, const FName& InColumnName);
+	TSharedRef<SWidget> CreateSuggestionItem(const TSharedPtr<FSequencerFilterSuggestionListEntry>& InSuggestionListItem, const FName& InColumnName);
+
+	TSharedPtr<STableViewBase> OwnerTable;
+
+	TSharedPtr<FSequencerFilterSuggestionListEntryBase> ListEntry;
+
+	TAttribute<FText> HighlightText;
+};
diff --git a/Engine/Source/Editor/Sequencer/Private/Filters/Widgets/SequencerFilterSuggestionListEntry.h b/Engine/Source/Editor/Sequencer/Private/Filters/Widgets/SequencerFilterSuggestionListEntry.h
new file mode 100644
--- /dev/null
+++ b/Engine/Source/Editor/Sequencer/Private/Filters/Widgets/SequencerFilterSuggestionListEntry.h
@@ -0,0 +1,72 @@
+﻿// Copyright Epic Games, Inc. All Rights Reserved.
+
+#pragma once
+
+#include "Internationalization/Text.h"
+#include "Templates/SharedPointer.h"
+
+enum class ESequencerFilterSuggestionListEntryType : uint8
+{
+	Header,
+	Suggestion
+};
+
+class FSequencerFilterSuggestionListEntryBase : public TSharedFromThis<FSequencerFilterSuggestionListEntryBase>
+{
+public:
+	FSequencerFilterSuggestionListEntryBase(const ESequencerFilterSuggestionListEntryType InItemType)
+		: ItemType(InItemType)
+	{}
+
+	ESequencerFilterSuggestionListEntryType GetItemType() const
+	{
+		return ItemType;
+	}
+
+	bool IsHeader() const
+	{
+		return ItemType == ESequencerFilterSuggestionListEntryType::Header;
+	}
+
+	TSharedPtr<class FSequencerFilterSuggestionListHeaderEntry> AsHeaderEntry()
+	{
+		return StaticCastSharedRef<FSequencerFilterSuggestionListHeaderEntry>(SharedThis(this));
+	}
+
+	TSharedPtr<class FSequencerFilterSuggestionListEntry> AsSuggestionEntry()
+	{
+		return StaticCastSharedRef<FSequencerFilterSuggestionListEntry>(SharedThis(this));
+	}
+
+protected:
+	ESequencerFilterSuggestionListEntryType ItemType;
+};
+
+
+class FSequencerFilterSuggestionListHeaderEntry : public FSequencerFilterSuggestionListEntryBase
+{
+public:
+	FSequencerFilterSuggestionListHeaderEntry(const FText& InTitle)
+		: FSequencerFilterSuggestionListEntryBase(ESequencerFilterSuggestionListEntryType::Header)
+		, Title(InTitle)
+	{}
+
+	FText Title;
+};
+
+
+class FSequencerFilterSuggestionListEntry : public FSequencerFilterSuggestionListEntryBase
+{
+public:
+	FSequencerFilterSuggestionListEntry()
+		: FSequencerFilterSuggestionListEntryBase(ESequencerFilterSuggestionListEntryType::Suggestion)
+	{}
+
+	FString Suggestion;
+
+	FText DisplayName;
+
+	FText CategoryName;
+
+	FText Description;
+};
diff --git a/Engine/Source/Editor/Sequencer/Private/SSequencer.cpp b/Engine/Source/Editor/Sequencer/Private/SSequencer.cpp
--- a/Engine/Source/Editor/Sequencer/Private/SSequencer.cpp
+++ b/Engine/Source/Editor/Sequencer/Private/SSequencer.cpp
@@ -138,6 +138,7 @@
 #include "Sidebar/SSidebar.h"
 #include "Sidebar/SSidebarContainer.h"
 #include "Widgets/SOverlay.h"
+#include "Widgets/SSequencerSearchBox.h"
 #include "IContentBrowserSingleton.h"
 #include "ContentBrowserModule.h"
 
@@ -1311,12 +1312,10 @@ TSharedRef<SWidget> SSequencer::ConstructSearchAndFilterRow()
 
 	if (!SearchBox.IsValid())
 	{
-		SearchBox = SNew(SFilterSearchBox)
+		SearchBox = SNew(SSequencerSearchBox, SharedThis(this))
 			.AddMetaData<FTagMetaData>(FTagMetaData(TEXT("SequencerFilterSearch")))
 			.HintText(LOCTEXT("FilterSearch", "Search..."))
 			.ToolTipText(LOCTEXT("FilterSearchHint", "Type here to search (pressing enter selects the results)"))
-			.DelayChangeNotificationsWhileTyping(true)
-			.ShowSearchHistory(true)
 			.OnTextChanged(this, &SSequencer::OnOutlinerSearchChanged)
 			.OnTextCommitted(this, &SSequencer::OnOutlinerSearchCommitted)
 			.OnSaveSearchClicked(this, &SSequencer::OnOutlinerSearchSaved);
diff --git a/Engine/Source/Editor/Sequencer/Private/SSequencer.h b/Engine/Source/Editor/Sequencer/Private/SSequencer.h
--- a/Engine/Source/Editor/Sequencer/Private/SSequencer.h
+++ b/Engine/Source/Editor/Sequencer/Private/SSequencer.h
@@ -47,6 +47,7 @@ class SSidebarContainer;
 class SWindow;
 class USequencerSettings;
 class SSequencerGroupManager;
+class SSequencerSearchBox;
 class SSequencerTreeFilterStatusBar;
 class SFilterSearchBox;
 enum class ESequencerFilterChange : uint8;
@@ -763,7 +764,7 @@ private:
 	TSharedPtr<SBreadcrumbTrail<FSequencerBreadcrumb>> BreadcrumbTrail;
 
 	/** The search box for filtering tracks. */
-	TSharedPtr<SFilterSearchBox> SearchBox;
+	TSharedPtr<SSequencerSearchBox> SearchBox;
 
 	/** The current playback time display. */
 	TSharedPtr<STemporarilyFocusedSpinBox<double>> PlayTimeDisplay;
diff --git a/Engine/Source/Editor/Sequencer/Private/Widgets/SSequencerSearchBox.cpp b/Engine/Source/Editor/Sequencer/Private/Widgets/SSequencerSearchBox.cpp
new file mode 100644
--- /dev/null
+++ b/Engine/Source/Editor/Sequencer/Private/Widgets/SSequencerSearchBox.cpp
@@ -0,0 +1,711 @@
+﻿// Copyright Epic Games, Inc. All Rights Reserved.
+
+#include "SSequencerSearchBox.h"
+#include "Framework/Application/SlateApplication.h"
+#include "SSequencer.h"
+#include "Filters/Filters/SequencerTrackFilter_Text.h"
+#include "Filters/SequencerFilterBar.h"
+#include "Filters/Widgets/SSequencerFilterSuggestionListRow.h"
+#include "Widgets/Input/SMenuAnchor.h"
+#include "Widgets/Layout/SBorder.h"
+#include "Widgets/Layout/SBox.h"
+#include "Widgets/SBoxPanel.h"
+#include "Widgets/Views/SListView.h"
+
+class ITableRow;
+class STableViewBase;
+class SWidget;
+struct FGeometry;
+
+#define LOCTEXT_NAMESPACE "SSequencerSearchBox"
+
+/** Case sensitive hashing function for TMap */
+template <typename ValueType>
+struct FSequencerSearchCategoryKeyMapFuncs : BaseKeyFuncs<ValueType, FText, /*bInAllowDuplicateKeys*/false>
+{
+	static FORCEINLINE const FString& GetSourceString(const FText& InText)
+	{
+		const FString* SourceString = FTextInspector::GetSourceString(InText);
+		check(SourceString);
+		return *SourceString;
+	}
+	static FORCEINLINE const FText& GetSetKey(const TPair<FText, ValueType>& Element)
+	{
+		return Element.Key;
+	}
+	static FORCEINLINE bool Matches(const FText& A, const FText& B)
+	{
+		return GetSourceString(A).Equals(GetSourceString(B), ESearchCase::CaseSensitive);
+	}
+	static FORCEINLINE uint32 GetKeyHash(const FText& Key)
+	{
+		return FLocKey::ProduceHash(GetSourceString(Key));
+	}
+};
+
+void SSequencerSearchBox::Construct(const FArguments& InArgs, const TWeakPtr<SSequencer>& InWeakSequencerWidget)
+{
+	WeakSequencerWidget = InWeakSequencerWidget;
+
+	OnTextChanged = InArgs._OnTextChanged;
+	OnTextCommitted = InArgs._OnTextCommitted;
+	OnKeyDownHandler = InArgs._OnKeyDownHandler;
+	PossibleSuggestions = InArgs._PossibleSuggestions;
+	OnSuggestionFilter = InArgs._OnSuggestionFilter;
+	OnSuggestionChosen = InArgs._OnSuggestionChosen;
+	PreCommittedText = InArgs._InitialText.Get();
+	bMustMatchPossibleSuggestions = InArgs._MustMatchPossibleSuggestions.Get();
+
+	if (!OnSuggestionFilter.IsBound())
+	{
+		OnSuggestionFilter.BindSP(this, &SSequencerSearchBox::DefaultSuggestionFilterImpl);
+	}
+
+	if (!OnSuggestionChosen.IsBound())
+	{
+		OnSuggestionChosen.BindSP(this, &SSequencerSearchBox::DefaultSuggestionChosenImpl);
+	}
+
+	ChildSlot
+		[
+			SAssignNew(SuggestionMenuAnchor, SMenuAnchor)
+			.Placement(InArgs._SuggestionListPlacement)
+			[
+				/** Use SFilterSearchBox internally to add the ability to show search history and save searches as filters */
+				SAssignNew(SearchBox, SFilterSearchBox)
+				.InitialText(InArgs._InitialText)
+				.HintText(InArgs._HintText)
+				.ShowSearchHistory(InArgs._ShowSearchHistory)
+				.DelayChangeNotificationsWhileTyping(InArgs._DelayChangeNotificationsWhileTyping)
+				.OnTextChanged(this, &SSequencerSearchBox::HandleTextChanged)
+				.OnTextCommitted(this, &SSequencerSearchBox::HandleTextCommitted)
+				.OnKeyDownHandler(this, &SSequencerSearchBox::HandleKeyDown)
+				.OnSaveSearchClicked(InArgs._OnSaveSearchClicked)
+			]
+			.MenuContent(GetSuggestionListMenuContent())
+		];
+}
+
+TSharedRef<SWidget> SSequencerSearchBox::GetSuggestionListMenuContent()
+{
+	const USequencerSettings* const SequencerSettings = GetDefault<USequencerSettings>();
+	check(IsValid(SequencerSettings));
+
+	return SNew(SBorder)
+		.BorderImage(FAppStyle::GetBrush(TEXT("Menu.Background")))
+		.Padding(FMargin(2))
+		[
+			SNew(SBox)
+			.MinDesiredWidth(180)
+			.MinDesiredHeight(16)
+			.MaxDesiredHeight(260)
+			[
+				SAssignNew(SuggestionListView, SListView<TSharedPtr<FSequencerFilterSuggestionListEntryBase>>)
+				.ListItemsSource(&SuggestionItems)
+				.SelectionMode(ESelectionMode::Single) // Ideally the mouse over would not highlight while keyboard controls the UI
+				.OnGenerateRow(this, &SSequencerSearchBox::MakeSuggestionListItemWidget)
+				.OnSelectionChanged(this, &SSequencerSearchBox::OnSelectionChanged)
+				.ScrollbarDragFocusCause(EFocusCause::SetDirectly) // Use SetDirect so that clicking the scrollbar doesn't close the suggestions list
+				.HeaderRow
+				(
+					SNew(SHeaderRow)
+					.Visibility(EVisibility::Collapsed)
+					+ SHeaderRow::Column("Suggestion")
+					.DefaultLabel(LOCTEXT("SuggestionColumnHeader", "Suggestion"))
+					.FixedWidth(180.f)
+					.VAlignCell(VAlign_Center)
+					.HAlignCell(HAlign_Left)
+					+ SHeaderRow::Column("Description")
+					.DefaultLabel(LOCTEXT("DescriptionColumnHeader", "Description"))
+					.FillWidth(1.f)
+					.VAlignCell(VAlign_Center)
+					.HAlignCell(HAlign_Left)
+				)
+			]
+		];
+}
+
+FText SSequencerSearchBox::GetText() const
+{
+	return SearchBox->GetText();
+}
+
+void SSequencerSearchBox::SetText(const TAttribute<FText>& InNewText)
+{
+	SearchBox->SetText(InNewText);
+	PreCommittedText = InNewText.Get();
+
+	// Set the cursor interaction location to the end of the newly inserted suggestion
+	if (LastChosenCursorOffset != INDEX_NONE)
+	{
+		const FTextLocation CursorLocation = FTextLocation(0, LastChosenCursorOffset);
+		SearchBox->SelectText(CursorLocation, CursorLocation);
+
+		LastChosenCursorOffset = INDEX_NONE;
+	}
+}
+
+void SSequencerSearchBox::SetError(const FText& InError)
+{
+	SearchBox->SetError(InError);
+}
+
+void SSequencerSearchBox::SetError(const FString& InError)
+{
+	SearchBox->SetError(InError);
+}
+
+FReply SSequencerSearchBox::OnPreviewKeyDown(const FGeometry& InGeometry, const FKeyEvent& InKeyEvent)
+{
+	if (SuggestionMenuAnchor->IsOpen() && InKeyEvent.GetKey() == EKeys::Escape)
+	{
+		// Clear any selection first to prevent the currently selection being set in the text box
+		SuggestionListView->ClearSelection();
+		SuggestionMenuAnchor->SetIsOpen(false, false);
+
+		return FReply::Handled();
+	}
+
+	return FReply::Unhandled();
+}
+
+FReply SSequencerSearchBox::HandleKeyDown(const FGeometry& InGeometry, const FKeyEvent& InKeyEvent)
+{
+	if (!SearchBox->GetText().IsEmpty()
+		&& (InKeyEvent.GetKey() == EKeys::Up || InKeyEvent.GetKey() == EKeys::Down))
+	{
+		if (!SuggestionMenuAnchor->IsOpen())
+		{
+			SuggestionMenuAnchor->SetIsOpen(true, false);
+		}
+
+		const bool bSelectingUp = InKeyEvent.GetKey() == EKeys::Up;
+		const TSharedPtr<FSequencerFilterSuggestionListEntryBase> SelectedSuggestion = GetSelectedSuggestion();
+
+		int32 TargetIdx = INDEX_NONE;
+		if (SelectedSuggestion.IsValid())
+		{
+			const int32 SelectionDirection = bSelectingUp ? -1 : 1;
+
+			// Select the next non-header suggestion, based on the direction of travel
+			TargetIdx = SuggestionItems.IndexOfByKey(SelectedSuggestion);
+			if (SuggestionItems.IsValidIndex(TargetIdx))
+			{
+				do
+				{
+					TargetIdx += SelectionDirection;
+				}
+				while (SuggestionItems.IsValidIndex(TargetIdx) && SuggestionItems[TargetIdx]->IsHeader());
+			}
+		}
+		else if (!bSelectingUp && SuggestionItems.Num() > 0)
+		{
+			// Nothing selected and pressed down, select the first non-header suggestion
+			TargetIdx = 0;
+			while (SuggestionItems.IsValidIndex(TargetIdx) && SuggestionItems[TargetIdx]->IsHeader())
+			{
+				TargetIdx += 1;
+			}
+		}
+
+		if (SuggestionItems.IsValidIndex(TargetIdx))
+		{
+			SuggestionListView->SetSelection(SuggestionItems[TargetIdx]);
+			SuggestionListView->RequestScrollIntoView(SuggestionItems[TargetIdx]);
+		}
+
+		return FReply::Handled();
+	}
+
+	if (OnKeyDownHandler.IsBound())
+	{
+		return OnKeyDownHandler.Execute(InGeometry, InKeyEvent);
+	}
+
+	return FReply::Unhandled();
+}
+
+bool SSequencerSearchBox::SupportsKeyboardFocus() const
+{
+	return SearchBox->SupportsKeyboardFocus();
+}
+
+bool SSequencerSearchBox::HasKeyboardFocus() const
+{
+	// Since keyboard focus is forwarded to our editable text, we will test it instead
+	return SearchBox->HasKeyboardFocus();
+}
+
+FReply SSequencerSearchBox::OnFocusReceived(const FGeometry& InGeometry, const FFocusEvent& InFocusEvent)
+{
+	// Forward keyboard focus to our editable text widget
+	return SearchBox->OnFocusReceived(InGeometry, InFocusEvent);
+}
+
+void SSequencerSearchBox::HandleTextChanged(const FText& InNewText)
+{
+	LastCursorLocation = SearchBox->GetSelection().GetEnd();
+
+	OnTextChanged.ExecuteIfBound(InNewText);
+
+	UpdateSuggestionList();
+}
+
+void SSequencerSearchBox::HandleTextCommitted(const FText& InNewText, const ETextCommit::Type InCommitType)
+{
+	TSharedPtr<FSequencerFilterSuggestionListEntryBase> SelectedSuggestion = GetSelectedSuggestion();
+
+	FText CommittedText;
+	if (SelectedSuggestion.IsValid() && !SelectedSuggestion->IsHeader() && InCommitType != ETextCommit::OnCleared)
+	{
+		CommittedText = OnSuggestionChosen.Execute(InNewText, SelectedSuggestion->AsSuggestionEntry()->Suggestion);
+	}
+	else
+	{
+		if (InCommitType == ETextCommit::OnCleared)
+		{
+			// Clear text when escape is pressed then commit an empty string
+			CommittedText = FText::GetEmpty();
+		}
+		else if (bMustMatchPossibleSuggestions)
+		{
+			const FString NewTextStr = InNewText.ToString();
+			const bool bIsSuggestion = PossibleSuggestions.Get().ContainsByPredicate(
+				[NewTextStr](const FSequencerFilterSuggestion& InSuggestion)
+				{
+					return InSuggestion.Suggestion == NewTextStr;
+				});
+			if (bIsSuggestion)
+			{
+				CommittedText = InNewText;
+			}
+			else
+			{
+				// commit the original text if we have to match a suggestion
+				CommittedText = PreCommittedText;
+			}
+		}
+		else
+		{
+			// otherwise, set the typed text
+			CommittedText = InNewText;
+		}
+	}
+
+	// Set the text and execute the delegate
+	bDisableOpeningSuggestions = true;
+	SetText(CommittedText);
+	OnTextCommitted.ExecuteIfBound(CommittedText, InCommitType);
+
+	if (InCommitType != ETextCommit::Default)
+	{
+		// Clear the suggestion box if the user has navigated away or set their own text
+		SuggestionMenuAnchor->SetIsOpen(false, false);
+	}
+}
+
+void SSequencerSearchBox::OnSelectionChanged(TSharedPtr<FSequencerFilterSuggestionListEntryBase> InNewValue, const ESelectInfo::Type InSelectInfo)
+{
+	// If the user clicked directly on an item to select it, then accept the choice and close the window
+	if (InSelectInfo == ESelectInfo::OnMouseClick && !InNewValue->IsHeader())
+	{
+		const FText NewText = OnSuggestionChosen.Execute(SearchBox->GetText(), InNewValue->AsSuggestionEntry()->Suggestion);
+		SetText(NewText);
+
+		SuggestionMenuAnchor->SetIsOpen(false, false);
+		FocusEditBox();
+	}
+}
+
+TSharedRef<ITableRow> SSequencerSearchBox::MakeSuggestionListItemWidget(const TSharedPtr<FSequencerFilterSuggestionListEntryBase> InSuggestion, const TSharedRef<STableViewBase>& InOwnerTable)
+{
+	return SNew(SSequencerFilterSuggestionListRow, InOwnerTable)
+		.ListItem(InSuggestion)
+		.HighlightText(this, &SSequencerSearchBox::GetHighlightText);
+}
+
+FText SSequencerSearchBox::GetHighlightText() const
+{
+	return SuggestionHighlightText;
+}
+
+void SSequencerSearchBox::UpdateSuggestionList()
+{
+	FText SearchText = SearchBox->GetText();
+
+	SuggestionItems.Reset();
+	SuggestionHighlightText = FText::GetEmpty();
+
+	if (!SearchText.IsEmpty())
+	{
+		typedef TMap<FText
+			, TArray<TSharedPtr<FSequencerFilterSuggestionListEntryBase>>
+			, FDefaultSetAllocator
+			, FSequencerSearchCategoryKeyMapFuncs<TArray<TSharedPtr<FSequencerFilterSuggestionListEntryBase>>>> FCategorizedSuggestionsMap;
+
+		FParsedExpression ParsedExpression;
+		ExtractSearchFilterTerms(SearchText, ParsedExpression);
+
+		// Get the potential suggestions
+		TArray<FSequencerFilterSuggestion> FilteredSuggestions;
+		if (PossibleSuggestions.IsBound())
+		{
+			FilteredSuggestions = PossibleSuggestions.Get();
+		}
+		else
+		{
+			if (!ParsedExpression.Key.IsSet() && ParsedExpression.Value.IsSet())
+			{
+				DefaultKeySuggestions(SearchText.ToString(), FilteredSuggestions);
+			}
+			else if (ParsedExpression.Key.IsSet())
+			{
+				DefaultValueSuggestions(ParsedExpression.Key.GetValue(), FilteredSuggestions);
+			}
+
+			if (ParsedExpression.Value.IsSet())
+			{
+				SearchText = FText::FromString(ParsedExpression.Value.GetValue());
+			}
+		}
+
+		// Run them through the filter
+		OnSuggestionFilter.Execute(SearchText, FilteredSuggestions, SuggestionHighlightText);
+
+		// Split the suggestions list into categories
+		FCategorizedSuggestionsMap CategorizedSuggestions;
+		for (const FSequencerFilterSuggestion& Suggestion : FilteredSuggestions)
+		{
+			TArray<TSharedPtr<FSequencerFilterSuggestionListEntryBase>>& CategorySuggestions = CategorizedSuggestions.FindOrAdd(Suggestion.CategoryName);
+
+			const TSharedRef<FSequencerFilterSuggestionListEntry> NewSuggestionEntry = MakeShared<FSequencerFilterSuggestionListEntry>();
+			NewSuggestionEntry->Suggestion = Suggestion.Suggestion;
+			NewSuggestionEntry->DisplayName = Suggestion.DisplayName;
+			NewSuggestionEntry->Description = Suggestion.Description;
+
+			CategorySuggestions.Add(NewSuggestionEntry);
+		}
+
+		// Rebuild the flat list in categorized groups
+		// If there is only one category, and that category is empty (undefined), then skip adding the category headers
+		const bool bSkipCategoryHeaders = CategorizedSuggestions.Num() == 1 && CategorizedSuggestions.Contains(FText::GetEmpty());
+		for (const TPair<FText, TArray<TSharedPtr<FSequencerFilterSuggestionListEntryBase>>>& CategorySuggestionsPair : CategorizedSuggestions)
+		{
+			if (!bSkipCategoryHeaders)
+			{
+				const FText CategoryDisplayName = CategorySuggestionsPair.Key.IsEmpty()
+					? LOCTEXT("UndefinedCategory", "Undefined")
+					: CategorySuggestionsPair.Key;
+
+				const TSharedRef<FSequencerFilterSuggestionListHeaderEntry> NewHeaderEntry = MakeShared<FSequencerFilterSuggestionListHeaderEntry>(CategoryDisplayName);
+
+				SuggestionItems.Add(NewHeaderEntry);
+			}
+			SuggestionItems.Append(CategorySuggestionsPair.Value);
+		}
+	}
+
+	if (!bDisableOpeningSuggestions && SuggestionItems.Num() > 0 && HasKeyboardFocus())
+	{
+		SuggestionMenuAnchor->SetIsOpen(true, false);
+	}
+	else
+	{
+		SuggestionMenuAnchor->SetIsOpen(false, false);
+	}
+
+	bDisableOpeningSuggestions = false;
+
+	SuggestionListView->RequestListRefresh();
+}
+
+void SSequencerSearchBox::FocusEditBox()
+{
+	FWidgetPath WidgetToFocusPath;
+	FSlateApplication::Get().GeneratePathToWidgetUnchecked(SearchBox.ToSharedRef(), WidgetToFocusPath);
+	FSlateApplication::Get().SetKeyboardFocus(WidgetToFocusPath, EFocusCause::SetDirectly);
+}
+
+TSharedPtr<FSequencerFilterSuggestionListEntryBase> SSequencerSearchBox::GetSelectedSuggestion() const
+{
+	TSharedPtr<FSequencerFilterSuggestionListEntryBase> SelectedSuggestion;
+
+	if (SuggestionMenuAnchor->IsOpen())
+	{
+		const TArray<TSharedPtr<FSequencerFilterSuggestionListEntryBase>>& SelectedSuggestionList = SuggestionListView->GetSelectedItems();
+		if (SelectedSuggestionList.Num() > 0)
+		{
+			// Selection mode is Single, so there should only be one suggestion at the most
+			SelectedSuggestion = SelectedSuggestionList[0];
+		}
+	}
+
+	return SelectedSuggestion;
+}
+
+void SSequencerSearchBox::SetOnSaveSearchHandler(SFilterSearchBox::FOnSaveSearchClicked InOnSaveSearchHandler)
+{
+	SearchBox->SetOnSaveSearchHandler(InOnSaveSearchHandler);
+}
+
+void SSequencerSearchBox::ExtractSearchFilterTerms(const FText& InSearchText, FParsedExpression& OutParsedPair)
+{
+	const TSharedPtr<SSequencer> SequencerWidget = WeakSequencerWidget.Pin();
+	if (!SequencerWidget.IsValid())
+	{
+		return;
+	}
+
+	const TSharedPtr<FSequencerFilterBar> FilterBar = SequencerWidget->GetFilterBar();
+	if (!FilterBar.IsValid())
+	{
+		return;
+	}
+
+	const FTextFilterExpressionEvaluator& TextFilterExpressionEvaluator = FilterBar->GetTextFilter()->GetTextFilterExpressionEvaluator();
+	const TArray<FExpressionToken>& ExpressionTokens = TextFilterExpressionEvaluator.GetFilterExpressionTokens();
+	if (ExpressionTokens.Num() == 0)
+	{
+		return;
+	}
+
+	const FString SearchString = InSearchText.ToString();
+
+	OutParsedPair.SuggestionIndex = SearchString.Len();
+
+	const int32 CaretLocationIndex = LastCursorLocation.GetOffset();
+	const int32 CaretTokenIndex = FindTokenIndex(ExpressionTokens, CaretLocationIndex);
+
+	if (CaretTokenIndex == INDEX_NONE)
+	{
+		return;
+	}
+
+	// Inspect the tokens to see what the last part of the search term was
+	// If it was a key->value pair then we'll use that to control what kinds of results we show
+	// For anything else we just use the text from the last token as our filter term to allow incremental auto-complete
+	const FExpressionToken& LastToken = ExpressionTokens[CaretTokenIndex];
+
+	const int32 SecondToLastIndex = CaretTokenIndex - 1;
+	const int32 ThirdToLastIndex = CaretTokenIndex - 2;
+
+	// If the last token is a text token, then consider it as a value and walk back to see if we also have a key
+	if (LastToken.Node.Cast<TextFilterExpressionParser::FTextToken>())
+	{
+		OutParsedPair.Value = LastToken.Context.GetString();
+		OutParsedPair.SuggestionIndex = LastToken.Context.GetCharacterIndex();
+
+		if (ExpressionTokens.IsValidIndex(ThirdToLastIndex))
+		{
+			// Check if the second from last token is a operator (=, !=, <, >, etc)
+			const FExpressionToken& ComparisonToken = ExpressionTokens[SecondToLastIndex];
+			if (IsOperatorToken(ComparisonToken))
+			{
+				const FExpressionToken& KeyToken = ExpressionTokens[ThirdToLastIndex];
+				if (KeyToken.Node.Cast<TextFilterExpressionParser::FTextToken>())
+				{
+					OutParsedPair.Key = KeyToken.Context.GetString();
+					OutParsedPair.SuggestionIndex = KeyToken.Context.GetCharacterIndex();
+				}
+			}
+			else if (IsLogicalOperatorToken(ComparisonToken))
+			{
+				const FExpressionToken& KeyToken = ExpressionTokens[SecondToLastIndex];
+				if (KeyToken.Node.Cast<TextFilterExpressionParser::FTextToken>())
+				{
+					OutParsedPair.Key = KeyToken.Context.GetString();
+					OutParsedPair.SuggestionIndex = LastToken.Context.GetCharacterIndex();
+				}
+			}
+		}
+	}
+	// If the last token is a comparison operator, then walk back and see if we have a key
+	else if (IsOperatorToken(LastToken))
+	{
+		if (ExpressionTokens.IsValidIndex(SecondToLastIndex))
+		{
+			const FExpressionToken& KeyToken = ExpressionTokens[SecondToLastIndex];
+			if (KeyToken.Node.Cast<TextFilterExpressionParser::FTextToken>())
+			{
+				OutParsedPair.Key = KeyToken.Context.GetString();
+				OutParsedPair.Value = TEXT("");
+				OutParsedPair.SuggestionIndex = LastToken.Context.GetCharacterIndex();
+			}
+		}
+	}
+}
+
+void SSequencerSearchBox::DefaultSuggestionFilterImpl(const FText& InSearchText, TArray<FSequencerFilterSuggestion>& OutPossibleSuggestions, FText& OutSuggestionHighlightText)
+{
+	OutPossibleSuggestions.RemoveAll([SearchStr = InSearchText.ToString()](const FSequencerFilterSuggestion& InSuggestion)
+		{
+			return !InSuggestion.Suggestion.Contains(SearchStr);
+		});
+
+	OutSuggestionHighlightText = InSearchText;
+}
+
+FText SSequencerSearchBox::DefaultSuggestionChosenImpl(const FText& InSearchText, const FString& InSuggestion)
+{
+	const FTextFilterExpressionEvaluator* const TextFilterExpressionEvaluator = GetTextFilterExpressionEvaluator();
+	if (!TextFilterExpressionEvaluator)
+	{
+		return InSearchText;
+	}
+
+	const TArray<FExpressionToken>& ExpressionTokens = TextFilterExpressionEvaluator->GetFilterExpressionTokens();
+	const int32 CaretLocationIndex = LastCursorLocation.GetOffset();
+	const int32 CaretTokenIndex = FindTokenIndex(ExpressionTokens, CaretLocationIndex);
+	const FExpressionToken& CaretToken = ExpressionTokens[CaretTokenIndex];
+	const int32 SuggestionInsertionIndex = CaretToken.Context.GetCharacterIndex();
+	const FString CaretTokenString = CaretToken.Context.GetString();
+
+	// Replace the value with the suggestion
+	FString SearchString = InSearchText.ToString();
+	SearchString.RemoveAt(SuggestionInsertionIndex, CaretTokenString.Len(), EAllowShrinking::No);
+	SearchString.InsertAt(SuggestionInsertionIndex, InSuggestion);
+
+	// Set the cursor interaction location to the end of the newly inserted suggestion
+	LastChosenCursorOffset = SuggestionInsertionIndex + InSuggestion.Len();
+
+	return FText::FromString(SearchString);
+}
+
+void SSequencerSearchBox::DefaultKeySuggestions(const FString& InKeyValue, TArray<FSequencerFilterSuggestion>& OutPossibleSuggestions)
+{
+	const TSharedPtr<SSequencer> SequencerWidget = WeakSequencerWidget.Pin();
+	if (!SequencerWidget.IsValid())
+	{
+		return;
+	}
+
+	const TSharedPtr<FSequencerFilterBar> FilterBar = SequencerWidget->GetFilterBar();
+	if (!FilterBar.IsValid())
+	{
+		return;
+	}
+
+	const TArray<TSharedRef<FSequencerTextFilterExpressionContext>>& FilterExpressions = FilterBar->GetTextFilter()->GetTextFilterExpressionContexts();
+
+	for (const TSharedRef<FSequencerTextFilterExpressionContext>& Expression : FilterExpressions)
+	{
+		const TArray<FName> Keys = Expression->GetKeys().Array();
+		for (const FName Key : Keys)
+		{
+			const FString KeyString = Key.ToString();
+
+			FSequencerFilterSuggestion NewSuggestion;
+			NewSuggestion.Suggestion = FString::Printf(TEXT("%s"), *KeyString);
+			NewSuggestion.DisplayName = FText::FromString(KeyString);
+			NewSuggestion.Description = Expression->GetDescription();
+
+			OutPossibleSuggestions.Add(NewSuggestion);
+		}
+	}
+
+	OutPossibleSuggestions.Sort([](const FSequencerFilterSuggestion& InA, const FSequencerFilterSuggestion& InB)
+		{
+			return InA.DisplayName.CompareTo(InB.DisplayName) < 0;
+		});
+}
+
+void SSequencerSearchBox::DefaultValueSuggestions(const FString& InKeyValue, TArray<FSequencerFilterSuggestion>& OutPossibleSuggestions)
+{
+	const TSharedPtr<SSequencer> SequencerWidget = WeakSequencerWidget.Pin();
+	if (!SequencerWidget.IsValid())
+	{
+		return;
+	}
+
+	const TSharedPtr<FSequencerFilterBar> FilterBar = SequencerWidget->GetFilterBar();
+	if (!FilterBar.IsValid())
+	{
+		return;
+	}
+
+	const TArray<TSharedRef<FSequencerTextFilterExpressionContext>>& FilterExpressions = FilterBar->GetTextFilter()->GetTextFilterExpressionContexts();
+
+	FSequencerTextFilterKeyword FoundKeyword;
+	TSharedPtr<FSequencerTextFilterExpressionContext> FoundExpression;
+
+	for (const TSharedRef<FSequencerTextFilterExpressionContext>& Expression : FilterExpressions)
+	{
+		for (const FName Key : Expression->GetKeys())
+		{
+			if (InKeyValue.Equals(Key.ToString(), ESearchCase::IgnoreCase))
+			{
+				FoundExpression = Expression;
+				goto end;
+			}
+		}
+	}
+
+end:
+	if (FoundExpression.IsValid())
+	{
+		const TArray<FSequencerTextFilterKeyword> TextFilterKeywords = FoundExpression->GetValueKeywords();
+		for (const FSequencerTextFilterKeyword& TextFilterKeyword : TextFilterKeywords)
+		{
+			FSequencerFilterSuggestion NewSuggestion;
+			NewSuggestion.Suggestion = FString::Printf(TEXT("%s"), *TextFilterKeyword.Keyword);
+			NewSuggestion.DisplayName = FText::FromString(TextFilterKeyword.Keyword);
+			NewSuggestion.Description = TextFilterKeyword.Description;
+
+			OutPossibleSuggestions.Add(NewSuggestion);
+		}
+
+		OutPossibleSuggestions.Sort([](const FSequencerFilterSuggestion& InA, const FSequencerFilterSuggestion& InB)
+			{
+				return InA.DisplayName.CompareTo(InB.DisplayName) < 0;
+			});
+	}
+}
+
+bool SSequencerSearchBox::IsOperatorToken(const FExpressionToken& InToken)
+{
+	return InToken.Node.Cast<TextFilterExpressionParser::FEqual>()
+		|| InToken.Node.Cast<TextFilterExpressionParser::FNotEqual>()
+		|| InToken.Node.Cast<TextFilterExpressionParser::FLess>()
+		|| InToken.Node.Cast<TextFilterExpressionParser::FLessOrEqual>()
+		|| InToken.Node.Cast<TextFilterExpressionParser::FGreater>()
+		|| InToken.Node.Cast<TextFilterExpressionParser::FGreaterOrEqual>();
+};
+
+bool SSequencerSearchBox::IsLogicalOperatorToken(const FExpressionToken& InToken)
+{
+	return InToken.Node.Cast<TextFilterExpressionParser::FAnd>()
+		|| InToken.Node.Cast<TextFilterExpressionParser::FOr>();
+};
+
+int32 SSequencerSearchBox::FindTokenIndex(const TArray<FExpressionToken>& ExpressionTokens, const int32 InIndexToFind)
+{
+	for (int32 Index = ExpressionTokens.Num() - 1; Index > 0; --Index)
+	{
+		const int32 CharacterIndex = ExpressionTokens[Index].Context.GetCharacterIndex();
+		const int32 TokenLength = ExpressionTokens[Index].Context.GetString().Len();
+		if (InIndexToFind > CharacterIndex)// && InIndexToFind <= CharacterIndex + TokenLength)
+		{
+			return Index;
+		}
+	}
+	return ExpressionTokens.IsEmpty() ? INDEX_NONE : 0;
+}
+
+const FTextFilterExpressionEvaluator* SSequencerSearchBox::GetTextFilterExpressionEvaluator() const
+{
+	const TSharedPtr<SSequencer> SequencerWidget = WeakSequencerWidget.Pin();
+	if (!SequencerWidget.IsValid())
+	{
+		return nullptr;
+	}
+
+	const TSharedPtr<FSequencerFilterBar> FilterBar = SequencerWidget->GetFilterBar();
+	if (!FilterBar.IsValid())
+	{
+		return nullptr;
+	}
+
+	return &FilterBar->GetTextFilter()->GetTextFilterExpressionEvaluator();
+}
+
+#undef LOCTEXT_NAMESPACE
diff --git a/Engine/Source/Editor/Sequencer/Private/Widgets/SSequencerSearchBox.h b/Engine/Source/Editor/Sequencer/Private/Widgets/SSequencerSearchBox.h
new file mode 100644
--- /dev/null
+++ b/Engine/Source/Editor/Sequencer/Private/Widgets/SSequencerSearchBox.h
@@ -0,0 +1,213 @@
+﻿// Copyright Epic Games, Inc. All Rights Reserved.
+
+#pragma once
+
+#include "Filters/SFilterSearchBox.h"
+#include "Framework/Text/TextLayout.h"
+#include "Widgets/DeclarativeSyntaxSupport.h"
+#include "Widgets/SCompoundWidget.h"
+
+class FExpressionToken;
+class FSequencerFilterSuggestionListEntryBase;
+class FTextFilterExpressionEvaluator;
+class ITableRow;
+class SMenuAnchor;
+class SSequencer;
+class STableViewBase;
+struct FFocusEvent;
+struct FGeometry;
+struct FKeyEvent;
+
+struct FSequencerFilterSuggestion
+{
+	FString Suggestion;
+
+	FText DisplayName;
+
+	FText CategoryName;
+
+	FText Description;
+};
+
+/** A delegate for a callback to filter the given suggestion list, to allow custom filtering behavior */
+DECLARE_DELEGATE_ThreeParams(FOnSequencerSearchBoxSuggestionFilter, const FText& /*InSearchText*/, TArray<FSequencerFilterSuggestion>& /*InPossibleSuggestions*/, FText& /*InSuggestionHighlightText*/);
+
+/** A delegate for a callback when a suggestion entry is chosen during a search, to allow custom compositing behavior of the suggestion into the search text */
+DECLARE_DELEGATE_RetVal_TwoParams(FText, FOnSequencerSearchBoxSuggestionChosen, const FText& /*InSearchText*/, const FString& /*InSuggestion*/);
+
+/**
+ * A wrapper widget around SFilterSearchBox to provide filter text expression suggestions in a dropdown menu.
+ */
+class SSequencerSearchBox : public SCompoundWidget
+{
+public:
+	SLATE_BEGIN_ARGS(SSequencerSearchBox)
+		: _SuggestionListPlacement(MenuPlacement_BelowAnchor)
+		, _PossibleSuggestions(TArray<FSequencerFilterSuggestion>())
+		, _DelayChangeNotificationsWhileTyping(true)
+		, _MustMatchPossibleSuggestions(false)
+		, _ShowSearchHistory(true)
+	{}
+
+		/** Where to place the suggestion list */
+		SLATE_ARGUMENT(EMenuPlacement, SuggestionListPlacement)
+
+		/** Invoked whenever the text changes */
+		SLATE_EVENT(FOnTextChanged, OnTextChanged)
+
+		/** Invoked whenever the text is committed (e.g. user presses enter) */
+		SLATE_EVENT(FOnTextCommitted, OnTextCommitted)
+
+		/** Initial text to display for the search text */
+		SLATE_ATTRIBUTE(FText, InitialText)
+
+		/** Hint text to display for the search text when there is no value */
+		SLATE_ATTRIBUTE(FText, HintText)
+
+		/** All possible suggestions for the search text */
+		SLATE_ATTRIBUTE(TArray<FSequencerFilterSuggestion>, PossibleSuggestions)
+
+		/** Whether the SearchBox should delay notifying listeners of text changed events until the user is done typing */
+		SLATE_ATTRIBUTE(bool, DelayChangeNotificationsWhileTyping)
+
+		/** Whether the SearchBox allows entries that don't match the possible suggestions */
+		SLATE_ATTRIBUTE(bool, MustMatchPossibleSuggestions)
+
+		/** Callback to filter the given suggestion list, to allow custom filtering behavior */
+		SLATE_EVENT(FOnSequencerSearchBoxSuggestionFilter, OnSuggestionFilter)
+
+		/** Callback when a suggestion entry is chosen during an asset search, to allow custom compositing behavior of the suggestion into the search text */
+		SLATE_EVENT(FOnSequencerSearchBoxSuggestionChosen, OnSuggestionChosen)
+
+		/** Callback delegate to have first chance handling of the OnKeyDown event */
+		SLATE_EVENT(FOnKeyDown, OnKeyDownHandler)
+	
+		/** Whether we should show a dropdown containing the last few searches */
+		SLATE_ATTRIBUTE(bool, ShowSearchHistory)
+
+		/** Handler for when the + Button next to a search is clicked */
+		SLATE_EVENT(SFilterSearchBox::FOnSaveSearchClicked, OnSaveSearchClicked)
+
+	SLATE_END_ARGS()
+
+	/** Constructs this widget with InArgs */
+	void Construct(const FArguments& InArgs, const TWeakPtr<SSequencer>& InWeakSequencerWidget);
+
+	/** Get the text string currently being edited */
+	FText GetText() const;
+
+	/** Sets the text string currently being edited */
+	void SetText(const TAttribute<FText>& InNewText);
+
+	/** Set or clear the current error reporting information for this search box */
+	void SetError(const FText& InError);
+	void SetError(const FString& InError);
+
+	/** Show a + button next to the current search and set the handler for when that is clicked */
+	void SetOnSaveSearchHandler(SFilterSearchBox::FOnSaveSearchClicked InOnSaveSearchHandler);
+
+	//~ Begin SWidget
+	virtual FReply OnPreviewKeyDown(const FGeometry& InGeometry, const FKeyEvent& InKeyEvent) override;
+	virtual bool SupportsKeyboardFocus() const override;
+	virtual bool HasKeyboardFocus() const override;
+	virtual FReply OnFocusReceived(const FGeometry& InGeometry, const FFocusEvent& InFocusEvent) override;
+	//~ End SWidget
+
+private:
+	struct FParsedExpression
+	{
+		TOptional<FString> Key;
+		TOptional<FString> Value;
+		int32 SuggestionIndex = 0;
+	};
+
+	void ExtractSearchFilterTerms(const FText& InSearchText, FParsedExpression& OutParsedPair);
+
+	TSharedRef<SWidget> GetSuggestionListMenuContent();
+
+	/** First chance handler for key down events to the editable text widget */
+	FReply HandleKeyDown(const FGeometry& InGeometry, const FKeyEvent& InKeyEvent);
+
+	/** Handler for when text in the editable text box changed */
+	void HandleTextChanged(const FText& InNewText);
+
+	/** Handler for when text in the editable text box changed */
+	void HandleTextCommitted(const FText& InNewText, const ETextCommit::Type InCommitType);
+
+	/** Called by SListView when the selection changes in the suggestion list */
+	void OnSelectionChanged(TSharedPtr<FSequencerFilterSuggestionListEntryBase> InNewValue, const ESelectInfo::Type InSelectInfo);
+
+	/** Makes the widget for a suggestion message in the list view */
+	TSharedRef<ITableRow> MakeSuggestionListItemWidget(const TSharedPtr<FSequencerFilterSuggestionListEntryBase> InSuggestion, const TSharedRef<STableViewBase>& InOwnerTable);
+
+	/** Gets the text to highlight in the suggestion list */
+	FText GetHighlightText() const;
+
+	/** Updates and shows or hides the suggestion list */
+	void UpdateSuggestionList();
+
+	/** Sets the focus to the InputText box */
+	void FocusEditBox();
+
+	/** Returns the currently selected suggestion */
+	TSharedPtr<FSequencerFilterSuggestionListEntryBase> GetSelectedSuggestion() const;
+
+	/** Default implementation of OnSuggestionFilter, if no external implementation is provided */
+	void DefaultSuggestionFilterImpl(const FText& InSearchText, TArray<FSequencerFilterSuggestion>& OutPossibleSuggestions, FText& OutSuggestionHighlightText);
+
+	/** Default implementation of OnSuggestionChosen, if no external implementation is provided */
+	FText DefaultSuggestionChosenImpl(const FText& InSearchText, const FString& InSuggestion);
+
+	/** Provides default key suggestions, if no suggestions are provided */
+	void DefaultKeySuggestions(const FString& InKeyValue, TArray<FSequencerFilterSuggestion>& OutPossibleSuggestions);
+
+	/** Provides default value suggestions, if no suggestions are provided */
+	void DefaultValueSuggestions(const FString& InKeyValue, TArray<FSequencerFilterSuggestion>& OutPossibleSuggestions);
+
+	bool IsOperatorToken(const FExpressionToken& InToken);
+	bool IsLogicalOperatorToken(const FExpressionToken& InToken);
+
+	int32 FindTokenIndex(const TArray<FExpressionToken>& ExpressionTokens, const int32 InIndexToFind);
+
+	const FTextFilterExpressionEvaluator* GetTextFilterExpressionEvaluator() const;
+
+	TWeakPtr<SSequencer> WeakSequencerWidget;
+
+	TSharedPtr<SFilterSearchBox> SearchBox;
+
+	TSharedPtr<SListView<TSharedPtr<FSequencerFilterSuggestionListEntryBase>>> SuggestionListView;
+
+	TSharedPtr<SMenuAnchor> SuggestionMenuAnchor;
+
+	TArray<TSharedPtr<FSequencerFilterSuggestionListEntryBase>> SuggestionItems;
+
+	/** The state of the text prior to being committed */
+	FText PreCommittedText;
+
+	/** The highlight text to use for the suggestions list */
+	FText SuggestionHighlightText;
+
+	FOnTextChanged OnTextChanged;
+	FOnTextCommitted OnTextCommitted;
+
+	/** Delegate to filter the given suggestion list, to allow custom filtering behavior */
+	FOnSequencerSearchBoxSuggestionFilter OnSuggestionFilter;
+
+	/** Delegate when a suggestion entry is chosen during a search, to allow custom compositing behavior of the suggestion into the search text */
+	FOnSequencerSearchBoxSuggestionChosen OnSuggestionChosen;
+
+	/** Delegate for first chance handling for key down events */
+	FOnKeyDown OnKeyDownHandler;
+
+	/** All possible suggestions for the search text */
+	TAttribute<TArray<FSequencerFilterSuggestion>> PossibleSuggestions;
+
+	/** Determines whether the committed text should match a suggestion */
+	bool bMustMatchPossibleSuggestions = false;
+
+	bool bDisableOpeningSuggestions = true;
+
+	FTextLocation LastCursorLocation = FTextLocation(0, 0);
+
+	int32 LastChosenCursorOffset = INDEX_NONE;
+};
