Sequencer: Wrap the advanced time panel in a scrollbox so that the supported frame rates don't draw outside

#rnx

diff --git a/Engine/Source/Editor/Sequencer/Private/SSequencer.cpp b/Engine/Source/Editor/Sequencer/Private/SSequencer.cpp
--- a/Engine/Source/Editor/Sequencer/Private/SSequencer.cpp
+++ b/Engine/Source/Editor/Sequencer/Private/SSequencer.cpp
@@ -1121,7 +1121,7 @@ void SSequencer::OpenTickResolutionOptions()
 	TSharedRef<SWindow> TickResolutionOptionsWindow = SNew(SWindow)
 		.Title(LOCTEXT("TickResolutionOptions_Title", "Advanced Time Properties"))
 		.SupportsMaximize(false)
-		.ClientSize(FVector2D(600.f, 500.f))
+		.ClientSize(FVector2D(600.f, 510.f))
 		.Content()
 		[
 			SNew(SSequencerTimePanel, SequencerPtr)
diff --git a/Engine/Source/Editor/Sequencer/Private/SSequencerTimePanel.cpp b/Engine/Source/Editor/Sequencer/Private/SSequencerTimePanel.cpp
--- a/Engine/Source/Editor/Sequencer/Private/SSequencerTimePanel.cpp
+++ b/Engine/Source/Editor/Sequencer/Private/SSequencerTimePanel.cpp
@@ -50,161 +50,165 @@ void SSequencerTimePanel::Construct(const FArguments& InArgs, TWeakPtr<FSequence
 	static FLinearColor WarningColor(FColor(0xffbbbb44));
 	ChildSlot
 	[
-		SNew(SBorder)
-		.BorderImage(FAppStyle::Get().GetBrush("Brushes.Panel"))
-		.Padding(16.f)
+		SNew(SScrollBox)
+		+ SScrollBox::Slot()
 		[
-			SNew(SVerticalBox)
-
-			+SVerticalBox::Slot()
-			.AutoHeight()
+			SNew(SBorder)
+			.BorderImage(FAppStyle::Get().GetBrush("Brushes.Panel"))
+			.Padding(16.f)
 			[
-				SNew(STextBlock)
-				.AutoWrapText(true)
-				.Text(Description)
-			]
+				SNew(SVerticalBox)
 
-			+SVerticalBox::Slot()
-			.FillHeight(1.0f)
-			.Padding(0.0f, 10.0f)
-			[
-				SNew(SBorder)
-				.Padding(FMargin(10.0f, 10.0f))
-				.BorderImage(FAppStyle::Get().GetBrush("Brushes.Recessed"))
+				+SVerticalBox::Slot()
+				.AutoHeight()
 				[
-					SNew(SGridPanel)
-					.FillColumn(1, 1.f)
-
-					+ SGridPanel::Slot(0, 0)
-					.Padding(Col1Padding)
-					.HAlign(HAlign_Right)
-					[
-						SNew(STextBlock)
-						.Text(LOCTEXT("NewTickInterval", "Desired Tick Interval"))
-					]
-					+ SGridPanel::Slot(1, 0)
-					.HAlign(HAlign_Left)
-					.Padding(Col2Padding)
-					[
-						SNew(SFrameRatePicker)
-						.RecommendedText(LOCTEXT("CompatibleWithDisplayRate", "Compatible with this sequence"))
-						.NotRecommendedText(LOCTEXT("NotCompatibleWithDisplayRate", "Other"))
-						.NotRecommendedToolTip(LOCTEXT("NotCompatibleWithDisplayRate_Tip", "All other preset frame rates that are not compatible with the current display and tick rate"))
-						.IsPresetRecommended(this, &SSequencerTimePanel::IsRecommendedResolution)
-						.PresetValues(MoveTemp(TickResolutionRates))
-						.Value(this, &SSequencerTimePanel::GetCurrentTickResolution)
-						.OnValueChanged(this, &SSequencerTimePanel::OnSetTickResolution)
-					]
+					SNew(STextBlock)
+					.AutoWrapText(true)
+					.Text(Description)
+				]
 
-					+ SGridPanel::Slot(0, 1)
-					.HAlign(HAlign_Right)
-					.Padding(FMargin(0, 0, HorizontalGridPadding, 0))
-					[
-						SNew(STextBlock)
-						.Text(LOCTEXT("ApplyRecursively", "Apply Recursively"))
-						.ToolTipText(LOCTEXT("ApplyRecursively_Tooltip", "If true then the tick rate change will propagate to all child sub-sequences of the current sequence. The UI warnings only apply to the top level sequence."))
-					]
-					+ SGridPanel::Slot(1, 1)
-					.HAlign(HAlign_Left)
-					.Padding(FMargin(HorizontalGridPadding, 0, 0, 0))
+				+SVerticalBox::Slot()
+				.FillHeight(1.0f)
+				.Padding(0.0f, 10.0f)
+				[
+					SNew(SBorder)
+					.Padding(FMargin(10.0f, 10.0f))
+					.BorderImage(FAppStyle::Get().GetBrush("Brushes.Recessed"))
 					[
-						SNew(SCheckBox)
-						.IsChecked(this, &SSequencerTimePanel::GetIsRecursive)
-						.OnCheckStateChanged(this, &SSequencerTimePanel::OnSetIsRecursive)
+						SNew(SGridPanel)
+						.FillColumn(1, 1.f)
+
+						+ SGridPanel::Slot(0, 0)
+						.Padding(Col1Padding)
+						.HAlign(HAlign_Right)
+						[
+							SNew(STextBlock)
+							.Text(LOCTEXT("NewTickInterval", "Desired Tick Interval"))
+						]
+						+ SGridPanel::Slot(1, 0)
+						.HAlign(HAlign_Left)
+						.Padding(Col2Padding)
+						[
+							SNew(SFrameRatePicker)
+							.RecommendedText(LOCTEXT("CompatibleWithDisplayRate", "Compatible with this sequence"))
+							.NotRecommendedText(LOCTEXT("NotCompatibleWithDisplayRate", "Other"))
+							.NotRecommendedToolTip(LOCTEXT("NotCompatibleWithDisplayRate_Tip", "All other preset frame rates that are not compatible with the current display and tick rate"))
+							.IsPresetRecommended(this, &SSequencerTimePanel::IsRecommendedResolution)
+							.PresetValues(MoveTemp(TickResolutionRates))
+							.Value(this, &SSequencerTimePanel::GetCurrentTickResolution)
+							.OnValueChanged(this, &SSequencerTimePanel::OnSetTickResolution)
+						]
+
+						+ SGridPanel::Slot(0, 1)
+						.HAlign(HAlign_Right)
+						.Padding(FMargin(0, 0, HorizontalGridPadding, 0))
+						[
+							SNew(STextBlock)
+							.Text(LOCTEXT("ApplyRecursively", "Apply Recursively"))
+							.ToolTipText(LOCTEXT("ApplyRecursively_Tooltip", "If true then the tick rate change will propagate to all child sub-sequences of the current sequence. The UI warnings only apply to the top level sequence."))
+						]
+						+ SGridPanel::Slot(1, 1)
+						.HAlign(HAlign_Left)
+						.Padding(FMargin(HorizontalGridPadding, 0, 0, 0))
+						[
+							SNew(SCheckBox)
+							.IsChecked(this, &SSequencerTimePanel::GetIsRecursive)
+							.OnCheckStateChanged(this, &SSequencerTimePanel::OnSetIsRecursive)
+						]
+
+						+ SGridPanel::Slot(0, 2)
+						.ColumnSpan(2)
+						.Padding(FMargin(0.f, VerticalGridPadding, 0.f, VerticalGridPadding))
+						.HAlign(HAlign_Left)
+						.VAlign(VAlign_Top)
+						[
+							SNew(STextBlock)
+							.Text(LOCTEXT("NewTickInterval_Tip", "Sequence will have the following properties if applied:"))
+						]
+
+						+ SGridPanel::Slot(0, 3)
+						.Padding(Col1Padding)
+						.HAlign(HAlign_Right)
+						[
+							SNew(STextBlock)
+							.Text(LOCTEXT("ResultingRange", "Time Range"))
+						]
+						+ SGridPanel::Slot(1, 3)
+						.Padding(Col2Padding)
+						[
+							SNew(STextBlock)
+							.Text(this, &SSequencerTimePanel::GetSupportedTimeRange)
+						]
+
+						+ SGridPanel::Slot(0, 4)
+						.Padding(Col1Padding)
+						.HAlign(HAlign_Right)
+						[
+							SNew(STextBlock)
+							.Text(LOCTEXT("SupportedFrameRates", "Supported Rates"))
+						]
+						+ SGridPanel::Slot(1, 4)
+						.Padding(Col2Padding)
+						[
+							SAssignNew(CommonFrameRates, SVerticalBox)
+						]
 					]
+				]
 
-					+ SGridPanel::Slot(0, 2)
-					.ColumnSpan(2)
-					.Padding(FMargin(0.f, VerticalGridPadding, 0.f, VerticalGridPadding))
-					.HAlign(HAlign_Left)
-					.VAlign(VAlign_Top)
-					[
-						SNew(STextBlock)
-						.Text(LOCTEXT("NewTickInterval_Tip", "Sequence will have the following properties if applied:"))
-					]
+				+ SVerticalBox::Slot()
+				.AutoHeight()
+				.Padding(5.f)
+				.HAlign(HAlign_Center)
+				.VAlign(VAlign_Center)
+				[
+					SNew(SHorizontalBox)
+					.Visibility(this, &SSequencerTimePanel::GetWarningVisibility)
 
-					+ SGridPanel::Slot(0, 3)
-					.Padding(Col1Padding)
-					.HAlign(HAlign_Right)
-					[
-						SNew(STextBlock)
-						.Text(LOCTEXT("ResultingRange", "Time Range"))
-					]
-					+ SGridPanel::Slot(1, 3)
-					.Padding(Col2Padding)
+					+ SHorizontalBox::Slot()
+					.Padding(FMargin(0.f, 0.f, 7.f, 0.f))
+					.AutoWidth()
 					[
 						SNew(STextBlock)
-						.Text(this, &SSequencerTimePanel::GetSupportedTimeRange)
+						.ColorAndOpacity(WarningColor)
+						.Font(FAppStyle::Get().GetFontStyle("FontAwesome.11"))
+						.Text(FText::FromString(FString(TEXT("\xf071"))) /*fa-exclamation-triangle*/)
 					]
 
-					+ SGridPanel::Slot(0, 4)
-					.Padding(Col1Padding)
-					.HAlign(HAlign_Right)
+					+ SHorizontalBox::Slot()
+					.AutoWidth()
 					[
 						SNew(STextBlock)
-						.Text(LOCTEXT("SupportedFrameRates", "Supported Rates"))
+						.ColorAndOpacity(WarningColor)
+						.Text(LOCTEXT("ApplyWarning", "Applying these settings may result in changes to key positions or section boundaries."))
 					]
-					+ SGridPanel::Slot(1, 4)
-					.Padding(Col2Padding)
-					[
-						SAssignNew(CommonFrameRates, SVerticalBox)
-					]
-				]
-			]
-
-			+ SVerticalBox::Slot()
-			.AutoHeight()
-			.Padding(5.f)
-			.HAlign(HAlign_Center)
-			.VAlign(VAlign_Center)
-			[
-				SNew(SHorizontalBox)
-				.Visibility(this, &SSequencerTimePanel::GetWarningVisibility)
-
-				+ SHorizontalBox::Slot()
-				.Padding(FMargin(0.f, 0.f, 7.f, 0.f))
-				.AutoWidth()
-				[
-					SNew(STextBlock)
-					.ColorAndOpacity(WarningColor)
-					.Font(FAppStyle::Get().GetFontStyle("FontAwesome.11"))
-					.Text(FText::FromString(FString(TEXT("\xf071"))) /*fa-exclamation-triangle*/)
 				]
 
-				+ SHorizontalBox::Slot()
-				.AutoWidth()
-				[
-					SNew(STextBlock)
-					.ColorAndOpacity(WarningColor)
-					.Text(LOCTEXT("ApplyWarning", "Applying these settings may result in changes to key positions or section boundaries."))
-				]
-			]
-
-			+ SVerticalBox::Slot()
-			.AutoHeight()
-			.HAlign(HAlign_Right)
-			.VAlign(VAlign_Bottom)
-			.Padding(0.0f, 3.0f)
-			[
-				SNew(SHorizontalBox)
-
-				+SHorizontalBox::Slot()
-				.AutoWidth()
+				+ SVerticalBox::Slot()
+				.AutoHeight()
+				.HAlign(HAlign_Right)
 				.VAlign(VAlign_Bottom)
-				.Padding(8, 0)
+				.Padding(0.0f, 3.0f)
 				[
-					SNew(SPrimaryButton)
-					.Text(LOCTEXT("ApplyButtonText", "Apply"))
-					.OnClicked(this, &SSequencerTimePanel::Apply)
-				]
+					SNew(SHorizontalBox)
 
-				+SHorizontalBox::Slot()
-				.AutoWidth()
-				.VAlign(VAlign_Bottom)
-				[
-					SNew(SButton)
-					.Text(LOCTEXT("CancelButtonText", "Cancel"))
-					.OnClicked(this, &SSequencerTimePanel::Close)
+					+SHorizontalBox::Slot()
+					.AutoWidth()
+					.VAlign(VAlign_Bottom)
+					.Padding(8, 0)
+					[
+						SNew(SPrimaryButton)
+						.Text(LOCTEXT("ApplyButtonText", "Apply"))
+						.OnClicked(this, &SSequencerTimePanel::Apply)
+					]
+
+					+SHorizontalBox::Slot()
+					.AutoWidth()
+					.VAlign(VAlign_Bottom)
+					[
+						SNew(SButton)
+						.Text(LOCTEXT("CancelButtonText", "Cancel"))
+						.OnClicked(this, &SSequencerTimePanel::Close)
+					]
 				]
 			]
 		]
