[Sequencer]
 - Isolate newly added sequencer object tracks if there is already an existing isolation
 - Remove hidden/isolated tracks on delete

diff --git a/Engine/Source/Editor/Sequencer/Private/Filters/Filters/SequencerTrackFilter_HideIsolate.cpp b/Engine/Source/Editor/Sequencer/Private/Filters/Filters/SequencerTrackFilter_HideIsolate.cpp
--- a/Engine/Source/Editor/Sequencer/Private/Filters/Filters/SequencerTrackFilter_HideIsolate.cpp
+++ b/Engine/Source/Editor/Sequencer/Private/Filters/Filters/SequencerTrackFilter_HideIsolate.cpp
@@ -101,6 +101,19 @@ void FSequencerTrackFilter_HideIsolate::HideTracks(const TSet<TWeakViewModelPtr<
 	BroadcastChangedEvent();
 }
 
+void FSequencerTrackFilter_HideIsolate::UnhideTracks(const TSet<TWeakViewModelPtr<IOutlinerExtension>>& InTracks)
+{
+	for (const TWeakViewModelPtr<IOutlinerExtension>& TrackModekWeak : InTracks)
+	{
+		if (const TViewModelPtr<IOutlinerExtension> TrackModel = TrackModekWeak.Pin())
+		{
+			HiddenTracks.Remove(TrackModel);
+		}
+	}
+
+	BroadcastChangedEvent();
+}
+
 void FSequencerTrackFilter_HideIsolate::IsolateTracks(const TSet<TWeakViewModelPtr<IOutlinerExtension>>& InTracks, const bool bInAddToExisting)
 {
 	if (!bInAddToExisting)
diff --git a/Engine/Source/Editor/Sequencer/Private/Filters/Filters/SequencerTrackFilter_HideIsolate.h b/Engine/Source/Editor/Sequencer/Private/Filters/Filters/SequencerTrackFilter_HideIsolate.h
--- a/Engine/Source/Editor/Sequencer/Private/Filters/Filters/SequencerTrackFilter_HideIsolate.h
+++ b/Engine/Source/Editor/Sequencer/Private/Filters/Filters/SequencerTrackFilter_HideIsolate.h
@@ -15,6 +15,8 @@ public:
 	TSet<UE::Sequencer::TWeakViewModelPtr<UE::Sequencer::IOutlinerExtension>> GetIsolatedTracks() const;
 
 	void HideTracks(const TSet<UE::Sequencer::TWeakViewModelPtr<UE::Sequencer::IOutlinerExtension>>& InTracks, const bool bInAddToExisting);
+	void UnhideTracks(const TSet<UE::Sequencer::TWeakViewModelPtr<UE::Sequencer::IOutlinerExtension>>& InTracks);
+
 	void IsolateTracks(const TSet<UE::Sequencer::TWeakViewModelPtr<UE::Sequencer::IOutlinerExtension>>& InTracks, const bool bInAddToExisting);
 	void UnisolateTracks(const TSet<UE::Sequencer::TWeakViewModelPtr<UE::Sequencer::IOutlinerExtension>>& InTracks);
 
diff --git a/Engine/Source/Editor/Sequencer/Private/Filters/SequencerFilterBar.cpp b/Engine/Source/Editor/Sequencer/Private/Filters/SequencerFilterBar.cpp
--- a/Engine/Source/Editor/Sequencer/Private/Filters/SequencerFilterBar.cpp
+++ b/Engine/Source/Editor/Sequencer/Private/Filters/SequencerFilterBar.cpp
@@ -367,11 +367,21 @@ void FSequencerFilterBar::HideTracks(const TSet<TWeakViewModelPtr<IOutlinerExten
 	HideIsolateFilter->HideTracks(InTracks, bInAddToExisting);
 }
 
+void FSequencerFilterBar::UnhideTracks(const TSet<TWeakViewModelPtr<IOutlinerExtension>>& InTracks)
+{
+	HideIsolateFilter->UnhideTracks(InTracks);
+}
+
 void FSequencerFilterBar::IsolateTracks(const TSet<TWeakViewModelPtr<IOutlinerExtension>>& InTracks, const bool bInAddToExisting)
 {
 	HideIsolateFilter->IsolateTracks(InTracks, bInAddToExisting);
 }
 
+void FSequencerFilterBar::UnisolateTracks(const TSet<TWeakViewModelPtr<IOutlinerExtension>>& InTracks)
+{
+	HideIsolateFilter->UnisolateTracks(InTracks);
+}
+
 void FSequencerFilterBar::ShowAllTracks()
 {
 	HideIsolateFilter->ShowAllTracks();
diff --git a/Engine/Source/Editor/Sequencer/Private/Filters/SequencerFilterBar.h b/Engine/Source/Editor/Sequencer/Private/Filters/SequencerFilterBar.h
--- a/Engine/Source/Editor/Sequencer/Private/Filters/SequencerFilterBar.h
+++ b/Engine/Source/Editor/Sequencer/Private/Filters/SequencerFilterBar.h
@@ -77,7 +77,11 @@ public:
 	TSet<UE::Sequencer::TWeakViewModelPtr<UE::Sequencer::IOutlinerExtension>> GetIsolatedTracks() const;
 
 	void HideTracks(const TSet<UE::Sequencer::TWeakViewModelPtr<UE::Sequencer::IOutlinerExtension>>& InTracks, const bool bInAddToExisting);
+	void UnhideTracks(const TSet<UE::Sequencer::TWeakViewModelPtr<UE::Sequencer::IOutlinerExtension>>& InTracks);
+
 	void IsolateTracks(const TSet<UE::Sequencer::TWeakViewModelPtr<UE::Sequencer::IOutlinerExtension>>& InTracks, const bool bInAddToExisting);
+	void UnisolateTracks(const TSet<UE::Sequencer::TWeakViewModelPtr<UE::Sequencer::IOutlinerExtension>>& InTracks);
+
 	void ShowAllTracks();
 
 	bool HasHiddenTracks() const;
diff --git a/Engine/Source/Editor/Sequencer/Private/SSequencer.cpp b/Engine/Source/Editor/Sequencer/Private/SSequencer.cpp
--- a/Engine/Source/Editor/Sequencer/Private/SSequencer.cpp
+++ b/Engine/Source/Editor/Sequencer/Private/SSequencer.cpp
@@ -2864,6 +2864,21 @@ void SSequencer::UpdateLayoutTree()
 			NodePathToRename.Empty();
 		}
 
+		// Isolate binding object guids after the tree view is refreshed and the new tracks are created
+		if (!NewNodePathsToIsolate.IsEmpty())
+		{
+			for (const TViewModelPtr<IOutlinerExtension>& OutlinerItem : Sequencer->GetNodeTree()->GetRootNode()->GetDescendantsOfType<IOutlinerExtension>())
+			{
+				const FString ItemPath = IOutlinerExtension::GetPathName(OutlinerItem);
+				if (NewNodePathsToIsolate.Contains(ItemPath))
+				{
+					Sequencer->GetFilterBar()->IsolateTracks({ OutlinerItem }, true);
+					NewNodePathsToIsolate.Remove(ItemPath);
+				}
+			}
+			NewNodePathsToIsolate.Empty();
+		}
+
 		if (Sequencer->GetFocusedMovieSceneSequence())
 		{
 			bool bAnyChanged = false;
diff --git a/Engine/Source/Editor/Sequencer/Private/SSequencer.h b/Engine/Source/Editor/Sequencer/Private/SSequencer.h
--- a/Engine/Source/Editor/Sequencer/Private/SSequencer.h
+++ b/Engine/Source/Editor/Sequencer/Private/SSequencer.h
@@ -641,6 +641,9 @@ public:
 	/** This adds the specified path to the selection set to be restored the next time the tree view is refreshed. */
 	void AddAdditionalPathToSelectionSet(const FString& Path) { AdditionalSelectionsToAdd.Add(Path); }
 
+	/** Adds a node path to be isolated after the tree view is refreshed and the new tracks are created. */
+	void AddNewNodePathsToIsolate(const TSet<FString>& InPaths) { NewNodePathsToIsolate.Append(InPaths); }
+
 	/** Request to rename the given node path. */
 	void RequestRenameNode(const FString& Path) { NodePathToRename = Path; }
 
@@ -891,6 +894,9 @@ private:
 
 	FString NodePathToRename;
 
+	/** List of node paths that should be isolated on next tree update */
+	TSet<FString> NewNodePathsToIsolate;
+
 	TWeakPtr<SWindow> WeakTickResolutionOptionsWindow;
 
 	TWeakPtr<SWindow> WeakExposedBindingsWindow;
diff --git a/Engine/Source/Editor/Sequencer/Private/Sequencer.cpp b/Engine/Source/Editor/Sequencer/Private/Sequencer.cpp
--- a/Engine/Source/Editor/Sequencer/Private/Sequencer.cpp
+++ b/Engine/Source/Editor/Sequencer/Private/Sequencer.cpp
@@ -5882,6 +5882,19 @@ bool FSequencer::OnRequestNodeDeleted( TSharedRef<FViewModel> NodeToBeDeleted, c
 	if (IDeletableExtension* Deletable = NodeToBeDeleted->CastThis<IDeletableExtension>())
 	{
 		Deletable->Delete();
+
+		const TViewModelPtr<IOutlinerExtension> OutlinerItem = NodeToBeDeleted->CastThisShared<IOutlinerExtension>();
+
+		if (FilterBar->HasIsolatedTracks())
+		{
+			FilterBar->UnisolateTracks({ OutlinerItem });
+		}
+
+		if (FilterBar->HasHiddenTracks())
+		{
+			FilterBar->UnhideTracks({ OutlinerItem });
+		}
+
 		return true;
 	}
 
@@ -9193,6 +9206,8 @@ void FixSortingOrders(FMovieSceneBinding* InBinding, UMovieScene* MovieScene)
 
 void FSequencer::OnAddBinding(const FGuid& ObjectBinding, UMovieScene* MovieScene)
 {
+	using namespace UE::Sequencer;
+
 	FMovieSceneBinding* Binding = MovieScene->FindBinding(ObjectBinding);
 	if (Binding)
 	{
@@ -9200,6 +9215,19 @@ void FSequencer::OnAddBinding(const FGuid& ObjectBinding, UMovieScene* MovieScen
 	}
 
 	NodeTree->SortAllNodesAndDescendants();
+
+	if (FilterBar->HasIsolatedTracks())
+	{
+		for (const TViewModelPtr<IObjectBindingExtension>& ObjectBindingItem : NodeTree->GetRootNode()->GetDescendantsOfType<IObjectBindingExtension>())
+		{
+			if (ObjectBinding == ObjectBindingItem->GetObjectGuid())
+			{
+				const FString NewNodePath = IOutlinerExtension::GetPathName(ObjectBindingItem);
+				SequencerWidget->AddNewNodePathsToIsolate({ NewNodePath });
+				break;
+			}
+		}
+	}
 }
 
 void FSequencer::OnAddTrack(const TWeakObjectPtr<UMovieSceneTrack>&InTrack, const FGuid & ObjectBinding)
@@ -9253,6 +9281,11 @@ void FSequencer::OnAddTrack(const TWeakObjectPtr<UMovieSceneTrack>&InTrack, cons
 	ThrobSectionSelection();
 
 	NodeTree->SortAllNodesAndDescendants();
+
+	if (FilterBar->HasIsolatedTracks())
+	{
+		SequencerWidget->AddNewNodePathsToIsolate({ NewNodePath });
+	}
 }
 
 
