Sequencer: Blueprint read only properties should not be keyable

diff --git a/Engine/Source/Editor/Sequencer/Private/SequencerObjectChangeListener.cpp b/Engine/Source/Editor/Sequencer/Private/SequencerObjectChangeListener.cpp
--- a/Engine/Source/Editor/Sequencer/Private/SequencerObjectChangeListener.cpp
+++ b/Engine/Source/Editor/Sequencer/Private/SequencerObjectChangeListener.cpp
@@ -219,6 +219,13 @@ const FOnAnimatablePropertyChanged* FSequencerObjectChangeListener::FindProperty
 			return nullptr;
 		}
 
+		// If not editable, we cannot use it
+		const bool bEditable = Property.HasAnyPropertyFlags(CPF_Edit) && !Property.HasAnyPropertyFlags(CPF_DisableEditOnInstance) && !Property.HasAnyPropertyFlags(CPF_BlueprintReadOnly);
+		if (!bEditable)
+		{
+			return nullptr;
+		}
+
 		// Interp properties are always keyable
 		if (PropertyOrContainer->HasAnyPropertyFlags(CPF_Interp))
 		{
@@ -257,12 +264,10 @@ const FOnAnimatablePropertyChanged* FSequencerObjectChangeListener::FindProperty
 				FName FunctionName = FName(*(Set + PropertyVarName));
 				UFunction* Function = Class->FindFunctionByName(FunctionName);
 
-				// @TODO: should we early out of our property path iteration if we find an "edit defaults only" property?
-				const bool bEditable           = Property.HasAnyPropertyFlags(CPF_Edit) && !Property.HasAnyPropertyFlags(CPF_DisableEditOnInstance);
 				const bool bFoundValidFunction = Function && !Function->HasMetaData(DeprecatedFunctionName);
-			
+
 				// Valid if there's a setter function and the property is editable.
-				if (bFoundValidFunction && bEditable)
+				if (bFoundValidFunction)
 				{
 					return DelegatePtr;
 				}
