Avoid capturing a widget by rawptr in an attribute

[RN] Fixed crash after tearing down a level sequencer object

diff --git a/Engine/Source/Editor/Sequencer/Private/SSequencer.cpp b/Engine/Source/Editor/Sequencer/Private/SSequencer.cpp
--- a/Engine/Source/Editor/Sequencer/Private/SSequencer.cpp
+++ b/Engine/Source/Editor/Sequencer/Private/SSequencer.cpp
@@ -185,13 +185,17 @@ void SSequencer::Construct(const FArguments& InArgs, TSharedRef<FSequencer> InSe
 	RootCustomization.OnActorsDrop = InArgs._OnActorsDrop;
 	RootCustomization.OnFoldersDrop = InArgs._OnFoldersDrop;
 
+	TWeakPtr<SSequencer> WeakSelf = StaticCastSharedRef<SSequencer>(AsShared());
 	// Get the desired display format from the user's settings each time.
 	TAttribute<EFrameNumberDisplayFormats> GetDisplayFormatAttr = MakeAttributeLambda(
-		[this]
+		[WeakSelf]
 		{
-			if (USequencerSettings* Settings = GetSequencerSettings())
+			if(TSharedPtr<SSequencer> Target = WeakSelf.Pin())
 			{
-				return Settings->GetTimeDisplayFormat();
+				if (USequencerSettings* Settings = Target->GetSequencerSettings())
+				{
+					return Settings->GetTimeDisplayFormat();
+				}
 			}
 			return EFrameNumberDisplayFormats::Frames;
 		}
@@ -199,11 +203,14 @@ void SSequencer::Construct(const FArguments& InArgs, TSharedRef<FSequencer> InSe
 
 	// Get the number of zero pad frames from the user's settings as well.
 	TAttribute<uint8> GetZeroPadFramesAttr = MakeAttributeLambda(
-		[this]()->uint8
+		[WeakSelf]()->uint8
 		{
-			if (USequencerSettings* Settings = GetSequencerSettings())
+			if(TSharedPtr<SSequencer> Target = WeakSelf.Pin())
 			{
-				return Settings->GetZeroPadFrames();
+				if (USequencerSettings* Settings = Target->GetSequencerSettings())
+				{
+					return Settings->GetZeroPadFrames();
+				}
 			}
 			return 0;
 		}
