Sequencer: Browse to selected section's objects

Added UMovieSceneSection::GetSourceObject()

Implements the source object for audio, skeletal animation, subsequences/shots, media, geometry cache, geometry collection, camera shake

diff --git a/Engine/Plugins/Experimental/GeometryCollectionPlugin/Source/GeometryCollectionTracks/Private/MovieSceneGeometryCollectionSection.cpp b/Engine/Plugins/Experimental/GeometryCollectionPlugin/Source/GeometryCollectionTracks/Private/MovieSceneGeometryCollectionSection.cpp
--- a/Engine/Plugins/Experimental/GeometryCollectionPlugin/Source/GeometryCollectionTracks/Private/MovieSceneGeometryCollectionSection.cpp
+++ b/Engine/Plugins/Experimental/GeometryCollectionPlugin/Source/GeometryCollectionTracks/Private/MovieSceneGeometryCollectionSection.cpp
@@ -55,6 +55,11 @@ void UMovieSceneGeometryCollectionSection::MigrateFrameTimes(FFrameRate SourceRa
 	}
 }
 
+UObject* UMovieSceneGeometryCollectionSection::GetSourceObject() const
+{
+	return Params.GeometryCollectionCache.ResolveObject();
+}
+
 FFrameNumber GetStartOffsetAtTrimTime(FQualifiedFrameTime TrimTime, const FMovieSceneGeometryCollectionParams& Params, FFrameNumber StartFrame, FFrameRate FrameRate)
 {
 	float AnimPlayRate = FMath::IsNearlyZero(Params.PlayRate) ? 1.0f : Params.PlayRate;
diff --git a/Engine/Plugins/Experimental/GeometryCollectionPlugin/Source/GeometryCollectionTracks/Public/MovieSceneGeometryCollectionSection.h b/Engine/Plugins/Experimental/GeometryCollectionPlugin/Source/GeometryCollectionTracks/Public/MovieSceneGeometryCollectionSection.h
--- a/Engine/Plugins/Experimental/GeometryCollectionPlugin/Source/GeometryCollectionTracks/Public/MovieSceneGeometryCollectionSection.h
+++ b/Engine/Plugins/Experimental/GeometryCollectionPlugin/Source/GeometryCollectionTracks/Public/MovieSceneGeometryCollectionSection.h
@@ -89,6 +89,7 @@ protected:
 	virtual void GetSnapTimes(TArray<FFrameNumber>& OutSnapTimes, bool bGetSectionBorders) const override;
 	virtual TOptional<FFrameTime> GetOffsetTime() const override;
 	virtual void MigrateFrameTimes(FFrameRate SourceRate, FFrameRate DestinationRate) override;
+	virtual UObject* GetSourceObject() const override;
 
 private:
 
diff --git a/Engine/Plugins/Media/MediaCompositing/Source/MediaCompositing/Private/MovieScene/MovieSceneMediaSection.cpp b/Engine/Plugins/Media/MediaCompositing/Source/MediaCompositing/Private/MovieScene/MovieSceneMediaSection.cpp
--- a/Engine/Plugins/Media/MediaCompositing/Source/MediaCompositing/Private/MovieScene/MovieSceneMediaSection.cpp
+++ b/Engine/Plugins/Media/MediaCompositing/Source/MediaCompositing/Private/MovieScene/MovieSceneMediaSection.cpp
@@ -122,6 +122,11 @@ void UMovieSceneMediaSection::GetReferencedBindings(TArray<FGuid>& OutBindings)
 	OutBindings.Add(MediaSourceProxyBindingID.GetGuid());
 }
 
+UObject* UMovieSceneMediaSection::GetSourceObject() const
+{
+	return GetMediaSource();
+}
+
 UMediaSource* UMovieSceneMediaSection::GetMediaSource() const
 {
 	return MediaSource.Get();
diff --git a/Engine/Plugins/Media/MediaCompositing/Source/MediaCompositing/Public/MovieSceneMediaSection.h b/Engine/Plugins/Media/MediaCompositing/Source/MediaCompositing/Public/MovieSceneMediaSection.h
--- a/Engine/Plugins/Media/MediaCompositing/Source/MediaCompositing/Public/MovieSceneMediaSection.h
+++ b/Engine/Plugins/Media/MediaCompositing/Source/MediaCompositing/Public/MovieSceneMediaSection.h
@@ -99,7 +99,7 @@ protected:
 	virtual void MigrateFrameTimes(FFrameRate SourceRate, FFrameRate DestinationRate) override;
 	virtual void OnBindingIDsUpdated(const TMap<UE::MovieScene::FFixedObjectBindingID, UE::MovieScene::FFixedObjectBindingID>& OldFixedToNewFixedMap, FMovieSceneSequenceID LocalSequenceID, TSharedRef<UE::MovieScene::FSharedPlaybackState> SharedPlaybackState) override;
 	virtual void GetReferencedBindings(TArray<FGuid>& OutBindings) override;
-
+	virtual UObject* GetSourceObject() const override;
 
 public:
 
diff --git a/Engine/Plugins/Runtime/GeometryCache/Source/GeometryCacheTracks/Classes/MovieSceneGeometryCacheSection.h b/Engine/Plugins/Runtime/GeometryCache/Source/GeometryCacheTracks/Classes/MovieSceneGeometryCacheSection.h
--- a/Engine/Plugins/Runtime/GeometryCache/Source/GeometryCacheTracks/Classes/MovieSceneGeometryCacheSection.h
+++ b/Engine/Plugins/Runtime/GeometryCache/Source/GeometryCacheTracks/Classes/MovieSceneGeometryCacheSection.h
@@ -76,6 +76,7 @@ protected:
 	virtual void GetSnapTimes(TArray<FFrameNumber>& OutSnapTimes, bool bGetSectionBorders) const override;
 	virtual TOptional<FFrameTime> GetOffsetTime() const override;
 	virtual void MigrateFrameTimes(FFrameRate SourceRate, FFrameRate DestinationRate) override;
+	virtual UObject* GetSourceObject() const override;
 
 	/** ~UObject interface */
 	virtual void PostLoad() override;
diff --git a/Engine/Plugins/Runtime/GeometryCache/Source/GeometryCacheTracks/Private/MovieSceneGeometryCacheSection.cpp b/Engine/Plugins/Runtime/GeometryCache/Source/GeometryCacheTracks/Private/MovieSceneGeometryCacheSection.cpp
--- a/Engine/Plugins/Runtime/GeometryCache/Source/GeometryCacheTracks/Private/MovieSceneGeometryCacheSection.cpp
+++ b/Engine/Plugins/Runtime/GeometryCache/Source/GeometryCacheTracks/Private/MovieSceneGeometryCacheSection.cpp
@@ -64,6 +64,11 @@ void UMovieSceneGeometryCacheSection::MigrateFrameTimes(FFrameRate SourceRate, F
 	}
 }
 
+UObject* UMovieSceneGeometryCacheSection::GetSourceObject() const
+{
+	return Params.GeometryCacheAsset;
+}
+
 void UMovieSceneGeometryCacheSection::PostLoad()
 {
 	Super::PostLoad();
diff --git a/Engine/Source/Editor/Sequencer/Private/Sequencer.cpp b/Engine/Source/Editor/Sequencer/Private/Sequencer.cpp
--- a/Engine/Source/Editor/Sequencer/Private/Sequencer.cpp
+++ b/Engine/Source/Editor/Sequencer/Private/Sequencer.cpp
@@ -5750,6 +5750,28 @@ void FSequencer::FindInContentBrowser()
 	}
 }
 
+void FSequencer::BrowseToObject()
+{
+	using namespace UE::Sequencer;
+
+	TArray<UObject*> ObjectsToFocus;
+
+	for (TViewModelPtr<FSectionModel> SectionModel : ViewModel->GetSelection()->TrackArea.Filter<FSectionModel>())
+	{
+		if (UMovieSceneSection* Section = SectionModel->GetSection())
+		{	
+			if (UObject* ObjectToFocus = Section->GetSourceObject())
+			{
+				ObjectsToFocus.Add(ObjectToFocus);
+			}
+		}
+	}
+
+	if (ObjectsToFocus.Num() > 0)
+	{
+		GEditor->SyncBrowserToObjects(ObjectsToFocus);
+	}
+}
 
 UObject* FSequencer::GetCurrentAsset() const
 {
@@ -10973,6 +10995,20 @@ void FSequencer::BindCommands()
 		return false;
 	};
 
+	auto CanBrowseToObject = [this] {
+		for (TViewModelPtr<FSectionModel> SectionModel : ViewModel->GetSelection()->TrackArea.Filter<FSectionModel>())
+		{
+			if (UMovieSceneSection* Section = SectionModel->GetSection())
+			{
+				if (UObject* ObjectToFocus = Section->GetSourceObject())
+				{
+					return true;
+				}
+			}
+		}
+		return false;
+	};
+
 	SequencerCommandBindings->MapAction(
 		FSequencerCommands::Get().ConvertToPossessable,
 		FExecuteAction::CreateSP(this, &FSequencer::ConvertSelectedNodesToPossessables),
@@ -11097,6 +11133,11 @@ void FSequencer::BindCommands()
 		Commands.FindInContentBrowser,
 		FExecuteAction::CreateSP( this, &FSequencer::FindInContentBrowser ) );
 
+	SequencerCommandBindings->MapAction(
+		Commands.BrowseToObject,
+		FExecuteAction::CreateSP( this, &FSequencer::BrowseToObject),
+		FCanExecuteAction::CreateLambda(CanBrowseToObject));
+
 	SequencerCommandBindings->MapAction(
 		Commands.ToggleLayerBars,
 		FExecuteAction::CreateLambda( [this]{
diff --git a/Engine/Source/Editor/Sequencer/Private/Sequencer.h b/Engine/Source/Editor/Sequencer/Private/Sequencer.h
--- a/Engine/Source/Editor/Sequencer/Private/Sequencer.h
+++ b/Engine/Source/Editor/Sequencer/Private/Sequencer.h
@@ -993,6 +993,9 @@ protected:
 	/** Find the viewed sequence asset in the content browser. */
 	void FindInContentBrowser();
 
+	/** Browse to the selected section's object. */
+	void BrowseToObject();
+
 	/** Get the asset we're currently editing, if applicable. */
 	UObject* GetCurrentAsset() const;
 
diff --git a/Engine/Source/Editor/Sequencer/Private/SequencerCommands.cpp b/Engine/Source/Editor/Sequencer/Private/SequencerCommands.cpp
--- a/Engine/Source/Editor/Sequencer/Private/SequencerCommands.cpp
+++ b/Engine/Source/Editor/Sequencer/Private/SequencerCommands.cpp
@@ -150,6 +150,7 @@ void FSequencerCommands::RegisterCommands()
 	UI_COMMAND( ToggleSnapCurveValueToInterval, "Snap Curve Key Values", "Snap curve keys to the value snapping interval", EUserInterfaceActionType::ToggleButton, FInputChord() );
 
 	UI_COMMAND( FindInContentBrowser, "Find in Content Browser", "Find the viewed sequence asset in the content browser", EUserInterfaceActionType::Button, FInputChord() );
+	UI_COMMAND( BrowseToObject, "Browse to Object", "Browse to the selected section's object", EUserInterfaceActionType::Button, FInputChord(EModifierKey::Control, EKeys::B) );
 	UI_COMMAND( ToggleLayerBars, "Layer Bars", "Show/hide the layer bars to edit keyframes in bulk", EUserInterfaceActionType::ToggleButton, FInputChord() );
 	UI_COMMAND( ToggleKeyBars, "Key Bars", "Show/hide key bar connectors for quickly retiming pairs of keys", EUserInterfaceActionType::ToggleButton, FInputChord() );
 	UI_COMMAND( ToggleChannelColors, "Channel Colors", "Show/hide the channel colors for the key bars", EUserInterfaceActionType::ToggleButton, FInputChord() );
diff --git a/Engine/Source/Editor/Sequencer/Public/SequencerCommands.h b/Engine/Source/Editor/Sequencer/Public/SequencerCommands.h
--- a/Engine/Source/Editor/Sequencer/Public/SequencerCommands.h
+++ b/Engine/Source/Editor/Sequencer/Public/SequencerCommands.h
@@ -379,6 +379,9 @@ public:
 	/** Finds the viewed sequence asset in the content browser. */
 	TSharedPtr< FUICommandInfo > FindInContentBrowser;
 
+	/** Browse to the selected objects. */
+	TSharedPtr< FUICommandInfo > BrowseToObject;
+
 	/** Toggles whether to show layer bars to edit keyframes in bulk. */
 	TSharedPtr< FUICommandInfo > ToggleLayerBars;
 
diff --git a/Engine/Source/Runtime/MovieScene/Private/Sections/MovieSceneSubSection.cpp b/Engine/Source/Runtime/MovieScene/Private/Sections/MovieSceneSubSection.cpp
--- a/Engine/Source/Runtime/MovieScene/Private/Sections/MovieSceneSubSection.cpp
+++ b/Engine/Source/Runtime/MovieScene/Private/Sections/MovieSceneSubSection.cpp
@@ -467,6 +467,11 @@ FMovieSceneTimeWarpVariant* UMovieSceneSubSection::GetTimeWarp()
 	return &Parameters.TimeScale;
 }
 
+UObject* UMovieSceneSubSection::GetSourceObject() const
+{
+	return GetSequence();
+}
+
 #if WITH_EDITOR
 void UMovieSceneSubSection::PreEditChange(FProperty* PropertyAboutToChange)
 {
diff --git a/Engine/Source/Runtime/MovieScene/Public/MovieSceneSection.h b/Engine/Source/Runtime/MovieScene/Public/MovieSceneSection.h
--- a/Engine/Source/Runtime/MovieScene/Public/MovieSceneSection.h
+++ b/Engine/Source/Runtime/MovieScene/Public/MovieSceneSection.h
@@ -675,6 +675,8 @@ public:
 	**/
 	MOVIESCENE_API virtual UObject* GetImplicitObjectOwner();
 
+	/** Get the object associated with this section */
+	MOVIESCENE_API virtual UObject* GetSourceObject() const { return nullptr; }
 
 #if WITH_EDITOR
 	MOVIESCENE_API virtual void PostEditChangeProperty(FPropertyChangedEvent& PropertyChangedEvent) override;
diff --git a/Engine/Source/Runtime/MovieScene/Public/Sections/MovieSceneSubSection.h b/Engine/Source/Runtime/MovieScene/Public/Sections/MovieSceneSubSection.h
--- a/Engine/Source/Runtime/MovieScene/Public/Sections/MovieSceneSubSection.h
+++ b/Engine/Source/Runtime/MovieScene/Public/Sections/MovieSceneSubSection.h
@@ -229,6 +229,7 @@ public:
 	MOVIESCENE_API virtual void GetSnapTimes(TArray<FFrameNumber>& OutSnapTimes, bool bGetSectionBorders) const override;
 	MOVIESCENE_API virtual void MigrateFrameTimes(FFrameRate SourceRate, FFrameRate DestinationRate) override;
 	MOVIESCENE_API virtual FMovieSceneTimeWarpVariant* GetTimeWarp() override;
+	MOVIESCENE_API virtual UObject* GetSourceObject() const override;
 
 protected:
 
diff --git a/Engine/Source/Runtime/MovieSceneTracks/Private/Sections/MovieSceneAudioSection.cpp b/Engine/Source/Runtime/MovieSceneTracks/Private/Sections/MovieSceneAudioSection.cpp
--- a/Engine/Source/Runtime/MovieSceneTracks/Private/Sections/MovieSceneAudioSection.cpp
+++ b/Engine/Source/Runtime/MovieSceneTracks/Private/Sections/MovieSceneAudioSection.cpp
@@ -149,6 +149,11 @@ EMovieSceneChannelProxyType  UMovieSceneAudioSection::CacheChannelProxy()
 	return EMovieSceneChannelProxyType::Dynamic;
 }
 
+UObject* UMovieSceneAudioSection::GetSourceObject() const
+{
+	return GetSound();
+}
+
 void UMovieSceneAudioSection::SetupSoundInputParameters(USoundBase* InSoundBase)
 {
 	// Populate with defaults.
diff --git a/Engine/Source/Runtime/MovieSceneTracks/Private/Sections/MovieSceneCameraShakeSection.cpp b/Engine/Source/Runtime/MovieSceneTracks/Private/Sections/MovieSceneCameraShakeSection.cpp
--- a/Engine/Source/Runtime/MovieSceneTracks/Private/Sections/MovieSceneCameraShakeSection.cpp
+++ b/Engine/Source/Runtime/MovieSceneTracks/Private/Sections/MovieSceneCameraShakeSection.cpp
@@ -46,6 +46,11 @@ void UMovieSceneCameraShakeSection::PostLoad()
 	Super::PostLoad();
 }
 
+UObject* UMovieSceneCameraShakeSection::GetSourceObject() const
+{
+	return ShakeData.ShakeClass;
+}
+
 void UMovieSceneCameraShakeSection::ImportEntityImpl(UMovieSceneEntitySystemLinker* EntityLinker, const FEntityImportParams& Params, FImportedEntity* OutImportedEntity)
 {
 	using namespace UE::MovieScene;
@@ -67,4 +72,3 @@ void UMovieSceneCameraShakeSection::ImportEntityImpl(UMovieSceneEntitySystemLink
 		.Add(TrackComponents->CameraShake, ComponentData)
 	);
 }
-
diff --git a/Engine/Source/Runtime/MovieSceneTracks/Private/Sections/MovieSceneCameraShakeSourceShakeSection.cpp b/Engine/Source/Runtime/MovieSceneTracks/Private/Sections/MovieSceneCameraShakeSourceShakeSection.cpp
--- a/Engine/Source/Runtime/MovieSceneTracks/Private/Sections/MovieSceneCameraShakeSourceShakeSection.cpp
+++ b/Engine/Source/Runtime/MovieSceneTracks/Private/Sections/MovieSceneCameraShakeSourceShakeSection.cpp
@@ -11,6 +11,11 @@ UMovieSceneCameraShakeSourceShakeSection::UMovieSceneCameraShakeSourceShakeSecti
 	EvalOptions.EnableAndSetCompletionMode(EMovieSceneCompletionMode::ProjectDefault);
 }
 
+UObject* UMovieSceneCameraShakeSourceShakeSection::GetSourceObject() const
+{
+	return ShakeData.ShakeClass;
+}
+
 void UMovieSceneCameraShakeSourceShakeSection::ImportEntityImpl(UMovieSceneEntitySystemLinker* EntityLinker, const FEntityImportParams& Params, FImportedEntity* OutImportedEntity)
 {
 	using namespace UE::MovieScene;
diff --git a/Engine/Source/Runtime/MovieSceneTracks/Private/Sections/MovieSceneSkeletalAnimationSection.cpp b/Engine/Source/Runtime/MovieSceneTracks/Private/Sections/MovieSceneSkeletalAnimationSection.cpp
--- a/Engine/Source/Runtime/MovieSceneTracks/Private/Sections/MovieSceneSkeletalAnimationSection.cpp
+++ b/Engine/Source/Runtime/MovieSceneTracks/Private/Sections/MovieSceneSkeletalAnimationSection.cpp
@@ -486,6 +486,11 @@ FMovieSceneTimeWarpVariant* UMovieSceneSkeletalAnimationSection::GetTimeWarp()
 	return &Params.PlayRate;
 }
 
+UObject* UMovieSceneSkeletalAnimationSection::GetSourceObject() const
+{
+	return Params.Animation;
+}
+
 void UMovieSceneSkeletalAnimationSection::ImportEntityImpl(UMovieSceneEntitySystemLinker* EntityLinker, const FEntityImportParams& InParams, FImportedEntity* OutImportedEntity)
 {
 	using namespace UE::MovieScene;
diff --git a/Engine/Source/Runtime/MovieSceneTracks/Public/Sections/MovieSceneAudioSection.h b/Engine/Source/Runtime/MovieSceneTracks/Public/Sections/MovieSceneAudioSection.h
--- a/Engine/Source/Runtime/MovieSceneTracks/Public/Sections/MovieSceneAudioSection.h
+++ b/Engine/Source/Runtime/MovieSceneTracks/Public/Sections/MovieSceneAudioSection.h
@@ -190,6 +190,7 @@ public:
 	MOVIESCENETRACKS_API virtual TOptional<FFrameTime> GetOffsetTime() const override;
 	MOVIESCENETRACKS_API virtual void MigrateFrameTimes(FFrameRate SourceRate, FFrameRate DestinationRate) override;
 	MOVIESCENETRACKS_API virtual EMovieSceneChannelProxyType CacheChannelProxy() override;
+	MOVIESCENETRACKS_API virtual UObject* GetSourceObject() const override;
 
 	//~ IMovieSceneEntityProvider interface
 	MOVIESCENETRACKS_API virtual bool PopulateEvaluationFieldImpl(const TRange<FFrameNumber>& EffectiveRange, const FMovieSceneEvaluationFieldEntityMetaData& InMetaData, FMovieSceneEntityComponentFieldBuilder* OutFieldBuilder) override;
diff --git a/Engine/Source/Runtime/MovieSceneTracks/Public/Sections/MovieSceneCameraShakeSection.h b/Engine/Source/Runtime/MovieSceneTracks/Public/Sections/MovieSceneCameraShakeSection.h
--- a/Engine/Source/Runtime/MovieSceneTracks/Public/Sections/MovieSceneCameraShakeSection.h
+++ b/Engine/Source/Runtime/MovieSceneTracks/Public/Sections/MovieSceneCameraShakeSection.h
@@ -58,6 +58,7 @@ public:
 	UMovieSceneCameraShakeSection(const FObjectInitializer& ObjectInitializer);
 
 	virtual void PostLoad() override;
+	virtual UObject* GetSourceObject() const override;
 
 	/** IMovieSceneEntityProvider interface */
 	void ImportEntityImpl(UMovieSceneEntitySystemLinker* EntityLinker, const FEntityImportParams& Params, FImportedEntity* OutImportedEntity) override;
diff --git a/Engine/Source/Runtime/MovieSceneTracks/Public/Sections/MovieSceneCameraShakeSourceShakeSection.h b/Engine/Source/Runtime/MovieSceneTracks/Public/Sections/MovieSceneCameraShakeSourceShakeSection.h
--- a/Engine/Source/Runtime/MovieSceneTracks/Public/Sections/MovieSceneCameraShakeSourceShakeSection.h
+++ b/Engine/Source/Runtime/MovieSceneTracks/Public/Sections/MovieSceneCameraShakeSourceShakeSection.h
@@ -23,6 +23,8 @@ class UMovieSceneCameraShakeSourceShakeSection
 public:
 	UMovieSceneCameraShakeSourceShakeSection(const FObjectInitializer& ObjectInitializer);
 
+	virtual UObject* GetSourceObject() const override;
+
 	/** IMovieSceneEntityProvider interface */
 	void ImportEntityImpl(UMovieSceneEntitySystemLinker* EntityLinker, const FEntityImportParams& Params, FImportedEntity* OutImportedEntity) override;
 	
diff --git a/Engine/Source/Runtime/MovieSceneTracks/Public/Sections/MovieSceneSkeletalAnimationSection.h b/Engine/Source/Runtime/MovieSceneTracks/Public/Sections/MovieSceneSkeletalAnimationSection.h
--- a/Engine/Source/Runtime/MovieSceneTracks/Public/Sections/MovieSceneSkeletalAnimationSection.h
+++ b/Engine/Source/Runtime/MovieSceneTracks/Public/Sections/MovieSceneSkeletalAnimationSection.h
@@ -124,6 +124,7 @@ public:
 	virtual void SetStartFrame(TRangeBound<FFrameNumber> NewStartFrame) override;
 	virtual void SetEndFrame(TRangeBound<FFrameNumber> NewEndFrame)override;
 	virtual FMovieSceneTimeWarpVariant* GetTimeWarp() override;
+	virtual UObject* GetSourceObject() const override;
 
 	MOVIESCENETRACKS_API void DeleteChannels(TArrayView<const FName> ChannelNames);
 
