Sequencer: Fixed potential for nullptr crash when updating selection state on outliner nodes

diff --git a/Engine/Source/Editor/Sequencer/Private/MVVM/Selection.cpp b/Engine/Source/Editor/Sequencer/Private/MVVM/Selection.cpp
--- a/Engine/Source/Editor/Sequencer/Private/MVVM/Selection.cpp
+++ b/Engine/Source/Editor/Sequencer/Private/MVVM/Selection.cpp
@@ -103,14 +103,15 @@ void FSequencerSelection::PreBroadcastChangeEvent()
 		{
 			ParentOutlinerNode->ToggleSelectionState(EOutlinerSelectionState::HasSelectedTrackAreaItems, true);
 			NodesWithKeysOrSections.Add(ParentOutlinerNode);
-		}
 
-		ParentOutlinerNode = ParentOutlinerNode.AsModel()->FindAncestorOfType<IOutlinerExtension>();
-		while (ParentOutlinerNode)
-		{
-			ParentOutlinerNode->ToggleSelectionState(EOutlinerSelectionState::DescendentHasSelectedTrackAreaItems, true);
 			ParentOutlinerNode = ParentOutlinerNode.AsModel()->FindAncestorOfType<IOutlinerExtension>();
+			while (ParentOutlinerNode)
+			{
+				ParentOutlinerNode->ToggleSelectionState(EOutlinerSelectionState::DescendentHasSelectedTrackAreaItems, true);
+				ParentOutlinerNode = ParentOutlinerNode.AsModel()->FindAncestorOfType<IOutlinerExtension>();
+			}
 		}
+
 	}
 
 	// Gather selection states from selected keys
@@ -122,13 +123,13 @@ void FSequencerSelection::PreBroadcastChangeEvent()
 		{
 			ParentOutlinerNode->ToggleSelectionState(EOutlinerSelectionState::HasSelectedKeys, true);
 			NodesWithKeysOrSections.Add(ParentOutlinerNode);
-		}
 
-		ParentOutlinerNode = ParentOutlinerNode.AsModel()->FindAncestorOfType<IOutlinerExtension>();
-		while (ParentOutlinerNode)
-		{
-			ParentOutlinerNode->ToggleSelectionState(EOutlinerSelectionState::DescendentHasSelectedKeys, true);
 			ParentOutlinerNode = ParentOutlinerNode.AsModel()->FindAncestorOfType<IOutlinerExtension>();
+			while (ParentOutlinerNode)
+			{
+				ParentOutlinerNode->ToggleSelectionState(EOutlinerSelectionState::DescendentHasSelectedKeys, true);
+				ParentOutlinerNode = ParentOutlinerNode.AsModel()->FindAncestorOfType<IOutlinerExtension>();
+			}
 		}
 	}
 
