Sequencer: Reset the time controller in tick resolution and not display rate.

NewGlobalTime is in tick resolution and NewPlayPosition is in display rate. So, it needs to use NewGlobalTime.

Tested that looping a subsequence, correctly loops back to the start of the subsequence.

#rnx

diff --git a/Engine/Source/Editor/Sequencer/Private/Sequencer.cpp b/Engine/Source/Editor/Sequencer/Private/Sequencer.cpp
--- a/Engine/Source/Editor/Sequencer/Private/Sequencer.cpp
+++ b/Engine/Source/Editor/Sequencer/Private/Sequencer.cpp
@@ -5219,7 +5219,7 @@ void FSequencer::SetLocalTimeLooped(FFrameTime NewLocalTime, const FMovieSceneTr
 	if (bResetPosition)
 	{
 		PlayPosition.Reset(NewPlayPosition);
-		TimeController->Reset(FQualifiedFrameTime(NewPlayPosition, RootTickResolution));
+		TimeController->Reset(FQualifiedFrameTime(NewGlobalTime.GetValue(), RootTickResolution));
 	}
 
 	// Ensure breadcrumbs are up to date
