Sequencer: Time Warp filter needs to look at sections as well

#rnx

diff --git a/Engine/Source/Editor/Sequencer/Private/Filters/Filters/SequencerTrackFilter_TimeWarp.cpp b/Engine/Source/Editor/Sequencer/Private/Filters/Filters/SequencerTrackFilter_TimeWarp.cpp
new file mode 100644
--- /dev/null
+++ b/Engine/Source/Editor/Sequencer/Private/Filters/Filters/SequencerTrackFilter_TimeWarp.cpp
@@ -0,0 +1,78 @@
+// Copyright Epic Games, Inc. All Rights Reserved.
+
+#include "SequencerTrackFilter_TimeWarp.h"
+#include "Filters/SequencerFilterBar.h"
+#include "Filters/SequencerTrackFilterCommands.h"
+#include "MovieSceneSection.h"
+#include "MVVM/ViewModels/ObjectBindingModel.h"
+#include "MVVM/ViewModels/SectionModel.h"
+#include "MVVM/ViewModels/TrackModel.h"
+#include "Sequencer.h"
+#include "Tracks/MovieSceneTimeWarpTrack.h"
+#include "Variants/MovieSceneTimeWarpVariant.h"
+
+using namespace UE::Sequencer;
+
+#define LOCTEXT_NAMESPACE "SequencerTrackFilter_TimeWarp"
+
+FSequencerTrackFilter_TimeWarp::FSequencerTrackFilter_TimeWarp(ISequencerTrackFilters& InFilterInterface, TSharedPtr<FFilterCategory> InCategory)
+	: FSequencerTrackFilter(InFilterInterface, MoveTemp(InCategory))
+{
+}
+
+bool FSequencerTrackFilter_TimeWarp::ShouldUpdateOnTrackValueChanged() const
+{
+	return true;
+}
+
+FText FSequencerTrackFilter_TimeWarp::GetDefaultToolTipText() const
+{
+	return LOCTEXT("SequencerTrackFilter_TimeWarpToolTip", "Show only Time Warp tracks");
+}
+
+TSharedPtr<FUICommandInfo> FSequencerTrackFilter_TimeWarp::GetToggleCommand() const
+{
+	return FSequencerTrackFilterCommands::Get().ToggleFilter_TimeWarp;
+}
+
+FText FSequencerTrackFilter_TimeWarp::GetDisplayName() const
+{
+	return LOCTEXT("FSequencerTrackFilter_TimeWarp", "TimeWarp");
+}
+
+FSlateIcon FSequencerTrackFilter_TimeWarp::GetIcon() const
+{
+	return FSlateIcon(FAppStyle::GetAppStyleSetName(), TEXT("Sequencer.Tracks.TimeWarp"));
+}
+
+FString FSequencerTrackFilter_TimeWarp::GetName() const
+{
+	return TEXT("TimeWarp");
+}
+
+bool FSequencerTrackFilter_TimeWarp::PassesFilter(FSequencerTrackFilterType InItem) const
+{
+	if (const TViewModelPtr<ITrackExtension> TrackModel = InItem->FindAncestorOfType<ITrackExtension>(true))
+	{
+		for (const TViewModelPtr<FSectionModel>& SectionModel : TrackModel->GetSectionModels().IterateSubList<FSectionModel>())
+		{
+			UMovieSceneSection* Section = SectionModel->GetSection();
+			if (!Section)
+			{
+				continue;
+			}
+
+			if (FMovieSceneTimeWarpVariant* Variant = Section->GetTimeWarp())
+			{
+				if (Variant->GetType() == EMovieSceneTimeWarpType::Custom)
+				{
+					return true;
+				}
+			}
+		}
+	}
+
+	return false;
+}
+
+#undef LOCTEXT_NAMESPACE
diff --git a/Engine/Source/Editor/Sequencer/Private/Filters/Filters/SequencerTrackFilter_TimeWarp.h b/Engine/Source/Editor/Sequencer/Private/Filters/Filters/SequencerTrackFilter_TimeWarp.h
new file mode 100644
--- /dev/null
+++ b/Engine/Source/Editor/Sequencer/Private/Filters/Filters/SequencerTrackFilter_TimeWarp.h
@@ -0,0 +1,27 @@
+// Copyright Epic Games, Inc. All Rights Reserved.
+
+#pragma once
+
+#include "Filters/SequencerTrackFilterBase.h"
+
+class FSequencerTrackFilter_TimeWarp : public FSequencerTrackFilter
+{
+public:
+	FSequencerTrackFilter_TimeWarp(ISequencerTrackFilters& InFilterInterface, TSharedPtr<FFilterCategory> InCategory = nullptr);
+
+	//~ Begin FSequencerTrackFilter
+	virtual bool ShouldUpdateOnTrackValueChanged() const override;
+	virtual FText GetDefaultToolTipText() const override;
+	virtual TSharedPtr<FUICommandInfo> GetToggleCommand() const override;
+	//~ End FSequencerTrackFilter
+
+	//~ Begin FFilterBase
+	virtual FText GetDisplayName() const override;
+	virtual FSlateIcon GetIcon() const override;
+	//~ End FFilterBase
+
+	//~ Begin IFilter
+	virtual FString GetName() const override;
+	virtual bool PassesFilter(FSequencerTrackFilterType InItem) const override;
+	//~ End IFilter
+};
diff --git a/Engine/Source/Editor/Sequencer/Private/Filters/Filters/SequencerTrackFilters.h b/Engine/Source/Editor/Sequencer/Private/Filters/Filters/SequencerTrackFilters.h
--- a/Engine/Source/Editor/Sequencer/Private/Filters/Filters/SequencerTrackFilters.h
+++ b/Engine/Source/Editor/Sequencer/Private/Filters/Filters/SequencerTrackFilters.h
@@ -23,7 +23,6 @@
 #include "Tracks/MovieSceneParticleTrack.h"
 #include "Tracks/MovieSceneSlomoTrack.h"
 #include "Tracks/MovieSceneSubTrack.h"
-#include "Tracks/MovieSceneTimeWarpTrack.h"
 
 #define LOCTEXT_NAMESPACE "SequencerTrackFilters"
 
@@ -492,40 +491,6 @@ public:
 //////////////////////////////////////////////////////////////////////////
 //
 
-class FSequencerTrackFilter_TimeWarp : public FSequencerTrackFilter_ClassType<UMovieSceneTimeWarpTrack>
-{
-public:
-	FSequencerTrackFilter_TimeWarp(ISequencerTrackFilters& InFilterInterface, TSharedPtr<FFilterCategory> InCategory = nullptr)
-		: FSequencerTrackFilter_ClassType<UMovieSceneTimeWarpTrack>(InFilterInterface, InCategory)
-	{}
-
-	//~ Begin IFilter
-	virtual FString GetName() const override { return TEXT("TimeWarp"); }
-	//~ End IFilter
-
-	//~ Begin FFilterBase
-	virtual FText GetDisplayName() const override { return LOCTEXT("SequencerTrackFilter_TimeWarp", "Time Warp"); }
-	virtual FSlateIcon GetIcon() const override { return FSlateIcon(FAppStyle::GetAppStyleSetName(), TEXT("Sequencer.Tracks.TimeWarp")); }
-	//~ End FFilterBase
-
-	//~ Begin FSequencerTrackFilter
-
-	virtual FText GetDefaultToolTipText() const override
-	{
-		return LOCTEXT("SequencerTrackFilter_TimeWarpToolTip", "Show only Time Warp tracks");
-	}
-
-	virtual TSharedPtr<FUICommandInfo> GetToggleCommand() const override
-	{
-		return FSequencerTrackFilterCommands::Get().ToggleFilter_TimeWarp;
-	}
-
-	//~ End FSequencerTrackFilter
-};
-
-//////////////////////////////////////////////////////////////////////////
-//
-
 class FSequencerTrackFilter_Folder : public FSequencerTrackFilter_ModelType<UE::Sequencer::FFolderModel>
 {
 public:
diff --git a/Engine/Source/Editor/Sequencer/Private/Filters/SequencerFilterBar.cpp b/Engine/Source/Editor/Sequencer/Private/Filters/SequencerFilterBar.cpp
--- a/Engine/Source/Editor/Sequencer/Private/Filters/SequencerFilterBar.cpp
+++ b/Engine/Source/Editor/Sequencer/Private/Filters/SequencerFilterBar.cpp
@@ -11,6 +11,7 @@
 #include "Filters/Filters/SequencerTrackFilter_Modified.h"
 #include "Filters/Filters/SequencerTrackFilter_Selected.h"
 #include "Filters/Filters/SequencerTrackFilter_Text.h"
+#include "Filters/Filters/SequencerTrackFilter_TimeWarp.h"
 #include "Filters/Filters/SequencerTrackFilter_Unbound.h"
 #include "Filters/Filters/SequencerTrackFilters.h"
 #include "Filters/SequencerFilterBarConfig.h"
@@ -121,7 +122,6 @@ void FSequencerFilterBar::CreateDefaultFilters()
 	CommonFilters->Add(MakeShared<FSequencerTrackFilter_CinematicShot>(*this, ClassTypeCategory));
 	CommonFilters->Add(MakeShared<FSequencerTrackFilter_Subsequence>(*this, ClassTypeCategory));
 	CommonFilters->Add(MakeShared<FSequencerTrackFilter_TimeDilation>(*this, ClassTypeCategory));
-	CommonFilters->Add(MakeShared<FSequencerTrackFilter_TimeWarp>(*this, ClassTypeCategory));
 
 	// Add component type category filters
 	CommonFilters->Add(MakeShared<FSequencerTrackFilter_Camera>(*this, ComponentTypeCategory));
@@ -134,6 +134,7 @@ void FSequencerFilterBar::CreateDefaultFilters()
 	CommonFilters->Add(SelectedFilter);
 	CommonFilters->Add(MakeShared<FSequencerTrackFilter_Unbound>(*this, MiscCategory));
 	CommonFilters->Add(MakeShared<FSequencerTrackFilter_Condition>(*this, MiscCategory));
+	CommonFilters->Add(MakeShared<FSequencerTrackFilter_TimeWarp>(*this, MiscCategory));
 
 	// Add global user-defined filters
 	for (TObjectIterator<USequencerTrackFilterExtension> ExtensionIt(RF_NoFlags); ExtensionIt; ++ExtensionIt)
