Sequencer: Remove old filter using the old filter display name, not the newly changed name

#rnx

diff --git a/Engine/Source/Editor/Sequencer/Private/Filters/Widgets/SSequencerCustomTextFilterDialog.cpp b/Engine/Source/Editor/Sequencer/Private/Filters/Widgets/SSequencerCustomTextFilterDialog.cpp
--- a/Engine/Source/Editor/Sequencer/Private/Filters/Widgets/SSequencerCustomTextFilterDialog.cpp
+++ b/Engine/Source/Editor/Sequencer/Private/Filters/Widgets/SSequencerCustomTextFilterDialog.cpp
@@ -440,7 +440,7 @@ void SSequencerCustomTextFilterDialog::OnModifyCustomTextFilter()
 	const TSharedRef<FSequencerTrackFilter_CustomText> CustomTextFilterRef = CustomTextFilter.ToSharedRef();
 	const bool bWasFilterEnabled = FilterBar->IsFilterEnabled(CustomTextFilterRef);
 	const bool bWasFilterActive = FilterBar->IsFilterActive(CustomTextFilterRef);
-	const FString OldFilterName = CustomTextFilterData.FilterLabel.ToString();
+	const FString OldFilterName = CustomTextFilterRef->GetDisplayName().ToString();
 
 	FSequencerFilterBarConfig& Config = SequencerSettings->FindOrAddTrackFilterBar(FilterBar->GetIdentifier(), false);
 	Config.RemoveCustomTextFilter(OldFilterName);
