Sequencer: Fixed breadcrumbs breaking when calling focus_parent_sequence

There are three issues here; the main issue is that breadcrumbs broke when calling PopToSequenceInstance in some situations. The fix is to simply rebuild the breadcrumbs when necessary.
The second was that that PopToSequenceInstance takes a reference, and that was teetering dangerously close to resulting in an access violation if the array was mutated while that function was being called.
Finally, after auditing the code it appears that FSequencerBreadcrumb::BreadcrumbType is never actually assigned ShotType, so I have removed that member entirely for clarity.

diff --git a/Engine/Plugins/MovieScene/LevelSequenceEditor/Source/LevelSequenceEditor/Private/LevelSequenceEditorBlueprintLibrary.cpp b/Engine/Plugins/MovieScene/LevelSequenceEditor/Source/LevelSequenceEditor/Private/LevelSequenceEditorBlueprintLibrary.cpp
--- a/Engine/Plugins/MovieScene/LevelSequenceEditor/Source/LevelSequenceEditor/Private/LevelSequenceEditorBlueprintLibrary.cpp
+++ b/Engine/Plugins/MovieScene/LevelSequenceEditor/Source/LevelSequenceEditor/Private/LevelSequenceEditorBlueprintLibrary.cpp
@@ -78,7 +78,9 @@ void ULevelSequenceEditorBlueprintLibrary::FocusParentSequence()
 
 		if (Hierarchy.Num() > 1)
 		{
-			CurrentSequencer.Pin()->PopToSequenceInstance(Hierarchy[Hierarchy.Num() -2]);
+			// Copy the sequence ID since PopToSequenceInstance takes it by reference
+			FMovieSceneSequenceID SequenceID = Hierarchy[Hierarchy.Num() -2];
+			CurrentSequencer.Pin()->PopToSequenceInstance(SequenceID);
 		}
 	}
 }
diff --git a/Engine/Source/Editor/Sequencer/Private/SSequencer.cpp b/Engine/Source/Editor/Sequencer/Private/SSequencer.cpp
--- a/Engine/Source/Editor/Sequencer/Private/SSequencer.cpp
+++ b/Engine/Source/Editor/Sequencer/Private/SSequencer.cpp
@@ -2965,20 +2965,23 @@ void SSequencer::UpdateBreadcrumbs()
 		return;
 	}
 
-	FMovieSceneSequenceID FocusedID = Sequencer->GetFocusedTemplateID();
+	BreadcrumbTrail->ClearCrumbs();
 
-	if (BreadcrumbTrail->PeekCrumb().BreadcrumbType == FSequencerBreadcrumb::ShotType)
+	for (FMovieSceneSequenceID SequenceID : Sequencer->GetSubSequenceHierarchy())
 	{
-		BreadcrumbTrail->PopCrumb();
-	}
+		TAttribute<FText> CrumbNameAttribute;
 
-	if( BreadcrumbTrail->PeekCrumb().BreadcrumbType == FSequencerBreadcrumb::MovieSceneType && BreadcrumbTrail->PeekCrumb().SequenceID != FocusedID )
-	{
-		TWeakObjectPtr<UMovieSceneSubSection> SubSection = Sequencer->FindSubSection(FocusedID);
-		TAttribute<FText> CrumbNameAttribute = MakeAttributeSP(this, &SSequencer::GetBreadcrumbTextForSection, SubSection);
+		if (SequenceID == MovieSceneSequenceID::Root)
+		{
+			CrumbNameAttribute = MakeAttributeSP(this, &SSequencer::GetBreadcrumbTextForSequence, MakeWeakObjectPtr(SequencerPtr.Pin()->GetRootMovieSceneSequence()), true);
+		}
+		else
+		{
+			TWeakObjectPtr<UMovieSceneSubSection> SubSection = Sequencer->FindSubSection(SequenceID);
+			CrumbNameAttribute = MakeAttributeSP(this, &SSequencer::GetBreadcrumbTextForSection, SubSection);
+		}
 
-		// The current breadcrumb is not a moviescene so we need to make a new breadcrumb in order return to the parent moviescene later
-		BreadcrumbTrail->PushCrumb( CrumbNameAttribute, FSequencerBreadcrumb( FocusedID, CrumbNameAttribute.Get()) );
+		BreadcrumbTrail->PushCrumb( CrumbNameAttribute, FSequencerBreadcrumb( SequenceID, CrumbNameAttribute.Get()) );
 	}
 }
 
@@ -3465,16 +3468,9 @@ void SSequencer::OnFolderDropped(TSharedPtr<FFolderDragDropOp> DragDropOp )
 
 void SSequencer::OnCrumbClicked(const FSequencerBreadcrumb& Item)
 {
-	if (Item.BreadcrumbType != FSequencerBreadcrumb::ShotType)
+	if( SequencerPtr.Pin()->GetFocusedTemplateID() != Item.SequenceID )
 	{
-		if( SequencerPtr.Pin()->GetFocusedTemplateID() == Item.SequenceID ) 
-		{
-			// then do zooming
-		}
-		else
-		{
-			SequencerPtr.Pin()->PopToSequenceInstance( Item.SequenceID );
-		}
+		SequencerPtr.Pin()->PopToSequenceInstance( Item.SequenceID );
 	}
 }
 
diff --git a/Engine/Source/Editor/Sequencer/Private/SSequencer.h b/Engine/Source/Editor/Sequencer/Private/SSequencer.h
--- a/Engine/Source/Editor/Sequencer/Private/SSequencer.h
+++ b/Engine/Source/Editor/Sequencer/Private/SSequencer.h
@@ -97,15 +97,6 @@ namespace SequencerLayoutConstants
  */
 struct FSequencerBreadcrumb
 {
-	enum Type
-	{
-		ShotType,
-		MovieSceneType,
-	};
-
-	/** The type of breadcrumb this is */
-	FSequencerBreadcrumb::Type BreadcrumbType;
-
 	/** The movie scene this may point to */
 	FMovieSceneSequenceID SequenceID;
 
@@ -113,13 +104,7 @@ struct FSequencerBreadcrumb
 	FText BreadcrumbName;
 
 	FSequencerBreadcrumb(FMovieSceneSequenceIDRef InSequenceID, FText CrumbName)
-		: BreadcrumbType(FSequencerBreadcrumb::MovieSceneType)
-		, SequenceID(InSequenceID)
-		, BreadcrumbName(CrumbName)
-	{ }
-
-	FSequencerBreadcrumb(FText CrumbName)
-		: BreadcrumbType(FSequencerBreadcrumb::ShotType)
+		: SequenceID(InSequenceID)
 		, BreadcrumbName(CrumbName)
 	{ }
 };
