FDaySequenceEditorToolkit : Avoid force refresh of details if actor preview.
FLevelEditorSequencerIntegration : Preserve the original Options when OnPropertyEditorOpened() is invoked.

This removes unnecessary details panel forced refreshes ( one still happens later in FDaySequenceActorPreview::EnablePreview ).

diff --git a/Engine/Plugins/Experimental/DaySequence/Source/DaySequenceEditor/Private/DaySequenceEditorToolkit.cpp b/Engine/Plugins/Experimental/DaySequence/Source/DaySequenceEditor/Private/DaySequenceEditorToolkit.cpp
--- a/Engine/Plugins/Experimental/DaySequence/Source/DaySequenceEditor/Private/DaySequenceEditorToolkit.cpp
+++ b/Engine/Plugins/Experimental/DaySequence/Source/DaySequenceEditor/Private/DaySequenceEditorToolkit.cpp
@@ -312,6 +312,7 @@ void FDaySequenceEditorToolkit::InitializeInternal(
 	FLevelEditorSequencerIntegrationOptions Options;
 	Options.bRequiresLevelEvents = true;
 	Options.bRequiresActorEvents = true;
+	Options.bForceRefreshDetails = bIsSequenceEditor;
 	Options.bAttachOutlinerColumns = bIsSequenceEditor;
 	Options.bActivateSequencerEdMode = bIsSequenceEditor;
 	Options.bSyncBindingsToActorLabels = bIsSequenceEditor;
diff --git a/Engine/Source/Editor/Sequencer/Private/LevelEditorSequencerIntegration.cpp b/Engine/Source/Editor/Sequencer/Private/LevelEditorSequencerIntegration.cpp
--- a/Engine/Source/Editor/Sequencer/Private/LevelEditorSequencerIntegration.cpp
+++ b/Engine/Source/Editor/Sequencer/Private/LevelEditorSequencerIntegration.cpp
@@ -920,7 +920,7 @@ void FLevelEditorSequencerIntegration::ActivateDetailHandler(const FLevelEditorS
 		}
 	}
 
-	FDelegateHandle OnPropertyEditorOpenedHandle = EditModule.OnPropertyEditorOpened().AddRaw(this, &FLevelEditorSequencerIntegration::OnPropertyEditorOpened);
+	FDelegateHandle OnPropertyEditorOpenedHandle = EditModule.OnPropertyEditorOpened().AddLambda([this, Options]() { ActivateDetailHandler(Options); });
 
 	auto DeactivateDetailKeyframeHandler =
 		[this, OnPropertyEditorOpenedHandle]()
@@ -977,12 +977,6 @@ void FLevelEditorSequencerIntegration::ActivateDetailHandler(const FLevelEditorS
 	}
 }
 
-void FLevelEditorSequencerIntegration::OnPropertyEditorOpened()
-{
-	FLevelEditorSequencerIntegrationOptions Options;
-	ActivateDetailHandler(Options);
-}
-
 void FLevelEditorSequencerIntegration::BrowseToSelectedActor(AActor* Actor, FSequencer* Sequencer, FMovieSceneSequenceID SequenceID)
 {
 	Sequencer->PopToSequenceInstance(MovieSceneSequenceID::Root);
diff --git a/Engine/Source/Editor/Sequencer/Public/LevelEditorSequencerIntegration.h b/Engine/Source/Editor/Sequencer/Public/LevelEditorSequencerIntegration.h
--- a/Engine/Source/Editor/Sequencer/Public/LevelEditorSequencerIntegration.h
+++ b/Engine/Source/Editor/Sequencer/Public/LevelEditorSequencerIntegration.h
@@ -172,8 +172,6 @@ private:
 	/** Called to determine whether a binding is visible in the tree view */
 	bool IsBindingVisible(const FMovieSceneBinding& InBinding);
 
-	void OnPropertyEditorOpened();
-
 	void RegisterMenus();
 
 	void MakeBrowseToSelectedActorSubMenu(UToolMenu* Menu);
