Sequencer: Don't disable the filter bar widget when muting because then you can't right click on it to unmute filters

Added SetMuted() to the filter bar widget which will disable the search box, horizontal and vertical containers. This context menu is still enabled so allows for the right click menu to process.

Also, make sure the muted/enabled state is synced when the filter bar widget is rebuilt (ie. switching layout from horizontal to vertical or toggling visibility)

#rnx

diff --git a/Engine/Source/Editor/Sequencer/Private/Filters/SequencerFilterBar.cpp b/Engine/Source/Editor/Sequencer/Private/Filters/SequencerFilterBar.cpp
--- a/Engine/Source/Editor/Sequencer/Private/Filters/SequencerFilterBar.cpp
+++ b/Engine/Source/Editor/Sequencer/Private/Filters/SequencerFilterBar.cpp
@@ -293,13 +293,7 @@ void FSequencerFilterBar::MuteFilters(const bool bInMute)
 
 	if (const TSharedPtr<SSequencerFilterBar> FilterBarWidget = GetWidget())
 	{
-		FilterBarWidget->SetEnabled(!bFiltersMuted);
-
-		const TWeakPtr<SSequencerSearchBox> WeakSearchBox = FilterBarWidget->GetSearchBox();
-		if (WeakSearchBox.IsValid())
-		{
-			WeakSearchBox.Pin()->SetEnabled(!bFiltersMuted);
-		}
+		FilterBarWidget->SetMuted(bFiltersMuted);
 	}
 
 	RequestFilterUpdate();
diff --git a/Engine/Source/Editor/Sequencer/Private/Filters/Widgets/SSequencerFilterBar.cpp b/Engine/Source/Editor/Sequencer/Private/Filters/Widgets/SSequencerFilterBar.cpp
--- a/Engine/Source/Editor/Sequencer/Private/Filters/Widgets/SSequencerFilterBar.cpp
+++ b/Engine/Source/Editor/Sequencer/Private/Filters/Widgets/SSequencerFilterBar.cpp
@@ -385,6 +385,24 @@ TWeakPtr<SSequencerSearchBox> SSequencerFilterBar::GetSearchBox() const
 	return WeakSearchBox;
 }
 
+void SSequencerFilterBar::SetMuted(bool bInMuted)
+{
+	if (HorizontalContainerWidget.IsValid())
+	{
+		HorizontalContainerWidget->SetEnabled(!bInMuted);
+	}
+
+	if (VerticalContainerWidget.IsValid())
+	{
+		VerticalContainerWidget->SetEnabled(!bInMuted);
+	}
+
+	if (WeakSearchBox.IsValid())
+	{
+		WeakSearchBox.Pin()->SetEnabled(!bInMuted);
+	}
+}
+
 void SSequencerFilterBar::OnFiltersChanged(const ESequencerFilterChange InChangeType, const TSharedRef<FSequencerTrackFilter>& InFilter)
 {
 	switch (InChangeType)
diff --git a/Engine/Source/Editor/Sequencer/Private/Filters/Widgets/SSequencerFilterBar.h b/Engine/Source/Editor/Sequencer/Private/Filters/Widgets/SSequencerFilterBar.h
--- a/Engine/Source/Editor/Sequencer/Private/Filters/Widgets/SSequencerFilterBar.h
+++ b/Engine/Source/Editor/Sequencer/Private/Filters/Widgets/SSequencerFilterBar.h
@@ -86,6 +86,9 @@ public:
 
 	TWeakPtr<SSequencerSearchBox> GetSearchBox() const;
 
+	// Set the state of the filter bar. Muted means that the filters are muted, but the context menu is still enabled and accessible.
+	void SetMuted(bool bInMuted);
+
 protected:
 	void AddWidgetToLayout(const TSharedRef<SWidget>& InWidget);
 	void RemoveWidgetFromLayout(const TSharedRef<SWidget>& InWidget);
diff --git a/Engine/Source/Editor/Sequencer/Private/SSequencer.cpp b/Engine/Source/Editor/Sequencer/Private/SSequencer.cpp
--- a/Engine/Source/Editor/Sequencer/Private/SSequencer.cpp
+++ b/Engine/Source/Editor/Sequencer/Private/SSequencer.cpp
@@ -1335,6 +1335,8 @@ TSharedRef<SWidget> SSequencer::ConstructSearchAndFilterRow()
 		FilterComboButtonWidget = FilterBar->MakeAddFilterButton();
 	}
 
+	FilterBarWidget->SetMuted(Sequencer->GetFilterBar()->AreFiltersMuted());
+
 	return SNew(SHorizontalBox)
 
 		// Add Button
