Sequencer: Map filter toggle commands to curve editor

This fixes F10 for selected control rig controls not working in the curve editor.

#rnx

diff --git a/Engine/Source/Editor/Sequencer/Private/Filters/SequencerFilterBar.cpp b/Engine/Source/Editor/Sequencer/Private/Filters/SequencerFilterBar.cpp
--- a/Engine/Source/Editor/Sequencer/Private/Filters/SequencerFilterBar.cpp
+++ b/Engine/Source/Editor/Sequencer/Private/Filters/SequencerFilterBar.cpp
@@ -223,6 +223,7 @@ void FSequencerFilterBar::BindCommands()
 			{
 				const TSharedPtr<FUICommandList> CurveEditorSharedBindings = Sequencer.GetCommandBindings(ESequencerCommandBindings::CurveEditor);
 
+				// Add the general track filter commands
 				for (const TSharedPtr<FUICommandInfo>& Command : TrackFilterCommands.GetAllCommands())
 				{
 					if (Command.IsValid() && CommandList->IsActionMapped(Command))
@@ -231,7 +232,23 @@ void FSequencerFilterBar::BindCommands()
 					}
 				}
 
+				// Add the specific track filter toggle commands
+				for (const TSharedRef<FSequencerTrackFilter>& Filter : AllFilters)
+				{
+					if (Filter->SupportsSequence(FocusedSequence))
+					{
+						const TSharedPtr<FUICommandList>& FilterCommandList = Filter->GetFilterInterface().GetCommandList();
+						const TSharedPtr<FUICommandInfo>& FilterCommand = Filter->GetToggleCommand();
+
+						if (FilterCommand.IsValid() && FilterCommandList->IsActionMapped(FilterCommand))
+						{
+							CurveEditorSharedBindings->MapAction(FilterCommand, *FilterCommandList->GetActionForCommand(FilterCommand));
+						}
+					}
+				}
+
 				CurveEditorCommands->Append(CurveEditorSharedBindings.ToSharedRef());
+
 			}
 		}
 	}
