Sequencer: Restrict creating keys with middle mouse button to only when there are no modifier keys down

This fixes an issue where if you zoom scroll (ctrl and middle mouse wheel) you might accidentally press the middle mouse button down, creating a bunch of keys.

Tested that you can still create keys with the middle mouse button or by holding down alt and left mouse button.

diff --git a/Engine/Source/Editor/Sequencer/Private/MVVM/Views/SChannelView.cpp b/Engine/Source/Editor/Sequencer/Private/MVVM/Views/SChannelView.cpp
--- a/Engine/Source/Editor/Sequencer/Private/MVVM/Views/SChannelView.cpp
+++ b/Engine/Source/Editor/Sequencer/Private/MVVM/Views/SChannelView.cpp
@@ -251,7 +251,8 @@ FReply SChannelView::OnMouseButtonDown(const FGeometry& MyGeometry, const FPoint
 		HoveredKeys = KeyHotspot->Keys;
 	}
 
-	if (MouseEvent.GetEffectingButton() == EKeys::MiddleMouseButton)
+	// Create keys with middle mouse button only if there are no modifier keys down (to prevent zoom scrolling and creating keys at the same time)
+	if (!MouseEvent.GetModifierKeys().AnyModifiersDown() && MouseEvent.GetEffectingButton() == EKeys::MiddleMouseButton)
 	{
 		GEditor->BeginTransaction(LOCTEXT("CreateKey_Transaction", "Create Key"));
 
@@ -382,7 +383,7 @@ FReply SChannelView::OnMouseMove(const FGeometry& MyGeometry, const FPointerEven
 
 FReply SChannelView::OnMouseButtonUp(const FGeometry& MyGeometry, const FPointerEvent& MouseEvent)
 {
-	if (MouseEvent.GetEffectingButton() == EKeys::MiddleMouseButton)
+	if (!MouseEvent.GetModifierKeys().AnyModifiersDown() && MouseEvent.GetEffectingButton() == EKeys::MiddleMouseButton)
 	{
 		TSharedPtr<FSequencer> Sequencer = LegacyGetSequencer();
 
diff --git a/Engine/Source/Editor/Sequencer/Private/Tools/SequencerEditTool_Movement.cpp b/Engine/Source/Editor/Sequencer/Private/Tools/SequencerEditTool_Movement.cpp
--- a/Engine/Source/Editor/Sequencer/Private/Tools/SequencerEditTool_Movement.cpp
+++ b/Engine/Source/Editor/Sequencer/Private/Tools/SequencerEditTool_Movement.cpp
@@ -175,8 +175,11 @@ TSharedPtr<UE::Sequencer::ISequencerEditToolDragOperation> FSequencerEditTool_Mo
 		// Gather the draggable sections from all selected models
 		const bool bModelsSelected = Selection.TrackArea.Num() > 0;
 		const bool bKeySelected = Selection.KeySelection.Num() > 0;
+		
 		// @todo sequencer: Make this a customizable UI command modifier?
-		const bool bIsDuplicateEvent = MouseEvent.IsAltDown() || MouseEvent.GetEffectingButton() == EKeys::MiddleMouseButton;
+		const bool bIsDuplicateEvent = 
+			MouseEvent.IsAltDown() || 
+			(!MouseEvent.GetModifierKeys().AnyModifiersDown() && MouseEvent.GetEffectingButton() == EKeys::MiddleMouseButton);
 
 		FSectionHotspotBase* SectionHotspot = DelayedDrag->Hotspot->CastThis<FSectionHotspotBase>();
 		FKeyHotspot*         KeyHotspot     = DelayedDrag->Hotspot->CastThis<FKeyHotspot>();
@@ -306,7 +309,11 @@ FReply FSequencerEditTool_Movement::OnMouseButtonUp(SWidget& OwnerWidget, const
 		DragOperation->OnEndDrag(MouseEvent, MyGeometry.AbsoluteToLocal(MouseEvent.GetScreenSpacePosition()), SequencerWidget->GetVirtualTrackArea(&TrackArea));
 		DragOperation = nullptr;
 
-		if (MouseEvent.GetEffectingButton() == EKeys::MiddleMouseButton)
+		const bool bIsDuplicateEvent =
+			MouseEvent.IsAltDown() ||
+			(!MouseEvent.GetModifierKeys().AnyModifiersDown() && MouseEvent.GetEffectingButton() == EKeys::MiddleMouseButton);
+
+		if (bIsDuplicateEvent)
 		{
 			GEditor->EndTransaction();
 		}
