Sequencer: Do not re-initialize local time when compiling sequences if the transform is the same

diff --git a/Engine/Source/Editor/Sequencer/Private/Sequencer.cpp b/Engine/Source/Editor/Sequencer/Private/Sequencer.cpp
--- a/Engine/Source/Editor/Sequencer/Private/Sequencer.cpp
+++ b/Engine/Source/Editor/Sequencer/Private/Sequencer.cpp
@@ -926,6 +926,7 @@ void FSequencer::Tick(float InDeltaTime)
 			// Try and preserve the current local time across compilations.
 			//     When a modification results in a change of transform or time-warp,
 			//     this helps to stop the play head from jumping around
+			FMovieSceneSequenceTransform CachedRootToUnwarpedLocalTransform = RootToUnwarpedLocalTransform;
 			FFrameTime OldLocalTime = GetUnwarpedLocalTime().Time;
 
 			CompiledDataManager->Compile(RootSequencePtr);
@@ -952,8 +953,11 @@ void FSequencer::Tick(float InDeltaTime)
 				}
 			}
 
-			FTimeDomainOverride TimeDomain = OverrideTimeDomain(ETimeDomain::Unwarped);
-			SetLocalTime(OldLocalTime, ESnapTimeMode::STM_None, false /* bEvaluate */);
+			if (CachedRootToUnwarpedLocalTransform != RootToUnwarpedLocalTransform)
+			{
+				FTimeDomainOverride TimeDomain = OverrideTimeDomain(ETimeDomain::Unwarped);
+				SetLocalTime(OldLocalTime, ESnapTimeMode::STM_None, false /* bEvaluate */);
+			}
 
 			SuppressAutoEvalSignature.Reset();
 
