Sequencer: Added advanced text filters for custom bindings.

* BindingType/BindingClass = (BindingClassName)
* Custom/CustomBinding/IsCustom/IsCustomBinding
* Empty/EmptyBinding/IsEmpty/IsEmptyBinding
* Possessable/PossessableBinding/IsPossessable/IsPossessableBinding
* Replaceable/ReplaceableBinding/IsReplaceable/IsReplaceableBinding
* Spawnable/SpawnableBinding/IsSpawnable/IsSpawnableBinding

diff --git a/Engine/Source/Editor/Sequencer/Private/Filters/Filters/SequencerTrackFilter_Text.cpp b/Engine/Source/Editor/Sequencer/Private/Filters/Filters/SequencerTrackFilter_Text.cpp
--- a/Engine/Source/Editor/Sequencer/Private/Filters/Filters/SequencerTrackFilter_Text.cpp
+++ b/Engine/Source/Editor/Sequencer/Private/Filters/Filters/SequencerTrackFilter_Text.cpp
@@ -24,6 +24,9 @@
 #include "Filters/TextExpressions/SequencerTextFilterExpression_Time.h"
 #include "Filters/TextExpressions/SequencerTextFilterExpression_TrackClass.h"
 #include "Filters/TextExpressions/SequencerTextFilterExpression_Unbound.h"
+#include "Filters/TextExpressions/SequencerTextFilterExpression_BindingType.h"
+#include "Filters/TextExpressions/SequencerTextFilterExpression_CustomBinding.h"
+#include "Filters/TextExpressions/SequencerTextFilterExpression_EmptyBinding.h"
 #include "UObject/UObjectIterator.h"
 
 #define LOCTEXT_NAMESPACE "SequencerTrackFilter_Text"
@@ -58,6 +61,9 @@ FSequencerTrackFilter_Text::FSequencerTrackFilter_Text(ISequencerTrackFilters& I
 	TextFilterExpressionContexts.Add(MakeShared<FSequencerTextFilterExpression_Muted>(InFilterInterface));
 	TextFilterExpressionContexts.Add(MakeShared<FSequencerTextFilterExpression_Soloed>(InFilterInterface));
 	TextFilterExpressionContexts.Add(MakeShared<FSequencerTextFilterExpression_Tag>(InFilterInterface));
+	TextFilterExpressionContexts.Add(MakeShared<FSequencerTextFilterExpression_BindingType>(InFilterInterface));
+	TextFilterExpressionContexts.Add(MakeShared<FSequencerTextFilterExpression_CustomBinding>(InFilterInterface));
+	TextFilterExpressionContexts.Add(MakeShared<FSequencerTextFilterExpression_EmptyBinding>(InFilterInterface));
 
 	// Add global user-defined text expressions
 	for (TObjectIterator<USequencerTrackFilterTextExpressionExtension> ExtensionIt(RF_NoFlags); ExtensionIt; ++ExtensionIt)
diff --git a/Engine/Source/Editor/Sequencer/Private/Filters/TextExpressions/SequencerTextFilterExpression_BindingType.cpp b/Engine/Source/Editor/Sequencer/Private/Filters/TextExpressions/SequencerTextFilterExpression_BindingType.cpp
new file mode 100644
--- /dev/null
+++ b/Engine/Source/Editor/Sequencer/Private/Filters/TextExpressions/SequencerTextFilterExpression_BindingType.cpp
@@ -0,0 +1,119 @@
+// Copyright Epic Games, Inc. All Rights Reserved.
+
+#include "SequencerTextFilterExpression_BindingType.h"
+#include "MVVM/Extensions/IObjectBindingExtension.h"
+#include "MVVM/ViewModels/ViewModelIterators.h"
+#include "MVVM/ViewModelPtr.h"
+#include "MVVM/SharedViewModelData.h"
+#include "MovieSceneBindingReferences.h"
+#include "Bindings/MovieSceneSpawnableBinding.h"
+#include "Bindings/MovieSceneReplaceableBinding.h"
+
+using namespace UE::Sequencer;
+
+#define LOCTEXT_NAMESPACE "SequencerTextFilterExpression_BindingType"
+
+FSequencerTextFilterExpression_BindingType::FSequencerTextFilterExpression_BindingType(ISequencerTrackFilters& InFilterInterface)
+	: FSequencerTextFilterExpressionContext(InFilterInterface)
+{
+}
+
+TSet<FName> FSequencerTextFilterExpression_BindingType::GetKeys() const
+{
+	return { TEXT("BindingClass"), TEXT("BindingType") };
+}
+
+ESequencerTextFilterValueType FSequencerTextFilterExpression_BindingType::GetValueType() const
+{
+	return ESequencerTextFilterValueType::String;
+}
+
+FText FSequencerTextFilterExpression_BindingType::GetDescription() const
+{
+	return LOCTEXT("ExpressionDescription_BindingType", "Filter by presence of a binding with the given type");
+}
+
+bool FSequencerTextFilterExpression_BindingType::TestComplexExpression(const FName& InKey
+	, const FTextFilterString& InValue
+	, const ETextFilterComparisonOperation InComparisonOperation
+	, const ETextFilterTextComparisonMode InTextComparisonMode) const
+{
+	if (!FSequencerTextFilterExpressionContext::TestComplexExpression(InKey, InValue, InComparisonOperation, InTextComparisonMode))
+	{
+		return true;
+	}
+
+	const TViewModelPtr<IObjectBindingExtension> ObjectBindingExtension = FilterItem->FindAncestorOfType<IObjectBindingExtension>(true);
+	if (!ObjectBindingExtension.IsValid())
+	{
+		return false;
+	}
+
+	FGuid ObjectBindingID = ObjectBindingExtension->GetObjectGuid();
+
+	const UMovieSceneSequence* FocusedSequence = GetFocusedMovieSceneSequence();
+	if (!IsValid(FocusedSequence))
+	{
+		return true;
+	}
+
+	const FMovieSceneBindingReferences* BindingReferences = FocusedSequence->GetBindingReferences();
+	const FMovieSceneBindingReference* BindingReference = BindingReferences ? BindingReferences->GetReference(ObjectBindingID, 0) : nullptr;
+
+	// Special cases- possessable, spawnable, replaceable
+	if (TextFilterUtils::TestComplexExpression(TEXT("Possessable"), InValue, InComparisonOperation, InTextComparisonMode))
+	{
+		if (!BindingReferences)
+		{
+			return FocusedSequence->GetMovieScene()->FindPossessable(ObjectBindingID) != nullptr;
+		}
+		else if (!BindingReference)
+		{
+			return false;
+		}
+		else 
+		{
+			return !BindingReference->Locator.IsEmpty() && BindingReference->CustomBinding == nullptr;
+		}
+	}
+	else if (TextFilterUtils::TestComplexExpression(TEXT("Spawnable"), InValue, InComparisonOperation, InTextComparisonMode))
+	{
+		if (!BindingReferences)
+		{
+			return FocusedSequence->GetMovieScene()->FindSpawnable(ObjectBindingID) != nullptr;
+		}
+		else if (!BindingReference)
+		{
+			return false;
+		}
+		else
+		{
+			return Cast<UMovieSceneSpawnableBindingBase>(BindingReference->CustomBinding) != nullptr;
+		}
+	}
+	else if (TextFilterUtils::TestComplexExpression(TEXT("Replaceable"), InValue, InComparisonOperation, InTextComparisonMode))
+	{
+		if (!BindingReferences)
+		{
+			return FocusedSequence->GetMovieScene()->FindSpawnable(ObjectBindingID) != nullptr;
+		}
+		else if (!BindingReference)
+		{
+			return false;
+		}
+		else
+		{
+			return Cast<UMovieSceneReplaceableBindingBase>(BindingReference->CustomBinding) != nullptr;
+		}
+	}
+	else if (!BindingReference || !BindingReference->CustomBinding)
+	{
+		return false;
+	}
+	else
+	{
+		return TextFilterUtils::TestComplexExpression(BindingReference->CustomBinding->GetClass()->GetName(), InValue, InComparisonOperation, InTextComparisonMode);
+	}
+}
+
+#undef LOCTEXT_NAMESPACE
diff --git a/Engine/Source/Editor/Sequencer/Private/Filters/TextExpressions/SequencerTextFilterExpression_BindingType.h b/Engine/Source/Editor/Sequencer/Private/Filters/TextExpressions/SequencerTextFilterExpression_BindingType.h
new file mode 100644
--- /dev/null
+++ b/Engine/Source/Editor/Sequencer/Private/Filters/TextExpressions/SequencerTextFilterExpression_BindingType.h
@@ -0,0 +1,24 @@
+// Copyright Epic Games, Inc. All Rights Reserved.
+
+#pragma once
+
+#include "Filters/SequencerTextFilterExpressionContext.h"
+
+class FSequencerTextFilterExpression_BindingType : public FSequencerTextFilterExpressionContext
+{
+public:
+	FSequencerTextFilterExpression_BindingType(ISequencerTrackFilters& InFilterInterface);
+
+	//~ Begin FSequencerTextFilterExpressionContext
+	virtual TSet<FName> GetKeys() const override;
+	virtual ESequencerTextFilterValueType GetValueType() const override;
+	virtual FText GetDescription() const override;
+	//~ End FSequencerTextFilterExpressionContext
+
+	//~ Begin ITextFilterExpressionContext
+	virtual bool TestComplexExpression(const FName& InKey
+		, const FTextFilterString& InValue
+		, const ETextFilterComparisonOperation InComparisonOperation
+		, const ETextFilterTextComparisonMode InTextComparisonMode) const override;
+	//~ End ITextFilterExpressionContext
+};
diff --git a/Engine/Source/Editor/Sequencer/Private/Filters/TextExpressions/SequencerTextFilterExpression_CustomBinding.cpp b/Engine/Source/Editor/Sequencer/Private/Filters/TextExpressions/SequencerTextFilterExpression_CustomBinding.cpp
new file mode 100644
--- /dev/null
+++ b/Engine/Source/Editor/Sequencer/Private/Filters/TextExpressions/SequencerTextFilterExpression_CustomBinding.cpp
@@ -0,0 +1,69 @@
+// Copyright Epic Games, Inc. All Rights Reserved.
+
+#include "SequencerTextFilterExpression_CustomBinding.h"
+#include "MVVM/Extensions/IObjectBindingExtension.h"
+#include "MVVM/ViewModels/ViewModelIterators.h"
+#include "MVVM/ViewModelPtr.h"
+#include "MVVM/SharedViewModelData.h"
+#include "MovieSceneBindingReferences.h"
+
+using namespace UE::Sequencer;
+
+#define LOCTEXT_NAMESPACE "SequencerTextFilterExpression_CustomBinding"
+
+FSequencerTextFilterExpression_CustomBinding::FSequencerTextFilterExpression_CustomBinding(ISequencerTrackFilters& InFilterInterface)
+	: FSequencerTextFilterExpressionContext(InFilterInterface)
+{
+}
+
+TSet<FName> FSequencerTextFilterExpression_CustomBinding::GetKeys() const
+{
+	return { TEXT("Custom"), TEXT("CustomBinding"), TEXT("IsCustom"), TEXT("IsCustomBinding") };
+}
+
+ESequencerTextFilterValueType FSequencerTextFilterExpression_CustomBinding::GetValueType() const
+{
+	return ESequencerTextFilterValueType::Boolean;
+}
+
+FText FSequencerTextFilterExpression_CustomBinding::GetDescription() const
+{
+	return LOCTEXT("ExpressionDescription_CustomBinding", "Filter by presence of a custom binding");
+}
+
+bool FSequencerTextFilterExpression_CustomBinding::TestComplexExpression(const FName& InKey
+	, const FTextFilterString& InValue
+	, const ETextFilterComparisonOperation InComparisonOperation
+	, const ETextFilterTextComparisonMode InTextComparisonMode) const
+{
+	if (!FSequencerTextFilterExpressionContext::TestComplexExpression(InKey, InValue, InComparisonOperation, InTextComparisonMode))
+	{
+		return true;
+	}
+
+	const TViewModelPtr<IObjectBindingExtension> ObjectBindingExtension = FilterItem->FindAncestorOfType<IObjectBindingExtension>(true);
+	if (!ObjectBindingExtension.IsValid())
+	{
+		return false;
+	}
+
+	FGuid ObjectBindingID = ObjectBindingExtension->GetObjectGuid();
+
+	const UMovieSceneSequence* FocusedSequence = GetFocusedMovieSceneSequence();
+	if (!IsValid(FocusedSequence))
+	{
+		return true;
+	}
+
+	const FMovieSceneBindingReferences* BindingReferences = FocusedSequence->GetBindingReferences();
+	if (!BindingReferences)
+	{
+		return false;
+	}
+
+	bool bHasCustomBinding = BindingReferences->GetCustomBinding(ObjectBindingID, 0) != nullptr;
+
+	return CompareFStringForExactBool(InValue, InComparisonOperation, bHasCustomBinding);
+}
+
+#undef LOCTEXT_NAMESPACE
diff --git a/Engine/Source/Editor/Sequencer/Private/Filters/TextExpressions/SequencerTextFilterExpression_CustomBinding.h b/Engine/Source/Editor/Sequencer/Private/Filters/TextExpressions/SequencerTextFilterExpression_CustomBinding.h
new file mode 100644
--- /dev/null
+++ b/Engine/Source/Editor/Sequencer/Private/Filters/TextExpressions/SequencerTextFilterExpression_CustomBinding.h
@@ -0,0 +1,24 @@
+// Copyright Epic Games, Inc. All Rights Reserved.
+
+#pragma once
+
+#include "Filters/SequencerTextFilterExpressionContext.h"
+
+class FSequencerTextFilterExpression_CustomBinding : public FSequencerTextFilterExpressionContext
+{
+public:
+	FSequencerTextFilterExpression_CustomBinding(ISequencerTrackFilters& InFilterInterface);
+
+	//~ Begin FSequencerTextFilterExpressionContext
+	virtual TSet<FName> GetKeys() const override;
+	virtual ESequencerTextFilterValueType GetValueType() const override;
+	virtual FText GetDescription() const override;
+	//~ End FSequencerTextFilterExpressionContext
+
+	//~ Begin ITextFilterExpressionContext
+	virtual bool TestComplexExpression(const FName& InKey
+		, const FTextFilterString& InValue
+		, const ETextFilterComparisonOperation InComparisonOperation
+		, const ETextFilterTextComparisonMode InTextComparisonMode) const override;
+	//~ End ITextFilterExpressionContext
+};
diff --git a/Engine/Source/Editor/Sequencer/Private/Filters/TextExpressions/SequencerTextFilterExpression_EmptyBinding.cpp b/Engine/Source/Editor/Sequencer/Private/Filters/TextExpressions/SequencerTextFilterExpression_EmptyBinding.cpp
new file mode 100644
--- /dev/null
+++ b/Engine/Source/Editor/Sequencer/Private/Filters/TextExpressions/SequencerTextFilterExpression_EmptyBinding.cpp
@@ -0,0 +1,76 @@
+// Copyright Epic Games, Inc. All Rights Reserved.
+
+#include "SequencerTextFilterExpression_EmptyBinding.h"
+#include "MVVM/Extensions/IObjectBindingExtension.h"
+#include "MVVM/ViewModels/ViewModelIterators.h"
+#include "MVVM/ViewModelPtr.h"
+#include "MVVM/SharedViewModelData.h"
+#include "MovieSceneBindingReferences.h"
+
+using namespace UE::Sequencer;
+
+#define LOCTEXT_NAMESPACE "SequencerTextFilterExpression_EmptyBinding"
+
+FSequencerTextFilterExpression_EmptyBinding::FSequencerTextFilterExpression_EmptyBinding(ISequencerTrackFilters& InFilterInterface)
+	: FSequencerTextFilterExpressionContext(InFilterInterface)
+{
+}
+
+TSet<FName> FSequencerTextFilterExpression_EmptyBinding::GetKeys() const
+{
+	return { TEXT("Empty"), TEXT("EmptyBinding"), TEXT("IsEmpty"), TEXT("IsEmptyBinding") };
+}
+
+ESequencerTextFilterValueType FSequencerTextFilterExpression_EmptyBinding::GetValueType() const
+{
+	return ESequencerTextFilterValueType::Boolean;
+}
+
+FText FSequencerTextFilterExpression_EmptyBinding::GetDescription() const
+{
+	return LOCTEXT("ExpressionDescription_EmptyBinding", "Filter by presence of an empty binding");
+}
+
+bool FSequencerTextFilterExpression_EmptyBinding::TestComplexExpression(const FName& InKey
+	, const FTextFilterString& InValue
+	, const ETextFilterComparisonOperation InComparisonOperation
+	, const ETextFilterTextComparisonMode InTextComparisonMode) const
+{
+	if (!FSequencerTextFilterExpressionContext::TestComplexExpression(InKey, InValue, InComparisonOperation, InTextComparisonMode))
+	{
+		return true;
+	}
+
+	const TViewModelPtr<IObjectBindingExtension> ObjectBindingExtension = FilterItem->FindAncestorOfType<IObjectBindingExtension>(true);
+	if (!ObjectBindingExtension.IsValid())
+	{
+		return false;
+	}
+
+	FGuid ObjectBindingID = ObjectBindingExtension->GetObjectGuid();
+
+	const UMovieSceneSequence* FocusedSequence = GetFocusedMovieSceneSequence();
+	if (!IsValid(FocusedSequence))
+	{
+		return true;
+	}
+
+	const FMovieSceneBindingReferences* BindingReferences = FocusedSequence->GetBindingReferences();
+	if (!BindingReferences)
+	{
+		return false;
+	}
+
+	const FMovieSceneBindingReference* Reference = BindingReferences->GetReference(ObjectBindingID, 0);
+
+	if (!Reference)
+	{
+		return false;
+	}
+
+	bool bIsEmptyBinding = Reference->Locator.IsEmpty() && Reference->CustomBinding == nullptr;
+
+	return CompareFStringForExactBool(InValue, InComparisonOperation, bIsEmptyBinding);
+}
+
+#undef LOCTEXT_NAMESPACE
diff --git a/Engine/Source/Editor/Sequencer/Private/Filters/TextExpressions/SequencerTextFilterExpression_EmptyBinding.h b/Engine/Source/Editor/Sequencer/Private/Filters/TextExpressions/SequencerTextFilterExpression_EmptyBinding.h
new file mode 100644
--- /dev/null
+++ b/Engine/Source/Editor/Sequencer/Private/Filters/TextExpressions/SequencerTextFilterExpression_EmptyBinding.h
@@ -0,0 +1,24 @@
+// Copyright Epic Games, Inc. All Rights Reserved.
+
+#pragma once
+
+#include "Filters/SequencerTextFilterExpressionContext.h"
+
+class FSequencerTextFilterExpression_EmptyBinding : public FSequencerTextFilterExpressionContext
+{
+public:
+	FSequencerTextFilterExpression_EmptyBinding(ISequencerTrackFilters& InFilterInterface);
+
+	//~ Begin FSequencerTextFilterExpressionContext
+	virtual TSet<FName> GetKeys() const override;
+	virtual ESequencerTextFilterValueType GetValueType() const override;
+	virtual FText GetDescription() const override;
+	//~ End FSequencerTextFilterExpressionContext
+
+	//~ Begin ITextFilterExpressionContext
+	virtual bool TestComplexExpression(const FName& InKey
+		, const FTextFilterString& InValue
+		, const ETextFilterComparisonOperation InComparisonOperation
+		, const ETextFilterTextComparisonMode InTextComparisonMode) const override;
+	//~ End ITextFilterExpressionContext
+};
