[Sequencer] Request filter update on undo

diff --git a/Engine/Source/Editor/Sequencer/Private/Sequencer.cpp b/Engine/Source/Editor/Sequencer/Private/Sequencer.cpp
--- a/Engine/Source/Editor/Sequencer/Private/Sequencer.cpp
+++ b/Engine/Source/Editor/Sequencer/Private/Sequencer.cpp
@@ -5964,6 +5964,7 @@ void FSequencer::PostUndo(bool bSuccess)
 	}
 
 	NodeTree->SortAllNodesAndDescendants();
+	NodeTree->RequestFilterUpdate();
 
 	OnActivateSequenceEvent.Broadcast(ActiveTemplateIDs.Top());
 }
