[Sequencer] Expose GetHostCapabilities() and GetPlaybackRange() functions through ISequencer
NOTE: Will be used to implement the Motion Design Stagger Tool in the Sequencer sidebar.

diff --git a/Engine/Source/Editor/Sequencer/Private/Sequencer.h b/Engine/Source/Editor/Sequencer/Private/Sequencer.h
--- a/Engine/Source/Editor/Sequencer/Private/Sequencer.h
+++ b/Engine/Source/Editor/Sequencer/Private/Sequencer.h
@@ -206,7 +206,7 @@ public:
 	 * @return Playback range.
 	 * @see SetPlaybackRange, SetPlaybackRangeEnd, SetPlaybackRangeStart
 	 */
-	TRange<FFrameNumber> GetPlaybackRange() const;
+	virtual TRange<FFrameNumber> GetPlaybackRange() const override;
 
 	/**
 	 * Set the playback range.
@@ -525,7 +525,8 @@ public:
 	/** @return The toolkit that this sequencer is hosted in (if any) */
 	virtual TSharedPtr<IToolkitHost> GetToolkitHost() const override { return ToolkitHost.Pin(); }
 
-	const FSequencerHostCapabilities& GetHostCapabilities() const { return HostCapabilities; }
+	/** @return The capabilities of this sequencer host */
+	virtual const FSequencerHostCapabilities& GetHostCapabilities() const override { return HostCapabilities; }
 
 	/** @return Whether or not this sequencer is used in the level editor */
 	virtual bool IsLevelEditorSequencer() const override { return bIsEditingWithinLevelEditor; }
diff --git a/Engine/Source/Editor/Sequencer/Public/ISequencer.h b/Engine/Source/Editor/Sequencer/Public/ISequencer.h
--- a/Engine/Source/Editor/Sequencer/Public/ISequencer.h
+++ b/Engine/Source/Editor/Sequencer/Public/ISequencer.h
@@ -68,6 +68,7 @@ struct FMovieSceneBinding;
 struct FMovieSceneChannelHandle;
 struct FMovieSceneMarkedFrame;
 struct FQualifiedFrameTime;
+struct FSequencerHostCapabilities;
 struct FSidebarDrawerConfig;
 template <typename NumericType> struct INumericTypeInterface;
 
@@ -432,6 +433,9 @@ public:
 	/** Returns the Toolkit hosting the sequencer instance, if any */	
 	virtual TSharedPtr<IToolkitHost> GetToolkitHost() const = 0;
 
+	/** @return The capabilities of this sequencer host */
+	virtual const FSequencerHostCapabilities& GetHostCapabilities() const = 0;
+
 	/**
 	 * Gets the current time of the time slider relative to the currently focused movie scene
 	 */
@@ -649,6 +653,9 @@ protected:
 public:
 	virtual void NotifyMovieSceneDataChanged( EMovieSceneDataChangeType DataChangeType ) = 0;
 
+	/** Get the playback range. */
+	virtual TRange<FFrameNumber> GetPlaybackRange() const = 0;
+
 	virtual void UpdatePlaybackRange() = 0;
 
 	virtual void SetPlaybackSpeed(float InPlaybackSpeed) = 0;
