MUSIC_IN_SEQUENCER [Initial Check-In]
The beginnings of new, experimental, music features in sequencer (eg. editing/quantizing keyframes in "musical time", importing of music maps, etc.) This is a Work-In-Progress and does not implement a fully working feature. It is the first minimally required changelist to allow further collaboration.

This is a combination of a number of changelists that were reviewed in these swarm reviews...

https://p4-swarm.epicgames.net/reviews/36988914
https://p4-swarm.epicgames.net/reviews/36994249
https://p4-swarm.epicgames.net/reviews/36995429
https://p4-swarm.epicgames.net/reviews/36998041
https://p4-swarm.epicgames.net/reviews/37001010
https://p4-swarm.epicgames.net/reviews/37003290
https://p4-swarm.epicgames.net/reviews/37003378
https://p4-swarm.epicgames.net/reviews/37003982
https://p4-swarm.epicgames.net/reviews/37037496

(Checked in as one changelist due to dependency issues.)

diff --git a/Engine/Plugins/Animation/ControlRig/Source/ControlRig/Private/Sequencer/MovieSceneControlRigSpaceChannel.cpp b/Engine/Plugins/Animation/ControlRig/Source/ControlRig/Private/Sequencer/MovieSceneControlRigSpaceChannel.cpp
--- a/Engine/Plugins/Animation/ControlRig/Source/ControlRig/Private/Sequencer/MovieSceneControlRigSpaceChannel.cpp
+++ b/Engine/Plugins/Animation/ControlRig/Source/ControlRig/Private/Sequencer/MovieSceneControlRigSpaceChannel.cpp
@@ -87,6 +87,11 @@ void FMovieSceneControlRigSpaceChannel::ChangeFrameResolution(FFrameRate SourceR
 	GetData().ChangeFrameResolution(SourceRate, DestinationRate);
 }
 
+void FMovieSceneControlRigSpaceChannel::RemapMusicalFrameTimes(const UFrameBasedMusicMap* OldMap, const UFrameBasedMusicMap* NewMap)
+{
+	GetData().RemapMusicalFrameTimes(OldMap, NewMap);
+}
+
 TRange<FFrameNumber> FMovieSceneControlRigSpaceChannel::ComputeEffectiveRange() const
 {
 	return GetData().GetTotalRange();
diff --git a/Engine/Plugins/Animation/ControlRig/Source/ControlRig/Public/Sequencer/MovieSceneControlRigSpaceChannel.h b/Engine/Plugins/Animation/ControlRig/Source/ControlRig/Public/Sequencer/MovieSceneControlRigSpaceChannel.h
--- a/Engine/Plugins/Animation/ControlRig/Source/ControlRig/Public/Sequencer/MovieSceneControlRigSpaceChannel.h
+++ b/Engine/Plugins/Animation/ControlRig/Source/ControlRig/Public/Sequencer/MovieSceneControlRigSpaceChannel.h
@@ -107,6 +107,7 @@ public:
 	virtual void DeleteKeys(TArrayView<const FKeyHandle> InHandles) override;
 	virtual void DeleteKeysFrom(FFrameNumber InTime, bool bDeleteKeysBefore) override;
 	virtual void ChangeFrameResolution(FFrameRate SourceRate, FFrameRate DestinationRate) override;
+	virtual void RemapMusicalFrameTimes(const UFrameBasedMusicMap* OldMap, const UFrameBasedMusicMap* NewMap) override;
 	virtual TRange<FFrameNumber> ComputeEffectiveRange() const override;
 	virtual int32 GetNumKeys() const override;
 	virtual void Reset() override;
diff --git a/Engine/Plugins/Animation/ControlRig/Source/ControlRigEditor/Private/EditMode/SControlRigSnapper.cpp b/Engine/Plugins/Animation/ControlRig/Source/ControlRigEditor/Private/EditMode/SControlRigSnapper.cpp
--- a/Engine/Plugins/Animation/ControlRig/Source/ControlRigEditor/Private/EditMode/SControlRigSnapper.cpp
+++ b/Engine/Plugins/Animation/ControlRig/Source/ControlRigEditor/Private/EditMode/SControlRigSnapper.cpp
@@ -290,7 +290,7 @@ void SControlRigSnapper::Construct(const FArguments& InArgs)
 										TSharedPtr<ISequencer> SequencerPtr = Snapper.GetSequencer().Pin();
 										if (SequencerPtr.IsValid() && SequencerPtr->GetFocusedMovieSceneSequence())
 										{
-											TOptional<double> NewFrameTime = SequencerPtr->GetNumericTypeInterface()->FromString(InText.ToString(), 0);
+											TOptional<double> NewFrameTime = SequencerPtr->GetNumericPositionTypeInterface()->FromString(InText.ToString(), 0);
 											if (NewFrameTime.IsSet())
 											{
 												StartFrame = FFrameNumber((int32)NewFrameTime.GetValue());
@@ -324,7 +324,7 @@ void SControlRigSnapper::Construct(const FArguments& InArgs)
 										if (Sequencer.IsValid() && Sequencer.Pin()->GetFocusedMovieSceneSequence())
 										{
 											TSharedPtr<ISequencer> SequencerPtr = Sequencer.Pin();
-											TOptional<double> NewFrameTime = SequencerPtr->GetNumericTypeInterface()->FromString(InText.ToString(), 0);
+											TOptional<double> NewFrameTime = SequencerPtr->GetNumericPositionTypeInterface()->FromString(InText.ToString(), 0);
 											if (NewFrameTime.IsSet())
 											{
 												EndFrame = FFrameNumber((int32)NewFrameTime.GetValue());
@@ -466,7 +466,7 @@ FText SControlRigSnapper::GetStartFrameToSnapText()
 	TWeakPtr<ISequencer> Sequencer = Snapper.GetSequencer();
 	if (Sequencer.IsValid() && Sequencer.Pin()->GetFocusedMovieSceneSequence())
 	{
-		Value = FText::FromString(Sequencer.Pin()->GetNumericTypeInterface()->ToString(StartFrame.Value));
+		Value = FText::FromString(Sequencer.Pin()->GetNumericPositionTypeInterface()->ToString(StartFrame.Value));
 	}
 	return Value;
 }
@@ -477,7 +477,7 @@ FText SControlRigSnapper::GetEndFrameToSnapText()
 	TWeakPtr<ISequencer> Sequencer = Snapper.GetSequencer();
 	if (Sequencer.IsValid() && Sequencer.Pin()->GetFocusedMovieSceneSequence())
 	{
-		Value = FText::FromString(Sequencer.Pin()->GetNumericTypeInterface()->ToString(EndFrame.Value));
+		Value = FText::FromString(Sequencer.Pin()->GetNumericPositionTypeInterface()->ToString(EndFrame.Value));
 	}
 	return Value;
 }
diff --git a/Engine/Plugins/Animation/ControlRig/Source/ControlRigEditor/Private/Editor/Constraints/SConstraintsWidget.cpp b/Engine/Plugins/Animation/ControlRig/Source/ControlRigEditor/Private/Editor/Constraints/SConstraintsWidget.cpp
--- a/Engine/Plugins/Animation/ControlRig/Source/ControlRigEditor/Private/Editor/Constraints/SConstraintsWidget.cpp
+++ b/Engine/Plugins/Animation/ControlRig/Source/ControlRigEditor/Private/Editor/Constraints/SConstraintsWidget.cpp
@@ -1578,9 +1578,9 @@ void SConstraintBakeWidget::Construct(const FArguments& InArgs)
 	FPropertyEditorModule& PropertyEditor = FModuleManager::Get().LoadModuleChecked<FPropertyEditorModule>(TEXT("PropertyEditor"));
 
 	DetailsView = PropertyEditor.CreateStructureDetailView(ViewArgs, StructureViewArgs, TSharedPtr<FStructOnScope>());
-	TSharedPtr<INumericTypeInterface<double>> NumericTypeInterface = Sequencer->GetNumericTypeInterface();
+	TSharedPtr<INumericTypeInterface<double>> NumericPositionTypeInterface = Sequencer->GetNumericPositionTypeInterface();
 	DetailsView->GetDetailsView()->RegisterInstancedCustomPropertyTypeLayout("FrameNumber",
-		FOnGetPropertyTypeCustomizationInstance::CreateLambda([=]() {return MakeShared<FFrameNumberDetailsCustomization>(NumericTypeInterface); }));
+		FOnGetPropertyTypeCustomizationInstance::CreateLambda([=]() {return MakeShared<FFrameNumberDetailsCustomization>(NumericPositionTypeInterface); }));
 	DetailsView->SetStructureData(Settings);
 
 
diff --git a/Engine/Plugins/Animation/ControlRig/Source/ControlRigEditor/Private/Editor/SRigSpacePickerWidget.cpp b/Engine/Plugins/Animation/ControlRig/Source/ControlRigEditor/Private/Editor/SRigSpacePickerWidget.cpp
--- a/Engine/Plugins/Animation/ControlRig/Source/ControlRigEditor/Private/Editor/SRigSpacePickerWidget.cpp
+++ b/Engine/Plugins/Animation/ControlRig/Source/ControlRigEditor/Private/Editor/SRigSpacePickerWidget.cpp
@@ -1182,9 +1182,9 @@ void SRigSpacePickerBakeWidget::Construct(const FArguments& InArgs)
 	FPropertyEditorModule& PropertyEditor = FModuleManager::Get().LoadModuleChecked<FPropertyEditorModule>(TEXT("PropertyEditor"));
 
 	DetailsView = PropertyEditor.CreateStructureDetailView(ViewArgs, StructureViewArgs, TSharedPtr<FStructOnScope>());
-	TSharedPtr<INumericTypeInterface<double>> NumericTypeInterface = Sequencer->GetNumericTypeInterface();
+	TSharedPtr<INumericTypeInterface<double>> NumericPositionTypeInterface = Sequencer->GetNumericPositionTypeInterface();
 	DetailsView->GetDetailsView()->RegisterInstancedCustomPropertyTypeLayout("FrameNumber",
-		FOnGetPropertyTypeCustomizationInstance::CreateLambda([=]() {return MakeShared<FFrameNumberDetailsCustomization>(NumericTypeInterface); }));
+		FOnGetPropertyTypeCustomizationInstance::CreateLambda([=]() {return MakeShared<FFrameNumberDetailsCustomization>(NumericPositionTypeInterface); }));
 	DetailsView->SetStructureData(Settings);
 
 	ChildSlot
diff --git a/Engine/Plugins/Animation/ControlRig/Source/ControlRigEditor/Private/Sequencer/ControlRigParameterTrackEditor.cpp b/Engine/Plugins/Animation/ControlRig/Source/ControlRigEditor/Private/Sequencer/ControlRigParameterTrackEditor.cpp
--- a/Engine/Plugins/Animation/ControlRig/Source/ControlRigEditor/Private/Sequencer/ControlRigParameterTrackEditor.cpp
+++ b/Engine/Plugins/Animation/ControlRig/Source/ControlRigEditor/Private/Sequencer/ControlRigParameterTrackEditor.cpp
@@ -5034,9 +5034,9 @@ void SCollapseControlsWidget::Construct(const FArguments& InArgs)
 	FPropertyEditorModule& PropertyEditor = FModuleManager::Get().LoadModuleChecked<FPropertyEditorModule>(TEXT("PropertyEditor"));
 
 	DetailsView = PropertyEditor.CreateStructureDetailView(ViewArgs, StructureViewArgs, TSharedPtr<FStructOnScope>());
-	TSharedPtr<INumericTypeInterface<double>> NumericTypeInterface = Sequencer.Pin()->GetNumericTypeInterface();
+	TSharedPtr<INumericTypeInterface<double>> NumericPositionTypeInterface = Sequencer.Pin()->GetNumericPositionTypeInterface();
 	DetailsView->GetDetailsView()->RegisterInstancedCustomPropertyTypeLayout("FrameNumber",
-		FOnGetPropertyTypeCustomizationInstance::CreateLambda([=]() {return MakeShared<FFrameNumberDetailsCustomization>(NumericTypeInterface); }));
+		FOnGetPropertyTypeCustomizationInstance::CreateLambda([=]() {return MakeShared<FFrameNumberDetailsCustomization>(NumericPositionTypeInterface); }));
 	DetailsView->SetStructureData(Settings);
 
 	ChildSlot
diff --git a/Engine/Plugins/FX/Niagara/Source/NiagaraEditor/Private/Sequencer/NiagaraSequence/Sections/MovieSceneNiagaraEmitterSection.cpp b/Engine/Plugins/FX/Niagara/Source/NiagaraEditor/Private/Sequencer/NiagaraSequence/Sections/MovieSceneNiagaraEmitterSection.cpp
--- a/Engine/Plugins/FX/Niagara/Source/NiagaraEditor/Private/Sequencer/NiagaraSequence/Sections/MovieSceneNiagaraEmitterSection.cpp
+++ b/Engine/Plugins/FX/Niagara/Source/NiagaraEditor/Private/Sequencer/NiagaraSequence/Sections/MovieSceneNiagaraEmitterSection.cpp
@@ -90,6 +90,11 @@ void FMovieSceneNiagaraEmitterChannel::ChangeFrameResolution(FFrameRate SourceRa
 	GetData().ChangeFrameResolution(SourceRate, DestinationRate);
 }
 
+void FMovieSceneNiagaraEmitterChannel::RemapMusicalFrameTimes(const UFrameBasedMusicMap* OldMap, const UFrameBasedMusicMap* NewMap)
+{
+	GetData().RemapMusicalFrameTimes(OldMap, NewMap);
+}
+
 TRange<FFrameNumber> FMovieSceneNiagaraEmitterChannel::ComputeEffectiveRange() const
 {
 	return GetData().GetTotalRange();
diff --git a/Engine/Plugins/FX/Niagara/Source/NiagaraEditor/Private/Sequencer/NiagaraSequence/Sections/MovieSceneNiagaraEmitterSection.h b/Engine/Plugins/FX/Niagara/Source/NiagaraEditor/Private/Sequencer/NiagaraSequence/Sections/MovieSceneNiagaraEmitterSection.h
--- a/Engine/Plugins/FX/Niagara/Source/NiagaraEditor/Private/Sequencer/NiagaraSequence/Sections/MovieSceneNiagaraEmitterSection.h
+++ b/Engine/Plugins/FX/Niagara/Source/NiagaraEditor/Private/Sequencer/NiagaraSequence/Sections/MovieSceneNiagaraEmitterSection.h
@@ -97,6 +97,7 @@ public:
 	virtual void DeleteKeys(TArrayView<const FKeyHandle> InHandles) override;
 	virtual void DeleteKeysFrom(FFrameNumber InTime, bool bDeleteKeysBefore) override;
 	virtual void ChangeFrameResolution(FFrameRate SourceRate, FFrameRate DestinationRate) override;
+	virtual void RemapMusicalFrameTimes(const UFrameBasedMusicMap* OldMap, const UFrameBasedMusicMap* NewMap) override;
 	virtual TRange<FFrameNumber> ComputeEffectiveRange() const override;
 	virtual int32 GetNumKeys() const override;
 	virtual void Reset() override;
diff --git a/Engine/Plugins/MovieScene/LevelSequenceEditor/Source/LevelSequenceEditor/Private/CinematicViewport/SCinematicLevelViewport.cpp b/Engine/Plugins/MovieScene/LevelSequenceEditor/Source/LevelSequenceEditor/Private/CinematicViewport/SCinematicLevelViewport.cpp
--- a/Engine/Plugins/MovieScene/LevelSequenceEditor/Source/LevelSequenceEditor/Private/CinematicViewport/SCinematicLevelViewport.cpp
+++ b/Engine/Plugins/MovieScene/LevelSequenceEditor/Source/LevelSequenceEditor/Private/CinematicViewport/SCinematicLevelViewport.cpp
@@ -537,8 +537,8 @@ float SCinematicLevelViewport::GetPlayTimeMinDesiredWidth() const
 	{
 		TRange<double> ViewRange = Sequencer->GetViewRange();
 
-		FString LowerBoundStr = Sequencer->GetNumericTypeInterface()->ToString(ViewRange.GetLowerBoundValue());
-		FString UpperBoundStr = Sequencer->GetNumericTypeInterface()->ToString(ViewRange.GetUpperBoundValue());
+		FString LowerBoundStr = Sequencer->GetNumericPositionTypeInterface()->ToString(ViewRange.GetLowerBoundValue());
+		FString UpperBoundStr = Sequencer->GetNumericPositionTypeInterface()->ToString(ViewRange.GetUpperBoundValue());
 
 		// Always measure with the negative and subframe indicator so that the size doesn't change when there is and isn't a subframe
 		if (!LowerBoundStr.Contains(TEXT("*")))
@@ -634,7 +634,7 @@ void SCinematicLevelViewport::Setup(FLevelSequenceEditorToolkit& NewToolkit)
 	ISequencer* Sequencer = GetSequencer();
 	if (Sequencer)
 	{
-		TypeInterfaceProxy->Impl = Sequencer->GetNumericTypeInterface();
+		TypeInterfaceProxy->Impl = Sequencer->GetNumericPositionTypeInterface();
 
 		if (TransportRange.IsValid())
 		{
@@ -734,7 +734,7 @@ void SCinematicLevelViewport::Tick(const FGeometry& AllottedGeometry, const doub
 
 	FText TimeFormat = LOCTEXT("TimeFormat", "{0}");
 
-	TSharedPtr<INumericTypeInterface<double>> TimeDisplayFormatInterface = Sequencer->GetNumericTypeInterface();
+	TSharedPtr<INumericTypeInterface<double>> TimeDisplayFormatInterface = Sequencer->GetNumericPositionTypeInterface();
 
 	UMovieSceneSequence* SubSequence = CinematicShotSection ? CinematicShotSection->GetSequence() : nullptr;
 	if (SubSequence)
diff --git a/Engine/Plugins/MovieScene/LevelSequenceEditor/Source/LevelSequenceEditor/Private/LevelSequenceEditorSubsystem.cpp b/Engine/Plugins/MovieScene/LevelSequenceEditor/Source/LevelSequenceEditor/Private/LevelSequenceEditorSubsystem.cpp
--- a/Engine/Plugins/MovieScene/LevelSequenceEditor/Source/LevelSequenceEditor/Private/LevelSequenceEditorSubsystem.cpp
+++ b/Engine/Plugins/MovieScene/LevelSequenceEditor/Source/LevelSequenceEditor/Private/LevelSequenceEditorSubsystem.cpp
@@ -1473,7 +1473,7 @@ void SBakeTransformWidget::Construct(const FArguments& InArgs)
 	FPropertyEditorModule& PropertyEditor = FModuleManager::Get().LoadModuleChecked<FPropertyEditorModule>(TEXT("PropertyEditor"));
 
 	DetailsView = PropertyEditor.CreateStructureDetailView(ViewArgs, StructureViewArgs, TSharedPtr<FStructOnScope>());
-	TSharedPtr<INumericTypeInterface<double>> NumericTypeInterface = Sequencer->GetNumericTypeInterface();
+	TSharedPtr<INumericTypeInterface<double>> NumericTypeInterface = Sequencer->GetNumericPositionTypeInterface();
 	DetailsView->GetDetailsView()->RegisterInstancedCustomPropertyTypeLayout("FrameNumber",
 		FOnGetPropertyTypeCustomizationInstance::CreateLambda([=]() {return MakeShared<FFrameNumberDetailsCustomization>(NumericTypeInterface); }));
 	DetailsView->SetStructureData(Settings);
diff --git a/Engine/Plugins/MovieScene/MovieSceneTextTrack/Source/MovieSceneTextTrack/Private/MovieSceneTextChannel.cpp b/Engine/Plugins/MovieScene/MovieSceneTextTrack/Source/MovieSceneTextTrack/Private/MovieSceneTextChannel.cpp
--- a/Engine/Plugins/MovieScene/MovieSceneTextTrack/Source/MovieSceneTextTrack/Private/MovieSceneTextChannel.cpp
+++ b/Engine/Plugins/MovieScene/MovieSceneTextTrack/Source/MovieSceneTextTrack/Private/MovieSceneTextChannel.cpp
@@ -74,6 +74,11 @@ void FMovieSceneTextChannel::ChangeFrameResolution(FFrameRate SourceRate, FFrame
 	GetData().ChangeFrameResolution(SourceRate, DestinationRate);
 }
 
+void FMovieSceneTextChannel::RemapMusicalFrameTimes(const UFrameBasedMusicMap* OldMap, const UFrameBasedMusicMap* NewMap)
+{
+	GetData().RemapMusicalFrameTimes(OldMap, NewMap);
+}
+
 TRange<FFrameNumber> FMovieSceneTextChannel::ComputeEffectiveRange() const
 {
 	return GetData().GetTotalRange();
diff --git a/Engine/Plugins/MovieScene/MovieSceneTextTrack/Source/MovieSceneTextTrack/Public/MovieSceneTextChannel.h b/Engine/Plugins/MovieScene/MovieSceneTextTrack/Source/MovieSceneTextTrack/Public/MovieSceneTextChannel.h
--- a/Engine/Plugins/MovieScene/MovieSceneTextTrack/Source/MovieSceneTextTrack/Public/MovieSceneTextChannel.h
+++ b/Engine/Plugins/MovieScene/MovieSceneTextTrack/Source/MovieSceneTextTrack/Public/MovieSceneTextChannel.h
@@ -108,6 +108,7 @@ struct MOVIESCENETEXTTRACK_API FMovieSceneTextChannel : public FMovieSceneChanne
 	virtual void DeleteKeys(TArrayView<const FKeyHandle> InHandles) override;
 	virtual void DeleteKeysFrom(FFrameNumber InTime, bool bDeleteKeysBefore) override;
 	virtual void ChangeFrameResolution(FFrameRate SourceRate, FFrameRate DestinationRate) override;
+	virtual void RemapMusicalFrameTimes(const UFrameBasedMusicMap* OldMap, const UFrameBasedMusicMap* NewMap) override;
 	virtual TRange<FFrameNumber> ComputeEffectiveRange() const override;
 	virtual int32 GetNumKeys() const override;
 	virtual void Reset() override;
diff --git a/Engine/Plugins/Runtime/GameplayAbilities/Source/GameplayAbilities/Private/Sequencer/MovieSceneGameplayCueSections.cpp b/Engine/Plugins/Runtime/GameplayAbilities/Source/GameplayAbilities/Private/Sequencer/MovieSceneGameplayCueSections.cpp
--- a/Engine/Plugins/Runtime/GameplayAbilities/Source/GameplayAbilities/Private/Sequencer/MovieSceneGameplayCueSections.cpp
+++ b/Engine/Plugins/Runtime/GameplayAbilities/Source/GameplayAbilities/Private/Sequencer/MovieSceneGameplayCueSections.cpp
@@ -211,6 +211,11 @@ void FMovieSceneGameplayCueChannel::ChangeFrameResolution(FFrameRate SourceRate,
 	GetData().ChangeFrameResolution(SourceRate, DestinationRate);
 }
 
+void FMovieSceneGameplayCueChannel::RemapMusicalFrameTimes(const UFrameBasedMusicMap* OldMap, const UFrameBasedMusicMap* NewMap)
+{
+	GetData().RemapMusicalFrameTimes(OldMap, NewMap);
+}
+
 TRange<FFrameNumber> FMovieSceneGameplayCueChannel::ComputeEffectiveRange() const
 {
 	return GetData().GetTotalRange();
diff --git a/Engine/Plugins/Runtime/GameplayAbilities/Source/GameplayAbilities/Public/Sequencer/MovieSceneGameplayCueSections.h b/Engine/Plugins/Runtime/GameplayAbilities/Source/GameplayAbilities/Public/Sequencer/MovieSceneGameplayCueSections.h
--- a/Engine/Plugins/Runtime/GameplayAbilities/Source/GameplayAbilities/Public/Sequencer/MovieSceneGameplayCueSections.h
+++ b/Engine/Plugins/Runtime/GameplayAbilities/Source/GameplayAbilities/Public/Sequencer/MovieSceneGameplayCueSections.h
@@ -110,6 +110,7 @@ public:
 	virtual void DeleteKeys(TArrayView<const FKeyHandle> InHandles) override;
 	virtual void DeleteKeysFrom(FFrameNumber InTime, bool bDeleteKeysBefore) override;
 	virtual void ChangeFrameResolution(FFrameRate SourceRate, FFrameRate DestinationRate) override;
+	virtual void RemapMusicalFrameTimes(const UFrameBasedMusicMap* OldMap, const UFrameBasedMusicMap* NewMap) override;
 	virtual TRange<FFrameNumber> ComputeEffectiveRange() const override;
 	virtual int32 GetNumKeys() const override;
 	virtual void Reset() override;
diff --git a/Engine/Plugins/Runtime/Harmonix/Source/HarmonixMetasound/HarmonixMetasound.Build.cs b/Engine/Plugins/Runtime/Harmonix/Source/HarmonixMetasound/HarmonixMetasound.Build.cs
--- a/Engine/Plugins/Runtime/Harmonix/Source/HarmonixMetasound/HarmonixMetasound.Build.cs
+++ b/Engine/Plugins/Runtime/Harmonix/Source/HarmonixMetasound/HarmonixMetasound.Build.cs
@@ -20,6 +20,7 @@ public class HarmonixMetasound : ModuleRules
 				"Engine",
 				"Core",
 				"CoreUObject",
+				"MovieScene",
 				"MetasoundStandardNodes",
 				"HarmonixDsp",
 				"HarmonixMidi",
diff --git a/Engine/Plugins/Runtime/Harmonix/Source/HarmonixMetasound/Private/HarmonixMetasound/Components/MusicClockComponent.cpp b/Engine/Plugins/Runtime/Harmonix/Source/HarmonixMetasound/Private/HarmonixMetasound/Components/MusicClockComponent.cpp
--- a/Engine/Plugins/Runtime/Harmonix/Source/HarmonixMetasound/Private/HarmonixMetasound/Components/MusicClockComponent.cpp
+++ b/Engine/Plugins/Runtime/Harmonix/Source/HarmonixMetasound/Private/HarmonixMetasound/Components/MusicClockComponent.cpp
@@ -769,6 +769,36 @@ void UMusicClockComponent::MakeDefaultSongMap()
 	DefaultMaps.GetBarMap().AddTimeSignatureAtBarIncludingCountIn(0, TimeSignatureNum, TimeSignatureDenom);
 }
 
+void UMusicClockComponent::OnTick(float DeltaSeconds, float InPlayRate)
+{
+
+}
+
+void UMusicClockComponent::OnStartPlaying(const FQualifiedFrameTime& InStartTime)
+{
+
+}
+
+void UMusicClockComponent::OnStopPlaying(const FQualifiedFrameTime& InStopTime)
+{
+
+}
+
+FFrameTime UMusicClockComponent::OnRequestCurrentTime(const FQualifiedFrameTime& InCurrentTime, float InPlayRate)
+{
+	return InCurrentTime.Rate.AsFrameTime(CurrentVideoRenderSongPos.SecondsFromBarOne);
+}
+
+FMusicalTime UMusicClockComponent::OnRequestCurrentMusicalTime(const FQualifiedFrameTime& InCurrentTime, float InPlayRate)
+{
+	FMusicalTime Result;
+	Result.Bar = CurrentVideoRenderSongPos.Timestamp.Bar - 1;
+	Result.TicksPerBeat = (MusicalTime::TicksPerQuarterNote * 4 / CurrentVideoRenderSongPos.TimeSigDenominator);
+	Result.TicksPerBar = Result.TicksPerBeat * CurrentVideoRenderSongPos.TimeSigNumerator;
+	Result.TickInBar = static_cast<int32>((CurrentVideoRenderSongPos.Timestamp.Beat - 1.0f) * static_cast<float>(Result.TicksPerBeat));
+	return Result;
+}
+
 void FMusicClockDriverBase::EnsureClockIsValidForGameFrame()
 {
 	// Here the ClockDriver reaches back up and mutates its owning UMusicClockComponent
diff --git a/Engine/Plugins/Runtime/Harmonix/Source/HarmonixMetasound/Public/HarmonixMetasound/Components/MusicClockComponent.h b/Engine/Plugins/Runtime/Harmonix/Source/HarmonixMetasound/Public/HarmonixMetasound/Components/MusicClockComponent.h
--- a/Engine/Plugins/Runtime/Harmonix/Source/HarmonixMetasound/Public/HarmonixMetasound/Components/MusicClockComponent.h
+++ b/Engine/Plugins/Runtime/Harmonix/Source/HarmonixMetasound/Public/HarmonixMetasound/Components/MusicClockComponent.h
@@ -7,6 +7,7 @@
 #include "HarmonixMidi/SongMaps.h"
 #include "Delegates/DelegateCombinations.h"
 #include "Templates/UniquePtr.h"
+#include "Evaluation/IMovieSceneCustomClockSource.h"
 
 #include "MusicClockComponent.generated.h"
 
@@ -49,7 +50,7 @@ DECLARE_DYNAMIC_MULTICAST_DELEGATE(FMusicClockConnected);
 DECLARE_DYNAMIC_MULTICAST_DELEGATE(FMusicClockDisconnected);
 
 UCLASS(ClassGroup = (MetaSoundMusic), PrioritizeCategories = "MusicClock", meta = (BlueprintSpawnableComponent, DisplayName = "Music Clock", ScriptName = MusicClockComponent))
-class HARMONIXMETASOUND_API UMusicClockComponent : public UActorComponent
+class HARMONIXMETASOUND_API UMusicClockComponent : public UActorComponent, public IMovieSceneCustomClockSource
 {
 	GENERATED_BODY()
 
@@ -367,6 +368,16 @@ public:
 
 	const FMidiSongPos& GetRawUnsmoothedAudioRenderPos() const { return CurrentRawAudioRenderSongPos; }
 
+	// vvv BEGIN IMovieSceneCustomClockSource Implementation vvv
+	virtual void OnTick(float DeltaSeconds, float InPlayRate) override;
+	virtual void OnStartPlaying(const FQualifiedFrameTime& InStartTime) override;
+	virtual void OnStopPlaying(const FQualifiedFrameTime& InStopTime) override;
+	virtual FFrameTime OnRequestCurrentTime(const FQualifiedFrameTime& InCurrentTime, float InPlayRate) override;
+	virtual bool OnSupportsMusicalTime() override { return true; }
+	virtual FMusicalTime OnRequestCurrentMusicalTime(const FQualifiedFrameTime& InCurrentTime, float InPlayRate) override;
+	// ^^^ END IMovieSceneCustomClockSource Implementation ^^^
+
+
 protected:
 	virtual void TickComponent(float DeltaTime, enum ELevelTick TickType, FActorComponentTickFunction* ThisTickFunction);
 
diff --git a/Engine/Plugins/Runtime/Harmonix/Source/HarmonixMidi/HarmonixMidi.Build.cs b/Engine/Plugins/Runtime/Harmonix/Source/HarmonixMidi/HarmonixMidi.Build.cs
--- a/Engine/Plugins/Runtime/Harmonix/Source/HarmonixMidi/HarmonixMidi.Build.cs
+++ b/Engine/Plugins/Runtime/Harmonix/Source/HarmonixMidi/HarmonixMidi.Build.cs
@@ -15,6 +15,7 @@ public class HarmonixMidi : ModuleRules
 			{
 				"Core",
 				"CoreUObject",
+				"TimeManagement",
 			});
 		PrivateDependencyModuleNames.AddRange(
 			new string[]
diff --git a/Engine/Plugins/Runtime/Harmonix/Source/HarmonixMidi/Private/HarmonixMidi/MidiFile.cpp b/Engine/Plugins/Runtime/Harmonix/Source/HarmonixMidi/Private/HarmonixMidi/MidiFile.cpp
--- a/Engine/Plugins/Runtime/Harmonix/Source/HarmonixMidi/Private/HarmonixMidi/MidiFile.cpp
+++ b/Engine/Plugins/Runtime/Harmonix/Source/HarmonixMidi/Private/HarmonixMidi/MidiFile.cpp
@@ -232,6 +232,37 @@ void UMidiFile::SaveStdMidiFile(TSharedPtr<FArchive> Archive, const FString& Fil
 	Algo::ForEach(TheMidiData.Tracks, [&](const FMidiTrack& track) { track.WriteStdMidi(Writer); });
 }
 
+void UMidiFile::CreateFrameBasedMusicMap(UFrameBasedMusicMap* Map) const
+{
+	return TheMidiData.SongMaps.FillInFrameBasedMusicMap(Map);
+}
+
+FMarkerProviderResults UMidiFile::GatherMarkers(const UFrameBasedMusicMap* Map) const
+{
+	FMarkerProviderResults Results;
+
+	if (TheMidiData.Tracks.Num() > 0)
+	{
+		const FMidiTrack& ConductorTrack = TheMidiData.Tracks[0];
+		const FMidiEventList& Events = ConductorTrack.GetEvents();
+		for (const FMidiEvent& Event : Events)
+		{
+			const FMidiMsg& Msg = Event.GetMsg();
+			if (Msg.IsText() && (Msg.GetTextType() == Harmonix::Midi::Constants::GMeta_CuePoint || Msg.GetTextType() == Harmonix::Midi::Constants::GMeta_Marker))
+			{
+				if (Results.Channels.IsEmpty())
+				{
+					Results.Channels.Emplace(TEXT("Conductor"));
+				}
+				float Ms = TheMidiData.SongMaps.TickToMs(static_cast<float>(Event.GetTick()));
+				FFrameNumber Location = Map->FrameResolution.AsFrameNumber(static_cast<double>(Ms) / 1000.0);;
+				Results.Channels[0].Markers.Emplace(ConductorTrack.GetTextForMsg(Msg), Location);
+			}
+		}
+	}
+	return Results;
+}
+
 FMidiTrack* UMidiFile::AddTrack(const FString& Name)
 {
 	TheMidiData.Tracks.Emplace(Name);
diff --git a/Engine/Plugins/Runtime/Harmonix/Source/HarmonixMidi/Private/HarmonixMidi/SongMaps.cpp b/Engine/Plugins/Runtime/Harmonix/Source/HarmonixMidi/Private/HarmonixMidi/SongMaps.cpp
--- a/Engine/Plugins/Runtime/Harmonix/Source/HarmonixMidi/Private/HarmonixMidi/SongMaps.cpp
+++ b/Engine/Plugins/Runtime/Harmonix/Source/HarmonixMidi/Private/HarmonixMidi/SongMaps.cpp
@@ -117,6 +117,35 @@ bool FSongMaps::LoadFromStdMidiFile(TSharedPtr<FArchive> Archive, const FString&
 	return ReadWithReader(Reader);
 }
 
+void FSongMaps::FillInFrameBasedMusicMap(UFrameBasedMusicMap* Map) const
+{
+	Map->Clear();
+
+	auto TempoMapIterator = TempoMap.GetTempoPoints().begin();
+	auto BarMapIterator = BarMap.GetTimeSignaturePoints().begin();
+
+	check(TempoMapIterator != TempoMap.GetTempoPoints().end());
+	check(BarMapIterator != BarMap.GetTimeSignaturePoints().end());
+
+	int32 NextTempoPointTick = (*TempoMapIterator).StartTick;
+	int32 NextBarPointTick = (*BarMapIterator).StartTick;
+
+	while (NextTempoPointTick != std::numeric_limits<int32>::max() ||
+		NextBarPointTick != std::numeric_limits<int32>::max())
+	{
+		if (NextBarPointTick <= NextTempoPointTick)
+		{
+			Map->AddTimeSignature(NextBarPointTick, (*BarMapIterator).BarIndex, (*BarMapIterator).TimeSignature.Numerator, (*BarMapIterator).TimeSignature.Denominator);
+			++BarMapIterator;
+			NextBarPointTick = (BarMapIterator != BarMap.GetTimeSignaturePoints().end()) ? (*BarMapIterator).StartTick : std::numeric_limits<int32>::max();
+			continue;
+		}
+		Map->AddTempo(NextTempoPointTick, TempoMap.TickToMs(NextTempoPointTick),  (*TempoMapIterator).GetBPM());
+		++TempoMapIterator;
+		NextTempoPointTick = (TempoMapIterator != TempoMap.GetTempoPoints().end()) ? (*TempoMapIterator).StartTick : std::numeric_limits<int32>::max();
+	}
+}
+
 bool FSongMaps::FinalizeRead(IMidiReader* Reader)
 {
 	int32 LastTick = Reader->GetLastTick();
diff --git a/Engine/Plugins/Runtime/Harmonix/Source/HarmonixMidi/Public/HarmonixMidi/MidiFile.h b/Engine/Plugins/Runtime/Harmonix/Source/HarmonixMidi/Public/HarmonixMidi/MidiFile.h
--- a/Engine/Plugins/Runtime/Harmonix/Source/HarmonixMidi/Public/HarmonixMidi/MidiFile.h
+++ b/Engine/Plugins/Runtime/Harmonix/Source/HarmonixMidi/Public/HarmonixMidi/MidiFile.h
@@ -6,7 +6,7 @@
 #include "HarmonixMidi/SongMaps.h"
 #include "HarmonixMidi/MidiTrack.h"
 #include "IAudioProxyInitializer.h"
-
+#include "FrameBasedMusicMap.h"
 
 #include "MidiFile.generated.h"
 
@@ -113,7 +113,7 @@ struct TStructOpsTypeTraits<FMidiFileData> : public TStructOpsTypeTraitsBase2<FM
  * as serializing itself to standard Unreal Engine FArchives.
  */
 UCLASS(BlueprintType, Category="Music", Meta = (DisplayName = "Standard MIDI File"))
-class HARMONIXMIDI_API UMidiFile : public UObject, public IAudioProxyDataFactory
+class HARMONIXMIDI_API UMidiFile : public UObject, public IAudioProxyDataFactory, public IMusicMapSource
 {
 	GENERATED_BODY()
 
@@ -163,6 +163,11 @@ public:
 	/** A method for exporting the midi track data to a standard midi file. */
 	void SaveStdMidiFile(TSharedPtr<FArchive> Archive, const FString& Filename = FString()) const;
 
+	//** BEGIN IMusicMapSource
+	virtual void CreateFrameBasedMusicMap(UFrameBasedMusicMap* Map) const override;
+	virtual FMarkerProviderResults GatherMarkers(const UFrameBasedMusicMap* Map) const override;
+	//** END IMusicMapSource
+
 	FMidiTrack* AddTrack(const FString& Name);
 
 	void Empty();
diff --git a/Engine/Plugins/Runtime/Harmonix/Source/HarmonixMidi/Public/HarmonixMidi/SongMaps.h b/Engine/Plugins/Runtime/Harmonix/Source/HarmonixMidi/Public/HarmonixMidi/SongMaps.h
--- a/Engine/Plugins/Runtime/Harmonix/Source/HarmonixMidi/Public/HarmonixMidi/SongMaps.h
+++ b/Engine/Plugins/Runtime/Harmonix/Source/HarmonixMidi/Public/HarmonixMidi/SongMaps.h
@@ -6,6 +6,7 @@
 #include "HarmonixMidi/ChordMap.h"
 #include "HarmonixMidi/SectionMap.h"
 #include "Sound/QuartzQuantizationUtilities.h"
+#include "FrameBasedMusicMap.h"
 #include <limits>
 
 #include "SongMaps.generated.h"
@@ -220,6 +221,10 @@ public:
 	bool LoadFromStdMidiFile(void* Buffer, int32 BufferSize, const FString& Filename);
 	bool LoadFromStdMidiFile(TSharedPtr<FArchive> Archive, const FString& Filename);
 
+	//** BEGIN Support for IMusicMapSource
+	void FillInFrameBasedMusicMap(UFrameBasedMusicMap* Map) const;
+	//** END Support for IMusicMapSource
+
 	// tracks
 	TArray<FString>& GetTrackNames() { return TrackNames; }
 	const TArray<FString>& GetTrackNames() const { return TrackNames; }
diff --git a/Engine/Source/Editor/MovieSceneTools/Private/Channels/BuiltInChannelEditors.cpp b/Engine/Source/Editor/MovieSceneTools/Private/Channels/BuiltInChannelEditors.cpp
--- a/Engine/Source/Editor/MovieSceneTools/Private/Channels/BuiltInChannelEditors.cpp
+++ b/Engine/Source/Editor/MovieSceneTools/Private/Channels/BuiltInChannelEditors.cpp
@@ -529,7 +529,7 @@ TSharedRef<SWidget> CreateKeyEditor(const TMovieSceneChannelHandle<FMovieSceneTi
 	if (ChannelPtr && ChannelPtr->Domain == UE::MovieScene::ETimeWarpChannelDomain::Time)
 	{
 		// Set the numeric type interface for frame numbers if the channel is in the time domain
-		KeyEditor.SetNumericTypeInterface(Params.Sequencer->GetNumericTypeInterface());
+		KeyEditor.SetNumericTypeInterface(Params.Sequencer->GetNumericPositionTypeInterface());
 	}
 
 	typedef SNumericKeyEditor<FMovieSceneTimeWarpChannel, double> KeyEditorType;
diff --git a/Engine/Source/Editor/MovieSceneTools/Private/FrameNumberDetailsCustomization.cpp b/Engine/Source/Editor/MovieSceneTools/Private/FrameNumberDetailsCustomization.cpp
--- a/Engine/Source/Editor/MovieSceneTools/Private/FrameNumberDetailsCustomization.cpp
+++ b/Engine/Source/Editor/MovieSceneTools/Private/FrameNumberDetailsCustomization.cpp
@@ -17,6 +17,7 @@ void FFrameNumberDetailsCustomization::CustomizeChildren(TSharedRef<IPropertyHan
 	// Check for Min/Max metadata on the property itself and we'll apply it to the child when changed.
 	const FString& MetaUIMinString = PropertyHandle->GetMetaData(TEXT("UIMin"));
 	const FString& MetaUIMaxString = PropertyHandle->GetMetaData(TEXT("UIMax"));
+	const FString& MetaUIDisplayAsString = PropertyHandle->GetMetaData(TEXT("UIFrameDisplayAs"));
 
 	UIClampMin = TNumericLimits<int32>::Lowest();
 	UIClampMax = TNumericLimits<int32>::Max();
@@ -31,6 +32,22 @@ void FFrameNumberDetailsCustomization::CustomizeChildren(TSharedRef<IPropertyHan
 		TTypeFromString<int32>::FromString(UIClampMax, *MetaUIMaxString);
 	}
 
+	if (!MetaUIDisplayAsString.IsEmpty())
+	{
+		if (MetaUIDisplayAsString.Equals(TEXT("Duration"), ESearchCase::IgnoreCase) && NumericDurationTypeInterface.IsValid())
+		{
+			NumericTypeInterface = NumericDurationTypeInterface;
+		}
+		else if (MetaUIDisplayAsString.Equals(TEXT("NonMusical"), ESearchCase::IgnoreCase) && NumericNonMusicalTypeInterface.IsValid())
+		{
+			NumericTypeInterface = NumericNonMusicalTypeInterface;
+		}
+		else
+		{
+			NumericTypeInterface = NumericPositionTypeInterface;
+		}
+	}
+
 	TMap<FName, TSharedRef<IPropertyHandle>> CustomizedProperties;
 
 	uint32 NumChildren = 0;
diff --git a/Engine/Source/Editor/MovieSceneTools/Private/MovieSceneSectionDetailsCustomization.cpp b/Engine/Source/Editor/MovieSceneTools/Private/MovieSceneSectionDetailsCustomization.cpp
--- a/Engine/Source/Editor/MovieSceneTools/Private/MovieSceneSectionDetailsCustomization.cpp
+++ b/Engine/Source/Editor/MovieSceneTools/Private/MovieSceneSectionDetailsCustomization.cpp
@@ -205,7 +205,7 @@ FText FMovieSceneSectionDetailsCustomization::OnGetRangeStartText() const
 			return NSLOCTEXT("PropertyEditor", "MultipleValues", "Multiple Values");
 		case ERangeBoundValueType::Finite:
 		default:
-			return FText::FromString(NumericTypeInterface->ToString((double)FrameValue.Value));
+			return FText::FromString(NumericPositionTypeInterface->ToString((double)FrameValue.Value));
 	}
 }
 
@@ -229,7 +229,7 @@ FText FMovieSceneSectionDetailsCustomization::OnGetRangeStartToolTipText() const
 void FMovieSceneSectionDetailsCustomization::OnRangeStartTextCommitted(const FText& InText, ETextCommit::Type CommitInfo)
 {
 	// Find the new value for the start range.
-	TOptional<double> NewStart = NumericTypeInterface->FromString(InText.ToString(), 0.0);
+	TOptional<double> NewStart = NumericPositionTypeInterface->FromString(InText.ToString(), 0.0);
 
 	// Early out if we couldn't parse it, no need to reset them all to zero.
 	if (!NewStart.IsSet())
@@ -437,7 +437,7 @@ FText FMovieSceneSectionDetailsCustomization::OnGetRangeEndText() const
 			return NSLOCTEXT("PropertyEditor", "MultipleValues", "Multiple Values");
 		case ERangeBoundValueType::Finite:
 		default:
-			return FText::FromString(NumericTypeInterface->ToString((double)FrameValue.Value));
+			return FText::FromString(NumericPositionTypeInterface->ToString((double)FrameValue.Value));
 	}
 }
 
@@ -461,7 +461,7 @@ FText FMovieSceneSectionDetailsCustomization::OnGetRangeEndToolTipText() const
 void FMovieSceneSectionDetailsCustomization::OnRangeEndTextCommitted(const FText& InText, ETextCommit::Type CommitInfo)
 {
 	// Find the new value for the end range.
-	TOptional<double> NewEnd = NumericTypeInterface->FromString(InText.ToString(), 0.0);
+	TOptional<double> NewEnd = NumericPositionTypeInterface->FromString(InText.ToString(), 0.0);
 
 	// Early out if we couldn't parse it, no need to reset them all to zero.
 	if (!NewEnd.IsSet())
diff --git a/Engine/Source/Editor/MovieSceneTools/Private/MovieSceneToolHelpers.cpp b/Engine/Source/Editor/MovieSceneTools/Private/MovieSceneToolHelpers.cpp
--- a/Engine/Source/Editor/MovieSceneTools/Private/MovieSceneToolHelpers.cpp
+++ b/Engine/Source/Editor/MovieSceneTools/Private/MovieSceneToolHelpers.cpp
@@ -1962,8 +1962,8 @@ class SControlRigImportFBXSettings : public SCompoundWidget
 
 		Sequencer = InSequencer;
 
-		TSharedPtr<INumericTypeInterface<double>> NumericTypeInterface = (InSequencer->GetNumericTypeInterface());
-		DetailView->RegisterInstancedCustomPropertyTypeLayout("FrameNumber", FOnGetPropertyTypeCustomizationInstance::CreateStatic(&FFrameNumberDetailsCustomization::MakeInstance, NumericTypeInterface));
+		TSharedPtr<INumericTypeInterface<double>> NumericPositionTypeInterface = (InSequencer->GetNumericPositionTypeInterface());
+		DetailView->RegisterInstancedCustomPropertyTypeLayout("FrameNumber", FOnGetPropertyTypeCustomizationInstance::CreateStatic(&FFrameNumberDetailsCustomization::MakeInstance, NumericPositionTypeInterface));
 
 
 		ChildSlot
diff --git a/Engine/Source/Editor/MovieSceneTools/Private/SKeyEditInterface.cpp b/Engine/Source/Editor/MovieSceneTools/Private/SKeyEditInterface.cpp
--- a/Engine/Source/Editor/MovieSceneTools/Private/SKeyEditInterface.cpp
+++ b/Engine/Source/Editor/MovieSceneTools/Private/SKeyEditInterface.cpp
@@ -113,7 +113,7 @@ TSharedRef<IPropertyTypeCustomization> SKeyEditInterface::CreateFrameNumberCusto
 {
 	if (const TSharedPtr<ISequencer> SequencerPtr = WeakSequencer.Pin())
 	{
-		return MakeShared<FFrameNumberDetailsCustomization>(SequencerPtr->GetNumericTypeInterface());
+		return MakeShared<FFrameNumberDetailsCustomization>(SequencerPtr->GetNumericPositionTypeInterface(), SequencerPtr->GetNumericDurationTypeInterface(), SequencerPtr->GetNumericNonMusicalTypeInterface());
 	}
 	return MakeShared<FFrameNumberDetailsCustomization>(nullptr);
 }
diff --git a/Engine/Source/Editor/MovieSceneTools/Private/Sections/ThumbnailSection.cpp b/Engine/Source/Editor/MovieSceneTools/Private/Sections/ThumbnailSection.cpp
--- a/Engine/Source/Editor/MovieSceneTools/Private/Sections/ThumbnailSection.cpp
+++ b/Engine/Source/Editor/MovieSceneTools/Private/Sections/ThumbnailSection.cpp
@@ -136,7 +136,7 @@ void FThumbnailSection::BuildThumbnailsMenu(FMenuBuilder& InMenuBuilder)
 		return;
 	}
 
-	const FText CurrentTime = FText::FromString(Sequencer->GetNumericTypeInterface()->ToString(Sequencer->GetLocalTime().Time.GetFrame().Value));
+	const FText CurrentTime = FText::FromString(Sequencer->GetNumericPositionTypeInterface()->ToString(Sequencer->GetLocalTime().Time.GetFrame().Value));
 
 	InMenuBuilder.BeginSection(TEXT("Thumbnails"), LOCTEXT("ThumbnailsMenuSection", "Thumbnails"));
 	{
diff --git a/Engine/Source/Editor/MovieSceneTools/Public/FrameNumberDetailsCustomization.h b/Engine/Source/Editor/MovieSceneTools/Public/FrameNumberDetailsCustomization.h
--- a/Engine/Source/Editor/MovieSceneTools/Public/FrameNumberDetailsCustomization.h
+++ b/Engine/Source/Editor/MovieSceneTools/Public/FrameNumberDetailsCustomization.h
@@ -14,14 +14,24 @@ class IDetailLayoutBuilder;
 class MOVIESCENETOOLS_API FFrameNumberDetailsCustomization : public IPropertyTypeCustomization
 {
 public:
-	static TSharedRef<IPropertyTypeCustomization> MakeInstance(TSharedPtr<INumericTypeInterface<double>> InNumericTypeInterface)
+	static TSharedRef<IPropertyTypeCustomization> MakeInstance(TSharedPtr<INumericTypeInterface<double>> InNumericPositionTypeInterface)
 	{
-		return MakeShared<FFrameNumberDetailsCustomization>(InNumericTypeInterface);
+		return MakeShared<FFrameNumberDetailsCustomization>(InNumericPositionTypeInterface/*, nullptr, nullptr*/);
 	}
 
-	FFrameNumberDetailsCustomization(TSharedPtr<INumericTypeInterface<double>> InNumericTypeInterface)
+	static TSharedRef<IPropertyTypeCustomization> MakeInstance(TSharedPtr<INumericTypeInterface<double>> InNumericPositionTypeInterface,
+		TSharedPtr<INumericTypeInterface<double>> InNumericDurationTypeInterface, TSharedPtr<INumericTypeInterface<double>> InNumericNonMusicalTypeInterface)
 	{
-		NumericTypeInterface = InNumericTypeInterface;
+		return MakeShared<FFrameNumberDetailsCustomization>(InNumericPositionTypeInterface, InNumericDurationTypeInterface, InNumericNonMusicalTypeInterface);
+	}
+
+	FFrameNumberDetailsCustomization(TSharedPtr<INumericTypeInterface<double>> InNumericPositionTypeInterface, TSharedPtr<INumericTypeInterface<double>> InNumericDurationTypeInterface = nullptr, TSharedPtr<INumericTypeInterface<double>> InNumericNonMusicalTypeInterface = nullptr)
+	{
+		NumericPositionTypeInterface = InNumericPositionTypeInterface;
+		NumericDurationTypeInterface = InNumericDurationTypeInterface;
+		NumericNonMusicalTypeInterface = InNumericNonMusicalTypeInterface;
+		// default to position in the current display format...
+		NumericTypeInterface = NumericPositionTypeInterface.IsValid() ? NumericPositionTypeInterface : NumericNonMusicalTypeInterface.IsValid() ? NumericNonMusicalTypeInterface : NumericDurationTypeInterface;
 	}
 
 	/** IPropertyTypeCustomization interface */
@@ -36,6 +46,15 @@ private:
 	/** The Numeric Type interface used to convert between display formats and internal tick resolution. */
 	TSharedPtr<INumericTypeInterface<double>> NumericTypeInterface;
 
+	/** The Numeric Type interface used to convert between display formats and internal tick resolution. */
+	TSharedPtr<INumericTypeInterface<double>> NumericPositionTypeInterface;
+
+	/** The Numeric Type interface used to convert between display formats and internal tick resolution. */
+	TSharedPtr<INumericTypeInterface<double>> NumericDurationTypeInterface;
+
+	/** The Numeric Type interface used to convert between display formats and internal tick resolution. */
+	TSharedPtr<INumericTypeInterface<double>> NumericNonMusicalTypeInterface;
+
 	/** Store the property handle to the FrameNumber field so we can get/set the value on the object via text box callbacks. */
 	TSharedPtr<IPropertyHandle> FrameNumberProperty;
 
diff --git a/Engine/Source/Editor/MovieSceneTools/Public/MovieSceneSectionDetailsCustomization.h b/Engine/Source/Editor/MovieSceneTools/Public/MovieSceneSectionDetailsCustomization.h
--- a/Engine/Source/Editor/MovieSceneTools/Public/MovieSceneSectionDetailsCustomization.h
+++ b/Engine/Source/Editor/MovieSceneTools/Public/MovieSceneSectionDetailsCustomization.h
@@ -17,14 +17,16 @@ class UMovieScene;
 class MOVIESCENETOOLS_API FMovieSceneSectionDetailsCustomization : public IDetailCustomization
 {
 public:
-	static TSharedRef<IDetailCustomization> MakeInstance(TSharedPtr<INumericTypeInterface<double>> InNumericTypeInterface, TWeakObjectPtr<UMovieScene> InParentMovieScene)
+	static TSharedRef<IDetailCustomization> MakeInstance(TSharedPtr<INumericTypeInterface<double>> InNumericPositionTypeInterface, TSharedPtr<INumericTypeInterface<double>> InNumericDurationTypeInterface, TSharedPtr<INumericTypeInterface<double>> InNumericNonMusicalTypeInterface, TWeakObjectPtr<UMovieScene> InParentMovieScene)
 	{
-		return MakeShared<FMovieSceneSectionDetailsCustomization>(InNumericTypeInterface, InParentMovieScene);
+		return MakeShared<FMovieSceneSectionDetailsCustomization>(InNumericPositionTypeInterface, InNumericDurationTypeInterface, InNumericNonMusicalTypeInterface, InParentMovieScene);
 	}
 
-	FMovieSceneSectionDetailsCustomization(TSharedPtr<INumericTypeInterface<double>> InNumericTypeInterface, TWeakObjectPtr<UMovieScene> InParentMovieScene)
+	FMovieSceneSectionDetailsCustomization(TSharedPtr<INumericTypeInterface<double>> InNumericPositionTypeInterface, TSharedPtr<INumericTypeInterface<double>> InNumericDurationTypeInterface, TSharedPtr<INumericTypeInterface<double>> InNumericNonMusicalTypeInterface, TWeakObjectPtr<UMovieScene> InParentMovieScene)
 	{
-		NumericTypeInterface = InNumericTypeInterface;
+		NumericPositionTypeInterface = InNumericPositionTypeInterface;
+		NumericDurationTypeInterface = InNumericDurationTypeInterface;
+		NumericNonMusicalTypeInterface = InNumericNonMusicalTypeInterface;
 		ParentMovieScene = InParentMovieScene;
 	}
 
@@ -89,7 +91,13 @@ private:
 private:
 
 	/** The Numeric Type interface used to convert between display formats and internal tick resolution. */
-	TSharedPtr<INumericTypeInterface<double>> NumericTypeInterface;
+	TSharedPtr<INumericTypeInterface<double>> NumericPositionTypeInterface;
+
+	/** The Numeric Type interface used to convert between display formats and internal tick resolution. */
+	TSharedPtr<INumericTypeInterface<double>> NumericDurationTypeInterface;
+
+	/** The Numeric Type interface used to convert between display formats and internal tick resolution. */
+	TSharedPtr<INumericTypeInterface<double>> NumericNonMusicalTypeInterface;
 
 	/** Store the property handle to the FrameNumber field so we can get/set the value on the object via text box callbacks. */
 	TSharedPtr<IPropertyHandle> MovieSceneSectionPropertyHandle;
diff --git a/Engine/Source/Editor/Sequencer/Private/ISequencer.cpp b/Engine/Source/Editor/Sequencer/Private/ISequencer.cpp
--- a/Engine/Source/Editor/Sequencer/Private/ISequencer.cpp
+++ b/Engine/Source/Editor/Sequencer/Private/ISequencer.cpp
@@ -72,6 +72,18 @@ FFrameRate ISequencer::GetFocusedDisplayRate() const
 	return FFrameRate();
 }
 
+const UFrameBasedMusicMap* ISequencer::GetFocusedMusicMap() const
+{
+	UMovieSceneSequence* FocusedSequence = GetFocusedMovieSceneSequence();
+	if (FocusedSequence)
+	{
+		return FocusedSequence->GetMovieScene()->GetMusicMap();
+	}
+
+	ensureMsgf(false, TEXT("No valid sequence found."));
+	return nullptr;
+}
+
 FGuid ISequencer::CreateBinding(UObject& InObject, const FString& InName)
 {
 	UE::Sequencer::FCreateBindingParams BindingParams;
diff --git a/Engine/Source/Editor/Sequencer/Private/MVVM/CurveEditorExtension.cpp b/Engine/Source/Editor/Sequencer/Private/MVVM/CurveEditorExtension.cpp
--- a/Engine/Source/Editor/Sequencer/Private/MVVM/CurveEditorExtension.cpp
+++ b/Engine/Source/Editor/Sequencer/Private/MVVM/CurveEditorExtension.cpp
@@ -161,7 +161,7 @@ public:
 
 	void HandleSequenceActivated(FMovieSceneSequenceIDRef NewSequenceID)
 	{
-		FocusedTimeAxis->NumericTypeInterface = WeakSequencer.Pin()->GetNumericTypeInterface();
+		FocusedTimeAxis->NumericTypeInterface = WeakSequencer.Pin()->GetNumericPositionTypeInterface();
 	}
 };
 
@@ -229,7 +229,7 @@ const FName FCurveEditorExtension::CurveEditorTabName = FName(TEXT("SequencerGra
 TSharedRef<IPropertyTypeCustomization> CreateFrameNumberCustomization(TWeakPtr<ISequencer> WeakSequencer)
 {
 	TSharedPtr<ISequencer> SequencerPtr = WeakSequencer.Pin();
-	return MakeShared<FFrameNumberDetailsCustomization>(SequencerPtr->GetNumericTypeInterface());
+	return MakeShared<FFrameNumberDetailsCustomization>(SequencerPtr->GetNumericPositionTypeInterface());
 }
 
 FCurveEditorExtension::FCurveEditorExtension()
diff --git a/Engine/Source/Editor/Sequencer/Private/SSequencer.cpp b/Engine/Source/Editor/Sequencer/Private/SSequencer.cpp
--- a/Engine/Source/Editor/Sequencer/Private/SSequencer.cpp
+++ b/Engine/Source/Editor/Sequencer/Private/SSequencer.cpp
@@ -104,6 +104,7 @@
 #include "EditorFontGlyphs.h"
 #include "HAL/PlatformApplicationMisc.h"
 #include "SSequencerPlayRateCombo.h"
+#include "SSequencerMusicModeCombo.h"
 #include "Camera/CameraActor.h"
 #include "SCurveEditorPanel.h"
 #include "Tree/SCurveEditorTree.h"
@@ -141,6 +142,7 @@
 #include "Widgets/SSequencerSearchBox.h"
 #include "IContentBrowserSingleton.h"
 #include "ContentBrowserModule.h"
+#include "Settings/EditorExperimentalSettings.h"
 
 #define LOCTEXT_NAMESPACE "Sequencer"
 
@@ -218,11 +220,16 @@ void SSequencer::Construct(const FArguments& InArgs, TSharedRef<FSequencer> InSe
 
 	TAttribute<FFrameRate> GetTickResolutionAttr = TAttribute<FFrameRate>(InSequencer, &FSequencer::GetFocusedTickResolution);
 	TAttribute<FFrameRate> GetDisplayRateAttr    = TAttribute<FFrameRate>(InSequencer, &FSequencer::GetFocusedDisplayRate);
+	TAttribute<const UFrameBasedMusicMap*> GetMusicMapAttr = TAttribute<const UFrameBasedMusicMap*>(InSequencer, &FSequencer::GetFocusedMusicMap);
 
 	// Create our numeric type interface so we can pass it to the time slider below.
-	TSharedPtr<FFrameNumberInterface> FrameNumberInterface = MakeShareable(new FFrameNumberInterface(GetDisplayFormatAttr, GetZeroPadFramesAttr, GetTickResolutionAttr, GetDisplayRateAttr));
-	NumericTypeInterface = FrameNumberInterface;
-	
+	TSharedPtr<FFrameNumberInterface> FrameNumberInterface = MakeShareable(new FFrameNumberInterface(GetDisplayFormatAttr, GetZeroPadFramesAttr, GetTickResolutionAttr, GetDisplayRateAttr, GetMusicMapAttr, FrameBasedMusicMap::EBarBeatFormat::Position));
+	NumericPositionTypeInterface = FrameNumberInterface;
+	FrameNumberInterface = MakeShareable(new FFrameNumberInterface(GetDisplayFormatAttr, GetZeroPadFramesAttr, GetTickResolutionAttr, GetDisplayRateAttr, GetMusicMapAttr, FrameBasedMusicMap::EBarBeatFormat::Duration));
+	NumericDurationTypeInterface = FrameNumberInterface;
+	FrameNumberInterface = MakeShareable(new FFrameNumberInterface(GetDisplayFormatAttr, GetZeroPadFramesAttr, GetTickResolutionAttr, GetDisplayRateAttr, GetMusicMapAttr, FrameBasedMusicMap::EBarBeatFormat::NotAllowed));
+	NumericNonMusicalTypeInterface = FrameNumberInterface;
+
 	FTimeSliderArgs TimeSliderArgs;
 	{
 		TimeSliderArgs.ViewRange = InArgs._ViewRange;
@@ -265,7 +272,7 @@ void SSequencer::Construct(const FArguments& InArgs, TSharedRef<FSequencer> InSe
 		TimeSliderArgs.AreMarkedFramesLocked = InArgs._AreMarkedFramesLocked;
 		TimeSliderArgs.OnToggleMarkedFramesLocked = InArgs._OnToggleMarkedFramesLocked;
 
-		TimeSliderArgs.NumericTypeInterface = GetNumericTypeInterface();
+		TimeSliderArgs.NumericTypeInterface = GetNumericPositionTypeInterface();
 	}
 
 	OnGetPlaybackSpeeds = InArgs._OnGetPlaybackSpeeds;
@@ -273,7 +280,7 @@ void SSequencer::Construct(const FArguments& InArgs, TSharedRef<FSequencer> InSe
 	RootCustomization.AddMenuExtender = InArgs._AddMenuExtender;
 	RootCustomization.ToolbarExtender = InArgs._ToolbarExtender;
 
-	PlayTimeDisplay = StaticCastSharedRef<STemporarilyFocusedSpinBox<double>>(SequencerPtr.Pin()->MakePlayTimeDisplay(GetNumericTypeInterface()));
+	PlayTimeDisplay = StaticCastSharedRef<STemporarilyFocusedSpinBox<double>>(SequencerPtr.Pin()->MakePlayTimeDisplay(GetNumericPositionTypeInterface()));
 
 	TAttribute<FAnimatedRange> ViewRangeAttribute = InArgs._ViewRange;
 
@@ -714,7 +721,7 @@ TSharedRef<SGridPanel> SSequencer::ConstructTrackAreaGridPanel(const FArguments&
 			Ranges,
 			TimeSliderControllerRef,
 			TAttribute<EVisibility>(this, &SSequencer::GetTimeRangeVisibility),
-			NumericTypeInterface.ToSharedRef()
+			NumericPositionTypeInterface.ToSharedRef()
 		),
 		SequencerWidgets.CreateTimeRangeSlider(TimeSliderControllerRef)
 	);
@@ -988,7 +995,7 @@ TSharedRef<SGridPanel> SSequencer::ConstructTrackAreaGridPanel(const FArguments&
 		.VAlign(VAlign_Top)
 		[
 			// Transform box
-			SAssignNew(TransformBox, SSequencerTransformBox, SequencerRef, *SequencerSettings, NumericTypeInterface.ToSharedRef())
+			SAssignNew(TransformBox, SSequencerTransformBox, SequencerRef, *SequencerSettings, NumericPositionTypeInterface.ToSharedRef())
 		]
 
 		+ SGridPanel::Slot(Column1, Row2, SGridPanel::Layer(40))
@@ -997,7 +1004,7 @@ TSharedRef<SGridPanel> SSequencer::ConstructTrackAreaGridPanel(const FArguments&
 		.VAlign(VAlign_Top)
 		[
 			// Stretch box
-			SAssignNew(StretchBox, SSequencerStretchBox, SequencerRef, *SequencerSettings, NumericTypeInterface.ToSharedRef())
+			SAssignNew(StretchBox, SSequencerStretchBox, SequencerRef, *SequencerSettings, NumericPositionTypeInterface.ToSharedRef())
 		]
 
 		// debug vis
@@ -1157,9 +1164,19 @@ void SSequencer::NotifyPostChange(const FPropertyChangedEvent& PropertyChangedEv
 /* SSequencer implementation
  *****************************************************************************/
 
-TSharedRef<INumericTypeInterface<double>> SSequencer::GetNumericTypeInterface() const
+TSharedRef<INumericTypeInterface<double>> SSequencer::GetNumericPositionTypeInterface() const
+{
+	return NumericPositionTypeInterface.ToSharedRef();
+}
+
+TSharedRef<INumericTypeInterface<double>> SSequencer::GetNumericDurationTypeInterface() const
 {
-	return NumericTypeInterface.ToSharedRef();
+	return NumericDurationTypeInterface.ToSharedRef();
+}
+
+TSharedRef<INumericTypeInterface<double>> SSequencer::GetNumericNonMusicalTypeInterface() const
+{
+	return NumericNonMusicalTypeInterface.ToSharedRef();
 }
 
 void SSequencer::UpdateOutlinerViewColumns()
@@ -1739,6 +1756,20 @@ void SSequencer::PopulateToolBar(UToolMenu* InMenu)
 
 		SnappingSection.AddSeparator(NAME_None);
 
+		static const auto MusicModeEnabledCVar = IConsoleManager::Get().FindConsoleVariable(TEXT("Sequencer.EnableMusicModeEditing"));
+		if (MusicModeEnabledCVar->GetInt() != 0)
+		{
+			FToolMenuEntry MusicModeEntry = FToolMenuEntry::InitWidget(
+				"MusicMode",
+				SNew(SSequencerMusicModeCombo, Sequencer.Pin(), SequencerWidget.ToSharedRef()),
+				LOCTEXT("MusicMode", "MusicMode")
+			);
+			MusicModeEntry.StyleNameOverride = SequencerToolbarStyleName;
+			SnappingSection.AddEntry(MusicModeEntry);
+
+			SnappingSection.AddSeparator(NAME_None);
+		}
+
 		FToolMenuEntry PlayRateEntry = FToolMenuEntry::InitWidget(
 			"PlayRate",
 			SNew(SSequencerPlayRateCombo, Sequencer.Pin(), SequencerWidget.ToSharedRef()),
@@ -2283,6 +2314,10 @@ void SSequencer::FillTimeDisplayFormatMenu(FMenuBuilder& MenuBuilder)
 {
 	TSharedPtr<FSequencer> Sequencer = SequencerPtr.Pin();
 	bool bShouldDisplayDropFormat = FTimecode::UseDropFormatTimecode(Sequencer->GetFocusedDisplayRate());
+	
+	const UFrameBasedMusicMap* MusicMap = Sequencer->GetFocusedMusicMap();
+	static const auto MusicModeEnabledCVar = IConsoleManager::Get().FindConsoleVariable(TEXT("Sequencer.EnableMusicModeEditing"));
+	bool bShouldDisplayBarsBeats = MusicMap && !MusicMap->IsEmpty() && MusicModeEnabledCVar->GetInt() != 0;
 
 	const UEnum* FrameNumberDisplayEnum = StaticEnum<EFrameNumberDisplayFormats>();
 	check(FrameNumberDisplayEnum);
@@ -2307,6 +2342,12 @@ void SSequencer::FillTimeDisplayFormatMenu(FMenuBuilder& MenuBuilder)
 					continue;
 				}
 
+				// Don't show Bars.Beats if the sequence doesn't have a music map.
+				if (Value == EFrameNumberDisplayFormats::BarsBeats && !bShouldDisplayBarsBeats)
+				{
+					continue;
+				}
+
 				MenuBuilder.AddMenuEntry(
 					FrameNumberDisplayEnum->GetDisplayNameTextByIndex(Index),
 					FrameNumberDisplayEnum->GetToolTipTextByIndex(Index),
@@ -2373,7 +2414,7 @@ TSharedRef<SWidget> SSequencer::MakePlaybackMenu()
 				.AutoWidth()
 				[
 					SNew(SSpinBox<double>)
-						.TypeInterface(NumericTypeInterface)
+						.TypeInterface(NumericPositionTypeInterface)
 						.IsEnabled_Lambda([this]() {
 							return !SequencerPtr.Pin()->IsPlaybackRangeLocked();
 						})
@@ -2434,7 +2475,7 @@ TSharedRef<SWidget> SSequencer::MakePlaybackMenu()
 				.AutoWidth()
 				[
 					SNew(SSpinBox<double>)
-						.TypeInterface(NumericTypeInterface)
+						.TypeInterface(NumericPositionTypeInterface)
 						.IsEnabled_Lambda([this]() {
 					 		return !SequencerPtr.Pin()->IsPlaybackRangeLocked();
 						})
@@ -2506,7 +2547,7 @@ TSharedRef<SWidget> SSequencer::MakePlaybackMenu()
 			.AutoWidth()
 			[
 				SNew(SSpinBox<double>)
-					.TypeInterface(NumericTypeInterface)
+					.TypeInterface(NumericDurationTypeInterface)
 					.Style(&FAppStyle::GetWidgetStyle<FSpinBoxStyle>("Sequencer.HyperlinkSpinBox"))
 					.OnValueCommitted_Lambda([=](double Value, ETextCommit::Type){ OnJumpFrameIncrementChanged(Value); })
 					.OnValueChanged_Lambda([=](double Value) { OnJumpFrameIncrementChanged(Value); })
@@ -2768,7 +2809,7 @@ TSharedRef<SWidget> SSequencer::MakeTimeRange(const TSharedRef<SWidget>& InnerCo
 		ShowRange,
 		TimeSliderController.ToSharedRef(),
 		EVisibility::Visible,
-		NumericTypeInterface.ToSharedRef()
+		NumericPositionTypeInterface.ToSharedRef()
 		);
 	return SequencerWidgets.CreateTimeRange(Args, InnerContent);
 }
diff --git a/Engine/Source/Editor/Sequencer/Private/SSequencer.h b/Engine/Source/Editor/Sequencer/Private/SSequencer.h
--- a/Engine/Source/Editor/Sequencer/Private/SSequencer.h
+++ b/Engine/Source/Editor/Sequencer/Private/SSequencer.h
@@ -336,8 +336,14 @@ public:
 	/** Access this widget's track area widget */
 	TSharedPtr<UE::Sequencer::STrackAreaView> GetTrackAreaWidget() const { return TrackArea; }
 
-	/** @return a numeric type interface that will parse and display numbers as frames and times correctly */
-	TSharedRef<INumericTypeInterface<double>> GetNumericTypeInterface() const;
+	/** @return a numeric type interface that will parse and display numbers as frame / time / bar.beat positions correctly */
+	TSharedRef<INumericTypeInterface<double>> GetNumericPositionTypeInterface() const;
+
+	/** @return a numeric type interface that will parse and display numbers as frame / time / bar.beat durations correctly */
+	TSharedRef<INumericTypeInterface<double>> GetNumericDurationTypeInterface() const;
+
+	/** @return a numeric type interface that will parse and display numbers as frame / time durations correctly. Bars.Beats will be forced to time. */
+	TSharedRef<INumericTypeInterface<double>> GetNumericNonMusicalTypeInterface() const;
 
 	void OpenTickResolutionOptions();
 
@@ -816,7 +822,13 @@ private:
 	TArray<TSharedPtr<FExtender>> ViewMenuExtenders;
 
 	/** Numeric type interface used for converting parsing and generating strings from numbers */
-	TSharedPtr<INumericTypeInterface<double>> NumericTypeInterface;
+	TSharedPtr<INumericTypeInterface<double>> NumericPositionTypeInterface;
+
+	/** Numeric type interface used for converting parsing and generating strings from numbers */
+	TSharedPtr<INumericTypeInterface<double>> NumericDurationTypeInterface;
+
+	/** Numeric type interface used for converting parsing and generating strings from numbers */
+	TSharedPtr<INumericTypeInterface<double>> NumericNonMusicalTypeInterface;
 
 	/** Time slider controller for this sequencer */
 	TSharedPtr<FSequencerTimeSliderController> TimeSliderController;
diff --git a/Engine/Source/Editor/Sequencer/Private/SSequencerMarkerImportDialog.cpp b/Engine/Source/Editor/Sequencer/Private/SSequencerMarkerImportDialog.cpp
new file mode 100644
--- /dev/null
+++ b/Engine/Source/Editor/Sequencer/Private/SSequencerMarkerImportDialog.cpp
@@ -0,0 +1,142 @@
+// Copyright Epic Games, Inc. All Rights Reserved.
+
+#include "SSequencerMarkerImportDialog.h"
+#include "Widgets/SBoxPanel.h"
+#include "Widgets/Layout/SGridPanel.h"
+#include "Widgets/Layout/SScrollBox.h"
+#include "Widgets/Text/STextBlock.h"
+#include "Widgets/Input/SCheckBox.h"
+#include "Widgets/Input/SComboBox.h"
+#include "Widgets/Input/SButton.h"
+#include "SPrimaryButton.h"
+#include "Styling/AppStyle.h"
+#include "Framework/Notifications/NotificationManager.h"
+#include "Widgets/Notifications/SNotificationList.h"
+#include "Editor.h"
+
+#define LOCTEXT_NAMESPACE "SSequencerMarkerImportDialog"
+
+FSequencerMarkerImportConfig SSequencerMarkerImportDialog::GetMarkerImportOptions(const TArray<FMarkerProviderEntry>& AvailableMarkers)
+{
+	const FVector2D DEFAULT_WINDOW_SIZE = FVector2D(600, 700);
+
+	/** Create the window to host our package dialog widget */
+	TSharedRef<SWindow> ImportMarkerOptionsWindow = SNew(SWindow)
+		.Title(FText::FromString("Import Markers"))
+		.SizingRule(ESizingRule::Autosized);
+
+	/** Set the content of the window to our package dialog widget */
+	TSharedRef<SSequencerMarkerImportDialog> MarkerImportDialog =
+		SNew(SSequencerMarkerImportDialog, AvailableMarkers);
+
+	ImportMarkerOptionsWindow->SetContent(MarkerImportDialog);
+
+	/** Show the dialog window as a modal window */
+	GEditor->EditorAddModalWindow(ImportMarkerOptionsWindow);
+
+	return MarkerImportDialog->GetResult();
+}
+
+void SSequencerMarkerImportDialog::Construct(const FArguments& InArgs, const TArray<FMarkerProviderEntry>& AvailableMarkers)
+{
+	TSharedPtr<SScrollBox> MarkerList;
+
+	ChildSlot
+	[
+		SNew(SVerticalBox)
+		+ SVerticalBox::Slot()
+		.AutoHeight()
+		.Padding(16.f)
+		[
+			SNew(STextBlock)
+			.Text(LOCTEXT("Explanation", "The imported music map contains timeline markers..."))
+		]
+		+ SVerticalBox::Slot()
+		.MinHeight(300)
+		.MaxHeight(300)
+		.Padding(16.f,0.f)
+		[
+			SNew(SBorder)
+			.BorderImage(FAppStyle::Get().GetBrush("Brushes.Panel"))
+			.Padding(8.f, 0.f)
+			[
+				SAssignNew(MarkerList, SScrollBox)
+			]
+		]
+		+ SVerticalBox::Slot()
+		.AutoHeight()
+		.HAlign(HAlign_Center)
+		.VAlign(VAlign_Bottom)
+		.Padding(16.0f, 3.0f)
+		[
+			SNew(SHorizontalBox)
+			+SHorizontalBox::Slot()
+			.AutoWidth()
+			[
+				SNew(SCheckBox)
+				.IsChecked(false)
+				.OnCheckStateChanged_Lambda(
+					[this](ECheckBoxState NewCheckedState)
+					{
+						Result.bClearPreviouslyImported = NewCheckedState == ECheckBoxState::Checked ? true : false;
+					}
+				)
+				[
+					SNew(STextBlock).Text(LOCTEXT("ClearExistingMarkers", "Clear Previously Imported Markers"))
+				]
+			]
+		]
+		+ SVerticalBox::Slot()
+		.AutoHeight()
+		.HAlign(HAlign_Right)
+		.VAlign(VAlign_Bottom)
+		.Padding(16.0f, 3.0f)
+		[
+			SNew(SHorizontalBox)
+
+			+SHorizontalBox::Slot()
+			.AutoWidth()
+			.VAlign(VAlign_Bottom)
+			.Padding(8, 0)
+			[
+				SNew(SPrimaryButton)
+				.Text(LOCTEXT("ImportButtonText", "Import"))
+				.OnClicked(this, &SSequencerMarkerImportDialog::DoImport)
+			]
+
+			+SHorizontalBox::Slot()
+			.AutoWidth()
+			.VAlign(VAlign_Bottom)
+			[
+				SNew(SButton)
+				.Text(LOCTEXT("CancelButtonText", "Cancel"))
+				.OnClicked(this, &SSequencerMarkerImportDialog::Cancel)
+			]
+		]
+	];
+
+	for (const FMarkerProviderEntry& Marker : AvailableMarkers)
+	{
+		MarkerList->AddSlot()
+		[
+			SNew(STextBlock)
+				.Text(FText::FromString(Marker.Label))
+		];
+	}
+}
+
+FReply SSequencerMarkerImportDialog::Cancel()
+{
+	TSharedRef<SWindow> ParentWindow = FSlateApplication::Get().FindWidgetWindow( AsShared() ).ToSharedRef();
+	FSlateApplication::Get().RequestDestroyWindow( ParentWindow );
+	return FReply::Handled();
+}
+
+FReply SSequencerMarkerImportDialog::DoImport()
+{
+	Result.bImport = true;
+	return Cancel();
+}
+
+
+#undef LOCTEXT_NAMESPACE
\ No newline at end of file
diff --git a/Engine/Source/Editor/Sequencer/Private/SSequencerMarkerImportDialog.h b/Engine/Source/Editor/Sequencer/Private/SSequencerMarkerImportDialog.h
new file mode 100644
--- /dev/null
+++ b/Engine/Source/Editor/Sequencer/Private/SSequencerMarkerImportDialog.h
@@ -0,0 +1,36 @@
+// Copyright Epic Games, Inc. All Rights Reserved.
+
+#pragma once
+
+#include "CoreTypes.h"
+#include "Templates/SharedPointer.h"
+#include "Styling/SlateTypes.h"
+#include "Widgets/DeclarativeSyntaxSupport.h"
+#include "Widgets/SCompoundWidget.h"
+#include "Widgets/SBoxPanel.h"
+#include "FrameBasedMusicMap.h"
+
+
+struct FSequencerMarkerImportConfig
+{
+	bool bImport = false;
+	bool bClearPreviouslyImported = false;
+};
+
+class SSequencerMarkerImportDialog : public SCompoundWidget
+{
+public:
+	static FSequencerMarkerImportConfig GetMarkerImportOptions(const TArray<FMarkerProviderEntry>& AvailableMarkers);
+
+	SLATE_BEGIN_ARGS(SSequencerMarkerImportDialog){}
+	SLATE_END_ARGS()
+
+	void Construct(const FArguments& InArgs, const TArray<FMarkerProviderEntry>& AvailableMarkers);
+
+	FSequencerMarkerImportConfig GetResult() const { return Result; }
+
+private:
+	FReply Cancel();
+	FReply DoImport();
+	FSequencerMarkerImportConfig Result;
+};
\ No newline at end of file
diff --git a/Engine/Source/Editor/Sequencer/Private/SSequencerMusicModeCombo.cpp b/Engine/Source/Editor/Sequencer/Private/SSequencerMusicModeCombo.cpp
new file mode 100644
--- /dev/null
+++ b/Engine/Source/Editor/Sequencer/Private/SSequencerMusicModeCombo.cpp
@@ -0,0 +1,537 @@
+// Copyright Epic Games, Inc. All Rights Reserved.
+
+#include "SSequencerMusicModeCombo.h"
+#include "Sequencer.h"
+#include "SSequencer.h"
+#include "Widgets/SOverlay.h"
+#include "Widgets/Layout/SBox.h"
+#include "Widgets/Text/STextBlock.h"
+#include "Widgets/Input/SComboButton.h"
+#include "Widgets/Input/SEditableTextBox.h"
+#include "Widgets/Images/SImage.h"
+#include "Framework/MultiBox/MultiBoxBuilder.h"
+#include "Misc/CString.h"
+#include "MovieSceneSequence.h"
+#include "MovieScene.h"
+#include "MovieSceneCommonHelpers.h"
+#include "Compilation/MovieSceneCompiledDataManager.h"
+#include "ScopedTransaction.h"
+#include "Widgets/SFrameRateEntryBox.h"
+#include "EditorFontGlyphs.h"
+#include "Evaluation/IMovieSceneCustomClockSource.h"
+#include "SceneOutlinerModule.h"
+#include "SceneOutlinerPublicTypes.h"
+#include "IContentBrowserSingleton.h"
+#include "ContentBrowserModule.h"
+#include "Misc/Timecode.h"
+#include "Styling/ToolBarStyle.h"
+#include "ActorTreeItem.h"
+
+#define LOCTEXT_NAMESPACE "SSequencerMusicModeCombo"
+
+void SSequencerMusicModeCombo::Construct(const FArguments& InArgs, TWeakPtr<FSequencer> InWeakSequencer, TWeakPtr<SSequencer> InWeakSequencerWidget)
+{
+	WeakSequencer = InWeakSequencer;
+	WeakSequencerWidget = InWeakSequencerWidget;
+
+	const FToolBarStyle& SequencerToolBarStyle = FAppStyle::Get().GetWidgetStyle<FToolBarStyle>(InArgs._StyleName);
+
+	SetToolTipText(MakeAttributeSP(this, &SSequencerMusicModeCombo::GetToolTipText));
+
+	ChildSlot
+	.VAlign(VAlign_Fill)
+	[
+		SNew(SComboButton)
+		.ContentPadding(FMargin(2.f, 1.0f))
+		.VAlign(VAlign_Fill)
+		.ComboButtonStyle(&SequencerToolBarStyle.ComboButtonStyle)
+		.OnGetMenuContent(this, &SSequencerMusicModeCombo::OnCreateMenu)
+		.ButtonContent()
+		[
+			SNew(SHorizontalBox)
+
+			+ SHorizontalBox::Slot()
+			.HAlign(HAlign_Center)
+			.VAlign(VAlign_Center)
+			.AutoWidth()
+			[
+				SNew(STextBlock)
+				.Text(this, &SSequencerMusicModeCombo::GetTempoTimeSignatureText)
+				.TextStyle(&SequencerToolBarStyle.LabelStyle)
+			]
+		]
+	];
+}
+
+bool SSequencerMusicModeCombo::GetIsSequenceReadOnly() const
+{
+	TSharedPtr<FSequencer> Sequencer = WeakSequencer.Pin();
+	return !Sequencer.IsValid() || Sequencer->IsReadOnly();
+}
+
+TSharedRef<SWidget> SSequencerMusicModeCombo::OnCreateMenu()
+{
+	TSharedPtr<FSequencer> Sequencer         = WeakSequencer.Pin();
+	TSharedPtr<SSequencer> SequencerWidget   = WeakSequencerWidget.Pin();
+	UMovieSceneSequence*   FocusedSequence   = Sequencer.IsValid() ? Sequencer->GetFocusedMovieSceneSequence() : nullptr;
+	UMovieScene*           FocusedMovieScene = FocusedSequence ? FocusedSequence->GetMovieScene() : nullptr;
+
+	if (!Sequencer.IsValid() || !SequencerWidget.IsValid() || !FocusedSequence || !FocusedMovieScene)
+	{
+		return SNullWidget::NullWidget;
+	}
+
+	const UFrameBasedMusicMap* MusicMap = FocusedMovieScene->GetMusicMap();
+
+	FMenuBuilder MenuBuilder(true, nullptr);
+
+	PopulateSimpleMusicMapSection(MenuBuilder, MusicMap);
+	PopulateImportMusicMapSection(MenuBuilder);
+
+	if (FocusedMovieScene->HasMusicMap())
+	{
+		MenuBuilder.AddMenuSeparator();
+		MenuBuilder.AddMenuEntry(
+			LOCTEXT("CearMusicMap", "Remove Music Map"),
+			LOCTEXT("ClearMusicMap_tooltip", "Removes any and all tempo and time signature information."),
+			FSlateIcon(),
+			FUIAction(
+				FExecuteAction::CreateSP(this, &SSequencerMusicModeCombo::OnRemoveAllMusicMaps),
+				FCanExecuteAction::CreateLambda([this] { return !GetIsSequenceReadOnly(); })
+			)
+		);
+	}
+
+ 	return MenuBuilder.MakeWidget();
+}
+
+void SSequencerMusicModeCombo::PopulateSimpleMusicMapSection(FMenuBuilder& MenuBuilder, const UFrameBasedMusicMap* MusicMap)
+{
+	MenuBuilder.BeginSection(NAME_None, LOCTEXT("SimpleMapSection", "Simple Music Map"));
+	{
+		if (!MusicMap)
+		{
+			MenuBuilder.AddMenuEntry(
+				LOCTEXT("MakeSimpleMusicMap", "Add Simple Music Map"),
+				LOCTEXT("MakeSimpleMusicMap_Description", "Creates a simple, one tempo, one time signature music map for the sequence."),
+				FSlateIcon(),
+				FUIAction( FExecuteAction::CreateSP(this, &SSequencerMusicModeCombo::OnCreateSimpleMusicMap),
+					       FCanExecuteAction::CreateLambda([this] { return !GetIsSequenceReadOnly(); }))
+				);
+		}
+		else if (!MusicMap->IsSimple())
+		{
+			MenuBuilder.AddMenuEntry(
+				LOCTEXT("ClearMusicMapAndMakeSimple", "Clear Imported Music Map And Make Simple"),
+				LOCTEXT("ClearMusicMapAndMakeSimple_Description", "Clears all imported tempo and time signature changes and creates a simple, one tempo, one time signature music map for the sequence."),
+				FSlateIcon(),
+				FUIAction(FExecuteAction::CreateSP(this, &SSequencerMusicModeCombo::OnCreateSimpleMusicMap),
+					FCanExecuteAction::CreateLambda([this] { return !GetIsSequenceReadOnly(); }))
+			);
+		}
+		else
+		{
+			MenuBuilder.AddWidget(
+				SNew(SBox)
+				.HAlign(HAlign_Right)
+				.MaxDesiredWidth(100.f)
+				[
+					SNew(SEditableTextBox)
+						.Text_Lambda([this]()
+							{
+								TSharedPtr<FSequencer> Sequencer = WeakSequencer.Pin();
+								UMovieSceneSequence* FocusedSequence = Sequencer.IsValid() ? Sequencer->GetFocusedMovieSceneSequence() : nullptr;
+								UMovieScene* FocusedMovieScene = FocusedSequence ? FocusedSequence->GetMovieScene() : nullptr;
+								const UFrameBasedMusicMap* MusicMap = FocusedMovieScene ? FocusedMovieScene->GetMusicMap() : nullptr;
+								if (MusicMap)
+								{
+									return FText::Format(LOCTEXT("MusicModeSettings_Bpm", "{0} bpm"), MusicMap->GetTempoForSeconds(0.0));
+								}
+								return FText();
+							})
+						.OnTextCommitted(this, &SSequencerMusicModeCombo::OnTempoTextCommitted)
+						.OnVerifyTextChanged_Lambda([](const FText& InText, FText& OutErrorMessage)
+							{
+								float NewTempo = FCString::Atof(*(InText.ToString()));
+								bool Valid = NewTempo > 0.0f && NewTempo < 999.0f;
+								if (!Valid)
+								{
+									OutErrorMessage = LOCTEXT("MusicModeSettings_Error_Bpm", "Tempo must be greater than 0 and less than 999.");
+								}
+								return Valid;
+							})
+						.IsEnabled(this, &SSequencerMusicModeCombo::SimpleMapEntryEnabled)
+						.SelectAllTextWhenFocused(true)
+				],
+				LOCTEXT("MusicModeSettings_Tempo", "Tempo"),
+				/*InNoIndent*/ true, /*bInSearchable*/ true,
+				LOCTEXT("MusicModeSettings_Tempo_ToolTip", "This is the \"natural tempo\" this sequence will play at when not being driven by an external musical clock source.")
+			);
+			MenuBuilder.AddWidget(
+				SNew(SBox)
+				.HAlign(HAlign_Right)
+				.MaxDesiredWidth(100.f)
+				[
+					SNew(SEditableTextBox)
+						.Text_Lambda([this]()
+							{
+								TSharedPtr<FSequencer> Sequencer = WeakSequencer.Pin();
+								UMovieSceneSequence* FocusedSequence = Sequencer.IsValid() ? Sequencer->GetFocusedMovieSceneSequence() : nullptr;
+								UMovieScene* FocusedMovieScene = FocusedSequence ? FocusedSequence->GetMovieScene() : nullptr;
+								const UFrameBasedMusicMap* MusicMap = FocusedMovieScene ? FocusedMovieScene->GetMusicMap() : nullptr;
+								if (MusicMap)
+								{
+									return FText::Format(LOCTEXT("TempoFormat", "{0}/{1}"),
+										FText::AsNumber(MusicMap->GetInitialTimeSignature().Numerator),
+										FText::AsNumber(MusicMap->GetInitialTimeSignature().Denominator));
+								}
+								return FText();
+							})
+						.OnTextCommitted(this, &SSequencerMusicModeCombo::OnTimeSignatureTextCommitted)
+						.OnVerifyTextChanged(this, &SSequencerMusicModeCombo::ValidateTimeSignatureString)
+						.IsEnabled(this, &SSequencerMusicModeCombo::SimpleMapEntryEnabled)
+						.SelectAllTextWhenFocused(true)
+				],
+				LOCTEXT("MusicModeSettings_TimeSignature", "Time Signature"),
+				/*InNoIndent*/ true, /*bInSearchable*/ true,
+				LOCTEXT("MusicModeSettings_Tempo_ToolTip", "The \"natural time signature\" this sequence is authored for.")
+			);
+
+
+		}
+	}
+	MenuBuilder.EndSection();
+}
+
+void SSequencerMusicModeCombo::PopulateImportMusicMapSection(FMenuBuilder& MenuBuilder)
+{
+	if (GetIsSequenceReadOnly())
+	{
+		return;
+	}
+
+	MenuBuilder.BeginSection(NAME_None, LOCTEXT("ImportMapSection", "Import Music Map"));
+	{
+		auto MusicMapAssetsMenu = [this](FMenuBuilder& ActorMenuBuilder)
+			{
+				FAssetPickerConfig AssetPickerConfig;
+				AssetPickerConfig.Filter.ClassPaths.Add(UMusicMapSource::StaticClass()->GetClassPathName());
+				AssetPickerConfig.bAllowNullSelection = false;
+				AssetPickerConfig.Filter.bRecursiveClasses = true;
+				AssetPickerConfig.OnAssetSelected = FOnAssetSelected::CreateLambda([this](const FAssetData& In) { this->OnImportMusicMap(In.GetAsset()); });
+				AssetPickerConfig.InitialAssetViewType = EAssetViewType::List;
+				AssetPickerConfig.bAllowDragging = false;
+
+				FContentBrowserModule& ContentBrowserModule = FModuleManager::Get().LoadModuleChecked<FContentBrowserModule>(TEXT("ContentBrowser"));
+
+				ActorMenuBuilder.AddWidget(
+					SNew(SBox)
+					.WidthOverride(300)
+					.HeightOverride(300)
+					[
+						ContentBrowserModule.Get().CreateAssetPicker(AssetPickerConfig)
+					],
+					FText(),
+					true /*bNoIndent*/
+				);
+			};
+
+		MenuBuilder.AddSubMenu(
+			LOCTEXT("ImportMusicMap", "Import Music Map..."),
+			LOCTEXT("ImportMusicMap_tooltip", "Imports tempo and time signature changes from the selected asset. NOTE: This will remove any existing tempo and time signature changes!"),
+			FNewMenuDelegate::CreateLambda(MusicMapAssetsMenu)
+		);
+	}
+	MenuBuilder.EndSection();
+}
+
+FText SSequencerMusicModeCombo::GetTempoTimeSignatureText() const
+{
+	TSharedPtr<FSequencer>     Sequencer         = WeakSequencer.Pin();
+	UMovieSceneSequence*       FocusedSequence   = Sequencer.IsValid() ? Sequencer->GetFocusedMovieSceneSequence() : nullptr;
+	UMovieScene*               FocusedMovieScene = FocusedSequence ? FocusedSequence->GetMovieScene() : nullptr;
+	const UFrameBasedMusicMap* MusicMap          = FocusedMovieScene ? FocusedMovieScene->GetMusicMap() : nullptr;
+
+	if (MusicMap)
+	{
+		FNumberFormattingOptions Formatting;
+		Formatting.MaximumFractionalDigits = 3;
+		if (MusicMap->IsSimple())
+		{
+			float Tempo = MusicMap->GetInitialTempo();
+			FFrameBasedTimeSignature TimeSignature = MusicMap->GetInitialTimeSignature();
+			return FText::Format(LOCTEXT("TempoTimeSigDisplay", "{0} bpm {1}/{2}"),
+				FText::AsNumber(Tempo, &Formatting),
+				FText::AsNumber(TimeSignature.Numerator, &Formatting),
+				FText::AsNumber(TimeSignature.Denominator, &Formatting));
+		}
+		else
+		{
+			FQualifiedFrameTime CursorTime = Sequencer->GetLocalTime();
+			float Tempo = MusicMap->GetTempoForFrameTime(CursorTime.Time);
+			FFrameBasedTimeSignature TimeSignature = MusicMap->GetTimeSignatureForFrameTime(CursorTime.Time);
+			return FText::Format(LOCTEXT("TempoTimeSigDisplay", "Current: {0} bpm {1}/{2}"),
+				FText::AsNumber(Tempo, &Formatting),
+				FText::AsNumber(TimeSignature.Numerator, &Formatting),
+				FText::AsNumber(TimeSignature.Denominator, &Formatting));
+		}
+	}
+
+	return LOCTEXT("MusicModeNotEnabled", "Music Mode...");
+}
+
+FText SSequencerMusicModeCombo::GetToolTipText() const
+{
+	TSharedPtr<FSequencer>     Sequencer         = WeakSequencer.Pin();
+	UMovieSceneSequence*       FocusedSequence   = Sequencer.IsValid() ? Sequencer->GetFocusedMovieSceneSequence() : nullptr;
+	UMovieScene*               FocusedMovieScene = FocusedSequence ? FocusedSequence->GetMovieScene() : nullptr;
+	const UFrameBasedMusicMap* MusicMap          = FocusedMovieScene ? FocusedMovieScene->GetMusicMap() : nullptr;
+
+	if (MusicMap && !MusicMap->IsSimple())
+	{
+		return FText::Format(LOCTEXT("MusicModeMenuToolTip", "{0} tempos & {1} time signatures"), MusicMap->TempoMap.Num(), MusicMap->BarMap.Num());
+	}
+	return FText();
+}
+
+void SSequencerMusicModeCombo::OnRemoveAllMusicMaps()
+{
+	TSharedPtr<FSequencer>     Sequencer         = WeakSequencer.Pin();
+	UMovieSceneSequence*       FocusedSequence   = Sequencer.IsValid() ? Sequencer->GetFocusedMovieSceneSequence() : nullptr;
+	UMovieScene*               FocusedMovieScene = FocusedSequence ? FocusedSequence->GetMovieScene() : nullptr;
+	const UFrameBasedMusicMap* MusicMap          = FocusedMovieScene ? FocusedMovieScene->GetMusicMap() : nullptr;
+
+	if (MusicMap)
+	{
+		FScopedTransaction ScopedTransaction(FText(LOCTEXT("SetMusicMode", "Turning music mode off for sequence")));
+		FocusedMovieScene->Modify();
+		FocusedMovieScene->RemoveMusicMap();
+	}
+}
+
+void SSequencerMusicModeCombo::OnCreateSimpleMusicMap()
+{
+	TSharedPtr<FSequencer> Sequencer         = WeakSequencer.Pin();
+	UMovieSceneSequence*   FocusedSequence   = Sequencer.IsValid() ? Sequencer->GetFocusedMovieSceneSequence() : nullptr;
+	UMovieScene*           FocusedMovieScene = FocusedSequence ? FocusedSequence->GetMovieScene() : nullptr;
+
+	if (FocusedMovieScene)
+	{
+		if (FocusedMovieScene->HasMusicMap() && !FocusedMovieScene->GetMusicMap()->IsSimple())
+		{
+			// WARN THE USER HERE! Make sure they want to blow away the existing music map!
+		}
+
+		FScopedTransaction ScopedTransaction(FText::Format(LOCTEXT("SetMusicMode", "Turning music mode {0} for sequence"), FocusedMovieScene->HasMusicMap() ? FText::FromString(TEXT("On")) : FText::FromString(TEXT("Off"))));
+		FocusedMovieScene->Modify();
+
+		UFrameBasedMusicMap* NewSimpleMusicMap = NewObject<UFrameBasedMusicMap>(FocusedMovieScene);
+		NewSimpleMusicMap->SetFrameResolution(FocusedMovieScene->GetTickResolution());
+		NewSimpleMusicMap->Init(120.0f, { 4, 4 });
+		if (FocusedMovieScene->HasMusicMap())
+		{
+			// This will remap from the old to the new map AND set the movie scene's tempo map to the new map...
+			UE::MovieScene::TimeHelpers::RemapMusicalFrameTimes(FocusedMovieScene, NewSimpleMusicMap);
+		}
+		else
+		{
+			FocusedMovieScene->SetMusicMapDirectly(NewSimpleMusicMap);
+		}
+	}
+}
+
+void SSequencerMusicModeCombo::OnImportMusicMap(UObject* Object)
+{
+	TSharedPtr<FSequencer> Sequencer         = WeakSequencer.Pin();
+	UMovieSceneSequence*   FocusedSequence   = Sequencer.IsValid() ? Sequencer->GetFocusedMovieSceneSequence() : nullptr;
+	UMovieScene*           FocusedMovieScene = FocusedSequence ? FocusedSequence->GetMovieScene() : nullptr;
+
+	if (FocusedMovieScene)
+	{
+		IMusicMapSource* MapSource = Cast<IMusicMapSource>(Object);
+		check(MapSource);
+
+		UFrameBasedMusicMap* NewMusicMap = NewObject<UFrameBasedMusicMap>(FocusedMovieScene);
+		NewMusicMap->SetFrameResolution(FocusedMovieScene->GetTickResolution());
+		MapSource->CreateFrameBasedMusicMap(NewMusicMap);
+
+		FMarkerProviderResults Markers = MapSource->GatherMarkers(NewMusicMap);
+
+		FScopedTransaction ScopedTransaction(LOCTEXT("RemapSequence", "Remap Sequence to a new music map"));
+		FocusedMovieScene->Modify();
+
+		// This will remap from the old to the new map AND set the movie scene's tempo map to the new map...
+		UE::MovieScene::TimeHelpers::RemapMusicalFrameTimes(FocusedMovieScene, NewMusicMap);
+
+		if (!Markers.IsEmpty())
+		{
+			PromptForMarkerImport(Markers);
+		}
+	}
+}
+
+void SSequencerMusicModeCombo::OnTempoTextCommitted(const FText& NewValue, ETextCommit::Type)
+{
+	TSharedPtr<FSequencer> Sequencer         = WeakSequencer.Pin();
+	UMovieSceneSequence*   FocusedSequence   = Sequencer.IsValid() ? Sequencer->GetFocusedMovieSceneSequence() : nullptr;
+	UMovieScene*           FocusedMovieScene = FocusedSequence ? FocusedSequence->GetMovieScene() : nullptr;
+	UFrameBasedMusicMap*   MusicMap          = FocusedMovieScene ? FocusedMovieScene->GetMusicMap() : nullptr;
+
+	if (MusicMap)
+	{
+		float NewTempo = FCString::Atof(*(NewValue.ToString()));
+		if (NewTempo > 0.0f && NewTempo < 999.0f)
+		{
+			if (MusicMap->TempoMap[0].Bpm() != NewTempo)
+			{
+				UFrameBasedMusicMap* NewMusicMap = NewObject<UFrameBasedMusicMap>(FocusedMovieScene);
+				NewMusicMap->SetFrameResolution(FocusedMovieScene->GetTickResolution());
+				NewMusicMap->Init(NewTempo, MusicMap->GetInitialTimeSignature());
+
+				// This will remap from the old to the new map AND set the movie scene's tempo map to the new map...
+				UE::MovieScene::TimeHelpers::RemapMusicalFrameTimes(FocusedMovieScene, NewMusicMap);
+			}
+		}
+	}
+}
+
+void SSequencerMusicModeCombo::OnTimeSignatureTextCommitted(const FText& NewValue, ETextCommit::Type)
+{
+	TSharedPtr<FSequencer> Sequencer         = WeakSequencer.Pin();
+	UMovieSceneSequence*   FocusedSequence   = Sequencer.IsValid() ? Sequencer->GetFocusedMovieSceneSequence() : nullptr;
+	UMovieScene*           FocusedMovieScene = FocusedSequence ? FocusedSequence->GetMovieScene() : nullptr;
+	UFrameBasedMusicMap*   MusicMap          = FocusedMovieScene ? FocusedMovieScene->GetMusicMap() : nullptr;
+
+	if (MusicMap)
+	{
+		FString NumeratorString;
+		FString DenominatorString;
+		if (ensure(NewValue.ToString().Split("/", &NumeratorString, &DenominatorString)))
+		{
+			FFrameBasedTimeSignature OriginalTimeSignature = MusicMap->GetInitialTimeSignature();
+			FFrameBasedTimeSignature NewTimeSignature;
+			NewTimeSignature.Numerator = FCString::Atoi(*NumeratorString);
+			NewTimeSignature.Denominator = FCString::Atoi(*DenominatorString);
+			if (OriginalTimeSignature != NewTimeSignature)
+			{
+				FocusedMovieScene->Modify();
+				FocusedMovieScene->GetMusicMap()->SetInitialTimeSignature(NewTimeSignature);
+			}
+		}
+	}
+}
+
+bool SSequencerMusicModeCombo::ValidateTimeSignatureString(const FText& InText, FText& OutErrorMessage) const
+{
+	FString NumeratorString;
+	FString DenominatorString;
+	bool Success = InText.ToString().Split("/", &NumeratorString, &DenominatorString);
+	if (Success)
+	{
+		int32 NewNumerator = FCString::Atoi(*NumeratorString);
+		bool Valid = NewNumerator > 0 && NewNumerator < 100;
+		if (!Valid)
+		{
+			OutErrorMessage = LOCTEXT("MusicModeSettings_Error_Numerator", "Time signature numerator must be from 1 to 99.");
+			return false;
+		}
+		int32 NewDenominator = FCString::Atoi(*DenominatorString);
+		Valid = NewDenominator > 0 && NewDenominator < 100 && FMath::IsPowerOfTwo(NewDenominator);
+		if (!Valid)
+		{
+			OutErrorMessage = LOCTEXT("MusicModeSettings_Error_Denominator", "Time signature denominator must be from 1 to 64 and be a power of 2.");
+			return false;
+		}
+		return true;
+	}
+	OutErrorMessage = LOCTEXT("MusicModeSettings_Error_TimeSignature", "No '/' detected. Time signature must be in the form '4/4', '7/8', etc.");
+	return false;
+}
+
+bool SSequencerMusicModeCombo::SimpleMapEntryEnabled() const
+{
+	if (!GetIsSequenceReadOnly())
+	{
+		TSharedPtr<FSequencer> Sequencer         = WeakSequencer.Pin();
+		UMovieSceneSequence*   FocusedSequence   = Sequencer.IsValid() ? Sequencer->GetFocusedMovieSceneSequence() : nullptr;
+		UMovieScene*           FocusedMovieScene = FocusedSequence ? FocusedSequence->GetMovieScene() : nullptr;
+		UFrameBasedMusicMap*   MusicMap          = FocusedMovieScene ? FocusedMovieScene->GetMusicMap() : nullptr;
+
+		if (MusicMap)
+		{
+			return MusicMap->IsSimple();
+		}
+	}
+	return false;
+}
+
+
+bool SSequencerMusicModeCombo::PromptForMarkerImport(const FMarkerProviderResults& Markers)
+{
+	TSharedPtr<FSequencer> Sequencer         = WeakSequencer.Pin();
+	UMovieSceneSequence*   FocusedSequence   = Sequencer.IsValid() ? Sequencer->GetFocusedMovieSceneSequence() : nullptr;
+	UMovieScene*           FocusedMovieScene = FocusedSequence ? FocusedSequence->GetMovieScene() : nullptr;
+
+	if (!FocusedMovieScene || Markers.IsEmpty())
+	{
+		return false;
+	}
+
+	FSequencerMarkerImportConfig Results = SSequencerMarkerImportDialog::GetMarkerImportOptions(Markers.Channels[0].Markers);
+	if (!Results.bImport)
+	{
+		return false;
+	}
+
+	if (Results.bClearPreviouslyImported)
+	{
+		const TArray<FMovieSceneMarkedFrame>& ExistingMarkedFrames = FocusedMovieScene->GetMarkedFrames();
+		for (int32 Index = ExistingMarkedFrames.Num() - 1; Index >= 0; --Index)
+		{
+			if (ExistingMarkedFrames[Index].Comment.Compare(TEXT("From Imported Music Map")) == 0)
+			{
+				FocusedMovieScene->DeleteMarkedFrame(Index);
+			}
+		}
+	}
+
+	int32 ImportCount = 0;
+	FocusedMovieScene->Modify();
+	for (const FMarkerProviderEntry& Entry : Markers.Channels[0].Markers)
+	{
+		const TArray<FMovieSceneMarkedFrame>& ExistingMarkedFrames = FocusedMovieScene->GetMarkedFrames();
+		bool bAddIt = true;
+		if (!Results.bClearPreviouslyImported)
+		{
+			for (int32 Index = 0; Index < ExistingMarkedFrames.Num(); ++Index)
+			{
+				if (ExistingMarkedFrames[Index].Label == Entry.Label && ExistingMarkedFrames[Index].Comment.Compare(TEXT("From Imported Music Map")) == 0)
+				{
+					// Label already exists. 
+					// If it is already in the correct place move on...
+					if (ExistingMarkedFrames[Index].FrameNumber == Entry.FrameNumber)
+					{
+						bAddIt = false;
+						break;
+					}
+					// Otherwise delete the old one...
+					FocusedMovieScene->DeleteMarkedFrame(Index);
+					break;
+				}
+			}
+		}
+		if (!bAddIt)
+		{
+			continue;
+		}
+
+		FMovieSceneMarkedFrame MarkedFrame(Entry.FrameNumber);
+		MarkedFrame.Label = Entry.Label;
+		MarkedFrame.Comment = TEXT("From Imported Music Map");
+		FocusedMovieScene->AddMarkedFrame(MarkedFrame);
+		++ImportCount;
+	}
+	return ImportCount > 0;
+}
+
+#undef LOCTEXT_NAMESPACE
\ No newline at end of file
diff --git a/Engine/Source/Editor/Sequencer/Private/SSequencerMusicModeCombo.h b/Engine/Source/Editor/Sequencer/Private/SSequencerMusicModeCombo.h
new file mode 100644
--- /dev/null
+++ b/Engine/Source/Editor/Sequencer/Private/SSequencerMusicModeCombo.h
@@ -0,0 +1,60 @@
+// Copyright Epic Games, Inc. All Rights Reserved.
+
+#pragma once
+
+#include "CoreTypes.h"
+#include "Templates/SharedPointer.h"
+#include "Containers/ArrayView.h"
+#include "Styling/SlateTypes.h"
+#include "Widgets/DeclarativeSyntaxSupport.h"
+#include "Widgets/SCompoundWidget.h"
+#include "Styling/AppStyle.h"
+#include "MovieSceneFwd.h"
+#include "SSequencerMarkerImportDialog.h"
+
+class SSequencer;
+class FSequencer;
+class FMenuBuilder;
+class UFrameBasedMusicMap;
+struct FMarkerProviderResults;
+
+class SSequencerMusicModeCombo : public SCompoundWidget
+{
+public:
+
+	SLATE_BEGIN_ARGS(SSequencerMusicModeCombo)
+		: _StyleSet(&FAppStyle::Get())
+		, _StyleName("SlimToolBar")
+	{}
+
+		SLATE_ARGUMENT(const ISlateStyle*, StyleSet)
+
+		SLATE_ARGUMENT(FName, StyleName)
+
+	SLATE_END_ARGS()
+
+	void Construct(const FArguments& InArgs, TWeakPtr<FSequencer> Sequencer, TWeakPtr<SSequencer> InWeakSequencerWidget);
+
+private:
+
+	bool GetIsSequenceReadOnly() const;
+
+	TSharedRef<SWidget> OnCreateMenu();
+	void PopulateSimpleMusicMapSection(FMenuBuilder& MenuBuilder, const UFrameBasedMusicMap* MusicMap);
+	void PopulateImportMusicMapSection(FMenuBuilder& MenuBuilder);
+
+	FText GetTempoTimeSignatureText() const;
+	FText GetToolTipText() const;
+	void OnRemoveAllMusicMaps();
+	void OnCreateSimpleMusicMap();
+	void OnImportMusicMap(UObject* Object);
+	void OnTempoTextCommitted(const FText& NewValue, ETextCommit::Type);
+	void OnTimeSignatureTextCommitted(const FText& NewValue, ETextCommit::Type);
+	bool ValidateTimeSignatureString(const FText& InText, FText& OutErrorMessage) const;
+	bool SimpleMapEntryEnabled() const;
+	bool PromptForMarkerImport(const FMarkerProviderResults& Markers);
+
+	/** Sequencer pointer */
+	TWeakPtr<FSequencer> WeakSequencer;
+	TWeakPtr<SSequencer> WeakSequencerWidget;
+};
\ No newline at end of file
diff --git a/Engine/Source/Editor/Sequencer/Private/SSequencerPlayRateCombo.cpp b/Engine/Source/Editor/Sequencer/Private/SSequencerPlayRateCombo.cpp
--- a/Engine/Source/Editor/Sequencer/Private/SSequencerPlayRateCombo.cpp
+++ b/Engine/Source/Editor/Sequencer/Private/SSequencerPlayRateCombo.cpp
@@ -429,7 +429,23 @@ void SSequencerPlayRateCombo::PopulateCustomClockSourceMenu(FMenuBuilder& MenuBu
 	{
 		auto IsActorValid = [](const AActor* Actor)
 		{
-			return Actor && Actor->GetClass()->ImplementsInterface(UMovieSceneCustomClockSource::StaticClass());
+			if (!Actor)
+			{
+				return false;
+			}
+
+			if (Actor->GetClass()->ImplementsInterface(UMovieSceneCustomClockSource::StaticClass()))
+			{
+				return true;
+			}
+
+			TArray<UActorComponent*> ComponentsWithClockSource = Actor->GetComponentsByInterface(UMovieSceneCustomClockSource::StaticClass());
+			if (ComponentsWithClockSource.Num() > 0)
+			{
+				return true;
+			}
+
+			return false;
 		};
 
 		// Set up a menu entry to assign an actor to the object binding node
diff --git a/Engine/Source/Editor/Sequencer/Private/Sequencer.cpp b/Engine/Source/Editor/Sequencer/Private/Sequencer.cpp
--- a/Engine/Source/Editor/Sequencer/Private/Sequencer.cpp
+++ b/Engine/Source/Editor/Sequencer/Private/Sequencer.cpp
@@ -56,6 +56,7 @@
 #include "Widgets/Images/SImage.h"
 #include "Widgets/Layout/SBox.h"
 #include "Widgets/Input/SButton.h"
+#include "Widgets/Input/SComboBox.h"
 #include "Styling/AppStyle.h"
 #include "Editor/UnrealEdEngine.h"
 #include "Camera/CameraActor.h"
@@ -167,6 +168,7 @@
 #include "Sidebar/SidebarDrawerConfig.h"
 #include "Widgets/Sidebar/SequencerSelectionDrawer.h"
 #include "Misc/Thumbnail/ThumbnailCaptureUtils.h"
+#include "FrameBasedMusicMap.h"
 
 #include "EngineModule.h"
 #include "IViewportSelectableObject.h"
@@ -961,9 +963,14 @@ void FSequencer::Tick(float InDeltaTime)
 
 			SuppressAutoEvalSignature.Reset();
 
-			if (GetNumericTypeInterface()->GetOnSettingChanged() != nullptr)
+			if (GetNumericPositionTypeInterface()->GetOnSettingChanged() != nullptr)
 			{
-				GetNumericTypeInterface()->GetOnSettingChanged()->Broadcast();
+				GetNumericPositionTypeInterface()->GetOnSettingChanged()->Broadcast();
+			}
+
+			if (GetNumericDurationTypeInterface()->GetOnSettingChanged() != nullptr)
+			{
+				GetNumericDurationTypeInterface()->GetOnSettingChanged()->Broadcast();
 			}
 		}
 	}
@@ -1023,7 +1030,17 @@ void FSequencer::Tick(float InDeltaTime)
 
 	if (PlaybackState == EMovieScenePlayerStatus::Playing)
 	{
-		FFrameTime NewGlobalTime = TimeController->RequestCurrentTime(GlobalTime, PlaybackSpeed * Dilation, GetFocusedDisplayRate());
+		FFrameTime NewGlobalTime;
+		if (MovieScene->HasMusicMap() && TimeController->SupportsMusicalTime())
+		{
+			FMusicalTime MusicalTime = TimeController->RequestCurrentMusicalTime(GlobalTime, PlaybackSpeed * Dilation, GetFocusedDisplayRate());
+			FFrameTime FrameTime = MovieScene->GetMusicMap()->GetFrameTime(MusicalTime);
+			NewGlobalTime = FFrameRate::TransformTime(FrameTime, MovieScene->GetMusicMap()->FrameResolution, GetFocusedDisplayRate());
+		}
+		else
+		{
+			NewGlobalTime = TimeController->RequestCurrentTime(GlobalTime, PlaybackSpeed * Dilation, GetFocusedDisplayRate());
+		}
 
 		// Put the time into clamped local space
 		FMovieSceneTransformBreadcrumbs Breadcrumbs;
@@ -3962,9 +3979,19 @@ ISequencer::FOnCameraCut& FSequencer::OnCameraCut()
 	return OnCameraCutEvent;
 }
 
-TSharedRef<INumericTypeInterface<double>> FSequencer::GetNumericTypeInterface() const
+TSharedRef<INumericTypeInterface<double>> FSequencer::GetNumericPositionTypeInterface() const
+{
+	return SequencerWidget->GetNumericPositionTypeInterface();
+}
+
+TSharedRef<INumericTypeInterface<double>> FSequencer::GetNumericDurationTypeInterface() const
 {
-	return SequencerWidget->GetNumericTypeInterface();
+	return SequencerWidget->GetNumericDurationTypeInterface();
+}
+
+TSharedRef<INumericTypeInterface<double>> FSequencer::GetNumericNonMusicalTypeInterface() const
+{
+	return SequencerWidget->GetNumericNonMusicalTypeInterface();
 }
 
 TSharedRef<SWidget> FSequencer::MakeTimeRange(const TSharedRef<SWidget>& InnerContent, bool bShowWorkingRange, bool bShowViewRange, bool bShowPlaybackRange)
@@ -4418,9 +4445,13 @@ void FSequencer::ResetPerMovieSceneData()
 
 	SuppressAutoEvalSignature.Reset();
 
-	if (GetNumericTypeInterface()->GetOnSettingChanged() != nullptr)
+	if (GetNumericPositionTypeInterface()->GetOnSettingChanged() != nullptr)
+	{
+		GetNumericPositionTypeInterface()->GetOnSettingChanged()->Broadcast();
+	}
+	if (GetNumericDurationTypeInterface()->GetOnSettingChanged() != nullptr)
 	{
-		GetNumericTypeInterface()->GetOnSettingChanged()->Broadcast();
+		GetNumericDurationTypeInterface()->GetOnSettingChanged()->Broadcast();
 	}
 	// @todo run through all tracks for new movie scene changes
 	//  needed for audio track decompression
@@ -4529,8 +4560,8 @@ TSharedRef<SWidget> FSequencer::MakePlayTimeDisplay(const TSharedRef<INumericTyp
 		.MinDesiredWidth_Lambda([this] {
 			TRange<double> ViewRange = GetViewRange();
 
-			FString LowerBoundStr = GetNumericTypeInterface()->ToString(ViewRange.GetLowerBoundValue());
-			FString UpperBoundStr = GetNumericTypeInterface()->ToString(ViewRange.GetUpperBoundValue());
+			FString LowerBoundStr = GetNumericPositionTypeInterface()->ToString(ViewRange.GetLowerBoundValue());
+			FString UpperBoundStr = GetNumericPositionTypeInterface()->ToString(ViewRange.GetUpperBoundValue());
 
 			// Always measure with the negative and subframe indicator so that the size doesn't change when there is and isn't a subframe
 			if (!LowerBoundStr.Contains(TEXT("*")))
@@ -5690,7 +5721,7 @@ FString FSequencer::GetFrameTimeText() const
 	}
 
 	FFrameTime LocalTime = RootToParentChainTransform->TransformTime(GetGlobalTime().Time, FTransformTimeParams().IgnoreClamps());
-	return GetNumericTypeInterface()->ToString(LocalTime.GetFrame().Value);
+	return GetNumericPositionTypeInterface()->ToString(LocalTime.GetFrame().Value);
 }
 
 FMovieSceneSequenceID FSequencer::GetScrubPositionParent() const
@@ -7008,7 +7039,12 @@ void FSequencer::OnSelectedNodesOnlyChanged()
 
 void FSequencer::OnTimeDisplayFormatChanged()
 {
-	TSharedPtr<FFrameNumberInterface> FrameNumberInterface = StaticCastSharedPtr<FFrameNumberInterface>(GetNumericTypeInterface().ToSharedPtr());
+	TSharedPtr<FFrameNumberInterface> FrameNumberInterface = StaticCastSharedPtr<FFrameNumberInterface>(GetNumericPositionTypeInterface().ToSharedPtr());
+	if (FrameNumberInterface.IsValid())
+	{
+		FrameNumberInterface->DisplayFormatChanged();
+	}
+	FrameNumberInterface = StaticCastSharedPtr<FFrameNumberInterface>(GetNumericDurationTypeInterface().ToSharedPtr());
 	if (FrameNumberInterface.IsValid())
 	{
 		FrameNumberInterface->DisplayFormatChanged();
@@ -9882,7 +9918,7 @@ void FSequencer::SetKeyTime()
 	}
 
 	// Create a popup showing the existing time value and let the user set a new one.
- 	GenericTextEntryModeless(NSLOCTEXT("Sequencer.Popups", "SetKeyTimePopup", "New Time"), FText::FromString(GetNumericTypeInterface()->ToString(KeyTime.Value)),
+ 	GenericTextEntryModeless(NSLOCTEXT("Sequencer.Popups", "SetKeyTimePopup", "New Time"), FText::FromString(GetNumericPositionTypeInterface()->ToString(KeyTime.Value)),
  		FOnTextCommitted::CreateSP(this, &FSequencer::OnSetKeyTimeTextCommitted)
  	);
 }
@@ -9897,7 +9933,7 @@ void FSequencer::OnSetKeyTimeTextCommitted(const FText& InText, ETextCommit::Typ
 	CloseEntryPopupMenu();
 	if (CommitInfo == ETextCommit::OnEnter)
 	{
-		TOptional<double> NewFrameTime = GetNumericTypeInterface()->FromString(InText.ToString(), 0);
+		TOptional<double> NewFrameTime = GetNumericPositionTypeInterface()->FromString(InText.ToString(), 0);
 		if (!NewFrameTime.IsSet())
 			return;
 
@@ -9928,6 +9964,46 @@ void FSequencer::OnSetKeyTimeTextCommitted(const FText& InText, ETextCommit::Typ
 	}
 }
 
+void FSequencer::QuantizeSelectedKeyTimes(int32 QuantizationInterval)
+{
+	using namespace UE::Sequencer;
+
+	bool bAnythingChanged = false;
+
+	const UFrameBasedMusicMap* MusicMap = GetFocusedMusicMap();
+	FScopedTransaction SetKeyTimeTransaction(NSLOCTEXT("Sequencer", "SetKeyTime_Transaction", "Set Key Time"));
+	const FKeySelection& KeySelection = ViewModel->GetSelection()->KeySelection;
+	for (FKeyHandle Key : KeySelection)
+	{
+		TSharedPtr<FChannelModel> Channel = KeySelection.GetModelForKey(Key);
+		UMovieSceneSection* Section = Channel ? Channel->GetSection() : nullptr;
+		if (Channel && Section)
+		{
+			if (Section->TryModify())
+			{
+				FFrameTime FrameTime = Channel->GetKeyArea()->GetKeyTime(Key);
+				FMusicalTime MusicalTime = MusicMap->GetMusicalTime(FrameTime);
+				FMusicalTime QuantizedMusicalTime = MusicMap->Quantize(MusicalTime, QuantizationInterval);
+				FrameTime = MusicMap->GetFrameTime(QuantizedMusicalTime);
+				Channel->GetKeyArea()->SetKeyTime(Key, FrameTime.GetFrame());
+				bAnythingChanged = true;
+				Section->ExpandToFrame(FrameTime.GetFrame());
+			}
+		}
+	}
+
+	if (bAnythingChanged)
+	{
+		NotifyMovieSceneDataChanged(EMovieSceneDataChangeType::TrackValueChanged);
+	}
+}
+
+bool FSequencer::CanQuantizeKeyTime() const
+{
+	const UFrameBasedMusicMap* MusicMap = GetFocusedMusicMap();
+	return MusicMap && !MusicMap->IsEmpty() && ViewModel->GetSelection()->KeySelection.Num() > 0;
+}
+
 bool FSequencer::CanRekey() const
 {
 	return ViewModel->GetSelection()->KeySelection.Num() > 0;
@@ -12080,7 +12156,7 @@ bool FSequencer::GetGridMetrics(const float PhysicalWidth, const double InViewSt
 	TSharedRef<FSlateFontMeasure> FontMeasureService = FSlateApplication::Get().GetRenderer()->GetFontMeasureService();
 
 	// Use the end of the view as the longest number
-	FString TickString = GetNumericTypeInterface()->ToString((InViewEnd * GetFocusedDisplayRate()).FrameNumber.Value);
+	FString TickString = GetNumericPositionTypeInterface()->ToString((InViewEnd * GetFocusedDisplayRate()).FrameNumber.Value);
 	FVector2D MaxTextSize = FontMeasureService->Measure(TickString, SmallLayoutFont);
 
 	static float MajorTickMultiplier = 2.f;
diff --git a/Engine/Source/Editor/Sequencer/Private/Sequencer.h b/Engine/Source/Editor/Sequencer/Private/Sequencer.h
--- a/Engine/Source/Editor/Sequencer/Private/Sequencer.h
+++ b/Engine/Source/Editor/Sequencer/Private/Sequencer.h
@@ -597,6 +597,10 @@ public:
 	void SetKeyTime();
 	void OnSetKeyTimeTextCommitted(const FText& InText, ETextCommit::Type CommitInfo);
 
+	/** Called when a user executes the quantize key time for selected keys */
+	void QuantizeSelectedKeyTimes(int32 QuantizationInterval);
+	bool CanQuantizeKeyTime() const;
+
 	/** Called when a user executes the rekey for selected keys */
 	bool CanRekey() const;
 	void Rekey();
@@ -818,7 +822,9 @@ public:
 	virtual FOnPostSave& OnPostSave() override;
 	virtual FOnActivateSequence& OnActivateSequence() override;
 	virtual FOnCameraCut& OnCameraCut() override;
-	virtual TSharedRef<INumericTypeInterface<double>> GetNumericTypeInterface() const override;
+	virtual TSharedRef<INumericTypeInterface<double>> GetNumericPositionTypeInterface() const override;
+	virtual TSharedRef<INumericTypeInterface<double>> GetNumericDurationTypeInterface() const override;
+	virtual TSharedRef<INumericTypeInterface<double>> GetNumericNonMusicalTypeInterface() const override;
 	virtual TSharedRef<SWidget> MakeTransportControls(bool bExtended) override;
 	virtual FReply OnPlay(bool bTogglePlay = true) override;
 	virtual void Pause() override;
diff --git a/Engine/Source/Editor/Sequencer/Private/SequencerCommonHelpers.cpp b/Engine/Source/Editor/Sequencer/Private/SequencerCommonHelpers.cpp
--- a/Engine/Source/Editor/Sequencer/Private/SequencerCommonHelpers.cpp
+++ b/Engine/Source/Editor/Sequencer/Private/SequencerCommonHelpers.cpp
@@ -24,6 +24,7 @@
 #include "MovieScene.h"
 #include "MovieSceneSectionDetailsCustomization.h"
 #include "MovieSceneSequence.h"
+#include "Sections/MovieSceneSubSection.h"
 #include "PropertyEditorModule.h"
 #include "PropertyPermissionList.h"
 #include "SSequencer.h"
@@ -362,7 +363,30 @@ void SequencerHelpers::BuildEditSectionMenu(const TWeakPtr<FSequencer>& InWeakSe
 
 		UMovieSceneSequence* Sequence = Sequencer->GetFocusedMovieSceneSequence();
 		const TWeakObjectPtr<UMovieScene> CurrentScene = Sequence->GetMovieScene();
-		const TWeakPtr<INumericTypeInterface<double>> NumericTypeInterfafce = Sequencer->GetNumericTypeInterface();
+		const TWeakPtr<INumericTypeInterface<double>> NumericPositionTypeInterface = Sequencer->GetNumericPositionTypeInterface();
+		const TWeakPtr<INumericTypeInterface<double>> NumericDurationTypeInterface = Sequencer->GetNumericDurationTypeInterface();
+		const TWeakPtr<INumericTypeInterface<double>> NumericNonMusicalTypeInterface = Sequencer->GetNumericNonMusicalTypeInterface();
+
+		// If all of the selected sections are musical, we can show durations as bars.beats, but if not
+		// we need to show durations as seconds...
+		bool bAllSequencesAreMusical = true;
+		for (TWeakObjectPtr<> Section : InWeakSections)
+		{
+			if (Section.IsValid())
+			{
+				const UMovieSceneSubSection* SubSection = Cast<UMovieSceneSubSection>(Section);
+				if (SubSection)
+				{
+					UMovieSceneSequence* SubSequence = SubSection->GetSequence();
+					UMovieScene* MovieScene = SubSequence->GetMovieScene();
+					if (!MovieScene->HasMusicMap())
+					{
+						bAllSequencesAreMusical = false;
+						break;
+					}
+				}
+			}
+		}
 
 		TSharedRef<SSectionDetailsNotifyHookWrapper> DetailsNotifyWrapper = SNew(SSectionDetailsNotifyHookWrapper);
 		FDetailsViewArgs DetailsViewArgs;
@@ -385,21 +409,25 @@ void SequencerHelpers::BuildEditSectionMenu(const TWeakPtr<FSequencer>& InWeakSe
 
 		TSharedRef<IDetailsView> DetailsView = FModuleManager::GetModuleChecked<FPropertyEditorModule>("PropertyEditor").CreateDetailView(DetailsViewArgs);
 		DetailsView->RegisterInstancedCustomPropertyTypeLayout("FrameNumber",
-			FOnGetPropertyTypeCustomizationInstance::CreateLambda([NumericTypeInterfafce]()
+			FOnGetPropertyTypeCustomizationInstance::CreateLambda([NumericPositionTypeInterface, NumericDurationTypeInterface, NumericNonMusicalTypeInterface, bAllSequencesAreMusical]()
 			{
-				return MakeShared<FFrameNumberDetailsCustomization>(NumericTypeInterfafce.Pin());
+				// If all of the selected sections are musical, we can show durations as bars.beats, but if not
+				// we need to show durations as seconds...
+				return MakeShared<FFrameNumberDetailsCustomization>(NumericPositionTypeInterface.Pin(),
+				bAllSequencesAreMusical ? NumericDurationTypeInterface.Pin() : NumericNonMusicalTypeInterface.Pin(),
+				NumericNonMusicalTypeInterface.Pin());
 			}));
 		DetailsView->RegisterInstancedCustomPropertyLayout(UMovieSceneSection::StaticClass(),
-			FOnGetDetailCustomizationInstance::CreateLambda([NumericTypeInterfafce, CurrentScene]()
+			FOnGetDetailCustomizationInstance::CreateLambda([NumericPositionTypeInterface, NumericDurationTypeInterface, NumericNonMusicalTypeInterface, CurrentScene]()
 			{
-				return MakeShared<FMovieSceneSectionDetailsCustomization>(NumericTypeInterfafce.Pin(), CurrentScene.Get());
+				return MakeShared<FMovieSceneSectionDetailsCustomization>(NumericPositionTypeInterface.Pin(), NumericDurationTypeInterface.Pin(), NumericNonMusicalTypeInterface.Pin(), CurrentScene.Get());
 			}));
 
-	DetailsView->RegisterInstancedCustomPropertyTypeLayout("MovieSceneConditionContainer", FOnGetPropertyTypeCustomizationInstance::CreateLambda([Sequence, InWeakSequencer]() {
-		return FMovieSceneConditionCustomization::MakeInstance(Sequence, InWeakSequencer); }));
+		DetailsView->RegisterInstancedCustomPropertyTypeLayout("MovieSceneConditionContainer", FOnGetPropertyTypeCustomizationInstance::CreateLambda([Sequence, InWeakSequencer]() {
+			return FMovieSceneConditionCustomization::MakeInstance(Sequence, InWeakSequencer); }));
 
-	DetailsView->RegisterInstancedCustomPropertyTypeLayout("MovieSceneDirectorBlueprintConditionData", FOnGetPropertyTypeCustomizationInstance::CreateLambda([=]() {
-		return FMovieSceneDirectorBlueprintConditionCustomization::MakeInstance(CurrentScene.Get());}));
+		DetailsView->RegisterInstancedCustomPropertyTypeLayout("MovieSceneDirectorBlueprintConditionData", FOnGetPropertyTypeCustomizationInstance::CreateLambda([=]() {
+			return FMovieSceneDirectorBlueprintConditionCustomization::MakeInstance(CurrentScene.Get());}));
 	
 		DetailsView->SetIsPropertyVisibleDelegate(FIsPropertyVisible::CreateLambda([](const FPropertyAndParent& PropertyAndParent)
 			{			
diff --git a/Engine/Source/Editor/Sequencer/Private/SequencerContextMenus.cpp b/Engine/Source/Editor/Sequencer/Private/SequencerContextMenus.cpp
--- a/Engine/Source/Editor/Sequencer/Private/SequencerContextMenus.cpp
+++ b/Engine/Source/Editor/Sequencer/Private/SequencerContextMenus.cpp
@@ -50,7 +50,7 @@
 #include "Tracks/MovieScenePropertyTrack.h"
 #include "Algo/AnyOf.h"
 #include "IKeyArea.h"
-
+#include "FrameBasedMusicMap.h"
 
 #define LOCTEXT_NAMESPACE "SequencerContextMenus"
 
@@ -187,6 +187,17 @@ void FKeyContextMenu::PopulateMenu(FMenuBuilder& MenuBuilder, TSharedPtr<FExtend
 	MenuBuilder.BeginSection("SequencerKeys", LOCTEXT("KeysMenu", "Keys"));
 	{
 		MenuBuilder.AddMenuEntry(FSequencerCommands::Get().SetKeyTime);
+		if (Sequencer->CanQuantizeKeyTime())
+		{
+			MenuBuilder.AddSubMenu(
+				LOCTEXT("KeyQuantizeTime", "Quantize Key Time"),
+				LOCTEXT("KeyQuantizeTimeTooltip", "Quantize the key time to a musical subdivision"),
+				FNewMenuDelegate::CreateLambda([=](FMenuBuilder& SubMenuBuilder) { Shared->AddQuantizationMenu(SubMenuBuilder); }),
+				FUIAction(),
+				NAME_None,
+				EUserInterfaceActionType::Button
+			);
+		}
 		MenuBuilder.AddMenuEntry(FSequencerCommands::Get().Rekey);
 		MenuBuilder.AddMenuEntry(FSequencerCommands::Get().SnapToFrame);
 		MenuBuilder.AddMenuEntry(FSequencerCommands::Get().DeleteKeys);
@@ -215,6 +226,38 @@ void FKeyContextMenu::AddPropertiesMenu(FMenuBuilder& MenuBuilder)
 		, FText::GetEmpty(), true);
 }
 
+void FKeyContextMenu::AddQuantizationMenu(FMenuBuilder& MenuBuilder)
+{
+	const TSharedPtr<FSequencer> Sequencer = WeakSequencer.Pin();
+	if (!Sequencer.IsValid())
+	{
+		return;
+	}
+
+	for (auto& Item : FrameBasedMusicMap::QuantizationOptionsSource)
+	{
+		int32 QunatizationInterval = *Item;
+		auto LocalSequencerPointer = WeakSequencer;
+		MenuBuilder.AddMenuEntry(
+			FrameBasedMusicMap::GetQuantizationName(*Item),
+			FText(),
+			FSlateIcon(),
+			FUIAction(FExecuteAction::CreateLambda([=]
+				()
+				{
+					const TSharedPtr<FSequencer> Sequencer = LocalSequencerPointer.Pin();
+					if (!Sequencer.IsValid())
+					{
+						return;
+					}
+					Sequencer->QuantizeSelectedKeyTimes(QunatizationInterval);
+				}
+				)
+			)
+			);
+	}
+}
+
 
 FSectionContextMenu::FSectionContextMenu(TWeakPtr<FSequencer> InWeakSequencer, FFrameTime InMouseDownTime)
 	: WeakSequencer(InWeakSequencer)
@@ -2050,7 +2093,7 @@ void FEasingContextMenu::PopulateMenu(FMenuBuilder& MenuBuilder, TSharedPtr<FExt
 					.OnBeginSliderMovement_Lambda(OnBeginSliderMovement)
 					.OnEndSliderMovement_Lambda(OnEndSliderMovement)
 					.BorderForegroundColor(FAppStyle::GetSlateColor("DefaultForeground"))
-					.TypeInterface(Sequencer->GetNumericTypeInterface())
+					.TypeInterface(Sequencer->GetNumericPositionTypeInterface())
 				]
 			]
 
diff --git a/Engine/Source/Editor/Sequencer/Private/SequencerContextMenus.h b/Engine/Source/Editor/Sequencer/Private/SequencerContextMenus.h
--- a/Engine/Source/Editor/Sequencer/Private/SequencerContextMenus.h
+++ b/Engine/Source/Editor/Sequencer/Private/SequencerContextMenus.h
@@ -228,6 +228,10 @@ private:
 	/** Add the Properties sub-menu. */
 	void AddPropertiesMenu(FMenuBuilder& MenuBuilder);
 
+	/** Add the Quantization sub-menu. */
+	void AddQuantizationMenu(FMenuBuilder& MenuBuilder);
+	void QuantizeKeys(int32 QualtizationInterval);
+
 	void PopulateMenu(FMenuBuilder& MenuBuilder, TSharedPtr<FExtender> MenuExtender);
 
 	/** The sequencer */
diff --git a/Engine/Source/Editor/Sequencer/Private/SequencerTimeSliderController.cpp b/Engine/Source/Editor/Sequencer/Private/SequencerTimeSliderController.cpp
--- a/Engine/Source/Editor/Sequencer/Private/SequencerTimeSliderController.cpp
+++ b/Engine/Source/Editor/Sequencer/Private/SequencerTimeSliderController.cpp
@@ -41,6 +41,7 @@
 #include "ISinglePropertyView.h"
 #include "IStructureDetailsView.h"
 #include "FrameNumberDetailsCustomization.h"
+#include "FrameBasedMusicMap.h"
 
 #define LOCTEXT_NAMESPACE "TimeSlider"
 
@@ -244,6 +245,25 @@ struct FSequencerTimeSliderController::FDrawTickArgs
 };
 
 void FSequencerTimeSliderController::DrawTicks( FSlateWindowElementList& OutDrawElements, const TRange<double>& ViewRange, const FScrubRangeToScreen& RangeToScreen, FDrawTickArgs& InArgs ) const
+{
+	TSharedPtr<FSequencer> Sequencer = WeakSequencer.Pin();
+	if (!Sequencer.IsValid())
+	{
+		return;
+	}
+	const UFrameBasedMusicMap* MusicMap = Sequencer->GetFocusedMusicMap();
+
+	if (Sequencer->GetSequencerSettings()->GetTimeDisplayFormat() != EFrameNumberDisplayFormats::BarsBeats || !MusicMap || MusicMap->IsEmpty())
+	{
+		DrawLinearTicks(OutDrawElements, ViewRange, RangeToScreen, InArgs);
+	}
+	else
+	{
+		DrawBarBeatTicks(OutDrawElements, ViewRange, RangeToScreen, InArgs);
+	}
+}
+
+void FSequencerTimeSliderController::DrawLinearTicks(FSlateWindowElementList& OutDrawElements, const TRange<double>& ViewRange, const FScrubRangeToScreen& RangeToScreen, FDrawTickArgs& InArgs) const
 {
 	TSharedPtr<FSequencer> Sequencer = WeakSequencer.Pin();
 	if (!Sequencer.IsValid())
@@ -336,6 +356,101 @@ void FSequencerTimeSliderController::DrawTicks( FSlateWindowElementList& OutDraw
 	}
 }
 
+void FSequencerTimeSliderController::DrawBarBeatTicks(FSlateWindowElementList& OutDrawElements, const TRange<double>& ViewRange, const FScrubRangeToScreen& RangeToScreen, FDrawTickArgs& InArgs) const
+{
+	TSharedPtr<FSequencer> Sequencer = WeakSequencer.Pin();
+	if (!Sequencer.IsValid())
+	{
+		return;
+	}
+
+	const UFrameBasedMusicMap* MusicMap = Sequencer->GetFocusedMusicMap();
+	if (!MusicMap || MusicMap->IsEmpty())
+	{
+		return;
+	}
+
+	FFrameRate     TickResolution = GetTickResolution();
+	FFrameRate     DisplayRate = GetDisplayRate();
+	FPaintGeometry PaintGeometry = InArgs.AllottedGeometry.ToPaintGeometry();
+	FSlateFontInfo TickFrameFont = FCoreStyle::GetDefaultFontStyle("Regular", 8);
+
+	double PixelsPerSecond = InArgs.AllottedGeometry.Size.X / (ViewRange.GetUpperBoundValue() - ViewRange.GetLowerBoundValue());
+	double MinSecondGap = 20.0f / PixelsPerSecond;
+
+	FFrameTime StartFrameTime = GetTickResolution().AsFrameTime(ViewRange.GetLowerBoundValue());
+	FFrameTime EndFrameTime = GetTickResolution().AsFrameTime(ViewRange.GetUpperBoundValue());
+	FFrameTime MinSpacing = GetTickResolution().AsFrameTime(MinSecondGap);
+
+	TArray<FSubdivisionPoint> Points = MusicMap->GatherSubdivisonsForTimespan(StartFrameTime, EndFrameTime, MinSpacing, FrameBasedMusicMap::EWholeQuantizationTicks::QuarterNote);
+
+	TArray<FVector2D> LinePoints;
+	LinePoints.SetNumUninitialized(2);
+
+	const bool bAntiAliasLines = false;
+	float LastLabelRightEdge = -10000;
+	for (const FSubdivisionPoint& Point : Points)
+	{
+		float LinePx = RangeToScreen.InputToLocalX(TickResolution.AsSeconds(Point.FrameTime));
+
+		if (Point.MusicalTime.TickInBar == 0)
+		{
+			if (LinePx > (LastLabelRightEdge + 26))
+			{
+				LinePoints[0] = FVector2D(LinePx, InArgs.TickOffset - 25);
+				LinePoints[1] = FVector2D(LinePx, InArgs.TickOffset + InArgs.MajorTickHeight);
+
+				// Space the text between the tick mark but slightly above
+				FVector2D TextOffset(LinePx + 5.f, InArgs.bMirrorLabels ? 1.f : FMath::Abs(InArgs.AllottedGeometry.Size.Y - (30/*InArgs.MajorTickHeight*/ + 9.f)));
+				FSlateDrawElement::MakeText(
+					OutDrawElements,
+					InArgs.StartLayer + 1,
+					InArgs.AllottedGeometry.ToPaintGeometry(InArgs.AllottedGeometry.Size, FSlateLayoutTransform(TextOffset)),
+					Point.PrettyString,
+					TickFrameFont,
+					InArgs.DrawEffects,
+					InArgs.TickColor * 0.65f
+				);
+
+				FVector2D LabelSize = FontMeasureService->Measure(Point.PrettyString, TickFrameFont);
+				LastLabelRightEdge = LinePx + 5.f + LabelSize.Y;
+			}
+			else
+			{
+				LinePoints[0] = FVector2D(LinePx, InArgs.TickOffset);
+				LinePoints[1] = FVector2D(LinePx, InArgs.TickOffset + InArgs.MajorTickHeight);
+			}
+
+			FSlateDrawElement::MakeLines(
+				OutDrawElements,
+				InArgs.StartLayer,
+				PaintGeometry,
+				LinePoints,
+				InArgs.DrawEffects,
+				InArgs.TickColor,
+				bAntiAliasLines
+			);
+		}
+		else
+		{ 
+			const float MinorTickHeight = 2.0f; // : 2.0f;
+
+			LinePoints[0] = FVector2D(LinePx, InArgs.bMirrorLabels ? 0.0f : FMath::Abs(InArgs.AllottedGeometry.Size.Y - MinorTickHeight));
+			LinePoints[1] = FVector2D(LinePx, LinePoints[0].Y + MinorTickHeight);
+
+			FSlateDrawElement::MakeLines(
+				OutDrawElements,
+				InArgs.StartLayer,
+				PaintGeometry,
+				LinePoints,
+				InArgs.DrawEffects,
+				InArgs.TickColor,
+				bAntiAliasLines
+			);
+		}
+	}
+}
+
 int32 FSequencerTimeSliderController::DrawMarkedFrames( const FGeometry& AllottedGeometry, const FScrubRangeToScreen& RangeToScreen, FSlateWindowElementList& OutDrawElements, int32 LayerId, const ESlateDrawEffect& DrawEffects, const FWidgetStyle& InWidgetStyle, bool bDrawLabels ) const
 {
 	if (!WeakSequencer.Pin()->GetSequencerSettings()->GetShowMarkedFrames())
@@ -1047,6 +1162,20 @@ FReply FSequencerTimeSliderController::OnMouseButtonDoubleClick(TSharedRef<const
 	return FReply::Unhandled();
 }
 
+double FSequencerTimeSliderController::ComputeHeight() const
+{
+	TSharedPtr<FSequencer> Sequencer = WeakSequencer.Pin();
+	if (Sequencer.IsValid())
+	{
+		const UFrameBasedMusicMap* MusicMap = Sequencer->GetFocusedMusicMap();
+		if (Sequencer->GetSequencerSettings()->GetTimeDisplayFormat() == EFrameNumberDisplayFormats::BarsBeats && MusicMap && !MusicMap->IsEmpty())
+		{
+			return 55.0;
+		}
+	}
+	return 22.0;
+}
+
 FReply FSequencerTimeSliderController::OnMouseMove( SWidget& WidgetOwner, const FGeometry& MyGeometry, const FPointerEvent& MouseEvent )
 {
 	return OnMouseMoveImpl(WidgetOwner, MyGeometry, MouseEvent, false);
@@ -1928,7 +2057,7 @@ TSharedRef<SWidget> FSequencerTimeSliderController::OpenSetPlaybackRangeMenu(con
 
 					DetailsView = PropertyEditorModule.CreateStructureDetailView(DetailsViewArgs, StructureDetailsViewArgs, nullptr);
 					DetailsView->GetDetailsView()->RegisterInstancedCustomPropertyTypeLayout("FrameNumber", FOnGetPropertyTypeCustomizationInstance::CreateLambda([this]() {
-						return MakeShared<FFrameNumberDetailsCustomization>(WeakSequencer.Pin()->GetNumericTypeInterface()); }));
+						return MakeShared<FFrameNumberDetailsCustomization>(WeakSequencer.Pin()->GetNumericPositionTypeInterface()/*, WeakSequencer.Pin()->GetNumericDurationTypeInterface(), WeakSequencer.Pin()->GetNumericNonMusicalTypeInterface()*/); }));
 					DetailsView->SetStructureData(StructOnScope);
 
 					ChildSlot
diff --git a/Engine/Source/Editor/Sequencer/Private/Widgets/Sidebar/SMarkedFrameDetails.h b/Engine/Source/Editor/Sequencer/Private/Widgets/Sidebar/SMarkedFrameDetails.h
--- a/Engine/Source/Editor/Sequencer/Private/Widgets/Sidebar/SMarkedFrameDetails.h
+++ b/Engine/Source/Editor/Sequencer/Private/Widgets/Sidebar/SMarkedFrameDetails.h
@@ -1,4 +1,4 @@
-﻿// Copyright Epic Games, Inc. All Rights Reserved.
+// Copyright Epic Games, Inc. All Rights Reserved.
 
 #pragma once
 
@@ -69,7 +69,7 @@ public:
 		DetailsView->GetDetailsView()->RegisterInstancedCustomPropertyTypeLayout(TEXT("FrameNumber"),
 			FOnGetPropertyTypeCustomizationInstance::CreateLambda([this]()
 			{
-				return MakeShared<FFrameNumberDetailsCustomization>(WeakSequencer.Pin()->GetNumericTypeInterface());
+				return MakeShared<FFrameNumberDetailsCustomization>(WeakSequencer.Pin()->GetNumericPositionTypeInterface());
 			}));
 
 		DetailsView->SetStructureData(StructOnScope);
diff --git a/Engine/Source/Editor/Sequencer/Public/ISequencer.h b/Engine/Source/Editor/Sequencer/Public/ISequencer.h
--- a/Engine/Source/Editor/Sequencer/Public/ISequencer.h
+++ b/Engine/Source/Editor/Sequencer/Public/ISequencer.h
@@ -81,6 +81,7 @@ class IToolkitHost;
 struct FAnimatedRange;
 struct FMovieSceneChannelMetaData;
 struct FMovieSceneSequencePlaybackParams;
+class UFrameBasedMusicMap;
 
 namespace UE
 {
@@ -773,8 +774,14 @@ public:
 	/** Gets a multicast delegate when the curve edtior associated with this sequencer has it's selection change. */
 	virtual FOnCurveDisplayChanged& GetCurveDisplayChanged() = 0;
 
-	/** @return a numeric type interface that will parse and display numbers as frames and times correctly */
-	virtual TSharedRef<INumericTypeInterface<double>> GetNumericTypeInterface() const = 0;
+	/** @return a numeric type interface that will parse and display numbers as frame / time / bar.beat positions correctly */
+	virtual TSharedRef<INumericTypeInterface<double>> GetNumericPositionTypeInterface() const = 0;
+
+	/** @return a numeric type interface that will parse and display numbers as frame / time / bar.beat durations correctly */
+	virtual TSharedRef<INumericTypeInterface<double>> GetNumericDurationTypeInterface() const = 0;
+
+	/** @return a numeric type interface that will parse and display numbers as frame / time, and bar.beat positions/durations are forces to seconds */
+	virtual TSharedRef<INumericTypeInterface<double>> GetNumericNonMusicalTypeInterface() const = 0;
 
 	/** @return the command bindings for this sequencer */
 	virtual TSharedPtr<FUICommandList> GetCommandBindings(ESequencerCommandBindings Type = ESequencerCommandBindings::Sequencer) const = 0;
@@ -901,6 +908,10 @@ public:
 	 */
 	SEQUENCER_API FFrameRate GetFocusedDisplayRate() const;
 
+	/**
+	 * Get the music map from the currently focused sequence
+	 */
+	SEQUENCER_API const UFrameBasedMusicMap* GetFocusedMusicMap() const;
 
 	/**
 	 * Create a new binding for the specified object, lightly deprecated and no longer virtual in favor of the overload with binding parameters
diff --git a/Engine/Source/Editor/Sequencer/Public/SequencerTimeSliderController.h b/Engine/Source/Editor/Sequencer/Public/SequencerTimeSliderController.h
--- a/Engine/Source/Editor/Sequencer/Public/SequencerTimeSliderController.h
+++ b/Engine/Source/Editor/Sequencer/Public/SequencerTimeSliderController.h
@@ -47,6 +47,7 @@ public:
 	virtual FReply OnTimeSliderMouseMove( SWidget& OwnerWidget, const FGeometry& MyGeometry, const FPointerEvent& MouseEvent ) override;
 	virtual FCursorReply OnCursorQuery( TSharedRef<const SWidget> WidgetOwner, const FGeometry& MyGeometry, const FPointerEvent& CursorEvent ) const override;
 	virtual FReply OnMouseButtonDoubleClick( TSharedRef<const SWidget> WidgetOwner, const FGeometry& MyGeometry, const FPointerEvent& MouseEvent ) override;
+	virtual double ComputeHeight() const override;
 	/** End ITimeSliderController Interface */
 
 	/** Get the current play rate for this controller */
@@ -224,6 +225,26 @@ private:
 	 */
 	void DrawTicks( FSlateWindowElementList& OutDrawElements, const TRange<double>& ViewRange, const FScrubRangeToScreen& RangeToScreen, FDrawTickArgs& InArgs ) const;
 
+	/**
+	 * Draw time tick marks
+	 *
+	 * @param OutDrawElements	List to add draw elements to
+	 * @param ViewRange			The currently visible time range in seconds
+	 * @param RangeToScreen		Time range to screen space converter
+	 * @param InArgs			Parameters for drawing the tick lines
+	 */
+	void DrawLinearTicks(FSlateWindowElementList& OutDrawElements, const TRange<double>& ViewRange, const FScrubRangeToScreen& RangeToScreen, FDrawTickArgs& InArgs) const;
+
+	/**
+	 * Draw time tick marks
+	 *
+	 * @param OutDrawElements	List to add draw elements to
+	 * @param ViewRange			The currently visible time range in seconds
+	 * @param RangeToScreen		Time range to screen space converter
+	 * @param InArgs			Parameters for drawing the tick lines
+	 */
+	void DrawBarBeatTicks(FSlateWindowElementList& OutDrawElements, const TRange<double>& ViewRange, const FScrubRangeToScreen& RangeToScreen, FDrawTickArgs& InArgs) const;
+
 	/**
 	 * Draw the selection range.
 	 *
diff --git a/Engine/Source/Editor/SequencerWidgets/Private/SSequencerTimeSlider.cpp b/Engine/Source/Editor/SequencerWidgets/Private/SSequencerTimeSlider.cpp
--- a/Engine/Source/Editor/SequencerWidgets/Private/SSequencerTimeSlider.cpp
+++ b/Engine/Source/Editor/SequencerWidgets/Private/SSequencerTimeSlider.cpp
@@ -49,7 +49,7 @@ FReply SSequencerTimeSlider::OnMouseButtonDoubleClick(const FGeometry& MyGeometr
 
 FVector2D SSequencerTimeSlider::ComputeDesiredSize( float ) const
 {
-	return FVector2D(100, 22);
+	return FVector2D(100, TimeSliderController->ComputeHeight());
 }
 
 FReply SSequencerTimeSlider::OnMouseWheel( const FGeometry& MyGeometry, const FPointerEvent& MouseEvent )
diff --git a/Engine/Source/Editor/SequencerWidgets/Public/ITimeSlider.h b/Engine/Source/Editor/SequencerWidgets/Public/ITimeSlider.h
--- a/Engine/Source/Editor/SequencerWidgets/Public/ITimeSlider.h
+++ b/Engine/Source/Editor/SequencerWidgets/Public/ITimeSlider.h
@@ -108,6 +108,8 @@ public:
 	{
 		return OnMouseWheel(OwnerWidget, MyGeometry, MouseEvent);
 	}
+	
+	virtual double ComputeHeight() const { return 22.0; };
 
 	/** Get the current play rate for this controller */
 	virtual FFrameRate GetDisplayRate() const = 0;
diff --git a/Engine/Source/Editor/UnrealEd/Classes/Settings/EditorExperimentalSettings.h b/Engine/Source/Editor/UnrealEd/Classes/Settings/EditorExperimentalSettings.h
--- a/Engine/Source/Editor/UnrealEd/Classes/Settings/EditorExperimentalSettings.h
+++ b/Engine/Source/Editor/UnrealEd/Classes/Settings/EditorExperimentalSettings.h
@@ -149,6 +149,9 @@ public:
 	UPROPERTY(EditAnywhere, config, Category = LevelInstance)
 	bool bEnableLevelInstancePropertyOverrides;
 
+	UPROPERTY(EditAnywhere, config, Category = Sequencer, meta = (DisplayName = "Enable Music Mode Editing", ToolTip = "Enables \"music mode\" editing features"))
+	bool bMusicModeEditingInSequencerEnabled;
+
 	/**
 	 * Returns an event delegate that is executed when a setting has changed.
 	 *
diff --git a/Engine/Source/Editor/UnrealEd/Private/Settings/SettingsClasses.cpp b/Engine/Source/Editor/UnrealEd/Private/Settings/SettingsClasses.cpp
--- a/Engine/Source/Editor/UnrealEd/Private/Settings/SettingsClasses.cpp
+++ b/Engine/Source/Editor/UnrealEd/Private/Settings/SettingsClasses.cpp
@@ -181,6 +181,12 @@ static TAutoConsoleVariable<float> CVarEditorHDRNITLevel(
 	TEXT("Sets The desired NIT level of the editor when running on HDR"),
 	ECVF_Default);
 
+static TAutoConsoleVariable<int32> CVarSequenceEditorMusicModeEnabled(
+	TEXT("Sequencer.EnableMusicModeEditing"),
+	0,
+	TEXT("Sets whether or not we should allow the sequence editor to show music mode tools"),
+	ECVF_Default);
+
 UEditorExperimentalSettings::UEditorExperimentalSettings( const FObjectInitializer& ObjectInitializer )
 	: Super(ObjectInitializer)
 	, bEnableAsyncTextureCompilation(false)
@@ -193,6 +199,7 @@ UEditorExperimentalSettings::UEditorExperimentalSettings( const FObjectInitializ
 	, HDREditorNITLevel(160.0f)
 	, bUseOpenCLForConvexHullDecomp(false)
 	, bAllowPotentiallyUnsafePropertyEditing(false)
+	, bMusicModeEditingInSequencerEnabled(false)
 {
 }
 
@@ -210,6 +217,7 @@ void UEditorExperimentalSettings::PostInitProperties()
 
 	CVarEditorHDRSupport->Set(bHDREditor ? 1 : 0, ECVF_SetByProjectSetting);
 	CVarEditorHDRNITLevel->Set(HDREditorNITLevel, ECVF_SetByProjectSetting);
+	CVarSequenceEditorMusicModeEnabled->Set(bMusicModeEditingInSequencerEnabled ? 1 : 0, ECVF_SetByProjectSetting);
 	Super::PostInitProperties();
 }
 
@@ -231,6 +239,10 @@ void UEditorExperimentalSettings::PostEditChangeProperty( struct FPropertyChange
 	{
 		CVarEditorHDRNITLevel->Set(HDREditorNITLevel, ECVF_SetByProjectSetting);
 	}
+	else if (Name == GET_MEMBER_NAME_CHECKED(UEditorExperimentalSettings, bMusicModeEditingInSequencerEnabled))
+	{
+		CVarSequenceEditorMusicModeEnabled->Set(bMusicModeEditingInSequencerEnabled ? 1 : 0, ECVF_SetByProjectSetting);
+	}
 
 	if (!FUnrealEdMisc::Get().IsDeletePreferences())
 	{
diff --git a/Engine/Source/Runtime/Core/Public/Misc/MusicalTime.h b/Engine/Source/Runtime/Core/Public/Misc/MusicalTime.h
new file mode 100644
--- /dev/null
+++ b/Engine/Source/Runtime/Core/Public/Misc/MusicalTime.h
@@ -0,0 +1,118 @@
+// Copyright Epic Games, Inc. All Rights Reserved.
+
+#pragma once
+
+#include "CoreTypes.h"
+#include "Serialization/Archive.h"
+
+namespace MusicalTime
+{
+	static constexpr int32 TicksPerQuarterNote = 960;
+}
+
+struct FMusicalTime
+{
+	FMusicalTime()
+		: Bar(0)
+		, TickInBar(0)
+		, TicksPerBar(0)
+		, TicksPerBeat(0)
+	{}
+
+	FMusicalTime(int32 InBar, int32 InTickInBar = 0, int32 InTicksPerBar = MusicalTime::TicksPerQuarterNote * 4, int32 InTicksPerBeat = MusicalTime::TicksPerQuarterNote)
+		: Bar(InBar)
+		, TickInBar(InTickInBar)
+		, TicksPerBar(InTicksPerBar)
+		, TicksPerBeat(InTicksPerBeat)
+	{}
+
+	/** IMPORTANT: If you change the struct data, ensure that you also update the version in NoExportTypes.h  */
+
+	int32 Bar;
+	int32 TickInBar;
+	int32 TicksPerBar;
+	int32 TicksPerBeat;
+
+	/**
+	 * Verify that this musical time is valid to use
+	 */
+	bool IsValid() const
+	{
+		return TicksPerBar > 0 && TicksPerBeat > 0;
+	}
+
+	float BeatInBar() const
+	{
+		return static_cast<float>(TickInBar) / static_cast<float>(TicksPerBeat);
+	}
+
+	double FractionalBar() const
+	{
+		return static_cast<double>(Bar) + (static_cast<double>(TickInBar) / static_cast<double>(TicksPerBar));
+	}
+
+	/**
+	 * Serializes the given musical time from or into the specified archive
+	 */
+	bool Serialize(FArchive& Ar);
+
+	/**
+	 *  return fractional bars between two musical times
+	 */
+	double operator-(const FMusicalTime& Other)
+	{
+		return FractionalBar() - Other.FractionalBar();
+	}
+
+public:
+
+	friend FArchive& operator<<(FArchive& Ar, FMusicalTime& MusicalTime)
+	{
+		MusicalTime.Serialize(Ar);
+		return Ar;
+	}
+
+	friend FORCEINLINE_DEBUGGABLE bool operator==(const FMusicalTime& A, const FMusicalTime& B)
+	{
+		return A.Bar == B.Bar && A.TickInBar == B.TickInBar;
+	}
+
+
+	friend FORCEINLINE_DEBUGGABLE bool operator!=(const FMusicalTime& A, const FMusicalTime& B)
+	{
+		return A.Bar != B.Bar || A.TickInBar != B.TickInBar;
+	}
+
+
+	friend FORCEINLINE_DEBUGGABLE bool operator> (const FMusicalTime& A, const FMusicalTime& B)
+	{
+		return A.Bar >  B.Bar || ( A.Bar == B.Bar && A.TickInBar > B.TickInBar );
+	}
+
+
+	friend FORCEINLINE_DEBUGGABLE bool operator>=(const FMusicalTime& A, const FMusicalTime& B)
+	{
+		return A.Bar > B.Bar || ( A.Bar == B.Bar && A.TickInBar >= B.TickInBar );
+	}
+
+
+	friend FORCEINLINE_DEBUGGABLE bool operator< (const FMusicalTime& A, const FMusicalTime& B)
+	{
+		return A.Bar <  B.Bar || ( A.Bar == B.Bar && A.TickInBar < B.TickInBar );
+	}
+
+
+	friend FORCEINLINE_DEBUGGABLE bool operator<=(const FMusicalTime& A, const FMusicalTime& B)
+	{
+		return A.Bar < B.Bar || ( A.Bar == B.Bar && A.TickInBar <= B.TickInBar );
+	}
+};
+
+inline bool FMusicalTime::Serialize(FArchive& Ar)
+{
+	Ar << Bar;
+	Ar << TickInBar;
+	Ar << TicksPerBar;
+	Ar << TicksPerBeat;
+	return true;
+}
diff --git a/Engine/Source/Runtime/CoreUObject/Public/UObject/NoExportTypes.h b/Engine/Source/Runtime/CoreUObject/Public/UObject/NoExportTypes.h
--- a/Engine/Source/Runtime/CoreUObject/Public/UObject/NoExportTypes.h
+++ b/Engine/Source/Runtime/CoreUObject/Public/UObject/NoExportTypes.h
@@ -2233,6 +2233,28 @@ struct FTimespan
 	int64 Ticks;
 };
 
+/** 
+ * Represents a time by a context-free musical bar/beat. 
+ * @note The full C++ class is located here: Engine\Source\Runtime\Core\Public\Misc\MusicalTime.h
+ */
+USTRUCT(noexport)
+struct FMusicalTime
+{
+	/** 0-based Bar Index */
+	UPROPERTY()
+	int32 Bar;
+	
+	/** 0-based Tick Index */
+	UPROPERTY()
+	int32 TickInBar;
+
+	UPROPERTY()
+	int32 TicksPerBar;
+
+	UPROPERTY()
+	int32 TicksPerBeat;
+};
+
 /**
  * A struct that can reference a top level asset such as '/Path/To/Package.AssetName'
  * @note The full C++ class is located here: Engine\Source\Runtime\CoreUObject\Public\UObject\TopLevelAssetPath.h
diff --git a/Engine/Source/Runtime/MovieScene/Private/Channels/MovieSceneAudioTriggerChannel.cpp b/Engine/Source/Runtime/MovieScene/Private/Channels/MovieSceneAudioTriggerChannel.cpp
--- a/Engine/Source/Runtime/MovieScene/Private/Channels/MovieSceneAudioTriggerChannel.cpp
+++ b/Engine/Source/Runtime/MovieScene/Private/Channels/MovieSceneAudioTriggerChannel.cpp
@@ -75,6 +75,11 @@ void FMovieSceneAudioTriggerChannel::ChangeFrameResolution(FFrameRate SourceRate
 	GetData().ChangeFrameResolution(SourceRate, DestinationRate);
 }
 
+void FMovieSceneAudioTriggerChannel::RemapMusicalFrameTimes(const UFrameBasedMusicMap* OldMap, const UFrameBasedMusicMap* NewMap)
+{
+	GetData().RemapMusicalFrameTimes(OldMap, NewMap);
+}
+
 TRange<FFrameNumber> FMovieSceneAudioTriggerChannel::ComputeEffectiveRange() const
 {
 	return GetData().GetTotalRange();
diff --git a/Engine/Source/Runtime/MovieScene/Private/Channels/MovieSceneBoolChannel.cpp b/Engine/Source/Runtime/MovieScene/Private/Channels/MovieSceneBoolChannel.cpp
--- a/Engine/Source/Runtime/MovieScene/Private/Channels/MovieSceneBoolChannel.cpp
+++ b/Engine/Source/Runtime/MovieScene/Private/Channels/MovieSceneBoolChannel.cpp
@@ -153,6 +153,11 @@ void FMovieSceneBoolChannel::ChangeFrameResolution(FFrameRate SourceRate, FFrame
 	GetData().ChangeFrameResolution(SourceRate, DestinationRate);
 }
 
+void FMovieSceneBoolChannel::RemapMusicalFrameTimes(const UFrameBasedMusicMap* OldMap, const UFrameBasedMusicMap* NewMap)
+{
+	GetData().RemapMusicalFrameTimes(OldMap, NewMap);
+}
+
 TRange<FFrameNumber> FMovieSceneBoolChannel::ComputeEffectiveRange() const
 {
 	return GetData().GetTotalRange();
diff --git a/Engine/Source/Runtime/MovieScene/Private/Channels/MovieSceneByteChannel.cpp b/Engine/Source/Runtime/MovieScene/Private/Channels/MovieSceneByteChannel.cpp
--- a/Engine/Source/Runtime/MovieScene/Private/Channels/MovieSceneByteChannel.cpp
+++ b/Engine/Source/Runtime/MovieScene/Private/Channels/MovieSceneByteChannel.cpp
@@ -160,6 +160,11 @@ void FMovieSceneByteChannel::ChangeFrameResolution(FFrameRate SourceRate, FFrame
 	GetData().ChangeFrameResolution(SourceRate, DestinationRate);
 }
 
+void FMovieSceneByteChannel::RemapMusicalFrameTimes(const UFrameBasedMusicMap* OldMap, const UFrameBasedMusicMap* NewMap)
+{
+	GetData().RemapMusicalFrameTimes(OldMap, NewMap);
+}
+
 TRange<FFrameNumber> FMovieSceneByteChannel::ComputeEffectiveRange() const
 {
 	return GetData().GetTotalRange();
diff --git a/Engine/Source/Runtime/MovieScene/Private/Channels/MovieSceneChannelData.cpp b/Engine/Source/Runtime/MovieScene/Private/Channels/MovieSceneChannelData.cpp
--- a/Engine/Source/Runtime/MovieScene/Private/Channels/MovieSceneChannelData.cpp
+++ b/Engine/Source/Runtime/MovieScene/Private/Channels/MovieSceneChannelData.cpp
@@ -1,6 +1,7 @@
 // Copyright Epic Games, Inc. All Rights Reserved.
 
 #include "Channels/MovieSceneChannelData.h"
+#include "FrameBasedMusicMap.h"
 #include "Misc/FrameRate.h"
 
 #include UE_INLINE_GENERATED_CPP_BY_NAME(MovieSceneChannelData)
@@ -198,6 +199,30 @@ void FMovieSceneChannelData::ChangeFrameResolution(FFrameRate SourceRate, FFrame
 	}
 }
 
+void FMovieSceneChannelData::RemapMusicalFrameTimes(const UFrameBasedMusicMap* OldMap, const UFrameBasedMusicMap* NewMap)
+{
+	// Unlike ChangeFrameResolution, this function can change the actual location in
+	// time of the the keyframes. So I *think* the right thing to do here is broadcast to
+	// all of the callbacks that want to know when keyframes change. I cribbed this from
+	// the MoveKeyInternal function above - Buzz
+	TArray<FKeyMoveEventItem> Items;
+
+	for (int32 Index = 0; Index < Times->Num(); ++Index)
+	{
+		FFrameNumber OriginalTime = (*Times)[Index];
+		(*Times)[Index] = NewMap->RemapFrameTime((*Times)[Index], OldMap).RoundToFrame();
+		if ((*Times)[Index] != OriginalTime)
+		{
+			Items.Add(FKeyMoveEventItem(Index, OriginalTime, Index, (*Times)[Index]));
+		}
+	}
+
+	if (OwningChannel && OwningChannel->OnKeyMovedEvent().IsBound())
+	{
+		OwningChannel->OnKeyMovedEvent().Broadcast(OwningChannel, Items);
+	}
+}
+
 void FMovieSceneChannelData::GetKeys(const TRange<FFrameNumber>& WithinRange, TArray<FFrameNumber>* OutKeyTimes, TArray<FKeyHandle>* OutKeyHandles)
 {
 	if (!Times->Num())
diff --git a/Engine/Source/Runtime/MovieScene/Private/Channels/MovieSceneCurveChannelImpl.cpp b/Engine/Source/Runtime/MovieScene/Private/Channels/MovieSceneCurveChannelImpl.cpp
--- a/Engine/Source/Runtime/MovieScene/Private/Channels/MovieSceneCurveChannelImpl.cpp
+++ b/Engine/Source/Runtime/MovieScene/Private/Channels/MovieSceneCurveChannelImpl.cpp
@@ -6,6 +6,7 @@
 #include "Channels/MovieSceneInterpolation.h"
 #include "HAL/ConsoleManager.h"
 #include "MovieSceneFrameMigration.h"
+#include "FrameBasedMusicMap.h"
 #include "Curves/CurveEvaluation.h"
 #include "UObject/FortniteMainBranchObjectVersion.h"
 #include "UObject/SequencerObjectVersion.h"
@@ -868,6 +869,63 @@ void TMovieSceneCurveChannelImpl<ChannelType>::ChangeFrameResolution(ChannelType
 	}
 }
 
+/**  Remap original times to new times given the old music map and the new music map.
+  *  Assumes the original times were set in referenc to the OldMap!
+  */
+template<typename ChannelType>
+void TMovieSceneCurveChannelImpl<ChannelType>::RemapMusicalFrameTimes(ChannelType* InChannel, const UFrameBasedMusicMap* OldMap, const UFrameBasedMusicMap* NewMap)
+{
+	check(InChannel->Times.Num() == InChannel->Values.Num());
+
+	if (InChannel->Times.IsEmpty())
+	{
+		return;
+	}
+
+	TArray<FKeyMoveEventItem> MovedItems;
+
+	// Remapping may result in keyframes getting closer together or farther apart,
+	// so we want to scale the tangents as appropriate given any change to the 
+	// distance between the keyframes. This is why we can't just calculate one 
+	// IntervalFactor like is done above in ChangeFrameResolution.
+	FFrameNumber OriginalPreviousTime = InChannel->Times[0];
+	FFrameNumber NewPreviousTime = InChannel->Times[0];
+	for (int32 Index = 0; Index < InChannel->Times.Num(); ++Index)
+	{
+		FFrameNumber OriginalThisTime = InChannel->Times[Index];
+		InChannel->Times[Index] = NewMap->RemapFrameTime(InChannel->Times[Index], OldMap).RoundToFrame();
+		FFrameNumber NewThisTime = InChannel->Times[Index];
+
+		if (NewThisTime != OriginalThisTime)
+		{
+			MovedItems.Add(FKeyMoveEventItem(Index, OriginalThisTime, Index, NewThisTime));
+		}
+
+		if (Index != 0)
+		{
+			float OldInterval = static_cast<float>(OriginalThisTime.Value - OriginalPreviousTime.Value);
+			float NewInterval = static_cast<float>(NewThisTime.Value - NewPreviousTime.Value);
+			if (!FMath::IsNearlyEqual(NewInterval, OldInterval) && !FMath::IsNearlyZero(NewInterval))
+			{
+				float IntervalFactor = OldInterval / NewInterval;
+				ChannelValueType& PreviousValue = InChannel->Values[Index - 1];
+				PreviousValue.Tangent.LeaveTangent *= IntervalFactor;
+				ChannelValueType& ThisValue = InChannel->Values[Index];
+				ThisValue.Tangent.ArriveTangent *= IntervalFactor;
+			}
+		}
+		OriginalPreviousTime = OriginalThisTime;
+		NewPreviousTime = NewThisTime;
+	}
+
+	if (MovedItems.Num() > 0 && InChannel && InChannel->OnKeyMovedEvent().IsBound())
+	{
+		InChannel->OnKeyMovedEvent().Broadcast(InChannel, MovedItems);
+	}
+
+}
+
+
 template<typename ChannelType>
 void TMovieSceneCurveChannelImpl<ChannelType>::Optimize(ChannelType* InChannel, const FKeyDataOptimizationParams& InParameters)
 {
diff --git a/Engine/Source/Runtime/MovieScene/Private/Channels/MovieSceneCurveChannelImpl.h b/Engine/Source/Runtime/MovieScene/Private/Channels/MovieSceneCurveChannelImpl.h
--- a/Engine/Source/Runtime/MovieScene/Private/Channels/MovieSceneCurveChannelImpl.h
+++ b/Engine/Source/Runtime/MovieScene/Private/Channels/MovieSceneCurveChannelImpl.h
@@ -94,6 +94,9 @@ struct TMovieSceneCurveChannelImpl
 	/** Change the frame resolution of the channel */
 	static void ChangeFrameResolution(ChannelType* InChannel, FFrameRate SourceRate, FFrameRate DestinationRate);
 
+	/** Remap key positions from an old music map to a new music map */
+	static void RemapMusicalFrameTimes(ChannelType* InChannel, const UFrameBasedMusicMap* OldMap, const UFrameBasedMusicMap* NewMap);
+
 	/** Optimize the channel's curve */
 	static void Optimize(ChannelType* InChannel, const FKeyDataOptimizationParams& InParameters);
 
diff --git a/Engine/Source/Runtime/MovieScene/Private/Channels/MovieSceneDoubleChannel.cpp b/Engine/Source/Runtime/MovieScene/Private/Channels/MovieSceneDoubleChannel.cpp
--- a/Engine/Source/Runtime/MovieScene/Private/Channels/MovieSceneDoubleChannel.cpp
+++ b/Engine/Source/Runtime/MovieScene/Private/Channels/MovieSceneDoubleChannel.cpp
@@ -205,6 +205,11 @@ void FMovieSceneDoubleChannel::ChangeFrameResolution(FFrameRate SourceRate, FFra
 	FMovieSceneDoubleChannelImpl::ChangeFrameResolution(this, SourceRate, DestinationRate);
 }
 
+void FMovieSceneDoubleChannel::RemapMusicalFrameTimes(const UFrameBasedMusicMap* OldMap, const UFrameBasedMusicMap* NewMap)
+{
+	FMovieSceneDoubleChannelImpl::RemapMusicalFrameTimes(this, OldMap, NewMap);
+}
+
 TRange<FFrameNumber> FMovieSceneDoubleChannel::ComputeEffectiveRange() const
 {
 	return GetData().GetTotalRange();
diff --git a/Engine/Source/Runtime/MovieScene/Private/Channels/MovieSceneFloatChannel.cpp b/Engine/Source/Runtime/MovieScene/Private/Channels/MovieSceneFloatChannel.cpp
--- a/Engine/Source/Runtime/MovieScene/Private/Channels/MovieSceneFloatChannel.cpp
+++ b/Engine/Source/Runtime/MovieScene/Private/Channels/MovieSceneFloatChannel.cpp
@@ -119,6 +119,11 @@ void FMovieSceneFloatChannel::ChangeFrameResolution(FFrameRate SourceRate, FFram
 	FMovieSceneFloatChannelImpl::ChangeFrameResolution(this, SourceRate, DestinationRate);
 }
 
+void FMovieSceneFloatChannel::RemapMusicalFrameTimes(const UFrameBasedMusicMap* OldMap, const UFrameBasedMusicMap* NewMap)
+{
+	FMovieSceneFloatChannelImpl::RemapMusicalFrameTimes(this, OldMap, NewMap);
+}
+
 TRange<FFrameNumber> FMovieSceneFloatChannel::ComputeEffectiveRange() const
 {
 	return GetData().GetTotalRange();
diff --git a/Engine/Source/Runtime/MovieScene/Private/Channels/MovieSceneIntegerChannel.cpp b/Engine/Source/Runtime/MovieScene/Private/Channels/MovieSceneIntegerChannel.cpp
--- a/Engine/Source/Runtime/MovieScene/Private/Channels/MovieSceneIntegerChannel.cpp
+++ b/Engine/Source/Runtime/MovieScene/Private/Channels/MovieSceneIntegerChannel.cpp
@@ -183,6 +183,11 @@ void FMovieSceneIntegerChannel::ChangeFrameResolution(FFrameRate SourceRate, FFr
 	GetData().ChangeFrameResolution(SourceRate, DestinationRate);
 }
 
+void FMovieSceneIntegerChannel::RemapMusicalFrameTimes(const UFrameBasedMusicMap* OldMap, const UFrameBasedMusicMap* NewMap)
+{
+	GetData().RemapMusicalFrameTimes(OldMap, NewMap);
+}
+
 TRange<FFrameNumber> FMovieSceneIntegerChannel::ComputeEffectiveRange() const
 {
 	return GetData().GetTotalRange();
diff --git a/Engine/Source/Runtime/MovieScene/Private/Channels/MovieSceneObjectPathChannel.cpp b/Engine/Source/Runtime/MovieScene/Private/Channels/MovieSceneObjectPathChannel.cpp
--- a/Engine/Source/Runtime/MovieScene/Private/Channels/MovieSceneObjectPathChannel.cpp
+++ b/Engine/Source/Runtime/MovieScene/Private/Channels/MovieSceneObjectPathChannel.cpp
@@ -135,7 +135,12 @@ void FMovieSceneObjectPathChannel::ChangeFrameResolution(FFrameRate SourceRate,
 	GetData().ChangeFrameResolution(SourceRate, DestinationRate);
 }
 
-TRange<FFrameNumber> FMovieSceneObjectPathChannel::ComputeEffectiveRange() const 
+void FMovieSceneObjectPathChannel::RemapMusicalFrameTimes(const UFrameBasedMusicMap* OldMap, const UFrameBasedMusicMap* NewMap)
+{
+	GetData().RemapMusicalFrameTimes(OldMap, NewMap);
+}
+
+TRange<FFrameNumber> FMovieSceneObjectPathChannel::ComputeEffectiveRange() const
 {
 	return GetData().GetTotalRange();
 }
diff --git a/Engine/Source/Runtime/MovieScene/Private/Evaluation/IMovieSceneCustomClockSource.cpp b/Engine/Source/Runtime/MovieScene/Private/Evaluation/IMovieSceneCustomClockSource.cpp
--- a/Engine/Source/Runtime/MovieScene/Private/Evaluation/IMovieSceneCustomClockSource.cpp
+++ b/Engine/Source/Runtime/MovieScene/Private/Evaluation/IMovieSceneCustomClockSource.cpp
@@ -2,6 +2,8 @@
 
 #include "Evaluation/IMovieSceneCustomClockSource.h"
 #include "UObject/Package.h"
+#include "Components/ActorComponent.h"
+#include "GameFramework/Actor.h"
 
 #include UE_INLINE_GENERATED_CPP_BY_NAME(IMovieSceneCustomClockSource)
 
@@ -80,6 +82,36 @@ FFrameTime FMovieSceneTimeController_Custom::OnRequestCurrentTime(const FQualifi
 	return InCurrentTime.Time;
 }
 
+bool FMovieSceneTimeController_Custom::OnSupportsMusicalTime()
+{
+	if (WeakObject.IsStale())
+	{
+		ResolveInterfacePtr();
+	}
+
+	if (WeakObject.IsValid())
+	{
+		return InterfacePtr->OnSupportsMusicalTime();
+	}
+
+	return false;
+}
+
+FMusicalTime FMovieSceneTimeController_Custom::OnRequestCurrentMusicalTime(const FQualifiedFrameTime& InCurrentTime, float InPlayRate)
+{
+	if (WeakObject.IsStale())
+	{
+		ResolveInterfacePtr();
+	}
+
+	if (WeakObject.IsValid())
+	{
+		return InterfacePtr->OnRequestCurrentMusicalTime(InCurrentTime, InPlayRate);
+	}
+	
+	return FMusicalTime();
+}
+
 void FMovieSceneTimeController_Custom::ResolveInterfacePtr()
 {
 	WeakObject = nullptr;
@@ -87,11 +119,24 @@ void FMovieSceneTimeController_Custom::ResolveInterfacePtr()
 
 	if (UObject* ResolvedClockSource = ObjectPath.ResolveObject())
 	{
-		const bool bHasCompatibleInterface = ResolvedClockSource->GetClass()->ImplementsInterface(UMovieSceneCustomClockSource::StaticClass());
-		if (bHasCompatibleInterface)
+		const bool bObjectHasCompatibleInterface = ResolvedClockSource->GetClass()->ImplementsInterface(UMovieSceneCustomClockSource::StaticClass());
+		if (bObjectHasCompatibleInterface)
 		{
 			WeakObject = ResolvedClockSource;
 			InterfacePtr = static_cast<IMovieSceneCustomClockSource*>(ResolvedClockSource->GetInterfaceAddress(UMovieSceneCustomClockSource::StaticClass()));
+			return;
+		}
+
+		if (AActor* AsActor = Cast<AActor>(ResolvedClockSource))
+		{
+			TArray<UActorComponent*> ClockSources = AsActor->GetComponentsByInterface(UMovieSceneCustomClockSource::StaticClass());
+			if (ClockSources.Num() == 0)
+			{
+				return;
+			}
+			// hmmmm... warn if actor has more than one possible clock source? 
+			WeakObject = ClockSources[0];
+			InterfacePtr = static_cast<IMovieSceneCustomClockSource*>(ClockSources[0]->GetInterfaceAddress(UMovieSceneCustomClockSource::StaticClass()));
 		}
 	}
 }
diff --git a/Engine/Source/Runtime/MovieScene/Private/MovieScene.cpp b/Engine/Source/Runtime/MovieScene/Private/MovieScene.cpp
--- a/Engine/Source/Runtime/MovieScene/Private/MovieScene.cpp
+++ b/Engine/Source/Runtime/MovieScene/Private/MovieScene.cpp
@@ -522,6 +522,15 @@ FText UMovieScene::GetObjectDisplayName(const FGuid& ObjectId)
 }
 
 
+void UMovieScene::SetTickResolutionDirectly(FFrameRate InTickResolution)
+{
+	TickResolution = InTickResolution;
+	if (MusicMap)
+	{
+		MusicMap->SetFrameResolution(TickResolution);
+	}
+}
+
 void UMovieScene::AddNewBindingTag(const FName& NewTag)
 {
 	BindingGroups.Add(NewTag);
diff --git a/Engine/Source/Runtime/MovieScene/Private/MovieSceneSection.cpp b/Engine/Source/Runtime/MovieScene/Private/MovieSceneSection.cpp
--- a/Engine/Source/Runtime/MovieScene/Private/MovieSceneSection.cpp
+++ b/Engine/Source/Runtime/MovieScene/Private/MovieSceneSection.cpp
@@ -780,6 +780,33 @@ FColor UMovieSceneSection::GetColorTint() const
 #endif
 }
 
+void UMovieSceneSection::RemapMusicalFrameTimes(const UFrameBasedMusicMap* OldMap, const UFrameBasedMusicMap* NewMap)
+{
+	if (TryModify())
+	{
+		if (SectionRange.Value.GetLowerBound().IsClosed())
+		{
+			FFrameTime OldFrame(SectionRange.Value.GetLowerBoundValue());
+			FFrameTime NewFrame = NewMap->RemapFrameTime(OldFrame, OldMap);
+			SectionRange.Value.SetLowerBoundValue(NewFrame.FloorToFrame());
+		}
+		if (SectionRange.Value.GetUpperBound().IsClosed())
+		{
+			FFrameTime OldFrame(SectionRange.Value.GetUpperBoundValue());
+			FFrameTime NewFrame = NewMap->RemapFrameTime(OldFrame, OldMap);
+			SectionRange.Value.SetUpperBoundValue(NewFrame.CeilToFrame());
+		}
+
+		for (const FMovieSceneChannelEntry& Entry : GetChannelProxy().GetAllEntries())
+		{
+			for (FMovieSceneChannel* Channel : Entry.GetChannels())
+			{
+				Channel->RemapMusicalFrameTimes(OldMap, NewMap);
+			}
+		}
+	}
+}
+
 UMovieSceneSection* UMovieSceneSection::SplitSection(FQualifiedFrameTime SplitTime, bool bDeleteKeys)
 {
 	if (!SectionRange.Value.Contains(SplitTime.Time.GetFrame()))
diff --git a/Engine/Source/Runtime/MovieScene/Private/MovieSceneSequencePlayer.cpp b/Engine/Source/Runtime/MovieScene/Private/MovieSceneSequencePlayer.cpp
--- a/Engine/Source/Runtime/MovieScene/Private/MovieSceneSequencePlayer.cpp
+++ b/Engine/Source/Runtime/MovieScene/Private/MovieSceneSequencePlayer.cpp
@@ -1114,7 +1114,19 @@ void UMovieSceneSequencePlayer::Update(const float DeltaSeconds)
 		{
 			check(!IsEvaluating());
 
-			FFrameTime NewTime = TimeController->RequestCurrentTime(GetCurrentTime(), PlayRate, GetDisplayRate());
+			FFrameTime NewTime;
+			UMovieScene* MovieScene = Sequence->GetMovieScene();
+			if (MovieScene && MovieScene->HasMusicMap() && TimeController->SupportsMusicalTime())
+			{
+				FMusicalTime MusicalTime = TimeController->RequestCurrentMusicalTime(GetCurrentTime(), PlayRate, GetDisplayRate());
+				FFrameTime FrameTime = MovieScene->GetMusicMap()->GetFrameTime(MusicalTime);
+				NewTime = FFrameRate::TransformTime(FrameTime, MovieScene->GetMusicMap()->FrameResolution, GetDisplayRate());
+			}
+			else
+			{
+				NewTime = TimeController->RequestCurrentTime(GetCurrentTime(), PlayRate, GetDisplayRate());
+			}
+
 			UpdateTimeCursorPosition(NewTime, EUpdatePositionMethod::Play);
 		}
 
diff --git a/Engine/Source/Runtime/MovieScene/Private/MovieSceneTimeController.cpp b/Engine/Source/Runtime/MovieScene/Private/MovieSceneTimeController.cpp
--- a/Engine/Source/Runtime/MovieScene/Private/MovieSceneTimeController.cpp
+++ b/Engine/Source/Runtime/MovieScene/Private/MovieSceneTimeController.cpp
@@ -68,6 +68,17 @@ FFrameTime FMovieSceneTimeController::RequestCurrentTime(const FQualifiedFrameTi
 	return OnRequestCurrentTime(InCurrentTime, InPlayRate);
 }
 
+bool FMovieSceneTimeController::SupportsMusicalTime()
+{
+	return OnSupportsMusicalTime();
+}
+
+FMusicalTime FMovieSceneTimeController::RequestCurrentMusicalTime(const FQualifiedFrameTime& InCurrentTime, float InPlayRate, FFrameRate InDisplayRate)
+{
+	DisplayRate = InDisplayRate;
+	return OnRequestCurrentMusicalTime(InCurrentTime, InPlayRate);
+}
+
 void FMovieSceneTimeController_ExternalClock::OnStartPlaying(const FQualifiedFrameTime& InStartTime)
 {
 	ClockStartTime = ClockLastUpdateTime = GetCurrentTime();
diff --git a/Engine/Source/Runtime/MovieScene/Private/MovieSceneTimeHelpers.cpp b/Engine/Source/Runtime/MovieScene/Private/MovieSceneTimeHelpers.cpp
--- a/Engine/Source/Runtime/MovieScene/Private/MovieSceneTimeHelpers.cpp
+++ b/Engine/Source/Runtime/MovieScene/Private/MovieSceneTimeHelpers.cpp
@@ -107,6 +107,38 @@ void MigrateFrameTimes(FFrameRate SourceRate, FFrameRate DestinationRate, UMovie
 	Section->SetIsLocked(bSectionWasLocked);
 }
 
+void RemapMusicalFrameTimes(UMovieSceneSection* Section, const UFrameBasedMusicMap* OldMap, const UFrameBasedMusicMap* NewMap)
+{
+	Section->Modify();
+	const bool bSectionWasLocked = Section->IsLocked();
+	Section->SetIsLocked(false);
+
+	// We don't need to adjust the section 'Range' here because that happens in the next function call...
+
+	Section->RemapMusicalFrameTimes(OldMap, NewMap);
+
+	if (Section->GetPreRollFrames() > 0)
+	{
+		FFrameNumber NewPreRollFrameCount = NewMap->RemapFrameTime(FFrameTime(Section->GetPreRollFrames()), OldMap).FloorToFrame();
+		Section->SetPreRollFrames(NewPreRollFrameCount.Value);
+	}
+
+	if (Section->GetPostRollFrames() > 0)
+	{
+		FFrameNumber NewPostRollFrameCount = NewMap->RemapFrameTime(FFrameTime(Section->GetPostRollFrames()), OldMap).FloorToFrame();
+		Section->SetPostRollFrames(NewPostRollFrameCount.Value);
+	}
+
+	// NOT SURE THESE SHOULD CHANGE WHEN REMAPPING TO NEW MUSIC TIMINGS!
+	// QUESTION FOR ENGINE TEAM
+	Section->Easing.AutoEaseInDuration = NewMap->RemapFrameTime(Section->Easing.AutoEaseInDuration, OldMap).FloorToFrame().Value;
+	Section->Easing.AutoEaseOutDuration = NewMap->RemapFrameTime(Section->Easing.AutoEaseOutDuration, OldMap).FloorToFrame().Value;
+	Section->Easing.ManualEaseInDuration = NewMap->RemapFrameTime(Section->Easing.ManualEaseInDuration, OldMap).FloorToFrame().Value;
+	Section->Easing.ManualEaseOutDuration = NewMap->RemapFrameTime(Section->Easing.ManualEaseOutDuration, OldMap).FloorToFrame().Value;
+
+	Section->SetIsLocked(bSectionWasLocked);
+}
+
 void MigrateFrameTimes(FFrameRate SourceRate, FFrameRate DestinationRate, UMovieSceneTrack* Track)
 {
 	FScopedSlowTask SlowTask(Track->GetAllSections().Num());
@@ -118,6 +150,17 @@ void MigrateFrameTimes(FFrameRate SourceRate, FFrameRate DestinationRate, UMovie
 	}
 }
 
+void RemapMusicalFrameTimes(UMovieSceneTrack* Track, const UFrameBasedMusicMap* OldMap, const UFrameBasedMusicMap* NewMap)
+{
+	FScopedSlowTask SlowTask(Track->GetAllSections().Num());
+
+	for (UMovieSceneSection* Section : Track->GetAllSections())
+	{
+		SlowTask.EnterProgressFrame();
+		RemapMusicalFrameTimes(Section, OldMap, NewMap);
+	}
+}
+
 void TimeHelpers::MigrateFrameTimes(FFrameRate SourceRate, FFrameRate DestinationRate, UMovieScene* MovieScene, bool bApplyRecursively)
 {
 	int32 TotalNumTracks = MovieScene->GetTracks().Num() + (MovieScene->GetCameraCutTrack() ? 1 : 0);
@@ -200,11 +243,80 @@ void TimeHelpers::MigrateFrameTimes(FFrameRate SourceRate, FFrameRate Destinatio
 		MovieScene->SortMarkedFrames();
 
 	}
+
 	MovieScene->SetTickResolutionDirectly(DestinationRate);
 #if WITH_EDITOR
 	MovieScene->SetReadOnly(bMovieSceneReadOnly);
 #endif
 }
 
+void TimeHelpers::RemapMusicalFrameTimes(UMovieScene* MovieScene, UFrameBasedMusicMap* NewMap)
+{
+	int32 TotalNumTracks = MovieScene->GetTracks().Num() + (MovieScene->GetCameraCutTrack() ? 1 : 0);
+	for (const FMovieSceneBinding& Binding : MovieScene->GetBindings())
+	{
+		TotalNumTracks += Binding.GetTracks().Num();
+	}
+
+	FScopedSlowTask SlowTask(TotalNumTracks, NSLOCTEXT("MovieScene", "ChangingTickResolution", "Migrating sequence frame timing"));
+	SlowTask.MakeDialogDelayed(0.25f, true);
+
+	MovieScene->Modify();
+#if WITH_EDITOR
+	const bool bMovieSceneReadOnly = MovieScene->IsReadOnly();
+	MovieScene->SetReadOnly(false);
+#endif
+
+	const UFrameBasedMusicMap* SourceMap = MovieScene->GetMusicMap();
+
+	if (SourceMap)
+	{
+		for (UMovieSceneTrack* Track : MovieScene->GetTracks())
+		{
+			SlowTask.EnterProgressFrame();
+			UE::MovieScene::RemapMusicalFrameTimes(Track, SourceMap, NewMap);
+		}
+
+		if (UMovieSceneTrack* Track = MovieScene->GetCameraCutTrack())
+		{
+			SlowTask.EnterProgressFrame();
+			UE::MovieScene::RemapMusicalFrameTimes(Track, SourceMap, NewMap);
+		}
+
+		for (const FMovieSceneBinding& Binding : MovieScene->GetBindings())
+		{
+			for (UMovieSceneTrack* Track : Binding.GetTracks())
+			{
+				SlowTask.EnterProgressFrame();
+				UE::MovieScene::RemapMusicalFrameTimes(Track, SourceMap, NewMap);
+			}
+		}
+
+		{
+			TArray<FMovieSceneMarkedFrame> MarkedFrames = MovieScene->GetMarkedFrames();
+
+			// Clear the marked frames as the returned array is immutable
+			MovieScene->DeleteMarkedFrames();
+
+			for (FMovieSceneMarkedFrame& MarkedFrame : MarkedFrames)
+			{
+				MarkedFrame.FrameNumber = NewMap->RemapFrameTime(MarkedFrame.FrameNumber, SourceMap).RoundToFrame();
+
+				// Add it back in
+				MovieScene->AddMarkedFrame(MarkedFrame);
+			}
+
+			// Ensure they're in order as they may not have been before.
+			MovieScene->SortMarkedFrames();
+		}
+	}
+
+	MovieScene->SetMusicMapDirectly(NewMap);
+
+#if WITH_EDITOR
+	MovieScene->SetReadOnly(bMovieSceneReadOnly);
+#endif
+}
+
 } // namespace MovieScene
 } // namespace UE
\ No newline at end of file
diff --git a/Engine/Source/Runtime/MovieScene/Private/Sections/MovieSceneSectionTimingParameters.cpp b/Engine/Source/Runtime/MovieScene/Private/Sections/MovieSceneSectionTimingParameters.cpp
--- a/Engine/Source/Runtime/MovieScene/Private/Sections/MovieSceneSectionTimingParameters.cpp
+++ b/Engine/Source/Runtime/MovieScene/Private/Sections/MovieSceneSectionTimingParameters.cpp
@@ -3,8 +3,9 @@
 #include "Sections/MovieSceneSectionTimingParameters.h"
 #include "Evaluation/MovieSceneSequenceTransform.h"
 #include "Misc/FrameRate.h"
+#include "Variants/MovieSceneMusicTimeTransform.h"
 #include "Variants/MovieSceneTimeWarpVariantPayloads.h"
-
+#include "MovieScene.h"
 
 FMovieSceneSequenceTransform FMovieSceneSectionTimingParametersSeconds::MakeTransform(const FFrameRate& OuterFrameRate, const TRange<FFrameNumber>& OuterRange, double SourceDuration, double InnerPlayRate) const
 {
@@ -21,11 +22,11 @@ FMovieSceneSequenceTransform FMovieSceneSectionTimingParametersSeconds::MakeTran
 
 	// ----------------------------------------------------------------------------
 	// First things first, subtract the section start bound
-	Result.Add(FMovieSceneTimeTransform(-OuterRange.GetLowerBoundValue()));
+	AddPositionInOuterAsOffset(Result, -OuterRange.GetLowerBoundValue());
 
 	// ----------------------------------------------------------------------------
 	// Time warp
-	Result.Add(0, PlayRate.ShallowCopy());
+	AddPlayRate(Result);
 
 	// ----------------------------------------------------------------------------
 	// FrameRate conversion to seconds
@@ -38,7 +39,7 @@ FMovieSceneSequenceTransform FMovieSceneSectionTimingParametersSeconds::MakeTran
 	const double Duration    = EndTime - StartTime;
 	double StartOffset = InnerStartOffset + FirstLoopStartOffset;
 
-	// Accomodate negative play rates by playing from the end of the clip
+	// Accommodate negative play rates by playing from the end of the clip
 	if (PlayRate.GetType() == EMovieSceneTimeWarpType::FixedPlayRate && PlayRate.AsFixedPlayRate() < 0.0)
 	{
 		StartOffset += Duration;
@@ -47,9 +48,38 @@ FMovieSceneSequenceTransform FMovieSceneSectionTimingParametersSeconds::MakeTran
 	// Start offset
 	if (!FMath::IsNearlyZero(StartOffset))
 	{
-		Result.Add(FMovieSceneTimeTransform(FFrameTime::FromDecimal(StartOffset)));
+		AddInnerStartOffset(Result, FFrameTime::FromDecimal(StartOffset));
 	}
 
+	AddLoopingOrClampingAndReverse(Result, Duration, -StartTime);
+
+	return Result;
+}
+
+void FMovieSceneSectionTimingParametersSeconds::AddPositionInOuterAsOffset(FMovieSceneSequenceTransform& Transform, const FFrameTime& Offset) const
+{
+	Transform.Add(FMovieSceneTimeTransform(Offset));
+}
+
+void FMovieSceneSectionTimingParametersSeconds::AddPlayRate(FMovieSceneSequenceTransform& Transform) const
+{
+	Transform.Add(0, PlayRate.ShallowCopy());
+}
+
+void FMovieSceneSectionTimingParametersSeconds::AddFrameRateConversion(FMovieSceneSequenceTransform& Transform, const FFrameRate& OuterFrameRate, const FFrameRate& InnerFrameRate) const
+{
+	FMovieSceneTimeWarpVariant FrameRate;
+	FrameRate.Set(FMovieSceneTimeWarpFrameRate{ OuterFrameRate / InnerFrameRate });
+	Transform.Add(0, MoveTemp(FrameRate));
+}
+
+void FMovieSceneSectionTimingParametersSeconds::AddInnerStartOffset(FMovieSceneSequenceTransform& Transform, const FFrameTime& Offset) const
+{
+	Transform.Add(FMovieSceneTimeTransform(Offset));
+}
+
+void FMovieSceneSectionTimingParametersSeconds::AddLoopingOrClampingAndReverse(FMovieSceneSequenceTransform& Result, const double Duration, const double Offset) const
+{
 	// ----------------------------------------------------------------------------
 	// Looping or clamping
 	if (bLoop)
@@ -57,14 +87,14 @@ FMovieSceneSequenceTransform FMovieSceneSectionTimingParametersSeconds::MakeTran
 		// Loop
 		FMovieSceneTimeWarpVariant Loop;
 		Loop.Set(FMovieSceneTimeWarpLoopFloat{ static_cast<float>(Duration) });
-		Result.Add(FFrameTime::FromDecimal(-StartTime), MoveTemp(Loop));
+		Result.Add(FFrameTime::FromDecimal(Offset), MoveTemp(Loop));
 	}
 	else if (bClamp)
 	{
 		// Clamp
 		FMovieSceneTimeWarpVariant Clamp;
 		Clamp.Set(FMovieSceneTimeWarpClampFloat{ static_cast<float>(Duration) });
-		Result.Add(FFrameTime::FromDecimal(-StartTime), MoveTemp(Clamp));
+		Result.Add(FFrameTime::FromDecimal(Offset), MoveTemp(Clamp));
 	}
 
 	// ----------------------------------------------------------------------------
@@ -73,8 +103,6 @@ FMovieSceneSequenceTransform FMovieSceneSectionTimingParametersSeconds::MakeTran
 	{
 		Result.Add(FMovieSceneTimeTransform(FFrameTime::FromDecimal(Duration), -1.f));
 	}
-
-	return Result;
 }
 
 FMovieSceneSequenceTransform FMovieSceneSectionTimingParametersFrames::MakeTransform(const FFrameRate& OuterFrameRate, const TRange<FFrameNumber>& OuterRange, const FFrameRate& InnerFrameRate, const TRange<FFrameNumber>& InnerRange) const
@@ -86,19 +114,17 @@ FMovieSceneSequenceTransform FMovieSceneSectionTimingParametersFrames::MakeTrans
 
 	// ----------------------------------------------------------------------------
 	// First things first, subtract the section start bound
-	Result.Add(FMovieSceneTimeTransform(-OuterRange.GetLowerBoundValue()));
+	AddPositionInOuterAsOffset(Result, -OuterRange.GetLowerBoundValue());
 
 	// ----------------------------------------------------------------------------
 	// Time warp
-	Result.Add(0, PlayRate.ShallowCopy());
+	AddPlayRate(Result);
 
 	// ----------------------------------------------------------------------------
 	// FrameRate conversion
 	if (InnerFrameRate != OuterFrameRate)
 	{
-		FMovieSceneTimeWarpVariant FrameRate;
-		FrameRate.Set(FMovieSceneTimeWarpFrameRate{ OuterFrameRate / InnerFrameRate });
-		Result.Add(0, MoveTemp(FrameRate));
+		AddFrameRateConversion(Result, OuterFrameRate, InnerFrameRate);
 	}
 
 	FFrameNumber StartTime = InnerRange.GetLowerBoundValue() + InnerStartOffset;
@@ -114,31 +140,136 @@ FMovieSceneSequenceTransform FMovieSceneSectionTimingParametersFrames::MakeTrans
 	}
 
 	// Start offset
-	Result.Add(FMovieSceneTimeTransform(StartTime + LoopOffset + NegativeRateOffset));
+	AddInnerStartOffset(Result, StartTime + LoopOffset + NegativeRateOffset);
 
 	// ----------------------------------------------------------------------------
 	// Looping or clamping
+	AddLoopingOrClampingAndReverse(Result, Duration, -StartTime);
+
+	return Result;
+}
+
+MOVIESCENE_API FMovieSceneSequenceTransform FMovieSceneSectionTimingParametersFrames::MakeMusicalTransform(const FFrameRate& OuterFrameRate, const TRange<FFrameNumber>& OuterRange, UMovieScene* OuterMovieScene, const FFrameRate& InnerFrameRate, const TRange<FFrameNumber>& InnerRange, UMovieScene* InnerMovieScene) const
+{
+	FMovieSceneSequenceTransform Result;
+
+	check(OuterRange.HasLowerBound());
+	check(InnerRange.HasLowerBound() && InnerRange.HasUpperBound());
+	check(OuterMovieScene);
+	check(InnerMovieScene);
+
+
+	const UFrameBasedMusicMap* OuterMap = OuterMovieScene->GetMusicMap();
+	const UFrameBasedMusicMap* InnerMap = InnerMovieScene->GetMusicMap();
+	check(OuterMap);
+	check(InnerMap);
+
+	// ----------------------------------------------------------------------------
+	// First things first, subtract the section start bound
+	AddPositionInOuterAsOffset(Result, -OuterRange.GetLowerBoundValue());
+
+	// ----------------------------------------------------------------------------
+	// FrameRate conversion
+	if (InnerFrameRate != OuterFrameRate)
+	{
+		AddFrameRateConversion(Result, OuterFrameRate, InnerFrameRate);
+	}
+
+	AddMusicalTransform(Result, OuterRange.GetLowerBoundValue(), OuterMovieScene, InnerMovieScene);
+
+	// ----------------------------------------------------------------------------
+	// Time warp
+	AddPlayRate(Result);
+
+
+	// InnerStartOffset and InnerEndOffset are in frames that are both...
+	// a. In units of the outer music map, and...
+	// b. a count from the BEGINNING of the outer music map as a duration.
+	// So we have some conversions to use...
+	FMusicalTime StartOffsetAsMusicalTime = OuterMap->GetMusicalTime(InnerStartOffset);
+	FFrameNumber RealInnerStartOffset = InnerMap->GetFrameTime(StartOffsetAsMusicalTime).RoundToFrame();
+	FMusicalTime EndOffsetAsMusicalTime = OuterMap->GetMusicalTime(InnerEndOffset);
+	FFrameNumber RealInnerEndOffset = InnerMap->GetFrameTime(EndOffsetAsMusicalTime).RoundToFrame();
+
+	FMusicalTime LoopStartOffsetAsMusicalTime = OuterMap->GetMusicalTime(FirstLoopStartOffset);
+	FFrameNumber RealFirstLoopStartOffset = InnerMap->GetFrameTime(LoopStartOffsetAsMusicalTime).RoundToFrame();
+
+
+	FFrameNumber StartTime = InnerRange.GetLowerBoundValue() + RealInnerStartOffset;
+	FFrameNumber EndTime = InnerRange.GetUpperBoundValue() - RealInnerEndOffset;
+	FFrameNumber Duration = EndTime - StartTime;
+
+	FFrameNumber LoopOffset(bLoop ? RealFirstLoopStartOffset : 0);
+
+	FFrameNumber NegativeRateOffset = 0;
+	if (PlayRate.GetType() == EMovieSceneTimeWarpType::FixedPlayRate && PlayRate.AsFixedPlayRate() < 0.0)
+	{
+		NegativeRateOffset = Duration;
+	}
+
+	// Start offset
+	AddInnerStartOffset(Result, StartTime + LoopOffset + NegativeRateOffset);
+
+	// ----------------------------------------------------------------------------
+	// Looping or clamping
+	if (Duration > 0)
+	{
+		AddLoopingOrClampingAndReverse(Result, Duration, -StartTime);
+	}
+
+	return Result;
+}
+
+void FMovieSceneSectionTimingParametersFrames::AddPositionInOuterAsOffset(FMovieSceneSequenceTransform& Transform, const FFrameTime& Offset) const
+{
+	Transform.Add(FMovieSceneTimeTransform(Offset));
+}
+
+void FMovieSceneSectionTimingParametersFrames::AddPlayRate(FMovieSceneSequenceTransform& Transform) const
+{
+	Transform.Add(0, PlayRate.ShallowCopy());
+}
+
+void FMovieSceneSectionTimingParametersFrames::AddFrameRateConversion(FMovieSceneSequenceTransform& Transform, const FFrameRate& OuterFrameRate, const FFrameRate& InnerFrameRate) const
+{
+	FMovieSceneTimeWarpVariant FrameRate;
+	FrameRate.Set(FMovieSceneTimeWarpFrameRate{ OuterFrameRate / InnerFrameRate });
+	Transform.Add(0, MoveTemp(FrameRate));
+}
+
+void FMovieSceneSectionTimingParametersFrames::AddInnerStartOffset(FMovieSceneSequenceTransform& Transform, const FFrameTime& Offset) const
+{
+	Transform.Add(FMovieSceneTimeTransform(Offset));
+}
+
+void FMovieSceneSectionTimingParametersFrames::AddLoopingOrClampingAndReverse(FMovieSceneSequenceTransform& Result, const FFrameNumber& Duration, const FFrameNumber& Offset) const
+{
 	if (bLoop)
 	{
 		// Loop
 		FMovieSceneTimeWarpVariant Loop;
 		Loop.Set(FMovieSceneTimeWarpLoop{ Duration });
-		Result.Add(-StartTime, MoveTemp(Loop));
+		Result.Add(Offset, MoveTemp(Loop));
 	}
 	else if (bClamp)
 	{
 		// Clamp
 		FMovieSceneTimeWarpVariant Clamp;
 		Clamp.Set(FMovieSceneTimeWarpClamp{ Duration });
-		Result.Add(-StartTime, MoveTemp(Clamp));
+		Result.Add(Offset, MoveTemp(Clamp));
 	}
 
-	// ----------------------------------------------------------------------------
-	// Reverse
 	if (bReverse)
 	{
+		// Reverse
 		Result.Add(FMovieSceneTimeTransform(Duration, -1.f));
 	}
+}
 
-	return Result;
+void FMovieSceneSectionTimingParametersFrames::AddMusicalTransform(FMovieSceneSequenceTransform& Result, const FFrameNumber& Offset, UMovieScene* OuterMovieScene, UMovieScene* InnerMovieScene) const
+{
+	check(OuterMovieScene->HasMusicMap() && InnerMovieScene->HasMusicMap());
+	UMovieSceneMusicTimeTransform* Mapper = NewObject<UMovieSceneMusicTimeTransform>(InnerMovieScene);
+	Mapper->Initialize(OuterMovieScene, InnerMovieScene, Offset);
+	Result.Add(0, FMovieSceneTimeWarpVariant(Mapper));
 }
\ No newline at end of file
diff --git a/Engine/Source/Runtime/MovieScene/Private/Sections/MovieSceneSubSection.cpp b/Engine/Source/Runtime/MovieScene/Private/Sections/MovieSceneSubSection.cpp
--- a/Engine/Source/Runtime/MovieScene/Private/Sections/MovieSceneSubSection.cpp
+++ b/Engine/Source/Runtime/MovieScene/Private/Sections/MovieSceneSubSection.cpp
@@ -21,6 +21,8 @@
 #include "Tracks/MovieSceneTimeWarpTrack.h"
 #include "Sections/MovieSceneSectionTimingParameters.h"
 #include "Variants/MovieSceneTimeWarpGetter.h"
+#include "Variants/MovieSceneMusicTimeTransform.h"
+
 
 #include UE_INLINE_GENERATED_CPP_BY_NAME(MovieSceneSubSection)
 
@@ -231,7 +233,16 @@ FMovieSceneSequenceTransform UMovieSceneSubSection::OuterToInnerTransform_NoInne
 		false
 	};
 
-	return TimingParams.MakeTransform(OuterFrameRate, SubRange, InnerFrameRate, MovieScenePtr->GetPlaybackRange());
+	// determine if we need to generate a musical transform or a standard transform...
+	if (!MovieScenePtr->HasMusicMap() || !GetTypedOuter<UMovieScene>()->HasMusicMap())
+	{
+		FMovieSceneSequenceTransform Result = TimingParams.MakeTransform(OuterFrameRate, SubRange, InnerFrameRate, MovieScenePtr->GetPlaybackRange());
+		return Result;
+	}
+
+	FMovieSceneSequenceTransform Result = TimingParams.MakeMusicalTransform(OuterFrameRate, SubRange, GetTypedOuter<UMovieScene>(), InnerFrameRate, MovieScenePtr->GetPlaybackRange(), MovieScenePtr);
+
+	return Result;
 }
 
 FMovieSceneSequenceTransform UMovieSceneSubSection::OuterToInnerTransform() const
@@ -684,6 +695,35 @@ void UMovieSceneSubSection::MigrateFrameTimes(FFrameRate SourceRate, FFrameRate
 	}
 }
 
+MOVIESCENE_API void UMovieSceneSubSection::RemapMusicalFrameTimes(const UFrameBasedMusicMap* OldMap, const UFrameBasedMusicMap* NewMap)
+{
+	Super::RemapMusicalFrameTimes(OldMap, NewMap);
+	
+	UMovieSceneSequence* SequencePtr = GetSequence();
+	UMovieScene* MovieScenePtr = SequencePtr ? SequencePtr->GetMovieScene() : nullptr;
+	if (!MovieScenePtr || !MovieScenePtr->HasMusicMap())
+	{
+		return;
+	}
+
+	// The start, end, and loop-start offsets are in musical units
+	if (Parameters.StartFrameOffset.Value > 0)
+	{
+		FFrameNumber NewStartFrameOffset = NewMap->RemapFrameTime(FFrameTime(Parameters.StartFrameOffset), OldMap).RoundToFrame();
+		Parameters.StartFrameOffset = NewStartFrameOffset;
+	}
+	if (Parameters.EndFrameOffset.Value > 0)
+	{
+		FFrameNumber NewEndFrameOffset = NewMap->RemapFrameTime(FFrameTime(Parameters.EndFrameOffset), OldMap).RoundToFrame();
+		Parameters.EndFrameOffset = NewEndFrameOffset;
+	}
+	if (Parameters.FirstLoopStartFrameOffset.Value > 0)
+	{
+		FFrameNumber NewFirstLoopStartFrameOffset = NewMap->RemapFrameTime(FFrameTime(Parameters.FirstLoopStartFrameOffset), OldMap).RoundToFrame();
+		Parameters.FirstLoopStartFrameOffset = NewFirstLoopStartFrameOffset;
+	}
+}
+
 FMovieSceneSubSequenceData UMovieSceneSubSection::GenerateSubSequenceData(const FSubSequenceInstanceDataParams& Params) const
 {
 	return FMovieSceneSubSequenceData(*this);
diff --git a/Engine/Source/Runtime/MovieScene/Private/Variants/MovieSceneMusicTimeTransform.cpp b/Engine/Source/Runtime/MovieScene/Private/Variants/MovieSceneMusicTimeTransform.cpp
new file mode 100644
--- /dev/null
+++ b/Engine/Source/Runtime/MovieScene/Private/Variants/MovieSceneMusicTimeTransform.cpp
@@ -0,0 +1,177 @@
+// Copyright Epic Games, Inc. All Rights Reserved.
+
+#include "Variants/MovieSceneMusicTimeTransform.h"
+#include "FrameBasedMusicMap.h"
+#include "MovieScene.h"
+
+DEFINE_LOG_CATEGORY_STATIC(LogMusicTimeRemap, Log, All);
+
+namespace MusicTimeTransform
+{
+	class UScopedMapLock
+	{
+	public:
+		UScopedMapLock(const UMovieSceneMusicTimeTransform& Transform)
+		{
+			TStrongObjectPtr<UMovieScene> OuterScene = Transform.OuterMusicMapOwner.Pin();
+			TStrongObjectPtr<UMovieScene> InnerScene = Transform.InnerMusicMapOwner.Pin();
+			OuterMap = OuterScene ? TStrongObjectPtr<const UFrameBasedMusicMap>(OuterScene->GetMusicMap()) : nullptr;
+			InnerMap = InnerScene ? TStrongObjectPtr<const UFrameBasedMusicMap>(InnerScene->GetMusicMap()) : nullptr;
+		}
+
+		operator bool() const
+		{
+			return OuterMap && InnerMap;
+		}
+
+		TStrongObjectPtr<const UFrameBasedMusicMap> OuterMap;
+		TStrongObjectPtr<const UFrameBasedMusicMap> InnerMap;
+	};
+}
+
+UMovieSceneMusicTimeTransform::UMovieSceneMusicTimeTransform()
+{
+}
+
+void UMovieSceneMusicTimeTransform::Initialize(UMovieScene* Outer, UMovieScene* Inner, FFrameTime InOffset)
+{
+	OuterMusicMapOwner = TWeakObjectPtr<UMovieScene>(Outer);
+	InnerMusicMapOwner = TWeakObjectPtr<UMovieScene>(Inner);
+	Offset = InOffset;
+}
+
+void UMovieSceneMusicTimeTransform::InitializeDefaults()
+{
+	ScaleFactor = 1.0;
+	OuterMusicMapOwner = nullptr;
+	InnerMusicMapOwner = nullptr;
+}
+
+TRange<FFrameTime> UMovieSceneMusicTimeTransform::ComputeTraversedHull(const TRange<FFrameTime>& Range) const
+{
+	TRange<FFrameTime> Result = Range;
+
+	MusicTimeTransform::UScopedMapLock MapLock(*this);
+
+	if (!MapLock)
+	{
+		if (Range.GetLowerBound().IsClosed())
+		{
+			Result.SetLowerBoundValue(Range.GetLowerBoundValue() - Offset);
+		}
+		if (Range.GetUpperBound().IsClosed())
+		{
+			Result.SetUpperBoundValue(Range.GetUpperBoundValue() - Offset);
+		}
+		return Result;
+	}
+
+	// These can't be null if CanMapTime returns true!
+	FMusicalTime OuterPositionAtOffset = MapLock.OuterMap->GetMusicalTime(Offset);
+	
+	if (Range.GetLowerBound().IsClosed())
+	{
+		Result.SetLowerBoundValue(RemapTime_Internal(MapLock, Range.GetLowerBoundValue()));
+	}
+	if (Range.GetUpperBound().IsClosed())
+	{
+		Result.SetUpperBoundValue(RemapTime_Internal(MapLock, Range.GetUpperBoundValue()));
+	}
+
+	return Result;
+}
+
+FFrameTime UMovieSceneMusicTimeTransform::RemapTime(FFrameTime In) const
+{
+	MusicTimeTransform::UScopedMapLock MapLock(*this);
+	if (!MapLock)
+	{
+		return In;
+	}
+
+	return RemapTime_Internal(MapLock, In);
+}
+
+TOptional<FFrameTime> UMovieSceneMusicTimeTransform::InverseRemapTimeCycled(FFrameTime InValue, FFrameTime InTimeHint, const UE::MovieScene::FInverseTransformTimeParams& Params) const
+{
+	MusicTimeTransform::UScopedMapLock MapLock(*this);
+	if (!MapLock)
+	{
+		return TOptional<FFrameTime>(InValue + Offset);
+	}
+
+	FMusicalTime OuterPositionAtOffset = MapLock.OuterMap->GetMusicalTime(Offset);
+	FMusicalTime InnerMusicalTime = MapLock.InnerMap->GetMusicalTime(InValue);
+
+	double OuterBarAtOffset = OuterPositionAtOffset.FractionalBar();
+	double InnerBar = InnerMusicalTime.FractionalBar();
+
+	double OuterBarAtInValue = OuterBarAtOffset + InnerBar;
+
+	FMusicalTime OuterPositionAtInValue = MapLock.OuterMap->BarsToMusicalTime(OuterBarAtInValue);
+	FFrameTime Result = FFrameRate::TransformTime(MapLock.OuterMap->GetFrameTime(OuterPositionAtInValue) - Offset, MapLock.OuterMap->FrameResolution, MapLock.InnerMap->FrameResolution);
+/*
+	UE_LOG(LogMusicTimeRemap, Log, TEXT("InFrame: %f, InTimeHint: %f, OuterMusicalPosAtStart: %s, InnerMusicalPos: %s, OuterMusicalPos: %s, Result Frame: %f"), 
+		InValue.AsDecimal(),
+		InTimeHint.AsDecimal(),
+		*MapLock.OuterMap->MakeBarBeatStringForMusicalTime(OuterPositionAtOffset, FrameBasedMusicMap::EBarBeatFormat::Position),
+		*MapLock.InnerMap->MakeBarBeatStringForMusicalTime(InnerMusicalTime, FrameBasedMusicMap::EBarBeatFormat::Position),
+		*MapLock.OuterMap->MakeBarBeatStringForMusicalTime(OuterPositionAtInValue, FrameBasedMusicMap::EBarBeatFormat::Position),
+		Result.AsDecimal()
+		);
+*/
+	return TOptional<FFrameTime>(Result);
+}
+
+bool UMovieSceneMusicTimeTransform::InverseRemapTimeWithinRange(FFrameTime InTime, FFrameTime RangeStart, FFrameTime RangeEnd, const TFunctionRef<bool(FFrameTime)>& VisitorCallback) const
+{
+	MusicTimeTransform::UScopedMapLock MapLock(*this);
+	if (!MapLock)
+	{
+		return  VisitorCallback(InTime + Offset);
+	}
+
+	FMusicalTime OuterPositionAtOffset = MapLock.OuterMap->GetMusicalTime(Offset);
+	FMusicalTime InnerMusicalTime = MapLock.InnerMap->GetMusicalTime(InTime);
+
+	double OuterBarAtOffset = OuterPositionAtOffset.FractionalBar();
+	double InnerBar = InnerMusicalTime.FractionalBar();
+
+	double OuterBarAtInValue = OuterBarAtOffset + InnerBar;
+
+	FMusicalTime OuterPositionAtInValue = MapLock.OuterMap->BarsToMusicalTime(OuterBarAtInValue);
+	FFrameTime Result = FFrameRate::TransformTime(MapLock.OuterMap->GetFrameTime(OuterPositionAtInValue) - Offset, MapLock.OuterMap->FrameResolution, MapLock.InnerMap->FrameResolution);
+/*
+	UE_LOG(LogMusicTimeRemap, Log, TEXT("InTime: %f, OuterMusicalPosAtStart: %s, InnerMusicalPos: %s, OuterMusicalPos: %s, Result Frame: %f"),
+		InTime.AsDecimal(),
+		*MapLock.OuterMap->MakeBarBeatStringForMusicalTime(OuterPositionAtOffset, FrameBasedMusicMap::EBarBeatFormat::Position),
+		*MapLock.InnerMap->MakeBarBeatStringForMusicalTime(InnerMusicalTime, FrameBasedMusicMap::EBarBeatFormat::Position),
+		*MapLock.OuterMap->MakeBarBeatStringForMusicalTime(OuterPositionAtInValue, FrameBasedMusicMap::EBarBeatFormat::Position),
+		Result.AsDecimal()
+	);
+*/
+	return VisitorCallback(Result);
+}
+
+void UMovieSceneMusicTimeTransform::ScaleBy(double UnwarpedScaleFactor)
+{
+	ScaleFactor = UnwarpedScaleFactor;
+}
+
+UE::MovieScene::ETimeWarpChannelDomain UMovieSceneMusicTimeTransform::GetDomain() const
+{
+	return UE::MovieScene::ETimeWarpChannelDomain::PlayRate;
+}
+
+FFrameTime UMovieSceneMusicTimeTransform::RemapTime_Internal(const MusicTimeTransform::UScopedMapLock& MapLock, FFrameTime In) const
+{
+	check(MapLock);
+	FMusicalTime OuterMusicalTimeAtClipStart = MapLock.OuterMap->GetMusicalTime(Offset);
+	FFrameTime   OuterFrameAtIn = Offset + FFrameRate::TransformTime(In, MapLock.InnerMap->FrameResolution, MapLock.OuterMap->FrameResolution);
+	FMusicalTime OuterMusicalTimeAtIn = MapLock.OuterMap->GetMusicalTime(OuterFrameAtIn);
+	double BarsFromPositionToStart = OuterMusicalTimeAtIn - OuterMusicalTimeAtClipStart;
+	FMusicalTime InnerMusicalTimeAtIn = MapLock.InnerMap->BarsToMusicalTime(BarsFromPositionToStart);
+	FFrameTime InnerFrameAtIn = MapLock.InnerMap->GetFrameTime(InnerMusicalTimeAtIn);
+	return InnerFrameAtIn;
+}
+
diff --git a/Engine/Source/Runtime/MovieScene/Public/Channels/MovieSceneAudioTriggerChannel.h b/Engine/Source/Runtime/MovieScene/Public/Channels/MovieSceneAudioTriggerChannel.h
--- a/Engine/Source/Runtime/MovieScene/Public/Channels/MovieSceneAudioTriggerChannel.h
+++ b/Engine/Source/Runtime/MovieScene/Public/Channels/MovieSceneAudioTriggerChannel.h
@@ -91,6 +91,7 @@ public:
 	MOVIESCENE_API virtual void DeleteKeys(TArrayView<const FKeyHandle> InHandles) override;
 	MOVIESCENE_API virtual void DeleteKeysFrom(FFrameNumber InTime, bool bDeleteKeysBefore) override;
 	MOVIESCENE_API virtual void ChangeFrameResolution(FFrameRate SourceRate, FFrameRate DestinationRate) override;
+	MOVIESCENE_API virtual void RemapMusicalFrameTimes(const UFrameBasedMusicMap* OldMap, const UFrameBasedMusicMap* NewMap) override;
 	MOVIESCENE_API virtual TRange<FFrameNumber> ComputeEffectiveRange() const override;
 	MOVIESCENE_API virtual int32 GetNumKeys() const override;
 	MOVIESCENE_API virtual void Reset() override;
diff --git a/Engine/Source/Runtime/MovieScene/Public/Channels/MovieSceneBoolChannel.h b/Engine/Source/Runtime/MovieScene/Public/Channels/MovieSceneBoolChannel.h
--- a/Engine/Source/Runtime/MovieScene/Public/Channels/MovieSceneBoolChannel.h
+++ b/Engine/Source/Runtime/MovieScene/Public/Channels/MovieSceneBoolChannel.h
@@ -116,6 +116,7 @@ public:
 	MOVIESCENE_API virtual void DeleteKeys(TArrayView<const FKeyHandle> InHandles) override;
 	MOVIESCENE_API virtual void DeleteKeysFrom(FFrameNumber InTime, bool bDeleteKeysBefore) override;
 	MOVIESCENE_API virtual void ChangeFrameResolution(FFrameRate SourceRate, FFrameRate DestinationRate) override;
+	MOVIESCENE_API virtual void RemapMusicalFrameTimes(const UFrameBasedMusicMap* OldMap, const UFrameBasedMusicMap* NewMap) override;
 	MOVIESCENE_API virtual TRange<FFrameNumber> ComputeEffectiveRange() const override;
 	MOVIESCENE_API virtual int32 GetNumKeys() const override;
 	MOVIESCENE_API virtual void Reset() override;
diff --git a/Engine/Source/Runtime/MovieScene/Public/Channels/MovieSceneByteChannel.h b/Engine/Source/Runtime/MovieScene/Public/Channels/MovieSceneByteChannel.h
--- a/Engine/Source/Runtime/MovieScene/Public/Channels/MovieSceneByteChannel.h
+++ b/Engine/Source/Runtime/MovieScene/Public/Channels/MovieSceneByteChannel.h
@@ -122,6 +122,7 @@ public:
 	MOVIESCENE_API virtual void DeleteKeys(TArrayView<const FKeyHandle> InHandles) override;
 	MOVIESCENE_API virtual void DeleteKeysFrom(FFrameNumber InTime, bool bDeleteKeysBefore) override;
 	MOVIESCENE_API virtual void ChangeFrameResolution(FFrameRate SourceRate, FFrameRate DestinationRate) override;
+	MOVIESCENE_API virtual void RemapMusicalFrameTimes(const UFrameBasedMusicMap* OldMap, const UFrameBasedMusicMap* NewMap) override;
 	MOVIESCENE_API virtual TRange<FFrameNumber> ComputeEffectiveRange() const override;
 	MOVIESCENE_API virtual int32 GetNumKeys() const override;
 	MOVIESCENE_API virtual void Reset() override;
diff --git a/Engine/Source/Runtime/MovieScene/Public/Channels/MovieSceneChannel.h b/Engine/Source/Runtime/MovieScene/Public/Channels/MovieSceneChannel.h
--- a/Engine/Source/Runtime/MovieScene/Public/Channels/MovieSceneChannel.h
+++ b/Engine/Source/Runtime/MovieScene/Public/Channels/MovieSceneChannel.h
@@ -19,6 +19,7 @@ struct FFrameNumber;
 struct FFrameRate;
 struct FKeyDataOptimizationParams;
 struct FMovieSceneChannel;
+class UFrameBasedMusicMap;
 
 /*
 *  Events that are fired when adding, deleting or moving keys
@@ -189,6 +190,15 @@ struct FMovieSceneChannel
 	virtual void ChangeFrameResolution(FFrameRate SourceRate, FFrameRate DestinationRate)
 	{}
 
+	/**
+	 * Called when the music map is changed in a sequence that supports music mode editing
+	 *
+	 * @param OldMap The music map the keyframes were originally in reference to.
+	 * @param NewMap The music map the keyframes should be "remapped" to.
+	 */
+	virtual void RemapMusicalFrameTimes(const UFrameBasedMusicMap* OldMap, const UFrameBasedMusicMap* NewMap)
+	{}
+
 	/**
 	 * Compute the effective range of this channel, for example, the extents of its key times
 	 *
@@ -238,7 +248,7 @@ struct FMovieSceneChannel
 	{}
 
 	/**
-	 * Perfor a possibly heavy operation after an edit change 
+	 * Perform a possibly heavy operation after an edit change 
 	 *
 	 */
 	virtual void PostEditChange() {}
diff --git a/Engine/Source/Runtime/MovieScene/Public/Channels/MovieSceneChannelData.h b/Engine/Source/Runtime/MovieScene/Public/Channels/MovieSceneChannelData.h
--- a/Engine/Source/Runtime/MovieScene/Public/Channels/MovieSceneChannelData.h
+++ b/Engine/Source/Runtime/MovieScene/Public/Channels/MovieSceneChannelData.h
@@ -166,6 +166,14 @@ struct FMovieSceneChannelData
 	 */
 	MOVIESCENE_API void ChangeFrameResolution(FFrameRate SourceRate, FFrameRate DestinationRate);
 
+	/**
+	 * Moves the location of keys to account for a changing music map.
+	 *
+	 * @param OldMap   The music map that the key's original time was in reference to.
+	 * @param NewMap   The music map that the key's new time should be in refernce to.
+	 */
+	MOVIESCENE_API void RemapMusicalFrameTimes(const UFrameBasedMusicMap* OldMap, const UFrameBasedMusicMap* NewMap);
+
 	/**
 	 * Get all the keys in the given range. Resulting arrays must be the same size where indices correspond to both arrays.
 	 *
diff --git a/Engine/Source/Runtime/MovieScene/Public/Channels/MovieSceneDoubleChannel.h b/Engine/Source/Runtime/MovieScene/Public/Channels/MovieSceneDoubleChannel.h
--- a/Engine/Source/Runtime/MovieScene/Public/Channels/MovieSceneDoubleChannel.h
+++ b/Engine/Source/Runtime/MovieScene/Public/Channels/MovieSceneDoubleChannel.h
@@ -260,6 +260,7 @@ public:
 	MOVIESCENE_API virtual void DeleteKeys(TArrayView<const FKeyHandle> InHandles) override;
 	MOVIESCENE_API virtual void DeleteKeysFrom(FFrameNumber InTime, bool bDeleteKeysBefore) override;
 	MOVIESCENE_API virtual void ChangeFrameResolution(FFrameRate SourceRate, FFrameRate DestinationRate) override;
+	MOVIESCENE_API virtual void RemapMusicalFrameTimes(const UFrameBasedMusicMap* OldMap, const UFrameBasedMusicMap* NewMap) override;
 	MOVIESCENE_API virtual TRange<FFrameNumber> ComputeEffectiveRange() const override;
 	MOVIESCENE_API virtual int32 GetNumKeys() const override;
 	MOVIESCENE_API virtual void Reset() override;
diff --git a/Engine/Source/Runtime/MovieScene/Public/Channels/MovieSceneFloatChannel.h b/Engine/Source/Runtime/MovieScene/Public/Channels/MovieSceneFloatChannel.h
--- a/Engine/Source/Runtime/MovieScene/Public/Channels/MovieSceneFloatChannel.h
+++ b/Engine/Source/Runtime/MovieScene/Public/Channels/MovieSceneFloatChannel.h
@@ -199,6 +199,7 @@ public:
 	MOVIESCENE_API virtual void DeleteKeys(TArrayView<const FKeyHandle> InHandles) override;
 	MOVIESCENE_API virtual void DeleteKeysFrom(FFrameNumber InTime, bool bDeleteKeysBefore) override;
 	MOVIESCENE_API virtual void ChangeFrameResolution(FFrameRate SourceRate, FFrameRate DestinationRate) override;
+	MOVIESCENE_API virtual void RemapMusicalFrameTimes(const UFrameBasedMusicMap* OldMap, const UFrameBasedMusicMap* NewMap) override;
 	MOVIESCENE_API virtual TRange<FFrameNumber> ComputeEffectiveRange() const override;
 	MOVIESCENE_API virtual int32 GetNumKeys() const override;
 	MOVIESCENE_API virtual void Reset() override;
diff --git a/Engine/Source/Runtime/MovieScene/Public/Channels/MovieSceneIntegerChannel.h b/Engine/Source/Runtime/MovieScene/Public/Channels/MovieSceneIntegerChannel.h
--- a/Engine/Source/Runtime/MovieScene/Public/Channels/MovieSceneIntegerChannel.h
+++ b/Engine/Source/Runtime/MovieScene/Public/Channels/MovieSceneIntegerChannel.h
@@ -152,6 +152,7 @@ public:
 	MOVIESCENE_API virtual void DeleteKeys(TArrayView<const FKeyHandle> InHandles) override;
 	MOVIESCENE_API virtual void DeleteKeysFrom(FFrameNumber InTime, bool bDeleteKeysBefore) override;
 	MOVIESCENE_API virtual void ChangeFrameResolution(FFrameRate SourceRate, FFrameRate DestinationRate) override;
+	MOVIESCENE_API virtual void RemapMusicalFrameTimes(const UFrameBasedMusicMap* OldMap, const UFrameBasedMusicMap* NewMap) override;
 	MOVIESCENE_API virtual TRange<FFrameNumber> ComputeEffectiveRange() const override;
 	MOVIESCENE_API virtual int32 GetNumKeys() const override;
 	MOVIESCENE_API virtual void Reset() override;
diff --git a/Engine/Source/Runtime/MovieScene/Public/Channels/MovieSceneObjectPathChannel.h b/Engine/Source/Runtime/MovieScene/Public/Channels/MovieSceneObjectPathChannel.h
--- a/Engine/Source/Runtime/MovieScene/Public/Channels/MovieSceneObjectPathChannel.h
+++ b/Engine/Source/Runtime/MovieScene/Public/Channels/MovieSceneObjectPathChannel.h
@@ -125,6 +125,7 @@ public:
 	MOVIESCENE_API virtual void DeleteKeys(TArrayView<const FKeyHandle> InHandles) override;
 	MOVIESCENE_API virtual void DeleteKeysFrom(FFrameNumber InTime, bool bDeleteKeysBefore) override;
 	MOVIESCENE_API virtual void ChangeFrameResolution(FFrameRate SourceRate, FFrameRate DestinationRate) override;
+	MOVIESCENE_API virtual void RemapMusicalFrameTimes(const UFrameBasedMusicMap* OldMap, const UFrameBasedMusicMap* NewMap) override;
 	MOVIESCENE_API virtual TRange<FFrameNumber> ComputeEffectiveRange() const override;
 	MOVIESCENE_API virtual int32 GetNumKeys() const override;
 	MOVIESCENE_API virtual void Reset() override;
diff --git a/Engine/Source/Runtime/MovieScene/Public/Evaluation/IMovieSceneCustomClockSource.h b/Engine/Source/Runtime/MovieScene/Public/Evaluation/IMovieSceneCustomClockSource.h
--- a/Engine/Source/Runtime/MovieScene/Public/Evaluation/IMovieSceneCustomClockSource.h
+++ b/Engine/Source/Runtime/MovieScene/Public/Evaluation/IMovieSceneCustomClockSource.h
@@ -42,6 +42,12 @@ public:
 
 	UFUNCTION()
 	virtual FFrameTime OnRequestCurrentTime(const FQualifiedFrameTime& InCurrentTime, float InPlayRate) { return 0; }
+
+	UFUNCTION()
+	virtual bool OnSupportsMusicalTime() { return false; }
+
+	UFUNCTION()
+	virtual FMusicalTime OnRequestCurrentMusicalTime(const FQualifiedFrameTime& InCurrentTime, float InPlayRate) { return FMusicalTime(); }
 };
 
 struct FMovieSceneTimeController_Custom : FMovieSceneTimeController
@@ -54,6 +60,8 @@ private:
 	MOVIESCENE_API virtual void OnStartPlaying(const FQualifiedFrameTime& InStartTime) override final;
 	MOVIESCENE_API virtual void OnStopPlaying(const FQualifiedFrameTime& InStopTime) override final;
 	MOVIESCENE_API virtual FFrameTime OnRequestCurrentTime(const FQualifiedFrameTime& InCurrentTime, float InPlayRate) override final;
+	MOVIESCENE_API virtual bool OnSupportsMusicalTime() override final;
+	MOVIESCENE_API virtual FMusicalTime OnRequestCurrentMusicalTime(const FQualifiedFrameTime& InCurrentTime, float InPlayRate) override final;
 
 private:
 
diff --git a/Engine/Source/Runtime/MovieScene/Public/Evaluation/MovieSceneSectionParameters.h b/Engine/Source/Runtime/MovieScene/Public/Evaluation/MovieSceneSectionParameters.h
--- a/Engine/Source/Runtime/MovieScene/Public/Evaluation/MovieSceneSectionParameters.h
+++ b/Engine/Source/Runtime/MovieScene/Public/Evaluation/MovieSceneSectionParameters.h
@@ -72,7 +72,7 @@ struct FMovieSceneSectionParameters
 
 public:
 	/** Number of frames (in display rate) to skip at the beginning of the sub-sequence. */
-	UPROPERTY(BlueprintReadWrite, EditAnywhere, Category="Clipping")
+	UPROPERTY(BlueprintReadWrite, EditAnywhere, Category="Clipping", meta=(UIFrameDisplayAs=Duration))
 	FFrameNumber StartFrameOffset;
 
 	/** Whether this section supports looping the sub-sequence. */
@@ -80,11 +80,11 @@ public:
 	bool bCanLoop = false;
 
 	/** Number of frames (in display rate) to skip at the beginning of the sub-sequence. */
-	UPROPERTY(BlueprintReadWrite, EditAnywhere, Category="Clipping", meta=(EditCondition="bCanLoop"))
+	UPROPERTY(BlueprintReadWrite, EditAnywhere, Category="Clipping", meta=(EditCondition="bCanLoop", UIFrameDisplayAs=Duration))
 	FFrameNumber EndFrameOffset;
 
 	/** Number of frames (in display rate) to offset the first loop of the sub-sequence. */
-	UPROPERTY(BlueprintReadWrite, EditAnywhere, Category="Clipping", meta=(EditCondition="bCanLoop"))
+	UPROPERTY(BlueprintReadWrite, EditAnywhere, Category="Clipping", meta=(EditCondition="bCanLoop", UIFrameDisplayAs=Duration))
 	FFrameNumber FirstLoopStartFrameOffset;
 
 	/** Playback time scaling factor. */
diff --git a/Engine/Source/Runtime/MovieScene/Public/MovieScene.h b/Engine/Source/Runtime/MovieScene/Public/MovieScene.h
--- a/Engine/Source/Runtime/MovieScene/Public/MovieScene.h
+++ b/Engine/Source/Runtime/MovieScene/Public/MovieScene.h
@@ -11,6 +11,7 @@
 #include "CoreTypes.h"
 #include "Delegates/Delegate.h"
 #include "EventHandlers/MovieSceneDataEventContainer.h"
+#include "FrameBasedMusicMap.h"
 #include "HAL/PlatformCrt.h"
 #include "Internationalization/Text.h"
 #include "Math/Color.h"
@@ -799,10 +800,7 @@ public:
 	/**
 	 * Directly set the tick resolution for this movie scene without applying any conversion whatsoever, or modifying the data
 	 */
-	void SetTickResolutionDirectly(FFrameRate InTickResolution)
-	{
-		TickResolution = InTickResolution;
-	}
+	MOVIESCENE_API void SetTickResolutionDirectly(FFrameRate InTickResolution);
 
 	/**
 	 * Retrieve the display frame rate for this data, in which frame numbers should be displayed on UI, and interacted with in movie scene players
@@ -1201,6 +1199,34 @@ public:
 	/* Called by the compiler to empty the list of generated conditions*/
 	void ResetGeneratedConditions() { GeneratedConditions.Reset(); }
 
+
+	MOVIESCENE_API bool HasMusicMap() const
+	{
+		return (bool)MusicMap;
+	}
+
+	MOVIESCENE_API void RemoveMusicMap()
+	{
+		MusicMap = nullptr;
+	}
+
+	MOVIESCENE_API const UFrameBasedMusicMap* GetMusicMap() const
+	{
+		return MusicMap.Get();
+	}
+
+	MOVIESCENE_API UFrameBasedMusicMap* GetMusicMap()
+	{
+		return MusicMap.Get();
+	}
+
+	void SetMusicMapDirectly(UFrameBasedMusicMap* InMusicMap)
+	{
+		check(InMusicMap);
+		check(InMusicMap->GetOuter() == this);
+		MusicMap = InMusicMap;
+	}
+
 protected:
 
 	/**
@@ -1315,6 +1341,9 @@ private:
 	UPROPERTY()
 	TArray<TObjectPtr<UMovieSceneGroupCondition>> GeneratedConditions;
 
+	UPROPERTY()
+	TObjectPtr<UFrameBasedMusicMap> MusicMap;
+
 #if WITH_EDITORONLY_DATA
 
 	/** Indicates whether this movie scene is read only */
diff --git a/Engine/Source/Runtime/MovieScene/Public/MovieSceneSection.h b/Engine/Source/Runtime/MovieScene/Public/MovieSceneSection.h
--- a/Engine/Source/Runtime/MovieScene/Public/MovieSceneSection.h
+++ b/Engine/Source/Runtime/MovieScene/Public/MovieSceneSection.h
@@ -55,6 +55,7 @@ struct FMovieSceneSequenceID;
 struct FPropertyChangedEvent;
 struct FQualifiedFrameTime;
 struct FMovieSceneTimeWarpVariant;
+class UFrameBasedMusicMap;
 
 namespace UE::MovieScene
 {
@@ -608,6 +609,9 @@ public:
 	/* Migrate the frame times of the movie scene section from the source frame rate to the destination frame rate */
 	virtual void MigrateFrameTimes(FFrameRate SourceRate, FFrameRate DestinationRate) {}
 
+	/* Remap the frame times of the movie scene section from the source music map to the destination music map */
+	MOVIESCENE_API virtual void RemapMusicalFrameTimes(const UFrameBasedMusicMap* OldMap, const UFrameBasedMusicMap* NewMap);
+
 	/** When guid bindings are updated to allow this section to fix-up any internal bindings */
 	virtual void OnBindingIDsUpdated(const TMap<UE::MovieScene::FFixedObjectBindingID, UE::MovieScene::FFixedObjectBindingID>& OldFixedToNewFixedMap, FMovieSceneSequenceID LocalSequenceID, TSharedRef<UE::MovieScene::FSharedPlaybackState> SharedPlaybackState) {}
 
@@ -755,11 +759,11 @@ public:
 private:
 
 	/** The amount of time to prepare this section for evaluation before it actually starts. */
-	UPROPERTY(EditAnywhere, AdvancedDisplay, Category="Section", meta=(UIMin=0))
+	UPROPERTY(EditAnywhere, AdvancedDisplay, Category="Section", meta=(UIMin=0, UIFrameDisplayAs=NonMusical))
 	FFrameNumber PreRollFrames;
 
 	/** The amount of time to continue 'postrolling' this section for after evaluation has ended. */
-	UPROPERTY(EditAnywhere, AdvancedDisplay, Category="Section", meta=(UIMin=0))
+	UPROPERTY(EditAnywhere, AdvancedDisplay, Category="Section", meta=(UIMin=0, UIFrameDisplayAs=NonMusical))
 	FFrameNumber PostRollFrames;
 
 	/** The row index that this section sits on */
diff --git a/Engine/Source/Runtime/MovieScene/Public/MovieSceneTimeController.h b/Engine/Source/Runtime/MovieScene/Public/MovieSceneTimeController.h
--- a/Engine/Source/Runtime/MovieScene/Public/MovieSceneTimeController.h
+++ b/Engine/Source/Runtime/MovieScene/Public/MovieSceneTimeController.h
@@ -6,6 +6,7 @@
 #include "Misc/QualifiedFrameTime.h"
 #include "MovieSceneFwd.h"
 #include "Misc/Attribute.h"
+#include "Misc/MusicalTime.h"
 #include "UObject/ObjectMacros.h"
 
 class IMovieScenePlayer;
@@ -43,6 +44,16 @@ public:
 	 */
 	MOVIESCENE_API FFrameTime RequestCurrentTime(const FQualifiedFrameTime& InCurrentTime, float InPlayRate, FFrameRate InDisplayRate);
 
+	MOVIESCENE_API bool SupportsMusicalTime();
+
+	/**
+	 * Request the current musical time based on the specified existing time and play rate.
+	 *
+	 * @param InCurrentTime    The current time of the sequence
+	 * @param InPlayRate       The current play rate of the sequence, multiplied by any world actor settings global dilation
+	 */
+	MOVIESCENE_API FMusicalTime RequestCurrentMusicalTime(const FQualifiedFrameTime& InCurrentTime, float InPlayRate, FFrameRate InDisplayRate);
+
 	/**
 	 * Called when the status of the owning IMovieScenePlayer has changed
 	 */
@@ -59,6 +70,8 @@ protected:
 	virtual void OnStartPlaying(const FQualifiedFrameTime& InStartTime){}
 	virtual void OnStopPlaying(const FQualifiedFrameTime& InStopTime){}
 	virtual FFrameTime OnRequestCurrentTime(const FQualifiedFrameTime& InCurrentTime, float InPlayRate) = 0;
+	virtual FMusicalTime OnRequestCurrentMusicalTime(const FQualifiedFrameTime& InCurrentTime, float InPlayRate) { return FMusicalTime(); }
+	virtual bool OnSupportsMusicalTime() { return false; }
 
 protected:
 
diff --git a/Engine/Source/Runtime/MovieScene/Public/MovieSceneTimeHelpers.h b/Engine/Source/Runtime/MovieScene/Public/MovieSceneTimeHelpers.h
--- a/Engine/Source/Runtime/MovieScene/Public/MovieSceneTimeHelpers.h
+++ b/Engine/Source/Runtime/MovieScene/Public/MovieSceneTimeHelpers.h
@@ -14,6 +14,7 @@
 
 class UMovieScene;
 struct FFrameRate;
+class UFrameBasedMusicMap;
 
 namespace UE
 {
@@ -30,6 +31,13 @@ public:
 
 static MOVIESCENE_API void MigrateFrameTimes(FFrameRate SourceRate, FFrameRate DestinationRate, UMovieScene* MovieScene, bool bApplyRecursively = false);
 
+/**
+ * Remap musical frame times from the current music map to a new music map. This remapping will cause the 
+ * frames to maintain their musical position (bar.beat), but will likely cause them to change there position
+ * in seconds as the new map may have different tempo(s).
+ */
+static MOVIESCENE_API void RemapMusicalFrameTimes(UMovieScene* MovieScene, UFrameBasedMusicMap* NewMap);
+
 };
 
 /**
diff --git a/Engine/Source/Runtime/MovieScene/Public/Sections/MovieSceneSectionTimingParameters.h b/Engine/Source/Runtime/MovieScene/Public/Sections/MovieSceneSectionTimingParameters.h
--- a/Engine/Source/Runtime/MovieScene/Public/Sections/MovieSceneSectionTimingParameters.h
+++ b/Engine/Source/Runtime/MovieScene/Public/Sections/MovieSceneSectionTimingParameters.h
@@ -9,6 +9,7 @@
 
 
 struct FMovieSceneSequenceTransform;
+class UMovieScene;
 
 /**
  * Parameter utility that converts section timing parameters to a transform using Seconds values.
@@ -73,6 +74,16 @@ struct FMovieSceneSectionTimingParametersSeconds
 	 * Make a transform out of these parameters
 	 */
 	MOVIESCENE_API FMovieSceneSequenceTransform MakeTransform(const FFrameRate& OuterFrameRate, const TRange<FFrameNumber>& OuterRange, double SourceDuration, double InnerPlayRate) const;
+
+	MOVIESCENE_API void AddPositionInOuterAsOffset(FMovieSceneSequenceTransform& Transform, const FFrameTime& Offset) const;
+
+	MOVIESCENE_API void AddPlayRate(FMovieSceneSequenceTransform& Transform) const;
+
+	MOVIESCENE_API void AddFrameRateConversion(FMovieSceneSequenceTransform& Transform, const FFrameRate& OuterFrameRate, const FFrameRate& InnerFrameRate) const;
+
+	MOVIESCENE_API void AddInnerStartOffset(FMovieSceneSequenceTransform& Transform, const FFrameTime& Offset) const;
+
+	MOVIESCENE_API void AddLoopingOrClampingAndReverse(FMovieSceneSequenceTransform& Result, const double Duration, const double Offset) const;
 };
 
 
@@ -139,6 +150,27 @@ struct FMovieSceneSectionTimingParametersFrames
 	 * Make a transform out of these parameters
 	 */
 	MOVIESCENE_API FMovieSceneSequenceTransform MakeTransform(const FFrameRate& OuterFrameRate, const TRange<FFrameNumber>& OuterRange, const FFrameRate& InnerFrameRate, const TRange<FFrameNumber>& InnerRange) const;
+
+	/**
+	 * Make a musical transform out of these parameters
+	 */
+	MOVIESCENE_API FMovieSceneSequenceTransform MakeMusicalTransform(const FFrameRate& OuterFrameRate, const TRange<FFrameNumber>& OuterRange, UMovieScene* OuterMovieScene, const FFrameRate& InnerFrameRate, const TRange<FFrameNumber>& InnerRange, UMovieScene* InnerMovieScene) const;
+
+	/**
+	 * Helpers
+	 */
+	void AddPositionInOuterAsOffset(FMovieSceneSequenceTransform& Transform, const FFrameTime& Offset) const;
+
+	void AddPlayRate(FMovieSceneSequenceTransform& Transform) const;
+
+	void AddFrameRateConversion(FMovieSceneSequenceTransform& Transform, const FFrameRate& OuterFrameRate, const FFrameRate& InnerFrameRate) const;
+
+	void AddInnerStartOffset(FMovieSceneSequenceTransform& Transform, const FFrameTime& Offset) const;
+
+	void AddLoopingOrClampingAndReverse(FMovieSceneSequenceTransform& Result, const FFrameNumber& Duration, const FFrameNumber& Offset) const;
+
+	void AddMusicalTransform(FMovieSceneSequenceTransform& Result, const FFrameNumber& Offset, UMovieScene* OuterMovieScene, UMovieScene* InnerMovieScene) const;
+
 };
 
 namespace UE::MovieScene
diff --git a/Engine/Source/Runtime/MovieScene/Public/Sections/MovieSceneSubSection.h b/Engine/Source/Runtime/MovieScene/Public/Sections/MovieSceneSubSection.h
--- a/Engine/Source/Runtime/MovieScene/Public/Sections/MovieSceneSubSection.h
+++ b/Engine/Source/Runtime/MovieScene/Public/Sections/MovieSceneSubSection.h
@@ -228,6 +228,7 @@ public:
 	virtual TOptional<FFrameTime> GetOffsetTime() const override { return TOptional<FFrameTime>(FFrameTime(Parameters.StartFrameOffset)); }
 	MOVIESCENE_API virtual void GetSnapTimes(TArray<FFrameNumber>& OutSnapTimes, bool bGetSectionBorders) const override;
 	MOVIESCENE_API virtual void MigrateFrameTimes(FFrameRate SourceRate, FFrameRate DestinationRate) override;
+	MOVIESCENE_API virtual void RemapMusicalFrameTimes(const UFrameBasedMusicMap* OldMap, const UFrameBasedMusicMap* NewMap) override;
 	MOVIESCENE_API virtual FMovieSceneTimeWarpVariant* GetTimeWarp() override;
 	MOVIESCENE_API virtual UObject* GetSourceObject() const override;
 
diff --git a/Engine/Source/Runtime/MovieScene/Public/Variants/MovieSceneMusicTimeTransform.h b/Engine/Source/Runtime/MovieScene/Public/Variants/MovieSceneMusicTimeTransform.h
new file mode 100644
--- /dev/null
+++ b/Engine/Source/Runtime/MovieScene/Public/Variants/MovieSceneMusicTimeTransform.h
@@ -0,0 +1,47 @@
+// Copyright Epic Games, Inc. All Rights Reserved.
+#pragma once
+
+#include "Variants/MovieSceneTimeWarpGetter.h"
+#include "MovieSceneMusicTimeTransform.generated.h"
+
+namespace MusicTimeTransform { class UScopedMapLock; }
+
+/**
+ * A time-warp defined as a mapping between two musical sequences.
+ */
+ UCLASS(MinimalAPI, DisplayName="Music Time Transform")
+class UMovieSceneMusicTimeTransform : public UMovieSceneTimeWarpGetter
+{
+public:
+
+	GENERATED_BODY()
+
+	UMovieSceneMusicTimeTransform();
+
+	void Initialize(UMovieScene* Outer, UMovieScene* Inner, FFrameTime InOffset);
+
+	/* Begin UMovieSceneTimeWarpGetter Implementation */
+	FFrameTime RemapTime(FFrameTime In) const override;
+	TOptional<FFrameTime> InverseRemapTimeCycled(FFrameTime InValue, FFrameTime InTimeHint, const UE::MovieScene::FInverseTransformTimeParams& Params) const override;
+	TRange<FFrameTime> ComputeTraversedHull(const TRange<FFrameTime>& Range) const override;
+	bool InverseRemapTimeWithinRange(FFrameTime InTime, FFrameTime RangeStart, FFrameTime RangeEnd, const TFunctionRef<bool(FFrameTime)>& VisitorCallback) const override;
+	void InitializeDefaults() override;
+	void ScaleBy(double UnwarpedScaleFactor) override;
+	UE::MovieScene::ETimeWarpChannelDomain GetDomain() const override;
+	/* End UMovieSceneTimeWarpGetter Implementation */
+
+private:
+	friend class MusicTimeTransform::UScopedMapLock;
+
+	UPROPERTY(transient)
+	TWeakObjectPtr<UMovieScene> OuterMusicMapOwner;
+
+	UPROPERTY(transient)
+	TWeakObjectPtr<UMovieScene> InnerMusicMapOwner;
+
+	double ScaleFactor = 1.0;
+
+	FFrameTime Offset;
+
+	FFrameTime RemapTime_Internal(const MusicTimeTransform::UScopedMapLock& MapLock, FFrameTime In) const;
+};
\ No newline at end of file
diff --git a/Engine/Source/Runtime/MovieSceneTracks/Private/Channels/MovieSceneCameraShakeSourceTriggerChannel.cpp b/Engine/Source/Runtime/MovieSceneTracks/Private/Channels/MovieSceneCameraShakeSourceTriggerChannel.cpp
--- a/Engine/Source/Runtime/MovieSceneTracks/Private/Channels/MovieSceneCameraShakeSourceTriggerChannel.cpp
+++ b/Engine/Source/Runtime/MovieSceneTracks/Private/Channels/MovieSceneCameraShakeSourceTriggerChannel.cpp
@@ -39,6 +39,11 @@ void FMovieSceneCameraShakeSourceTriggerChannel::ChangeFrameResolution(FFrameRat
 	GetData().ChangeFrameResolution(SourceRate, DestinationRate);
 }
 
+void FMovieSceneCameraShakeSourceTriggerChannel::RemapMusicalFrameTimes(const UFrameBasedMusicMap* OldMap, const UFrameBasedMusicMap* NewMap)
+{
+	GetData().RemapMusicalFrameTimes(OldMap, NewMap);
+}
+
 TRange<FFrameNumber> FMovieSceneCameraShakeSourceTriggerChannel::ComputeEffectiveRange() const
 {
 	return GetData().GetTotalRange();
diff --git a/Engine/Source/Runtime/MovieSceneTracks/Private/Channels/MovieSceneEventChannel.cpp b/Engine/Source/Runtime/MovieSceneTracks/Private/Channels/MovieSceneEventChannel.cpp
--- a/Engine/Source/Runtime/MovieSceneTracks/Private/Channels/MovieSceneEventChannel.cpp
+++ b/Engine/Source/Runtime/MovieSceneTracks/Private/Channels/MovieSceneEventChannel.cpp
@@ -39,6 +39,11 @@ void FMovieSceneEventChannel::ChangeFrameResolution(FFrameRate SourceRate, FFram
 	GetData().ChangeFrameResolution(SourceRate, DestinationRate);
 }
 
+void FMovieSceneEventChannel::RemapMusicalFrameTimes(const UFrameBasedMusicMap* OldMap, const UFrameBasedMusicMap* NewMap)
+{
+	GetData().RemapMusicalFrameTimes(OldMap, NewMap);
+}
+
 TRange<FFrameNumber> FMovieSceneEventChannel::ComputeEffectiveRange() const
 {
 	return GetData().GetTotalRange();
diff --git a/Engine/Source/Runtime/MovieSceneTracks/Private/Channels/MovieSceneStringChannel.cpp b/Engine/Source/Runtime/MovieSceneTracks/Private/Channels/MovieSceneStringChannel.cpp
--- a/Engine/Source/Runtime/MovieSceneTracks/Private/Channels/MovieSceneStringChannel.cpp
+++ b/Engine/Source/Runtime/MovieSceneTracks/Private/Channels/MovieSceneStringChannel.cpp
@@ -107,6 +107,11 @@ void FMovieSceneStringChannel::ChangeFrameResolution(FFrameRate SourceRate, FFra
 	GetData().ChangeFrameResolution(SourceRate, DestinationRate);
 }
 
+void FMovieSceneStringChannel::RemapMusicalFrameTimes(const UFrameBasedMusicMap* OldMap, const UFrameBasedMusicMap* NewMap)
+{
+	GetData().RemapMusicalFrameTimes(OldMap, NewMap);
+}
+
 TRange<FFrameNumber> FMovieSceneStringChannel::ComputeEffectiveRange() const
 {
 	return GetData().GetTotalRange();
diff --git a/Engine/Source/Runtime/MovieSceneTracks/Private/Sections/MovieSceneActorReferenceSection.cpp b/Engine/Source/Runtime/MovieSceneTracks/Private/Sections/MovieSceneActorReferenceSection.cpp
--- a/Engine/Source/Runtime/MovieSceneTracks/Private/Sections/MovieSceneActorReferenceSection.cpp
+++ b/Engine/Source/Runtime/MovieSceneTracks/Private/Sections/MovieSceneActorReferenceSection.cpp
@@ -75,6 +75,11 @@ void FMovieSceneActorReferenceData::ChangeFrameResolution(FFrameRate SourceRate,
 	GetData().ChangeFrameResolution(SourceRate, DestinationRate);
 }
 
+void FMovieSceneActorReferenceData::RemapMusicalFrameTimes(const UFrameBasedMusicMap* OldMap, const UFrameBasedMusicMap* NewMap)
+{
+	GetData().RemapMusicalFrameTimes(OldMap, NewMap);
+}
+
 TRange<FFrameNumber> FMovieSceneActorReferenceData::ComputeEffectiveRange() const
 {
 	return GetData().GetTotalRange();
diff --git a/Engine/Source/Runtime/MovieSceneTracks/Private/Sections/MovieSceneEventSection.cpp b/Engine/Source/Runtime/MovieSceneTracks/Private/Sections/MovieSceneEventSection.cpp
--- a/Engine/Source/Runtime/MovieSceneTracks/Private/Sections/MovieSceneEventSection.cpp
+++ b/Engine/Source/Runtime/MovieSceneTracks/Private/Sections/MovieSceneEventSection.cpp
@@ -374,6 +374,11 @@ void FMovieSceneEventSectionData::ChangeFrameResolution(FFrameRate SourceRate, F
 	GetData().ChangeFrameResolution(SourceRate, DestinationRate);
 }
 
+void FMovieSceneEventSectionData::RemapMusicalFrameTimes(const UFrameBasedMusicMap* OldMap, const UFrameBasedMusicMap* NewMap)
+{
+	GetData().RemapMusicalFrameTimes(OldMap, NewMap);
+}
+
 TRange<FFrameNumber> FMovieSceneEventSectionData::ComputeEffectiveRange() const
 {
 	return GetData().GetTotalRange();
diff --git a/Engine/Source/Runtime/MovieSceneTracks/Public/Channels/MovieSceneCameraShakeSourceTriggerChannel.h b/Engine/Source/Runtime/MovieSceneTracks/Public/Channels/MovieSceneCameraShakeSourceTriggerChannel.h
--- a/Engine/Source/Runtime/MovieSceneTracks/Public/Channels/MovieSceneCameraShakeSourceTriggerChannel.h
+++ b/Engine/Source/Runtime/MovieSceneTracks/Public/Channels/MovieSceneCameraShakeSourceTriggerChannel.h
@@ -82,6 +82,7 @@ public:
 	MOVIESCENETRACKS_API virtual void DeleteKeys(TArrayView<const FKeyHandle> InHandles) override;
 	MOVIESCENETRACKS_API virtual void DeleteKeysFrom(FFrameNumber InTime, bool bDeleteKeysBefore) override;
 	MOVIESCENETRACKS_API virtual void ChangeFrameResolution(FFrameRate SourceRate, FFrameRate DestinationRate) override;
+	MOVIESCENETRACKS_API virtual void RemapMusicalFrameTimes(const UFrameBasedMusicMap* OldMap, const UFrameBasedMusicMap* NewMap) override;
 	MOVIESCENETRACKS_API virtual TRange<FFrameNumber> ComputeEffectiveRange() const override;
 	MOVIESCENETRACKS_API virtual int32 GetNumKeys() const override;
 	MOVIESCENETRACKS_API virtual void Reset() override;
diff --git a/Engine/Source/Runtime/MovieSceneTracks/Public/Channels/MovieSceneEventChannel.h b/Engine/Source/Runtime/MovieSceneTracks/Public/Channels/MovieSceneEventChannel.h
--- a/Engine/Source/Runtime/MovieSceneTracks/Public/Channels/MovieSceneEventChannel.h
+++ b/Engine/Source/Runtime/MovieSceneTracks/Public/Channels/MovieSceneEventChannel.h
@@ -56,6 +56,7 @@ public:
 	MOVIESCENETRACKS_API virtual void DeleteKeys(TArrayView<const FKeyHandle> InHandles) override;
 	MOVIESCENETRACKS_API virtual void DeleteKeysFrom(FFrameNumber InTime, bool bDeleteKeysBefore) override;
 	MOVIESCENETRACKS_API virtual void ChangeFrameResolution(FFrameRate SourceRate, FFrameRate DestinationRate) override;
+	MOVIESCENETRACKS_API virtual void RemapMusicalFrameTimes(const UFrameBasedMusicMap* OldMap, const UFrameBasedMusicMap* NewMap) override;
 	MOVIESCENETRACKS_API virtual TRange<FFrameNumber> ComputeEffectiveRange() const override;
 	MOVIESCENETRACKS_API virtual int32 GetNumKeys() const override;
 	MOVIESCENETRACKS_API virtual void Reset() override;
diff --git a/Engine/Source/Runtime/MovieSceneTracks/Public/Channels/MovieSceneStringChannel.h b/Engine/Source/Runtime/MovieSceneTracks/Public/Channels/MovieSceneStringChannel.h
--- a/Engine/Source/Runtime/MovieSceneTracks/Public/Channels/MovieSceneStringChannel.h
+++ b/Engine/Source/Runtime/MovieSceneTracks/Public/Channels/MovieSceneStringChannel.h
@@ -106,6 +106,7 @@ public:
 	MOVIESCENETRACKS_API virtual void DeleteKeys(TArrayView<const FKeyHandle> InHandles) override;
 	MOVIESCENETRACKS_API virtual void DeleteKeysFrom(FFrameNumber InTime, bool bDeleteKeysBefore) override;
 	MOVIESCENETRACKS_API virtual void ChangeFrameResolution(FFrameRate SourceRate, FFrameRate DestinationRate) override;
+	MOVIESCENETRACKS_API virtual void RemapMusicalFrameTimes(const UFrameBasedMusicMap* OldMap, const UFrameBasedMusicMap* NewMap) override;
 	MOVIESCENETRACKS_API virtual TRange<FFrameNumber> ComputeEffectiveRange() const override;
 	MOVIESCENETRACKS_API virtual int32 GetNumKeys() const override;
 	MOVIESCENETRACKS_API virtual void Reset() override;
diff --git a/Engine/Source/Runtime/MovieSceneTracks/Public/Sections/MovieSceneActorReferenceSection.h b/Engine/Source/Runtime/MovieSceneTracks/Public/Sections/MovieSceneActorReferenceSection.h
--- a/Engine/Source/Runtime/MovieSceneTracks/Public/Sections/MovieSceneActorReferenceSection.h
+++ b/Engine/Source/Runtime/MovieSceneTracks/Public/Sections/MovieSceneActorReferenceSection.h
@@ -116,6 +116,7 @@ public:
 	MOVIESCENETRACKS_API virtual void DeleteKeys(TArrayView<const FKeyHandle> InHandles) override;
 	MOVIESCENETRACKS_API virtual void DeleteKeysFrom(FFrameNumber InTime, bool bDeleteKeysBefore) override;
 	MOVIESCENETRACKS_API virtual void ChangeFrameResolution(FFrameRate SourceRate, FFrameRate DestinationRate) override;
+	MOVIESCENETRACKS_API virtual void RemapMusicalFrameTimes(const UFrameBasedMusicMap* OldMap, const UFrameBasedMusicMap* NewMap) override;
 	MOVIESCENETRACKS_API virtual TRange<FFrameNumber> ComputeEffectiveRange() const override;
 	MOVIESCENETRACKS_API virtual int32 GetNumKeys() const override;
 	MOVIESCENETRACKS_API virtual void Reset() override;
diff --git a/Engine/Source/Runtime/MovieSceneTracks/Public/Sections/MovieSceneAudioSection.h b/Engine/Source/Runtime/MovieSceneTracks/Public/Sections/MovieSceneAudioSection.h
--- a/Engine/Source/Runtime/MovieSceneTracks/Public/Sections/MovieSceneAudioSection.h
+++ b/Engine/Source/Runtime/MovieSceneTracks/Public/Sections/MovieSceneAudioSection.h
@@ -214,7 +214,7 @@ private:
 	TObjectPtr<USoundBase> Sound;
 
 	/** The offset into the beginning of the audio clip */
-	UPROPERTY(EditAnywhere, Category="Audio")
+	UPROPERTY(EditAnywhere, Category="Audio", meta=(UIFrameDisplayAs=NonMusical))
 	FFrameNumber StartFrameOffset;
 
 	/** The offset into the beginning of the audio clip */
diff --git a/Engine/Source/Runtime/MovieSceneTracks/Public/Sections/MovieSceneEventSection.h b/Engine/Source/Runtime/MovieSceneTracks/Public/Sections/MovieSceneEventSection.h
--- a/Engine/Source/Runtime/MovieSceneTracks/Public/Sections/MovieSceneEventSection.h
+++ b/Engine/Source/Runtime/MovieSceneTracks/Public/Sections/MovieSceneEventSection.h
@@ -189,6 +189,7 @@ public:
 	MOVIESCENETRACKS_API virtual void DeleteKeys(TArrayView<const FKeyHandle> InHandles) override;
 	MOVIESCENETRACKS_API virtual void DeleteKeysFrom(FFrameNumber InTime, bool bDeleteKeysBefore) override;
 	MOVIESCENETRACKS_API virtual void ChangeFrameResolution(FFrameRate SourceRate, FFrameRate DestinationRate) override;
+	MOVIESCENETRACKS_API virtual void RemapMusicalFrameTimes(const UFrameBasedMusicMap* OldMap, const UFrameBasedMusicMap* NewMap) override;
 	MOVIESCENETRACKS_API virtual TRange<FFrameNumber> ComputeEffectiveRange() const override;
 	MOVIESCENETRACKS_API virtual int32 GetNumKeys() const override;
 	MOVIESCENETRACKS_API virtual void Reset() override;
diff --git a/Engine/Source/Runtime/TimeManagement/Private/FrameBasedMusicMap.cpp b/Engine/Source/Runtime/TimeManagement/Private/FrameBasedMusicMap.cpp
new file mode 100644
--- /dev/null
+++ b/Engine/Source/Runtime/TimeManagement/Private/FrameBasedMusicMap.cpp
@@ -0,0 +1,569 @@
+// Copyright Epic Games, Inc. All Rights Reserved.
+
+#include "FrameBasedMusicMap.h"
+#include "Algo/FindLast.h"
+
+DEFINE_LOG_CATEGORY_STATIC(LogMusicMap, Log, All);
+#define LOCTEXT_NAMESPACE "FrameBasedMusicMap"
+
+namespace FrameBasedMusicMap
+{
+	FText GetQuantizationName(int32 Quantization)
+	{
+		switch (Quantization)
+		{
+		case ETripletQuantizationTicks::ThirtysecondNoteTriplet:	return LOCTEXT("ThirtySecondNoteTriplet", "1/32 Note Triplet");
+		case ETripletQuantizationTicks::SixteenthNoteTriplet:		return LOCTEXT("SixteenthNoteTriplet", "1/16 Note Triplet");
+		case ETripletQuantizationTicks::EighthNoteTriplet:			return LOCTEXT("EighthNoteTriplet", "1/8 Note Triplet");
+		case ETripletQuantizationTicks::QuarterNoteTriplet:			return LOCTEXT("QuarterNoteTriplet", "1/4 Note Triplet");
+		case ETripletQuantizationTicks::HalfNoteTriplet:			return LOCTEXT("HalfNoteTriplet", "1/2 Note Triplet");
+		case ETripletQuantizationTicks::WholeNoteTriplet:			return LOCTEXT("WholeNoteTriplet", "Whole Note Triplet");
+		
+		case EWholeQuantizationTicks::ThirtysecondNote:	return LOCTEXT("ThirtySecondNote", "1/32 Note");
+		case EWholeQuantizationTicks::SixteenthNote:	return LOCTEXT("SixteenthNote", "1/16 Note");
+		case EWholeQuantizationTicks::EighthNote:		return LOCTEXT("EighthNote", "1/8 Note");
+		case EWholeQuantizationTicks::QuarterNote:		return LOCTEXT("QuarterNote", "1/4 Note");
+		case EWholeQuantizationTicks::HalfNote:			return LOCTEXT("HalfNote", "1/2 Note");
+		case EWholeQuantizationTicks::WholeNote:		return LOCTEXT("WholeNote", "Whole Note");
+		case EWholeQuantizationTicks::Bar:				return LOCTEXT("Bar", "Bar");
+		case EWholeQuantizationTicks::Beat:				return LOCTEXT("Beat", "Beat");
+
+
+		case EDottedQuantizationTicks::DottedThirtysecondNote:	return LOCTEXT("DottedThirtySecondNote", "Dotted 1/32 Note");
+		case EDottedQuantizationTicks::DottedSixteenthNote:		return LOCTEXT("DottedSixteenthNote", "Dotted 1/16 Note");
+		case EDottedQuantizationTicks::DottedEighthNote:		return LOCTEXT("DottedEighthNote", "Dotted 1/8 Note");
+		case EDottedQuantizationTicks::DottedQuarterNote:		return LOCTEXT("DottedQuarterNote", "Dotted 1/4 Note");
+		case EDottedQuantizationTicks::DottedHalfNote:			return LOCTEXT("DottedHalfNote", "Dotted Half Note");
+		case EDottedQuantizationTicks::DottedWholeNote:			return LOCTEXT("DottedWholeNote", "Dotted Whole Note");
+		}
+		return LOCTEXT("Unknown", "Unrecognized Qunatization Factor");
+	}
+
+	TArray<TSharedPtr<int32>> QuantizationOptionsSource = {
+		MakeShared<int32>(EWholeQuantizationTicks::Bar),
+		MakeShared<int32>(EWholeQuantizationTicks::Beat),
+		MakeShared<int32>(EWholeQuantizationTicks::WholeNote),
+		MakeShared<int32>(EWholeQuantizationTicks::HalfNote),
+		MakeShared<int32>(EWholeQuantizationTicks::QuarterNote),
+		MakeShared<int32>(EWholeQuantizationTicks::EighthNote),
+		MakeShared<int32>(EWholeQuantizationTicks::SixteenthNote),
+		MakeShared<int32>(EWholeQuantizationTicks::ThirtysecondNote),
+		MakeShared<int32>(ETripletQuantizationTicks::WholeNoteTriplet),
+		MakeShared<int32>(ETripletQuantizationTicks::HalfNoteTriplet),
+		MakeShared<int32>(ETripletQuantizationTicks::QuarterNoteTriplet),
+		MakeShared<int32>(ETripletQuantizationTicks::EighthNoteTriplet),
+		MakeShared<int32>(ETripletQuantizationTicks::SixteenthNoteTriplet),
+		MakeShared<int32>(ETripletQuantizationTicks::ThirtysecondNoteTriplet),
+		MakeShared<int32>(EDottedQuantizationTicks::DottedWholeNote),
+		MakeShared<int32>(EDottedQuantizationTicks::DottedHalfNote),
+		MakeShared<int32>(EDottedQuantizationTicks::DottedQuarterNote),
+		MakeShared<int32>(EDottedQuantizationTicks::DottedEighthNote),
+		MakeShared<int32>(EDottedQuantizationTicks::DottedSixteenthNote),
+		MakeShared<int32>(EDottedQuantizationTicks::DottedThirtysecondNote)
+	};
+}
+
+void UFrameBasedMusicMap::SetFrameResolution(const FFrameRate& NewFrameResolution)
+{
+	check(NewFrameResolution.IsValid());
+
+	if (NewFrameResolution == FrameResolution)
+	{
+		return;
+	}
+
+	if (FrameResolution.IsValid() && !IsEmpty() && !IsSimple())
+	{
+		// must be changing sequence's resolution...
+		for (FFrameBasedTimeSignaturePoint& Bar : BarMap)
+		{
+			Bar.OnFrame = FFrameRate::TransformTime(FFrameTime::FromDecimal(Bar.OnFrame), FrameResolution, NewFrameResolution).AsDecimal();
+		}
+	}
+
+	FrameResolution = NewFrameResolution;
+}
+
+void UFrameBasedMusicMap::Init(float InBpm, const FFrameBasedTimeSignature& InTimeSignature)
+{
+	check(FrameResolution.IsValid());
+	Clear();
+	TempoMap.Add(FFrameBasedTempoPoint(InBpm, 0, 0.0f));
+	BarMap.Add(FFrameBasedTimeSignaturePoint(InTimeSignature, 0, 0, 0.0f));
+}
+
+void UFrameBasedMusicMap::SetInitialTempo(float Bpm)
+{
+	check(!TempoMap.IsEmpty());
+	TempoMap[0].SetBpm(Bpm);
+}
+
+float UFrameBasedMusicMap::GetInitialTempo() const
+{
+	return TempoMap.IsEmpty() ? 0.0f : TempoMap[0].Bpm();
+}
+
+void UFrameBasedMusicMap::SetInitialTimeSignature(const FFrameBasedTimeSignature& TimeSignature)
+{
+	check(!BarMap.IsEmpty());
+	BarMap[0].TimeSignature = TimeSignature;
+}
+
+FFrameBasedTimeSignature UFrameBasedMusicMap::GetInitialTimeSignature() const
+{
+	return BarMap.IsEmpty() ? FFrameBasedTimeSignature() : BarMap[0].TimeSignature;
+}
+
+double UFrameBasedMusicMap::GetFrame(const FMusicalTime& MusicalTime) const
+{
+	return FrameResolution.AsFrameTime(GetSeconds(MusicalTime)).AsDecimal();
+}
+
+double UFrameBasedMusicMap::GetFrame(int32 AtBar) const
+{
+	return FrameResolution.AsFrameTime(GetSeconds(AtBar)).AsDecimal();
+}
+
+FFrameTime UFrameBasedMusicMap::GetFrameTime(const FMusicalTime& MusicalTime) const
+{
+	return FrameResolution.AsFrameTime(GetSeconds(MusicalTime));
+}
+
+FFrameTime UFrameBasedMusicMap::GetFrameTime(int32 AtBar, float AtBeat) const
+{
+	return FrameResolution.AsFrameTime(GetSeconds(AtBar, AtBeat));
+}
+
+double UFrameBasedMusicMap::GetSeconds(const FMusicalTime& MusicalTime) const
+{
+	if (IsEmpty())
+	{
+		return 0.0;
+	}
+
+	check(MusicalTime.IsValid());
+
+	const FFrameBasedTimeSignaturePoint* TimeSignatureAtFrame = GetTimeSignaturePointForBar(MusicalTime.Bar);
+
+	int32 Tick = TimeSignatureAtFrame->MusicalTimeToTick(MusicalTime);
+
+	const FFrameBasedTempoPoint* TempoAtFrame = GetTempoPointForTick(Tick);
+
+	return TempoAtFrame->TickToSeconds(Tick);
+}
+
+double UFrameBasedMusicMap::GetSeconds(int32 AtBar, float AtBeat) const
+{
+	if (IsEmpty())
+	{
+		return 0.0;
+	}
+
+	const FFrameBasedTimeSignaturePoint* TimeSignatureAtFrame = GetTimeSignaturePointForBar(AtBar);
+
+	int32 Tick = TimeSignatureAtFrame->BarBeatToTick(AtBar, AtBeat);
+
+	const FFrameBasedTempoPoint* TempoAtFrame = GetTempoPointForTick(Tick);
+
+	return TempoAtFrame->TickToSeconds(Tick);
+}
+
+FString UFrameBasedMusicMap::MakeBarBeatStringForFrame(double DecimalFrame, FrameBasedMusicMap::EBarBeatFormat BarBeatFormat) const
+{
+	if (IsEmpty())
+	{
+		return FString(TEXT("0.0.00"));
+	}
+
+	const FFrameBasedTimeSignaturePoint* TimeSignatureAtFrame = GetTimeSignaturePointForFrame(DecimalFrame);
+	const FFrameBasedTempoPoint* TempoAtFrame = GetTempoPointForFrame(DecimalFrame);
+
+	return MakeBarBeatString(DecimalFrame, TimeSignatureAtFrame, TempoAtFrame, BarBeatFormat);
+}
+
+FString UFrameBasedMusicMap::MakeBarBeatStringForFrameTime(const FFrameTime& Frame, FrameBasedMusicMap::EBarBeatFormat BarBeatFormat) const
+{
+	return MakeBarBeatStringForFrame(Frame.AsDecimal(), BarBeatFormat);
+}
+
+FString UFrameBasedMusicMap::MakeBarBeatStringForMusicalTime(const FMusicalTime& MusicalTime, FrameBasedMusicMap::EBarBeatFormat BarBeatFormat) const
+{
+	FFrameTime FrameTime = GetFrameTime(MusicalTime);
+	return MakeBarBeatStringForFrameTime(FrameTime, BarBeatFormat);
+}
+
+FString UFrameBasedMusicMap::MakeBarBeatStringForSeconds(double Seconds, FrameBasedMusicMap::EBarBeatFormat BarBeatFormat) const
+{
+	return MakeBarBeatStringForFrameTime(FrameResolution.AsFrameTime(Seconds), BarBeatFormat);
+}
+
+float UFrameBasedMusicMap::GetTempoForFrame(double DecimalFrame) const
+{
+	return GetTempoForFrameTime(FFrameTime::FromDecimal(DecimalFrame));
+}
+
+float UFrameBasedMusicMap::GetTempoForFrameTime(const FFrameTime& FrameTime) const
+{
+	return GetTempoForSeconds(FrameResolution.AsSeconds(FrameTime));
+}
+
+float UFrameBasedMusicMap::GetTempoForSeconds(double Seconds) const
+{
+	if (TempoMap.IsEmpty())
+	{
+		return 0.0f;
+	}
+
+	if (Seconds <= 0.0)
+	{
+		return TempoMap[0].Bpm();
+	}
+
+	const FFrameBasedTempoPoint* TempoPoint = GetTempoPointForSeconds(Seconds);
+
+	return TempoPoint->Bpm();
+}
+
+FFrameBasedTimeSignature UFrameBasedMusicMap::GetTimeSignatureForFrame(double DecimalFrame) const
+{
+	const FFrameBasedTimeSignaturePoint* Point = GetTimeSignaturePointForFrame(DecimalFrame);
+	return Point ? Point->TimeSignature : FFrameBasedTimeSignature();
+}
+
+FFrameBasedTimeSignature UFrameBasedMusicMap::GetTimeSignatureForFrameTime(const FFrameTime& FrameTime) const
+{
+	const FFrameBasedTimeSignaturePoint* Point = GetTimeSignaturePointForFrameTime(FrameTime);
+	return Point ? Point->TimeSignature : FFrameBasedTimeSignature();
+}
+
+FFrameBasedTimeSignature UFrameBasedMusicMap::GetTimeSignatureForSeconds(double Seconds) const
+{
+	const FFrameBasedTimeSignaturePoint* Point = GetTimeSignaturePointForSeconds(Seconds);
+	return Point ? Point->TimeSignature : FFrameBasedTimeSignature();
+}
+
+const FFrameBasedTimeSignaturePoint* UFrameBasedMusicMap::GetTimeSignaturePointForTick(int32 Tick) const
+{
+	if (BarMap.IsEmpty())
+	{
+		return nullptr;
+	}
+
+	if (Tick <= 0)
+	{
+		return &BarMap[0];
+	}
+
+	const FFrameBasedTimeSignaturePoint* TimeSignatureAtTick = Algo::FindLastByPredicate(BarMap, [=](FFrameBasedTimeSignaturePoint const& Point) { return Tick >= Point.OnTick; });
+	return TimeSignatureAtTick ? TimeSignatureAtTick : &BarMap[0];
+}
+
+const FFrameBasedTimeSignaturePoint* UFrameBasedMusicMap::GetTimeSignaturePointForBar(int32 Bar) const
+{
+	if (BarMap.IsEmpty())
+	{
+		return nullptr;
+	}
+
+	if (Bar <= 0)
+	{
+		return &BarMap[0];
+	}
+
+	const FFrameBasedTimeSignaturePoint* TimeSignatureAtBar = Algo::FindLastByPredicate(BarMap, [=](FFrameBasedTimeSignaturePoint const& Point) { return Bar >= Point.OnBar; });
+	return TimeSignatureAtBar ? TimeSignatureAtBar : &BarMap[0];
+}
+
+const FFrameBasedTimeSignaturePoint* UFrameBasedMusicMap::GetTimeSignaturePointForFrame(double DecimalFrame) const
+{
+	if (BarMap.IsEmpty())
+	{
+		return nullptr;
+	}
+
+	if (DecimalFrame <= 0.0)
+	{
+		return &BarMap[0];
+	}
+
+	const FFrameBasedTimeSignaturePoint* TimeSignatureAtBar = Algo::FindLastByPredicate(BarMap, [=](FFrameBasedTimeSignaturePoint const& Point) { return DecimalFrame >= Point.OnFrame; });
+	return TimeSignatureAtBar ? TimeSignatureAtBar : &BarMap[0];
+}
+
+const FFrameBasedTimeSignaturePoint* UFrameBasedMusicMap::GetTimeSignaturePointForFrameTime(const FFrameTime& Frame) const
+{
+	return GetTimeSignaturePointForFrame(Frame.AsDecimal());
+}
+
+const FFrameBasedTimeSignaturePoint* UFrameBasedMusicMap::GetTimeSignaturePointForSeconds(double Seconds) const
+{
+	if (BarMap.IsEmpty() || TempoMap.IsEmpty())
+	{
+		return nullptr;
+	}
+
+	if (Seconds <= 0.0)
+	{
+		return &BarMap[0];
+	}
+
+	const FFrameBasedTempoPoint* TempoAtTime = GetTempoPointForSeconds(Seconds);
+	double Tick = TempoAtTime->SecondsToTick(Seconds);
+	return GetTimeSignaturePointForTick(Tick);
+}
+
+const FFrameBasedTempoPoint* UFrameBasedMusicMap::GetTempoPointForTick(int32 Tick) const
+{
+	if (TempoMap.IsEmpty())
+	{
+		return nullptr;
+	}
+
+	if (Tick <= 0)
+	{
+		return &TempoMap[0];
+	}
+
+	const FFrameBasedTempoPoint* TempoAtTick = Algo::FindLastByPredicate(TempoMap, [=](FFrameBasedTempoPoint const& Point) { return Tick >= Point.OnTick; });
+	return TempoAtTick ? TempoAtTick : &TempoMap[0];
+}
+
+const FFrameBasedTempoPoint* UFrameBasedMusicMap::GetTempoPointForFrame(double DecimalFrame) const
+{
+	return GetTempoPointForFrameTime(FFrameTime::FromDecimal(DecimalFrame));
+}
+
+const FFrameBasedTempoPoint* UFrameBasedMusicMap::GetTempoPointForFrameTime(const FFrameTime& Frame) const
+{
+	return GetTempoPointForSeconds(FrameResolution.AsSeconds(Frame));
+}
+
+const FFrameBasedTempoPoint* UFrameBasedMusicMap::GetTempoPointForSeconds(double Seconds) const
+{
+	if (TempoMap.IsEmpty())
+	{
+		return nullptr;
+	}
+
+	if (Seconds <= 0.0)
+	{
+		return &TempoMap[0];
+	}
+
+	float Ms = static_cast<float>(Seconds * 1000.0);
+	const FFrameBasedTempoPoint* TempoAtTime = Algo::FindLastByPredicate(TempoMap, [=](FFrameBasedTempoPoint const& Point) { return Ms >= Point.OnMs; });
+	return TempoAtTime ? TempoAtTime : &TempoMap[0];
+}
+
+FFrameTime UFrameBasedMusicMap::RemapFrameTime(FFrameTime SourceFrameTime, const UFrameBasedMusicMap* OldMap) const
+{
+	if (IsEmpty())
+	{
+		return SourceFrameTime;
+	}
+
+	FMusicalTime MusicalTime = OldMap->GetMusicalTime(SourceFrameTime);
+	const FFrameBasedTimeSignaturePoint* TimeSignaturePoint = GetTimeSignaturePointForBar(MusicalTime.Bar);
+	if (TimeSignaturePoint->TicksPerBar() != MusicalTime.TicksPerBar)
+	{
+		// Time Signature mismatch! Scale bar to bar!
+		MusicalTime.TickInBar = static_cast<int32>( static_cast<float>(MusicalTime.TickInBar * TimeSignaturePoint->TicksPerBar()) / static_cast<float>(MusicalTime.TicksPerBar));
+	}
+	MusicalTime.TicksPerBar = TimeSignaturePoint->TicksPerBar();
+	MusicalTime.TicksPerBeat = TimeSignaturePoint->TicksPerBeat();
+	return GetFrameTime(MusicalTime);
+}
+
+FMusicalTime UFrameBasedMusicMap::BarsToMusicalTime(double BarsFromStart) const
+{
+	FMusicalTime Result;
+	Result.Bar = FMath::FloorToInt32(BarsFromStart);
+	double FractionalBar = BarsFromStart - static_cast<double>(Result.Bar);
+
+	const FFrameBasedTimeSignaturePoint* TimeSignaturePoint = GetTimeSignaturePointForBar(Result.Bar);
+	Result.TicksPerBar = TimeSignaturePoint->TicksPerBar();
+	Result.TicksPerBeat = TimeSignaturePoint->TicksPerBeat();
+	Result.TickInBar = FMath::Clamp(FMath::RoundToInt32( static_cast<double>(Result.TicksPerBar) * FractionalBar), 0, Result.TicksPerBar - 1);
+	return Result;
+}
+
+FMusicalTime UFrameBasedMusicMap::Quantize(const FMusicalTime& MusicalTime, int32 QuantizationInterval, EQuantizeDirection Direction) const
+{
+	int32 TicksInQuantizationInterval = QuantizationInterval;
+	switch (QuantizationInterval)
+	{
+	case FrameBasedMusicMap::EWholeQuantizationTicks::Bar:
+		TicksInQuantizationInterval = MusicalTime.TicksPerBar;
+		break;
+	case FrameBasedMusicMap::EWholeQuantizationTicks::Beat:
+		TicksInQuantizationInterval = MusicalTime.TicksPerBeat;
+		break;
+	}
+
+	const FFrameBasedTimeSignaturePoint* TimeSignaturePoint = GetTimeSignaturePointForBar(MusicalTime.Bar);
+
+	int32 TotalTicks = TimeSignaturePoint->MusicalTimeToTick(MusicalTime);
+
+	int32 TicksInTimeSignature = TotalTicks - TimeSignaturePoint->OnTick;
+
+	int32 Intervals = 0;
+	switch (Direction)
+	{
+	default:
+	case EQuantizeDirection::Nearest:
+		Intervals = FMath::RoundToInt32(static_cast<float>(TicksInTimeSignature) / static_cast<float>(TicksInQuantizationInterval));
+		break;
+	case EQuantizeDirection::Ceiling:
+		Intervals = FMath::CeilToInt32(static_cast<float>(TicksInTimeSignature) / static_cast<float>(TicksInQuantizationInterval));
+		break;
+	case EQuantizeDirection::Floor:
+		Intervals = FMath::FloorToInt32(static_cast<float>(TicksInTimeSignature) / static_cast<float>(TicksInQuantizationInterval));
+		break;
+	}
+
+	int32 QuantizedTicksInTimeSignature = Intervals * TicksInQuantizationInterval;
+
+	int32 FinalTick = TimeSignaturePoint->OnTick + QuantizedTicksInTimeSignature;
+
+	TimeSignaturePoint = GetTimeSignaturePointForTick(FinalTick);
+
+	return TimeSignaturePoint->TickToMusicalTime(FinalTick);
+}
+
+FMusicalTime UFrameBasedMusicMap::Add(const FMusicalTime& MusicalTime, int32 QuantizationInterval) const
+{
+	const FFrameBasedTimeSignaturePoint* OriginalTimeSignature = GetTimeSignaturePointForBar(MusicalTime.Bar);
+	int32 AbsoluteTick = OriginalTimeSignature->MusicalTimeToTick(MusicalTime);
+	if (FMath::Abs(QuantizationInterval) == FrameBasedMusicMap::EWholeQuantizationTicks::Bar)
+	{
+		AbsoluteTick += (QuantizationInterval < 0) ? MusicalTime.TicksPerBar : -MusicalTime.TicksPerBar;
+	}
+	else if (FMath::Abs(QuantizationInterval) == FrameBasedMusicMap::EWholeQuantizationTicks::Beat)
+	{
+		AbsoluteTick += (QuantizationInterval < 0) ? MusicalTime.TicksPerBeat : -MusicalTime.TicksPerBeat;
+	}
+	else
+	{
+		AbsoluteTick += QuantizationInterval;
+	}
+	const FFrameBasedTimeSignaturePoint* NewTimeSignature = GetTimeSignaturePointForTick(AbsoluteTick);
+	return NewTimeSignature->TickToMusicalTime(AbsoluteTick);
+}
+
+void UFrameBasedMusicMap::AddTimeSignature(int32 AtTick, int32 AtBar, int16 Numerator, int16 Denominator)
+{
+	if (ensure(!BarMap.IsEmpty() || AtTick == 0))
+	{
+		if (ensure(AtTick == 0 || BarMap.Last().OnTick < AtTick))
+		{
+			BarMap.Add(FFrameBasedTimeSignaturePoint({Numerator, Denominator}, AtBar, AtTick, GetFrame(AtBar)));
+		}
+	}
+}
+
+void UFrameBasedMusicMap::AddTempo(int32 AtTick, float AtMs, float Bpm)
+{
+	if (ensure(TempoMap.IsEmpty() || AtTick >= TempoMap.Last().OnTick))
+	{
+		float ExpectedMs = TempoMap.IsEmpty() ? 0.0f : static_cast<float>(TempoMap.Last().TickToSeconds(static_cast<double>(AtTick)) * 1000.0);
+		TempoMap.Add(FFrameBasedTempoPoint(Bpm, AtTick, ExpectedMs));
+	}
+}
+
+TArray<FSubdivisionPoint> UFrameBasedMusicMap::GatherSubdivisonsForTimespan(const FFrameTime& Start, const FFrameTime& End, const FFrameTime& MinSpacing, int32 QuantizationInterval) const
+{
+	TArray<FSubdivisionPoint> Result;
+
+	FMusicalTime StartMusicalTime = Quantize(GetMusicalTime(Start), QuantizationInterval, EQuantizeDirection::Ceiling);
+	FMusicalTime EndMusicalTime = Quantize(GetMusicalTime(End), QuantizationInterval, EQuantizeDirection::Floor);
+
+	if (StartMusicalTime >= EndMusicalTime)
+	{
+		return Result;
+	}
+
+
+	Result.Add(MakeSubdivisionPoint(StartMusicalTime));
+	FFrameTime LastEntryFrameTime = GetFrameTime(StartMusicalTime);
+	int32 LastBarNumber = StartMusicalTime.Bar;
+	StartMusicalTime = Add(StartMusicalTime, QuantizationInterval);
+	if (StartMusicalTime.Bar != LastBarNumber)
+	{
+		StartMusicalTime.TickInBar = 0;
+		LastBarNumber = StartMusicalTime.Bar;
+	}
+	while (StartMusicalTime <= EndMusicalTime)
+	{
+		FFrameTime ThisEntryFrameTime = GetFrameTime(StartMusicalTime);
+		if ((ThisEntryFrameTime - LastEntryFrameTime) < MinSpacing)
+		{
+			if (StartMusicalTime.TickInBar == 0)
+			{
+				// This is a downbeat. We are too close to the last, BUT we 
+				// take precedence. So blow away the last...
+				Result.Pop(EAllowShrinking::No);
+			}
+			else
+			{
+				StartMusicalTime = Add(StartMusicalTime, QuantizationInterval);
+				if (StartMusicalTime.Bar != LastBarNumber)
+				{
+					StartMusicalTime.TickInBar = 0;
+					LastBarNumber = StartMusicalTime.Bar;
+				}
+				continue;
+			}
+		}
+		Result.Add(MakeSubdivisionPoint(StartMusicalTime));
+		LastEntryFrameTime = GetFrameTime(StartMusicalTime);
+		StartMusicalTime = Add(StartMusicalTime, QuantizationInterval);
+		if (StartMusicalTime.Bar != LastBarNumber)
+		{
+			StartMusicalTime.TickInBar = 0;
+			LastBarNumber = StartMusicalTime.Bar;
+		}
+	}
+	
+	return Result;
+}
+
+FMusicalTime UFrameBasedMusicMap::GetMusicalTime(const FFrameTime& FrameTime) const
+{
+	if (IsEmpty())
+	{
+		return FMusicalTime();
+	}
+
+	const FFrameBasedTimeSignaturePoint* TimeSignatureAtFrame = GetTimeSignaturePointForFrameTime(FrameTime);
+	const FFrameBasedTempoPoint* TempoAtFrame = GetTempoPointForFrameTime(FrameTime);
+
+	double Seconds = FrameResolution.AsSeconds(FrameTime);
+	float Tick = TempoAtFrame->SecondsToTick(Seconds);
+	return TimeSignatureAtFrame->TickToMusicalTime(Tick);
+}
+
+FSubdivisionPoint UFrameBasedMusicMap::MakeSubdivisionPoint(const FMusicalTime& MusicalTime) const
+{
+	FSubdivisionPoint Result;
+	Result.MusicalTime = MusicalTime;
+	Result.FrameTime = GetFrameTime(MusicalTime);
+	Result.PrettyString = MakeBarBeatStringForMusicalTime(MusicalTime, FrameBasedMusicMap::EBarBeatFormat::Position);
+	return Result;
+}
+
+FString UFrameBasedMusicMap::MakeBarBeatString(double DecimalFrame, const FFrameBasedTimeSignaturePoint* TimeSignature, const FFrameBasedTempoPoint* Tempo, FrameBasedMusicMap::EBarBeatFormat BarBeatFormat) const
+{
+	double Seconds = FrameResolution.AsSeconds(FFrameTime::FromDecimal(DecimalFrame));
+	float Tick = Tempo->SecondsToTick(Seconds);
+	FMusicalTime MusicalTime = TimeSignature->TickToMusicalTime(Tick);
+	// for display to users... bars and beats are 1 based!
+	MusicalTime.Bar += (BarBeatFormat == FrameBasedMusicMap::EBarBeatFormat::Position ? 1 : 0);
+	double BeatAsDouble = FMath::RoundToDouble((MusicalTime.BeatInBar() + (BarBeatFormat == FrameBasedMusicMap::EBarBeatFormat::Position ? 1.0f : 0.0f)) * 100.0f);
+	if (BeatAsDouble >= ((TimeSignature->TimeSignature.Numerator + (BarBeatFormat == FrameBasedMusicMap::EBarBeatFormat::Position ? 1 : 0)) * 100))
+	{
+		MusicalTime.Bar++;
+		BeatAsDouble = BarBeatFormat == FrameBasedMusicMap::EBarBeatFormat::Position ? 100.0 : 0.0;
+	}
+	return FString::Printf(TEXT("%d.%.2f"), MusicalTime.Bar, BeatAsDouble / 100.0);
+}
+#undef LOCTEXT_NAMESPACE
diff --git a/Engine/Source/Runtime/TimeManagement/Private/FrameNumberTimeEvaluator.cpp b/Engine/Source/Runtime/TimeManagement/Private/FrameNumberTimeEvaluator.cpp
--- a/Engine/Source/Runtime/TimeManagement/Private/FrameNumberTimeEvaluator.cpp
+++ b/Engine/Source/Runtime/TimeManagement/Private/FrameNumberTimeEvaluator.cpp
@@ -405,4 +405,64 @@ TValueOrError<FFrameTime, FExpressionError> FFrameNumberTimeEvaluator::EvaluateT
 	OutDirectlyParsed = false;
 	return MakeError(NSLOCTEXT("TimeManagement", "UnrecognizedTimeResult", "Unrecognized result returned from expression"));
 
-}
\ No newline at end of file
+}
+
+TValueOrError<FFrameTime, FExpressionError> FFrameNumberTimeEvaluator::EvaluateBarBeat(const TCHAR* InExpression, const UFrameBasedMusicMap* MusicMap, FrameBasedMusicMap::EBarBeatFormat BarBeatFormat, bool& OutDirectlyParsed) const
+{
+	OutDirectlyParsed = false;
+
+	if (!MusicMap || MusicMap->IsEmpty())
+	{
+		return MakeError(NSLOCTEXT("TimeManagement", "BarBeatError_NoTempoMap", "No supplied TempoMap"));
+	}
+
+	FString InAsString(InExpression);
+
+	// find bar.beat boundary and replace with space...
+	int32 PeriodCount = 0;
+	for (TCHAR& Character : InAsString)
+	{
+		if (Character == '.')
+		{
+			PeriodCount++;
+			if (PeriodCount == 1)
+			{
+				Character = ' ';
+			}
+		}
+	}
+
+	if (PeriodCount != 2)
+	{
+		return MakeError(NSLOCTEXT("TimeManagement", "BarBeatError_BadFormat", "Format not recognized as a Bar.Beat"));
+	}
+
+	TValueOrError<TArray<FExpressionToken>, FExpressionError> LexResult = ExpressionParser::Lex(*InAsString, TimeTokenDefinitions);
+	if (!LexResult.IsValid())
+	{
+		return MakeError(LexResult.StealError());
+	}
+	TArray<FExpressionToken> Tokens = LexResult.StealValue();
+	if (Tokens.Num() != 2)
+	{
+		return MakeError(NSLOCTEXT("TimeManagement", "BarBeatError_BadFormat", "Format not recognized as a Bar.Beat"));
+	}
+
+	if (!Tokens[0].Node.Cast<double>() || !Tokens[1].Node.Cast<double>())
+	{
+		return MakeError(NSLOCTEXT("TimeManagement", "BarBeatError_BadFormat", "Format not recognized as a Bar.Beat"));
+	}
+
+	int32 Bar = static_cast<int32>(*Tokens[0].Node.Cast<double>());
+	float BeatInBar = static_cast<float>(*Tokens[1].Node.Cast<double>());
+	if (BeatInBar < 1.0f && BarBeatFormat == FrameBasedMusicMap::EBarBeatFormat::Position)
+	{
+		return MakeError(NSLOCTEXT("TimeManagement", "BarBeatError_InvalidBeatValue", "Beats are one based, so the specified beat must be >= 1.0"));
+	}
+	
+	// in user land bars and beats are 1 based, so make them 0 based.
+	OutDirectlyParsed = true;
+	int32 BarCorrection = BarBeatFormat == FrameBasedMusicMap::EBarBeatFormat::Position ? 1 : 0;
+	float BeatCorrection = BarBeatFormat == FrameBasedMusicMap::EBarBeatFormat::Position ? 1.0f : 0.0f;
+	return MakeValue(MusicMap->GetFrameTime(Bar - BarCorrection, BeatInBar - BeatCorrection));
+}
diff --git a/Engine/Source/Runtime/TimeManagement/Public/FrameBasedMusicMap.h b/Engine/Source/Runtime/TimeManagement/Public/FrameBasedMusicMap.h
new file mode 100644
--- /dev/null
+++ b/Engine/Source/Runtime/TimeManagement/Public/FrameBasedMusicMap.h
@@ -0,0 +1,406 @@
+// Copyright Epic Games, Inc. All Rights Reserved.
+
+#pragma once
+
+#include "CoreTypes.h"
+#include "Misc/FrameRate.h"
+#include "Misc/MusicalTime.h"
+#include "UObject/Interface.h"
+
+#include "FrameBasedMusicMap.generated.h"
+
+struct FFrameTime;
+
+namespace FrameBasedMusicMap
+{
+	enum ETripletQuantizationTicks : int32
+	{
+		ThirtysecondNoteTriplet = (MusicalTime::TicksPerQuarterNote / 4) / 3,	// 80
+		SixteenthNoteTriplet    = (MusicalTime::TicksPerQuarterNote / 2) / 3,	// 160
+		EighthNoteTriplet       = (MusicalTime::TicksPerQuarterNote) / 3,		// 320
+		QuarterNoteTriplet      = (MusicalTime::TicksPerQuarterNote * 2) / 3,	// 640
+		HalfNoteTriplet         = (MusicalTime::TicksPerQuarterNote * 4) / 3,	// 1280
+		WholeNoteTriplet        = (MusicalTime::TicksPerQuarterNote * 8) / 3,	// 2560
+	};
+
+	enum EWholeQuantizationTicks : int32
+	{
+		ThirtysecondNote        = MusicalTime::TicksPerQuarterNote / 8,			// 120
+		SixteenthNote           = MusicalTime::TicksPerQuarterNote / 4,			// 240
+		EighthNote              = MusicalTime::TicksPerQuarterNote / 2,			// 480
+		QuarterNote	            = MusicalTime::TicksPerQuarterNote,				// 960
+		HalfNote                = MusicalTime::TicksPerQuarterNote * 2,			// 1920
+		WholeNote               = MusicalTime::TicksPerQuarterNote * 4,			// 3840
+		Bar                     = std::numeric_limits<int32>::min()+1,          // Ticks depends on Time Signature
+		Beat                    = std::numeric_limits<int32>::min()+2           // Ticks depends on Time Signature
+	};
+
+	enum EDottedQuantizationTicks : int32
+	{
+		DottedThirtysecondNote  = (MusicalTime::TicksPerQuarterNote / 16) * 3,	// 180
+		DottedSixteenthNote     = (MusicalTime::TicksPerQuarterNote / 8) * 3,	// 360
+		DottedEighthNote        = (MusicalTime::TicksPerQuarterNote / 4) * 3,	// 720
+		DottedQuarterNote       = (MusicalTime::TicksPerQuarterNote / 2) * 3,	// 1440
+		DottedHalfNote          = (MusicalTime::TicksPerQuarterNote) * 3,		// 2880
+		DottedWholeNote         = (MusicalTime::TicksPerQuarterNote) * 6,		// 5760
+	};
+
+	TIMEMANAGEMENT_API FText GetQuantizationName(int32 Quantization);
+	extern TIMEMANAGEMENT_API TArray<TSharedPtr<int32>> QuantizationOptionsSource;
+
+	enum class EBarBeatFormat
+	{
+		Position,
+		Duration,
+		NotAllowed // Don't display bars.beats! Fall back to seconds.
+	};
+}
+
+struct FSubdivisionPoint
+{
+	FFrameTime FrameTime;
+	FMusicalTime MusicalTime;
+	FString PrettyString;
+};
+
+USTRUCT()
+struct TIMEMANAGEMENT_API FFrameBasedTempoPoint
+{
+	GENERATED_BODY()
+
+	/**
+	 * NOTE: This structure treats BPM/Tempo exactly like Midi does to avoid any
+	 * floating point precision problems and differences between Midi code and 
+	 * this code. It is for that reason that Tempo is stored here a "microseconds
+	 * per quarter note" (MPQ) and not a float Bpm. 
+	 */
+
+	FFrameBasedTempoPoint()
+	{}
+
+	FFrameBasedTempoPoint(float InBpm, int32 InTick, float InAtMillisecond)
+		: OnTick(InTick)
+		, OnMs(InAtMillisecond)
+	{
+		SetBpm(InBpm);
+	}
+
+	void SetBpm(float InBpm)
+	{
+		MicrosecondsPerQuarterNote = InBpm >= 0.0f ? static_cast<int32>(60000000.0f / InBpm) : 0;
+	}
+
+	UPROPERTY()
+	int32 MicrosecondsPerQuarterNote = 500000; // 120 bpm default.
+	UPROPERTY()
+	int32 OnTick = 0;
+	UPROPERTY()
+	float OnMs = 0.0;
+
+	float Bpm() const
+	{
+		return MicrosecondsPerQuarterNote == 0 ? 0 : 60000000.0f / static_cast<float>(MicrosecondsPerQuarterNote);
+	}
+
+	int32 MsToTick(float Ms) const
+	{
+		if (Ms == 0.0f)
+		{
+			return 0;
+		}
+
+		return FMath::RoundToInt32(static_cast<double>(OnTick) +
+								static_cast<double>(Ms - OnMs) * 1000.0 / 
+								static_cast<double>(MicrosecondsPerQuarterNote) *
+								static_cast<double>(MusicalTime::TicksPerQuarterNote));
+	}
+
+	int32 SecondsToTick(double Second) const
+	{
+		return MsToTick(static_cast<float>(Second * 1000.0));
+	}
+
+	double TickToSeconds(int32 Tick) const
+	{
+		return static_cast<double>(TickToMs(Tick))/1000.0;
+	}
+
+	float TickToMs(int32 Tick) const
+	{
+		if (Tick == 0)
+		{
+			return 0.0f;
+		}
+		float MsPerTick = static_cast<float>(MicrosecondsPerQuarterNote) / (static_cast<float>(MusicalTime::TicksPerQuarterNote) * 1000.0f);
+		float DeltaTick = static_cast<float>(Tick - OnTick);
+		return OnMs + (DeltaTick * MsPerTick);
+	}
+};
+
+USTRUCT()
+struct TIMEMANAGEMENT_API FFrameBasedTimeSignature
+{
+	GENERATED_BODY()
+
+	FFrameBasedTimeSignature()
+	{}
+	FFrameBasedTimeSignature(int32 InNumerator, int32 InDenominator)
+		: Numerator(static_cast<int16>(InNumerator))
+		, Denominator(static_cast<int16>(InDenominator))
+	{}
+
+	UPROPERTY()
+	int16 Numerator = 4;
+	UPROPERTY()
+	int16 Denominator = 4;
+
+	bool operator==(const FFrameBasedTimeSignature& Other)
+	{
+		return Numerator == Other.Numerator && Denominator == Other.Denominator;
+	}
+
+	bool operator!=(const FFrameBasedTimeSignature& Other)
+	{
+		return Numerator != Other.Numerator || Denominator != Other.Denominator;
+	}
+};
+
+USTRUCT()
+struct TIMEMANAGEMENT_API FFrameBasedTimeSignaturePoint
+{
+	GENERATED_BODY()
+
+	FFrameBasedTimeSignaturePoint()
+	{}
+
+	FFrameBasedTimeSignaturePoint(const FFrameBasedTimeSignature& InTimeSignature, int32 InOnBar, int32 InOnTick, double InOnFrame)
+		: TimeSignature(InTimeSignature)
+		, OnBar(InOnBar)
+		, OnTick(InOnTick)
+		, OnFrame(InOnFrame)
+	{
+	}
+
+	UPROPERTY()
+	FFrameBasedTimeSignature TimeSignature;
+	UPROPERTY()
+	int32 OnBar = 0;
+	UPROPERTY()
+	int32 OnTick = 0;
+	UPROPERTY()
+	double OnFrame = 0.0;
+
+	int32 TicksPerBeat() const
+	{
+		return MusicalTime::TicksPerQuarterNote * 4 / TimeSignature.Denominator;
+	}
+
+	int32 TicksPerBar() const
+	{
+		return TicksPerBeat() * TimeSignature.Numerator;
+	}
+
+	FMusicalTime TickToMusicalTime(int32 InTick) const
+	{
+		FMusicalTime Result(0, 0, TicksPerBar(), TicksPerBeat());
+		int32 TicksPast = InTick - OnTick;
+		int32 BarsPassed = TicksPast / TicksPerBar();
+		if (TicksPast >= 0)
+		{
+			Result.Bar = OnBar + BarsPassed;
+			Result.TickInBar = TicksPast - (BarsPassed * TicksPerBar());
+		}
+		else
+		{
+			Result.Bar = OnBar + (BarsPassed - 1);
+			Result.TickInBar = TicksPast - ((-1 + BarsPassed) * TicksPerBar());
+		}
+		return Result;
+	}
+
+	int32 MusicalTimeToTick(const FMusicalTime& MusicalTime) const
+	{
+		int32 BarDelta = MusicalTime.Bar - OnBar;
+		return OnTick + (BarDelta * MusicalTime.TicksPerBar) + MusicalTime.TickInBar;
+	}
+
+	int32 BarToTick(int32 Bar) const
+	{
+		int32 BarDelta = Bar - OnBar;
+		return OnTick + (BarDelta * TicksPerBar());
+	}
+
+	int32 BarBeatToTick(int32 Bar, float Beat) const
+	{
+		int32 BarDelta = Bar - OnBar;
+		float TicksIntoBar = Beat * static_cast<float>(TicksPerBeat());
+		return OnTick + (BarDelta * TicksPerBar()) + FMath::RoundToInt32(TicksIntoBar);
+	}
+};
+
+UCLASS()
+class TIMEMANAGEMENT_API UFrameBasedMusicMap : public UObject
+{
+	GENERATED_BODY()
+
+public:
+	UFrameBasedMusicMap()
+		: FrameResolution(0,0)
+	{}
+
+	UFrameBasedMusicMap(const FFrameRate& InFrameResolution)
+		: FrameResolution(InFrameResolution)
+	{
+		
+	}
+
+	UFrameBasedMusicMap(const FFrameRate& InFrameResolution, float InBpm, const FFrameBasedTimeSignature& TimeSignature)
+		: FrameResolution(InFrameResolution)
+	{
+		TempoMap.Add(FFrameBasedTempoPoint(InBpm, 0, 0.0f));
+		BarMap.Add(FFrameBasedTimeSignaturePoint(TimeSignature, 0, 0, 0.0));
+	}
+
+	void SetFrameResolution(const FFrameRate& NewFrameResolution);
+
+	void Init(float InBpm, const FFrameBasedTimeSignature& InTimeSignature);
+
+	void SetInitialTempo(float Bpm);
+	float GetInitialTempo() const;
+	void SetInitialTimeSignature(const FFrameBasedTimeSignature& TimeSignature);
+	FFrameBasedTimeSignature GetInitialTimeSignature() const;
+
+	void Clear()
+	{
+		TempoMap.Empty();
+		BarMap.Empty();
+	}
+
+	bool IsEmpty() const { return TempoMap.IsEmpty() || BarMap.IsEmpty(); }
+	bool IsSimple() const { return TempoMap.Num() == 1 && BarMap.Num() == 1; }
+
+	double GetFrame(const FMusicalTime& MusicalTime) const;
+	double GetFrame(int32 AtBar) const;
+	FFrameTime GetFrameTime(const FMusicalTime& MusicalTime) const;
+	FFrameTime GetFrameTime(int32 AtBar, float AtBeat) const;
+	double GetSeconds(const FMusicalTime& MusicalTime) const;
+	double GetSeconds(int32 AtBar, float AtBeat = 0.0f) const;
+
+	FString MakeBarBeatStringForFrame(double DecimalFrame, FrameBasedMusicMap::EBarBeatFormat BarBeatFormat) const;
+	FString MakeBarBeatStringForFrameTime(const FFrameTime& Frame, FrameBasedMusicMap::EBarBeatFormat BarBeatFormat) const;
+	FString MakeBarBeatStringForMusicalTime(const FMusicalTime& MusicalTime, FrameBasedMusicMap::EBarBeatFormat BarBeatFormat) const;
+	FString MakeBarBeatStringForSeconds(double Seconds, FrameBasedMusicMap::EBarBeatFormat BarBeatFormat) const;
+
+	FMusicalTime GetMusicalTime(const FFrameTime& FrameTime) const;
+
+	float GetTempoForFrame(double DecimalFrame) const;
+	float GetTempoForFrameTime(const FFrameTime& FrameTime) const;
+	float GetTempoForSeconds(double Seconds) const;
+
+	FFrameBasedTimeSignature GetTimeSignatureForFrame(double DecimalFrame) const;
+	FFrameBasedTimeSignature GetTimeSignatureForFrameTime(const FFrameTime& FrameTime) const;
+	FFrameBasedTimeSignature GetTimeSignatureForSeconds(double Seconds) const;
+
+	FFrameTime RemapFrameTime(FFrameTime SourceTime, const UFrameBasedMusicMap* OldMap) const;
+
+	FMusicalTime BarsToMusicalTime(double BarsFromStart) const;
+
+	enum class EQuantizeDirection
+	{
+		Nearest,
+		Ceiling,
+		Floor
+	};
+	FMusicalTime Quantize(const FMusicalTime& MusicalTime, int32 QuantizationInterval, EQuantizeDirection Direction = EQuantizeDirection::Nearest) const;
+
+	/**
+	* FMusicalTime Add(const FMusicalTime& MusicalTime, int32 QuantizationInterval) const;
+	* 
+	* Increment the supplied MusicalTime by the specified QuantizationInterval. 
+	* NOTE: This does NOT quantize anything. QuantizationInterval is just used 
+	* as a convenient way to specify a musical subdivision. You can use any of the 
+	* TripletQuantizationTicks, WholeQuantizationTicks, or DottedQuantizationTicks
+	* enum values. AND... you can negate them to go back in time.
+	* 
+	* @param MusicalTime			Time to be added to
+	
+	* @param QuantizationInterval	One of the values from...
+	*                               TripletQuantizationTicks
+	*								WholeQuantizationTicks
+	*								DottedQuantizationTicks
+	* 
+	*                                 
+	*/
+	FMusicalTime Add(const FMusicalTime& MusicalTime, int32 QuantizationInterval) const;
+
+	UPROPERTY()
+	TArray<FFrameBasedTempoPoint> TempoMap;
+	UPROPERTY()
+	TArray<FFrameBasedTimeSignaturePoint> BarMap;
+	UPROPERTY()
+	FFrameRate FrameResolution;
+
+	void AddTimeSignature(int32 AtTick, int32 AtBar, int16 Numerator, int16 Denominator);
+	void AddTempo(int32 AtTick, float AtMs, float Bpm);
+
+	TArray<FSubdivisionPoint> GatherSubdivisonsForTimespan(const FFrameTime& Start, const FFrameTime& End, const FFrameTime& MinSpacing, int32 QuantizationInterval) const;
+
+private:
+	const FFrameBasedTimeSignaturePoint* GetTimeSignaturePointForTick(int32 Tick) const;
+	const FFrameBasedTimeSignaturePoint* GetTimeSignaturePointForBar(int32 Bar) const;
+	const FFrameBasedTimeSignaturePoint* GetTimeSignaturePointForFrame(double DecimalFrame) const;
+	const FFrameBasedTimeSignaturePoint* GetTimeSignaturePointForFrameTime(const FFrameTime& Frame) const;
+	const FFrameBasedTimeSignaturePoint* GetTimeSignaturePointForSeconds(double Seconds) const;
+
+	const FFrameBasedTempoPoint* GetTempoPointForTick(int32 Tick) const;
+	const FFrameBasedTempoPoint* GetTempoPointForFrame(double DecimalFrame) const;
+	const FFrameBasedTempoPoint* GetTempoPointForFrameTime(const FFrameTime& Frame) const;
+	const FFrameBasedTempoPoint* GetTempoPointForSeconds(double Seconds) const;
+
+	FSubdivisionPoint MakeSubdivisionPoint(const FMusicalTime& MusicalTime) const;
+
+	FString MakeBarBeatString(double DecimalFrame, const FFrameBasedTimeSignaturePoint* TimeSignature, const FFrameBasedTempoPoint* Tempo, FrameBasedMusicMap::EBarBeatFormat BarBeatFormat) const;
+};
+
+UINTERFACE(MinimalAPI)
+class UMusicMapSource : public UInterface
+{
+	GENERATED_BODY()
+};
+
+struct FMarkerProviderEntry
+{
+	FString Label;
+	FFrameNumber FrameNumber;
+};
+
+USTRUCT()
+struct FMarkerProviderResults
+{
+	GENERATED_BODY()
+
+	struct FChannel
+	{
+		FString Name;
+		TArray<FMarkerProviderEntry> Markers;
+	};
+
+	bool IsEmpty() const { return Channels.IsEmpty(); }
+
+	TArray<FChannel> Channels;
+};
+
+class TIMEMANAGEMENT_API IMusicMapSource
+{
+public:
+
+	GENERATED_BODY()
+
+	UFUNCTION()
+	virtual void CreateFrameBasedMusicMap(UFrameBasedMusicMap* Map) const { return; }
+
+	UFUNCTION()
+	virtual FMarkerProviderResults GatherMarkers(const UFrameBasedMusicMap* Map) const { return FMarkerProviderResults(); }
+
+};
+
diff --git a/Engine/Source/Runtime/TimeManagement/Public/FrameNumberDisplayFormat.h b/Engine/Source/Runtime/TimeManagement/Public/FrameNumberDisplayFormat.h
--- a/Engine/Source/Runtime/TimeManagement/Public/FrameNumberDisplayFormat.h
+++ b/Engine/Source/Runtime/TimeManagement/Public/FrameNumberDisplayFormat.h
@@ -17,6 +17,8 @@ enum class EFrameNumberDisplayFormats : uint8
 	Seconds,
 	
 	Frames,
-	
+
+	BarsBeats UMETA(DisplayName = "Bars : Beats"),
+
 	MAX_Count UMETA(Hidden)
 };
\ No newline at end of file
diff --git a/Engine/Source/Runtime/TimeManagement/Public/FrameNumberNumericInterface.h b/Engine/Source/Runtime/TimeManagement/Public/FrameNumberNumericInterface.h
--- a/Engine/Source/Runtime/TimeManagement/Public/FrameNumberNumericInterface.h
+++ b/Engine/Source/Runtime/TimeManagement/Public/FrameNumberNumericInterface.h
@@ -6,6 +6,7 @@
 #include "FrameNumberDisplayFormat.h"
 #include "Widgets/Input/NumericTypeInterface.h"
 #include "Misc/FrameRate.h"
+#include "FrameBasedMusicMap.h"
 #include "FrameNumberTimeEvaluator.h"
 #include "Misc/Timecode.h"
 #include "Misc/ExpressionParserTypes.h"
@@ -29,12 +30,14 @@
 */
 struct FFrameNumberInterface : public INumericTypeInterface<double>
 {
-	FFrameNumberInterface(const TAttribute<EFrameNumberDisplayFormats>& InDisplayFormatAttr, const TAttribute<uint8>& InOnGetZeroPadFrameNumber, const TAttribute<FFrameRate>& InTickResolutionAttr, const TAttribute<FFrameRate>& InDisplayRateAttr)
+	FFrameNumberInterface(const TAttribute<EFrameNumberDisplayFormats>& InDisplayFormatAttr, const TAttribute<uint8>& InOnGetZeroPadFrameNumber, const TAttribute<FFrameRate>& InTickResolutionAttr, const TAttribute<FFrameRate>& InDisplayRateAttr, const TAttribute<const UFrameBasedMusicMap*>& InMusicMap = TAttribute<const UFrameBasedMusicMap*>(), FrameBasedMusicMap::EBarBeatFormat InBarBeatFormat = FrameBasedMusicMap::EBarBeatFormat::Position)
 		: DisplayFormatAttr(InDisplayFormatAttr)
 		, TickResolutionAttr(InTickResolutionAttr)
 		, DisplayRateAttr(InDisplayRateAttr)
 		, ZeroPadFramesAttr(InOnGetZeroPadFrameNumber)
+		, MusicMapAttr(InMusicMap)
 		, SubframeIndicator(TEXT("*"))
+		, BarBeatFormat(InBarBeatFormat)
 	{
 		check(InDisplayFormatAttr.IsSet());
 		check(InTickResolutionAttr.IsSet());
@@ -84,6 +87,13 @@ struct FFrameNumberInterface : public INumericTypeInterface<double>
 			Format = EFrameNumberDisplayFormats::NonDropFrameTimecode;
 		}
 
+		// Don't show bars:beats under various conditions...
+		if (Format == EFrameNumberDisplayFormats::BarsBeats &&
+			(!MusicMapAttr.IsSet() || !MusicMapAttr.Get() || MusicMapAttr.Get()->IsEmpty() || BarBeatFormat == FrameBasedMusicMap::EBarBeatFormat::NotAllowed))
+		{
+			Format = EFrameNumberDisplayFormats::Seconds;
+		}
+
 		switch (Format)
 		{
 		case EFrameNumberDisplayFormats::Frames:
@@ -116,6 +126,11 @@ struct FFrameNumberInterface : public INumericTypeInterface<double>
 			FTimecode AsNonDropTimecode = FTimecode::FromFrameNumber(DisplayTime.FloorToFrame(), DestinationFrameRate, bIsDropTimecode);
 			return FString::Printf(TEXT("[%s%s]"), *AsNonDropTimecode.ToString(false), *SubframeSuffix);
 		}
+		case EFrameNumberDisplayFormats::BarsBeats:
+		{
+			check(MusicMapAttr.IsSet() && MusicMapAttr.Get());
+			return MusicMapAttr.Get()->MakeBarBeatStringForFrame(Value, BarBeatFormat);
+		}
 		default:
 			return FString(TEXT("Unsupported Format"));
 		}
@@ -126,6 +141,14 @@ struct FFrameNumberInterface : public INumericTypeInterface<double>
 		FFrameRate SourceFrameRate = DisplayRateAttr.Get();
 		FFrameRate DestinationFrameRate = TickResolutionAttr.Get();
 		EFrameNumberDisplayFormats FallbackFormat = DisplayFormatAttr.Get();
+		const UFrameBasedMusicMap* MusicMap = MusicMapAttr.Get();
+
+		// If we are currently displaying bars.beats but this specific property field does not allow bars.beats 
+		// display/entry, force the fallback here to be seconds...
+		if (FallbackFormat == EFrameNumberDisplayFormats::BarsBeats && BarBeatFormat == FrameBasedMusicMap::EBarBeatFormat::NotAllowed)
+		{
+			FallbackFormat = EFrameNumberDisplayFormats::Seconds;
+		}
 
 		// We allow input in any format (time, frames or timecode) and we just convert it into into the internal sequence resolution.
 		// The user's input can be ambiguous though (does "5" mean 5 frames or 5 seconds?) so when we check each possible result we
@@ -139,12 +162,14 @@ struct FFrameNumberInterface : public INumericTypeInterface<double>
 		TValueOrError<FFrameTime, FExpressionError> FrameResult = Eval.EvaluateFrame(*InString, SourceFrameRate, DestinationFrameRate, /*Out*/ bWasFrameText);
 		bool bWasTimeText;
 		TValueOrError<FFrameTime, FExpressionError> TimeResult = Eval.EvaluateTime(*InString, DestinationFrameRate, /*Out*/ bWasTimeText);
+		bool bWasBarBeatText;
+		TValueOrError<FFrameTime, FExpressionError> BarBeatResult = Eval.EvaluateBarBeat(*InString, MusicMap, BarBeatFormat, /*Out*/ bWasBarBeatText);
 
 
 		// All three formats support ambiguous conversion where the user can enter "5" and wants it in the logical unit based on the current
 		// display format. This means 5 -> 5f, 5 ->5s, and 5 -> 5 frames (in timecode). We also support specifically specifying in a different
 		// format than your current display, ie if your display is in frames and you enter 1s, you get 1s in frames or 30 (at 30fps play rate).
-		if (!bWasTimecodeText && !bWasFrameText && !bWasTimeText)
+		if (!bWasTimecodeText && !bWasFrameText && !bWasTimeText && !bWasBarBeatText)
 		{
 			// They've entered an ambiguous number, so we'll check the display format and see if we did successfully parse as that type. If it
 			// was able to parse the input for the given display format, we return that.
@@ -206,6 +231,10 @@ struct FFrameNumberInterface : public INumericTypeInterface<double>
 		{
 			return TOptional<double>(FrameResult.GetValue().GetFrame().Value);
 		}
+		else if (bWasBarBeatText)
+		{
+			return TOptional<double>(BarBeatResult.GetValue().GetFrame().Value);
+		}
 
 		// We're not sure what they typed in
 		return TOptional<double>();
@@ -238,7 +267,10 @@ private:
 	TAttribute<FFrameRate> TickResolutionAttr;
 	TAttribute<FFrameRate> DisplayRateAttr;
 	TAttribute<uint8> ZeroPadFramesAttr;
+	TAttribute<const UFrameBasedMusicMap*> MusicMapAttr;
 
 	FString SubframeIndicator;
+	FrameBasedMusicMap::EBarBeatFormat BarBeatFormat;
+
 	FOnSettingChanged OnSettingChangedEvent;
 };
diff --git a/Engine/Source/Runtime/TimeManagement/Public/FrameNumberTimeEvaluator.h b/Engine/Source/Runtime/TimeManagement/Public/FrameNumberTimeEvaluator.h
--- a/Engine/Source/Runtime/TimeManagement/Public/FrameNumberTimeEvaluator.h
+++ b/Engine/Source/Runtime/TimeManagement/Public/FrameNumberTimeEvaluator.h
@@ -3,6 +3,7 @@
 #pragma once
 
 #include "CoreTypes.h"
+#include "FrameBasedMusicMap.h"
 #include "Math/BasicMathExpressionEvaluator.h"
 #include "Misc/ExpressionParserTypes.h"
 #include "Misc/Optional.h"
@@ -30,6 +31,7 @@ struct FFrameNumberTimeEvaluator
 	TIMEMANAGEMENT_API TValueOrError<FFrameTime, FExpressionError> EvaluateTimecode(const TCHAR* InExpression, const FFrameRate& InDisplayFrameRate, const FFrameRate& InTickResolution, bool& OutDirectlyParsed) const;
 	TIMEMANAGEMENT_API TValueOrError<FFrameTime, FExpressionError> EvaluateFrame(const TCHAR* InExpression, const FFrameRate& InDisplayFrameRate, const FFrameRate& InTickResolution, bool& OutDirectlyParsed) const;
 	TIMEMANAGEMENT_API TValueOrError<FFrameTime, FExpressionError> EvaluateTime(const TCHAR* InExpression, FFrameRate InFrameRate, bool& OutDirectlyParsed) const;
+	TIMEMANAGEMENT_API TValueOrError<FFrameTime, FExpressionError> EvaluateBarBeat(const TCHAR* InExpression, const UFrameBasedMusicMap* MusicMap, FrameBasedMusicMap::EBarBeatFormat BarBeatFormat, bool& OutDirectlyParsed) const;
 
 	FTokenDefinitions TimecodeTokenDefinitions;
 	FTokenDefinitions FrameTokenDefinitions;
