Fix LOCTEXT issue.

diff --git a/Engine/Source/Editor/Sequencer/Private/SSequencerMusicModeCombo.cpp b/Engine/Source/Editor/Sequencer/Private/SSequencerMusicModeCombo.cpp
--- a/Engine/Source/Editor/Sequencer/Private/SSequencerMusicModeCombo.cpp
+++ b/Engine/Source/Editor/Sequencer/Private/SSequencerMusicModeCombo.cpp
@@ -194,7 +194,7 @@ void SSequencerMusicModeCombo::PopulateSimpleMusicMapSection(FMenuBuilder& MenuB
 				],
 				LOCTEXT("MusicModeSettings_TimeSignature", "Time Signature"),
 				/*InNoIndent*/ true, /*bInSearchable*/ true,
-				LOCTEXT("MusicModeSettings_Tempo_ToolTip", "The \"natural time signature\" this sequence is authored for.")
+				LOCTEXT("MusicModeSettings_TimeSignature_ToolTip", "The \"natural time signature\" this sequence is authored for.")
 			);
 
 
@@ -260,7 +260,7 @@ FText SSequencerMusicModeCombo::GetTempoTimeSignatureText() const
 		{
 			float Tempo = MusicMap->GetInitialTempo();
 			FFrameBasedTimeSignature TimeSignature = MusicMap->GetInitialTimeSignature();
-			return FText::Format(LOCTEXT("TempoTimeSigDisplay", "{0} bpm {1}/{2}"),
+			return FText::Format(LOCTEXT("TempoTimeSigDisplay_Simple", "{0} bpm {1}/{2}"),
 				FText::AsNumber(Tempo, &Formatting),
 				FText::AsNumber(TimeSignature.Numerator, &Formatting),
 				FText::AsNumber(TimeSignature.Denominator, &Formatting));
@@ -270,7 +270,7 @@ FText SSequencerMusicModeCombo::GetTempoTimeSignatureText() const
 			FQualifiedFrameTime CursorTime = Sequencer->GetLocalTime();
 			float Tempo = MusicMap->GetTempoForFrameTime(CursorTime.Time);
 			FFrameBasedTimeSignature TimeSignature = MusicMap->GetTimeSignatureForFrameTime(CursorTime.Time);
-			return FText::Format(LOCTEXT("TempoTimeSigDisplay", "Current: {0} bpm {1}/{2}"),
+			return FText::Format(LOCTEXT("TempoTimeSigDisplay_Complex", "Current: {0} bpm {1}/{2}"),
 				FText::AsNumber(Tempo, &Formatting),
 				FText::AsNumber(TimeSignature.Numerator, &Formatting),
 				FText::AsNumber(TimeSignature.Denominator, &Formatting));
@@ -303,7 +303,7 @@ void SSequencerMusicModeCombo::OnRemoveAllMusicMaps()
 
 	if (MusicMap)
 	{
-		FScopedTransaction ScopedTransaction(FText(LOCTEXT("SetMusicMode", "Turning music mode off for sequence")));
+		FScopedTransaction ScopedTransaction(FText(LOCTEXT("SetMusicMode_Off", "Turning music mode off for sequence")));
 		FocusedMovieScene->Modify();
 		FocusedMovieScene->RemoveMusicMap();
 	}
@@ -322,7 +322,7 @@ void SSequencerMusicModeCombo::OnCreateSimpleMusicMap()
 			// WARN THE USER HERE! Make sure they want to blow away the existing music map!
 		}
 
-		FScopedTransaction ScopedTransaction(FText::Format(LOCTEXT("SetMusicMode", "Turning music mode {0} for sequence"), FocusedMovieScene->HasMusicMap() ? FText::FromString(TEXT("On")) : FText::FromString(TEXT("Off"))));
+		FScopedTransaction ScopedTransaction(FText(LOCTEXT("SetMusicMode_On", "Turning music mode on for sequence")));
 		FocusedMovieScene->Modify();
 
 		UFrameBasedMusicMap* NewSimpleMusicMap = NewObject<UFrameBasedMusicMap>(FocusedMovieScene);
