Sequencer: CreateKeysUnderMouse should select them so that they can be moved in the same drag operation. This fixes a bug where middle mouse copy/paste key will end up on a subframe

diff --git a/Engine/Source/Editor/Sequencer/Private/MVVM/Views/SChannelView.cpp b/Engine/Source/Editor/Sequencer/Private/MVVM/Views/SChannelView.cpp
--- a/Engine/Source/Editor/Sequencer/Private/MVVM/Views/SChannelView.cpp
+++ b/Engine/Source/Editor/Sequencer/Private/MVVM/Views/SChannelView.cpp
@@ -754,10 +754,20 @@ void SChannelView::CreateKeysUnderMouse(const FVector2D& MousePosition, const FG
 		}
 
 		FScopedTransaction Transaction(LOCTEXT("CreateKeysUnderMouse", "Create keys under mouse"));
-		FAddKeyOperation::FromKeyAreas(TrackModel->GetTrackEditor().Get(), ValidKeyAreasUnderCursor).Commit(CurrentTime.FrameNumber, *LegacyGetSequencer());
+		TArray<FAddKeyResult> Results;
+		FAddKeyOperation::FromKeyAreas(TrackModel->GetTrackEditor().Get(), ValidKeyAreasUnderCursor).Commit(CurrentTime.FrameNumber, *LegacyGetSequencer(), &Results);
 
-		// Get the keys under the mouse as the newly created keys.
-		GetKeysAtPixelX(LocalSpaceMousePosition.X, OutKeys);
+		for (FAddKeyResult& Result : Results)
+		{
+			for (TViewModelPtr<FChannelModel> Channel : Model->GetDescendantsOfType<FChannelModel>(true))
+			{
+				if (Channel->GetChannel() == Result.KeyArea->GetChannel().Get())
+				{
+					OutKeys.Add(FSequencerSelectedKey(*SectionObject, Channel, Result.KeyHandle));
+					break;
+				}
+			}
+		}
 	}
 }
 
