Sequencer: Synchronize sequencer selection with external selection if a handle was created when adding a key

- Tested that an object that isn't in Sequencer gets selected in Sequencer when adding a key through shortcut s
- Tested that an object that isn't in Sequencer gets selected in Sequencer when adding a key through the details panel

diff --git a/Engine/Source/Editor/Sequencer/Private/MovieSceneTrackEditor.cpp b/Engine/Source/Editor/Sequencer/Private/MovieSceneTrackEditor.cpp
--- a/Engine/Source/Editor/Sequencer/Private/MovieSceneTrackEditor.cpp
+++ b/Engine/Source/Editor/Sequencer/Private/MovieSceneTrackEditor.cpp
@@ -157,6 +157,12 @@ void FMovieSceneTrackEditor::AnimatablePropertyChanged( FOnKeyProperty OnKeyProp
 			}
 		}
 
+		// If a handle was created, something new has been added to Sequencer, so its selection state must be synchronized
+		if (KeyPropertyResult.bHandleCreated)
+		{
+			Sequencer.Pin()->ExternalSelectionHasChanged();
+		}
+
 		UpdatePlaybackRange();
 	}
 }
