PR #12554: Sequencer: Add ability to set the sequencer sections grip height

diff --git a/Engine/Source/Editor/Sequencer/Private/SSequencerSection.cpp b/Engine/Source/Editor/Sequencer/Private/SSequencerSection.cpp
--- a/Engine/Source/Editor/Sequencer/Private/SSequencerSection.cpp
+++ b/Engine/Source/Editor/Sequencer/Private/SSequencerSection.cpp
@@ -1095,7 +1095,7 @@ bool SSequencerSection::CheckForEasingHandleInteraction( const FPointerEvent& Mo
 		return false;
 	}
 
-	FGeometry SectionGeometry = MakeSectionGeometryWithoutHandles(AllottedGeometry);
+	FGeometry SectionGeometry = MakeHandlesGeometry(AllottedGeometry);
 
 	const double MousePositionX = SectionGeometry.AbsoluteToLocal(MouseEvent.GetScreenSpacePosition()).X;
 
@@ -1196,7 +1196,7 @@ bool SSequencerSection::CheckForEdgeInteraction( const FPointerEvent& MouseEvent
 		}
 	}
 
-	FGeometry SectionGeometry = MakeSectionGeometryWithoutHandles(AllottedGeometry);
+	FGeometry SectionGeometry = MakeHandlesGeometry(AllottedGeometry);
 
 	for (const TSharedPtr<FSectionModel>& UnderlappingSection : AllUnderlappingSections)
 	{
@@ -1415,7 +1415,7 @@ int32 SSequencerSection::OnPaint( const FPaintArgs& Args, const FGeometry& Allot
 
 	const ESlateDrawEffect DrawEffects = bEnabled ? ESlateDrawEffect::None : ESlateDrawEffect::DisabledEffect;
 
-	FGeometry SectionGeometry = MakeSectionGeometryWithoutHandles(AllottedGeometry);
+	FGeometry SectionGeometry = MakeHandlesGeometry(AllottedGeometry);
 	*TimeToPixel = ConstructTimeConverterForSection(SectionGeometry, *SectionObject, GetSequencer()); 
 
 	FSequencerSectionPainterImpl Painter(GetSequencer(), TrackAreaViewModel, SectionModel, OutDrawElements, SectionGeometry, *this);
@@ -1826,7 +1826,7 @@ void SSequencerSection::DrawSectionHandles( const FGeometry& AllottedGeometry, F
 		}
 	}
 
-	FGeometry SectionGeometry = MakeSectionGeometryWithoutHandles(AllottedGeometry);
+	FGeometry SectionGeometry = MakeHandlesGeometry(AllottedGeometry);
 
 	for (TSharedPtr<FSectionModel> SectionModel : AllUnderlappingSections)
 	{
@@ -1862,7 +1862,7 @@ void SSequencerSection::DrawSectionHandles( const FGeometry& AllottedGeometry, F
 		}
 
 		const float ThisHandleOffset = UnderlappingSectionObj == ThisSection ? HandleOffsetPx : 0.f;
-		FVector2D GripSize( UnderlappingSection->GetSectionGripSize(), AllottedGeometry.Size.Y );
+		FVector2D GripSize( UnderlappingSection->GetSectionGripSize(), UnderlappingSection->GetSectionGripHeight(AllottedGeometry.Size.Y) );
 
 		float Opacity = 0.85f;
 		if (ThisHandleOffset != 0)
@@ -1992,6 +1992,20 @@ FGeometry SSequencerSection::MakeSectionGeometryWithoutHandles(const FGeometry&
 	);
 }
 
+FGeometry SSequencerSection::MakeHandlesGeometry(const FGeometry& AllottedGeometry) const
+{
+	const FVector2f LocalSize = AllottedGeometry.GetLocalSize();
+
+	FVector2f SectionSizeWithoutHandles = LocalSize - FVector2f(HandleOffsetPx * 2.f, 0.f);
+	SectionSizeWithoutHandles.X = FMath::Max(1.f, SectionSizeWithoutHandles.X);
+	SectionSizeWithoutHandles.Y = FMath::Max(1.f, SectionInterface->GetSectionGripHeight(SectionSizeWithoutHandles.Y));
+
+	return AllottedGeometry.MakeChild(
+		SectionSizeWithoutHandles,
+		FSlateLayoutTransform(FVector2D(HandleOffsetPx, 0))
+	);
+}
+
 void SSequencerSection::UpdateUnderlappingSegments()
 {
 	TSharedPtr<FSectionModel> SectionModel = WeakSectionModel.Pin();
diff --git a/Engine/Source/Editor/Sequencer/Private/SSequencerSection.h b/Engine/Source/Editor/Sequencer/Private/SSequencerSection.h
--- a/Engine/Source/Editor/Sequencer/Private/SSequencerSection.h
+++ b/Engine/Source/Editor/Sequencer/Private/SSequencerSection.h
@@ -157,6 +157,11 @@ private:
 	 */
 	FGeometry MakeSectionGeometryWithoutHandles(const FGeometry& AllottedGeometry) const;
 
+	/** 
+	 * Creates geometry for the section handles
+	 */
+	FGeometry MakeHandlesGeometry(const FGeometry& AllottedGeometry) const;
+
 public:
 
 	/** Indicate that the current section selection should throb the specified number of times. A single throb takes 0.2s. */
diff --git a/Engine/Source/Editor/Sequencer/Public/ISequencerSection.h b/Engine/Source/Editor/Sequencer/Public/ISequencerSection.h
--- a/Engine/Source/Editor/Sequencer/Public/ISequencerSection.h
+++ b/Engine/Source/Editor/Sequencer/Public/ISequencerSection.h
@@ -233,6 +233,11 @@ public:
 	 */
 	virtual float GetSectionGripSize() const { return SequencerSectionConstants::DefaultSectionGripSize; }
 
+	/**
+	 * @return The height of the section drag handles
+	 */
+	virtual float GetSectionGripHeight(float InSectionHeight) const { return InSectionHeight; }
+
 	/**
 	 * @ return The size of keyframe widgets
 	 */
