Sequencer: Update cached pinned state before filtering

This fixes a bug where unpinning a track that doesn't pass filters doesn't disappear.

Checked that pinning a binding still properly updates/displays the sections underneath it

diff --git a/Engine/Source/Editor/Sequencer/Private/SequencerNodeTree.cpp b/Engine/Source/Editor/Sequencer/Private/SequencerNodeTree.cpp
--- a/Engine/Source/Editor/Sequencer/Private/SequencerNodeTree.cpp
+++ b/Engine/Source/Editor/Sequencer/Private/SequencerNodeTree.cpp
@@ -108,6 +108,9 @@ void FSequencerNodeTree::Update()
 	UMovieScene* MovieScene = CurrentSequence->GetMovieScene();
 	CleanupMuteSolo(MovieScene);
 
+	// Cache pinned state of nodes, this needs to happen before UpdateFilters() below as some filters will look at the pinned state of child nodes
+	FPinnableExtensionShim::UpdateCachedPinnedState(RootNode);
+
 	// Re-filter the tree after updating 
 	// @todo sequencer: Newly added sections may need to be visible even when there is a filter
 	bFilterUpdateRequested = true;
@@ -129,9 +132,6 @@ void FSequencerNodeTree::Update()
 		IGeometryExtension::UpdateVirtualGeometry(0.f, RootNode);
 	}
 
-	// Cache pinned state of nodes, needs to happen after OnTreeRefreshed
-	FPinnableExtensionShim::UpdateCachedPinnedState(RootNode);
-
 	// Update curve editor tree based on new filtered hierarchy
 	auto CurveEditorIntegration = SequenceModel->CastDynamic<FCurveEditorIntegrationExtension>();
 	if (CurveEditorIntegration)
