Sequencer: Fixed some issues when pasting overlapping sections onto tracks that don't support blending and don't support multiple rows.

If a section overlaps and it's now allowed to, output an error and don't paste the section.

When pasting a section onto a binding, we try to paste onto all compatible tracks. The row fix up code wasn't looking at compatibility when considering the rows for the track.

diff --git a/Engine/Source/Editor/Sequencer/Private/SequencerUtilities.cpp b/Engine/Source/Editor/Sequencer/Private/SequencerUtilities.cpp
--- a/Engine/Source/Editor/Sequencer/Private/SequencerUtilities.cpp
+++ b/Engine/Source/Editor/Sequencer/Private/SequencerUtilities.cpp
@@ -1664,7 +1664,7 @@ void ExportObjectsToText(const TArray<UObject*>& ObjectsToExport, FString& Expor
 		UObject* ThisOuter = ObjectToExport->GetOuter();
 		if (LastOuter != nullptr && ThisOuter != LastOuter)
 		{
-			UE_LOG(LogSequencer, Warning, TEXT("Cannot copy objects from different outers. Only copying from %s"), *LastOuter->GetName());
+			UE_LOG(LogSequencer, Error, TEXT("Cannot copy objects from different outers. Only copying from %s"), *LastOuter->GetName());
 			continue;
 		}
 		LastOuter = ThisOuter;
@@ -2160,6 +2160,8 @@ bool FSequencerUtilities::PasteSections(const FString& TextToImport, FMovieScene
 	{
 		UMovieSceneTrack* Track = PasteSectionsParams.Tracks[Index];
 
+		const bool bAllowOverlap = Track->SupportsMultipleRows() || Track->GetSupportedBlendTypes().Num() > 0;
+
 		for (int32 SectionIndex = 0; SectionIndex < ImportedSections.Num(); ++SectionIndex)
 		{
 			UMovieSceneSection* Section = ImportedSections[SectionIndex];
@@ -2170,8 +2172,6 @@ bool FSequencerUtilities::PasteSections(const FString& TextToImport, FMovieScene
 
 			int32 RowIndex = SectionIndex < PasteSectionsParams.TrackRowIndices.Num() ? PasteSectionsParams.TrackRowIndices[SectionIndex] : Section->GetRowIndex();
 
-			SectionIndicesImported.AddUnique(SectionIndex);
-
 			Track->Modify();
 
 			Section->ClearFlags(RF_Transient);
@@ -2195,6 +2195,17 @@ bool FSequencerUtilities::PasteSections(const FString& TextToImport, FMovieScene
 				Section->MoveSection(NewStartFrame - Section->GetInclusiveStartFrame());
 			}
 
+			if (!bAllowOverlap)
+			{
+				if (Section->OverlapsWithSections(Track->GetAllSections()))
+				{
+					Track->RemoveSection(*Section);
+					UE_LOG(LogSequencer, Error, TEXT("Could not paste section because it overlaps with existing sections and this track type does not allow overlaps"));
+					continue;
+				}
+			}
+
+			SectionIndicesImported.AddUnique(SectionIndex);
 			OutSections.Add(Section);
 		}
 
@@ -2206,6 +2217,11 @@ bool FSequencerUtilities::PasteSections(const FString& TextToImport, FMovieScene
 			int32 MaxAvailableRowIndex = -1;
 			for (UMovieSceneSection* Section : OutSections)
 			{
+				if (!Track->SupportsType(Section->GetClass()))
+				{
+					continue;
+				}
+
 				if (MovieSceneToolHelpers::OverlapsSection(Track, Section, OutSections))
 				{
 					int32 AvailableRowIndex = MovieSceneToolHelpers::FindAvailableRowIndex(Track, Section, OutSections);
@@ -2246,7 +2262,7 @@ bool FSequencerUtilities::PasteSections(const FString& TextToImport, FMovieScene
 	{
 		if (!SectionIndicesImported.Contains(SectionIndex))
 		{
-			UE_LOG(LogSequencer, Display, TEXT("Could not paste section of type %s"), *ImportedSections[SectionIndex]->GetClass()->GetName());
+			UE_LOG(LogSequencer, Error, TEXT("Could not paste section of type %s"), *ImportedSections[SectionIndex]->GetClass()->GetName());
 		}
 	}
 
