- Movied previously committed squencer changes for music mode into a new Musical Mode plugin

diff --git a/Engine/Plugins/Runtime/GameplayAbilities/Source/GameplayAbilities/Private/Sequencer/MovieSceneGameplayCueSections.cpp b/Engine/Plugins/Runtime/GameplayAbilities/Source/GameplayAbilities/Private/Sequencer/MovieSceneGameplayCueSections.cpp
--- a/Engine/Plugins/Runtime/GameplayAbilities/Source/GameplayAbilities/Private/Sequencer/MovieSceneGameplayCueSections.cpp
+++ b/Engine/Plugins/Runtime/GameplayAbilities/Source/GameplayAbilities/Private/Sequencer/MovieSceneGameplayCueSections.cpp
@@ -206,14 +206,9 @@ void FMovieSceneGameplayCueChannel::DeleteKeysFrom(FFrameNumber InTime, bool bDe
 	GetData().DeleteKeysFrom(InTime, bDeleteKeysBefore);
 }
 
-void FMovieSceneGameplayCueChannel::ChangeFrameResolution(FFrameRate SourceRate, FFrameRate DestinationRate)
+void FMovieSceneGameplayCueChannel::RemapTimes(const UE::MovieScene::IRetimingInterface& Retimer)
 {
-	GetData().ChangeFrameResolution(SourceRate, DestinationRate);
-}
-
-void FMovieSceneGameplayCueChannel::RemapMusicalFrameTimes(const UFrameBasedMusicMap* OldMap, const UFrameBasedMusicMap* NewMap)
-{
-	GetData().RemapMusicalFrameTimes(OldMap, NewMap);
+	GetData().RemapTimes(Retimer);
 }
 
 TRange<FFrameNumber> FMovieSceneGameplayCueChannel::ComputeEffectiveRange() const
diff --git a/Engine/Plugins/Runtime/Harmonix/Source/HarmonixMetasound/HarmonixMetasound.Build.cs b/Engine/Plugins/Runtime/Harmonix/Source/HarmonixMetasound/HarmonixMetasound.Build.cs
--- a/Engine/Plugins/Runtime/Harmonix/Source/HarmonixMetasound/HarmonixMetasound.Build.cs
+++ b/Engine/Plugins/Runtime/Harmonix/Source/HarmonixMetasound/HarmonixMetasound.Build.cs
@@ -20,6 +20,8 @@ public class HarmonixMetasound : ModuleRules
 				"Engine",
 				"Core",
 				"CoreUObject",
+				"MusicEnvironment",
+				"TimeManagement",
 				"MovieScene",
 				"MetasoundStandardNodes",
 				"HarmonixDsp",
diff --git a/Engine/Plugins/Runtime/Harmonix/Source/HarmonixMetasound/Private/HarmonixMetasound/Components/MetasoundMusicClockDriver.cpp b/Engine/Plugins/Runtime/Harmonix/Source/HarmonixMetasound/Private/HarmonixMetasound/Components/MetasoundMusicClockDriver.cpp
--- a/Engine/Plugins/Runtime/Harmonix/Source/HarmonixMetasound/Private/HarmonixMetasound/Components/MetasoundMusicClockDriver.cpp
+++ b/Engine/Plugins/Runtime/Harmonix/Source/HarmonixMetasound/Private/HarmonixMetasound/Components/MetasoundMusicClockDriver.cpp
@@ -251,7 +251,7 @@ void FMetasoundMusicClockDriver::OnStart()
 	SongPosOffsetMs = 0.0f;
 	RenderStartSampleCount = 0;
 	RenderStartWallClockTimeSeconds = 0.0;
-	FreeRunStartTimeSecs = ClockComponent ? ClockComponent->GetWorld()->GetTimeSeconds() : 0.0;
+	FreeRunStartTimeSecs = GetWallClockTime();
 	bRunning = true;
 }
 
@@ -409,7 +409,7 @@ void FMetasoundMusicClockDriver::OnGeneratorDetached()
 			ClockComponent->DefaultMaps.Copy(*(CurrentMapChain->SongMaps), 0, LastTickSeen);
 		}
 		SongPosOffsetMs = ClockComponent->CurrentSmoothedAudioRenderSongPos.SecondsIncludingCountIn * 1000.0f;
-		FreeRunStartTimeSecs = ClockComponent->GetWorld()->GetTimeSeconds();
+		FreeRunStartTimeSecs = GetWallClockTime();
 	}
 	ClockHistory = nullptr;
 	SmoothedAudioRenderClockHistoryCursor = HarmonixMetasound::Analysis::FMidiClockSongPositionHistory::FReadCursor();
@@ -425,7 +425,7 @@ void FMetasoundMusicClockDriver::RefreshCurrentSongPosFromWallClock()
 
 	bool TempoChanged = ClockComponent->CurrentSmoothedAudioRenderSongPos.Tempo != ClockComponent->Tempo;
 
-	double FreeRunTime = (ClockComponent->GetWorld()->GetTimeSeconds() - FreeRunStartTimeSecs) * ClockComponent->CurrentClockAdvanceRate;
+	double FreeRunTime = (GetWallClockTime() - FreeRunStartTimeSecs) * ClockComponent->CurrentClockAdvanceRate;
 
 	ClockComponent->CurrentRawAudioRenderSongPos.SetByTime(((float)FreeRunTime * 1000.0) + SongPosOffsetMs, ClockComponent->DefaultMaps);
 	ClockComponent->CurrentSmoothedAudioRenderSongPos = ClockComponent->CurrentRawAudioRenderSongPos;
@@ -453,6 +453,19 @@ FString FMetasoundMusicClockDriver::HistoryFailureTypeToString(FMetasoundMusicCl
 	return "Unrecognized";
 }
 
+double FMetasoundMusicClockDriver::GetWallClockTime()
+{
+	if (ClockComponent && ClockComponent->GetWorld())
+	{
+		return ClockComponent->GetWorld()->GetTimeSeconds();
+	}
+	if (GWorld)
+	{
+		return GWorld->GetTimeSeconds();
+	}
+	return 0.0;
+}
+
 void FMetasoundMusicClockDriver::RefreshCurrentSongPosFromHistory()
 {
 	using namespace HarmonixMetasound::Analysis;
@@ -490,13 +503,13 @@ void FMetasoundMusicClockDriver::RefreshCurrentSongPosFromHistory()
 		// We are just starting up. We need to create the initial "sync point"...
 		// Wall clock <-> Render Samples...
 		RenderStartSampleCount = LastRenderPosSampleCount;
-		RenderStartWallClockTimeSeconds = ClockComponent->GetWorld()->GetTimeSeconds() - (double)RenderStartSampleCount / (double)ClockHistory->SampleRate;
+		RenderStartWallClockTimeSeconds = GetWallClockTime() - (double)RenderStartSampleCount / (double)ClockHistory->SampleRate;
 		RenderSmoothingLagSeconds = MetasoundMusicClockDriver::SmoothedAudioRenderLagSeconds;
 		ErrorTracker.Reset();
 		LastRefreshWallClockTimeSeconds = RenderStartWallClockTimeSeconds;
 	}
 
-	double CurrentWallClockSeconds = ClockComponent->GetWorld()->GetTimeSeconds();
+	double CurrentWallClockSeconds = GetWallClockTime();
 	DeltaSecondsBetweenRefreshes = CurrentWallClockSeconds - LastRefreshWallClockTimeSeconds;
 	LastRefreshWallClockTimeSeconds = CurrentWallClockSeconds;
 
@@ -512,7 +525,7 @@ void FMetasoundMusicClockDriver::RefreshCurrentSongPosFromHistory()
 		{
 			UE_LOG(LogMusicClock, Verbose, TEXT("======== MASSIVE ERROR (%f) - SEEKING ==========="), (float)Error);
 			RenderStartSampleCount = LastRenderPosSampleCount;
-			RenderStartWallClockTimeSeconds = ClockComponent->GetWorld()->GetTimeSeconds() - (double)RenderStartSampleCount / (double)ClockHistory->SampleRate;
+			RenderStartWallClockTimeSeconds = GetWallClockTime() - (double)RenderStartSampleCount / (double)ClockHistory->SampleRate;
 			ExpectedRenderedSeconds = RenderedSeconds;
 			RenderSmoothingLagSeconds = MetasoundMusicClockDriver::SmoothedAudioRenderLagSeconds;
 			ErrorTracker.Reset();
diff --git a/Engine/Plugins/Runtime/Harmonix/Source/HarmonixMetasound/Private/HarmonixMetasound/Components/MusicClockComponent.cpp b/Engine/Plugins/Runtime/Harmonix/Source/HarmonixMetasound/Private/HarmonixMetasound/Components/MusicClockComponent.cpp
--- a/Engine/Plugins/Runtime/Harmonix/Source/HarmonixMetasound/Private/HarmonixMetasound/Components/MusicClockComponent.cpp
+++ b/Engine/Plugins/Runtime/Harmonix/Source/HarmonixMetasound/Private/HarmonixMetasound/Components/MusicClockComponent.cpp
@@ -229,19 +229,28 @@ void UMusicClockComponent::BeginPlay()
 		CreateClockDriver();
 	}
 	UMidiClockUpdateSubsystem::TrackMusicClockComponent(this);
-	Super::BeginPlay();
+	if (IsRegistered())
+	{
+		Super::BeginPlay();
+	}
 }
 
 void UMusicClockComponent::EndPlay(const EEndPlayReason::Type EndPlayReason)
 {
-	Super::EndPlay(EndPlayReason);
+	if (IsRegistered())
+	{
+		Super::EndPlay(EndPlayReason);
+	}
 	UMidiClockUpdateSubsystem::StopTrackingMusicClockComponent(this);
 	DisconnectFromClockDriver();
 }
 
 void UMusicClockComponent::TickComponent(float DeltaTime, enum ELevelTick TickType, FActorComponentTickFunction* ThisTickFunction)
 {
-	Super::TickComponent(DeltaTime, TickType, ThisTickFunction);
+	if (IsRegistered())
+	{
+		Super::TickComponent(DeltaTime, TickType, ThisTickFunction);
+	}
 	EnsureClockIsValidForGameFrame();
 	BroadcastSongPosChanges();
 	BroadcastSeekLoopDetections();
@@ -769,6 +778,7 @@ void UMusicClockComponent::MakeDefaultSongMap()
 	DefaultMaps.GetBarMap().AddTimeSignatureAtBarIncludingCountIn(0, TimeSignatureNum, TimeSignatureDenom);
 }
 
+/*
 void UMusicClockComponent::OnTick(float DeltaSeconds, float InPlayRate)
 {
 
@@ -784,12 +794,7 @@ void UMusicClockComponent::OnStopPlaying(const FQualifiedFrameTime& InStopTime)
 
 }
 
-FFrameTime UMusicClockComponent::OnRequestCurrentTime(const FQualifiedFrameTime& InCurrentTime, float InPlayRate)
-{
-	return InCurrentTime.Rate.AsFrameTime(CurrentVideoRenderSongPos.SecondsFromBarOne);
-}
-
-FMusicalTime UMusicClockComponent::OnRequestCurrentMusicalTime(const FQualifiedFrameTime& InCurrentTime, float InPlayRate)
+FMusicalTime UMusicClockComponent::GetMusicalTime() const
 {
 	FMusicalTime Result;
 	Result.Bar = CurrentVideoRenderSongPos.Timestamp.Bar - 1;
@@ -798,6 +803,7 @@ FMusicalTime UMusicClockComponent::OnRequestCurrentMusicalTime(const FQualifiedF
 	Result.TickInBar = static_cast<int32>((CurrentVideoRenderSongPos.Timestamp.Beat - 1.0f) * static_cast<float>(Result.TicksPerBeat));
 	return Result;
 }
+*/
 
 void FMusicClockDriverBase::EnsureClockIsValidForGameFrame()
 {
diff --git a/Engine/Plugins/Runtime/Harmonix/Source/HarmonixMetasound/Private/HarmonixMetasound/MusicEnvironmentSupport/HarmonixMusicEnvironmentMetronome.h b/Engine/Plugins/Runtime/Harmonix/Source/HarmonixMetasound/Private/HarmonixMetasound/MusicEnvironmentSupport/HarmonixMusicEnvironmentMetronome.h
new file mode 100644
--- /dev/null
+++ b/Engine/Plugins/Runtime/Harmonix/Source/HarmonixMetasound/Private/HarmonixMetasound/MusicEnvironmentSupport/HarmonixMusicEnvironmentMetronome.h
@@ -0,0 +1,76 @@
+// Copyright Epic Games, Inc. All Rights Reserved.
+
+#pragma once
+
+#include "CoreMinimal.h"
+#include "MetasoundBuilderBase.h"
+#include "MusicEnvironmentMetronome.h"
+#include "UObject/Object.h"
+
+#include "HarmonixMusicEnvironmentMetronome.generated.h"
+
+class UAudioComponent;
+class UMetaSoundSourceBuilder;
+class UMusicClockComponent;
+class UWorld;
+class AActor;
+
+/**
+ * Implementation of the generic MusicEnvironmentSubsystem's metronome interface. So that subsystem
+ * can create metronomes based on the Harmonix tech.
+ */
+UCLASS()
+class HARMONIXMETASOUND_API UHarmonixMusicEnvironmentMetronome : public UObject, public IMusicEnvironmentMetronome
+{
+	GENERATED_BODY()
+	
+public:
+	virtual bool Initialize(UWorld* InWorld) override;
+	virtual void Tick(float DeltaSecs) override;
+	virtual void Start(double FromSeconds = 0.0) override;
+	virtual void Seek(double ToSeconds) override;
+	virtual void Stop() override;
+	virtual void Pause() override;
+	virtual void Resume() override;
+
+	virtual float GetCurrentTempo() override;
+	virtual double GetCurrentPositionSeconds() override;
+	virtual FMusicalTime GetMusicalTime() const override;
+	
+protected:
+	virtual void OnMusicMapSet() override;
+	virtual bool OnSetTempo(float Bpm) override;
+	
+private:
+	UPROPERTY(Transient)
+	TObjectPtr<AActor> Actor;
+	
+	UPROPERTY(Transient)
+	TObjectPtr<UAudioComponent> AudioComponent;
+	UPROPERTY(Transient)
+	TObjectPtr<UMusicClockComponent> MusicClockComponent;
+
+	UPROPERTY(Transient)
+	TObjectPtr<UMetaSoundSourceBuilder> SourceBuilder;
+	UPROPERTY(Transient)
+	FMetaSoundBuilderNodeOutputHandle OnPlayNodeOutput;
+	UPROPERTY(Transient)
+	FMetaSoundBuilderNodeInputHandle OnFinishedNodeInput;
+	UPROPERTY(Transient)
+	TArray<FMetaSoundBuilderNodeInputHandle> AudioOutNodeInputs;
+
+	bool bMetasoundIsPlaying = false;
+	float CurrentTempo = 120.0f;
+	
+	bool BuildAndStartMetasound(UWorld* InWorld);
+	bool BuildMetasound();
+	FMetaSoundBuilderNodeOutputHandle AddGraphInput(const FName& InputName, const float Value) const;
+	FMetaSoundBuilderNodeOutputHandle AddGraphInputTrigger(const FName& InputName) const;
+	bool SetNodeInputDefault(const FMetaSoundNodeHandle& Node, const FName& InputName, const int32 Value) const;
+	bool SetNodeInputDefault(const FMetaSoundNodeHandle& Node, const FName& InputName, const float Value) const;
+	bool SetNodeInputDefault(const FMetaSoundNodeHandle& Node, const FName& InputName, const TArray<float>& Values) const;
+	bool ConnectGraphInputToNodeInput(const FName& GraphInputName, const FMetaSoundNodeHandle& Node, const FName& NodeInputName) const;
+	bool ConnectNodeOutputToGraphOutput(const FMetaSoundNodeHandle& Node, const FName& NodeOutputName, const FName& GraphOutputName) const;
+	bool ConnectNodes(const FMetaSoundNodeHandle& SourceNode, const FName& OutName, const FMetaSoundNodeHandle& DestinationNode, const FName& InName) const;
+};
+
diff --git a/Engine/Plugins/Runtime/Harmonix/Source/HarmonixMetasound/Private/HarmonixMetasound/Nodes/MetronomeNode.cpp b/Engine/Plugins/Runtime/Harmonix/Source/HarmonixMetasound/Private/HarmonixMetasound/Nodes/MetronomeNode.cpp
--- a/Engine/Plugins/Runtime/Harmonix/Source/HarmonixMetasound/Private/HarmonixMetasound/Nodes/MetronomeNode.cpp
+++ b/Engine/Plugins/Runtime/Harmonix/Source/HarmonixMetasound/Private/HarmonixMetasound/Nodes/MetronomeNode.cpp
@@ -23,6 +23,17 @@ namespace HarmonixMetasound::Nodes::MetronomeNode
 {
 	using namespace Metasound;
 
+	const FNodeClassName& GetClassName()
+	{
+		static const FNodeClassName ClassName { HarmonixNodeNamespace, TEXT("Metronome"), TEXT("") };
+		return ClassName;
+	}
+
+	int32 GetCurrentMajorVersion()
+	{
+		return 0;
+	}
+
 	bool bSkipMetronomeLastProcessedClockTickCheck = true;
 	FAutoConsoleVariableRef CVarSkipMetronomeLastProcessedClockTickCheck(
 		TEXT("au.Metronome.SkipLastProcessedClockTickCheck"),
@@ -64,8 +75,8 @@ namespace HarmonixMetasound::Nodes::MetronomeNode
 		auto InitNodeInfo = []() -> FNodeClassMetadata
 		{
 			FNodeClassMetadata Info;
-			Info.ClassName        = { HarmonixNodeNamespace, TEXT("Metronome"), TEXT("") };
-			Info.MajorVersion     = 0;
+			Info.ClassName        = GetClassName();
+			Info.MajorVersion     = GetCurrentMajorVersion();
 			Info.MinorVersion     = 1;
 			Info.DisplayName      = METASOUND_LOCTEXT("MetronomeNode_DisplayName", "Metronome MIDI Clock Generator");
 			Info.Description      = METASOUND_LOCTEXT("MetronomeNode_Description", "Provides a MIDI clock at the specified tempo and speed.");
@@ -82,10 +93,25 @@ namespace HarmonixMetasound::Nodes::MetronomeNode
 		return Info;
 	}
 
-	const FVertexInterface& FMetronomeOperator::GetVertexInterface()
+	namespace Inputs
 	{
-		using namespace CommonPinNames;
+		DEFINE_METASOUND_PARAM_ALIAS(Transport, CommonPinNames::Inputs::Transport);
+		DEFINE_METASOUND_PARAM_ALIAS(Loop, CommonPinNames::Inputs::Loop);
+		DEFINE_METASOUND_PARAM_ALIAS(LoopLengthBars, CommonPinNames::Inputs::LoopLengthBars);
+		DEFINE_METASOUND_PARAM_ALIAS(TimeSigNumerator, CommonPinNames::Inputs::TimeSigNumerator);
+		DEFINE_METASOUND_PARAM_ALIAS(TimeSigDenominator, CommonPinNames::Inputs::TimeSigDenominator);
+		DEFINE_METASOUND_PARAM_ALIAS(Tempo, CommonPinNames::Inputs::Tempo);
+		DEFINE_METASOUND_PARAM_ALIAS(Speed, CommonPinNames::Inputs::Speed);
+		DEFINE_METASOUND_PARAM_ALIAS(PrerollBars, CommonPinNames::Inputs::PrerollBars);
+	}
 
+	namespace Outputs
+	{
+		DEFINE_METASOUND_PARAM_ALIAS(MidiClock, CommonPinNames::Outputs::MidiClock);
+	}
+	
+	const FVertexInterface& FMetronomeOperator::GetVertexInterface()
+	{
 		static const FVertexInterface Interface(
 			FInputVertexInterface(
 				TInputDataVertex<FMusicTransportEventStream>(METASOUND_GET_PARAM_NAME_AND_METADATA(Inputs::Transport)),
@@ -107,8 +133,6 @@ namespace HarmonixMetasound::Nodes::MetronomeNode
 
 	TUniquePtr<IOperator> FMetronomeOperator::CreateOperator(const FBuildOperatorParams& InParams, FBuildResults& OutResults)
 	{
-		using namespace CommonPinNames;
-
 		const FMetronomeNode& TempoClockNode = static_cast<const FMetronomeNode&>(InParams.Node);
 
 		const FInputVertexInterfaceData& InputData = InParams.InputData;
@@ -171,7 +195,6 @@ namespace HarmonixMetasound::Nodes::MetronomeNode
 
 	void FMetronomeOperatorBase::BindInputs(FInputVertexInterfaceData& InVertexData)
 	{
-		using namespace CommonPinNames;
 		InVertexData.BindReadVertex(METASOUND_GET_PARAM_NAME(Inputs::Transport), TransportInPin);
 		InVertexData.SetValue(METASOUND_GET_PARAM_NAME(Inputs::Loop), LoopInPin);
 		InVertexData.SetValue(METASOUND_GET_PARAM_NAME(Inputs::LoopLengthBars), LoopLengthBarsInPin);
@@ -186,7 +209,6 @@ namespace HarmonixMetasound::Nodes::MetronomeNode
 
 	void FMetronomeOperatorBase::BindOutputs(FOutputVertexInterfaceData& InVertexData)
 	{
-		using namespace CommonPinNames;
 		InVertexData.BindReadVertex(METASOUND_GET_PARAM_NAME(Outputs::MidiClock), MidiClockOutPin);
 	}
 
diff --git a/Engine/Plugins/Runtime/Harmonix/Source/HarmonixMetasound/Private/HarmonixMetasound/Nodes/MusicSeekTargetBuilder.cpp b/Engine/Plugins/Runtime/Harmonix/Source/HarmonixMetasound/Private/HarmonixMetasound/Nodes/MusicSeekTargetBuilder.cpp
--- a/Engine/Plugins/Runtime/Harmonix/Source/HarmonixMetasound/Private/HarmonixMetasound/Nodes/MusicSeekTargetBuilder.cpp
+++ b/Engine/Plugins/Runtime/Harmonix/Source/HarmonixMetasound/Private/HarmonixMetasound/Nodes/MusicSeekTargetBuilder.cpp
@@ -18,12 +18,13 @@ namespace HarmonixMetasound::Nodes::BarBeatToSeekTarget
 {
 	using namespace Metasound;
 	
-	FNodeClassName GetClassName()
+	const FNodeClassName& GetClassName()
 	{
-		return { HarmonixNodeNamespace, TEXT("BarBeatToSeekTarget"), TEXT("") };
+		static const FNodeClassName ClassName { HarmonixNodeNamespace, TEXT("BarBeatToSeekTarget"), TEXT("") };
+		return ClassName;
 	}
 
-	int32 GetMajorVersion()
+	int32 GetCurrentMajorVersion()
 	{
 		return 0;
 	}
@@ -48,7 +49,7 @@ namespace HarmonixMetasound::Nodes::BarBeatToSeekTarget
 			{
 				FNodeClassMetadata Info;
 				Info.ClassName        = GetClassName();
-				Info.MajorVersion     = GetMajorVersion();
+				Info.MajorVersion     = GetCurrentMajorVersion();
 				Info.MinorVersion     = 1;
 				Info.DisplayName      = METASOUND_LOCTEXT("BarBeatToSeekTarget_DisplayName", "Bar & Beat To Seek Target");
 				Info.Description      = METASOUND_LOCTEXT("BarBeatToSeekTarget_Description", "Build a Seek Target from a Bar and Beat.");
@@ -146,12 +147,13 @@ namespace HarmonixMetasound::Nodes::TimeMsToSeekTarget
 {
 	using namespace Metasound;
 	
-	FNodeClassName GetClassName()
+	const FNodeClassName& GetClassName()
 	{
-		return { HarmonixNodeNamespace, TEXT("TimeMsToSeekTarget"), TEXT("") };
+		static const FNodeClassName ClassName { HarmonixNodeNamespace, TEXT("TimeMsToSeekTarget"), TEXT("") };
+		return ClassName;
 	}
 
-	int32 GetMajorVersion()
+	int32 GetCurrentMajorVersion()
 	{
 		return 0;
 	}
@@ -175,7 +177,7 @@ namespace HarmonixMetasound::Nodes::TimeMsToSeekTarget
 			{
 				FNodeClassMetadata Info;
 				Info.ClassName        = GetClassName();
-				Info.MajorVersion     = GetMajorVersion();
+				Info.MajorVersion     = GetCurrentMajorVersion();
 				Info.MinorVersion     = 1;
 				Info.DisplayName      = METASOUND_LOCTEXT("TimeMsToSeekTarget_DisplayName", "Time (Ms) To Seek Target");
 				Info.Description      = METASOUND_LOCTEXT("TimeMsToSeekTarget_Description", "Build a Seek Target from Time in Milliseconds.");
diff --git a/Engine/Plugins/Runtime/Harmonix/Source/HarmonixMetasound/Private/HarmonixMetasound/Nodes/TriggerToTransportNode.cpp b/Engine/Plugins/Runtime/Harmonix/Source/HarmonixMetasound/Private/HarmonixMetasound/Nodes/TriggerToTransportNode.cpp
--- a/Engine/Plugins/Runtime/Harmonix/Source/HarmonixMetasound/Private/HarmonixMetasound/Nodes/TriggerToTransportNode.cpp
+++ b/Engine/Plugins/Runtime/Harmonix/Source/HarmonixMetasound/Private/HarmonixMetasound/Nodes/TriggerToTransportNode.cpp
@@ -12,13 +12,25 @@
 #include "HarmonixMetasound/DataTypes/MusicTransport.h"
 #include "HarmonixMetasound/DataTypes/MusicTimestamp.h"
 #include "HarmonixMetasound/DataTypes/MusicSeekRequest.h"
+#include "HarmonixMetasound/Nodes/TriggerToTransportNode.h"
 
 #define LOCTEXT_NAMESPACE "HarmonixMetaSound"
 
-namespace HarmonixMetasound
+namespace HarmonixMetasound::Nodes::TriggerToTransportNode
 {
 	using namespace Metasound;
 
+	const FNodeClassName& GetClassName()
+	{
+		static const FNodeClassName ClassName { HarmonixNodeNamespace, TEXT("TriggerToTransport"), TEXT("") };
+		return ClassName;
+	}
+
+	int32 GetCurrentMajorVersion()
+	{
+		return 0;
+	}
+
 	class FTriggerToTransportOperator : public TExecutableOperator<FTriggerToTransportOperator>
 	{
 	public:
@@ -76,8 +88,8 @@ namespace HarmonixMetasound
 		auto InitNodeInfo = []() -> FNodeClassMetadata
 		{
 			FNodeClassMetadata Info;
-			Info.ClassName        = { HarmonixNodeNamespace, TEXT("TriggerToTransport"), TEXT("") };
-			Info.MajorVersion     = 0;
+			Info.ClassName        = GetClassName();
+			Info.MajorVersion     = GetCurrentMajorVersion();
 			Info.MinorVersion     = 2;
 			Info.DisplayName      = METASOUND_LOCTEXT("TriggerToTransportNode_DisplayName", "Trigger To Music Transport");
 			Info.Description      = METASOUND_LOCTEXT("TriggerToTransportNode_Description", "Combines input triggers into meaningful music transport requests.");
@@ -93,10 +105,25 @@ namespace HarmonixMetasound
 		return Info;
 	}
 
-	const FVertexInterface& FTriggerToTransportOperator::GetVertexInterface()
+	namespace Inputs
+	{
+		DEFINE_METASOUND_PARAM_ALIAS(TransportPrepare, CommonPinNames::Inputs::TransportPrepare);
+		DEFINE_METASOUND_PARAM_ALIAS(TransportPlay, CommonPinNames::Inputs::TransportPlay);
+		DEFINE_METASOUND_PARAM_ALIAS(TransportPause, CommonPinNames::Inputs::TransportPause);
+		DEFINE_METASOUND_PARAM_ALIAS(TransportContinue, CommonPinNames::Inputs::TransportContinue);
+		DEFINE_METASOUND_PARAM_ALIAS(TransportStop, CommonPinNames::Inputs::TransportStop);
+		DEFINE_METASOUND_PARAM_ALIAS(TransportKill, CommonPinNames::Inputs::TransportKill);
+		DEFINE_METASOUND_PARAM_ALIAS(TriggerSeek, CommonPinNames::Inputs::TriggerSeek);
+		DEFINE_METASOUND_PARAM_ALIAS(SeekDestination, CommonPinNames::Inputs::SeekDestination);
+	}
+
+	namespace Outputs
 	{
-		using namespace CommonPinNames;
+		DEFINE_METASOUND_PARAM_ALIAS(Transport, CommonPinNames::Outputs::Transport);
+	}
 
+	const FVertexInterface& FTriggerToTransportOperator::GetVertexInterface()
+	{
 		static const FVertexInterface Interface(
 			FInputVertexInterface(
 				TInputDataVertex<FTrigger>(METASOUND_GET_PARAM_NAME_AND_METADATA(Inputs::TransportPrepare)),
@@ -118,8 +145,6 @@ namespace HarmonixMetasound
 
 	TUniquePtr<IOperator> FTriggerToTransportOperator::CreateOperator(const FBuildOperatorParams& InParams, FBuildResults& OutResults)
 	{
-		using namespace CommonPinNames;
-
 		const FInputVertexInterfaceData& InputData = InParams.InputData;
 		FTriggerReadRef InTriggerPrepare          = InputData.GetOrConstructDataReadReference<FTrigger>(METASOUND_GET_PARAM_NAME(Inputs::TransportPrepare), InParams.OperatorSettings);
 		FTriggerReadRef InTriggerPlay             = InputData.GetOrConstructDataReadReference<FTrigger>(METASOUND_GET_PARAM_NAME(Inputs::TransportPlay), InParams.OperatorSettings);
@@ -164,8 +189,6 @@ namespace HarmonixMetasound
 
 	void FTriggerToTransportOperator::BindInputs(FInputVertexInterfaceData& InVertexData)
 	{
-		using namespace CommonPinNames;
-
 		InVertexData.BindReadVertex(METASOUND_GET_PARAM_NAME(Inputs::TransportPrepare), PrepareInPin);
 		InVertexData.BindReadVertex(METASOUND_GET_PARAM_NAME(Inputs::TransportPlay), PlayInPin);
 		InVertexData.BindReadVertex(METASOUND_GET_PARAM_NAME(Inputs::TransportPause), PauseInPin);
@@ -178,8 +201,6 @@ namespace HarmonixMetasound
 
 	void FTriggerToTransportOperator::BindOutputs(FOutputVertexInterfaceData& InVertexData)
 	{
-		using namespace CommonPinNames;
-
 		InVertexData.BindReadVertex(METASOUND_GET_PARAM_NAME(Outputs::Transport), TransportOutPin);
 	}
 
diff --git a/Engine/Plugins/Runtime/Harmonix/Source/HarmonixMetasound/Private/HarmonixMetasoundModule.cpp b/Engine/Plugins/Runtime/Harmonix/Source/HarmonixMetasound/Private/HarmonixMetasoundModule.cpp
--- a/Engine/Plugins/Runtime/Harmonix/Source/HarmonixMetasound/Private/HarmonixMetasoundModule.cpp
+++ b/Engine/Plugins/Runtime/Harmonix/Source/HarmonixMetasound/Private/HarmonixMetasoundModule.cpp
@@ -4,6 +4,7 @@
 
 #include "MetasoundDataTypeRegistrationMacro.h"
 #include "MetasoundGeneratorHandle.h"
+#include "MusicEnvironmentSubsystem.h"
 #include "Analysis/MetasoundFrontendAnalyzerRegistry.h"
 
 #include "HarmonixMetasound/Analysis/MidiClockVertexAnalyzer.h"
@@ -13,7 +14,11 @@
 #include "HarmonixMetasound/Analysis/FFTAnalyzerResultVertexAnalyzer.h"
 #include "HarmonixMetasound/DataTypes/FFTAnalyzerResult.h"
 #include "HarmonixMetasound/DataTypes/MidiStream.h"
+#include "HarmonixMetasound/MusicEnvironmentSupport/HarmonixMusicEnvironmentMetronome.h"
+
+#include "Engine/Engine.h"
 #include "Logging/LogMacros.h"
+#include "Misc/CoreDelegates.h"
 #include "Modules/ModuleManager.h"
 #include "UObject/CoreRedirects.h"
 
@@ -50,11 +55,20 @@ void FHarmonixMetasoundModule::StartupModule()
 	METASOUND_REGISTER_VERTEX_ANALYZER_FACTORY(HarmonixMetasound::Analysis::FMusicTransportEventStreamVertexAnalyzer)
 	METASOUND_REGISTER_VERTEX_ANALYZER_FACTORY(HarmonixMetasound::Analysis::FFFTAnalyzerResultVertexAnalyzer)
 
-	// The first redirect for thie module
+	// The first redirect for the module
 	TArray<FCoreRedirect> Redirects;
 	Redirects.Emplace(ECoreRedirectFlags::Type_Function, TEXT("MusicClockComponent.CreateMusicClockComponent"), TEXT("MusicClockComponent.CreateMetasoundDrivenMusicClock"));
 	FCoreRedirects::AddRedirectList(Redirects, TEXT("HarmonixMetasoundModule"));
 
+	// When the engine is done loading we want to register our metronome type with the Music Environment Subsystem so
+	// it can spawn metronomes for things like the sequence player/editor.
+	FCoreDelegates::OnAllModuleLoadingPhasesComplete.AddLambda([]()
+		{
+			if (UMusicEnvironmentSubsystem* MusicEnvironment = GEngine->GetEngineSubsystem<UMusicEnvironmentSubsystem>())
+			{
+				MusicEnvironment->SetMetronomeClass(UHarmonixMusicEnvironmentMetronome::StaticClass());
+			}
+		});
 }
 
 void FHarmonixMetasoundModule::ShutdownModule()
diff --git a/Engine/Plugins/Runtime/Harmonix/Source/HarmonixMetasound/Public/HarmonixMetasound/Components/MusicClockComponent.h b/Engine/Plugins/Runtime/Harmonix/Source/HarmonixMetasound/Public/HarmonixMetasound/Components/MusicClockComponent.h
--- a/Engine/Plugins/Runtime/Harmonix/Source/HarmonixMetasound/Public/HarmonixMetasound/Components/MusicClockComponent.h
+++ b/Engine/Plugins/Runtime/Harmonix/Source/HarmonixMetasound/Public/HarmonixMetasound/Components/MusicClockComponent.h
@@ -50,7 +50,7 @@ DECLARE_DYNAMIC_MULTICAST_DELEGATE(FMusicClockConnected);
 DECLARE_DYNAMIC_MULTICAST_DELEGATE(FMusicClockDisconnected);
 
 UCLASS(ClassGroup = (MetaSoundMusic), PrioritizeCategories = "MusicClock", meta = (BlueprintSpawnableComponent, DisplayName = "Music Clock", ScriptName = MusicClockComponent))
-class HARMONIXMETASOUND_API UMusicClockComponent : public UActorComponent, public IMovieSceneCustomClockSource
+class HARMONIXMETASOUND_API UMusicClockComponent : public UActorComponent/*, public IMovieSceneMusicClockSource*/
 {
 	GENERATED_BODY()
 
@@ -368,17 +368,17 @@ public:
 
 	const FMidiSongPos& GetRawUnsmoothedAudioRenderPos() const { return CurrentRawAudioRenderSongPos; }
 
-	// vvv BEGIN IMovieSceneCustomClockSource Implementation vvv
+	// vvv BEGIN IMovieSceneMusicClockSource Implementation vvv
+	/*
 	virtual void OnTick(float DeltaSeconds, float InPlayRate) override;
 	virtual void OnStartPlaying(const FQualifiedFrameTime& InStartTime) override;
 	virtual void OnStopPlaying(const FQualifiedFrameTime& InStopTime) override;
-	virtual FFrameTime OnRequestCurrentTime(const FQualifiedFrameTime& InCurrentTime, float InPlayRate) override;
-	virtual bool OnSupportsMusicalTime() override { return true; }
-	virtual FMusicalTime OnRequestCurrentMusicalTime(const FQualifiedFrameTime& InCurrentTime, float InPlayRate) override;
-	// ^^^ END IMovieSceneCustomClockSource Implementation ^^^
+	virtual FMusicalTime GetMusicalTime() const override;
+	*/
+	// ^^^ END IMovieSceneMusicClockSource Implementation ^^^
 
 
-protected:
+public:
 	virtual void TickComponent(float DeltaTime, enum ELevelTick TickType, FActorComponentTickFunction* ThisTickFunction);
 
 private:
diff --git a/Engine/Plugins/Runtime/Harmonix/Source/HarmonixMetasound/Public/HarmonixMetasound/Nodes/MetronomeNode.h b/Engine/Plugins/Runtime/Harmonix/Source/HarmonixMetasound/Public/HarmonixMetasound/Nodes/MetronomeNode.h
--- a/Engine/Plugins/Runtime/Harmonix/Source/HarmonixMetasound/Public/HarmonixMetasound/Nodes/MetronomeNode.h
+++ b/Engine/Plugins/Runtime/Harmonix/Source/HarmonixMetasound/Public/HarmonixMetasound/Nodes/MetronomeNode.h
@@ -20,9 +20,26 @@ namespace HarmonixMetasound::Nodes::MetronomeNode
 {
 	using namespace Metasound;
 
-	//HARMONIXMETASOUND_API Metasound::FNodeClassName GetClassName();
-	//HARMONIXMETASOUND_API int32 GetCurrentMajorVersion();
+	HARMONIXMETASOUND_API const Metasound::FNodeClassName& GetClassName();
+	HARMONIXMETASOUND_API int32 GetCurrentMajorVersion();
 
+	namespace Inputs
+	{
+		DECLARE_METASOUND_PARAM_ALIAS(Transport);
+		DECLARE_METASOUND_PARAM_ALIAS(Loop);
+		DECLARE_METASOUND_PARAM_ALIAS(LoopLengthBars);
+		DECLARE_METASOUND_PARAM_ALIAS(TimeSigNumerator);
+		DECLARE_METASOUND_PARAM_ALIAS(TimeSigDenominator);
+		DECLARE_METASOUND_PARAM_ALIAS(Tempo);
+		DECLARE_METASOUND_PARAM_ALIAS(Speed);
+		DECLARE_METASOUND_PARAM_ALIAS(PrerollBars);
+	}
+
+	namespace Outputs
+	{
+		DECLARE_METASOUND_PARAM_ALIAS(MidiClock);
+	}
+ 
 	class HARMONIXMETASOUND_API FMetronomeOperatorBase : public TExecutableOperator<FMetronomeOperatorBase>, public FMusicTransportControllable
 	{
 	public:
diff --git a/Engine/Plugins/Runtime/Harmonix/Source/HarmonixMidi/HarmonixMidi.Build.cs b/Engine/Plugins/Runtime/Harmonix/Source/HarmonixMidi/HarmonixMidi.Build.cs
--- a/Engine/Plugins/Runtime/Harmonix/Source/HarmonixMidi/HarmonixMidi.Build.cs
+++ b/Engine/Plugins/Runtime/Harmonix/Source/HarmonixMidi/HarmonixMidi.Build.cs
@@ -16,6 +16,7 @@ public class HarmonixMidi : ModuleRules
 				"Core",
 				"CoreUObject",
 				"TimeManagement",
+				"MusicEnvironment",
 			});
 		PrivateDependencyModuleNames.AddRange(
 			new string[]
diff --git a/Engine/Source/Editor/MovieSceneTools/Private/Channels/BuiltInChannelEditors.cpp b/Engine/Source/Editor/MovieSceneTools/Private/Channels/BuiltInChannelEditors.cpp
--- a/Engine/Source/Editor/MovieSceneTools/Private/Channels/BuiltInChannelEditors.cpp
+++ b/Engine/Source/Editor/MovieSceneTools/Private/Channels/BuiltInChannelEditors.cpp
@@ -529,7 +529,7 @@ TSharedRef<SWidget> CreateKeyEditor(const TMovieSceneChannelHandle<FMovieSceneTi
 	if (ChannelPtr && ChannelPtr->Domain == UE::MovieScene::ETimeWarpChannelDomain::Time)
 	{
 		// Set the numeric type interface for frame numbers if the channel is in the time domain
-		KeyEditor.SetNumericTypeInterface(Params.Sequencer->GetNumericPositionTypeInterface());
+		KeyEditor.SetNumericTypeInterface(Params.Sequencer->GetNumericTypeInterface());
 	}
 
 	typedef SNumericKeyEditor<FMovieSceneTimeWarpChannel, double> KeyEditorType;
diff --git a/Engine/Source/Editor/MovieSceneTools/Private/FrameNumberDetailsCustomization.cpp b/Engine/Source/Editor/MovieSceneTools/Private/FrameNumberDetailsCustomization.cpp
--- a/Engine/Source/Editor/MovieSceneTools/Private/FrameNumberDetailsCustomization.cpp
+++ b/Engine/Source/Editor/MovieSceneTools/Private/FrameNumberDetailsCustomization.cpp
@@ -9,15 +9,17 @@
 #include "Widgets/Input/SEditableTextBox.h"
 #include "PropertyCustomizationHelpers.h"
 #include "Misc/FrameRate.h"
+#include "ISequencer.h"
 
 #define LOCTEXT_NAMESPACE "TimeManagement.FrameNumber"
 
 void FFrameNumberDetailsCustomization::CustomizeChildren(TSharedRef<IPropertyHandle> PropertyHandle, IDetailChildrenBuilder& ChildBuilder, IPropertyTypeCustomizationUtils& CustomizationUtils)
 {
+	using namespace UE::Sequencer;
+
 	// Check for Min/Max metadata on the property itself and we'll apply it to the child when changed.
 	const FString& MetaUIMinString = PropertyHandle->GetMetaData(TEXT("UIMin"));
 	const FString& MetaUIMaxString = PropertyHandle->GetMetaData(TEXT("UIMax"));
-	const FString& MetaUIDisplayAsString = PropertyHandle->GetMetaData(TEXT("UIFrameDisplayAs"));
 
 	UIClampMin = TNumericLimits<int32>::Lowest();
 	UIClampMax = TNumericLimits<int32>::Max();
@@ -32,19 +34,19 @@ void FFrameNumberDetailsCustomization::CustomizeChildren(TSharedRef<IPropertyHan
 		TTypeFromString<int32>::FromString(UIClampMax, *MetaUIMaxString);
 	}
 
-	if (!MetaUIDisplayAsString.IsEmpty())
+	TSharedPtr<ISequencer> Sequencer = WeakSequencer.Pin();
+	if (NumericTypeInterface == nullptr && Sequencer)
 	{
-		if (MetaUIDisplayAsString.Equals(TEXT("Duration"), ESearchCase::IgnoreCase) && NumericDurationTypeInterface.IsValid())
-		{
-			NumericTypeInterface = NumericDurationTypeInterface;
-		}
-		else if (MetaUIDisplayAsString.Equals(TEXT("NonMusical"), ESearchCase::IgnoreCase) && NumericNonMusicalTypeInterface.IsValid())
-		{
-			NumericTypeInterface = NumericNonMusicalTypeInterface;
-		}
-		else
+		// Find the most relevant interface for this property
+		int32 BestScore = MIN_int32;
+		for (TSharedRef<FSequencerNumericTypeInterface> Interface : Sequencer->GetNumericTypeInterfaces())
 		{
-			NumericTypeInterface = NumericPositionTypeInterface;
+			const int32 Score = Interface->GetRelevancyScore(*Sequencer, PropertyHandle);
+			if (Score > BestScore)
+			{
+				BestScore = Score;
+				NumericTypeInterface = Interface->Interface;
+			}
 		}
 	}
 
diff --git a/Engine/Source/Editor/MovieSceneTools/Private/MovieSceneSectionDetailsCustomization.cpp b/Engine/Source/Editor/MovieSceneTools/Private/MovieSceneSectionDetailsCustomization.cpp
--- a/Engine/Source/Editor/MovieSceneTools/Private/MovieSceneSectionDetailsCustomization.cpp
+++ b/Engine/Source/Editor/MovieSceneTools/Private/MovieSceneSectionDetailsCustomization.cpp
@@ -205,7 +205,7 @@ FText FMovieSceneSectionDetailsCustomization::OnGetRangeStartText() const
 			return NSLOCTEXT("PropertyEditor", "MultipleValues", "Multiple Values");
 		case ERangeBoundValueType::Finite:
 		default:
-			return FText::FromString(NumericPositionTypeInterface->ToString((double)FrameValue.Value));
+			return FText::FromString(NumericTypeInterface->ToString((double)FrameValue.Value));
 	}
 }
 
@@ -229,7 +229,7 @@ FText FMovieSceneSectionDetailsCustomization::OnGetRangeStartToolTipText() const
 void FMovieSceneSectionDetailsCustomization::OnRangeStartTextCommitted(const FText& InText, ETextCommit::Type CommitInfo)
 {
 	// Find the new value for the start range.
-	TOptional<double> NewStart = NumericPositionTypeInterface->FromString(InText.ToString(), 0.0);
+	TOptional<double> NewStart = NumericTypeInterface->FromString(InText.ToString(), 0.0);
 
 	// Early out if we couldn't parse it, no need to reset them all to zero.
 	if (!NewStart.IsSet())
@@ -437,7 +437,7 @@ FText FMovieSceneSectionDetailsCustomization::OnGetRangeEndText() const
 			return NSLOCTEXT("PropertyEditor", "MultipleValues", "Multiple Values");
 		case ERangeBoundValueType::Finite:
 		default:
-			return FText::FromString(NumericPositionTypeInterface->ToString((double)FrameValue.Value));
+			return FText::FromString(NumericTypeInterface->ToString((double)FrameValue.Value));
 	}
 }
 
@@ -461,7 +461,7 @@ FText FMovieSceneSectionDetailsCustomization::OnGetRangeEndToolTipText() const
 void FMovieSceneSectionDetailsCustomization::OnRangeEndTextCommitted(const FText& InText, ETextCommit::Type CommitInfo)
 {
 	// Find the new value for the end range.
-	TOptional<double> NewEnd = NumericPositionTypeInterface->FromString(InText.ToString(), 0.0);
+	TOptional<double> NewEnd = NumericTypeInterface->FromString(InText.ToString(), 0.0);
 
 	// Early out if we couldn't parse it, no need to reset them all to zero.
 	if (!NewEnd.IsSet())
diff --git a/Engine/Source/Editor/MovieSceneTools/Private/MovieSceneToolHelpers.cpp b/Engine/Source/Editor/MovieSceneTools/Private/MovieSceneToolHelpers.cpp
--- a/Engine/Source/Editor/MovieSceneTools/Private/MovieSceneToolHelpers.cpp
+++ b/Engine/Source/Editor/MovieSceneTools/Private/MovieSceneToolHelpers.cpp
@@ -1962,8 +1962,7 @@ class SControlRigImportFBXSettings : public SCompoundWidget
 
 		Sequencer = InSequencer;
 
-		TSharedPtr<INumericTypeInterface<double>> NumericPositionTypeInterface = (InSequencer->GetNumericPositionTypeInterface());
-		DetailView->RegisterInstancedCustomPropertyTypeLayout("FrameNumber", FOnGetPropertyTypeCustomizationInstance::CreateStatic(&FFrameNumberDetailsCustomization::MakeInstance, NumericPositionTypeInterface));
+		DetailView->RegisterInstancedCustomPropertyTypeLayout("FrameNumber", FOnGetPropertyTypeCustomizationInstance::CreateSP(InSequencer, &ISequencer::MakeFrameNumberDetailsCustomization));
 
 
 		ChildSlot
diff --git a/Engine/Source/Editor/MovieSceneTools/Private/SKeyEditInterface.cpp b/Engine/Source/Editor/MovieSceneTools/Private/SKeyEditInterface.cpp
--- a/Engine/Source/Editor/MovieSceneTools/Private/SKeyEditInterface.cpp
+++ b/Engine/Source/Editor/MovieSceneTools/Private/SKeyEditInterface.cpp
@@ -111,11 +111,7 @@ TSharedRef<IPropertyTypeCustomization> SKeyEditInterface::CreateBindingIDCustomi
 
 TSharedRef<IPropertyTypeCustomization> SKeyEditInterface::CreateFrameNumberCustomization()
 {
-	if (const TSharedPtr<ISequencer> SequencerPtr = WeakSequencer.Pin())
-	{
-		return MakeShared<FFrameNumberDetailsCustomization>(SequencerPtr->GetNumericPositionTypeInterface(), SequencerPtr->GetNumericDurationTypeInterface(), SequencerPtr->GetNumericNonMusicalTypeInterface());
-	}
-	return MakeShared<FFrameNumberDetailsCustomization>(nullptr);
+	return WeakSequencer.Pin()->MakeFrameNumberDetailsCustomization();
 }
 
 TSharedRef<IPropertyTypeCustomization> SKeyEditInterface::CreateEventCustomization()
diff --git a/Engine/Source/Editor/MovieSceneTools/Private/Sections/ThumbnailSection.cpp b/Engine/Source/Editor/MovieSceneTools/Private/Sections/ThumbnailSection.cpp
--- a/Engine/Source/Editor/MovieSceneTools/Private/Sections/ThumbnailSection.cpp
+++ b/Engine/Source/Editor/MovieSceneTools/Private/Sections/ThumbnailSection.cpp
@@ -136,7 +136,7 @@ void FThumbnailSection::BuildThumbnailsMenu(FMenuBuilder& InMenuBuilder)
 		return;
 	}
 
-	const FText CurrentTime = FText::FromString(Sequencer->GetNumericPositionTypeInterface()->ToString(Sequencer->GetLocalTime().Time.GetFrame().Value));
+	const FText CurrentTime = FText::FromString(Sequencer->GetNumericTypeInterface()->ToString(Sequencer->GetLocalTime().Time.GetFrame().Value));
 
 	InMenuBuilder.BeginSection(TEXT("Thumbnails"), LOCTEXT("ThumbnailsMenuSection", "Thumbnails"));
 	{
diff --git a/Engine/Source/Editor/MovieSceneTools/Public/FrameNumberDetailsCustomization.h b/Engine/Source/Editor/MovieSceneTools/Public/FrameNumberDetailsCustomization.h
--- a/Engine/Source/Editor/MovieSceneTools/Public/FrameNumberDetailsCustomization.h
+++ b/Engine/Source/Editor/MovieSceneTools/Public/FrameNumberDetailsCustomization.h
@@ -4,6 +4,7 @@
 
 #include "CoreMinimal.h"
 #include "IPropertyTypeCustomization.h"
+#include "ISequencerNumericTypeInterface.h"
 #include "Widgets/Input/NumericTypeInterface.h"
 
 class IDetailLayoutBuilder;
@@ -14,24 +15,24 @@ class IDetailLayoutBuilder;
 class MOVIESCENETOOLS_API FFrameNumberDetailsCustomization : public IPropertyTypeCustomization
 {
 public:
-	static TSharedRef<IPropertyTypeCustomization> MakeInstance(TSharedPtr<INumericTypeInterface<double>> InNumericPositionTypeInterface)
+
+	static TSharedRef<IPropertyTypeCustomization> MakeInstance(TSharedPtr<INumericTypeInterface<double>> InNumericTypeInterface)
 	{
-		return MakeShared<FFrameNumberDetailsCustomization>(InNumericPositionTypeInterface/*, nullptr, nullptr*/);
+		return MakeShared<FFrameNumberDetailsCustomization>(MoveTemp(InNumericTypeInterface));
 	}
 
-	static TSharedRef<IPropertyTypeCustomization> MakeInstance(TSharedPtr<INumericTypeInterface<double>> InNumericPositionTypeInterface,
-		TSharedPtr<INumericTypeInterface<double>> InNumericDurationTypeInterface, TSharedPtr<INumericTypeInterface<double>> InNumericNonMusicalTypeInterface)
+	static TSharedRef<IPropertyTypeCustomization> MakeInstance(TWeakPtr<ISequencer> Sequencer)
 	{
-		return MakeShared<FFrameNumberDetailsCustomization>(InNumericPositionTypeInterface, InNumericDurationTypeInterface, InNumericNonMusicalTypeInterface);
+		return MakeShared<FFrameNumberDetailsCustomization>(Sequencer);
 	}
 
-	FFrameNumberDetailsCustomization(TSharedPtr<INumericTypeInterface<double>> InNumericPositionTypeInterface, TSharedPtr<INumericTypeInterface<double>> InNumericDurationTypeInterface = nullptr, TSharedPtr<INumericTypeInterface<double>> InNumericNonMusicalTypeInterface = nullptr)
+	FFrameNumberDetailsCustomization(TWeakPtr<ISequencer> Sequencer)
+	{
+		WeakSequencer = Sequencer;
+	}
+	FFrameNumberDetailsCustomization(TSharedPtr<INumericTypeInterface<double>> InNumericTypeInterface)
 	{
-		NumericPositionTypeInterface = InNumericPositionTypeInterface;
-		NumericDurationTypeInterface = InNumericDurationTypeInterface;
-		NumericNonMusicalTypeInterface = InNumericNonMusicalTypeInterface;
-		// default to position in the current display format...
-		NumericTypeInterface = NumericPositionTypeInterface.IsValid() ? NumericPositionTypeInterface : NumericNonMusicalTypeInterface.IsValid() ? NumericNonMusicalTypeInterface : NumericDurationTypeInterface;
+		NumericTypeInterface = InNumericTypeInterface;
 	}
 
 	/** IPropertyTypeCustomization interface */
@@ -43,17 +44,9 @@ private:
 	FText OnGetTimeToolTipText() const;
 	void OnTimeTextCommitted(const FText& InText, ETextCommit::Type CommitInfo);
 
-	/** The Numeric Type interface used to convert between display formats and internal tick resolution. */
-	TSharedPtr<INumericTypeInterface<double>> NumericTypeInterface;
+	TWeakPtr<ISequencer> WeakSequencer;
 
-	/** The Numeric Type interface used to convert between display formats and internal tick resolution. */
-	TSharedPtr<INumericTypeInterface<double>> NumericPositionTypeInterface;
-
-	/** The Numeric Type interface used to convert between display formats and internal tick resolution. */
-	TSharedPtr<INumericTypeInterface<double>> NumericDurationTypeInterface;
-
-	/** The Numeric Type interface used to convert between display formats and internal tick resolution. */
-	TSharedPtr<INumericTypeInterface<double>> NumericNonMusicalTypeInterface;
+	TSharedPtr<INumericTypeInterface<double>> NumericTypeInterface;
 
 	/** Store the property handle to the FrameNumber field so we can get/set the value on the object via text box callbacks. */
 	TSharedPtr<IPropertyHandle> FrameNumberProperty;
diff --git a/Engine/Source/Editor/MovieSceneTools/Public/MovieSceneSectionDetailsCustomization.h b/Engine/Source/Editor/MovieSceneTools/Public/MovieSceneSectionDetailsCustomization.h
--- a/Engine/Source/Editor/MovieSceneTools/Public/MovieSceneSectionDetailsCustomization.h
+++ b/Engine/Source/Editor/MovieSceneTools/Public/MovieSceneSectionDetailsCustomization.h
@@ -17,16 +17,14 @@ class UMovieScene;
 class MOVIESCENETOOLS_API FMovieSceneSectionDetailsCustomization : public IDetailCustomization
 {
 public:
-	static TSharedRef<IDetailCustomization> MakeInstance(TSharedPtr<INumericTypeInterface<double>> InNumericPositionTypeInterface, TSharedPtr<INumericTypeInterface<double>> InNumericDurationTypeInterface, TSharedPtr<INumericTypeInterface<double>> InNumericNonMusicalTypeInterface, TWeakObjectPtr<UMovieScene> InParentMovieScene)
+	static TSharedRef<IDetailCustomization> MakeInstance(TSharedPtr<INumericTypeInterface<double>> InNumericTypeInterface, TWeakObjectPtr<UMovieScene> InParentMovieScene)
 	{
-		return MakeShared<FMovieSceneSectionDetailsCustomization>(InNumericPositionTypeInterface, InNumericDurationTypeInterface, InNumericNonMusicalTypeInterface, InParentMovieScene);
+		return MakeShared<FMovieSceneSectionDetailsCustomization>(InNumericTypeInterface, InParentMovieScene);
 	}
 
-	FMovieSceneSectionDetailsCustomization(TSharedPtr<INumericTypeInterface<double>> InNumericPositionTypeInterface, TSharedPtr<INumericTypeInterface<double>> InNumericDurationTypeInterface, TSharedPtr<INumericTypeInterface<double>> InNumericNonMusicalTypeInterface, TWeakObjectPtr<UMovieScene> InParentMovieScene)
+	FMovieSceneSectionDetailsCustomization(TSharedPtr<INumericTypeInterface<double>> InNumericTypeInterface, TWeakObjectPtr<UMovieScene> InParentMovieScene)
 	{
-		NumericPositionTypeInterface = InNumericPositionTypeInterface;
-		NumericDurationTypeInterface = InNumericDurationTypeInterface;
-		NumericNonMusicalTypeInterface = InNumericNonMusicalTypeInterface;
+		NumericTypeInterface = InNumericTypeInterface;
 		ParentMovieScene = InParentMovieScene;
 	}
 
@@ -91,13 +89,7 @@ private:
 private:
 
 	/** The Numeric Type interface used to convert between display formats and internal tick resolution. */
-	TSharedPtr<INumericTypeInterface<double>> NumericPositionTypeInterface;
-
-	/** The Numeric Type interface used to convert between display formats and internal tick resolution. */
-	TSharedPtr<INumericTypeInterface<double>> NumericDurationTypeInterface;
-
-	/** The Numeric Type interface used to convert between display formats and internal tick resolution. */
-	TSharedPtr<INumericTypeInterface<double>> NumericNonMusicalTypeInterface;
+	TSharedPtr<INumericTypeInterface<double>> NumericTypeInterface;
 
 	/** Store the property handle to the FrameNumber field so we can get/set the value on the object via text box callbacks. */
 	TSharedPtr<IPropertyHandle> MovieSceneSectionPropertyHandle;
diff --git a/Engine/Source/Editor/Sequencer/Private/ISequencer.cpp b/Engine/Source/Editor/Sequencer/Private/ISequencer.cpp
--- a/Engine/Source/Editor/Sequencer/Private/ISequencer.cpp
+++ b/Engine/Source/Editor/Sequencer/Private/ISequencer.cpp
@@ -72,18 +72,6 @@ FFrameRate ISequencer::GetFocusedDisplayRate() const
 	return FFrameRate();
 }
 
-const UFrameBasedMusicMap* ISequencer::GetFocusedMusicMap() const
-{
-	UMovieSceneSequence* FocusedSequence = GetFocusedMovieSceneSequence();
-	if (FocusedSequence)
-	{
-		return FocusedSequence->GetMovieScene()->GetMusicMap();
-	}
-
-	ensureMsgf(false, TEXT("No valid sequence found."));
-	return nullptr;
-}
-
 FGuid ISequencer::CreateBinding(UObject& InObject, const FString& InName)
 {
 	UE::Sequencer::FCreateBindingParams BindingParams;
diff --git a/Engine/Source/Editor/Sequencer/Private/MVVM/CurveEditorExtension.cpp b/Engine/Source/Editor/Sequencer/Private/MVVM/CurveEditorExtension.cpp
--- a/Engine/Source/Editor/Sequencer/Private/MVVM/CurveEditorExtension.cpp
+++ b/Engine/Source/Editor/Sequencer/Private/MVVM/CurveEditorExtension.cpp
@@ -161,7 +161,7 @@ public:
 
 	void HandleSequenceActivated(FMovieSceneSequenceIDRef NewSequenceID)
 	{
-		FocusedTimeAxis->NumericTypeInterface = WeakSequencer.Pin()->GetNumericPositionTypeInterface();
+		FocusedTimeAxis->NumericTypeInterface = WeakSequencer.Pin()->GetNumericTypeInterface();
 	}
 };
 
@@ -226,12 +226,6 @@ struct FSequencerCurveEditorBounds : ICurveEditorBounds
 
 const FName FCurveEditorExtension::CurveEditorTabName = FName(TEXT("SequencerGraphEditor"));
 
-TSharedRef<IPropertyTypeCustomization> CreateFrameNumberCustomization(TWeakPtr<ISequencer> WeakSequencer)
-{
-	TSharedPtr<ISequencer> SequencerPtr = WeakSequencer.Pin();
-	return MakeShared<FFrameNumberDetailsCustomization>(SequencerPtr->GetNumericPositionTypeInterface());
-}
-
 FCurveEditorExtension::FCurveEditorExtension()
 {
 }
@@ -292,7 +286,7 @@ void FCurveEditorExtension::CreateCurveEditor(const FTimeSliderArgs& TimeSliderA
 	TSharedRef<ITimeSliderController> CurveEditorTimeSliderController = MakeShared<FSequencerCurveEditorTimeSliderController>(
 			TimeSliderArgs, Sequencer, CurveEditorModel.ToSharedRef());
 	
-	PlayTimeDisplay = StaticCastSharedRef<STemporarilyFocusedSpinBox<double>>(Sequencer->MakePlayTimeDisplay(TimeSliderArgs.NumericTypeInterface.ToSharedRef()));
+	PlayTimeDisplay = StaticCastSharedRef<STemporarilyFocusedSpinBox<double>>(Sequencer->MakePlayTimeDisplay());
 
 	CurveEditorTreeView = SNew(SCurveEditorTree, CurveEditorModel);
 	CurveEditorPanel = SNew(SCurveEditorPanel, CurveEditorModel.ToSharedRef())
@@ -417,10 +411,10 @@ void FCurveEditorExtension::CreateCurveEditor(const FTimeSliderArgs& TimeSliderA
 	TWeakPtr<ISequencer> WeakSequencer(Sequencer);
 	CurveEditorPanel->GetKeyDetailsView()->GetPropertyRowGenerator()->RegisterInstancedCustomPropertyTypeLayout(
 			"FrameNumber", 
-			FOnGetPropertyTypeCustomizationInstance::CreateStatic(CreateFrameNumberCustomization, WeakSequencer));
+			FOnGetPropertyTypeCustomizationInstance::CreateSP(Sequencer.ToSharedRef(), &FSequencer::MakeFrameNumberDetailsCustomization));
 	CurveEditorPanel->GetToolPropertiesPanel()->GetStructureDetailsView()->GetDetailsView()->RegisterInstancedCustomPropertyTypeLayout(
 		"FrameNumber",
-		FOnGetPropertyTypeCustomizationInstance::CreateStatic(CreateFrameNumberCustomization, WeakSequencer));
+		FOnGetPropertyTypeCustomizationInstance::CreateSP(Sequencer.ToSharedRef(), &FSequencer::MakeFrameNumberDetailsCustomization));
 
 	// And jump to the Curve Editor tree search if you have the Curve Editor focused
 	CurveEditorModel->GetCommands()->MapAction(
@@ -468,7 +462,7 @@ void FCurveEditorExtension::FilterClassChanged()
 			TWeakPtr<ISequencer> WeakSequencer(Sequencer);
 			FilterPanel->GetDetailsView()->RegisterInstancedCustomPropertyTypeLayout(
 				"FrameNumber",
-				FOnGetPropertyTypeCustomizationInstance::CreateStatic(CreateFrameNumberCustomization, WeakSequencer));
+				FOnGetPropertyTypeCustomizationInstance::CreateSP(Sequencer.ToSharedRef(), &FSequencer::MakeFrameNumberDetailsCustomization));
 		}
 	}
 }
diff --git a/Engine/Source/Editor/Sequencer/Private/MVVM/SequencerViewModelTypeIDs.cpp b/Engine/Source/Editor/Sequencer/Private/MVVM/SequencerViewModelTypeIDs.cpp
--- a/Engine/Source/Editor/Sequencer/Private/MVVM/SequencerViewModelTypeIDs.cpp
+++ b/Engine/Source/Editor/Sequencer/Private/MVVM/SequencerViewModelTypeIDs.cpp
@@ -20,6 +20,7 @@
 #include "MVVM/Extensions/IBindingLifetimeExtension.h"
 #include "MVVM/Extensions/IConditionableExtension.h"
 #include "MVVM/Extensions/ITimeDomainExtension.h"
+#include "MVVM/Extensions/IClockExtension.h"
 
 #include "MVVM/ViewModels/SequencerOutlinerViewModel.h"
 #include "MVVM/ViewModels/SequencerEditorViewModel.h"
@@ -62,6 +63,7 @@ UE_SEQUENCER_DEFINE_CASTABLE(FTrackRowModel);
 // Interface types
 UE_SEQUENCER_DEFINE_VIEW_MODEL_TYPE_ID(ITrackExtension);
 UE_SEQUENCER_DEFINE_VIEW_MODEL_TYPE_ID(IConditionableExtension);
+UE_SEQUENCER_DEFINE_VIEW_MODEL_TYPE_ID(IClockExtension);
 
 // View model types
 UE_SEQUENCER_DEFINE_CASTABLE(FSequencerEditorViewModel);
diff --git a/Engine/Source/Editor/Sequencer/Private/MVVM/ViewModels/SequenceModel.cpp b/Engine/Source/Editor/Sequencer/Private/MVVM/ViewModels/SequenceModel.cpp
--- a/Engine/Source/Editor/Sequencer/Private/MVVM/ViewModels/SequenceModel.cpp
+++ b/Engine/Source/Editor/Sequencer/Private/MVVM/ViewModels/SequenceModel.cpp
@@ -39,26 +39,28 @@ FSequenceModel::FSequenceModel(TWeakPtr<FSequencerEditorViewModel> InEditorViewM
 	WeakEditor = InEditorViewModel;
 
 	RegisterChildList(&RootOutlinerItems);
+
+	TSharedPtr<FSharedViewModelData> NewSharedData = MakeShared<FEditorSharedViewModelData>(WeakEditor.Pin().ToSharedRef());
+	SetSharedDataFromConstructor(NewSharedData);
 }
 
 void FSequenceModel::InitializeExtensions()
 {
-	TSharedPtr<FSharedViewModelData> NewSharedData = MakeShared<FEditorSharedViewModelData>(WeakEditor.Pin().ToSharedRef());
-	SetSharedData(NewSharedData);
+	TSharedPtr<FSharedViewModelData> ThisSharedData = GetSharedData();
 
 	// Re-generate hierarchical caches when the sequence changes
-	NewSharedData->AddDynamicExtension(FOutlinerCacheExtension::ID);
-	NewSharedData->AddDynamicExtension(FMuteStateCacheExtension::ID);
-	NewSharedData->AddDynamicExtension(FSoloStateCacheExtension::ID);
-	NewSharedData->AddDynamicExtension(FLockStateCacheExtension::ID);
-	NewSharedData->AddDynamicExtension(FConditionStateCacheExtension::ID);
+	ThisSharedData->AddDynamicExtension(FOutlinerCacheExtension::ID);
+	ThisSharedData->AddDynamicExtension(FMuteStateCacheExtension::ID);
+	ThisSharedData->AddDynamicExtension(FSoloStateCacheExtension::ID);
+	ThisSharedData->AddDynamicExtension(FLockStateCacheExtension::ID);
+	ThisSharedData->AddDynamicExtension(FConditionStateCacheExtension::ID);
 
 	// Add our hierarchical cache processor
-	TSharedPtr<FOutlinerCacheExtension> OutlinerCache = NewSharedData->CastThisSharedChecked<FOutlinerCacheExtension>();
+	TSharedPtr<FOutlinerCacheExtension> OutlinerCache = ThisSharedData->CastThisSharedChecked<FOutlinerCacheExtension>();
 	OutlinerCache->Initialize(AsShared());
 
 	// Make sure the hierarchical cache is updated when our hierarchy changes
-	FSimpleMulticastDelegate& HierarchyChanged = NewSharedData->SubscribeToHierarchyChanged(AsShared());
+	FSimpleMulticastDelegate& HierarchyChanged = ThisSharedData->SubscribeToHierarchyChanged(AsShared());
 	HierarchyChanged.AddSP(OutlinerCache.ToSharedRef(), &FOutlinerCacheExtension::OnHierarchyUpdated);
 
 	// This needs to be run outside of the constructor because a shared pointer to 'this' can't
@@ -101,7 +103,10 @@ void FSequenceModel::SetSequence(UMovieSceneSequence* InSequence, FMovieSceneSeq
 
 		// Remove everything and re-add the bottom spacer node
 		DiscardAllChildren();
-		GetChildrenForList(&RootOutlinerItems).AddChild(BottomSpacer);
+		if (BottomSpacer)
+		{
+			GetChildrenForList(&RootOutlinerItems).AddChild(BottomSpacer);
+		}
 
 		WeakSequence = InSequence;
 		SequenceID = InSequenceID;
diff --git a/Engine/Source/Editor/Sequencer/Private/MVVM/ViewModels/SequencerEditorViewModel.cpp b/Engine/Source/Editor/Sequencer/Private/MVVM/ViewModels/SequencerEditorViewModel.cpp
--- a/Engine/Source/Editor/Sequencer/Private/MVVM/ViewModels/SequencerEditorViewModel.cpp
+++ b/Engine/Source/Editor/Sequencer/Private/MVVM/ViewModels/SequencerEditorViewModel.cpp
@@ -53,9 +53,7 @@ void FSequencerEditorViewModel::PreInitializeEditorImpl()
 
 TSharedPtr<FViewModel> FSequencerEditorViewModel::CreateRootModelImpl()
 {
-	TSharedPtr<FSequenceModel> RootSequenceModel = MakeShared<FSequenceModel>(SharedThis(this));
-	RootSequenceModel->InitializeExtensions();
-	return RootSequenceModel;
+	return MakeShared<FSequenceModel>(SharedThis(this));
 }
 
 TSharedPtr<FOutlinerViewModel> FSequencerEditorViewModel::CreateOutlinerImpl()
diff --git a/Engine/Source/Editor/Sequencer/Private/SSequencer.cpp b/Engine/Source/Editor/Sequencer/Private/SSequencer.cpp
--- a/Engine/Source/Editor/Sequencer/Private/SSequencer.cpp
+++ b/Engine/Source/Editor/Sequencer/Private/SSequencer.cpp
@@ -104,7 +104,6 @@
 #include "EditorFontGlyphs.h"
 #include "HAL/PlatformApplicationMisc.h"
 #include "SSequencerPlayRateCombo.h"
-#include "SSequencerMusicModeCombo.h"
 #include "Camera/CameraActor.h"
 #include "SCurveEditorPanel.h"
 #include "Tree/SCurveEditorTree.h"
@@ -218,18 +217,6 @@ void SSequencer::Construct(const FArguments& InArgs, TSharedRef<FSequencer> InSe
 		}
 	);
 
-	TAttribute<FFrameRate> GetTickResolutionAttr = TAttribute<FFrameRate>(InSequencer, &FSequencer::GetFocusedTickResolution);
-	TAttribute<FFrameRate> GetDisplayRateAttr    = TAttribute<FFrameRate>(InSequencer, &FSequencer::GetFocusedDisplayRate);
-	TAttribute<const UFrameBasedMusicMap*> GetMusicMapAttr = TAttribute<const UFrameBasedMusicMap*>(InSequencer, &FSequencer::GetFocusedMusicMap);
-
-	// Create our numeric type interface so we can pass it to the time slider below.
-	TSharedPtr<FFrameNumberInterface> FrameNumberInterface = MakeShareable(new FFrameNumberInterface(GetDisplayFormatAttr, GetZeroPadFramesAttr, GetTickResolutionAttr, GetDisplayRateAttr, GetMusicMapAttr, FrameBasedMusicMap::EBarBeatFormat::Position));
-	NumericPositionTypeInterface = FrameNumberInterface;
-	FrameNumberInterface = MakeShareable(new FFrameNumberInterface(GetDisplayFormatAttr, GetZeroPadFramesAttr, GetTickResolutionAttr, GetDisplayRateAttr, GetMusicMapAttr, FrameBasedMusicMap::EBarBeatFormat::Duration));
-	NumericDurationTypeInterface = FrameNumberInterface;
-	FrameNumberInterface = MakeShareable(new FFrameNumberInterface(GetDisplayFormatAttr, GetZeroPadFramesAttr, GetTickResolutionAttr, GetDisplayRateAttr, GetMusicMapAttr, FrameBasedMusicMap::EBarBeatFormat::NotAllowed));
-	NumericNonMusicalTypeInterface = FrameNumberInterface;
-
 	FTimeSliderArgs TimeSliderArgs;
 	{
 		TimeSliderArgs.ViewRange = InArgs._ViewRange;
@@ -271,8 +258,6 @@ void SSequencer::Construct(const FArguments& InArgs, TSharedRef<FSequencer> InSe
 		TimeSliderArgs.OnDeleteAllMarkedFrames = InArgs._OnDeleteAllMarkedFrames;
 		TimeSliderArgs.AreMarkedFramesLocked = InArgs._AreMarkedFramesLocked;
 		TimeSliderArgs.OnToggleMarkedFramesLocked = InArgs._OnToggleMarkedFramesLocked;
-
-		TimeSliderArgs.NumericTypeInterface = GetNumericPositionTypeInterface();
 	}
 
 	OnGetPlaybackSpeeds = InArgs._OnGetPlaybackSpeeds;
@@ -280,7 +265,7 @@ void SSequencer::Construct(const FArguments& InArgs, TSharedRef<FSequencer> InSe
 	RootCustomization.AddMenuExtender = InArgs._AddMenuExtender;
 	RootCustomization.ToolbarExtender = InArgs._ToolbarExtender;
 
-	PlayTimeDisplay = StaticCastSharedRef<STemporarilyFocusedSpinBox<double>>(SequencerPtr.Pin()->MakePlayTimeDisplay(GetNumericPositionTypeInterface()));
+	PlayTimeDisplay = StaticCastSharedRef<STemporarilyFocusedSpinBox<double>>(SequencerPtr.Pin()->MakePlayTimeDisplay());
 
 	TAttribute<FAnimatedRange> ViewRangeAttribute = InArgs._ViewRange;
 
@@ -721,11 +706,22 @@ TSharedRef<SGridPanel> SSequencer::ConstructTrackAreaGridPanel(const FArguments&
 			Ranges,
 			TimeSliderControllerRef,
 			TAttribute<EVisibility>(this, &SSequencer::GetTimeRangeVisibility),
-			NumericPositionTypeInterface.ToSharedRef()
+			MakeAttributeSP(SequencerRef, &FSequencer::GetNumericTypeInterface, UE::Sequencer::ENumericIntent::Position)
 		),
 		SequencerWidgets.CreateTimeRangeSlider(TimeSliderControllerRef)
 	);
 
+	static const FName TransportControlsLeft("Sequencer.TransportControls.Left");
+	static const FName TransportControlsRight("Sequencer.TransportControls.Right");
+	if (!UToolMenus::Get()->IsMenuRegistered(TransportControlsLeft))
+	{
+		UToolMenus::Get()->RegisterMenu(TransportControlsLeft, NAME_None, EMultiBoxType::ToolBar);
+		UToolMenus::Get()->RegisterMenu(TransportControlsRight, NAME_None, EMultiBoxType::ToolBar);
+	}
+
+	USequencerToolMenuContext* ContextObject = NewObject<USequencerToolMenuContext>();
+	ContextObject->WeakSequencer = SequencerRef;
+
 	return SNew(SGridPanel)
 		.FillRow(2, 1.f)
 		.FillColumn(0, FillCoefficient_0)
@@ -882,13 +878,18 @@ TSharedRef<SGridPanel> SSequencer::ConstructTrackAreaGridPanel(const FArguments&
 				]
 
 				+ SHorizontalBox::Slot()
+				.AutoWidth()
+				.VAlign(VAlign_Center)
+				.HAlign(HAlign_Left)
+				.Padding(FMargin(CommonPadding, 0.f, 0.f, 0.f))
 				[
-					SNew(SBorder)
-					.BorderImage(FAppStyle::GetBrush(TEXT("ToolPanel.GroupBorder")))
-					.HAlign(HAlign_Center)
-					[
-						SequencerPtr.Pin()->MakeTransportControls(true)
-					]
+					UToolMenus::Get()->GenerateWidget(TransportControlsLeft, FToolMenuContext(SequencerRef->GetCommandBindings(), nullptr, ContextObject))
+				]
+
+				+ SHorizontalBox::Slot()
+				.HAlign(HAlign_Center)
+				[
+					SequencerPtr.Pin()->MakeTransportControls(true)
 				]
 
 				+ SHorizontalBox::Slot()
@@ -909,11 +910,7 @@ TSharedRef<SGridPanel> SSequencer::ConstructTrackAreaGridPanel(const FArguments&
 						.HAlign(HAlign_Right)
 						.Padding(FMargin(CommonPadding, 0.f, 0.f, 0.f))
 						[
-							SNew(SBorder)
-							.BorderImage(nullptr)
-							[
-								PlayTimeDisplay.ToSharedRef()
-							]
+							UToolMenus::Get()->GenerateWidget(TransportControlsRight, FToolMenuContext(SequencerRef->GetCommandBindings(), nullptr, ContextObject))
 						]
 
 						+ SHorizontalBox::Slot()
@@ -922,17 +919,22 @@ TSharedRef<SGridPanel> SSequencer::ConstructTrackAreaGridPanel(const FArguments&
 						.HAlign(HAlign_Right)
 						.Padding(FMargin(CommonPadding, 0.f, 0.f, 0.f))
 						[
-							SNew(SBorder)
-							.BorderImage(nullptr)
-							[
-								// Current loop index, if any
-								SAssignNew(LoopIndexDisplay, STextBlock)
-								.Font(FCoreStyle::GetDefaultFontStyle(TEXT("Regular"), 8))
-								.Text_Lambda([this]() -> FText {
-									TOptional<int32> LoopIndex = SequencerPtr.Pin()->GetLocalLoopIndex();
-									return LoopIndex ? FText::AsNumber(LoopIndex.GetValue()) : FText();
-								})
-							]
+							PlayTimeDisplay.ToSharedRef()
+						]
+
+						+ SHorizontalBox::Slot()
+						.AutoWidth()
+						.VAlign(VAlign_Center)
+						.HAlign(HAlign_Right)
+						.Padding(FMargin(CommonPadding, 0.f, 0.f, 0.f))
+						[
+							// Current loop index, if any
+							SAssignNew(LoopIndexDisplay, STextBlock)
+							.Font(FCoreStyle::GetDefaultFontStyle(TEXT("Regular"), 8))
+							.Text_Lambda([this]() -> FText {
+								TOptional<int32> LoopIndex = SequencerPtr.Pin()->GetLocalLoopIndex();
+								return LoopIndex ? FText::AsNumber(LoopIndex.GetValue()) : FText();
+							})
 						]
 					]
 				]
@@ -995,7 +997,7 @@ TSharedRef<SGridPanel> SSequencer::ConstructTrackAreaGridPanel(const FArguments&
 		.VAlign(VAlign_Top)
 		[
 			// Transform box
-			SAssignNew(TransformBox, SSequencerTransformBox, SequencerRef, *SequencerSettings, NumericPositionTypeInterface.ToSharedRef())
+			SAssignNew(TransformBox, SSequencerTransformBox, SequencerRef, *SequencerSettings, SequencerRef->GetNumericTypeInterface().ToSharedRef())
 		]
 
 		+ SGridPanel::Slot(Column1, Row2, SGridPanel::Layer(40))
@@ -1004,7 +1006,7 @@ TSharedRef<SGridPanel> SSequencer::ConstructTrackAreaGridPanel(const FArguments&
 		.VAlign(VAlign_Top)
 		[
 			// Stretch box
-			SAssignNew(StretchBox, SSequencerStretchBox, SequencerRef, *SequencerSettings, NumericPositionTypeInterface.ToSharedRef())
+			SAssignNew(StretchBox, SSequencerStretchBox, SequencerRef, *SequencerSettings, SequencerRef->GetNumericTypeInterface().ToSharedRef())
 		]
 
 		// debug vis
@@ -1164,21 +1166,6 @@ void SSequencer::NotifyPostChange(const FPropertyChangedEvent& PropertyChangedEv
 /* SSequencer implementation
  *****************************************************************************/
 
-TSharedRef<INumericTypeInterface<double>> SSequencer::GetNumericPositionTypeInterface() const
-{
-	return NumericPositionTypeInterface.ToSharedRef();
-}
-
-TSharedRef<INumericTypeInterface<double>> SSequencer::GetNumericDurationTypeInterface() const
-{
-	return NumericDurationTypeInterface.ToSharedRef();
-}
-
-TSharedRef<INumericTypeInterface<double>> SSequencer::GetNumericNonMusicalTypeInterface() const
-{
-	return NumericNonMusicalTypeInterface.ToSharedRef();
-}
-
 void SSequencer::UpdateOutlinerViewColumns()
 {
 	using namespace UE::Sequencer;
@@ -1440,7 +1427,7 @@ TSharedRef<SWidget> SSequencer::MakeToolBar()
 	TSharedPtr<FSequencer> Sequencer = SequencerPtr.Pin();
 
 	USequencerToolMenuContext* ContextObject = NewObject<USequencerToolMenuContext>();
-	ContextObject->SequencerWidget = SharedThis(this);
+	ContextObject->WeakSequencer = SequencerPtr;
 
 	FToolMenuContext Context(SequencerPtr.Pin()->GetCommandBindings(), Extender, ContextObject);
 
@@ -1460,30 +1447,27 @@ void SSequencer::PopulateToolBar(UToolMenu* InMenu)
 		return;
 	}
 
-	TSharedPtr<SSequencer> SequencerWidget = ContextObject->SequencerWidget.Pin();
-	if (!SequencerWidget.Get())
-	{
-		return;
-	}
-	
-	TWeakPtr<FSequencer> Sequencer = SequencerWidget.Get()->GetSequencer();
-	if (!Sequencer.IsValid())
+	TWeakPtr<FSequencer> WeakSequencer = StaticCastWeakPtr<FSequencer>(ContextObject->WeakSequencer);
+	TSharedPtr<FSequencer> Sequencer = WeakSequencer.Pin();
+	if (!Sequencer)
 	{
 		return;
 	}
 
+	TSharedRef<SSequencer> SequencerWidget = StaticCastSharedRef<SSequencer>(Sequencer->GetSequencerWidget());
+
 	{
 		FToolMenuSection& Section = InMenu->AddSection("BaseCommands");
 
-		if (Sequencer.Pin()->IsLevelEditorSequencer())
+		if (Sequencer->IsLevelEditorSequencer())
 		{
 			TAttribute<FSlateIcon> SaveIcon;
-			SaveIcon.Bind(TAttribute<FSlateIcon>::FGetter::CreateLambda([Sequencer] 
+			SaveIcon.Bind(TAttribute<FSlateIcon>::FGetter::CreateLambda([WeakSequencer] 
 			{
-				if (Sequencer.IsValid())
+				if (WeakSequencer.IsValid())
 				{
 					TArray<UMovieScene*> MovieScenesToSave;
-					MovieSceneHelpers::GetDescendantMovieScenes(Sequencer.Pin()->GetRootMovieSceneSequence(), MovieScenesToSave);
+					MovieSceneHelpers::GetDescendantMovieScenes(WeakSequencer.Pin()->GetRootMovieSceneSequence(), MovieScenesToSave);
 					for (UMovieScene* MovieSceneToSave : MovieScenesToSave)
 					{
 						UPackage* MovieScenePackageToSave = MovieSceneToSave->GetOuter()->GetOutermost();
@@ -1496,11 +1480,11 @@ void SSequencer::PopulateToolBar(UToolMenu* InMenu)
 				return FSlateIcon(FAppStyle::Get().GetStyleSetName(), "Icons.Save");
 			}));
 
-			if (Sequencer.Pin()->GetHostCapabilities().bSupportsSaveMovieSceneAsset)
+			if (Sequencer->GetHostCapabilities().bSupportsSaveMovieSceneAsset)
 			{
 				FToolMenuEntry SaveEntry = FToolMenuEntry::InitToolBarButton(
 					"Save",
-					FUIAction(FExecuteAction::CreateSP(SequencerWidget.Get(), &SSequencer::OnSaveMovieSceneClicked)),
+					FUIAction(FExecuteAction::CreateSP(SequencerWidget, &SSequencer::OnSaveMovieSceneClicked)),
 					LOCTEXT("SaveDirtyPackages", "Save"),
 					LOCTEXT("SaveDirtyPackagesTooltip", "Saves the current sequence and any subsequences"),
 					SaveIcon
@@ -1517,7 +1501,7 @@ void SSequencer::PopulateToolBar(UToolMenu* InMenu)
 			CreateCameraEntry.StyleNameOverride = SequencerToolbarStyleName;
 			Section.AddEntry(CreateCameraEntry);
 
-			if (Sequencer.Pin()->GetHostCapabilities().bSupportsRenderMovie)
+			if (Sequencer->GetHostCapabilities().bSupportsRenderMovie)
 			{
 				FToolMenuEntry RenderMovieEntry = FToolMenuEntry::InitToolBarButton(FSequencerCommands::Get().RenderMovie, TAttribute<FText>(), TAttribute<FText>(), FSlateIcon(FAppStyle::GetAppStyleSetName(), "LevelEditor.OpenCinematic"));
 				RenderMovieEntry.StyleNameOverride = SequencerToolbarStyleName;
@@ -1526,7 +1510,7 @@ void SSequencer::PopulateToolBar(UToolMenu* InMenu)
 				FToolMenuEntry RenderMovieOptionsEntry = FToolMenuEntry::InitComboButton(
 					"RenderMovieOptions",
 					FUIAction(),
-					FOnGetContent::CreateSP(SequencerWidget.Get(), &SSequencer::MakeRenderMovieMenu),
+					FOnGetContent::CreateSP(SequencerWidget, &SSequencer::MakeRenderMovieMenu),
 					LOCTEXT("RenderMovieOptions", "Render Movie Options"),
 					LOCTEXT("RenderMovieOptionsToolTip", "Render Movie Options"),
 					TAttribute<FSlateIcon>(),
@@ -1536,7 +1520,7 @@ void SSequencer::PopulateToolBar(UToolMenu* InMenu)
 				Section.AddEntry(RenderMovieOptionsEntry);
 			}
 
-			UMovieSceneSequence* RootSequence = Sequencer.Pin()->GetRootMovieSceneSequence();
+			UMovieSceneSequence* RootSequence = Sequencer->GetRootMovieSceneSequence();
 			if (RootSequence->GetTypedOuter<UBlueprint>() == nullptr && UMovieScene::IsTrackClassAllowed(UMovieSceneEventTrack::StaticClass()))
 			{
 				// Only show this button where it makes sense (ie, if the sequence is not contained within a blueprint already)
@@ -1551,7 +1535,7 @@ void SSequencer::PopulateToolBar(UToolMenu* InMenu)
 		FToolMenuEntry ActionsEntry = FToolMenuEntry::InitComboButton(
 			"Actions",
 			FUIAction(),
-			FOnGetContent::CreateSP(SequencerWidget.Get(), &SSequencer::MakeActionsMenu),
+			FOnGetContent::CreateSP(SequencerWidget, &SSequencer::MakeActionsMenu),
 			LOCTEXT("Actions", "Actions"),
 			LOCTEXT("ActionsToolTip", "Actions"),
 			FSlateIcon(FAppStyle::GetAppStyleSetName(), "Sequencer.Actions")
@@ -1562,7 +1546,7 @@ void SSequencer::PopulateToolBar(UToolMenu* InMenu)
 		FToolMenuEntry ViewOptionsEntry = FToolMenuEntry::InitComboButton(
 			"ViewOptions",
 			FUIAction(),
-			FOnGetContent::CreateSP(SequencerWidget.Get(), &SSequencer::MakeViewMenu),
+			FOnGetContent::CreateSP(SequencerWidget, &SSequencer::MakeViewMenu),
 			LOCTEXT("ViewOptions", "View Options"),
 			LOCTEXT("ViewOptionsToolTip", "View Options"),
 			FSlateIcon(FAppStyle::Get().GetStyleSetName(), "Icons.Visibility")
@@ -1573,7 +1557,7 @@ void SSequencer::PopulateToolBar(UToolMenu* InMenu)
 		FToolMenuEntry PlaybackOptionsEntry = FToolMenuEntry::InitComboButton(
 			"PlaybackOptions",
 			FUIAction(),
-			FOnGetContent::CreateSP(SequencerWidget.Get(), &SSequencer::MakePlaybackMenu),
+			FOnGetContent::CreateSP(SequencerWidget, &SSequencer::MakePlaybackMenu),
 			LOCTEXT("PlaybackOptions", "Playback Options"),
 			LOCTEXT("PlaybackOptionsToolTip", "Playback Options"),
 			FSlateIcon(FAppStyle::GetAppStyleSetName(), "Sequencer.PlaybackOptions")
@@ -1584,9 +1568,9 @@ void SSequencer::PopulateToolBar(UToolMenu* InMenu)
 		Section.AddSeparator(NAME_None);
 
 		TAttribute<FSlateIcon> KeyGroupModeIcon;
-		KeyGroupModeIcon.Bind(TAttribute<FSlateIcon>::FGetter::CreateLambda([Sequencer] 
+		KeyGroupModeIcon.Bind(TAttribute<FSlateIcon>::FGetter::CreateLambda([WeakSequencer] 
 		{
-			switch (Sequencer.Pin()->GetKeyGroupMode())
+			switch (WeakSequencer.Pin()->GetKeyGroupMode())
 			{
 			case EKeyGroupMode::KeyAll:
 				return FSequencerCommands::Get().SetKeyAll->GetIcon();
@@ -1598,9 +1582,9 @@ void SSequencer::PopulateToolBar(UToolMenu* InMenu)
 		}));
 
 		TAttribute<FText> KeyGroupModeToolTip;
-		KeyGroupModeToolTip.Bind(TAttribute<FText>::FGetter::CreateLambda([Sequencer] 
+		KeyGroupModeToolTip.Bind(TAttribute<FText>::FGetter::CreateLambda([WeakSequencer] 
 		{
-			switch (Sequencer.Pin()->GetKeyGroupMode())
+			switch (WeakSequencer.Pin()->GetKeyGroupMode())
 			{
 			case EKeyGroupMode::KeyAll:
 				return FSequencerCommands::Get().SetKeyAll->GetDescription();
@@ -1614,7 +1598,7 @@ void SSequencer::PopulateToolBar(UToolMenu* InMenu)
 		FToolMenuEntry KeyGroupEntry = FToolMenuEntry::InitComboButton(
 			"KeyGroup",
 			FUIAction(),
-			FOnGetContent::CreateSP(SequencerWidget.Get(), &SSequencer::MakeKeyGroupMenu),
+			FOnGetContent::CreateSP(SequencerWidget, &SSequencer::MakeKeyGroupMenu),
 			LOCTEXT("KeyGroup", "Key All"),
 			KeyGroupModeToolTip,
 			KeyGroupModeIcon
@@ -1622,12 +1606,12 @@ void SSequencer::PopulateToolBar(UToolMenu* InMenu)
 		KeyGroupEntry.StyleNameOverride = SequencerToolbarStyleName;
 		Section.AddEntry(KeyGroupEntry);
 	
-		if (IVREditorModule::Get().IsVREditorModeActive() || (Sequencer.Pin()->IsLevelEditorSequencer() && ExactCast<ULevelSequence>(Sequencer.Pin()->GetFocusedMovieSceneSequence()) == nullptr))
+		if (IVREditorModule::Get().IsVREditorModeActive() || (Sequencer->IsLevelEditorSequencer() && ExactCast<ULevelSequence>(Sequencer->GetFocusedMovieSceneSequence()) == nullptr))
 		{
 			TAttribute<FSlateIcon> AutoChangeModeIcon;
-			AutoChangeModeIcon.Bind(TAttribute<FSlateIcon>::FGetter::CreateLambda([Sequencer] 
+			AutoChangeModeIcon.Bind(TAttribute<FSlateIcon>::FGetter::CreateLambda([WeakSequencer] 
 			{
-				switch ( Sequencer.Pin()->GetAutoChangeMode() )
+				switch ( WeakSequencer.Pin()->GetAutoChangeMode() )
 				{
 				case EAutoChangeMode::AutoKey:
 					return FSequencerCommands::Get().SetAutoKey->GetIcon();
@@ -1641,9 +1625,9 @@ void SSequencer::PopulateToolBar(UToolMenu* InMenu)
 			} ) );
 
 			TAttribute<FText> AutoChangeModeToolTip;
-			AutoChangeModeToolTip.Bind( TAttribute<FText>::FGetter::CreateLambda([Sequencer] 
+			AutoChangeModeToolTip.Bind( TAttribute<FText>::FGetter::CreateLambda([WeakSequencer] 
 			{
-				switch ( Sequencer.Pin()->GetAutoChangeMode() )
+				switch ( WeakSequencer.Pin()->GetAutoChangeMode() )
 				{
 				case EAutoChangeMode::AutoKey:
 					return FSequencerCommands::Get().SetAutoKey->GetDescription();
@@ -1659,7 +1643,7 @@ void SSequencer::PopulateToolBar(UToolMenu* InMenu)
 			FToolMenuEntry AutoChangeEntry = FToolMenuEntry::InitComboButton(
 				"AutoChange",
 				FUIAction(),
-				FOnGetContent::CreateSP(SequencerWidget.Get(), &SSequencer::MakeAutoChangeMenu),
+				FOnGetContent::CreateSP(SequencerWidget, &SSequencer::MakeAutoChangeMenu),
 				LOCTEXT("AutoChangeMode", "Auto-Change Mode"),
 				AutoChangeModeToolTip,
 				AutoChangeModeIcon
@@ -1674,12 +1658,12 @@ void SSequencer::PopulateToolBar(UToolMenu* InMenu)
 			Section.AddEntry(ToggleAutoKeyEntry);
 		}
 
-		if (Sequencer.Pin()->IsLevelEditorSequencer())
+		if (Sequencer->IsLevelEditorSequencer())
 		{
 			TAttribute<FSlateIcon> AllowEditsModeIcon;
-			AllowEditsModeIcon.Bind(TAttribute<FSlateIcon>::FGetter::CreateLambda([Sequencer] 
+			AllowEditsModeIcon.Bind(TAttribute<FSlateIcon>::FGetter::CreateLambda([WeakSequencer] 
 			{
-				switch ( Sequencer.Pin()->GetAllowEditsMode() )
+				switch ( WeakSequencer.Pin()->GetAllowEditsMode() )
 				{
 				case EAllowEditsMode::AllEdits:
 					return FSequencerCommands::Get().AllowAllEdits->GetIcon();
@@ -1691,9 +1675,9 @@ void SSequencer::PopulateToolBar(UToolMenu* InMenu)
 			} ) );
 
 			TAttribute<FText> AllowEditsModeToolTip;
-			AllowEditsModeToolTip.Bind( TAttribute<FText>::FGetter::CreateLambda([Sequencer] 
+			AllowEditsModeToolTip.Bind( TAttribute<FText>::FGetter::CreateLambda([WeakSequencer] 
 			{
-				switch ( Sequencer.Pin()->GetAllowEditsMode() )
+				switch ( WeakSequencer.Pin()->GetAllowEditsMode() )
 				{
 				case EAllowEditsMode::AllEdits:
 					return FSequencerCommands::Get().AllowAllEdits->GetDescription();
@@ -1707,7 +1691,7 @@ void SSequencer::PopulateToolBar(UToolMenu* InMenu)
 			FToolMenuEntry AllowEditsModeEntry = FToolMenuEntry::InitComboButton(
 				"AllowEditsMode",
 				FUIAction(),
-				FOnGetContent::CreateSP(SequencerWidget.Get(), &SSequencer::MakeAllowEditsMenu),
+				FOnGetContent::CreateSP(SequencerWidget, &SSequencer::MakeAllowEditsMenu),
 				LOCTEXT("AllowEditsMode", "Allow Edits"),
 				AllowEditsModeToolTip,
 				AllowEditsModeIcon
@@ -1720,7 +1704,7 @@ void SSequencer::PopulateToolBar(UToolMenu* InMenu)
 	}
 
 	{
-		if (Sequencer.Pin()->GetHostCapabilities().bSupportsViewportSelectability)
+		if (Sequencer->GetHostCapabilities().bSupportsViewportSelectability)
 		{
 			FToolMenuSection& SelectionSection = InMenu->AddSection(TEXT("Selection"));
 
@@ -1743,7 +1727,7 @@ void SSequencer::PopulateToolBar(UToolMenu* InMenu)
 		FToolMenuEntry SnapOptionsEntry = FToolMenuEntry::InitComboButton(
 			"SnapOptions",
 			FUIAction(),
-			FOnGetContent::CreateSP(SequencerWidget.Get(), &SSequencer::MakeSnapMenu),
+			FOnGetContent::CreateSP(SequencerWidget, &SSequencer::MakeSnapMenu),
 			LOCTEXT("SnapOptions", "Options"),
 			LOCTEXT("SnapOptionsToolTip", "Snapping Options"),
 			TAttribute<FSlateIcon>(),
@@ -1752,27 +1736,13 @@ void SSequencer::PopulateToolBar(UToolMenu* InMenu)
 		SnapOptionsEntry.StyleNameOverride = SequencerToolbarStyleName;
 		SnappingSection.AddEntry(SnapOptionsEntry);
 
-		SnappingSection.AddSeparator(NAME_None);
-
-		static const auto MusicModeEnabledCVar = IConsoleManager::Get().FindConsoleVariable(TEXT("Sequencer.EnableMusicModeEditing"));
-		if (MusicModeEnabledCVar->GetInt() != 0)
-		{
-			FToolMenuEntry MusicModeEntry = FToolMenuEntry::InitWidget(
-				"MusicMode",
-				SNew(SSequencerMusicModeCombo, Sequencer.Pin(), SequencerWidget.ToSharedRef()),
-				LOCTEXT("MusicMode", "MusicMode")
-			);
-			MusicModeEntry.StyleNameOverride = SequencerToolbarStyleName;
-			SnappingSection.AddEntry(MusicModeEntry);
-
-			SnappingSection.AddSeparator(NAME_None);
-		}
+		SnappingSection.AddSeparator("PlayRate");
 
 		FToolMenuEntry PlayRateEntry = FToolMenuEntry::InitWidget(
 			"PlayRate",
-			SNew(SSequencerPlayRateCombo, Sequencer.Pin(), SequencerWidget.ToSharedRef()),
-			LOCTEXT("PlayRate", "PlayRate")
-		);
+			SNew(SSequencerPlayRateCombo, Sequencer, SequencerWidget)
+				.Visibility(Sequencer.Get(), &FSequencer::GetPlayRateComboVisibility),
+		   LOCTEXT("PlayRate", "PlayRate"));
 		PlayRateEntry.StyleNameOverride = SequencerToolbarStyleName;
 		SnappingSection.AddEntry(PlayRateEntry);
 	}
@@ -1781,7 +1751,7 @@ void SSequencer::PopulateToolBar(UToolMenu* InMenu)
 		FToolMenuSection& CurveEditorSection = InMenu->AddSection("CurveEditor");
 
 		// Only add the button if supported
-		if (Sequencer.Pin()->GetHostCapabilities().bSupportsCurveEditor)
+		if (Sequencer->GetHostCapabilities().bSupportsCurveEditor)
 		{
 			FToolMenuEntry ShowCurveEditorEntry = FToolMenuEntry::InitToolBarButton(FSequencerCommands::Get().ToggleShowCurveEditor, TAttribute<FText>(), TAttribute<FText>(), FSlateIcon(FAppStyle::GetAppStyleSetName(), "GenericCurveEditor"));
 			ShowCurveEditorEntry.StyleNameOverride = SequencerToolbarStyleName;
@@ -2308,19 +2278,27 @@ void SSequencer::FillColumnVisibilityMenu(FMenuBuilder& InMenuBuilder)
 	}
 }
 
-void SSequencer::FillTimeDisplayFormatMenu(FMenuBuilder& MenuBuilder)
+void SSequencer::FillTimeDisplayFormatMenu(UToolMenu* InMenu)
 {
-	TSharedPtr<FSequencer> Sequencer = SequencerPtr.Pin();
+	USequencerToolMenuContext* ContextObject = InMenu->FindContext<USequencerToolMenuContext>();
+	if (!ContextObject)
+	{
+		return;
+	}
+
+	TWeakPtr<FSequencer> WeakSequencer = StaticCastWeakPtr<FSequencer>(ContextObject->WeakSequencer);
+	TSharedPtr<FSequencer> Sequencer = WeakSequencer.Pin();
+	if (!Sequencer)
+	{
+		return;
+	}
+
 	bool bShouldDisplayDropFormat = FTimecode::UseDropFormatTimecode(Sequencer->GetFocusedDisplayRate());
-	
-	const UFrameBasedMusicMap* MusicMap = Sequencer->GetFocusedMusicMap();
-	static const auto MusicModeEnabledCVar = IConsoleManager::Get().FindConsoleVariable(TEXT("Sequencer.EnableMusicModeEditing"));
-	bool bShouldDisplayBarsBeats = MusicMap && !MusicMap->IsEmpty() && MusicModeEnabledCVar->GetInt() != 0;
 
 	const UEnum* FrameNumberDisplayEnum = StaticEnum<EFrameNumberDisplayFormats>();
 	check(FrameNumberDisplayEnum);
 
-	if (GetSequencerSettings())
+	if (USequencerSettings* Settings = Sequencer->GetSequencerSettings())
 	{
 		for (int32 Index = 0; Index < FrameNumberDisplayEnum->NumEnums() - 1; Index++)
 		{
@@ -2340,24 +2318,21 @@ void SSequencer::FillTimeDisplayFormatMenu(FMenuBuilder& MenuBuilder)
 					continue;
 				}
 
-				// Don't show Bars.Beats if the sequence doesn't have a music map.
-				if (Value == EFrameNumberDisplayFormats::BarsBeats && !bShouldDisplayBarsBeats)
-				{
-					continue;
-				}
-
-				MenuBuilder.AddMenuEntry(
+				FToolMenuEntry Entry = FToolMenuEntry::InitMenuEntry(
+					FrameNumberDisplayEnum->GetNameByIndex(Index),
 					FrameNumberDisplayEnum->GetDisplayNameTextByIndex(Index),
 					FrameNumberDisplayEnum->GetToolTipTextByIndex(Index),
 					FSlateIcon(),
-					FUIAction(
-						FExecuteAction::CreateUObject(GetSequencerSettings(), &USequencerSettings::SetTimeDisplayFormat, Value),
-						FCanExecuteAction(),
-						FIsActionChecked::CreateLambda([this, Value] { return GetSequencerSettings()->GetTimeDisplayFormat() == Value; })
+					FToolUIActionChoice(
+						FUIAction(
+							FExecuteAction::CreateUObject(Settings, &USequencerSettings::SetTimeDisplayFormat, Value),
+							FCanExecuteAction(),
+							FIsActionChecked::CreateLambda([Settings, Value] { return Settings->GetTimeDisplayFormat() == Value; })
+						)
 					),
-					NAME_None,
 					EUserInterfaceActionType::RadioButton
 				);
+				InMenu->AddMenuEntry(NAME_None, Entry);
 			}
 		}
 	}
@@ -2412,7 +2387,7 @@ TSharedRef<SWidget> SSequencer::MakePlaybackMenu()
 				.AutoWidth()
 				[
 					SNew(SSpinBox<double>)
-						.TypeInterface(NumericPositionTypeInterface)
+						.TypeInterface(SequencerPtr.Pin().Get(), &FSequencer::GetNumericTypeInterface, UE::Sequencer::ENumericIntent::Position)
 						.IsEnabled_Lambda([this]() {
 							return !SequencerPtr.Pin()->IsPlaybackRangeLocked();
 						})
@@ -2473,7 +2448,7 @@ TSharedRef<SWidget> SSequencer::MakePlaybackMenu()
 				.AutoWidth()
 				[
 					SNew(SSpinBox<double>)
-						.TypeInterface(NumericPositionTypeInterface)
+						.TypeInterface(SequencerPtr.Pin().Get(), &FSequencer::GetNumericTypeInterface, UE::Sequencer::ENumericIntent::Position)
 						.IsEnabled_Lambda([this]() {
 					 		return !SequencerPtr.Pin()->IsPlaybackRangeLocked();
 						})
@@ -2545,7 +2520,7 @@ TSharedRef<SWidget> SSequencer::MakePlaybackMenu()
 			.AutoWidth()
 			[
 				SNew(SSpinBox<double>)
-					.TypeInterface(NumericDurationTypeInterface)
+					.TypeInterface(SequencerPtr.Pin().Get(), &FSequencer::GetNumericTypeInterface, UE::Sequencer::ENumericIntent::Duration)
 					.Style(&FAppStyle::GetWidgetStyle<FSpinBoxStyle>("Sequencer.HyperlinkSpinBox"))
 					.OnValueCommitted_Lambda([=](double Value, ETextCommit::Type){ OnJumpFrameIncrementChanged(Value); })
 					.OnValueChanged_Lambda([=](double Value) { OnJumpFrameIncrementChanged(Value); })
@@ -2807,7 +2782,7 @@ TSharedRef<SWidget> SSequencer::MakeTimeRange(const TSharedRef<SWidget>& InnerCo
 		ShowRange,
 		TimeSliderController.ToSharedRef(),
 		EVisibility::Visible,
-		NumericPositionTypeInterface.ToSharedRef()
+		MakeAttributeSP(SequencerPtr.Pin().Get(), &FSequencer::GetNumericTypeInterface, UE::Sequencer::ENumericIntent::Position)
 		);
 	return SequencerWidgets.CreateTimeRange(Args, InnerContent);
 }
diff --git a/Engine/Source/Editor/Sequencer/Private/SSequencer.h b/Engine/Source/Editor/Sequencer/Private/SSequencer.h
--- a/Engine/Source/Editor/Sequencer/Private/SSequencer.h
+++ b/Engine/Source/Editor/Sequencer/Private/SSequencer.h
@@ -336,15 +336,6 @@ public:
 	/** Access this widget's track area widget */
 	TSharedPtr<UE::Sequencer::STrackAreaView> GetTrackAreaWidget() const { return TrackArea; }
 
-	/** @return a numeric type interface that will parse and display numbers as frame / time / bar.beat positions correctly */
-	TSharedRef<INumericTypeInterface<double>> GetNumericPositionTypeInterface() const;
-
-	/** @return a numeric type interface that will parse and display numbers as frame / time / bar.beat durations correctly */
-	TSharedRef<INumericTypeInterface<double>> GetNumericDurationTypeInterface() const;
-
-	/** @return a numeric type interface that will parse and display numbers as frame / time durations correctly. Bars.Beats will be forced to time. */
-	TSharedRef<INumericTypeInterface<double>> GetNumericNonMusicalTypeInterface() const;
-
 	void OpenTickResolutionOptions();
 
 	/** Sets the play time for the sequence but clamped by the working range. This is useful for cases where we can't clamp via the UI control. */
@@ -459,7 +450,7 @@ private:
 
 public:
 	/** Makes the time display format menu for the toolbar and the play rate menu. */
-	void FillTimeDisplayFormatMenu(FMenuBuilder& MenuBuilder);
+	static void FillTimeDisplayFormatMenu(UToolMenu* InMenu);
 
 	void OpenNodeGroupsManager();
 
@@ -821,15 +812,6 @@ private:
 	/** Extender to use for the view menu */
 	TArray<TSharedPtr<FExtender>> ViewMenuExtenders;
 
-	/** Numeric type interface used for converting parsing and generating strings from numbers */
-	TSharedPtr<INumericTypeInterface<double>> NumericPositionTypeInterface;
-
-	/** Numeric type interface used for converting parsing and generating strings from numbers */
-	TSharedPtr<INumericTypeInterface<double>> NumericDurationTypeInterface;
-
-	/** Numeric type interface used for converting parsing and generating strings from numbers */
-	TSharedPtr<INumericTypeInterface<double>> NumericNonMusicalTypeInterface;
-
 	/** Time slider controller for this sequencer */
 	TSharedPtr<FSequencerTimeSliderController> TimeSliderController;
 
diff --git a/Engine/Source/Editor/Sequencer/Private/SSequencerPlayRateCombo.cpp b/Engine/Source/Editor/Sequencer/Private/SSequencerPlayRateCombo.cpp
--- a/Engine/Source/Editor/Sequencer/Private/SSequencerPlayRateCombo.cpp
+++ b/Engine/Source/Editor/Sequencer/Private/SSequencerPlayRateCombo.cpp
@@ -11,6 +11,7 @@
 #include "Framework/MultiBox/MultiBoxBuilder.h"
 #include "MovieSceneSequence.h"
 #include "MovieScene.h"
+#include "MovieSceneClock.h"
 #include "MovieSceneCommonHelpers.h"
 #include "Compilation/MovieSceneCompiledDataManager.h"
 #include "ScopedTransaction.h"
@@ -24,6 +25,7 @@
 #include "Misc/Timecode.h"
 #include "Styling/ToolBarStyle.h"
 #include "ActorTreeItem.h"
+#include "SequencerToolMenuContext.h"
 
 #define LOCTEXT_NAMESPACE "SSequencerPlayRateCombo"
 
@@ -303,7 +305,16 @@ TSharedRef<SWidget> SSequencerPlayRateCombo::OnCreateMenu()
 		FNewMenuDelegate::CreateLambda([this](FMenuBuilder& InMenuBuilder) {
 			if (WeakSequencerWidget.IsValid())
 			{
-				WeakSequencerWidget.Pin()->FillTimeDisplayFormatMenu(InMenuBuilder);
+				if (!UToolMenus::Get()->IsMenuRegistered("Sequencer.TimeDisplayMenu"))
+				{
+					UToolMenu* Toolbar = UToolMenus::Get()->RegisterMenu("Sequencer.TimeDisplayMenu", NAME_None, EMultiBoxType::Menu);
+					Toolbar->AddDynamicSection("DefaultFormats", FNewToolMenuDelegate::CreateStatic(&SSequencer::FillTimeDisplayFormatMenu));
+				}
+
+				USequencerToolMenuContext* ContextObject = NewObject<USequencerToolMenuContext>();
+				ContextObject->WeakSequencer = WeakSequencer;
+
+				InMenuBuilder.AddWidget(UToolMenus::Get()->GenerateWidget("Sequencer.TimeDisplayMenu", FToolMenuContext(ContextObject)), FText(), true, false);
 			}
 		})
 	);
@@ -371,56 +382,73 @@ void SSequencerPlayRateCombo::PopulateIncompatibleRatesMenu(FMenuBuilder& MenuBu
 
 void SSequencerPlayRateCombo::PopulateClockSourceMenu(FMenuBuilder& MenuBuilder)
 {
-	TSharedPtr<FSequencer> Sequencer    = WeakSequencer.Pin();
-	UMovieSceneSequence*   RootSequence = Sequencer.IsValid() ? Sequencer->GetRootMovieSceneSequence() : nullptr;
-
-	const UEnum* ClockSourceEnum = StaticEnum<EUpdateClockSource>();
-
-	check(ClockSourceEnum);
-
-	if (RootSequence)
+	static const FName MenuName("Sequencer.PlayRateCombo.ClockSource");
+	if (!UToolMenus::Get()->IsMenuRegistered(MenuName))
 	{
-		const bool IsFocusedOnRootSequence = Sequencer->GetRootMovieSceneSequence() == Sequencer->GetFocusedMovieSceneSequence();
+		UToolMenu* ClockSourceMenu = UToolMenus::Get()->RegisterMenu(MenuName, NAME_None, EMultiBoxType::Menu);
+		
+		ClockSourceMenu->AddDynamicSection("StandardClocks", FNewToolMenuDelegate::CreateStatic([](UToolMenu* InMenu){
+			USequencerClockSourceMenuContext* ContextObject = InMenu->FindContext<USequencerClockSourceMenuContext>();
+			TSharedPtr<SSequencerPlayRateCombo> Combo = ContextObject ? ContextObject->WeakSequencerCombo.Pin() : nullptr;
+			TSharedPtr<ISequencer>     Sequencer     = Combo ? Combo->WeakSequencer.Pin() : nullptr;
+			UMovieSceneSequence*	   RootSequence  = Sequencer.IsValid() ? Sequencer->GetRootMovieSceneSequence() : nullptr;
 
-		for (int32 Index = 0; Index < ClockSourceEnum->NumEnums() - 1; Index++)
-		{
-			if (!ClockSourceEnum->HasMetaData(TEXT("Hidden"), Index))
+			const UEnum* ClockSourceEnum = StaticEnum<EUpdateClockSource>();
+
+			check(ClockSourceEnum);
+			if (RootSequence)
 			{
-				EUpdateClockSource Value = (EUpdateClockSource)ClockSourceEnum->GetValueByIndex(Index);
+				const bool IsFocusedOnRootSequence = Sequencer->GetRootMovieSceneSequence() == Sequencer->GetFocusedMovieSceneSequence();
 
-				if (Value == EUpdateClockSource::Custom)
-				{
-					MenuBuilder.AddSubMenu(
-						ClockSourceEnum->GetDisplayNameTextByIndex(Index),
-						ClockSourceEnum->GetToolTipTextByIndex(Index),
-						FNewMenuDelegate::CreateSP(this, &SSequencerPlayRateCombo::PopulateCustomClockSourceMenu),
-						FUIAction(
-							FExecuteAction::CreateSP(this, &SSequencerPlayRateCombo::SetClockSource, Value),
-							FCanExecuteAction::CreateLambda([this, IsFocusedOnRootSequence]{ return !GetIsSequenceReadOnly() && IsFocusedOnRootSequence; }),
-							FIsActionChecked::CreateLambda([=]{ return RootSequence->GetMovieScene()->GetClockSource() == Value; })
-						),
-						NAME_None,
-						EUserInterfaceActionType::RadioButton
-					);
-				}
-				else
+				for (int32 Index = 0; Index < ClockSourceEnum->NumEnums() - 1; Index++)
 				{
-					MenuBuilder.AddMenuEntry(
-						ClockSourceEnum->GetDisplayNameTextByIndex(Index),
-						ClockSourceEnum->GetToolTipTextByIndex(Index),
-						FSlateIcon(),
-						FUIAction(
-							FExecuteAction::CreateSP(this, &SSequencerPlayRateCombo::SetClockSource, Value),
-							FCanExecuteAction::CreateLambda([this, IsFocusedOnRootSequence]{ return !GetIsSequenceReadOnly() && IsFocusedOnRootSequence; }),
-							FIsActionChecked::CreateLambda([=]{ return RootSequence->GetMovieScene()->GetClockSource() == Value; })
-						),
-						NAME_None,
-						EUserInterfaceActionType::RadioButton
-					);
+					if (!ClockSourceEnum->HasMetaData(TEXT("Hidden"), Index))
+					{
+						EUpdateClockSource Value = (EUpdateClockSource)ClockSourceEnum->GetValueByIndex(Index);
+
+						if (Value == EUpdateClockSource::Custom)
+						{
+							FToolMenuEntry Item = FToolMenuEntry::InitSubMenu(NAME_None,
+								ClockSourceEnum->GetDisplayNameTextByIndex(Index),
+								ClockSourceEnum->GetToolTipTextByIndex(Index),
+								FNewToolMenuChoice(FNewMenuDelegate::CreateSP(Combo.Get(), &SSequencerPlayRateCombo::PopulateCustomClockSourceMenu)),
+								FToolUIActionChoice(FUIAction(
+									FExecuteAction::CreateSP(Combo.Get(), &SSequencerPlayRateCombo::SetClockSource, Value),
+									FCanExecuteAction::CreateLambda([Combo, IsFocusedOnRootSequence]{ return !Combo->GetIsSequenceReadOnly() && IsFocusedOnRootSequence; }),
+									FIsActionChecked::CreateLambda([=]{ return RootSequence->GetMovieScene()->GetClockSource() == Value; })
+								)),
+								EUserInterfaceActionType::RadioButton);
+							
+							InMenu->AddMenuEntry("StandardClocks", Item);
+						}
+						else
+						{
+							FToolMenuEntry Item = FToolMenuEntry::InitMenuEntry(
+								NAME_None,
+								ClockSourceEnum->GetDisplayNameTextByIndex(Index),
+								ClockSourceEnum->GetToolTipTextByIndex(Index),
+								FSlateIcon(),
+								FToolUIActionChoice(FUIAction(
+									FExecuteAction::CreateSP(Combo.Get(), &SSequencerPlayRateCombo::SetClockSource, Value),
+									FCanExecuteAction::CreateLambda([Combo, IsFocusedOnRootSequence]{ return !Combo->GetIsSequenceReadOnly() && IsFocusedOnRootSequence; }),
+									FIsActionChecked::CreateLambda([=]{ return RootSequence->GetMovieScene()->GetClockSource() == Value; })
+									)),
+								EUserInterfaceActionType::RadioButton
+								);
+							InMenu->AddMenuEntry("StandardClocks", Item);
+						}
+					}
 				}
+				FToolMenuEntry Sep = FToolMenuEntry::InitSeparator("ExtensionClockTypes");
+				InMenu->AddMenuEntry("StandardClocks", Sep);
 			}
-		}
+		}));
 	}
+	USequencerClockSourceMenuContext* ContextObject = NewObject<USequencerClockSourceMenuContext>();
+	ContextObject->WeakSequencer = WeakSequencer;
+	ContextObject->WeakSequencerCombo = StaticCastWeakPtr<SSequencerPlayRateCombo>(AsWeak()); 
+	FToolMenuContext MenuContext(ContextObject);
+	MenuBuilder.AddWidget(UToolMenus::Get()->GenerateWidget(MenuName, MenuContext), FText(), true, false);
 }
 
 void SSequencerPlayRateCombo::PopulateCustomClockSourceMenu(FMenuBuilder& MenuBuilder)
@@ -562,7 +590,16 @@ void SSequencerPlayRateCombo::SetCustomClockSource(UObject* InClockSource)
 		FScopedTransaction ScopedTransaction(LOCTEXT("SetClockSource", "Set Clock Source"));
 
 		MovieScene->Modify();
-		MovieScene->SetClockSource(InClockSource);
+
+		UMovieSceneExternalClock* External = MovieScene->GetCustomClockAs<UMovieSceneExternalClock>();
+		if (External == nullptr)
+		{
+			External = NewObject<UMovieSceneExternalClock>(MovieScene);
+			MovieScene->SetCustomClock(External);
+		}
+
+		External->Modify();
+		External->CustomClockSourcePath = InClockSource;
 
 		Sequencer->ResetTimeController();
 	}
