- Movied previously committed squencer changes for music mode into a new Musical Mode plugin

diff --git a/Engine/Source/Editor/Sequencer/Private/Sequencer.cpp b/Engine/Source/Editor/Sequencer/Private/Sequencer.cpp
--- a/Engine/Source/Editor/Sequencer/Private/Sequencer.cpp
+++ b/Engine/Source/Editor/Sequencer/Private/Sequencer.cpp
@@ -7,6 +7,7 @@
 #include "Algo/RemoveIf.h"
 #include "MVVM/SharedViewModelData.h"
 #include "MVVM/Selection/Selection.h"
+#include "MVVM/Extensions/IClockExtension.h"
 #include "MVVM/ViewModels/SequenceModel.h"
 #include "MVVM/ViewModels/SectionModel.h"
 #include "MVVM/ViewModels/FolderModel.h"
@@ -72,6 +73,7 @@
 #include "MovieSceneClipboard.h"
 #include "SequencerCommonHelpers.h"
 #include "SequencerMarkedFrameHelper.h"
+#include "ISequencerNumericTypeInterface.h"
 #include "SSequencer.h"
 #include "SSequencerSection.h"
 #include "SequencerKeyCollection.h"
@@ -84,6 +86,7 @@
 #include "MovieSceneSequence.h"
 #include "MovieSceneFolder.h"
 #include "PropertyEditorModule.h"
+#include "PropertyHandle.h"
 #include "EditorWidgetsModule.h"
 #include "IAssetViewport.h"
 #include "EditorSupportDelegates.h"
@@ -141,6 +144,7 @@
 #include "Fonts/FontMeasure.h"
 #include "MovieSceneTimeHelpers.h"
 #include "FrameNumberNumericInterface.h"
+#include "FrameNumberDetailsCustomization.h"
 #include "UObject/StrongObjectPtr.h"
 #include "LevelUtils.h"
 #include "Engine/Blueprint.h"
@@ -168,7 +172,8 @@
 #include "Sidebar/SidebarDrawerConfig.h"
 #include "Widgets/Sidebar/SequencerSelectionDrawer.h"
 #include "Misc/Thumbnail/ThumbnailCaptureUtils.h"
-#include "FrameBasedMusicMap.h"
+#include "ToolMenus.h"
+#include "SequencerToolMenuContext.h"
 
 #include "EngineModule.h"
 #include "IViewportSelectableObject.h"
@@ -296,6 +301,22 @@ namespace UE
 			TOptional<UE::MovieScene::FScopedSignedObjectModifyDefer> DeferImplicitChanges;
 		};
 
+
+
+		struct FPositionNumericTypeInterface : UE::Sequencer::FSequencerNumericTypeInterface
+		{
+			FPositionNumericTypeInterface(TSharedRef<FFrameNumberInterface> InInterface)
+				: FSequencerNumericTypeInterface(InInterface, ENumericIntent::Position)
+			{}
+		};
+
+		struct FDurationNumericTypeInterface : UE::Sequencer::FSequencerNumericTypeInterface
+		{
+			FDurationNumericTypeInterface(TSharedRef<FFrameNumberInterface> InInterface)
+				: FSequencerNumericTypeInterface(InInterface, ENumericIntent::Duration)
+			{}
+		};
+
 	} // namespace Sequencer
 } // namespace UE
 
@@ -360,6 +381,41 @@ void FSequencer::InitSequencer(const FSequencerInitParams& InitParams
 	Settings->GetOnShowSelectedNodesOnlyChanged().AddSP(this, &FSequencer::OnSelectedNodesOnlyChanged);
 	Settings->GetOnTimeDisplayFormatChanged().AddSP(this, &FSequencer::OnTimeDisplayFormatChanged);
 
+
+
+
+
+	// Initialize numeric type interfaces
+	{
+
+		TAttribute<EFrameNumberDisplayFormats> GetDisplayFormatAttr = MakeAttributeLambda(
+			[Settings = this->Settings]
+			{
+				return Settings->GetTimeDisplayFormat();
+			}
+		);
+		TAttribute<uint8> GetZeroPadFramesAttr = MakeAttributeLambda(
+			[Settings = this->Settings]()->uint8
+			{
+				return Settings->GetZeroPadFrames();
+			}
+		);
+
+		TAttribute<FFrameRate> GetTickResolutionAttr = MakeAttributeSP(this, &FSequencer::GetFocusedTickResolution);
+		TAttribute<FFrameRate> GetDisplayRateAttr    = MakeAttributeSP(this, &FSequencer::GetFocusedDisplayRate);
+
+		NumericTypeInterfaces.Add(MakeShared<FPositionNumericTypeInterface>(
+			MakeShared<FFrameNumberInterface>(GetDisplayFormatAttr, GetZeroPadFramesAttr, GetTickResolutionAttr, GetDisplayRateAttr)
+			)
+		);
+
+		NumericTypeInterfaces.Add(MakeShared<FDurationNumericTypeInterface>(
+			MakeShared<FFrameNumberInterface>(GetDisplayFormatAttr, GetZeroPadFramesAttr, GetTickResolutionAttr, GetDisplayRateAttr)
+			)
+		);
+	}
+
+
 	ObjectBindingTagCache = MakeUnique<FObjectBindingTagCache>();
 
 	{
@@ -564,6 +620,8 @@ void FSequencer::InitSequencer(const FSequencerInitParams& InitParams
 	ViewModel = MakeShared<FSequencerEditorViewModel>(SharedThis(this), GetHostCapabilities());
 	ViewModel->InitializeEditor();
 	ViewModel->SetSequence(InitParams.RootSequence);
+	ViewModel->GetRootSequenceModel()->InitializeExtensions();
+	
 
 	ViewModel->GetSelection()->OnChanged.AddRaw(this, &FSequencer::OnSelectionChanged);
 
@@ -967,14 +1025,12 @@ void FSequencer::Tick(float InDeltaTime)
 
 			SuppressAutoEvalSignature.Reset();
 
-			if (GetNumericPositionTypeInterface()->GetOnSettingChanged() != nullptr)
-			{
-				GetNumericPositionTypeInterface()->GetOnSettingChanged()->Broadcast();
-			}
-
-			if (GetNumericDurationTypeInterface()->GetOnSettingChanged() != nullptr)
+			for (TSharedRef<UE::Sequencer::FSequencerNumericTypeInterface> Interface : GetNumericTypeInterfaces())
 			{
-				GetNumericDurationTypeInterface()->GetOnSettingChanged()->Broadcast();
+				if (Interface->Interface->GetOnSettingChanged() != nullptr)
+				{
+					Interface->Interface->GetOnSettingChanged()->Broadcast();
+				}
 			}
 		}
 	}
@@ -1034,17 +1090,7 @@ void FSequencer::Tick(float InDeltaTime)
 
 	if (PlaybackState == EMovieScenePlayerStatus::Playing)
 	{
-		FFrameTime NewGlobalTime;
-		if (MovieScene->HasMusicMap() && TimeController->SupportsMusicalTime())
-		{
-			FMusicalTime MusicalTime = TimeController->RequestCurrentMusicalTime(GlobalTime, PlaybackSpeed * Dilation, GetFocusedDisplayRate());
-			FFrameTime FrameTime = MovieScene->GetMusicMap()->GetFrameTime(MusicalTime);
-			NewGlobalTime = FFrameRate::TransformTime(FrameTime, MovieScene->GetMusicMap()->FrameResolution, GetFocusedDisplayRate());
-		}
-		else
-		{
-			NewGlobalTime = TimeController->RequestCurrentTime(GlobalTime, PlaybackSpeed * Dilation, GetFocusedDisplayRate());
-		}
+		FFrameTime NewGlobalTime = TimeController->RequestCurrentTime(GlobalTime, PlaybackSpeed * Dilation, GetFocusedDisplayRate());
 
 		// Put the time into clamped local space
 		FMovieSceneTransformBreadcrumbs Breadcrumbs;
@@ -3995,19 +4041,47 @@ ISequencer::FOnCameraCut& FSequencer::OnCameraCut()
 	return OnCameraCutEvent;
 }
 
-TSharedRef<INumericTypeInterface<double>> FSequencer::GetNumericPositionTypeInterface() const
+void FSequencer::AddNumericTypeInterface(TSharedRef<UE::Sequencer::FSequencerNumericTypeInterface> InNumericTypeInterface)
 {
-	return SequencerWidget->GetNumericPositionTypeInterface();
+	NumericTypeInterfaces.AddUnique(InNumericTypeInterface);
 }
 
-TSharedRef<INumericTypeInterface<double>> FSequencer::GetNumericDurationTypeInterface() const
+void FSequencer::RemoveNumericTypeInterface(TSharedRef<UE::Sequencer::FSequencerNumericTypeInterface> InNumericTypeInterface)
 {
-	return SequencerWidget->GetNumericDurationTypeInterface();
+	NumericTypeInterfaces.Remove(InNumericTypeInterface);
 }
 
-TSharedRef<INumericTypeInterface<double>> FSequencer::GetNumericNonMusicalTypeInterface() const
+TArrayView<const TSharedRef<UE::Sequencer::FSequencerNumericTypeInterface>> FSequencer::GetNumericTypeInterfaces() const
 {
-	return SequencerWidget->GetNumericNonMusicalTypeInterface();
+	return NumericTypeInterfaces;
+}
+
+TSharedPtr<INumericTypeInterface<double>> FSequencer::GetNumericTypeInterface(UE::Sequencer::ENumericIntent Intent) const
+{
+	TSharedPtr<INumericTypeInterface<double>> Result;
+
+	int32 Score = MIN_int32;
+	for (TSharedRef<UE::Sequencer::FSequencerNumericTypeInterface> Interface : NumericTypeInterfaces)
+	{
+		if (Interface->Intent != Intent)
+		{
+			continue;
+		}
+		const int32 ThisScore = Interface->GetRelevancyScore(*this, nullptr);
+		if (ThisScore > Score)
+		{
+			Result = Interface->Interface;
+			Score = ThisScore;
+		}
+	}
+
+	check(Result);
+	return Result.ToSharedRef();
+}
+
+TSharedRef<IPropertyTypeCustomization> FSequencer::MakeFrameNumberDetailsCustomization()
+{
+	return MakeShared<FFrameNumberDetailsCustomization>(TWeakPtr<ISequencer>(AsShared()));
 }
 
 TSharedRef<SWidget> FSequencer::MakeTimeRange(const TSharedRef<SWidget>& InnerContent, bool bShowWorkingRange, bool bShowViewRange, bool bShowPlaybackRange)
@@ -4463,14 +4537,14 @@ void FSequencer::ResetPerMovieSceneData()
 
 	SuppressAutoEvalSignature.Reset();
 
-	if (GetNumericPositionTypeInterface()->GetOnSettingChanged() != nullptr)
-	{
-		GetNumericPositionTypeInterface()->GetOnSettingChanged()->Broadcast();
-	}
-	if (GetNumericDurationTypeInterface()->GetOnSettingChanged() != nullptr)
+	for (TSharedRef<UE::Sequencer::FSequencerNumericTypeInterface> Interface : GetNumericTypeInterfaces())
 	{
-		GetNumericDurationTypeInterface()->GetOnSettingChanged()->Broadcast();
+		if (Interface->Interface->GetOnSettingChanged() != nullptr)
+		{
+			Interface->Interface->GetOnSettingChanged()->Broadcast();
+		}
 	}
+
 	// @todo run through all tracks for new movie scene changes
 	//  needed for audio track decompression
 }
@@ -4536,7 +4610,7 @@ TSharedRef<SWidget> FSequencer::MakeTransportControls(bool bExtended)
 	return EditorWidgetsModule.CreateTransportControl( TransportControlArgs );
 }
 
-TSharedRef<SWidget> FSequencer::MakePlayTimeDisplay(const TSharedRef<INumericTypeInterface<double>>& InNumericTypeInterface)
+TSharedRef<SWidget> FSequencer::MakePlayTimeDisplay()
 {
 	return SNew(STemporarilyFocusedSpinBox<double>)
 		.Style(&FAppStyle::GetWidgetStyle<FSpinBoxStyle>("Sequencer.PlayTimeSpinBox"))
@@ -4572,14 +4646,14 @@ TSharedRef<SWidget> FSequencer::MakePlayTimeDisplay(const TSharedRef<INumericTyp
 
 			return FText::FromString(FString::Printf(TEXT("%d of %d"), FrameCount.Value, FrameDuration.Value));
 		})
-		.TypeInterface(InNumericTypeInterface)
+		.TypeInterface(this, &FSequencer::GetNumericTypeInterface, UE::Sequencer::ENumericIntent::Position)
 		.Delta(this, &FSequencer::GetDisplayRateDeltaFrameCount)
 		.LinearDeltaSensitivity(25)
 		.MinDesiredWidth_Lambda([this] {
 			TRange<double> ViewRange = GetViewRange();
 
-			FString LowerBoundStr = GetNumericPositionTypeInterface()->ToString(ViewRange.GetLowerBoundValue());
-			FString UpperBoundStr = GetNumericPositionTypeInterface()->ToString(ViewRange.GetUpperBoundValue());
+			FString LowerBoundStr = GetNumericTypeInterface()->ToString(ViewRange.GetLowerBoundValue());
+			FString UpperBoundStr = GetNumericTypeInterface()->ToString(ViewRange.GetUpperBoundValue());
 
 			// Always measure with the negative and subframe indicator so that the size doesn't change when there is and isn't a subframe
 			if (!LowerBoundStr.Contains(TEXT("*")))
@@ -5741,7 +5815,7 @@ FString FSequencer::GetFrameTimeText() const
 	}
 
 	FFrameTime LocalTime = RootToParentChainTransform->TransformTime(GetGlobalTime().Time, FTransformTimeParams().IgnoreClamps());
-	return GetNumericPositionTypeInterface()->ToString(LocalTime.GetFrame().Value);
+	return GetNumericTypeInterface()->ToString(LocalTime.GetFrame().Value);
 }
 
 FMovieSceneSequenceID FSequencer::GetScrubPositionParent() const
@@ -7059,15 +7133,14 @@ void FSequencer::OnSelectedNodesOnlyChanged()
 
 void FSequencer::OnTimeDisplayFormatChanged()
 {
-	TSharedPtr<FFrameNumberInterface> FrameNumberInterface = StaticCastSharedPtr<FFrameNumberInterface>(GetNumericPositionTypeInterface().ToSharedPtr());
-	if (FrameNumberInterface.IsValid())
-	{
-		FrameNumberInterface->DisplayFormatChanged();
-	}
-	FrameNumberInterface = StaticCastSharedPtr<FFrameNumberInterface>(GetNumericDurationTypeInterface().ToSharedPtr());
-	if (FrameNumberInterface.IsValid())
+	using namespace UE::Sequencer;
+
+	for (TSharedRef<FSequencerNumericTypeInterface> Interface : GetNumericTypeInterfaces())
 	{
-		FrameNumberInterface->DisplayFormatChanged();
+		if (Interface->Interface->GetOnSettingChanged() != nullptr)
+		{
+			Interface->Interface->GetOnSettingChanged()->Broadcast();
+		}
 	}
 }
 
@@ -9938,7 +10011,7 @@ void FSequencer::SetKeyTime()
 	}
 
 	// Create a popup showing the existing time value and let the user set a new one.
- 	GenericTextEntryModeless(NSLOCTEXT("Sequencer.Popups", "SetKeyTimePopup", "New Time"), FText::FromString(GetNumericPositionTypeInterface()->ToString(KeyTime.Value)),
+ 	GenericTextEntryModeless(NSLOCTEXT("Sequencer.Popups", "SetKeyTimePopup", "New Time"), FText::FromString(GetNumericTypeInterface()->ToString(KeyTime.Value)),
  		FOnTextCommitted::CreateSP(this, &FSequencer::OnSetKeyTimeTextCommitted)
  	);
 }
@@ -9953,7 +10026,7 @@ void FSequencer::OnSetKeyTimeTextCommitted(const FText& InText, ETextCommit::Typ
 	CloseEntryPopupMenu();
 	if (CommitInfo == ETextCommit::OnEnter)
 	{
-		TOptional<double> NewFrameTime = GetNumericPositionTypeInterface()->FromString(InText.ToString(), 0);
+		TOptional<double> NewFrameTime = GetNumericTypeInterface()->FromString(InText.ToString(), 0);
 		if (!NewFrameTime.IsSet())
 			return;
 
@@ -9984,45 +10057,6 @@ void FSequencer::OnSetKeyTimeTextCommitted(const FText& InText, ETextCommit::Typ
 	}
 }
 
-void FSequencer::QuantizeSelectedKeyTimes(int32 QuantizationInterval)
-{
-	using namespace UE::Sequencer;
-
-	bool bAnythingChanged = false;
-
-	const UFrameBasedMusicMap* MusicMap = GetFocusedMusicMap();
-	FScopedTransaction SetKeyTimeTransaction(NSLOCTEXT("Sequencer", "SetKeyTime_Transaction", "Set Key Time"));
-	const FKeySelection& KeySelection = ViewModel->GetSelection()->KeySelection;
-	for (FKeyHandle Key : KeySelection)
-	{
-		TSharedPtr<FChannelModel> Channel = KeySelection.GetModelForKey(Key);
-		UMovieSceneSection* Section = Channel ? Channel->GetSection() : nullptr;
-		if (Channel && Section)
-		{
-			if (Section->TryModify())
-			{
-				FFrameTime FrameTime = Channel->GetKeyArea()->GetKeyTime(Key);
-				FMusicalTime MusicalTime = MusicMap->GetMusicalTime(FrameTime);
-				FMusicalTime QuantizedMusicalTime = MusicMap->Quantize(MusicalTime, QuantizationInterval);
-				FrameTime = MusicMap->GetFrameTime(QuantizedMusicalTime);
-				Channel->GetKeyArea()->SetKeyTime(Key, FrameTime.GetFrame());
-				bAnythingChanged = true;
-				Section->ExpandToFrame(FrameTime.GetFrame());
-			}
-		}
-	}
-
-	if (bAnythingChanged)
-	{
-		NotifyMovieSceneDataChanged(EMovieSceneDataChangeType::TrackValueChanged);
-	}
-}
-
-bool FSequencer::CanQuantizeKeyTime() const
-{
-	const UFrameBasedMusicMap* MusicMap = GetFocusedMusicMap();
-	return MusicMap && !MusicMap->IsEmpty() && ViewModel->GetSelection()->KeySelection.Num() > 0;
-}
 
 bool FSequencer::CanRekey() const
 {
@@ -12178,7 +12212,7 @@ bool FSequencer::GetGridMetrics(const float PhysicalWidth, const double InViewSt
 	TSharedRef<FSlateFontMeasure> FontMeasureService = FSlateApplication::Get().GetRenderer()->GetFontMeasureService();
 
 	// Use the end of the view as the longest number
-	FString TickString = GetNumericPositionTypeInterface()->ToString((InViewEnd * GetFocusedDisplayRate()).FrameNumber.Value);
+	FString TickString = GetNumericTypeInterface()->ToString((InViewEnd * GetFocusedDisplayRate()).FrameNumber.Value);
 	FVector2D MaxTextSize = FontMeasureService->Measure(TickString, SmallLayoutFont);
 
 	static float MajorTickMultiplier = 2.f;
@@ -12547,4 +12581,16 @@ TSharedRef<FSequencerFilterBar> FSequencer::GetFilterBar() const
 	return FilterBar.ToSharedRef();
 }
 
+EVisibility FSequencer::GetPlayRateComboVisibility() const
+{
+	using namespace UE::Sequencer;
+
+	TViewModelPtr<IClockExtension> ClockExtension = GetViewModel()->GetRootSequenceModel().ImplicitCast();
+	if (!ClockExtension || ClockExtension->ShouldShowPlayRateCombo(AsShared()))
+	{
+		return EVisibility::Visible;
+	}
+	return EVisibility::Collapsed;
+}
+
 #undef LOCTEXT_NAMESPACE
diff --git a/Engine/Source/Editor/Sequencer/Private/Sequencer.h b/Engine/Source/Editor/Sequencer/Private/Sequencer.h
--- a/Engine/Source/Editor/Sequencer/Private/Sequencer.h
+++ b/Engine/Source/Editor/Sequencer/Private/Sequencer.h
@@ -597,10 +597,6 @@ public:
 	void SetKeyTime();
 	void OnSetKeyTimeTextCommitted(const FText& InText, ETextCommit::Type CommitInfo);
 
-	/** Called when a user executes the quantize key time for selected keys */
-	void QuantizeSelectedKeyTimes(int32 QuantizationInterval);
-	bool CanQuantizeKeyTime() const;
-
 	/** Called when a user executes the rekey for selected keys */
 	bool CanRekey() const;
 	void Rekey();
@@ -822,9 +818,11 @@ public:
 	virtual FOnPostSave& OnPostSave() override;
 	virtual FOnActivateSequence& OnActivateSequence() override;
 	virtual FOnCameraCut& OnCameraCut() override;
-	virtual TSharedRef<INumericTypeInterface<double>> GetNumericPositionTypeInterface() const override;
-	virtual TSharedRef<INumericTypeInterface<double>> GetNumericDurationTypeInterface() const override;
-	virtual TSharedRef<INumericTypeInterface<double>> GetNumericNonMusicalTypeInterface() const override;
+	virtual void AddNumericTypeInterface(TSharedRef<UE::Sequencer::FSequencerNumericTypeInterface> InNumericTypeInterface);
+	virtual void RemoveNumericTypeInterface(TSharedRef<UE::Sequencer::FSequencerNumericTypeInterface> InNumericTypeInterface);
+	virtual TArrayView<const TSharedRef<UE::Sequencer::FSequencerNumericTypeInterface>> GetNumericTypeInterfaces() const override;
+	virtual TSharedPtr<INumericTypeInterface<double>> GetNumericTypeInterface(UE::Sequencer::ENumericIntent Intent = UE::Sequencer::ENumericIntent::Position) const override;
+	virtual TSharedRef<IPropertyTypeCustomization> MakeFrameNumberDetailsCustomization() override;
 	virtual TSharedRef<SWidget> MakeTransportControls(bool bExtended) override;
 	virtual FReply OnPlay(bool bTogglePlay = true) override;
 	virtual void Pause() override;
@@ -1191,7 +1189,7 @@ public:
 	/** Get the name of the movie renderer to use, defaults to the first available if the setting is empty */
 	FString GetMovieRendererName() const;
 
-	TSharedRef<SWidget> MakePlayTimeDisplay(const TSharedRef<INumericTypeInterface<double>>& InNumericTypeInterface);
+	TSharedRef<SWidget> MakePlayTimeDisplay();
 
 	/** @return True if the Sequencer is currently limiting selection to only Sequencer key-able objects */
 	bool IsViewportSelectionLimited() const override;
@@ -1258,6 +1256,9 @@ public:
 
 	TSharedRef<FSequencerFilterBar> GetFilterBar() const;
 
+	/** Checks to see if a clock extension wants the standard PlayRate combo menu hidden */
+	EVisibility GetPlayRateComboVisibility() const;
+
 private:
 
 	/** Update the time bases for the current movie scene */
@@ -1423,6 +1424,8 @@ private:
 	/** Generic Popup Entry */
 	TWeakPtr<IMenu> EntryPopupMenu;
 
+	TArray<TSharedRef<UE::Sequencer::FSequencerNumericTypeInterface>> NumericTypeInterfaces;
+
 	/** Stores a dirty bit for whether the sequencer tree (and other UI bits) may need to be refreshed.  We
 	    do this simply to avoid refreshing the UI more than once per frame. (e.g. during live recording where
 		the MovieScene data can change many times per frame.) */
diff --git a/Engine/Source/Editor/Sequencer/Private/SequencerCommonHelpers.cpp b/Engine/Source/Editor/Sequencer/Private/SequencerCommonHelpers.cpp
--- a/Engine/Source/Editor/Sequencer/Private/SequencerCommonHelpers.cpp
+++ b/Engine/Source/Editor/Sequencer/Private/SequencerCommonHelpers.cpp
@@ -363,30 +363,6 @@ void SequencerHelpers::BuildEditSectionMenu(const TWeakPtr<FSequencer>& InWeakSe
 
 		UMovieSceneSequence* Sequence = Sequencer->GetFocusedMovieSceneSequence();
 		const TWeakObjectPtr<UMovieScene> CurrentScene = Sequence->GetMovieScene();
-		const TWeakPtr<INumericTypeInterface<double>> NumericPositionTypeInterface = Sequencer->GetNumericPositionTypeInterface();
-		const TWeakPtr<INumericTypeInterface<double>> NumericDurationTypeInterface = Sequencer->GetNumericDurationTypeInterface();
-		const TWeakPtr<INumericTypeInterface<double>> NumericNonMusicalTypeInterface = Sequencer->GetNumericNonMusicalTypeInterface();
-
-		// If all of the selected sections are musical, we can show durations as bars.beats, but if not
-		// we need to show durations as seconds...
-		bool bAllSequencesAreMusical = true;
-		for (TWeakObjectPtr<> Section : InWeakSections)
-		{
-			if (Section.IsValid())
-			{
-				const UMovieSceneSubSection* SubSection = Cast<UMovieSceneSubSection>(Section);
-				if (SubSection)
-				{
-					UMovieSceneSequence* SubSequence = SubSection->GetSequence();
-					UMovieScene* MovieScene = SubSequence->GetMovieScene();
-					if (!MovieScene->HasMusicMap())
-					{
-						bAllSequencesAreMusical = false;
-						break;
-					}
-				}
-			}
-		}
 
 		TSharedRef<SSectionDetailsNotifyHookWrapper> DetailsNotifyWrapper = SNew(SSectionDetailsNotifyHookWrapper);
 		FDetailsViewArgs DetailsViewArgs;
@@ -409,18 +385,12 @@ void SequencerHelpers::BuildEditSectionMenu(const TWeakPtr<FSequencer>& InWeakSe
 
 		TSharedRef<IDetailsView> DetailsView = FModuleManager::GetModuleChecked<FPropertyEditorModule>("PropertyEditor").CreateDetailView(DetailsViewArgs);
 		DetailsView->RegisterInstancedCustomPropertyTypeLayout("FrameNumber",
-			FOnGetPropertyTypeCustomizationInstance::CreateLambda([NumericPositionTypeInterface, NumericDurationTypeInterface, NumericNonMusicalTypeInterface, bAllSequencesAreMusical]()
-			{
-				// If all of the selected sections are musical, we can show durations as bars.beats, but if not
-				// we need to show durations as seconds...
-				return MakeShared<FFrameNumberDetailsCustomization>(NumericPositionTypeInterface.Pin(),
-				bAllSequencesAreMusical ? NumericDurationTypeInterface.Pin() : NumericNonMusicalTypeInterface.Pin(),
-				NumericNonMusicalTypeInterface.Pin());
-			}));
+			FOnGetPropertyTypeCustomizationInstance::CreateSP(Sequencer.ToSharedRef(), &FSequencer::MakeFrameNumberDetailsCustomization));
+
 		DetailsView->RegisterInstancedCustomPropertyLayout(UMovieSceneSection::StaticClass(),
-			FOnGetDetailCustomizationInstance::CreateLambda([NumericPositionTypeInterface, NumericDurationTypeInterface, NumericNonMusicalTypeInterface, CurrentScene]()
+			FOnGetDetailCustomizationInstance::CreateLambda([TypeInterface = Sequencer->GetNumericTypeInterface(), CurrentScene]()
 			{
-				return MakeShared<FMovieSceneSectionDetailsCustomization>(NumericPositionTypeInterface.Pin(), NumericDurationTypeInterface.Pin(), NumericNonMusicalTypeInterface.Pin(), CurrentScene.Get());
+				return MakeShared<FMovieSceneSectionDetailsCustomization>(TypeInterface, CurrentScene.Get());
 			}));
 
 		DetailsView->RegisterInstancedCustomPropertyTypeLayout("MovieSceneConditionContainer", FOnGetPropertyTypeCustomizationInstance::CreateLambda([Sequence, InWeakSequencer]() {
diff --git a/Engine/Source/Editor/Sequencer/Private/SequencerContextMenus.cpp b/Engine/Source/Editor/Sequencer/Private/SequencerContextMenus.cpp
--- a/Engine/Source/Editor/Sequencer/Private/SequencerContextMenus.cpp
+++ b/Engine/Source/Editor/Sequencer/Private/SequencerContextMenus.cpp
@@ -50,7 +50,8 @@
 #include "Tracks/MovieScenePropertyTrack.h"
 #include "Algo/AnyOf.h"
 #include "IKeyArea.h"
-#include "FrameBasedMusicMap.h"
+#include "SequencerToolMenuContext.h"
+#include "ToolMenus.h"
 
 #define LOCTEXT_NAMESPACE "SequencerContextMenus"
 
@@ -173,36 +174,39 @@ void FKeyContextMenu::PopulateMenu(FMenuBuilder& MenuBuilder, TSharedPtr<FExtend
 		);
 	}
 
-	MenuBuilder.BeginSection("SequencerKeyEdit", LOCTEXT("EditMenu", "Edit"));
+	if (!UToolMenus::Get()->IsMenuRegistered("Sequencer.KeyContextMenu"))
 	{
-		if (HotspotCast<FKeyHotspot>(Sequencer->GetViewModel()->GetHotspot()))
-		{
-			MenuBuilder.AddMenuEntry(FGenericCommands::Get().Cut);
-			MenuBuilder.AddMenuEntry(FGenericCommands::Get().Copy);
-			MenuBuilder.AddMenuEntry(FGenericCommands::Get().Duplicate);
-		}
-	}
-	MenuBuilder.EndSection(); // SequencerKeyEdit
+		UToolMenu* KeyContextMenu = UToolMenus::Get()->RegisterMenu("Sequencer.KeyContextMenu", NAME_None, EMultiBoxType::Menu);
 
-	MenuBuilder.BeginSection("SequencerKeys", LOCTEXT("KeysMenu", "Keys"));
-	{
-		MenuBuilder.AddMenuEntry(FSequencerCommands::Get().SetKeyTime);
-		if (Sequencer->CanQuantizeKeyTime())
-		{
-			MenuBuilder.AddSubMenu(
-				LOCTEXT("KeyQuantizeTime", "Quantize Key Time"),
-				LOCTEXT("KeyQuantizeTimeTooltip", "Quantize the key time to a musical subdivision"),
-				FNewMenuDelegate::CreateLambda([=](FMenuBuilder& SubMenuBuilder) { Shared->AddQuantizationMenu(SubMenuBuilder); }),
-				FUIAction(),
-				NAME_None,
-				EUserInterfaceActionType::Button
-			);
-		}
-		MenuBuilder.AddMenuEntry(FSequencerCommands::Get().Rekey);
-		MenuBuilder.AddMenuEntry(FSequencerCommands::Get().SnapToFrame);
-		MenuBuilder.AddMenuEntry(FSequencerCommands::Get().DeleteKeys);
+		KeyContextMenu->AddDynamicSection("Edit", FNewToolMenuDelegate::CreateStatic([](UToolMenu* InMenu){
+			USequencerToolMenuContext* ContextObject = InMenu->FindContext<USequencerToolMenuContext>();
+			TSharedPtr<ISequencer>     Sequencer     = ContextObject ? ContextObject->WeakSequencer.Pin() : nullptr;
+
+			if (Sequencer && HotspotCast<FKeyHotspot>(Sequencer->GetViewModel()->GetHotspot()))
+			{
+				FToolMenuSection& Section = InMenu->AddSection(NAME_None, LOCTEXT("EditMenu", "Edit"));
+
+				Section.AddMenuEntry(FGenericCommands::Get().Cut);
+				Section.AddMenuEntry(FGenericCommands::Get().Copy);
+				Section.AddMenuEntry(FGenericCommands::Get().Duplicate);
+			}
+		}));
+
+		FToolMenuSection& KeysSection = KeyContextMenu->AddSection("Keys", LOCTEXT("KeysMenu", "Keys"));
+		KeysSection.AddMenuEntry(FSequencerCommands::Get().SetKeyTime);
+		KeysSection.AddMenuEntry(FSequencerCommands::Get().Rekey);
+		KeysSection.AddMenuEntry(FSequencerCommands::Get().SnapToFrame);
+		FToolMenuSection& KeysRemovalSection = KeyContextMenu->AddSection("KeyRemoval");
+		KeysRemovalSection.AddSeparator(NAME_None);
+		KeysRemovalSection.AddMenuEntry(FSequencerCommands::Get().DeleteKeys);
 	}
-	MenuBuilder.EndSection(); // SequencerKeys
+
+	USequencerToolMenuContext* ContextObject = NewObject<USequencerToolMenuContext>();
+	ContextObject->WeakSequencer = Sequencer;
+
+	FToolMenuContext MenuContext(ContextObject);
+	MenuContext.AppendCommandList(Sequencer->GetCommandBindings(ESequencerCommandBindings::Sequencer));
+	MenuBuilder.AddWidget(UToolMenus::Get()->GenerateWidget("Sequencer.KeyContextMenu", MenuContext), FText(), true, false);
 }
 
 void FKeyContextMenu::AddPropertiesMenu(FMenuBuilder& MenuBuilder)
@@ -226,38 +230,6 @@ void FKeyContextMenu::AddPropertiesMenu(FMenuBuilder& MenuBuilder)
 		, FText::GetEmpty(), true);
 }
 
-void FKeyContextMenu::AddQuantizationMenu(FMenuBuilder& MenuBuilder)
-{
-	const TSharedPtr<FSequencer> Sequencer = WeakSequencer.Pin();
-	if (!Sequencer.IsValid())
-	{
-		return;
-	}
-
-	for (auto& Item : FrameBasedMusicMap::QuantizationOptionsSource)
-	{
-		int32 QunatizationInterval = *Item;
-		auto LocalSequencerPointer = WeakSequencer;
-		MenuBuilder.AddMenuEntry(
-			FrameBasedMusicMap::GetQuantizationName(*Item),
-			FText(),
-			FSlateIcon(),
-			FUIAction(FExecuteAction::CreateLambda([=]
-				()
-				{
-					const TSharedPtr<FSequencer> Sequencer = LocalSequencerPointer.Pin();
-					if (!Sequencer.IsValid())
-					{
-						return;
-					}
-					Sequencer->QuantizeSelectedKeyTimes(QunatizationInterval);
-				}
-				)
-			)
-			);
-	}
-}
-
 
 FSectionContextMenu::FSectionContextMenu(TWeakPtr<FSequencer> InWeakSequencer, FFrameTime InMouseDownTime)
 	: WeakSequencer(InWeakSequencer)
@@ -2093,7 +2065,7 @@ void FEasingContextMenu::PopulateMenu(FMenuBuilder& MenuBuilder, TSharedPtr<FExt
 					.OnBeginSliderMovement_Lambda(OnBeginSliderMovement)
 					.OnEndSliderMovement_Lambda(OnEndSliderMovement)
 					.BorderForegroundColor(FAppStyle::GetSlateColor("DefaultForeground"))
-					.TypeInterface(Sequencer->GetNumericPositionTypeInterface())
+					.TypeInterface(Sequencer->GetNumericTypeInterface())
 				]
 			]
 
diff --git a/Engine/Source/Editor/Sequencer/Private/SequencerTimeSliderController.cpp b/Engine/Source/Editor/Sequencer/Private/SequencerTimeSliderController.cpp
--- a/Engine/Source/Editor/Sequencer/Private/SequencerTimeSliderController.cpp
+++ b/Engine/Source/Editor/Sequencer/Private/SequencerTimeSliderController.cpp
@@ -5,6 +5,7 @@
 #include "MVVM/Selection/Selection.h"
 #include "MVVM/SharedViewModelData.h"
 #include "MVVM/Views/SOutlinerView.h"
+#include "MVVM/Extensions/IClockExtension.h"
 #include "TrackEditors/TimeWarpTrackEditor.h"
 #include "Channels/MovieSceneTimeWarpChannel.h"
 #include "SequencerTimeDomainOverride.h"
@@ -41,7 +42,6 @@
 #include "ISinglePropertyView.h"
 #include "IStructureDetailsView.h"
 #include "FrameNumberDetailsCustomization.h"
-#include "FrameBasedMusicMap.h"
 
 #define LOCTEXT_NAMESPACE "TimeSlider"
 
@@ -221,46 +221,24 @@ FSequencerTimeSliderController::FScrubberMetrics FSequencerTimeSliderController:
 	return Metrics;
 }
 
-struct FSequencerTimeSliderController::FDrawTickArgs
-{
-	/** Geometry of the area */
-	FGeometry AllottedGeometry;
-	/** Culling rect of the area */
-	FSlateRect CullingRect;
-	/** Color of each tick */
-	FLinearColor TickColor;
-	/** Offset in Y where to start the tick */
-	float TickOffset;
-	/** Height in of major ticks */
-	float MajorTickHeight;
-	/** Start layer for elements */
-	int32 StartLayer;
-	/** Draw effects to apply */
-	ESlateDrawEffect DrawEffects;
-	/** Whether or not to only draw major ticks */
-	bool bOnlyDrawMajorTicks;
-	/** Whether or not to mirror labels */
-	bool bMirrorLabels;
-	
-};
-
 void FSequencerTimeSliderController::DrawTicks( FSlateWindowElementList& OutDrawElements, const TRange<double>& ViewRange, const FScrubRangeToScreen& RangeToScreen, FDrawTickArgs& InArgs ) const
 {
+	using namespace UE::Sequencer;
+
 	TSharedPtr<FSequencer> Sequencer = WeakSequencer.Pin();
 	if (!Sequencer.IsValid())
 	{
 		return;
 	}
-	const UFrameBasedMusicMap* MusicMap = Sequencer->GetFocusedMusicMap();
 
-	if (Sequencer->GetSequencerSettings()->GetTimeDisplayFormat() != EFrameNumberDisplayFormats::BarsBeats || !MusicMap || MusicMap->IsEmpty())
-	{
-		DrawLinearTicks(OutDrawElements, ViewRange, RangeToScreen, InArgs);
-	}
-	else
+	TViewModelPtr<IClockExtension> Clock = Sequencer->GetViewModel()->GetRootSequenceModel().ImplicitCast();
+
+	if (Clock && Clock->DrawTicks(Sequencer, OutDrawElements, ViewRange, RangeToScreen, InArgs))
 	{
-		DrawBarBeatTicks(OutDrawElements, ViewRange, RangeToScreen, InArgs);
+		return;
 	}
+
+	DrawLinearTicks(OutDrawElements, ViewRange, RangeToScreen, InArgs);
 }
 
 void FSequencerTimeSliderController::DrawLinearTicks(FSlateWindowElementList& OutDrawElements, const TRange<double>& ViewRange, const FScrubRangeToScreen& RangeToScreen, FDrawTickArgs& InArgs) const
@@ -298,6 +276,8 @@ void FSequencerTimeSliderController::DrawLinearTicks(FSlateWindowElementList& Ou
 
 	const float  FlooredScrubPx  = RangeToScreen.InputToLocalX(ConvertFrameTime(TimeSliderArgs.ScrubPosition.Get(), TickResolution, GetDisplayRate()).FloorToFrame() / DisplayRate);
 
+	TSharedPtr<INumericTypeInterface<double>> NumericTypeInterface = Sequencer->GetNumericTypeInterface();
+
 	for (double CurrentMajorLine = FirstMajorLine; CurrentMajorLine < LastMajorLine; CurrentMajorLine += MajorGridStep)
 	{
 		float MajorLinePx = RangeToScreen.InputToLocalX( CurrentMajorLine );
@@ -318,7 +298,7 @@ void FSequencerTimeSliderController::DrawLinearTicks(FSlateWindowElementList& Ou
 
 		if (!InArgs.bOnlyDrawMajorTicks && !FMath::IsNearlyEqual(MajorLinePx, FlooredScrubPx, 3.f))
 		{
-			FString FrameString = TimeSliderArgs.NumericTypeInterface->ToString((CurrentMajorLine * TickResolution).RoundToFrame().Value);
+			FString FrameString = NumericTypeInterface->ToString((CurrentMajorLine * TickResolution).RoundToFrame().Value);
 
 			// Space the text between the tick mark but slightly above
 			FVector2D TextOffset( MajorLinePx + 5.f, InArgs.bMirrorLabels ? 1.f : FMath::Abs( InArgs.AllottedGeometry.Size.Y - (InArgs.MajorTickHeight+3.f) ) );
@@ -356,101 +336,6 @@ void FSequencerTimeSliderController::DrawLinearTicks(FSlateWindowElementList& Ou
 	}
 }
 
-void FSequencerTimeSliderController::DrawBarBeatTicks(FSlateWindowElementList& OutDrawElements, const TRange<double>& ViewRange, const FScrubRangeToScreen& RangeToScreen, FDrawTickArgs& InArgs) const
-{
-	TSharedPtr<FSequencer> Sequencer = WeakSequencer.Pin();
-	if (!Sequencer.IsValid())
-	{
-		return;
-	}
-
-	const UFrameBasedMusicMap* MusicMap = Sequencer->GetFocusedMusicMap();
-	if (!MusicMap || MusicMap->IsEmpty())
-	{
-		return;
-	}
-
-	FFrameRate     TickResolution = GetTickResolution();
-	FFrameRate     DisplayRate = GetDisplayRate();
-	FPaintGeometry PaintGeometry = InArgs.AllottedGeometry.ToPaintGeometry();
-	FSlateFontInfo TickFrameFont = FCoreStyle::GetDefaultFontStyle("Regular", 8);
-
-	double PixelsPerSecond = InArgs.AllottedGeometry.Size.X / (ViewRange.GetUpperBoundValue() - ViewRange.GetLowerBoundValue());
-	double MinSecondGap = 20.0f / PixelsPerSecond;
-
-	FFrameTime StartFrameTime = GetTickResolution().AsFrameTime(ViewRange.GetLowerBoundValue());
-	FFrameTime EndFrameTime = GetTickResolution().AsFrameTime(ViewRange.GetUpperBoundValue());
-	FFrameTime MinSpacing = GetTickResolution().AsFrameTime(MinSecondGap);
-
-	TArray<FSubdivisionPoint> Points = MusicMap->GatherSubdivisonsForTimespan(StartFrameTime, EndFrameTime, MinSpacing, FrameBasedMusicMap::EWholeQuantizationTicks::QuarterNote);
-
-	TArray<FVector2D> LinePoints;
-	LinePoints.SetNumUninitialized(2);
-
-	const bool bAntiAliasLines = false;
-	float LastLabelRightEdge = -10000;
-	for (const FSubdivisionPoint& Point : Points)
-	{
-		float LinePx = RangeToScreen.InputToLocalX(TickResolution.AsSeconds(Point.FrameTime));
-
-		if (Point.MusicalTime.TickInBar == 0)
-		{
-			if (LinePx > (LastLabelRightEdge + 26))
-			{
-				LinePoints[0] = FVector2D(LinePx, InArgs.TickOffset - 25);
-				LinePoints[1] = FVector2D(LinePx, InArgs.TickOffset + InArgs.MajorTickHeight);
-
-				// Space the text between the tick mark but slightly above
-				FVector2D TextOffset(LinePx + 5.f, InArgs.bMirrorLabels ? 1.f : FMath::Abs(InArgs.AllottedGeometry.Size.Y - (30/*InArgs.MajorTickHeight*/ + 9.f)));
-				FSlateDrawElement::MakeText(
-					OutDrawElements,
-					InArgs.StartLayer + 1,
-					InArgs.AllottedGeometry.ToPaintGeometry(InArgs.AllottedGeometry.Size, FSlateLayoutTransform(TextOffset)),
-					Point.PrettyString,
-					TickFrameFont,
-					InArgs.DrawEffects,
-					InArgs.TickColor * 0.65f
-				);
-
-				FVector2D LabelSize = FontMeasureService->Measure(Point.PrettyString, TickFrameFont);
-				LastLabelRightEdge = LinePx + 5.f + LabelSize.Y;
-			}
-			else
-			{
-				LinePoints[0] = FVector2D(LinePx, InArgs.TickOffset);
-				LinePoints[1] = FVector2D(LinePx, InArgs.TickOffset + InArgs.MajorTickHeight);
-			}
-
-			FSlateDrawElement::MakeLines(
-				OutDrawElements,
-				InArgs.StartLayer,
-				PaintGeometry,
-				LinePoints,
-				InArgs.DrawEffects,
-				InArgs.TickColor,
-				bAntiAliasLines
-			);
-		}
-		else
-		{ 
-			const float MinorTickHeight = 2.0f; // : 2.0f;
-
-			LinePoints[0] = FVector2D(LinePx, InArgs.bMirrorLabels ? 0.0f : FMath::Abs(InArgs.AllottedGeometry.Size.Y - MinorTickHeight));
-			LinePoints[1] = FVector2D(LinePx, LinePoints[0].Y + MinorTickHeight);
-
-			FSlateDrawElement::MakeLines(
-				OutDrawElements,
-				InArgs.StartLayer,
-				PaintGeometry,
-				LinePoints,
-				InArgs.DrawEffects,
-				InArgs.TickColor,
-				bAntiAliasLines
-			);
-		}
-	}
-}
-
 int32 FSequencerTimeSliderController::DrawMarkedFrames( const FGeometry& AllottedGeometry, const FScrubRangeToScreen& RangeToScreen, FSlateWindowElementList& OutDrawElements, int32 LayerId, const ESlateDrawEffect& DrawEffects, const FWidgetStyle& InWidgetStyle, bool bDrawLabels ) const
 {
 	if (!WeakSequencer.Pin()->GetSequencerSettings()->GetShowMarkedFrames())
@@ -703,7 +588,7 @@ int32 FSequencerTimeSliderController::OnPaintTimeSlider( bool bMirrorLabels, con
 			}
 			else
 			{
-				FrameString = TimeSliderArgs.NumericTypeInterface->ToString(ScrubPosition.Time.GetFrame().Value);
+				FrameString = Sequencer->GetNumericTypeInterface()->ToString(ScrubPosition.Time.GetFrame().Value);
 			}
 
 			if (TimeSliderArgs.ScrubPositionParent.Get() != MovieSceneSequenceID::Invalid)
@@ -1164,13 +1049,15 @@ FReply FSequencerTimeSliderController::OnMouseButtonDoubleClick(TSharedRef<const
 
 double FSequencerTimeSliderController::ComputeHeight() const
 {
+	using namespace UE::Sequencer;
+
 	TSharedPtr<FSequencer> Sequencer = WeakSequencer.Pin();
 	if (Sequencer.IsValid())
 	{
-		const UFrameBasedMusicMap* MusicMap = Sequencer->GetFocusedMusicMap();
-		if (Sequencer->GetSequencerSettings()->GetTimeDisplayFormat() == EFrameNumberDisplayFormats::BarsBeats && MusicMap && !MusicMap->IsEmpty())
+		TViewModelPtr<IClockExtension> Clock = Sequencer->GetViewModel()->GetRootSequenceModel().ImplicitCast();
+		if (Clock)
 		{
-			return 55.0;
+			return Clock->GetDesiredTimeSliderHeight(Sequencer);
 		}
 	}
 	return 22.0;
@@ -1883,7 +1770,7 @@ TSharedRef<SWidget> FSequencerTimeSliderController::OpenSetPlaybackRangeMenu(con
 	FMenuBuilder MenuBuilder(bShouldCloseWindowAfterMenuSelection, nullptr);
 
 	FText CurrentTimeText;
-	CurrentTimeText = FText::FromString(TimeSliderArgs.NumericTypeInterface->ToString(FrameNumber.Value));
+	CurrentTimeText = FText::FromString(Sequencer->GetNumericTypeInterface()->ToString(FrameNumber.Value));
 	
 	TRange<FFrameNumber> PlaybackRange = TimeSliderArgs.PlaybackRange.Get();
 	TOptional<TRange<FFrameNumber>> SubSequenceRange = TimeSliderArgs.SubSequenceRange.Get();
@@ -2056,8 +1943,8 @@ TSharedRef<SWidget> FSequencerTimeSliderController::OpenSetPlaybackRangeMenu(con
 					TSharedPtr<FStructOnScope> StructOnScope = MakeShared<FStructOnScope>(FMovieSceneMarkedFrame::StaticStruct(), (uint8 *)&InMovieScene->GetMarkedFrames()[InMarkedFrameIndex]);
 
 					DetailsView = PropertyEditorModule.CreateStructureDetailView(DetailsViewArgs, StructureDetailsViewArgs, nullptr);
-					DetailsView->GetDetailsView()->RegisterInstancedCustomPropertyTypeLayout("FrameNumber", FOnGetPropertyTypeCustomizationInstance::CreateLambda([this]() {
-						return MakeShared<FFrameNumberDetailsCustomization>(WeakSequencer.Pin()->GetNumericPositionTypeInterface()/*, WeakSequencer.Pin()->GetNumericDurationTypeInterface(), WeakSequencer.Pin()->GetNumericNonMusicalTypeInterface()*/); }));
+					DetailsView->GetDetailsView()->RegisterInstancedCustomPropertyTypeLayout("FrameNumber", 
+						FOnGetPropertyTypeCustomizationInstance::CreateSP(WeakSequencer.Pin().ToSharedRef(), &FSequencer::MakeFrameNumberDetailsCustomization));
 					DetailsView->SetStructureData(StructOnScope);
 
 					ChildSlot
diff --git a/Engine/Source/Editor/Sequencer/Private/Widgets/Sidebar/SMarkedFrameDetails.h b/Engine/Source/Editor/Sequencer/Private/Widgets/Sidebar/SMarkedFrameDetails.h
--- a/Engine/Source/Editor/Sequencer/Private/Widgets/Sidebar/SMarkedFrameDetails.h
+++ b/Engine/Source/Editor/Sequencer/Private/Widgets/Sidebar/SMarkedFrameDetails.h
@@ -67,10 +67,7 @@ public:
 		DetailsView = PropertyEditorModule.CreateStructureDetailView(DetailsViewArgs, StructureDetailsViewArgs, nullptr);
 
 		DetailsView->GetDetailsView()->RegisterInstancedCustomPropertyTypeLayout(TEXT("FrameNumber"),
-			FOnGetPropertyTypeCustomizationInstance::CreateLambda([this]()
-			{
-				return MakeShared<FFrameNumberDetailsCustomization>(WeakSequencer.Pin()->GetNumericPositionTypeInterface());
-			}));
+			FOnGetPropertyTypeCustomizationInstance::CreateSP(Sequencer.ToSharedRef(), &FSequencer::MakeFrameNumberDetailsCustomization));
 
 		DetailsView->SetStructureData(StructOnScope);
 
diff --git a/Engine/Source/Editor/Sequencer/Public/ISequencer.h b/Engine/Source/Editor/Sequencer/Public/ISequencer.h
--- a/Engine/Source/Editor/Sequencer/Public/ISequencer.h
+++ b/Engine/Source/Editor/Sequencer/Public/ISequencer.h
@@ -32,6 +32,7 @@
 #include "UObject/WeakObjectPtrTemplates.h"
 #include "Widgets/Input/NumericTypeInterface.h"
 #include "Widgets/SWidget.h"
+#include "ISequencerNumericTypeInterface.h"
 
 #if UE_ENABLE_INCLUDE_ORDER_DEPRECATED_IN_5_3
 #include "ITimeSlider.h"
@@ -71,6 +72,7 @@ struct FQualifiedFrameTime;
 struct FSequencerHostCapabilities;
 struct FSidebarDrawerConfig;
 template <typename NumericType> struct INumericTypeInterface;
+class IPropertyTypeCustomization;
 
 enum class EMapChangeType : uint8;
 enum class EPropertyKeyedStatus : uint8;
@@ -81,7 +83,6 @@ class IToolkitHost;
 struct FAnimatedRange;
 struct FMovieSceneChannelMetaData;
 struct FMovieSceneSequencePlaybackParams;
-class UFrameBasedMusicMap;
 
 namespace UE
 {
@@ -92,6 +93,7 @@ enum class ETimeDomain : uint8;
 class FSequencerEditorViewModel;
 struct FCreateBindingParams;
 struct FTimeDomainOverride;
+struct FSequencerNumericTypeInterface;
 
 } // namespace Sequencer
 } // namespace UE
@@ -774,14 +776,16 @@ public:
 	/** Gets a multicast delegate when the curve edtior associated with this sequencer has it's selection change. */
 	virtual FOnCurveDisplayChanged& GetCurveDisplayChanged() = 0;
 
-	/** @return a numeric type interface that will parse and display numbers as frame / time / bar.beat positions correctly */
-	virtual TSharedRef<INumericTypeInterface<double>> GetNumericPositionTypeInterface() const = 0;
+	virtual void AddNumericTypeInterface(TSharedRef<UE::Sequencer::FSequencerNumericTypeInterface> InNumericTypeInterface) = 0;
+
+	virtual void RemoveNumericTypeInterface(TSharedRef<UE::Sequencer::FSequencerNumericTypeInterface> InNumericTypeInterface) = 0;
 
-	/** @return a numeric type interface that will parse and display numbers as frame / time / bar.beat durations correctly */
-	virtual TSharedRef<INumericTypeInterface<double>> GetNumericDurationTypeInterface() const = 0;
+	virtual TArrayView<const TSharedRef<UE::Sequencer::FSequencerNumericTypeInterface>> GetNumericTypeInterfaces() const = 0;
+
+	/** @return a numeric type interface that will parse and display numbers as frame / time / bar.beat positions correctly */
+	virtual TSharedPtr<INumericTypeInterface<double>> GetNumericTypeInterface(UE::Sequencer::ENumericIntent Intent = UE::Sequencer::ENumericIntent::Position) const = 0;
 
-	/** @return a numeric type interface that will parse and display numbers as frame / time, and bar.beat positions/durations are forces to seconds */
-	virtual TSharedRef<INumericTypeInterface<double>> GetNumericNonMusicalTypeInterface() const = 0;
+	virtual TSharedRef<IPropertyTypeCustomization> MakeFrameNumberDetailsCustomization() = 0;
 
 	/** @return the command bindings for this sequencer */
 	virtual TSharedPtr<FUICommandList> GetCommandBindings(ESequencerCommandBindings Type = ESequencerCommandBindings::Sequencer) const = 0;
@@ -908,11 +912,6 @@ public:
 	 */
 	SEQUENCER_API FFrameRate GetFocusedDisplayRate() const;
 
-	/**
-	 * Get the music map from the currently focused sequence
-	 */
-	SEQUENCER_API const UFrameBasedMusicMap* GetFocusedMusicMap() const;
-
 	/**
 	 * Create a new binding for the specified object, lightly deprecated and no longer virtual in favor of the overload with binding parameters
 	 */
diff --git a/Engine/Source/Editor/Sequencer/Public/SequencerTimeSliderController.h b/Engine/Source/Editor/Sequencer/Public/SequencerTimeSliderController.h
--- a/Engine/Source/Editor/Sequencer/Public/SequencerTimeSliderController.h
+++ b/Engine/Source/Editor/Sequencer/Public/SequencerTimeSliderController.h
@@ -192,6 +192,29 @@ public:
 		}
 	};
 
+	struct FDrawTickArgs
+	{
+		/** Geometry of the area */
+		FGeometry AllottedGeometry;
+		/** Culling rect of the area */
+		FSlateRect CullingRect;
+		/** Color of each tick */
+		FLinearColor TickColor;
+		/** Offset in Y where to start the tick */
+		float TickOffset;
+		/** Height in of major ticks */
+		float MajorTickHeight;
+		/** Start layer for elements */
+		int32 StartLayer;
+		/** Draw effects to apply */
+		ESlateDrawEffect DrawEffects;
+		/** Whether or not to only draw major ticks */
+		bool bOnlyDrawMajorTicks;
+		/** Whether or not to mirror labels */
+		bool bMirrorLabels;
+		
+	};
+
 	/** Set that's evaluating */
 	void SetIsEvaluating()
 	{
@@ -203,8 +226,6 @@ private:
 	FReply OnMouseMoveImpl( SWidget& WidgetOwner, const FGeometry& MyGeometry, const FPointerEvent& MouseEvent, bool bFromTimeSlider );
 
 private:
-	// forward declared as class members to prevent name collision with similar types defined in other units
-	struct FDrawTickArgs;
 
 	/**
 	 * Call this method when the user's interaction has changed the scrub position
@@ -235,16 +256,6 @@ private:
 	 */
 	void DrawLinearTicks(FSlateWindowElementList& OutDrawElements, const TRange<double>& ViewRange, const FScrubRangeToScreen& RangeToScreen, FDrawTickArgs& InArgs) const;
 
-	/**
-	 * Draw time tick marks
-	 *
-	 * @param OutDrawElements	List to add draw elements to
-	 * @param ViewRange			The currently visible time range in seconds
-	 * @param RangeToScreen		Time range to screen space converter
-	 * @param InArgs			Parameters for drawing the tick lines
-	 */
-	void DrawBarBeatTicks(FSlateWindowElementList& OutDrawElements, const TRange<double>& ViewRange, const FScrubRangeToScreen& RangeToScreen, FDrawTickArgs& InArgs) const;
-
 	/**
 	 * Draw the selection range.
 	 *
diff --git a/Engine/Source/Editor/SequencerCore/Private/MVVM/Extensions/DynamicExtensionContainer.cpp b/Engine/Source/Editor/SequencerCore/Private/MVVM/Extensions/DynamicExtensionContainer.cpp
--- a/Engine/Source/Editor/SequencerCore/Private/MVVM/Extensions/DynamicExtensionContainer.cpp
+++ b/Engine/Source/Editor/SequencerCore/Private/MVVM/Extensions/DynamicExtensionContainer.cpp
@@ -10,7 +10,8 @@ const void* FDynamicExtensionContainer::CastDynamic(FViewModelTypeID Type) const
 {
 	for (const FDynamicExtensionInfo& DynamicExtension : DynamicExtensions)
 	{
-		if (const void* Result = DynamicExtension.TypeTable->Cast(&DynamicExtension.Extension.Get(), Type.GetTypeID()))
+		const void* Result = DynamicExtension.Extension ? DynamicExtension.TypeTable->Cast(DynamicExtension.Extension.Get(), Type.GetTypeID()) : nullptr;
+		if (Result)
 		{
 			return Result;
 		}
@@ -23,9 +24,10 @@ void FDynamicExtensionContainer::RemoveDynamicExtension(FViewModelTypeID Type)
 	for (int32 Index = DynamicExtensions.Num()-1; Index >= 0; --Index)
 	{
 		const FDynamicExtensionInfo& DynamicExtension = DynamicExtensions[Index];
-		if (DynamicExtension.TypeTable->Cast(&DynamicExtension.Extension.Get(), Type.GetTypeID()) != nullptr)
+		if (DynamicExtension.Extension && DynamicExtension.TypeTable->Cast(DynamicExtension.Extension.Get(), Type.GetTypeID()) != nullptr)
 		{
-			DynamicExtensions.RemoveAt(Index, 1);
+			// null it out so we can add/remove while we're iterating
+			DynamicExtensions[Index].Extension = nullptr;
 			return;
 		}
 	}
@@ -38,7 +40,7 @@ FDynamicExtensionContainerIterator::FDynamicExtensionContainerIterator(IteratorT
 {
 	for ( ; Iterator; ++Iterator)
 	{
-		CurrentExtension = const_cast<void*>(Iterator->TypeTable->Cast(&Iterator->Extension.Get(), Type.GetTypeID()));
+		CurrentExtension = Iterator->Extension ? const_cast<void*>(Iterator->TypeTable->Cast(Iterator->Extension.Get(), Type.GetTypeID())) : nullptr;
 		if (CurrentExtension)
 		{
 			break;
@@ -52,7 +54,7 @@ FDynamicExtensionContainerIterator& FDynamicExtensionContainerIterator::operator
 
 	for (CurrentExtension = nullptr; Iterator; ++Iterator)
 	{
-		CurrentExtension = const_cast<void*>(Iterator->TypeTable->Cast(&Iterator->Extension.Get(), Type.GetTypeID()));
+		CurrentExtension = Iterator->Extension ? const_cast<void*>(Iterator->TypeTable->Cast(Iterator->Extension.Get(), Type.GetTypeID())) : nullptr;
 		if (CurrentExtension)
 		{
 			break;
diff --git a/Engine/Source/Editor/SequencerCore/Private/MVVM/ViewModels/ViewModel.cpp b/Engine/Source/Editor/SequencerCore/Private/MVVM/ViewModels/ViewModel.cpp
--- a/Engine/Source/Editor/SequencerCore/Private/MVVM/ViewModels/ViewModel.cpp
+++ b/Engine/Source/Editor/SequencerCore/Private/MVVM/ViewModels/ViewModel.cpp
@@ -172,6 +172,11 @@ TSharedPtr<FSharedViewModelData> FViewModel::GetSharedData() const
 	return SharedData;
 }
 
+void FViewModel::SetSharedDataFromConstructor(TSharedPtr<FSharedViewModelData> InSharedData)
+{
+	SharedData = InSharedData;
+}
+
 void FViewModel::SetSharedData(TSharedPtr<FSharedViewModelData> InSharedData)
 {
 	if (SharedData == InSharedData)
diff --git a/Engine/Source/Editor/SequencerCore/Public/MVVM/Extensions/DynamicExtensionContainer.h b/Engine/Source/Editor/SequencerCore/Public/MVVM/Extensions/DynamicExtensionContainer.h
--- a/Engine/Source/Editor/SequencerCore/Public/MVVM/Extensions/DynamicExtensionContainer.h
+++ b/Engine/Source/Editor/SequencerCore/Public/MVVM/Extensions/DynamicExtensionContainer.h
@@ -146,17 +146,41 @@ protected:
 
 	void PostInitializeExtensions()
 	{
-		for (const FDynamicExtensionInfo& DynamicExtension : DynamicExtensions)
+		// Iterate by index with null checking to support mutation during iteration
+		for (int32 Index = 0; Index < DynamicExtensions.Num(); ++Index)
 		{
-			DynamicExtension.Extension->OnPostInitialize();
+			if (DynamicExtensions[Index].Extension)
+			{
+				DynamicExtensions[Index].Extension->OnPostInitialize();
+			}
+		}
+		// Clean up nulls
+		for (int32 Index = DynamicExtensions.Num()-1; Index >= 0; --Index)
+		{
+			if (!DynamicExtensions[Index].Extension)
+			{
+				DynamicExtensions.RemoveAtSwap(Index, 1);
+			}
 		}
 	}
 
 	void ReinitializeExtensions()
 	{
-		for (const FDynamicExtensionInfo& DynamicExtension : DynamicExtensions)
+		// Iterate by index with null checking to support mutation during iteration
+		for (int32 Index = 0; Index < DynamicExtensions.Num(); ++Index)
+		{
+			if (DynamicExtensions[Index].Extension)
+			{
+				DynamicExtensions[Index].Extension->OnReinitialize();
+			}
+		}
+		// Clean up nulls
+		for (int32 Index = DynamicExtensions.Num()-1; Index >= 0; --Index)
 		{
-			DynamicExtension.Extension->OnReinitialize();
+			if (!DynamicExtensions[Index].Extension)
+			{
+				DynamicExtensions.RemoveAtSwap(Index, 1);
+			}
 		}
 	}
 
@@ -168,7 +192,7 @@ private:
 	struct FDynamicExtensionInfo
 	{
 		FCastableTypeTable* TypeTable;
-		TSharedRef<IDynamicExtension> Extension;
+		TSharedPtr<IDynamicExtension> Extension;
 	};
 
 	TArray<FDynamicExtensionInfo> DynamicExtensions;
diff --git a/Engine/Source/Editor/SequencerCore/Public/MVVM/ViewModels/ViewModel.h b/Engine/Source/Editor/SequencerCore/Public/MVVM/ViewModels/ViewModel.h
--- a/Engine/Source/Editor/SequencerCore/Public/MVVM/ViewModels/ViewModel.h
+++ b/Engine/Source/Editor/SequencerCore/Public/MVVM/ViewModels/ViewModel.h
@@ -77,28 +77,32 @@ public:
 	template<typename T>
 	TViewModelPtr<T> CastThisShared()
 	{
-		return TViewModelPtr<T>(AsShared(), CastThis<T>());
+		T* Casted = CastThis<T>();
+		return TViewModelPtr<T>(AsShared(), Casted ? Casted : CastDynamic<T>());
 	}
 
 	/** Casts this data model to an extension, or to a child class implementation */
 	template<typename T>
 	TViewModelPtr<const T> CastThisShared() const
 	{
-		return TViewModelPtr<const T>(AsShared(), CastThis<const T>());
+		const T* Casted = CastThis<const T>();
+		return TViewModelPtr<const T>(AsShared(), Casted ? Casted : CastDynamic<const T>());
 	}
 
 	/** Casts this data model to an extension, or to a child class implementation */
 	template<typename T>
 	TSharedPtr<T> CastThisSharedChecked()
 	{
-		return TSharedPtr<T>(AsShared(), CastThisChecked<T>());
+		T* Casted = CastThisChecked<T>();
+		return TSharedPtr<T>(AsShared(), Casted ? Casted : CastDynamicChecked<T>());
 	}
 
 	/** Casts this data model to an extension, or to a child class implementation */
 	template<typename T>
 	TSharedPtr<const T> CastThisSharedChecked() const
 	{
-		return TSharedPtr<const T>(AsShared(), CastThisChecked<T>());
+		const T* Casted = CastThisChecked<const T>();
+		return TSharedPtr<const T>(AsShared(), Casted ? Casted : CastDynamic<const T>());
 	}
 
 	/**
@@ -244,6 +248,8 @@ protected:
 	/** Retrieve the children for the specified child list head */
 	FViewModelChildren GetChildrenForList(FViewModelListHead* ListHead);
 
+	void SetSharedDataFromConstructor(TSharedPtr<FSharedViewModelData> InSharedData);
+
 private:
 
 	enum class ESetParentResult
diff --git a/Engine/Source/Editor/SequencerWidgets/Private/STimeRange.cpp b/Engine/Source/Editor/SequencerWidgets/Private/STimeRange.cpp
--- a/Engine/Source/Editor/SequencerWidgets/Private/STimeRange.cpp
+++ b/Engine/Source/Editor/SequencerWidgets/Private/STimeRange.cpp
@@ -35,7 +35,7 @@ template <typename NumericType> struct INumericTypeInterface;
 
 #define LOCTEXT_NAMESPACE "STimeRange"
 
-void STimeRange::Construct( const STimeRange::FArguments& InArgs, TSharedRef<ITimeSliderController> InTimeSliderController, TSharedRef<INumericTypeInterface<double>> NumericTypeInterface )
+void STimeRange::Construct( const STimeRange::FArguments& InArgs, TSharedRef<ITimeSliderController> InTimeSliderController )
 {
 	TimeSliderController = InTimeSliderController;
 
@@ -51,7 +51,7 @@ void STimeRange::Construct( const STimeRange::FArguments& InArgs, TSharedRef<ITi
 		.MinValue(TOptional<double>())
 		.MaxValue(TOptional<double>())
 		.Style(&FAppStyle::Get().GetWidgetStyle<FSpinBoxStyle>("Sequencer.HyperlinkSpinBox"))
-		.TypeInterface(NumericTypeInterface)
+		.TypeInterface(InArgs._NumericTypeInterface)
 		.ClearKeyboardFocusOnCommit(true)
 		.Delta(this, &STimeRange::GetSpinboxDelta)
 		.LinearDeltaSensitivity(25);
@@ -65,7 +65,7 @@ void STimeRange::Construct( const STimeRange::FArguments& InArgs, TSharedRef<ITi
 		.MinValue(TOptional<double>())
 		.MaxValue(TOptional<double>())
 		.Style(&FAppStyle::Get().GetWidgetStyle<FSpinBoxStyle>("Sequencer.HyperlinkSpinBox"))
-		.TypeInterface(NumericTypeInterface)
+		.TypeInterface(InArgs._NumericTypeInterface)
 		.ClearKeyboardFocusOnCommit(true)
 		.Delta(this, &STimeRange::GetSpinboxDelta)
 		.LinearDeltaSensitivity(25);
@@ -83,7 +83,7 @@ void STimeRange::Construct( const STimeRange::FArguments& InArgs, TSharedRef<ITi
 		.MinValue(TOptional<double>())
 		.MaxValue(TOptional<double>())
 		.Style(&FAppStyle::Get().GetWidgetStyle<FSpinBoxStyle>("Sequencer.HyperlinkSpinBox"))
-		.TypeInterface(NumericTypeInterface)
+		.TypeInterface(InArgs._NumericTypeInterface)
 		.ClearKeyboardFocusOnCommit(true)
 		.Delta(this, &STimeRange::GetSpinboxDelta)
 		.LinearDeltaSensitivity(25);
@@ -98,7 +98,7 @@ void STimeRange::Construct( const STimeRange::FArguments& InArgs, TSharedRef<ITi
 		.MinValue(TOptional<double>())
 		.MaxValue(TOptional<double>())
 		.Style(&FAppStyle::Get().GetWidgetStyle<FSpinBoxStyle>("Sequencer.HyperlinkSpinBox"))
-		.TypeInterface(NumericTypeInterface)
+		.TypeInterface(InArgs._NumericTypeInterface)
 		.ClearKeyboardFocusOnCommit(true)
 		.Delta(this, &STimeRange::GetSpinboxDelta)
 		.LinearDeltaSensitivity(25);
@@ -116,7 +116,7 @@ void STimeRange::Construct( const STimeRange::FArguments& InArgs, TSharedRef<ITi
 		.MinValue(TOptional<double>())
 		.MaxValue(TOptional<double>())
 		.Style(&FAppStyle::Get().GetWidgetStyle<FSpinBoxStyle>("Sequencer.HyperlinkSpinBox"))
-		.TypeInterface(NumericTypeInterface)
+		.TypeInterface(InArgs._NumericTypeInterface)
 		.ClearKeyboardFocusOnCommit(true)
 		.Delta(this, &STimeRange::GetSpinboxDelta)
 		.LinearDeltaSensitivity(25);
@@ -131,7 +131,7 @@ void STimeRange::Construct( const STimeRange::FArguments& InArgs, TSharedRef<ITi
 		.MinValue(TOptional<double>())
 		.MaxValue(TOptional<double>())
 		.Style(&FAppStyle::Get().GetWidgetStyle<FSpinBoxStyle>("Sequencer.HyperlinkSpinBox"))
-		.TypeInterface(NumericTypeInterface)
+		.TypeInterface(InArgs._NumericTypeInterface)
 		.ClearKeyboardFocusOnCommit(true)
 		.Delta(this, &STimeRange::GetSpinboxDelta)
 		.LinearDeltaSensitivity(25);
diff --git a/Engine/Source/Editor/SequencerWidgets/Private/STimeRange.h b/Engine/Source/Editor/SequencerWidgets/Private/STimeRange.h
--- a/Engine/Source/Editor/SequencerWidgets/Private/STimeRange.h
+++ b/Engine/Source/Editor/SequencerWidgets/Private/STimeRange.h
@@ -16,6 +16,8 @@ public:
 	SLATE_BEGIN_ARGS(STimeRange)
 		: _ShowWorkingRange(true), _ShowViewRange(false), _ShowPlaybackRange(true)
 	{}
+		/** Whether to show the working range */
+		SLATE_ATTRIBUTE( TSharedPtr<INumericTypeInterface<double>>, NumericTypeInterface )
 		/** Whether to show the working range */
 		SLATE_ARGUMENT( bool, ShowWorkingRange )
 		/** Whether to show the view range */
@@ -37,7 +39,7 @@ public:
 	 * 
 	 * @param InArgs   A declaration from which to construct the widget
 	 */
-	void Construct( const FArguments& InArgs, TSharedRef<ITimeSliderController> InTimeSliderController, TSharedRef<INumericTypeInterface<double>> NumericTypeInterface );
+	void Construct( const FArguments& InArgs, TSharedRef<ITimeSliderController> InTimeSliderController);
 
 protected:
 	double GetSpinboxDelta() const;
diff --git a/Engine/Source/Editor/SequencerWidgets/Private/SequencerWidgetsModule.cpp b/Engine/Source/Editor/SequencerWidgets/Private/SequencerWidgetsModule.cpp
--- a/Engine/Source/Editor/SequencerWidgets/Private/SequencerWidgetsModule.cpp
+++ b/Engine/Source/Editor/SequencerWidgets/Private/SequencerWidgetsModule.cpp
@@ -45,7 +45,8 @@ public:
 
 	TSharedRef<ITimeSlider> CreateTimeRange(const FTimeRangeArgs& InArgs, const TSharedRef<SWidget>& Content) override
 	{
-		return SNew( STimeRange, InArgs.Controller, InArgs.NumericTypeInterface)
+		return SNew( STimeRange, InArgs.Controller )
+		.NumericTypeInterface(InArgs.NumericTypeInterface)
 		.Visibility(InArgs.VisibilityDelegate)
 		.ShowWorkingRange(!!(InArgs.ShowRanges & EShowRange::WorkingRange))
 		.ShowViewRange(!!(InArgs.ShowRanges & EShowRange::ViewRange))
diff --git a/Engine/Source/Editor/SequencerWidgets/Public/ISequencerWidgetsModule.h b/Engine/Source/Editor/SequencerWidgets/Public/ISequencerWidgetsModule.h
--- a/Engine/Source/Editor/SequencerWidgets/Public/ISequencerWidgetsModule.h
+++ b/Engine/Source/Editor/SequencerWidgets/Public/ISequencerWidgetsModule.h
@@ -29,7 +29,7 @@ struct FTimeRangeArgs
 		EShowRange InShowRanges,
 		TSharedRef<ITimeSliderController> InController,
 		TAttribute<EVisibility> InVisibilityDelegate,
-		TSharedRef<INumericTypeInterface<double>> InNumericTypeInterface
+		TAttribute<TSharedPtr<INumericTypeInterface<double>>> InNumericTypeInterface
 		)
 		: ShowRanges(InShowRanges)
 		, EnableRanges(InShowRanges)
@@ -44,7 +44,7 @@ struct FTimeRangeArgs
 		EShowRange InEnableRanges,
 		TSharedRef<ITimeSliderController> InController,
 		TAttribute<EVisibility> InVisibilityDelegate,
-		TSharedRef<INumericTypeInterface<double>> InNumericTypeInterface
+		TAttribute<TSharedPtr<INumericTypeInterface<double>>> InNumericTypeInterface
 		)
 		: ShowRanges(InShowRanges)
 		, EnableRanges(InEnableRanges)
@@ -63,7 +63,7 @@ struct FTimeRangeArgs
 	/** Visibility delegate */
 	TAttribute<EVisibility> VisibilityDelegate;
 	/** Numeric type interface to use for frame<->time conversion */
-	TSharedRef<INumericTypeInterface<double>> NumericTypeInterface;
+	TAttribute<TSharedPtr<INumericTypeInterface<double>>> NumericTypeInterface;
 };
 
 /**
diff --git a/Engine/Source/Editor/UnrealEd/Classes/Settings/EditorExperimentalSettings.h b/Engine/Source/Editor/UnrealEd/Classes/Settings/EditorExperimentalSettings.h
--- a/Engine/Source/Editor/UnrealEd/Classes/Settings/EditorExperimentalSettings.h
+++ b/Engine/Source/Editor/UnrealEd/Classes/Settings/EditorExperimentalSettings.h
@@ -153,9 +153,6 @@ public:
 	UPROPERTY(EditAnywhere, config, Category = LevelInstance)
 	bool bEnableLevelInstanceLandscapeEdit;
 
-	UPROPERTY(EditAnywhere, config, Category = Sequencer, meta = (DisplayName = "Enable Music Mode Editing", ToolTip = "Enables \"music mode\" editing features"))
-	bool bMusicModeEditingInSequencerEnabled;
-
 	/**
 	 * Returns an event delegate that is executed when a setting has changed.
 	 *
diff --git a/Engine/Source/Editor/UnrealEd/Private/Settings/SettingsClasses.cpp b/Engine/Source/Editor/UnrealEd/Private/Settings/SettingsClasses.cpp
--- a/Engine/Source/Editor/UnrealEd/Private/Settings/SettingsClasses.cpp
+++ b/Engine/Source/Editor/UnrealEd/Private/Settings/SettingsClasses.cpp
@@ -181,12 +181,6 @@ static TAutoConsoleVariable<float> CVarEditorHDRNITLevel(
 	TEXT("Sets The desired NIT level of the editor when running on HDR"),
 	ECVF_Default);
 
-static TAutoConsoleVariable<int32> CVarSequenceEditorMusicModeEnabled(
-	TEXT("Sequencer.EnableMusicModeEditing"),
-	0,
-	TEXT("Sets whether or not we should allow the sequence editor to show music mode tools"),
-	ECVF_Default);
-
 UEditorExperimentalSettings::UEditorExperimentalSettings( const FObjectInitializer& ObjectInitializer )
 	: Super(ObjectInitializer)
 	, bEnableAsyncTextureCompilation(false)
@@ -199,7 +193,6 @@ UEditorExperimentalSettings::UEditorExperimentalSettings( const FObjectInitializ
 	, HDREditorNITLevel(160.0f)
 	, bUseOpenCLForConvexHullDecomp(false)
 	, bAllowPotentiallyUnsafePropertyEditing(false)
-	, bMusicModeEditingInSequencerEnabled(false)
 {
 }
 
@@ -217,7 +210,6 @@ void UEditorExperimentalSettings::PostInitProperties()
 
 	CVarEditorHDRSupport->Set(bHDREditor ? 1 : 0, ECVF_SetByProjectSetting);
 	CVarEditorHDRNITLevel->Set(HDREditorNITLevel, ECVF_SetByProjectSetting);
-	CVarSequenceEditorMusicModeEnabled->Set(bMusicModeEditingInSequencerEnabled ? 1 : 0, ECVF_SetByProjectSetting);
 	Super::PostInitProperties();
 }
 
@@ -239,10 +231,6 @@ void UEditorExperimentalSettings::PostEditChangeProperty( struct FPropertyChange
 	{
 		CVarEditorHDRNITLevel->Set(HDREditorNITLevel, ECVF_SetByProjectSetting);
 	}
-	else if (Name == GET_MEMBER_NAME_CHECKED(UEditorExperimentalSettings, bMusicModeEditingInSequencerEnabled))
-	{
-		CVarSequenceEditorMusicModeEnabled->Set(bMusicModeEditingInSequencerEnabled ? 1 : 0, ECVF_SetByProjectSetting);
-	}
 
 	if (!FUnrealEdMisc::Get().IsDeletePreferences())
 	{
diff --git a/Engine/Source/Runtime/AudioMixerCore/Private/AudioMixerCoreModule.cpp b/Engine/Source/Runtime/AudioMixerCore/Private/AudioMixerCoreModule.cpp
--- a/Engine/Source/Runtime/AudioMixerCore/Private/AudioMixerCoreModule.cpp
+++ b/Engine/Source/Runtime/AudioMixerCore/Private/AudioMixerCoreModule.cpp
@@ -14,6 +14,7 @@ public:
 	virtual void StartupModule() override
 	{
 		FModuleManager::Get().LoadModuleChecked(TEXT("SignalProcessing"));
+		FModuleManager::Get().LoadModuleChecked(TEXT("MusicEnvironment"));
 	}
 };
 
diff --git a/Engine/Source/Runtime/MovieScene/Private/Channels/MovieSceneAudioTriggerChannel.cpp b/Engine/Source/Runtime/MovieScene/Private/Channels/MovieSceneAudioTriggerChannel.cpp
--- a/Engine/Source/Runtime/MovieScene/Private/Channels/MovieSceneAudioTriggerChannel.cpp
+++ b/Engine/Source/Runtime/MovieScene/Private/Channels/MovieSceneAudioTriggerChannel.cpp
@@ -70,14 +70,9 @@ void FMovieSceneAudioTriggerChannel::DeleteKeysFrom(FFrameNumber InTime, bool bD
 	GetData().DeleteKeysFrom(InTime, bDeleteKeysBefore);
 }
 
-void FMovieSceneAudioTriggerChannel::ChangeFrameResolution(FFrameRate SourceRate, FFrameRate DestinationRate)
+void FMovieSceneAudioTriggerChannel::RemapTimes(const UE::MovieScene::IRetimingInterface& Retimer)
 {
-	GetData().ChangeFrameResolution(SourceRate, DestinationRate);
-}
-
-void FMovieSceneAudioTriggerChannel::RemapMusicalFrameTimes(const UFrameBasedMusicMap* OldMap, const UFrameBasedMusicMap* NewMap)
-{
-	GetData().RemapMusicalFrameTimes(OldMap, NewMap);
+	GetData().RemapTimes(Retimer);
 }
 
 TRange<FFrameNumber> FMovieSceneAudioTriggerChannel::ComputeEffectiveRange() const
diff --git a/Engine/Source/Runtime/MovieScene/Private/Channels/MovieSceneBoolChannel.cpp b/Engine/Source/Runtime/MovieScene/Private/Channels/MovieSceneBoolChannel.cpp
--- a/Engine/Source/Runtime/MovieScene/Private/Channels/MovieSceneBoolChannel.cpp
+++ b/Engine/Source/Runtime/MovieScene/Private/Channels/MovieSceneBoolChannel.cpp
@@ -148,15 +148,11 @@ void FMovieSceneBoolChannel::DeleteKeysFrom(FFrameNumber InTime, bool bDeleteKey
 	GetData().DeleteKeysFrom(InTime, bDeleteKeysBefore);
 }
 
-void FMovieSceneBoolChannel::ChangeFrameResolution(FFrameRate SourceRate, FFrameRate DestinationRate)
+void FMovieSceneBoolChannel::RemapTimes(const UE::MovieScene::IRetimingInterface& Retimer)
 {
-	GetData().ChangeFrameResolution(SourceRate, DestinationRate);
+	GetData().RemapTimes(Retimer);
 }
 
-void FMovieSceneBoolChannel::RemapMusicalFrameTimes(const UFrameBasedMusicMap* OldMap, const UFrameBasedMusicMap* NewMap)
-{
-	GetData().RemapMusicalFrameTimes(OldMap, NewMap);
-}
 
 TRange<FFrameNumber> FMovieSceneBoolChannel::ComputeEffectiveRange() const
 {
diff --git a/Engine/Source/Runtime/MovieScene/Private/Channels/MovieSceneByteChannel.cpp b/Engine/Source/Runtime/MovieScene/Private/Channels/MovieSceneByteChannel.cpp
--- a/Engine/Source/Runtime/MovieScene/Private/Channels/MovieSceneByteChannel.cpp
+++ b/Engine/Source/Runtime/MovieScene/Private/Channels/MovieSceneByteChannel.cpp
@@ -155,15 +155,11 @@ void FMovieSceneByteChannel::DeleteKeysFrom(FFrameNumber InTime, bool bDeleteKey
 	GetData().DeleteKeysFrom(InTime, bDeleteKeysBefore);
 }
 
-void FMovieSceneByteChannel::ChangeFrameResolution(FFrameRate SourceRate, FFrameRate DestinationRate)
+void FMovieSceneByteChannel::RemapTimes(const UE::MovieScene::IRetimingInterface& Retimer)
 {
-	GetData().ChangeFrameResolution(SourceRate, DestinationRate);
+	GetData().RemapTimes(Retimer);
 }
 
-void FMovieSceneByteChannel::RemapMusicalFrameTimes(const UFrameBasedMusicMap* OldMap, const UFrameBasedMusicMap* NewMap)
-{
-	GetData().RemapMusicalFrameTimes(OldMap, NewMap);
-}
 
 TRange<FFrameNumber> FMovieSceneByteChannel::ComputeEffectiveRange() const
 {
diff --git a/Engine/Source/Runtime/MovieScene/Private/Channels/MovieSceneChannel.cpp b/Engine/Source/Runtime/MovieScene/Private/Channels/MovieSceneChannel.cpp
--- a/Engine/Source/Runtime/MovieScene/Private/Channels/MovieSceneChannel.cpp
+++ b/Engine/Source/Runtime/MovieScene/Private/Channels/MovieSceneChannel.cpp
@@ -2,6 +2,7 @@
 
 #include "Channels/MovieSceneChannel.h"
 #include "Curves/KeyHandle.h"
+#include "IMovieSceneRetimingInterface.h"
 
 #include UE_INLINE_GENERATED_CPP_BY_NAME(MovieSceneChannel)
 
@@ -14,3 +15,8 @@ void FMovieSceneChannel::SetKeyTime(const FKeyHandle InHandle, const FFrameNumbe
 {
 	SetKeyTimes(MakeArrayView(&InHandle, 1), MakeArrayView(&InKeyTime, 1));
 }
+
+void FMovieSceneChannel::ChangeFrameResolution(FFrameRate SourceRate, FFrameRate DestinationRate)
+{
+	RemapTimes(UE::MovieScene::FFrameRateRetiming{ SourceRate, DestinationRate });
+}
\ No newline at end of file
diff --git a/Engine/Source/Runtime/MovieScene/Private/Channels/MovieSceneChannelData.cpp b/Engine/Source/Runtime/MovieScene/Private/Channels/MovieSceneChannelData.cpp
--- a/Engine/Source/Runtime/MovieScene/Private/Channels/MovieSceneChannelData.cpp
+++ b/Engine/Source/Runtime/MovieScene/Private/Channels/MovieSceneChannelData.cpp
@@ -1,7 +1,6 @@
 // Copyright Epic Games, Inc. All Rights Reserved.
 
 #include "Channels/MovieSceneChannelData.h"
-#include "FrameBasedMusicMap.h"
 #include "Misc/FrameRate.h"
 
 #include UE_INLINE_GENERATED_CPP_BY_NAME(MovieSceneChannelData)
@@ -10,6 +9,7 @@ namespace UE
 {
 namespace MovieScene
 {
+
 	void EvaluateTime(TArrayView<const FFrameNumber> InTimes, FFrameTime InTime, int32& OutIndex1, int32& OutIndex2)
 	{
 		const int32 Index2 = Algo::UpperBound(InTimes, InTime.FrameNumber);
@@ -199,30 +199,6 @@ void FMovieSceneChannelData::ChangeFrameResolution(FFrameRate SourceRate, FFrame
 	}
 }
 
-void FMovieSceneChannelData::RemapMusicalFrameTimes(const UFrameBasedMusicMap* OldMap, const UFrameBasedMusicMap* NewMap)
-{
-	// Unlike ChangeFrameResolution, this function can change the actual location in
-	// time of the the keyframes. So I *think* the right thing to do here is broadcast to
-	// all of the callbacks that want to know when keyframes change. I cribbed this from
-	// the MoveKeyInternal function above - Buzz
-	TArray<FKeyMoveEventItem> Items;
-
-	for (int32 Index = 0; Index < Times->Num(); ++Index)
-	{
-		FFrameNumber OriginalTime = (*Times)[Index];
-		(*Times)[Index] = NewMap->RemapFrameTime((*Times)[Index], OldMap).RoundToFrame();
-		if ((*Times)[Index] != OriginalTime)
-		{
-			Items.Add(FKeyMoveEventItem(Index, OriginalTime, Index, (*Times)[Index]));
-		}
-	}
-
-	if (OwningChannel && OwningChannel->OnKeyMovedEvent().IsBound())
-	{
-		OwningChannel->OnKeyMovedEvent().Broadcast(OwningChannel, Items);
-	}
-}
-
 void FMovieSceneChannelData::GetKeys(const TRange<FFrameNumber>& WithinRange, TArray<FFrameNumber>* OutKeyTimes, TArray<FKeyHandle>* OutKeyHandles)
 {
 	if (!Times->Num())
diff --git a/Engine/Source/Runtime/MovieScene/Private/Channels/MovieSceneCurveChannelImpl.cpp b/Engine/Source/Runtime/MovieScene/Private/Channels/MovieSceneCurveChannelImpl.cpp
--- a/Engine/Source/Runtime/MovieScene/Private/Channels/MovieSceneCurveChannelImpl.cpp
+++ b/Engine/Source/Runtime/MovieScene/Private/Channels/MovieSceneCurveChannelImpl.cpp
@@ -6,7 +6,6 @@
 #include "Channels/MovieSceneInterpolation.h"
 #include "HAL/ConsoleManager.h"
 #include "MovieSceneFrameMigration.h"
-#include "FrameBasedMusicMap.h"
 #include "Curves/CurveEvaluation.h"
 #include "UObject/FortniteMainBranchObjectVersion.h"
 #include "UObject/SequencerObjectVersion.h"
@@ -856,76 +855,25 @@ void TMovieSceneCurveChannelImpl<ChannelType>::DeleteKeysFrom(ChannelType* InCha
 template<typename ChannelType>
 void TMovieSceneCurveChannelImpl<ChannelType>::ChangeFrameResolution(ChannelType* InChannel, FFrameRate SourceRate, FFrameRate DestinationRate)
 {
-	check(InChannel->Times.Num() == InChannel->Values.Num());
-
-	float IntervalFactor = DestinationRate.AsInterval() / SourceRate.AsInterval();
-	for (int32 Index = 0; Index < InChannel->Times.Num(); ++Index)
-	{
-		InChannel->Times[Index] = ConvertFrameTime(InChannel->Times[Index], SourceRate, DestinationRate).RoundToFrame();
-
-		ChannelValueType& Value = InChannel->Values[Index];
-		Value.Tangent.ArriveTangent *= IntervalFactor;
-		Value.Tangent.LeaveTangent  *= IntervalFactor;
-	}
+	RemapTimes(InChannel, UE::MovieScene::FFrameRateRetiming(SourceRate, DestinationRate));
 }
 
-/**  Remap original times to new times given the old music map and the new music map.
-  *  Assumes the original times were set in referenc to the OldMap!
-  */
 template<typename ChannelType>
-void TMovieSceneCurveChannelImpl<ChannelType>::RemapMusicalFrameTimes(ChannelType* InChannel, const UFrameBasedMusicMap* OldMap, const UFrameBasedMusicMap* NewMap)
+void TMovieSceneCurveChannelImpl<ChannelType>::RemapTimes(ChannelType* InChannel, const UE::MovieScene::IRetimingInterface& Retimer)
 {
 	check(InChannel->Times.Num() == InChannel->Values.Num());
 
-	if (InChannel->Times.IsEmpty())
-	{
-		return;
-	}
+	double IntervalFactor = Retimer.GetScale();
+	InChannel->GetData().RemapTimes(Retimer);
 
-	TArray<FKeyMoveEventItem> MovedItems;
-
-	// Remapping may result in keyframes getting closer together or farther apart,
-	// so we want to scale the tangents as appropriate given any change to the 
-	// distance between the keyframes. This is why we can't just calculate one 
-	// IntervalFactor like is done above in ChangeFrameResolution.
-	FFrameNumber OriginalPreviousTime = InChannel->Times[0];
-	FFrameNumber NewPreviousTime = InChannel->Times[0];
 	for (int32 Index = 0; Index < InChannel->Times.Num(); ++Index)
 	{
-		FFrameNumber OriginalThisTime = InChannel->Times[Index];
-		InChannel->Times[Index] = NewMap->RemapFrameTime(InChannel->Times[Index], OldMap).RoundToFrame();
-		FFrameNumber NewThisTime = InChannel->Times[Index];
-
-		if (NewThisTime != OriginalThisTime)
-		{
-			MovedItems.Add(FKeyMoveEventItem(Index, OriginalThisTime, Index, NewThisTime));
-		}
-
-		if (Index != 0)
-		{
-			float OldInterval = static_cast<float>(OriginalThisTime.Value - OriginalPreviousTime.Value);
-			float NewInterval = static_cast<float>(NewThisTime.Value - NewPreviousTime.Value);
-			if (!FMath::IsNearlyEqual(NewInterval, OldInterval) && !FMath::IsNearlyZero(NewInterval))
-			{
-				float IntervalFactor = OldInterval / NewInterval;
-				ChannelValueType& PreviousValue = InChannel->Values[Index - 1];
-				PreviousValue.Tangent.LeaveTangent *= IntervalFactor;
-				ChannelValueType& ThisValue = InChannel->Values[Index];
-				ThisValue.Tangent.ArriveTangent *= IntervalFactor;
-			}
-		}
-		OriginalPreviousTime = OriginalThisTime;
-		NewPreviousTime = NewThisTime;
-	}
-
-	if (MovedItems.Num() > 0 && InChannel->OnKeyMovedEvent().IsBound())
-	{
-		InChannel->OnKeyMovedEvent().Broadcast(InChannel, MovedItems);
+		ChannelValueType& Value = InChannel->Values[Index];
+		Value.Tangent.ArriveTangent *= IntervalFactor;
+		Value.Tangent.LeaveTangent  *= IntervalFactor;
 	}
-
 }
 
-
 template<typename ChannelType>
 void TMovieSceneCurveChannelImpl<ChannelType>::Optimize(ChannelType* InChannel, const FKeyDataOptimizationParams& InParameters)
 {
diff --git a/Engine/Source/Runtime/MovieScene/Private/Channels/MovieSceneCurveChannelImpl.h b/Engine/Source/Runtime/MovieScene/Private/Channels/MovieSceneCurveChannelImpl.h
--- a/Engine/Source/Runtime/MovieScene/Private/Channels/MovieSceneCurveChannelImpl.h
+++ b/Engine/Source/Runtime/MovieScene/Private/Channels/MovieSceneCurveChannelImpl.h
@@ -92,10 +92,11 @@ struct TMovieSceneCurveChannelImpl
 	static void DeleteKeysFrom(ChannelType* InChannel, FFrameNumber InTime, bool bDeleteKeysBefore);
 
 	/** Change the frame resolution of the channel */
+	UE_DEPRECATED(5.6, "Please use RemapTimes")
 	static void ChangeFrameResolution(ChannelType* InChannel, FFrameRate SourceRate, FFrameRate DestinationRate);
 
-	/** Remap key positions from an old music map to a new music map */
-	static void RemapMusicalFrameTimes(ChannelType* InChannel, const UFrameBasedMusicMap* OldMap, const UFrameBasedMusicMap* NewMap);
+	/** Remap key positions using a custom remap interface */
+	static void RemapTimes(ChannelType* InChannel, const UE::MovieScene::IRetimingInterface& Retimer);
 
 	/** Optimize the channel's curve */
 	static void Optimize(ChannelType* InChannel, const FKeyDataOptimizationParams& InParameters);
diff --git a/Engine/Source/Runtime/MovieScene/Private/Channels/MovieSceneDoubleChannel.cpp b/Engine/Source/Runtime/MovieScene/Private/Channels/MovieSceneDoubleChannel.cpp
--- a/Engine/Source/Runtime/MovieScene/Private/Channels/MovieSceneDoubleChannel.cpp
+++ b/Engine/Source/Runtime/MovieScene/Private/Channels/MovieSceneDoubleChannel.cpp
@@ -21,9 +21,43 @@ static_assert(
 		"You also need to update the layout in FMovieSceneDoubleValue so that they match!");
 
 
-namespace UE::MovieScene::Interpolation
+namespace UE::MovieScene
 {
+	void OnRemapChannelKeyTime(const FMovieSceneChannel* Channel, const IRetimingInterface& Retimer, FFrameNumber PreviousTime, FFrameNumber CurrentTime, FMovieSceneDoubleValue& InOutValue)
+	{
+		if (InOutValue.InterpMode == ERichCurveInterpMode::RCIM_Cubic)
+		{
+			// This is a bit of a hack, but we scale tangents if the remapper has stretched the time around the key that was remapped
+			//    We figure out this stretch factor by retiming a time slightly ahead (1/4 of a second) of the key, and seeing how it differs from the new key time
+			FFrameTime Diff = 0.25 * static_cast<const FMovieSceneDoubleChannel*>(Channel)->GetTickResolution();
+
+			const double StretchFactor = (Retimer.RemapTime(FFrameTime(PreviousTime + Diff)) - CurrentTime).AsDecimal() / Diff.AsDecimal();
+			if (!FMath::IsNearlyEqual(StretchFactor, 1.0) && !FMath::IsNearlyEqual(StretchFactor, 0.0))
+			{
+				if (InOutValue.Tangent.ArriveTangent != 0.0)
+				{
+					InOutValue.Tangent.ArriveTangent = static_cast<float>(InOutValue.Tangent.ArriveTangent / StretchFactor);
+				}
+				else
+				{
+					InOutValue.Tangent.ArriveTangentWeight = static_cast<float>(InOutValue.Tangent.ArriveTangentWeight * StretchFactor);
+				}
+
+				if (InOutValue.Tangent.LeaveTangent != 0.0)
+				{
+					InOutValue.Tangent.LeaveTangent  = static_cast<float>(InOutValue.Tangent.LeaveTangent  / StretchFactor);
+				}
+				else
+				{
+					InOutValue.Tangent.LeaveTangentWeight = static_cast<float>(InOutValue.Tangent.LeaveTangentWeight * StretchFactor);
+				}
+			}
+		}
+	}
+} // namespace UE::MovieScene
 
+namespace UE::MovieScene::Interpolation
+{
 	struct FDoubleChannelPiecewiseData
 	{
 		const FMovieSceneDoubleChannel* Channel;
@@ -91,7 +125,8 @@ namespace UE::MovieScene::Interpolation
 		}
 	};
 
-} // namespace UE::MovieScene
+} // namespace UE::MovieScene::Interpolation
+
 
 
 bool FMovieSceneDoubleValue::Serialize(FArchive& Ar)
@@ -200,14 +235,9 @@ void FMovieSceneDoubleChannel::DeleteKeysFrom(FFrameNumber InTime, bool bDeleteK
 	AutoSetTangents();
 }
 
-void FMovieSceneDoubleChannel::ChangeFrameResolution(FFrameRate SourceRate, FFrameRate DestinationRate)
-{
-	FMovieSceneDoubleChannelImpl::ChangeFrameResolution(this, SourceRate, DestinationRate);
-}
-
-void FMovieSceneDoubleChannel::RemapMusicalFrameTimes(const UFrameBasedMusicMap* OldMap, const UFrameBasedMusicMap* NewMap)
+void FMovieSceneDoubleChannel::RemapTimes(const UE::MovieScene::IRetimingInterface& Retimer)
 {
-	FMovieSceneDoubleChannelImpl::RemapMusicalFrameTimes(this, OldMap, NewMap);
+	FMovieSceneDoubleChannelImpl::RemapTimes(this, Retimer);
 }
 
 TRange<FFrameNumber> FMovieSceneDoubleChannel::ComputeEffectiveRange() const
diff --git a/Engine/Source/Runtime/MovieScene/Private/Channels/MovieSceneFloatChannel.cpp b/Engine/Source/Runtime/MovieScene/Private/Channels/MovieSceneFloatChannel.cpp
--- a/Engine/Source/Runtime/MovieScene/Private/Channels/MovieSceneFloatChannel.cpp
+++ b/Engine/Source/Runtime/MovieScene/Private/Channels/MovieSceneFloatChannel.cpp
@@ -16,6 +16,42 @@ static_assert(
 		"The size of the float channel value has changed. You need to update the padding byte at the end of the structure. "
 		"You also need to update the layout in FMovieSceneDoubleValue so that they match!");
 
+
+namespace UE::MovieScene
+{
+	void OnRemapChannelKeyTime(const FMovieSceneChannel* Channel, const IRetimingInterface& Retimer, FFrameNumber PreviousTime, FFrameNumber CurrentTime, FMovieSceneFloatValue& InOutValue)
+	{
+		if (InOutValue.InterpMode == ERichCurveInterpMode::RCIM_Cubic)
+		{
+			// This is a bit of a hack, but we scale tangents if the remapper has stretched the time around the key that was remapped
+			//    We figure out this stretch factor by retiming a time slightly ahead (1/4 of a second) of the key, and seeing how it differs from the new key time
+			FFrameTime Diff = 0.25 * static_cast<const FMovieSceneFloatChannel*>(Channel)->GetTickResolution();
+
+			const double StretchFactor = (Retimer.RemapTime(FFrameTime(PreviousTime + Diff)) - CurrentTime).AsDecimal() / Diff.AsDecimal();
+			if (!FMath::IsNearlyEqual(StretchFactor, 1.0) && !FMath::IsNearlyEqual(StretchFactor, 0.0))
+			{
+				if (InOutValue.Tangent.ArriveTangent != 0.0)
+				{
+					InOutValue.Tangent.ArriveTangent = static_cast<float>(InOutValue.Tangent.ArriveTangent / StretchFactor);
+				}
+				else
+				{
+					InOutValue.Tangent.ArriveTangentWeight = static_cast<float>(InOutValue.Tangent.ArriveTangentWeight * StretchFactor);
+				}
+
+				if (InOutValue.Tangent.LeaveTangent != 0.0)
+				{
+					InOutValue.Tangent.LeaveTangent  = static_cast<float>(InOutValue.Tangent.LeaveTangent  / StretchFactor);
+				}
+				else
+				{
+					InOutValue.Tangent.LeaveTangentWeight = static_cast<float>(InOutValue.Tangent.LeaveTangentWeight * StretchFactor);
+				}
+			}
+		}
+	}
+} // namespace UE::MovieScene
+
 bool FMovieSceneFloatValue::Serialize(FArchive& Ar)
 {
 	return TMovieSceneCurveChannelImpl<FMovieSceneFloatChannel>::SerializeChannelValue(*this, Ar);
@@ -114,14 +150,9 @@ void FMovieSceneFloatChannel::DeleteKeysFrom(FFrameNumber InTime, bool bDeleteKe
 	AutoSetTangents();
 }
 
-void FMovieSceneFloatChannel::ChangeFrameResolution(FFrameRate SourceRate, FFrameRate DestinationRate)
-{
-	FMovieSceneFloatChannelImpl::ChangeFrameResolution(this, SourceRate, DestinationRate);
-}
-
-void FMovieSceneFloatChannel::RemapMusicalFrameTimes(const UFrameBasedMusicMap* OldMap, const UFrameBasedMusicMap* NewMap)
+void FMovieSceneFloatChannel::RemapTimes(const UE::MovieScene::IRetimingInterface& Retimer)
 {
-	FMovieSceneFloatChannelImpl::RemapMusicalFrameTimes(this, OldMap, NewMap);
+	FMovieSceneFloatChannelImpl::RemapTimes(this, Retimer);
 }
 
 TRange<FFrameNumber> FMovieSceneFloatChannel::ComputeEffectiveRange() const
diff --git a/Engine/Source/Runtime/MovieScene/Private/Channels/MovieSceneIntegerChannel.cpp b/Engine/Source/Runtime/MovieScene/Private/Channels/MovieSceneIntegerChannel.cpp
--- a/Engine/Source/Runtime/MovieScene/Private/Channels/MovieSceneIntegerChannel.cpp
+++ b/Engine/Source/Runtime/MovieScene/Private/Channels/MovieSceneIntegerChannel.cpp
@@ -178,14 +178,9 @@ void FMovieSceneIntegerChannel::DeleteKeysFrom(FFrameNumber InTime, bool bDelete
 	GetData().DeleteKeysFrom(InTime, bDeleteKeysBefore);
 }
 
-void FMovieSceneIntegerChannel::ChangeFrameResolution(FFrameRate SourceRate, FFrameRate DestinationRate)
+void FMovieSceneIntegerChannel::RemapTimes(const UE::MovieScene::IRetimingInterface& Retimer)
 {
-	GetData().ChangeFrameResolution(SourceRate, DestinationRate);
-}
-
-void FMovieSceneIntegerChannel::RemapMusicalFrameTimes(const UFrameBasedMusicMap* OldMap, const UFrameBasedMusicMap* NewMap)
-{
-	GetData().RemapMusicalFrameTimes(OldMap, NewMap);
+	GetData().RemapTimes(Retimer);
 }
 
 TRange<FFrameNumber> FMovieSceneIntegerChannel::ComputeEffectiveRange() const
diff --git a/Engine/Source/Runtime/MovieScene/Private/Channels/MovieSceneObjectPathChannel.cpp b/Engine/Source/Runtime/MovieScene/Private/Channels/MovieSceneObjectPathChannel.cpp
--- a/Engine/Source/Runtime/MovieScene/Private/Channels/MovieSceneObjectPathChannel.cpp
+++ b/Engine/Source/Runtime/MovieScene/Private/Channels/MovieSceneObjectPathChannel.cpp
@@ -130,14 +130,9 @@ void FMovieSceneObjectPathChannel::DeleteKeys(TArrayView<const FKeyHandle> InHan
 	GetData().DeleteKeys(InHandles);
 }
 
-void FMovieSceneObjectPathChannel::ChangeFrameResolution(FFrameRate SourceRate, FFrameRate DestinationRate)
+void FMovieSceneObjectPathChannel::RemapTimes(const UE::MovieScene::IRetimingInterface& Retimer)
 {
-	GetData().ChangeFrameResolution(SourceRate, DestinationRate);
-}
-
-void FMovieSceneObjectPathChannel::RemapMusicalFrameTimes(const UFrameBasedMusicMap* OldMap, const UFrameBasedMusicMap* NewMap)
-{
-	GetData().RemapMusicalFrameTimes(OldMap, NewMap);
+	GetData().RemapTimes(Retimer);
 }
 
 TRange<FFrameNumber> FMovieSceneObjectPathChannel::ComputeEffectiveRange() const
diff --git a/Engine/Source/Runtime/MovieScene/Private/Evaluation/IMovieSceneCustomClockSource.cpp b/Engine/Source/Runtime/MovieScene/Private/Evaluation/IMovieSceneCustomClockSource.cpp
--- a/Engine/Source/Runtime/MovieScene/Private/Evaluation/IMovieSceneCustomClockSource.cpp
+++ b/Engine/Source/Runtime/MovieScene/Private/Evaluation/IMovieSceneCustomClockSource.cpp
@@ -82,36 +82,6 @@ FFrameTime FMovieSceneTimeController_Custom::OnRequestCurrentTime(const FQualifi
 	return InCurrentTime.Time;
 }
 
-bool FMovieSceneTimeController_Custom::OnSupportsMusicalTime()
-{
-	if (WeakObject.IsStale())
-	{
-		ResolveInterfacePtr();
-	}
-
-	if (WeakObject.IsValid())
-	{
-		return InterfacePtr->OnSupportsMusicalTime();
-	}
-
-	return false;
-}
-
-FMusicalTime FMovieSceneTimeController_Custom::OnRequestCurrentMusicalTime(const FQualifiedFrameTime& InCurrentTime, float InPlayRate)
-{
-	if (WeakObject.IsStale())
-	{
-		ResolveInterfacePtr();
-	}
-
-	if (WeakObject.IsValid())
-	{
-		return InterfacePtr->OnRequestCurrentMusicalTime(InCurrentTime, InPlayRate);
-	}
-	
-	return FMusicalTime();
-}
-
 void FMovieSceneTimeController_Custom::ResolveInterfacePtr()
 {
 	WeakObject = nullptr;
diff --git a/Engine/Source/Runtime/MovieScene/Private/MovieScene.cpp b/Engine/Source/Runtime/MovieScene/Private/MovieScene.cpp
--- a/Engine/Source/Runtime/MovieScene/Private/MovieScene.cpp
+++ b/Engine/Source/Runtime/MovieScene/Private/MovieScene.cpp
@@ -15,6 +15,7 @@
 #include "Misc/FrameRate.h"
 #include "UObject/ObjectSaveContext.h"
 #include "UObject/UObjectHash.h"
+#include "MovieSceneClock.h"
 #include "UObject/FortniteMainBranchObjectVersion.h"
 
 #include UE_INLINE_GENERATED_CPP_BY_NAME(MovieScene)
@@ -129,6 +130,14 @@ void UMovieScene::PostLoad()
 	SortMarkedFrames();
 
 #if WITH_EDITORONLY_DATA
+
+	if (!CustomClockSourcePath_DEPRECATED.IsNull())
+	{
+		UMovieSceneExternalClock* ExternalClock = NewObject<UMovieSceneExternalClock>(this);
+		ExternalClock->CustomClockSourcePath = CustomClockSourcePath_DEPRECATED;
+		SetCustomClock(ExternalClock);
+	}
+
 	for (FMovieSceneMarkedFrame& MarkedFrame : MarkedFrames)
 	{
 		const FLinearColor DefaultDeprecatedColor(0.f, 1.f, 1.f, 0.4f);
@@ -524,10 +533,11 @@ FText UMovieScene::GetObjectDisplayName(const FGuid& ObjectId)
 
 void UMovieScene::SetTickResolutionDirectly(FFrameRate InTickResolution)
 {
+	FFrameRate PreviousTickResolution = TickResolution;
 	TickResolution = InTickResolution;
-	if (MusicMap)
+	if (CustomClock)
 	{
-		MusicMap->SetFrameResolution(TickResolution);
+		CustomClock->HandleTickResolutionChange(PreviousTickResolution, TickResolution);
 	}
 }
 
@@ -1706,9 +1716,47 @@ const FMovieSceneBinding* UMovieScene::FindBinding(const FGuid& ForGuid) const
 	return nullptr;
 }
 
+void UMovieScene::SetClockSource(EUpdateClockSource InNewClockSource)
+{
+	if (ClockSource != InNewClockSource)
+	{
+		ClockSource = InNewClockSource;
+		EventHandlers.Trigger(&UE::MovieScene::ISequenceDataEventHandler::OnClockChanged, InNewClockSource, CustomClock);
+	}
+}
+
+void UMovieScene::SetClockSource(UObject* InNewClockSource)
+{
+	UMovieSceneExternalClock* Clock = NewObject<UMovieSceneExternalClock>(this);
+	Clock->CustomClockSourcePath = InNewClockSource;
+
+	ClockSource = EUpdateClockSource::Custom;
+	CustomClock = Clock;
+}
+
+void UMovieScene::SetCustomClock(UMovieSceneClock* InNewClockSource)
+{
+	if (ClockSource == EUpdateClockSource::Custom && CustomClock == InNewClockSource)
+	{
+		return;
+	}
+
+	ClockSource = EUpdateClockSource::Custom;
+	if (!InNewClockSource->IsIn(this))
+	{
+		InNewClockSource = DuplicateObject<UMovieSceneClock>(InNewClockSource, this);
+	}
+	CustomClock = InNewClockSource;
+	EventHandlers.Trigger(&UE::MovieScene::ISequenceDataEventHandler::OnClockChanged, EUpdateClockSource::Custom, CustomClock);
+}
+
 TSharedPtr<FMovieSceneTimeController> UMovieScene::MakeCustomTimeController(UObject* PlaybackContext)
 {
-	return MakeShared<FMovieSceneTimeController_Custom>(CustomClockSourcePath, PlaybackContext);
+	if (CustomClock != nullptr && ClockSource == EUpdateClockSource::Custom)
+	{
+		return CustomClock->MakeTimeController(PlaybackContext);
+	}
+	return nullptr;
 }
 
 FMovieSceneTimecodeSource UMovieScene::GetEarliestTimecodeSource() const
diff --git a/Engine/Source/Runtime/MovieScene/Private/MovieSceneSection.cpp b/Engine/Source/Runtime/MovieScene/Private/MovieSceneSection.cpp
--- a/Engine/Source/Runtime/MovieScene/Private/MovieSceneSection.cpp
+++ b/Engine/Source/Runtime/MovieScene/Private/MovieSceneSection.cpp
@@ -780,33 +780,6 @@ FColor UMovieSceneSection::GetColorTint() const
 #endif
 }
 
-void UMovieSceneSection::RemapMusicalFrameTimes(const UFrameBasedMusicMap* OldMap, const UFrameBasedMusicMap* NewMap)
-{
-	if (TryModify())
-	{
-		if (SectionRange.Value.GetLowerBound().IsClosed())
-		{
-			FFrameTime OldFrame(SectionRange.Value.GetLowerBoundValue());
-			FFrameTime NewFrame = NewMap->RemapFrameTime(OldFrame, OldMap);
-			SectionRange.Value.SetLowerBoundValue(NewFrame.FloorToFrame());
-		}
-		if (SectionRange.Value.GetUpperBound().IsClosed())
-		{
-			FFrameTime OldFrame(SectionRange.Value.GetUpperBoundValue());
-			FFrameTime NewFrame = NewMap->RemapFrameTime(OldFrame, OldMap);
-			SectionRange.Value.SetUpperBoundValue(NewFrame.CeilToFrame());
-		}
-
-		for (const FMovieSceneChannelEntry& Entry : GetChannelProxy().GetAllEntries())
-		{
-			for (FMovieSceneChannel* Channel : Entry.GetChannels())
-			{
-				Channel->RemapMusicalFrameTimes(OldMap, NewMap);
-			}
-		}
-	}
-}
-
 UMovieSceneSection* UMovieSceneSection::SplitSection(FQualifiedFrameTime SplitTime, bool bDeleteKeys)
 {
 	if (!SectionRange.Value.Contains(SplitTime.Time.GetFrame()))
diff --git a/Engine/Source/Runtime/MovieScene/Private/MovieSceneSequencePlayer.cpp b/Engine/Source/Runtime/MovieScene/Private/MovieSceneSequencePlayer.cpp
--- a/Engine/Source/Runtime/MovieScene/Private/MovieSceneSequencePlayer.cpp
+++ b/Engine/Source/Runtime/MovieScene/Private/MovieSceneSequencePlayer.cpp
@@ -1114,19 +1114,7 @@ void UMovieSceneSequencePlayer::Update(const float DeltaSeconds)
 		{
 			check(!IsEvaluating());
 
-			FFrameTime NewTime;
-			UMovieScene* MovieScene = Sequence->GetMovieScene();
-			if (MovieScene && MovieScene->HasMusicMap() && TimeController->SupportsMusicalTime())
-			{
-				FMusicalTime MusicalTime = TimeController->RequestCurrentMusicalTime(GetCurrentTime(), PlayRate, GetDisplayRate());
-				FFrameTime FrameTime = MovieScene->GetMusicMap()->GetFrameTime(MusicalTime);
-				NewTime = FFrameRate::TransformTime(FrameTime, MovieScene->GetMusicMap()->FrameResolution, GetDisplayRate());
-			}
-			else
-			{
-				NewTime = TimeController->RequestCurrentTime(GetCurrentTime(), PlayRate, GetDisplayRate());
-			}
-
+			FFrameTime NewTime = TimeController->RequestCurrentTime(GetCurrentTime(), PlayRate, GetDisplayRate());
 			UpdateTimeCursorPosition(NewTime, EUpdatePositionMethod::Play);
 		}
 
diff --git a/Engine/Source/Runtime/MovieScene/Private/MovieSceneTimeController.cpp b/Engine/Source/Runtime/MovieScene/Private/MovieSceneTimeController.cpp
--- a/Engine/Source/Runtime/MovieScene/Private/MovieSceneTimeController.cpp
+++ b/Engine/Source/Runtime/MovieScene/Private/MovieSceneTimeController.cpp
@@ -68,17 +68,6 @@ FFrameTime FMovieSceneTimeController::RequestCurrentTime(const FQualifiedFrameTi
 	return OnRequestCurrentTime(InCurrentTime, InPlayRate);
 }
 
-bool FMovieSceneTimeController::SupportsMusicalTime()
-{
-	return OnSupportsMusicalTime();
-}
-
-FMusicalTime FMovieSceneTimeController::RequestCurrentMusicalTime(const FQualifiedFrameTime& InCurrentTime, float InPlayRate, FFrameRate InDisplayRate)
-{
-	DisplayRate = InDisplayRate;
-	return OnRequestCurrentMusicalTime(InCurrentTime, InPlayRate);
-}
-
 void FMovieSceneTimeController_ExternalClock::OnStartPlaying(const FQualifiedFrameTime& InStartTime)
 {
 	ClockStartTime = ClockLastUpdateTime = GetCurrentTime();
