Sequencer: Refined default relevancy scores for position/duration numeric type interfaces

New scores define the following behavior: For detail customizations, Position/Duration type interfaces will always be used based on the meta-data. For non-detail customizations uses, the default is to prefer position interfaces, unless pre-filtered by specifying an intent.

diff --git a/Engine/Source/Editor/Sequencer/Private/ISequencerNumericTypeInterface.cpp b/Engine/Source/Editor/Sequencer/Private/ISequencerNumericTypeInterface.cpp
--- a/Engine/Source/Editor/Sequencer/Private/ISequencerNumericTypeInterface.cpp
+++ b/Engine/Source/Editor/Sequencer/Private/ISequencerNumericTypeInterface.cpp
@@ -17,11 +17,12 @@ int32 FSequencerNumericTypeInterface::GetRelevancyScore(const ISequencer& Sequen
 		// Don't use non-duration intents for duration properties and visa-versa
 		if ((Intent == ENumericIntent::Duration) == bIsDurationProperty)
 		{
-			return 0;
+			return 100;
 		}
 	}
 
-	return Intent == ENumericIntent::Position ? 0 : MIN_int32;
+	// By default, prefer position type interfaces.
+	return Intent == ENumericIntent::Position ? 10 : 0;
 }
 
 } // namespace UE::Sequencer
\ No newline at end of file
