Sequencer: DeleteKeys for ISequencerChannelInterface

This allows channels to customize deletion, ie. set the default value at the current time if all keys are removed

diff --git a/Engine/Plugins/Experimental/PropertyAnimator/Source/PropertyAnimatorEditor/Private/Sequencer/PropertyAnimatorEditorCurveChannelInterface.h b/Engine/Plugins/Experimental/PropertyAnimator/Source/PropertyAnimatorEditor/Private/Sequencer/PropertyAnimatorEditorCurveChannelInterface.h
--- a/Engine/Plugins/Experimental/PropertyAnimator/Source/PropertyAnimatorEditor/Private/Sequencer/PropertyAnimatorEditorCurveChannelInterface.h
+++ b/Engine/Plugins/Experimental/PropertyAnimator/Source/PropertyAnimatorEditor/Private/Sequencer/PropertyAnimatorEditorCurveChannelInterface.h
@@ -23,6 +23,7 @@ struct TPropertyAnimatorEditorCurveChannelInterface : ISequencerChannelInterface
 
 	//~ Begin ISequencerChannelInterface
 	virtual FKeyHandle AddOrUpdateKey_Raw(FMovieSceneChannel*, UMovieSceneSection*, const void*, FFrameNumber, ISequencer&, const FGuid&, FTrackInstancePropertyBindings*) const override { return FKeyHandle::Invalid(); }
+	virtual void DeleteKeys_Raw(FMovieSceneChannel*, TConstArrayView<FKeyHandle>, FFrameNumber) const override {}
 	virtual void CopyKeys_Raw(FMovieSceneChannel*, const UMovieSceneSection*, FName, FMovieSceneClipboardBuilder&, TConstArrayView<FKeyHandle>) const override {}
 	virtual void PasteKeys_Raw(FMovieSceneChannel*, UMovieSceneSection*, const FMovieSceneClipboardKeyTrack&, const FMovieSceneClipboardEnvironment&, const FSequencerPasteEnvironment&, TArray<FKeyHandle>&) const override {}
 	virtual TSharedPtr<FStructOnScope> GetKeyStruct_Raw(FMovieSceneChannelHandle, FKeyHandle) const override { return nullptr; }
diff --git a/Engine/Source/Editor/MovieSceneTools/Public/Channels/PerlinNoiseChannelInterface.h b/Engine/Source/Editor/MovieSceneTools/Public/Channels/PerlinNoiseChannelInterface.h
--- a/Engine/Source/Editor/MovieSceneTools/Public/Channels/PerlinNoiseChannelInterface.h
+++ b/Engine/Source/Editor/MovieSceneTools/Public/Channels/PerlinNoiseChannelInterface.h
@@ -86,6 +86,10 @@ struct TPerlinNoiseChannelInterface : ISequencerChannelInterface
 		return FKeyHandle::Invalid();
 	}
 
+	virtual void DeleteKeys_Raw(FMovieSceneChannel* Channel, TArrayView<const FKeyHandle> InHandles, FFrameNumber InTime) const override
+	{
+	}
+
 	virtual void CopyKeys_Raw(FMovieSceneChannel* Channel, const UMovieSceneSection* Section, FName KeyAreaName, FMovieSceneClipboardBuilder& ClipboardBuilder, TArrayView<const FKeyHandle> KeyMask) const override
 	{
 	}
diff --git a/Engine/Source/Editor/Sequencer/Private/IKeyArea.cpp b/Engine/Source/Editor/Sequencer/Private/IKeyArea.cpp
--- a/Engine/Source/Editor/Sequencer/Private/IKeyArea.cpp
+++ b/Engine/Source/Editor/Sequencer/Private/IKeyArea.cpp
@@ -113,6 +113,17 @@ FKeyHandle IKeyArea::AddOrUpdateKey(FFrameNumber Time, const FGuid& ObjectBindin
 	return FKeyHandle();
 }
 
+void IKeyArea::DeleteKeys(TArrayView<const FKeyHandle> InHandles, FFrameNumber InTime)
+{
+	ISequencerChannelInterface* EditorInterface = FindChannelEditorInterface();
+	FMovieSceneChannel* Channel = ChannelHandle.Get();
+
+	if (Channel)
+	{
+		EditorInterface->DeleteKeys_Raw(Channel, InHandles, InTime);
+	}
+}
+
 FKeyHandle IKeyArea::DuplicateKey(FKeyHandle InKeyHandle) const
 {
 	FKeyHandle NewHandle = FKeyHandle::Invalid();
diff --git a/Engine/Source/Editor/Sequencer/Private/Sequencer.cpp b/Engine/Source/Editor/Sequencer/Private/Sequencer.cpp
--- a/Engine/Source/Editor/Sequencer/Private/Sequencer.cpp
+++ b/Engine/Source/Editor/Sequencer/Private/Sequencer.cpp
@@ -1891,6 +1891,8 @@ void FSequencer::DeleteSections(const TSet<UMovieSceneSection*>& Sections)
 
 void FSequencer::DeleteSelectedKeys()
 {
+	using namespace UE::Sequencer;
+
 	FScopedTransaction DeleteKeysTransaction( NSLOCTEXT("Sequencer", "DeleteSelectedKeys_Transaction", "Delete Selected Keys") );
 	bool bAnythingRemoved = false;
 
@@ -1904,29 +1906,42 @@ void FSequencer::DeleteSelectedKeys()
 			continue;
 		}
 
-		FMovieSceneChannel* Channel = ChannelInfo.Channel.Get();
-		if (Channel)
+		if (ChannelInfo.KeyHandles.Num() == 0)
 		{
-			UObject* Owner = nullptr;
-			const FMovieSceneChannelMetaData* ChannelMetaData = ChannelInfo.Channel.GetMetaData();
-			if (ChannelMetaData)
-			{
-				Owner = ChannelMetaData->WeakOwningObject.Get();
-			}
-			if (!Owner)
-			{
-				Owner = ChannelInfo.OwningSection;
-			}
+			continue;
+		}
 
-			if (!ModifiedObjects.Contains(Owner))
-			{
-				Owner->Modify();
-				ModifiedObjects.Add(ChannelInfo.OwningSection);
-			}
+		TSharedPtr<FChannelModel> Channel = ViewModel->GetSelection()->KeySelection.GetModelForKey(ChannelInfo.KeyHandles[0]);
+		if (!Channel)
+		{
+			continue;
+		}
 
-			Channel->DeleteKeys(ChannelInfo.KeyHandles);
-			bAnythingRemoved = true;
+		TSharedPtr<IKeyArea> KeyArea = Channel->GetKeyArea();
+		if (!KeyArea)
+		{
+			continue;
 		}
+
+		UObject* Owner = nullptr;
+		const FMovieSceneChannelMetaData* ChannelMetaData = ChannelInfo.Channel.GetMetaData();
+		if (ChannelMetaData)
+		{
+			Owner = ChannelMetaData->WeakOwningObject.Get();
+		}
+		if (!Owner)
+		{
+			Owner = ChannelInfo.OwningSection;
+		}
+
+		if (!ModifiedObjects.Contains(Owner))
+		{
+			Owner->Modify();
+			ModifiedObjects.Add(ChannelInfo.OwningSection);
+		}
+
+		KeyArea->DeleteKeys(ChannelInfo.KeyHandles, GetLocalTime().Time.FloorToFrame());
+		bAnythingRemoved = true;
 	}
 
 	if (bAnythingRemoved)
diff --git a/Engine/Source/Editor/Sequencer/Public/IKeyArea.h b/Engine/Source/Editor/Sequencer/Public/IKeyArea.h
--- a/Engine/Source/Editor/Sequencer/Public/IKeyArea.h
+++ b/Engine/Source/Editor/Sequencer/Public/IKeyArea.h
@@ -161,6 +161,15 @@ public:
 	 */
 	FKeyHandle AddOrUpdateKey(FFrameNumber Time, const FGuid& ObjectBindingID, ISequencer& InSequencer);
 
+	/**
+	 * Delete the specified keys. If all keys are removed, the current value at that time will be set as the default value for the channel
+	 * 
+	 * @param Channel               The channel to remove the keys from
+	 * @param InHandles             The key handles to delete
+	 * @param InTime                The time at which to evaluate for the default value if there are no keys remaining
+	 */
+	void DeleteKeys(TArrayView<const FKeyHandle> InHandles, FFrameNumber InTime);
+
 	/**
 	 * Duplicate the key represented by the specified handle
 	 *
diff --git a/Engine/Source/Editor/Sequencer/Public/ISequencerChannelInterface.h b/Engine/Source/Editor/Sequencer/Public/ISequencerChannelInterface.h
--- a/Engine/Source/Editor/Sequencer/Public/ISequencerChannelInterface.h
+++ b/Engine/Source/Editor/Sequencer/Public/ISequencerChannelInterface.h
@@ -97,6 +97,15 @@ struct ISequencerChannelInterface
 	 */
 	virtual FKeyHandle AddOrUpdateKey_Raw(FMovieSceneChannel* Channel, UMovieSceneSection* SectionToKey, const void* ExtendedEditorData, FFrameNumber InTime, ISequencer& InSequencer, const FGuid& ObjectBindingID, FTrackInstancePropertyBindings* PropertyBindings) const = 0;
 
+	/**
+	 * Delete the specified keys. If all keys are removed, the current value at that time will be set as the default value for the channel
+	 * 
+	 * @param Channel               The channel to remove the keys from
+	 * @param InHandles             The key handles to delete
+	 * @param InTime                The time at which to evaluate for the default value if there are no keys remaining
+	 */ 
+	virtual void DeleteKeys_Raw(FMovieSceneChannel* Channel, TArrayView<const FKeyHandle> InHandles, FFrameNumber InTime) const = 0;
+
 	/**
 	 * Copy all the keys specified in KeyMask to the specified clipboard
 	 *
diff --git a/Engine/Source/Editor/Sequencer/Public/SequencerChannelInterface.h b/Engine/Source/Editor/Sequencer/Public/SequencerChannelInterface.h
--- a/Engine/Source/Editor/Sequencer/Public/SequencerChannelInterface.h
+++ b/Engine/Source/Editor/Sequencer/Public/SequencerChannelInterface.h
@@ -59,6 +59,18 @@ void CreateKeyEditorDeprecatedSignature()
 template<typename ChannelType>
 struct TSequencerChannelInterfaceCommon : ISequencerChannelInterface
 {
+	/**
+	 * Delete the specified keys. If all keys are removed, the current value at that time will be set as the default value for the channel
+	 *
+	 * @param Channel               The channel to remove the keys from
+	 * @param InHandles             The key handles to delete
+	 * @param InTime                The time at which to evaluate for the default value if there are no keys remaining
+	 */
+	virtual void DeleteKeys_Raw(FMovieSceneChannel* InChannel, TArrayView<const FKeyHandle> InHandles, FFrameNumber InTime) const override
+	{
+		using namespace Sequencer;
+		DeleteKeys(static_cast<ChannelType*>(InChannel), InHandles, InTime);
+	}
 
 	/**
 	 * Copy all the keys specified in KeyMask to the specified clipboard
diff --git a/Engine/Source/Editor/Sequencer/Public/SequencerChannelTraits.h b/Engine/Source/Editor/Sequencer/Public/SequencerChannelTraits.h
--- a/Engine/Source/Editor/Sequencer/Public/SequencerChannelTraits.h
+++ b/Engine/Source/Editor/Sequencer/Public/SequencerChannelTraits.h
@@ -334,6 +334,37 @@ namespace Sequencer
 	 */
 	SEQUENCER_API int32 DrawExtra(FMovieSceneChannel* InChannel, const UMovieSceneSection* InOwner, const FSequencerChannelPaintArgs& PaintArgs, int32 LayerId);
 
+	/**
+	 * Delete the specified keys. If all keys are removed, the current value at that time will be set as the default value for the channel
+	 *
+	 * @param Channel               The channel to remove the keys from
+	 * @param InHandles             The key handles to delete
+	 * @param InTime                The time at which to evaluate for the default value if there are no keys remaining
+	 */
+	template<typename ChannelType>
+	typename TEnableIf<TMovieSceneChannelTraits<ChannelType>::SupportsDefaults>::Type DeleteKeys(ChannelType* InChannel, TArrayView<const FKeyHandle> InHandles, FFrameNumber InTime)
+	{
+		using namespace UE::MovieScene;
+
+		using CurveValueType = typename ChannelType::CurveValueType;
+		CurveValueType Value;
+
+		EvaluateChannel(InChannel, InTime, Value);
+
+		InChannel->DeleteKeys(InHandles);
+
+		if (InChannel->GetNumKeys() == 0)
+		{
+			InChannel->SetDefault(Value);
+		}
+	}
+
+	template<typename ChannelType>
+	typename TEnableIf<!TMovieSceneChannelTraits<ChannelType>::SupportsDefaults>::Type DeleteKeys(ChannelType* InChannel, TArrayView<const FKeyHandle> InHandles, FFrameNumber InTime)
+	{
+		InChannel->DeleteKeys(InHandles);
+	}
+
 	/**
 	 * Copy the specified keys from a channel
 	 *
