[Sequencer] Don't check track object when filtering for "Name" text expression or component filters

diff --git a/Engine/Source/Editor/Sequencer/Private/Filters/TextExpressions/SequencerTextFilterExpression_Name.cpp b/Engine/Source/Editor/Sequencer/Private/Filters/TextExpressions/SequencerTextFilterExpression_Name.cpp
--- a/Engine/Source/Editor/Sequencer/Private/Filters/TextExpressions/SequencerTextFilterExpression_Name.cpp
+++ b/Engine/Source/Editor/Sequencer/Private/Filters/TextExpressions/SequencerTextFilterExpression_Name.cpp
@@ -39,30 +39,11 @@ bool FSequencerTextFilterExpression_Name::TestComplexExpression(const FName& InK
 		return true;
 	}
 
-	if (WeakTrackObject.IsValid())
-	{
-		if (TextFilterUtils::TestComplexExpression(WeakTrackObject->GetName(), InValue, InComparisonOperation, InTextComparisonMode))
-		{
-			return true;
-		}
-
-		UMovieSceneNameableTrack* const NameableTrack = Cast<UMovieSceneNameableTrack>(WeakTrackObject);
-		if (IsValid(NameableTrack))
-		{
-			const FString DisplayName = NameableTrack->GetDisplayName().ToString();
-			if (TextFilterUtils::TestComplexExpression(DisplayName, InValue, InComparisonOperation, InTextComparisonMode))
-			{
-				return true;
-			}
-		}
-	}
-
 	if (const TViewModelPtr<IOutlinerExtension> OutlinerExtension = FilterItem.ImplicitCast())
 	{
-		if (TextFilterUtils::TestComplexExpression(OutlinerExtension->GetLabel().ToString(), InValue, InComparisonOperation, InTextComparisonMode))
-		{
-			return true;
-		}
+		const FString Label = OutlinerExtension->GetLabel().ToString();
+		const bool bPassed = TextFilterUtils::TestComplexExpression(Label, InValue, InComparisonOperation, InTextComparisonMode);
+		return bPassed;
 	}
 
 	return false;
diff --git a/Engine/Source/Editor/Sequencer/Public/Filters/SequencerTrackFilterBase.h b/Engine/Source/Editor/Sequencer/Public/Filters/SequencerTrackFilterBase.h
--- a/Engine/Source/Editor/Sequencer/Public/Filters/SequencerTrackFilterBase.h
+++ b/Engine/Source/Editor/Sequencer/Public/Filters/SequencerTrackFilterBase.h
@@ -183,32 +183,27 @@ public:
 	{
 		FSequencerFilterData& FilterData = FilterInterface.GetFilterData();
 
-		const UMovieSceneTrack* const TrackObject = FilterData.ResolveMovieSceneTrackObject(InItem);
-		if (IsValid(TrackObject) && TrackObject->IsA(InComponentType::StaticClass()))
+		UObject* const BoundObject = FilterData.ResolveTrackBoundObject(GetSequencer(), InItem);
+		if (!IsValid(BoundObject))
 		{
-			return true;
+			return false;
 		}
 
-		const TWeakObjectPtr<> BoundObject = FilterData.ResolveTrackBoundObject(GetSequencer(), InItem);
-		if (BoundObject.IsValid() && BoundObject->IsA(InComponentType::StaticClass()))
+		if (BoundObject->IsA(InComponentType::StaticClass()))
 		{
 			return true;
 		}
 
-		const AActor* const Actor = Cast<const AActor>(BoundObject.Get());
-		if (IsValid(Actor) && Actor->FindComponentByClass(InComponentType::StaticClass()))
+		if (BoundObject->IsA(AActor::StaticClass()))
 		{
-			return true;
+			const AActor* const Actor = CastChecked<const AActor>(BoundObject);
+			if (Actor->FindComponentByClass(InComponentType::StaticClass()))
+			{
+				return true;
+			}
 		}
 
 		return false;
 	}
 	//~ End IFilter
-
-	//~ Begin FSequencerTrackFilter
-	TSubclassOf<UMovieSceneTrack> GetTrackClass() const override
-	{
-		return InComponentType::StaticClass();
-	}
-	//~ End FSequencerTrackFilter
 };
