[Sequencer] Limited Selectability:
- Fix selectability gizmo hitch when selecting axis

diff --git a/Engine/Source/Editor/Sequencer/Private/SequencerEdMode.cpp b/Engine/Source/Editor/Sequencer/Private/SequencerEdMode.cpp
--- a/Engine/Source/Editor/Sequencer/Private/SequencerEdMode.cpp
+++ b/Engine/Source/Editor/Sequencer/Private/SequencerEdMode.cpp
@@ -1065,8 +1065,6 @@ bool FSequencerEdMode::IsViewportSelectionLimited() const
 void FSequencerEdMode::EnableSelectabilityTool(const bool bInEnabled)
 {
 	SelectabilityTool->EnableLimitedSelection(bInEnabled);
-
-	SetCurrentTool(bInEnabled ? SelectabilityTool->GetID() : DefaultTool->GetID());
 }
 
 bool FSequencerEdMode::IsObjectSelectableInViewport(UObject* const InObject) const
