SequenceUtilities: Fix an issue when pasting bindings that were copied from a cooked sequence.

The paste code relied on some editor only data to determine if the binding should be a child of the binding that was selected when the paste action was executed by seeing if the binding to paste would resolve to an actor. This change adds a fallback which manually checks for actors in the world that the binding resolves to even without the missing data.

diff --git a/Engine/Source/Editor/Sequencer/Private/SequencerUtilities.cpp b/Engine/Source/Editor/Sequencer/Private/SequencerUtilities.cpp
--- a/Engine/Source/Editor/Sequencer/Private/SequencerUtilities.cpp
+++ b/Engine/Source/Editor/Sequencer/Private/SequencerUtilities.cpp
@@ -3175,16 +3175,6 @@ bool FSequencerUtilities::PasteBindings(const FString& TextToImport, TSharedRef<
 					}
 				}
 
-				// Set the parent to the pasted target only if it's not an actor
-				const UClass* PossessedObjectClass = CopyableBinding->Possessable.GetPossessedObjectClass();
-				if (PossessedObjectClass && !PossessedObjectClass->IsChildOf(AActor::StaticClass()))
-				{
-					if (TargetIndex < PasteBindingsParams.Bindings.Num())
-					{
-						Possessable->SetParent(PasteBindingsParams.Bindings[TargetIndex].BindingID, MovieScene);
-					}
-				}
-
 				// Find the objects that this pasted binding should bind to
 				TArray<UObject*> ObjectsToBind;
 
@@ -3213,6 +3203,29 @@ bool FSequencerUtilities::PasteBindings(const FString& TextToImport, TSharedRef<
 					}
 				}
 
+				bool bSetParent = false;
+				if (const UClass* PossessedObjectClass = CopyableBinding->Possessable.GetPossessedObjectClass())
+				{
+					if (!PossessedObjectClass->IsChildOf(AActor::StaticClass()))
+					{
+						// Attempt to set the parent to be the paste target only if the possessed object class is not an actor
+						bSetParent = true;
+					}
+				}
+				else if (ObjectsToBind.IsEmpty())
+				{
+					// Attempt to set the parent to be the paste target only if the binding does not resolve to an actor in the world.
+					bSetParent = true;
+				}
+				
+				if (bSetParent)
+				{
+					if (TargetIndex < PasteBindingsParams.Bindings.Num())
+					{
+						Possessable->SetParent(PasteBindingsParams.Bindings[TargetIndex].BindingID, MovieScene);
+					}
+				}
+				
 				if (ObjectsToBind.Num() != 0)
 				{
 					if (PasteBindingsParams.bDuplicateExistingActors)
