Sequencer: Remove viewport UI hiding

The new viewport toolbar is not an overlay so it causes a pop when hiding it since it pushes the viewport down. There's not much we can do here except just get rid of hiding the viewport toolbar.

#rnx

diff --git a/Engine/Source/Editor/Sequencer/Private/Sequencer.cpp b/Engine/Source/Editor/Sequencer/Private/Sequencer.cpp
--- a/Engine/Source/Editor/Sequencer/Private/Sequencer.cpp
+++ b/Engine/Source/Editor/Sequencer/Private/Sequencer.cpp
@@ -4940,8 +4940,6 @@ FSequencer::FCachedViewState::StoreViewState()
 	}
 
 	bValid = true;
-	bIsViewportUIHidden = GLevelEditorModeTools().IsViewportUIHidden();
-	GLevelEditorModeTools().SetHideViewportUI(!GLevelEditorModeTools().IsViewportUIHidden());
 
 	GameViewStates.Empty();
 	for (int32 ViewIndex = 0; ViewIndex < GEditor->GetLevelViewportClients().Num(); ++ViewIndex)
@@ -4964,7 +4962,6 @@ FSequencer::FCachedViewState::RestoreViewState()
 	}
 
 	bValid = false;
-	GLevelEditorModeTools().SetHideViewportUI(bIsViewportUIHidden);
 
 	for (int32 Index = 0; Index < GameViewStates.Num(); ++Index)
 	{
diff --git a/Engine/Source/Editor/Sequencer/Private/Sequencer.h b/Engine/Source/Editor/Sequencer/Private/Sequencer.h
--- a/Engine/Source/Editor/Sequencer/Private/Sequencer.h
+++ b/Engine/Source/Editor/Sequencer/Private/Sequencer.h
@@ -1578,8 +1578,7 @@ private:
 	struct FCachedViewState
 	{
 		FCachedViewState()
-			: bValid(false)
-			, bIsViewportUIHidden(false) {}
+			: bValid(false) {}
 
 	public:
 		void StoreViewState();
@@ -1587,7 +1586,6 @@ private:
 
 	private:
 		bool bValid;
-		bool bIsViewportUIHidden;
 		TArray<TPair<int32, bool> > GameViewStates;
 	};
 
