Sequencer: Update "Rerun Construction Scripts" menu to "Allow Construction Scripts to Rerun Every Frame" and clarify tooltip as "Blueprint Actors need to have Run Construction Script in Sequencer active in their Class Settings"

diff --git a/Engine/Source/Editor/Sequencer/Private/SequencerCommands.cpp b/Engine/Source/Editor/Sequencer/Private/SequencerCommands.cpp
--- a/Engine/Source/Editor/Sequencer/Private/SequencerCommands.cpp
+++ b/Engine/Source/Editor/Sequencer/Private/SequencerCommands.cpp
@@ -49,7 +49,7 @@ void FSequencerCommands::RegisterCommands()
 
 	UI_COMMAND( TogglePlaybackRangeLocked, "Playback Range Locked", "Prevent editing the start and end times for the sequence.", EUserInterfaceActionType::ToggleButton, FInputChord() );
 	UI_COMMAND( ToggleCleanPlaybackMode, "Game View (Clean Playback Mode)", "Enable game view and hide viewport buttons while playing.", EUserInterfaceActionType::ToggleButton, FInputChord() );
-	UI_COMMAND( ToggleRerunConstructionScripts, "Rerun Construction Scripts", "Rerun construction scripts on bound actors every frame.", EUserInterfaceActionType::ToggleButton, FInputChord() );
+	UI_COMMAND( ToggleRerunConstructionScripts, "Allow Construction Scripts to Rerun Every Frame", "Blueprint Actors need to have Run Construction Script in Sequencer active in their Class Settings.", EUserInterfaceActionType::ToggleButton, FInputChord() );
 	UI_COMMAND( ToggleAsyncEvaluation, "Async Evaluation", "When enabled, enables a single asynchronous evaluation once per-frame. When disabled, forces a full blocking evaluation every time this sequence is evaluated (should be avoided for real-time content).", EUserInterfaceActionType::ToggleButton, FInputChord() );
 	UI_COMMAND( ToggleDynamicWeighting, "Dynamic Weighting", "When enabled, all blendable tracks will cache their initial values to ensure that they are able to correctly blend in/out when dynamic weights are being used.", EUserInterfaceActionType::ToggleButton, FInputChord() );
 
