Sequencer: Rebuild the track model's sections only if necessary, ie. if the section model has changed or if the channel proxy has changed. This fixes performance issues especially when dragging keys, which calls modify and ultimately rebuilds everything (deferred or not).

The section model now stores the channel proxy that was used to layout the section model. This is used to determine whether the section model needs to be rebuilt.

diff --git a/Engine/Source/Editor/Sequencer/Private/MVVM/ViewModels/TrackModel.cpp b/Engine/Source/Editor/Sequencer/Private/MVVM/ViewModels/TrackModel.cpp
--- a/Engine/Source/Editor/Sequencer/Private/MVVM/ViewModels/TrackModel.cpp
+++ b/Engine/Source/Editor/Sequencer/Private/MVVM/ViewModels/TrackModel.cpp
@@ -201,8 +201,8 @@ void FTrackModel::ForceUpdate()
 
 		FScopedViewModelListHead RecycledModels(AsShared(), EViewModelListType::Recycled);
 		GetChildrenForList(&SectionList).MoveChildrenTo<IRecyclableExtension>(RecycledModels.GetChildren(), IRecyclableExtension::CallOnRecycle);
-		OutlinerChildren.MoveChildrenTo<IRecyclableExtension>(RecycledModels.GetChildren(), IRecyclableExtension::CallOnRecycle);
-		TopLevelChannelChildren.MoveChildrenTo<IRecyclableExtension>(RecycledModels.GetChildren(), IRecyclableExtension::CallOnRecycle);
+
+		bool bNeedsLayout = false;
 
 		// Add all sections directly to this track row
 		for (UMovieSceneSection* Section : Track->GetAllSections())
@@ -212,30 +212,43 @@ void FTrackModel::ForceUpdate()
 			{
 				TSharedRef<ISequencerSection> SectionInterface = TrackEditor->MakeSectionInterface(*Section, *Track, ObjectBinding);
 				SectionModel = SectionModelStorage->CreateModelForSection(Section, SectionInterface);
+				bNeedsLayout = true;
 			}
 
 			if (ensure(SectionModel))
 			{
+				bNeedsLayout |= SectionModel->NeedsLayout();
+
 				// Move the child back into the real section list
 				SectionChildren.InsertChild(SectionModel, SectionsTail);
 				SectionsTail = SectionModel;
 			}
 		}
 
-		// Rebuild the outliner layout for this track. This will clear our children and rebuild them if needed
-		// (with potentially recycled children), so if we went from, say, 2 rows to 1 row, it should correctly
-		// discard any children we don't need anymore.
-		FTrackModelLayoutBuilder LayoutBuilder(AsShared());
+		// If we are discarding any sections (because they still remain in the recycled list) we must run the layout
+		bNeedsLayout |= !RecycledModels.GetChildren().IsEmpty();
 
-		for (TSharedPtr<FSectionModel> Section : TViewModelListIterator<FSectionModel>(&SectionList))
+		if (bNeedsLayout)
 		{
-			LayoutBuilder.RefreshLayout(Section);
-		}
 
-		if (OutlinerChildren.IsEmpty())
-		{
-			// Reset expansion state if this track can no longer be expanded
-			SetExpansion(false);
+			OutlinerChildren.MoveChildrenTo<IRecyclableExtension>(RecycledModels.GetChildren(), IRecyclableExtension::CallOnRecycle);
+			TopLevelChannelChildren.MoveChildrenTo<IRecyclableExtension>(RecycledModels.GetChildren(), IRecyclableExtension::CallOnRecycle);
+
+			// Rebuild the outliner layout for this track. This will clear our children and rebuild them if needed
+			// (with potentially recycled children), so if we went from, say, 2 rows to 1 row, it should correctly
+			// discard any children we don't need anymore.
+			FTrackModelLayoutBuilder LayoutBuilder(AsShared());
+
+			for (TSharedPtr<FSectionModel> Section : TViewModelListIterator<FSectionModel>(&SectionList))
+			{
+				LayoutBuilder.RefreshLayout(Section);
+			}
+
+			if (OutlinerChildren.IsEmpty())
+			{
+				// Reset expansion state if this track can no longer be expanded
+				SetExpansion(false);
+			}
 		}
 	}
 	else
@@ -262,6 +275,7 @@ void FTrackModel::ForceUpdate()
 			TSharedPtr<FTrackRowModel> Row;
 			TSharedPtr<FViewModel> SectionsTail;
 			TUniquePtr<FScopedViewModelListHead> RecycledModels;
+			bool bNeedsLayout = false;
 		};
 		TArray<FRowData, TInlineAllocator<8>> RowModels;
 		RowModels.SetNum(PopulatedRows.Num());
@@ -276,6 +290,7 @@ void FTrackModel::ForceUpdate()
 			if (!TrackRowModel)
 			{
 				TrackRowModel = TrackRowModelStorage->CreateModelForTrackRow(Track, RowIndex);
+				RowModels[RowIndex].bNeedsLayout = true;
 			}
 
 			if (ensure(TrackRowModel))
@@ -289,8 +304,6 @@ void FTrackModel::ForceUpdate()
 				RowModels[RowIndex].RecycledModels = MakeUnique<FScopedViewModelListHead>(TrackRowModel, EViewModelListType::Recycled);
 				FViewModelChildren RecycledRowModels = RowModels[RowIndex].RecycledModels->GetChildren();
 				TrackRowModel->GetSectionModels().MoveChildrenTo<IRecyclableExtension>(RecycledRowModels, IRecyclableExtension::CallOnRecycle);
-				TrackRowModel->GetChildList(EViewModelListType::Outliner).MoveChildrenTo<IRecyclableExtension>(RecycledRowModels, IRecyclableExtension::CallOnRecycle);
-				TrackRowModel->GetTopLevelChannels().MoveChildrenTo<IRecyclableExtension>(RecycledRowModels, IRecyclableExtension::CallOnRecycle);
 			}
 		}
 
@@ -304,6 +317,11 @@ void FTrackModel::ForceUpdate()
 			{
 				TSharedRef<ISequencerSection> SectionInterface = TrackEditor->MakeSectionInterface(*Section, *Track, ObjectBinding);
 				SectionModel = SectionModelStorage->CreateModelForSection(Section, SectionInterface);
+				RowModels[RowIndex].bNeedsLayout = true;
+			}
+			else if (SectionModel)
+			{
+				RowModels[RowIndex].bNeedsLayout |= SectionModel->NeedsLayout();
 			}
 
 			RowModels[RowIndex].Row->GetSectionModels().InsertChild(SectionModel, RowModels[RowIndex].SectionsTail);
@@ -311,10 +329,15 @@ void FTrackModel::ForceUpdate()
 		}
 
 		// Rebuild the outliner layout for each track row
-		for (const FRowData& RowData : RowModels)
+		for (FRowData& RowData : RowModels)
 		{
-			if (RowData.Row)
+			const bool bStillHasRecycledChildren = RowData.RecycledModels && !RowData.RecycledModels->GetChildren().IsEmpty();
+			if (RowData.Row && (RowData.bNeedsLayout || bStillHasRecycledChildren))
 			{
+				FViewModelChildren RecycledRowModels = RowData.RecycledModels->GetChildren();
+				RowData.Row->GetChildList(EViewModelListType::Outliner).MoveChildrenTo<IRecyclableExtension>(RecycledRowModels, IRecyclableExtension::CallOnRecycle);
+				RowData.Row->GetTopLevelChannels().MoveChildrenTo<IRecyclableExtension>(RecycledRowModels, IRecyclableExtension::CallOnRecycle);
+
 				FTrackModelLayoutBuilder LayoutBuilder(RowData.Row->AsShared());
 				for (TSharedPtr<FSectionModel> Section : RowData.Row->GetChildrenOfType<FSectionModel>(EViewModelListType::TrackArea))
 				{
diff --git a/Engine/Source/Editor/Sequencer/Private/MVVM/ViewModels/TrackModelLayoutBuilder.cpp b/Engine/Source/Editor/Sequencer/Private/MVVM/ViewModels/TrackModelLayoutBuilder.cpp
--- a/Engine/Source/Editor/Sequencer/Private/MVVM/ViewModels/TrackModelLayoutBuilder.cpp
+++ b/Engine/Source/Editor/Sequencer/Private/MVVM/ViewModels/TrackModelLayoutBuilder.cpp
@@ -4,6 +4,7 @@
 
 #include "Channels/MovieSceneChannelEditorData.h"
 #include "Channels/MovieSceneChannelHandle.h"
+#include "Channels/MovieSceneChannelProxy.h"
 #include "Containers/UnrealString.h"
 #include "HAL/Platform.h"
 #include "ISequencerSection.h"
@@ -154,6 +155,14 @@ FTrackModelLayoutBuilder::~FTrackModelLayoutBuilder()
 void FTrackModelLayoutBuilder::RefreshLayout(TSharedPtr<FSectionModel> InSection)
 {
 	InSection->SetLinkedOutlinerItem(CastViewModelChecked<IOutlinerExtension>(Root));
+	if (UMovieSceneSection* SectionObject = InSection->GetSection())
+	{
+		InSection->SetLayoutChannelProxy(SectionObject->GetChannelProxy().AsWeak());
+	}
+	else
+	{
+		InSection->SetLayoutChannelProxy(nullptr);
+	}
 
 	// Reset everything
 	OutlinerList.Reset();
diff --git a/Engine/Source/Editor/Sequencer/Public/MVVM/ViewModels/SectionModel.h b/Engine/Source/Editor/Sequencer/Public/MVVM/ViewModels/SectionModel.h
--- a/Engine/Source/Editor/Sequencer/Public/MVVM/ViewModels/SectionModel.h
+++ b/Engine/Source/Editor/Sequencer/Public/MVVM/ViewModels/SectionModel.h
@@ -20,6 +20,7 @@
 
 class ISequencerSection;
 class UMovieSceneSection;
+struct FMovieSceneChannelProxy;
 
 namespace UE
 {
@@ -90,6 +91,18 @@ public:
 	int32 GetPreRollFrames() const;
 	int32 GetPostRollFrames() const;
 
+	/** Returns whether this section model needs to be rebuilt, ie. the channel proxy is no longer valid */
+	bool NeedsLayout() const
+	{
+		return !PreviousLayoutChannelProxy.Pin().IsValid();
+	}
+
+	/** Set the channel proxy that was most recently used to layout this section model */
+	void SetLayoutChannelProxy(TWeakPtr<FMovieSceneChannelProxy> InPreviousLayoutChannelProxy)
+	{
+		PreviousLayoutChannelProxy = InPreviousLayoutChannelProxy;
+	}
+
 public:
 
 	/*~ ILayerBarExtension Interface */
@@ -139,6 +152,8 @@ private:
 	TWeakObjectPtr<UMovieSceneSection> WeakSection;
 	TRange<FFrameNumber> SectionRange;
 	TRange<FFrameNumber> LayerBarRange;
+
+	TWeakPtr<FMovieSceneChannelProxy> PreviousLayoutChannelProxy;
 };
 
 struct FOverlappingSections
