Sequencer: Check row index for NeedsLayout

This is a follow up fix to 38754753 which only rebuilds track models when necessary. This adds a case for when the section's row index changes but it's channel proxy does not.

#rnx

diff --git a/Engine/Source/Editor/Sequencer/Private/MVVM/ViewModels/TrackModelLayoutBuilder.cpp b/Engine/Source/Editor/Sequencer/Private/MVVM/ViewModels/TrackModelLayoutBuilder.cpp
--- a/Engine/Source/Editor/Sequencer/Private/MVVM/ViewModels/TrackModelLayoutBuilder.cpp
+++ b/Engine/Source/Editor/Sequencer/Private/MVVM/ViewModels/TrackModelLayoutBuilder.cpp
@@ -158,10 +158,12 @@ void FTrackModelLayoutBuilder::RefreshLayout(TSharedPtr<FSectionModel> InSection
 	if (UMovieSceneSection* SectionObject = InSection->GetSection())
 	{
 		InSection->SetLayoutChannelProxy(SectionObject->GetChannelProxy().AsWeak());
+		InSection->SetLayoutRowIndex(SectionObject->GetRowIndex());
 	}
 	else
 	{
 		InSection->SetLayoutChannelProxy(nullptr);
+		InSection->SetLayoutRowIndex(-1);
 	}
 
 	// Reset everything
diff --git a/Engine/Source/Editor/Sequencer/Public/MVVM/ViewModels/SectionModel.h b/Engine/Source/Editor/Sequencer/Public/MVVM/ViewModels/SectionModel.h
--- a/Engine/Source/Editor/Sequencer/Public/MVVM/ViewModels/SectionModel.h
+++ b/Engine/Source/Editor/Sequencer/Public/MVVM/ViewModels/SectionModel.h
@@ -3,6 +3,7 @@
 #pragma once
 
 #include "Channels/MovieSceneChannelHandle.h"
+#include "MovieSceneSection.h"
 #include "MVVM/ViewModels/ViewModel.h"
 #include "MVVM/ViewModels/ViewModelHierarchy.h"
 #include "MVVM/Extensions/ILayerBarExtension.h"
@@ -94,7 +95,7 @@ public:
 	/** Returns whether this section model needs to be rebuilt, ie. the channel proxy is no longer valid */
 	bool NeedsLayout() const
 	{
-		return !PreviousLayoutChannelProxy.Pin().IsValid();
+		return !PreviousLayoutChannelProxy.Pin().IsValid() || (GetSection() && GetSection()->GetRowIndex() != PreviousLayoutRowIndex);
 	}
 
 	/** Set the channel proxy that was most recently used to layout this section model */
@@ -103,6 +104,12 @@ public:
 		PreviousLayoutChannelProxy = InPreviousLayoutChannelProxy;
 	}
 
+	/** Set the row index that was most recently used to layout this section model */
+	void SetLayoutRowIndex(int32 InRowIndex)
+	{
+		PreviousLayoutRowIndex = InRowIndex;
+	}
+
 public:
 
 	/*~ ILayerBarExtension Interface */
@@ -154,6 +161,7 @@ private:
 	TRange<FFrameNumber> LayerBarRange;
 
 	TWeakPtr<FMovieSceneChannelProxy> PreviousLayoutChannelProxy;
+	int32 PreviousLayoutRowIndex = -1;
 };
 
 struct FOverlappingSections
