Sequencer: Fix copy/paste bindings so that it uses the world when building relative path names

Copy bindings/paste bindings stores and restores bound object names relative to a context, previously the resolution context. The resolution context could have been the world or a level sequence actor in the world. The copy/paste could fail because it depended upon the same set up for when you do the copy and paste. For example, if would fail if you did the copy when there is a level sequence actor and then the paste when there is no level sequence actor. This fix eliminates that by always resolving the world.

diff --git a/Engine/Source/Editor/Sequencer/Private/SequencerUtilities.cpp b/Engine/Source/Editor/Sequencer/Private/SequencerUtilities.cpp
--- a/Engine/Source/Editor/Sequencer/Private/SequencerUtilities.cpp
+++ b/Engine/Source/Editor/Sequencer/Private/SequencerUtilities.cpp
@@ -2409,6 +2409,8 @@ void FSequencerUtilities::CopyBindings(TSharedRef<ISequencer> Sequencer, const T
 
 void FSequencerUtilities::CopyBindings(TSharedRef<ISequencer> Sequencer, const TArray<FMovieSceneBindingProxy>& Bindings, const TArray<UMovieSceneFolder*>& InFolders, FOutputDevice& Ar)
 {
+	UWorld* World = GCurrentLevelEditingViewportClient ? GCurrentLevelEditingViewportClient->GetWorld() : nullptr;
+
 	TArray<UMovieSceneCopyableBinding*> Objects;
 	for (const FMovieSceneBindingProxy& ObjectBinding : Bindings)
 	{
@@ -2424,12 +2426,6 @@ void FSequencerUtilities::CopyBindings(TSharedRef<ISequencer> Sequencer, const T
 		FMovieScenePossessable* Possessable = MovieScene->FindPossessable(ObjectBinding.BindingID);
 		if (Possessable)
 		{
-			UObject* ResolutionContext = FindResolutionContext(Sequencer
-				, *MovieScene->GetTypedOuter<UMovieSceneSequence>()
-				, *MovieScene
-				, Possessable->GetParent()
-				, Sequencer->GetPlaybackContext());
-
 			CopyableBinding->Possessable = *Possessable;
 
 			// Store any custom bindings
@@ -2460,7 +2456,7 @@ void FSequencerUtilities::CopyBindings(TSharedRef<ISequencer> Sequencer, const T
 					}
 					else if (UObject* RuntimeObject = MovieSceneHelpers::GetSingleBoundObject(ObjectBinding.Sequence, ObjectBinding.BindingID, Sequencer->GetSharedPlaybackState(), BindingIndex))
 					{
-						CopyableBinding->BoundObjectNames.Add(RuntimeObject->GetPathName(ResolutionContext));
+						CopyableBinding->BoundObjectNames.Add(RuntimeObject->GetPathName(World));
 					}
 					BindingIndex++;
 				}
@@ -2471,7 +2467,7 @@ void FSequencerUtilities::CopyBindings(TSharedRef<ISequencer> Sequencer, const T
 
 				for (TWeakObjectPtr<> RuntimeObject : Sequencer->FindBoundObjects(CopyableBinding->Possessable.GetGuid(), Sequencer->GetFocusedTemplateID()))
 				{
-					CopyableBinding->BoundObjectNames.Add(RuntimeObject->GetPathName(ResolutionContext));
+					CopyableBinding->BoundObjectNames.Add(RuntimeObject->GetPathName(World));
 				}
 			}
 		}
@@ -3189,7 +3185,7 @@ bool FSequencerUtilities::PasteBindings(const FString& TextToImport, TSharedRef<
 					for (TActorIterator<AActor> ActorItr(World); ActorItr; ++ActorItr)
 					{
 						AActor* Actor = *ActorItr;
-						if (Actor && CopyableBinding->BoundObjectNames.Contains(Actor->GetPathName(ResolutionContext)))
+						if (Actor && CopyableBinding->BoundObjectNames.Contains(Actor->GetPathName(World)))
 						{
 							// If this actor is already bound and we're not duplicating actors, don't bind to anything
 							if (!PasteBindingsParams.bDuplicateExistingActors && Sequencer->FindObjectId(*Actor, Sequencer->GetFocusedTemplateID()).IsValid())
@@ -3198,7 +3194,7 @@ bool FSequencerUtilities::PasteBindings(const FString& TextToImport, TSharedRef<
 							}
 
 							ObjectsToBind.Add(Actor);
-							CopyableBinding->BoundObjectNames.Remove(Actor->GetPathName(ResolutionContext));
+							CopyableBinding->BoundObjectNames.Remove(Actor->GetPathName(World));
 						}
 					}
 				}
