Sequencer: Keep track of num rows for track layout. The previous fix 38855482 to check the section row index wasn't enough because the section row index may not change if you're adding/removing sections from a track.

Tested that deleting sections and merging sections with anim layer properly updates the track/section display
Tested that dragging keys is still fast

#rnx

diff --git a/Engine/Source/Editor/Sequencer/Private/MVVM/ViewModels/TrackModel.cpp b/Engine/Source/Editor/Sequencer/Private/MVVM/ViewModels/TrackModel.cpp
--- a/Engine/Source/Editor/Sequencer/Private/MVVM/ViewModels/TrackModel.cpp
+++ b/Engine/Source/Editor/Sequencer/Private/MVVM/ViewModels/TrackModel.cpp
@@ -202,7 +202,7 @@ void FTrackModel::ForceUpdate()
 		FScopedViewModelListHead RecycledModels(AsShared(), EViewModelListType::Recycled);
 		GetChildrenForList(&SectionList).MoveChildrenTo<IRecyclableExtension>(RecycledModels.GetChildren(), IRecyclableExtension::CallOnRecycle);
 
-		bool bNeedsLayout = false;
+		bool bNeedsLayout = NumRows != PreviousLayoutNumRows;
 
 		// Add all sections directly to this track row
 		for (UMovieSceneSection* Section : Track->GetAllSections())
@@ -286,6 +286,8 @@ void FTrackModel::ForceUpdate()
 		{
 			const int32 RowIndex = It.GetIndex();
 
+			RowModels[RowIndex].bNeedsLayout = NumRows != PreviousLayoutNumRows;
+
 			TSharedPtr<FTrackRowModel> TrackRowModel = TrackRowModelStorage->FindModelForTrackRow(Track, RowIndex);
 			if (!TrackRowModel)
 			{
@@ -348,6 +350,8 @@ void FTrackModel::ForceUpdate()
 			//       we fixup row indices
 		}
 	}
+
+	PreviousLayoutNumRows = NumRows;
 }
 
 FOutlinerSizing FTrackModel::GetOutlinerSizing() const
diff --git a/Engine/Source/Editor/Sequencer/Public/MVVM/ViewModels/TrackModel.h b/Engine/Source/Editor/Sequencer/Public/MVVM/ViewModels/TrackModel.h
--- a/Engine/Source/Editor/Sequencer/Public/MVVM/ViewModels/TrackModel.h
+++ b/Engine/Source/Editor/Sequencer/Public/MVVM/ViewModels/TrackModel.h
@@ -169,6 +169,8 @@ private:
 	TSharedPtr<ISequencerTrackEditor> TrackEditor;
 
 	bool bNeedsUpdate;
+
+	int32 PreviousLayoutNumRows = -1;
 };
 
 } // namespace Sequencer
