Fix crash when drawing transform origins in subsequence with multiple siblings.

Sections are now imported using the default interrogation channel index rather than the row index.

#rnx

diff --git a/Engine/Source/Editor/Sequencer/Private/SequencerEdMode.cpp b/Engine/Source/Editor/Sequencer/Private/SequencerEdMode.cpp
--- a/Engine/Source/Editor/Sequencer/Private/SequencerEdMode.cpp
+++ b/Engine/Source/Editor/Sequencer/Private/SequencerEdMode.cpp
@@ -124,7 +124,7 @@ struct FTrackTransforms
 				{
 					if (UMovieSceneTrack* SubTrack = Cast<UMovieSceneTrack>(SubSection->GetOuter()))
 					{
-						Interrogator.ImportTrack(SubTrack, FInterrogationChannel::FromIndex(SubSection->GetRowIndex()), OwningSequenceID);
+						Interrogator.ImportTrack(SubTrack, FInterrogationChannel::Default(), OwningSequenceID);
 					}
 				}
 				OwningSequenceID = SubSequenceID;
