[Sequencer]
- Add "Sort by [REDACTED_PERSON]/Descending" to outliner track context menu
- Add "Sort by [REDACTED_PERSON]/Descending" to outliner track context menu

diff --git a/Engine/Source/Editor/Sequencer/Private/MVVM/ViewModels/OutlinerItemModel.cpp b/Engine/Source/Editor/Sequencer/Private/MVVM/ViewModels/OutlinerItemModel.cpp
--- a/Engine/Source/Editor/Sequencer/Private/MVVM/ViewModels/OutlinerItemModel.cpp
+++ b/Engine/Source/Editor/Sequencer/Private/MVVM/ViewModels/OutlinerItemModel.cpp
@@ -646,6 +646,14 @@ void FOutlinerItemModelMixin::BuildOrganizeContextMenu(FMenuBuilder& MenuBuilder
 				FExecuteAction::CreateSP(SequencerRaw, &FSequencer::RemoveSelectedNodesFromFolders),
 				FCanExecuteAction::CreateLambda([SequencerRaw] { return SequencerRaw->GetSelectedNodesInFolders().Num() > 0; } )));
 	}
+
+	if (!bIsReadOnly)
+	{
+		MenuBuilder.AddSubMenu(
+			LOCTEXT("SortBy", "Sort by"),
+			LOCTEXT("SortByTooltip", "Sort the selected tracks by start time of the first layer bar"),
+			FNewMenuDelegate::CreateSP(SequencerRaw, &FSequencer::BuildSortMenu));
+	}
 }
 
 void FOutlinerItemModelMixin::BuildDisplayOptionsMenu(FMenuBuilder& MenuBuilder)
diff --git a/Engine/Source/Editor/Sequencer/Private/Sequencer.cpp b/Engine/Source/Editor/Sequencer/Private/Sequencer.cpp
--- a/Engine/Source/Editor/Sequencer/Private/Sequencer.cpp
+++ b/Engine/Source/Editor/Sequencer/Private/Sequencer.cpp
@@ -12155,6 +12155,35 @@ void FSequencer::BuildAddSelectedToNodeGroupMenu(FMenuBuilder& MenuBuilder)
 	}
 }
 
+void FSequencer::BuildSortMenu(FMenuBuilder& MenuBuilder)
+{
+	MenuBuilder.AddMenuEntry(
+		LOCTEXT("SortByStartTimeAscending", "Start Time Ascending"),
+		LOCTEXT("SortByStartTimeAscendingTooltip", "Sort the selected tracks by start time of the first layer bar ascending"),
+		FSlateIcon(FAppStyle::GetAppStyleSetName(), TEXT("Profiler.Misc.SortAscending")),
+		FUIAction(FExecuteAction::CreateSP(this, &FSequencer::SortSelectedOutlinerItems, false, false)));
+
+	MenuBuilder.AddMenuEntry(
+		LOCTEXT("SortByStartTimeDescending", "Start Time Descending"),
+		LOCTEXT("SortByStartTimeDescendingTooltip", "Sort the selected tracks by start time of the first layer bar descending"),
+		FSlateIcon(FAppStyle::GetAppStyleSetName(), TEXT("Profiler.Misc.SortDescending")),
+		FUIAction(FExecuteAction::CreateSP(this, &FSequencer::SortSelectedOutlinerItems, false, true)));
+
+	MenuBuilder.AddMenuSeparator();
+
+	MenuBuilder.AddMenuEntry(
+		LOCTEXT("SortBySelectionAscending", "Selection Order Ascending"),
+		LOCTEXT("SortBySelectionAscendingTooltip", "Sort the selected tracks by the order selected ascending"),
+		FSlateIcon(FAppStyle::GetAppStyleSetName(), TEXT("Profiler.Misc.SortAscending")),
+		FUIAction(FExecuteAction::CreateSP(this, &FSequencer::SortSelectedOutlinerItems, true, false)));
+
+	MenuBuilder.AddMenuEntry(
+		LOCTEXT("SortBySelectionDescending", "Selection Order Descending"),
+		LOCTEXT("SortBySelectionDescendingTooltip", "Sort the selected tracks by the order selected descending"),
+		FSlateIcon(FAppStyle::GetAppStyleSetName(), TEXT("Profiler.Misc.SortDescending")),
+		FUIAction(FExecuteAction::CreateSP(this, &FSequencer::SortSelectedOutlinerItems, true, true)));
+}
+
 void FSequencer::UpdateTimeBases()
 {
 	UMovieSceneSequence* RootSequencePtr = GetRootMovieSceneSequence();
@@ -12605,4 +12634,23 @@ EVisibility FSequencer::GetPlayRateComboVisibility() const
 	return EVisibility::Collapsed;
 }
 
+void FSequencer::SortSelectedOutlinerItems(const bool bInSortBySelectOrder, const bool bInDescending)
+{
+	using namespace UE::Sequencer;
+
+	const TSharedPtr<FSequencerSelection> SequencerSelection = ViewModel->GetSelection();
+	if (!SequencerSelection.IsValid())
+	{
+		return;
+	}
+
+	TArray<TViewModelPtr<IOutlinerExtension>> OutlinerItems;
+	for (const TViewModelPtr<IOutlinerExtension> Item : SequencerSelection->Outliner)
+	{
+		OutlinerItems.Add(Item);
+	}
+
+	SequencerHelpers::SortOutlinerItems(*this, OutlinerItems, bInSortBySelectOrder, bInDescending);
+}
+
 #undef LOCTEXT_NAMESPACE
diff --git a/Engine/Source/Editor/Sequencer/Private/Sequencer.h b/Engine/Source/Editor/Sequencer/Private/Sequencer.h
--- a/Engine/Source/Editor/Sequencer/Private/Sequencer.h
+++ b/Engine/Source/Editor/Sequencer/Private/Sequencer.h
@@ -492,6 +492,8 @@ public:
 	 */
 	void BuildAddSelectedToNodeGroupMenu(FMenuBuilder& MenuBuilder);
 
+	void BuildSortMenu(FMenuBuilder& MenuBuilder);
+
 	/** Called when an actor is dropped into Sequencer */
 	void OnActorsDropped( const TArray<TWeakObjectPtr<AActor> >& Actors );
 
@@ -1271,6 +1273,8 @@ private:
 
 	FText GetSidebarSelectionDrawerToolTipText() const;
 
+	void SortSelectedOutlinerItems(const bool bInSortBySelectOrder, const bool bInDescending);
+
 	/** User-supplied settings object for this sequencer */
 	TObjectPtr<USequencerSettings> Settings;
 
diff --git a/Engine/Source/Editor/Sequencer/Private/SequencerCommonHelpers.cpp b/Engine/Source/Editor/Sequencer/Private/SequencerCommonHelpers.cpp
--- a/Engine/Source/Editor/Sequencer/Private/SequencerCommonHelpers.cpp
+++ b/Engine/Source/Editor/Sequencer/Private/SequencerCommonHelpers.cpp
@@ -2,6 +2,7 @@
 
 #include "SequencerCommonHelpers.h"
 
+#include "Algo/Reverse.h"
 #include "EntitySystem/IMovieSceneBlenderSystemSupport.h"
 #include "EntitySystem/MovieSceneBlenderSystem.h"
 #include "FrameNumberDetailsCustomization.h"
@@ -10,6 +11,8 @@
 #include "IStructureDetailsView.h"
 #include "MVVM/Extensions/ITrackAreaExtension.h"
 #include "MVVM/ViewModels/ChannelModel.h"
+#include "MVVM/ViewModels/FolderModel.h"
+#include "MVVM/ViewModels/LayerBarModel.h"
 #include "MVVM/ViewModels/ViewModelIterators.h"
 #include "MVVM/ViewModels/ObjectBindingModel.h"
 #include "MVVM/ViewModels/SectionModel.h"
@@ -602,3 +605,223 @@ TArray<TWeakObjectPtr<>> SequencerHelpers::GetSectionObjectsFromTrackAreaModels(
 
 	return OutWeakSectionObjects;
 }
+
+void SequencerHelpers::SortOutlinerItems(FSequencer& Sequencer
+	, const TArray<UE::Sequencer::TViewModelPtr<UE::Sequencer::IOutlinerExtension>>& InItems
+	, const bool bInSortByItemOrder
+	, const bool bInDescending
+	, const bool bInTransact)
+{
+	using namespace UE::Sequencer;
+
+	struct FSortItemCache
+	{
+		TViewModelPtr<FLayerBarModel> LayerBar;
+		TViewModelPtr<ISortableExtension> SortableItem;
+		TViewModelPtr<IOutlinerExtension> OutlinerItem;
+	};
+
+	if (InItems.IsEmpty())
+	{
+		return;
+	}
+
+	UMovieSceneSequence* const Sequence = Sequencer.GetFocusedMovieSceneSequence();
+	if (!IsValid(Sequence))
+	{
+		return;
+	}
+
+	UMovieScene* const MovieScene = Sequence->GetMovieScene();
+	if (!IsValid(Sequence))
+	{
+		return;
+	}
+
+	const FSequencerSelection SequencerSelection = Sequencer.GetSelection();
+
+	// Gather all the selected folder and layer bar models from InItems so we can sort them
+	TMap<TViewModelPtr<FFolderModel>, TArray<TViewModelPtr<FLayerBarModel>>> FolderPairs;
+	TMap<TViewModelPtr<IOutlinerExtension>, TArray<FSortItemCache>> LayerBarSortItems;
+
+	for (const TViewModelPtr<IOutlinerExtension>& TrackItem : InItems)
+	{
+		// Handle folders
+		if (!bInSortByItemOrder)
+		{
+			if (const TViewModelPtr<FFolderModel> FolderModel = TrackItem.ImplicitCast())
+			{
+				TArray<TViewModelPtr<FLayerBarModel>> FolderLayerBars;
+
+				for (const TViewModelPtr<IOutlinerExtension>& ChildTrack : FolderModel->GetChildrenOfType<IOutlinerExtension>())
+				{
+					if (const TViewModelPtr<ITrackAreaExtension> ChildTrackArea = ChildTrack.ImplicitCast())
+					{
+						for (const FViewModelPtr& TrackAreaModel : ChildTrackArea->GetTopLevelChildTrackAreaModels())
+						{
+							if (const TViewModelPtr<FLayerBarModel> LayerBarModel = TrackAreaModel.ImplicitCast())
+							{
+								FolderLayerBars.Add(LayerBarModel);
+							}
+						}
+					}
+				}
+
+				if (!FolderLayerBars.IsEmpty())
+				{
+					FolderPairs.Add(FolderModel, FolderLayerBars);
+				}
+			}
+		}
+
+		// Handle track areas
+		if (const TSharedPtr<ITrackAreaExtension> TrackArea = TrackItem.ImplicitCast())
+		{
+			for (const FViewModelPtr& TrackAreaModel : TrackArea->GetTopLevelChildTrackAreaModels())
+			{
+				if (const TViewModelPtr<FLayerBarModel> LayerBarModel = TrackAreaModel.ImplicitCast())
+				{
+					if (const TViewModelPtr<IOutlinerExtension> OutlinerItem = LayerBarModel->GetLinkedOutlinerItem())
+					{
+						if (const TViewModelPtr<ISortableExtension> SortableItem = OutlinerItem.ImplicitCast())
+						{
+							if (const TViewModelPtr<IOutlinerExtension> ParentOutlinerItem = TrackItem.AsModel()->FindAncestorOfType<IOutlinerExtension>())
+							{
+								// Cache this item info and add it to the array in the map for sorting
+								FSortItemCache NewSortItem;
+								NewSortItem.LayerBar = LayerBarModel;
+								NewSortItem.SortableItem = SortableItem;
+								NewSortItem.OutlinerItem = OutlinerItem;
+								LayerBarSortItems.FindOrAdd(ParentOutlinerItem).Add(NewSortItem);
+							}
+						}
+					}
+				}
+			}
+		}
+	}
+
+	// Begin the transaction operation
+	const FText TransactionText = bInDescending
+		? NSLOCTEXT("Sequencer", "SortByAscending_Transaction", "Sort Tracks (Ascending)")
+		: NSLOCTEXT("Sequencer", "SortByDescending_Transaction", "Sort Tracks (Descending)");
+	FScopedTransaction SortNodesTransaction(TransactionText, bInTransact);
+
+	if (bInTransact)
+	{
+		MovieScene->Modify();
+	}
+
+	// Set the sort order for all selected folders
+	if (!bInSortByItemOrder)
+	{
+		int32 SortOrder = 0;
+		for (TTuple<TViewModelPtr<FFolderModel>, TArray<TViewModelPtr<FLayerBarModel>>>& FolderModel : FolderPairs)
+		{
+			FolderModel.Value.Sort([bInDescending](const TViewModelPtr<FLayerBarModel>& InLayerBarA, const TViewModelPtr<FLayerBarModel>& InLayerBarB) -> bool
+				{
+					if (bInDescending)
+					{
+						return InLayerBarA->ComputeRange().GetLowerBoundValue() > InLayerBarB->ComputeRange().GetLowerBoundValue();
+					}
+					return InLayerBarA->ComputeRange().GetLowerBoundValue() < InLayerBarB->ComputeRange().GetLowerBoundValue();
+				});
+
+			SortOrder = 0;
+			for (const TViewModelPtr<FLayerBarModel>& LayerBar : FolderModel.Value)
+			{
+				if (const TViewModelPtr<IOutlinerExtension> OutlinerItem = LayerBar->GetLinkedOutlinerItem())
+				{
+					if (const TViewModelPtr<ISortableExtension> SortableExtension = OutlinerItem.ImplicitCast())
+					{
+						SortableExtension->SetCustomOrder(SortOrder);
+						++SortOrder;
+					}
+				}
+			}
+		}
+	}
+
+	// Set the sort order for all selected layer bars
+	for (TPair<TViewModelPtr<IOutlinerExtension>, TArray<FSortItemCache>>& LayerBarSortItemPair : LayerBarSortItems)
+	{
+		TViewModelPtr<IOutlinerExtension>& ParentOutlinerItem = LayerBarSortItemPair.Key;
+		TArray<FSortItemCache>& SortItems = LayerBarSortItemPair.Value;
+
+		// Remove all layer models that have a descendant that is selected
+		SortItems.RemoveAll([&SequencerSelection](const FSortItemCache& InSortItem) -> bool
+			{
+				for (const TViewModelPtr<IOutlinerExtension>& ChildOutlinerExtension : InSortItem.OutlinerItem.AsModel()->GetDescendantsOfType<IOutlinerExtension>())
+				{
+					if (SequencerSelection.Outliner.IsSelected(ChildOutlinerExtension))
+					{
+						return true;
+					}
+				}
+				return false;
+			});
+
+		if (!bInSortByItemOrder)
+		{
+			SortItems.Sort([bInDescending](const FSortItemCache& InSortItemA, const FSortItemCache& InSortItemB) -> bool
+				{
+					if (bInDescending)
+					{
+						return InSortItemA.LayerBar->ComputeRange().GetLowerBoundValue() > InSortItemB.LayerBar->ComputeRange().GetLowerBoundValue();
+					}
+					return InSortItemA.LayerBar->ComputeRange().GetLowerBoundValue() < InSortItemB.LayerBar->ComputeRange().GetLowerBoundValue();
+				});
+		}
+
+		// Save the current sortable children and begin the re-ordering operation.
+		// We will attempt to maintain any current custom ordering that exists.
+		TArray<TViewModelPtr<ISortableExtension>> SortableChildren;
+		ParentOutlinerItem.AsModel()->GetChildrenOfType<ISortableExtension>().ToArray(SortableChildren);
+
+		// Sort to make sure we are in the correct order since the array can be out of order from what is actually displayed
+		SortableChildren.Sort([](const TViewModelPtr<ISortableExtension>& InA, const TViewModelPtr<ISortableExtension>& InB) -> bool
+			{
+				return FSortingKey::CompareCustomOrderFirst(InA->GetSortingKey(), InB->GetSortingKey());
+			});
+
+		// Loop backwards and remove items to re-insert while also re-ordering items
+		TArray<int32> ItemIndicesToReinsert;
+
+		for (int32 Index = SortableChildren.Num() - 1; Index >= 0; --Index)
+		{
+			const TViewModelPtr<ISortableExtension>& SortableChild = SortableChildren[Index];
+			const bool bHasItemToSort = SortItems.ContainsByPredicate([&SortableChild](const FSortItemCache& InItemCache)
+				{
+					return InItemCache.SortableItem == SortableChild;
+				});
+			if (bHasItemToSort)
+			{
+				SortableChildren.RemoveAt(Index);
+				ItemIndicesToReinsert.Add(Index);
+			}
+		}
+
+		// Reverse the array since we added while looping backwards through the array
+		Algo::Reverse(ItemIndicesToReinsert);
+
+		// Insert the sorted (selected) items to the array
+		int32 CurrentIndex = 0;
+		for (const int32 ItemIndex : ItemIndicesToReinsert)
+		{
+			if (SortItems.IsValidIndex(CurrentIndex))
+			{
+				SortableChildren.Insert(SortItems[CurrentIndex].SortableItem, ItemIndex);
+				++CurrentIndex;
+			}
+		}
+
+		// Use the SortableChildren index to set the custom order
+		int32 CustomOrder = 0;
+		for (const TViewModelPtr<ISortableExtension> SortableChild : SortableChildren)
+		{
+			SortableChild->SetCustomOrder(CustomOrder++);
+		}
+	}
+
+	Sequencer.RefreshTree();
+}
diff --git a/Engine/Source/Editor/Sequencer/Private/SequencerCommonHelpers.h b/Engine/Source/Editor/Sequencer/Private/SequencerCommonHelpers.h
--- a/Engine/Source/Editor/Sequencer/Private/SequencerCommonHelpers.h
+++ b/Engine/Source/Editor/Sequencer/Private/SequencerCommonHelpers.h
@@ -21,6 +21,7 @@ namespace UE
 namespace Sequencer
 {
 	class FChannelModel;
+	class IOutlinerExtension;
 	struct FViewModelVariantIterator;
 
 } // namespace Sequencer
@@ -105,6 +106,20 @@ public:
 	 * Gets all section objects from track area models
 	 */
 	static TArray<TWeakObjectPtr<>> GetSectionObjectsFromTrackAreaModels(const UE::Sequencer::FViewModelVariantIterator& InTrackAreaModels);
+
+	/**
+	 * Sorts an array of outliner items by start time of the first layer bar or selection order
+	 * @param Sequencer The Sequencer that the outliner items belong to
+	 * @param InItems The list of outliner items to sort
+	 * @param bInSortByItemOrder If true, sorts by the order of the item in the array instead of by start time of the layer bar
+	 * @param bInDescending If true, sorts in reverse order instead of ascending
+	 * @param bInTransact If true, creates an editor transaction that can be undone
+	 */
+	static void SortOutlinerItems(FSequencer& Sequencer
+		, const TArray<UE::Sequencer::TViewModelPtr<UE::Sequencer::IOutlinerExtension>>& InItems
+		, const bool bInSortByItemOrder
+		, const bool bInDescending
+		, const bool bInTransact = true);
 };
 
 #undef LOCTEXT_NAMESPACE
