[Sequencer] Fix SequencerHelpers::SortOutlinerItems(): pass by const ref instead of copy

diff --git a/Engine/Source/Editor/Sequencer/Private/SequencerCommonHelpers.cpp b/Engine/Source/Editor/Sequencer/Private/SequencerCommonHelpers.cpp
--- a/Engine/Source/Editor/Sequencer/Private/SequencerCommonHelpers.cpp
+++ b/Engine/Source/Editor/Sequencer/Private/SequencerCommonHelpers.cpp
@@ -817,7 +817,7 @@ void SequencerHelpers::SortOutlinerItems(FSequencer& Sequencer
 
 		// Use the SortableChildren index to set the custom order
 		int32 CustomOrder = 0;
-		for (const TViewModelPtr<ISortableExtension> SortableChild : SortableChildren)
+		for (const TViewModelPtr<ISortableExtension>& SortableChild : SortableChildren)
 		{
 			SortableChild->SetCustomOrder(CustomOrder++);
 		}
