Sequencer: Fix it so that custom bindings that are neither replaceables nor spawnables can be created properly.

diff --git a/Engine/Source/Editor/Sequencer/Private/SequencerUtilities.cpp b/Engine/Source/Editor/Sequencer/Private/SequencerUtilities.cpp
--- a/Engine/Source/Editor/Sequencer/Private/SequencerUtilities.cpp
+++ b/Engine/Source/Editor/Sequencer/Private/SequencerUtilities.cpp
@@ -1618,7 +1618,7 @@ FMovieScenePossessable* FSequencerUtilities::ConvertToCustomBinding(TSharedRef<I
 	{
 		Sequencer->GetSpawnRegister().DestroySpawnedObject(BindingGuid, Sequencer->GetFocusedTemplateID(), Sequencer->GetSharedPlaybackState(), BindingIndex);
 	}
-	else if (bConvertFromPossessable)
+	else if (bConvertFromPossessable && NewCustomBinding->WillSpawnObject(Sequencer->GetSharedPlaybackState()))
 	{
 		// We have an old possessable to destroy
 		Sequencer->GetSpawnRegister().HandleConvertPossessableToSpawnable(ObjectToConvert, *Sequencer, TransformData);
@@ -2668,11 +2668,14 @@ FGuid TryCreateCustomBinding(TSharedPtr<ISequencer> Sequencer, UObject* CustomBi
 
 		for (const TSubclassOf<UMovieSceneCustomBinding>& CustomBindingType : PrioritySortedCustomBindingTypes)
 		{
-			// If 'spawnable' has been passed in, we only want to use children of UMovieSceneSpawnableBindingBase and vice versa
+			// If 'spawnable' has been passed in, we can use children of UMovieSceneSpawnableBindingBase
+			// If 'replaceable' has been passed in, we can use children of UMovieSceneReplaceableBindingBase
+			// Otherwise if neither has been passed in, we only want to use bindings that aren't children of either.
 			const bool bIsCustomSpawnableBinding = CustomBindingType->IsChildOf<UMovieSceneSpawnableBindingBase>();
 			const bool bIsCustomReplaceableBinding = CustomBindingType->IsChildOf<UMovieSceneReplaceableBindingBase>();
 			if ((bSpawnable && bIsCustomSpawnableBinding) ||
-				(bReplaceable && bIsCustomReplaceableBinding))
+				(bReplaceable && bIsCustomReplaceableBinding) ||
+				(!bSpawnable && !bReplaceable && !bIsCustomSpawnableBinding && !bIsCustomReplaceableBinding))
 			{
 				if (UMovieSceneCustomBinding* CustomBindingCDO = CustomBindingType ? CustomBindingType->GetDefaultObject<UMovieSceneCustomBinding>() : nullptr)
 				{
