[Sequencer] Advanced Filtering:
- Make SSequencerFilter widget more generic so it can be used with all Sequencer filter bars
- Add base ISequencerFilterBar interface (mainly adds key and value suggestions)
- Make ISequencerTrackFilters inherit from ISequencerFilterBar to add additional track filter functionality
- Add ISequencerTextFilterExpressionContext
- This CL is part 1 of 2 for the Navigation Tool

diff --git a/Engine/Plugins/Animation/ControlRig/Source/ControlRigEditor/Private/Sequencer/AnimLayerSequencerFilter.cpp b/Engine/Plugins/Animation/ControlRig/Source/ControlRigEditor/Private/Sequencer/AnimLayerSequencerFilter.cpp
--- a/Engine/Plugins/Animation/ControlRig/Source/ControlRigEditor/Private/Sequencer/AnimLayerSequencerFilter.cpp
+++ b/Engine/Plugins/Animation/ControlRig/Source/ControlRigEditor/Private/Sequencer/AnimLayerSequencerFilter.cpp
@@ -31,13 +31,11 @@ class FSequencerTrackFilter_AnimLayerSequencerFilterCommands
 public:
 
 	FSequencerTrackFilter_AnimLayerSequencerFilterCommands()
-		: TCommands<FSequencerTrackFilter_AnimLayerSequencerFilterCommands>
-		(
+		: TCommands<FSequencerTrackFilter_AnimLayerSequencerFilterCommands>(
 			"FSequencerTrackFilter_AnimLayerSequencerFilter",
 			NSLOCTEXT("Contexts", "FSequencerTrackFilter_AnimLayerSequencer", "FSequencerTrackFilter_AnimLayerSequencer"),
 			NAME_None,
-			FAppStyle::GetAppStyleSetName() // Icon Style Set
-			)
+			FAppStyle::GetAppStyleSetName())
 	{ }
 
 	/** Toggle the control rig controls filter */
@@ -73,12 +71,11 @@ public:
 
 	virtual FString GetName() const override { return TEXT("SelectedAnimLayersFilter"); }
 	virtual FText GetDisplayName() const override { return LOCTEXT("SequenceTrackFilter_SelectedAnimLayers", "Selected Anim Layers"); }
-	virtual FSlateIcon GetIcon() const { return FSlateIconFinder::FindIconForClass(UControlRig::StaticClass()); } //UanimLayer
+	virtual FSlateIcon GetIcon() const override { return FSlateIconFinder::FindIconForClass(UControlRig::StaticClass()); }
 
 	virtual bool PassesFilter(FSequencerTrackFilterType InItem) const override
 	{
-
-		FSequencerFilterData& FilterData = FilterInterface.GetFilterData();
+		FSequencerFilterData& FilterData = GetFilterInterface().GetFilterData();
 
 		const TViewModelPtr<IOutlinerExtension> OutlinerExtension = InItem.AsModel()->FindAncestorOfType<IOutlinerExtension>();
 		if (!OutlinerExtension.IsValid())
diff --git a/Engine/Plugins/Animation/ControlRig/Source/ControlRigEditor/Private/Sequencer/ControlRigSequencerFilter.cpp b/Engine/Plugins/Animation/ControlRig/Source/ControlRigEditor/Private/Sequencer/ControlRigSequencerFilter.cpp
--- a/Engine/Plugins/Animation/ControlRig/Source/ControlRigEditor/Private/Sequencer/ControlRigSequencerFilter.cpp
+++ b/Engine/Plugins/Animation/ControlRig/Source/ControlRigEditor/Private/Sequencer/ControlRigSequencerFilter.cpp
@@ -72,7 +72,7 @@ public:
 
 	virtual bool PassesFilter(FSequencerTrackFilterType InItem) const override
 	{
-		FSequencerFilterData& FilterData = FilterInterface.GetFilterData();
+		FSequencerFilterData& FilterData = GetFilterInterface().GetFilterData();
 		const UMovieSceneTrack* const TrackObject = FilterData.ResolveMovieSceneTrackObject(InItem);
 		const UMovieSceneControlRigParameterTrack* const Track = Cast<UMovieSceneControlRigParameterTrack>(TrackObject);
 		return IsValid(Track);
@@ -156,7 +156,7 @@ public:
 
 	virtual bool PassesFilter(FSequencerTrackFilterType InItem) const override
 	{
-		FSequencerFilterData& FilterData = FilterInterface.GetFilterData();
+		FSequencerFilterData& FilterData = GetFilterInterface().GetFilterData();
 
 		UMovieSceneTrack* const TrackObject = FilterData.ResolveMovieSceneTrackObject(InItem);
 		URigHierarchy* const ControlRigHierarchy = GetControlRigHierarchyFromTrackObject(TrackObject);
diff --git a/Engine/Plugins/FX/Niagara/Source/NiagaraEditor/Private/Sequencer/NiagaraSequence/NiagaraSequencerFilters.cpp b/Engine/Plugins/FX/Niagara/Source/NiagaraEditor/Private/Sequencer/NiagaraSequence/NiagaraSequencerFilters.cpp
--- a/Engine/Plugins/FX/Niagara/Source/NiagaraEditor/Private/Sequencer/NiagaraSequence/NiagaraSequencerFilters.cpp
+++ b/Engine/Plugins/FX/Niagara/Source/NiagaraEditor/Private/Sequencer/NiagaraSequence/NiagaraSequencerFilters.cpp
@@ -66,7 +66,7 @@ public:
 
 	virtual bool PassesFilter(FSequencerTrackFilterType InItem) const override
 	{
-		FSequencerFilterData& FilterData = FilterInterface.GetFilterData();
+		FSequencerFilterData& FilterData = GetFilterInterface().GetFilterData();
 
 		const UObject* const BoundObject = FilterData.ResolveTrackBoundObject(GetSequencer(), InItem);
 		if (IsValid(BoundObject))
diff --git a/Engine/Source/Editor/Sequencer/Private/Filters/Filters/SequencerTrackFilterBase.cpp b/Engine/Source/Editor/Sequencer/Private/Filters/Filters/SequencerTrackFilterBase.cpp
--- a/Engine/Source/Editor/Sequencer/Private/Filters/Filters/SequencerTrackFilterBase.cpp
+++ b/Engine/Source/Editor/Sequencer/Private/Filters/Filters/SequencerTrackFilterBase.cpp
@@ -6,46 +6,22 @@
 #include "LevelSequence.h"
 #include "MovieSceneSequence.h"
 #include "MVVM/ViewModelPtr.h"
-#include "Sequencer.h"
 
 using namespace UE::Sequencer;
 
 FSequencerTrackFilter::FSequencerTrackFilter(ISequencerTrackFilters& InOutFilterInterface, TSharedPtr<FFilterCategory>&& InCategory)
-	: FFilterBase<FSequencerTrackFilterType>(MoveTemp(InCategory))
-	, FilterInterface(InOutFilterInterface)
+	: FSequencerFilterBase<FSequencerTrackFilterType>(InOutFilterInterface, MoveTemp(InCategory))
 {
 }
 
-FText FSequencerTrackFilter::GetToolTipText() const
-{
-	if (const TSharedPtr<FUICommandInfo> ToggleCommand = GetToggleCommand())
-	{
-		return BuildTooltipTextForCommand(GetDefaultToolTipText(), ToggleCommand);
-	}
-	return GetDefaultToolTipText();
-}
-
 bool FSequencerTrackFilter::SupportsSequence(UMovieSceneSequence* const InSequence) const
 {
 	return SupportsLevelSequence(InSequence);
 }
 
-void FSequencerTrackFilter::BindCommands()
-{
-	if (const TSharedPtr<FUICommandInfo> ToggleCommand = GetToggleCommand())
-	{
-		MapToggleAction(ToggleCommand);
-	}
-}
-
 ISequencerTrackFilters& FSequencerTrackFilter::GetFilterInterface() const
 {
-	return FilterInterface;
-}
-
-ISequencer& FSequencerTrackFilter::GetSequencer() const
-{
-	return FilterInterface.GetSequencer();
+	return static_cast<ISequencerTrackFilters&>(FilterInterface);
 }
 
 UMovieSceneSequence* FSequencerTrackFilter::GetFocusedMovieSceneSequence() const
@@ -74,35 +50,3 @@ bool FSequencerTrackFilter::SupportsUMGSequence(UMovieSceneSequence* const InSeq
 		&& IsValid(WidgetAnimationClass)
 		&& InSequence->GetClass()->IsChildOf(WidgetAnimationClass);
 }
-
-FText FSequencerTrackFilter::BuildTooltipTextForCommand(const FText& InBaseText, const TSharedPtr<FUICommandInfo>& InCommand)
-{
-	const TSharedRef<const FInputChord> FirstValidChord = InCommand->GetFirstValidChord();
-	if (FirstValidChord->IsValidChord())
-	{
-		return FText::Format(NSLOCTEXT("Sequencer", "TrackFilterTooltipText", "{0} ({1})"), InBaseText, FirstValidChord->GetInputText());
-	}
-	return InBaseText;
-}
-
-bool FSequencerTrackFilter::CanToggleFilter() const
-{
-	const FString FilterName = GetDisplayName().ToString();
-	return FilterInterface.IsFilterActiveByDisplayName(FilterName);
-}
-
-void FSequencerTrackFilter::ToggleFilter()
-{
-	const FString FilterName = GetDisplayName().ToString();
-	const bool bNewState = !FilterInterface.IsFilterActiveByDisplayName(FilterName);
-	FilterInterface.SetFilterActiveByDisplayName(FilterName, bNewState);
-}
-
-void FSequencerTrackFilter::MapToggleAction(const TSharedPtr<FUICommandInfo>& InCommand)
-{
-	FilterInterface.GetCommandList()->MapAction(
-		InCommand,
-		FExecuteAction::CreateSP(this, &FSequencerTrackFilter::ToggleFilter),
-		FCanExecuteAction(),
-		FIsActionChecked::CreateSP(this, &FSequencerTrackFilter::CanToggleFilter));
-}
diff --git a/Engine/Source/Editor/Sequencer/Private/Filters/Filters/SequencerTrackFilter_CustomText.cpp b/Engine/Source/Editor/Sequencer/Private/Filters/Filters/SequencerTrackFilter_CustomText.cpp
--- a/Engine/Source/Editor/Sequencer/Private/Filters/Filters/SequencerTrackFilter_CustomText.cpp
+++ b/Engine/Source/Editor/Sequencer/Private/Filters/Filters/SequencerTrackFilter_CustomText.cpp
@@ -2,7 +2,6 @@
 
 #include "SequencerTrackFilter_CustomText.h"
 #include "Filters/SequencerFilterBar.h"
-#include "Filters/SequencerTrackFilterCommands.h"
 
 using namespace UE::Sequencer;
 
diff --git a/Engine/Source/Editor/Sequencer/Private/Filters/Filters/SequencerTrackFilter_HideIsolate.cpp b/Engine/Source/Editor/Sequencer/Private/Filters/Filters/SequencerTrackFilter_HideIsolate.cpp
--- a/Engine/Source/Editor/Sequencer/Private/Filters/Filters/SequencerTrackFilter_HideIsolate.cpp
+++ b/Engine/Source/Editor/Sequencer/Private/Filters/Filters/SequencerTrackFilter_HideIsolate.cpp
@@ -22,19 +22,19 @@ FSequencerTrackFilter_HideIsolate::FSequencerTrackFilter_HideIsolate(ISequencerT
 void FSequencerTrackFilter_HideIsolate::BindCommands()
 {
 	const FSequencerTrackFilterCommands& TrackFilterCommands = FSequencerTrackFilterCommands::Get();
-
-	const TSharedPtr<FUICommandList>& FilterBarBindings = FilterInterface.GetCommandList();
+	ISequencerTrackFilters& TrackFiltersInterface = GetFilterInterface();
+	const TSharedPtr<FUICommandList>& FilterBarBindings = TrackFiltersInterface.GetCommandList();
 
 	FilterBarBindings->MapAction(TrackFilterCommands.HideSelectedTracks,
 		FUIAction(
-			FExecuteAction::CreateRaw(&FilterInterface, &ISequencerTrackFilters::HideSelectedTracks),
-			FCanExecuteAction::CreateRaw(&FilterInterface, &ISequencerTrackFilters::HasSelectedTracks)
+			FExecuteAction::CreateRaw(&TrackFiltersInterface, &ISequencerTrackFilters::HideSelectedTracks),
+			FCanExecuteAction::CreateRaw(&TrackFiltersInterface, &ISequencerTrackFilters::HasSelectedTracks)
 		));
 
 	FilterBarBindings->MapAction(TrackFilterCommands.IsolateSelectedTracks,
 		FUIAction(
-			FExecuteAction::CreateRaw(&FilterInterface, &ISequencerTrackFilters::IsolateSelectedTracks),
-			FCanExecuteAction::CreateRaw(&FilterInterface, &ISequencerTrackFilters::HasSelectedTracks)
+			FExecuteAction::CreateRaw(&TrackFiltersInterface, &ISequencerTrackFilters::IsolateSelectedTracks),
+			FCanExecuteAction::CreateRaw(&TrackFiltersInterface, &ISequencerTrackFilters::HasSelectedTracks)
 		));
 
 	FilterBarBindings->MapAction(TrackFilterCommands.ClearHiddenTracks,
@@ -55,9 +55,9 @@ void FSequencerTrackFilter_HideIsolate::BindCommands()
 			FCanExecuteAction::CreateSP(this, &FSequencerTrackFilter_HideIsolate::HasHiddenOrIsolatedTracks)
 		));
 
-	FilterBarBindings->MapAction(TrackFilterCommands.ShowLocationCategoryGroups, FExecuteAction::CreateRaw(&FilterInterface, &ISequencerTrackFilters::ShowOnlyLocationCategoryGroups));
-	FilterBarBindings->MapAction(TrackFilterCommands.ShowRotationCategoryGroups, FExecuteAction::CreateRaw(&FilterInterface, &ISequencerTrackFilters::ShowOnlyRotationCategoryGroups));
-	FilterBarBindings->MapAction(TrackFilterCommands.ShowScaleCategoryGroups, FExecuteAction::CreateRaw(&FilterInterface, &ISequencerTrackFilters::ShowOnlyScaleCategoryGroups));
+	FilterBarBindings->MapAction(TrackFilterCommands.ShowLocationCategoryGroups, FExecuteAction::CreateRaw(&TrackFiltersInterface, &ISequencerTrackFilters::ShowOnlyLocationCategoryGroups));
+	FilterBarBindings->MapAction(TrackFilterCommands.ShowRotationCategoryGroups, FExecuteAction::CreateRaw(&TrackFiltersInterface, &ISequencerTrackFilters::ShowOnlyRotationCategoryGroups));
+	FilterBarBindings->MapAction(TrackFilterCommands.ShowScaleCategoryGroups, FExecuteAction::CreateRaw(&TrackFiltersInterface, &ISequencerTrackFilters::ShowOnlyScaleCategoryGroups));
 }
 
 void FSequencerTrackFilter_HideIsolate::ResetFilter()
diff --git a/Engine/Source/Editor/Sequencer/Private/Filters/Filters/SequencerTrackFilter_Keyed.cpp b/Engine/Source/Editor/Sequencer/Private/Filters/Filters/SequencerTrackFilter_Keyed.cpp
--- a/Engine/Source/Editor/Sequencer/Private/Filters/Filters/SequencerTrackFilter_Keyed.cpp
+++ b/Engine/Source/Editor/Sequencer/Private/Filters/Filters/SequencerTrackFilter_Keyed.cpp
@@ -4,7 +4,6 @@
 #include "Filters/SequencerFilterBar.h"
 #include "Filters/SequencerTrackFilterCommands.h"
 #include "Framework/Commands/Commands.h"
-#include "MVVM/Extensions/IOutlinerExtension.h"
 #include "MVVM/ViewModelPtr.h"
 #include "MVVM/ViewModels/CategoryModel.h"
 #include "Sequencer.h"
@@ -55,7 +54,7 @@ bool FSequencerTrackFilter_Keyed::PassesFilter(FSequencerTrackFilterType InItem)
 		return ChannelGroupOutlinerModel->IsAnimated();
 	}
 
-	const TWeakViewModelPtr<ITrackExtension> WeakTrack = FilterInterface.GetFilterData().ResolveTrack(InItem);
+	const TWeakViewModelPtr<ITrackExtension> WeakTrack = GetFilterInterface().GetFilterData().ResolveTrack(InItem);
 	return DoesTrackExtensionHaveKeys(WeakTrack);
 }
 
diff --git a/Engine/Source/Editor/Sequencer/Private/Filters/Filters/SequencerTrackFilter_Modified.cpp b/Engine/Source/Editor/Sequencer/Private/Filters/Filters/SequencerTrackFilter_Modified.cpp
--- a/Engine/Source/Editor/Sequencer/Private/Filters/Filters/SequencerTrackFilter_Modified.cpp
+++ b/Engine/Source/Editor/Sequencer/Private/Filters/Filters/SequencerTrackFilter_Modified.cpp
@@ -44,7 +44,7 @@ FString FSequencerTrackFilter_Modified::GetName() const
 
 bool FSequencerTrackFilter_Modified::PassesFilter(FSequencerTrackFilterType InItem) const
 {
-	FSequencerFilterData& FilterData = FilterInterface.GetFilterData();
+	FSequencerFilterData& FilterData = GetFilterInterface().GetFilterData();
 
 	const UMovieSceneTrack* const TrackObject = FilterData.ResolveMovieSceneTrackObject(InItem);
 	if (!IsValid(TrackObject))
diff --git a/Engine/Source/Editor/Sequencer/Private/Filters/Filters/SequencerTrackFilter_Selected.cpp b/Engine/Source/Editor/Sequencer/Private/Filters/Filters/SequencerTrackFilter_Selected.cpp
--- a/Engine/Source/Editor/Sequencer/Private/Filters/Filters/SequencerTrackFilter_Selected.cpp
+++ b/Engine/Source/Editor/Sequencer/Private/Filters/Filters/SequencerTrackFilter_Selected.cpp
@@ -117,7 +117,7 @@ bool FSequencerTrackFilter_Selected::PassesFilter(FSequencerTrackFilterType InIt
 		return false;
 	}
 
-	FSequencerFilterData& FilterData = FilterInterface.GetFilterData();
+	FSequencerFilterData& FilterData = GetFilterInterface().GetFilterData();
 
 	const UObject* const TrackObject = FilterData.ResolveTrackBoundObject(GetSequencer(), InItem);
 	if (!IsValid(TrackObject))
diff --git a/Engine/Source/Editor/Sequencer/Private/Filters/Filters/SequencerTrackFilter_Text.cpp b/Engine/Source/Editor/Sequencer/Private/Filters/Filters/SequencerTrackFilter_Text.cpp
--- a/Engine/Source/Editor/Sequencer/Private/Filters/Filters/SequencerTrackFilter_Text.cpp
+++ b/Engine/Source/Editor/Sequencer/Private/Filters/Filters/SequencerTrackFilter_Text.cpp
@@ -74,7 +74,7 @@ FSequencerTrackFilter_Text::FSequencerTrackFilter_Text(ISequencerTrackFilters& I
 			&& !PotentialExtension->GetClass()->HasAnyClassFlags(CLASS_Deprecated | CLASS_Abstract))
 		{
 			TArray<TSharedRef<FSequencerTextFilterExpressionContext>> ExtendedTextExpressions;
-			PotentialExtension->AddTrackFilterTextExpressionExtensions(FilterInterface, ExtendedTextExpressions);
+			PotentialExtension->AddTrackFilterTextExpressionExtensions(static_cast<ISequencerTrackFilters&>(FilterInterface), ExtendedTextExpressions);
 
 			for (const TSharedRef<FSequencerTextFilterExpressionContext>& TextExpression : ExtendedTextExpressions)
 			{
@@ -101,7 +101,7 @@ FString FSequencerTrackFilter_Text::GetName() const
 
 bool FSequencerTrackFilter_Text::PassesFilter(FSequencerTrackFilterType InItem) const
 {
-	FSequencerFilterData& FilterData = FilterInterface.GetFilterData();
+	FSequencerFilterData& FilterData = GetFilterInterface().GetFilterData();
 
 	for (const TSharedRef<FSequencerTextFilterExpressionContext>& TextFilterExpressionContext : TextFilterExpressionContexts)
 	{
@@ -149,7 +149,7 @@ const TArray<TSharedRef<FSequencerTextFilterExpressionContext>>& FSequencerTrack
 	return TextFilterExpressionContexts;
 }
 
-bool FSequencerTrackFilter_Text::DoesTextFilterStringContainExpressionPair(const FSequencerTextFilterExpressionContext& InExpression) const
+bool FSequencerTrackFilter_Text::DoesTextFilterStringContainExpressionPair(const ISequencerTextFilterExpressionContext& InExpression) const
 {
 	const TArray<FExpressionToken>& ExpressionTokens = TextFilterExpressionEvaluator.GetFilterExpressionTokens();
 	const int32 ExpressionCount = ExpressionTokens.Num();
diff --git a/Engine/Source/Editor/Sequencer/Private/Filters/Filters/SequencerTrackFilter_Text.h b/Engine/Source/Editor/Sequencer/Private/Filters/Filters/SequencerTrackFilter_Text.h
--- a/Engine/Source/Editor/Sequencer/Private/Filters/Filters/SequencerTrackFilter_Text.h
+++ b/Engine/Source/Editor/Sequencer/Private/Filters/Filters/SequencerTrackFilter_Text.h
@@ -30,7 +30,7 @@ public:
 	const FTextFilterExpressionEvaluator& GetTextFilterExpressionEvaluator() const;
 	const TArray<TSharedRef<FSequencerTextFilterExpressionContext>>& GetTextFilterExpressionContexts() const;
 
-	bool DoesTextFilterStringContainExpressionPair(const FSequencerTextFilterExpressionContext& InExpression) const;
+	bool DoesTextFilterStringContainExpressionPair(const ISequencerTextFilterExpressionContext& InExpression) const;
 
 protected:
 	static bool IsTokenKey(const FExpressionToken& InToken, const TSet<FName>& InKeys);
diff --git a/Engine/Source/Editor/Sequencer/Private/Filters/Filters/SequencerTrackFilters.h b/Engine/Source/Editor/Sequencer/Private/Filters/Filters/SequencerTrackFilters.h
--- a/Engine/Source/Editor/Sequencer/Private/Filters/Filters/SequencerTrackFilters.h
+++ b/Engine/Source/Editor/Sequencer/Private/Filters/Filters/SequencerTrackFilters.h
@@ -221,7 +221,7 @@ public:
 
 	virtual bool PassesFilter(FSequencerTrackFilterType InItem) const override
 	{
-		FSequencerFilterData& FilterData = FilterInterface.GetFilterData();
+		FSequencerFilterData& FilterData = GetFilterInterface().GetFilterData();
 		const UMovieSceneTrack* const TrackObject = FilterData.ResolveMovieSceneTrackObject(InItem);
 		return IsValid(TrackObject)
 			&& TrackObject->IsA(UMovieSceneSubTrack::StaticClass())
@@ -519,7 +519,7 @@ public:
 		return FSequencerTrackFilterCommands::Get().ToggleFilter_Folder;
 	}
 
-	TSubclassOf<UMovieSceneTrack> GetTrackClass() const override
+	virtual TSubclassOf<UMovieSceneTrack> GetTrackClass() const override
 	{
 		return UMovieSceneFolder::StaticClass();
 	}
diff --git a/Engine/Source/Editor/Sequencer/Private/Filters/Menus/SequencerFilterMenuContext.h b/Engine/Source/Editor/Sequencer/Private/Filters/Menus/SequencerFilterMenuContext.h
--- a/Engine/Source/Editor/Sequencer/Private/Filters/Menus/SequencerFilterMenuContext.h
+++ b/Engine/Source/Editor/Sequencer/Private/Filters/Menus/SequencerFilterMenuContext.h
@@ -3,7 +3,6 @@
 #pragma once
 
 #include "Filters/SequencerTrackFilterBase.h"
-#include "Filters/Widgets/SSequencerFilter.h"
 #include "SequencerFilterBarContext.h"
 #include "Templates/SharedPointer.h"
 #include "SequencerFilterMenuContext.generated.h"
@@ -14,23 +13,18 @@ class USequencerFilterMenuContext : public UObject
 	GENERATED_BODY()
 
 public:
-	void Init(const TWeakPtr<SSequencerFilter>& InWeakFilterWidget)
+	void Init(const TWeakPtr<FSequencerTrackFilter>& InFilter)
 	{
-		WeakFilterWidget = InWeakFilterWidget;
-	}
-
-	TSharedPtr<SSequencerFilter> GetFilterWidget() const
-	{
-		return WeakFilterWidget.Pin();
+		WeakFilter = InFilter;
 	}
 
 	TSharedPtr<FSequencerTrackFilter> GetFilter() const
 	{
-		return WeakFilterWidget.IsValid() ? WeakFilterWidget.Pin()->GetFilter() : nullptr;
+		return WeakFilter.Pin();
 	}
 
 	FOnPopulateFilterBarMenu OnPopulateFilterBarMenu;
 
 protected:
-	TWeakPtr<SSequencerFilter> WeakFilterWidget;
+	TWeakPtr<FSequencerTrackFilter> WeakFilter;
 };
diff --git a/Engine/Source/Editor/Sequencer/Private/Filters/Menus/SequencerTrackFilterContextMenu.cpp b/Engine/Source/Editor/Sequencer/Private/Filters/Menus/SequencerTrackFilterContextMenu.cpp
--- a/Engine/Source/Editor/Sequencer/Private/Filters/Menus/SequencerTrackFilterContextMenu.cpp
+++ b/Engine/Source/Editor/Sequencer/Private/Filters/Menus/SequencerTrackFilterContextMenu.cpp
@@ -6,7 +6,6 @@
 #include "Filters/SequencerFilterBar.h"
 #include "Filters/SequencerTrackFilterCommands.h"
 #include "Filters/Widgets/SSequencerCustomTextFilterDialog.h"
-#include "Filters/Widgets/SSequencerFilter.h"
 #include "Sequencer.h"
 #include "SequencerSettings.h"
 #include "ToolMenu.h"
@@ -14,7 +13,7 @@
 
 #define LOCTEXT_NAMESPACE "SequencerTrackFilterContextMenu"
 
-TSharedRef<SWidget> FSequencerTrackFilterContextMenu::CreateMenuWidget(const TSharedRef<SSequencerFilter>& InFilterWidget)
+TSharedRef<SWidget> FSequencerTrackFilterContextMenu::CreateMenuWidget(const TSharedRef<FSequencerTrackFilter>& InFilter)
 {
 	UToolMenus* const ToolMenus = UToolMenus::Get();
 	check(IsValid(ToolMenus));
@@ -32,11 +31,10 @@ TSharedRef<SWidget> FSequencerTrackFilterContextMenu::CreateMenuWidget(const TSh
 			}));
 	}
 
-	const TSharedPtr<FSequencerTrackFilter> Filter = InFilterWidget->GetFilter();
-	const TSharedPtr<FUICommandList> CommandList = Filter.IsValid() ? Filter->GetFilterInterface().GetCommandList() : nullptr;
+	const TSharedPtr<FUICommandList> CommandList = InFilter->GetFilterInterface().GetCommandList();
 
 	USequencerFilterMenuContext* const ContextObject = NewObject<USequencerFilterMenuContext>();
-	ContextObject->Init(InFilterWidget);
+	ContextObject->Init(InFilter);
 	ContextObject->OnPopulateFilterBarMenu = FOnPopulateFilterBarMenu::CreateRaw(this, &FSequencerTrackFilterContextMenu::PopulateMenu);
 
 	const FToolMenuContext MenuContext(CommandList, nullptr, ContextObject);
@@ -56,7 +54,7 @@ void FSequencerTrackFilterContextMenu::PopulateMenu(UToolMenu* const InMenu)
 		return;
 	}
 
-	WeakFilterWidget = Context->GetFilterWidget();
+	WeakFilter = Context->GetFilter();
 
 	UToolMenu& MenuRef = *InMenu;
 
@@ -272,7 +270,7 @@ void FSequencerTrackFilterContextMenu::OnDeleteFilter()
 
 const TSharedPtr<FSequencerTrackFilter> FSequencerTrackFilterContextMenu::GetFilter() const
 {
-	return WeakFilterWidget.IsValid() ? WeakFilterWidget.Pin()->GetFilter() : nullptr;
+	return WeakFilter.IsValid() ? WeakFilter.Pin() : nullptr;
 }
 
 #undef LOCTEXT_NAMESPACE
diff --git a/Engine/Source/Editor/Sequencer/Private/Filters/Menus/SequencerTrackFilterContextMenu.h b/Engine/Source/Editor/Sequencer/Private/Filters/Menus/SequencerTrackFilterContextMenu.h
--- a/Engine/Source/Editor/Sequencer/Private/Filters/Menus/SequencerTrackFilterContextMenu.h
+++ b/Engine/Source/Editor/Sequencer/Private/Filters/Menus/SequencerTrackFilterContextMenu.h
@@ -5,14 +5,13 @@
 #include "Templates/SharedPointer.h"
 
 class FSequencerTrackFilter;
-class SSequencerFilter;
 class SWidget;
 class UToolMenu;
 
 class FSequencerTrackFilterContextMenu : public TSharedFromThis<FSequencerTrackFilterContextMenu>
 {
 public:
-	TSharedRef<SWidget> CreateMenuWidget(const TSharedRef<SSequencerFilter>& InFilterWidget);
+	TSharedRef<SWidget> CreateMenuWidget(const TSharedRef<FSequencerTrackFilter>& InFilter);
 
 protected:
 	void PopulateMenu(UToolMenu* const InMenu);
@@ -35,5 +34,5 @@ protected:
 
 	const TSharedPtr<FSequencerTrackFilter> GetFilter() const;
 
-	TWeakPtr<SSequencerFilter> WeakFilterWidget;
+	TWeakPtr<FSequencerTrackFilter> WeakFilter;
 };
diff --git a/Engine/Source/Editor/Sequencer/Private/Filters/Menus/SequencerTrackFilterMenu.cpp b/Engine/Source/Editor/Sequencer/Private/Filters/Menus/SequencerTrackFilterMenu.cpp
--- a/Engine/Source/Editor/Sequencer/Private/Filters/Menus/SequencerTrackFilterMenu.cpp
+++ b/Engine/Source/Editor/Sequencer/Private/Filters/Menus/SequencerTrackFilterMenu.cpp
@@ -5,10 +5,10 @@
 #include "Filters/Filters/SequencerTrackFilter_CustomText.h"
 #include "Filters/Widgets/SSequencerFilterBar.h"
 #include "Filters/Widgets/SSequencerCustomTextFilterDialog.h"
+#include "ISequencer.h"
 #include "LevelSequence.h"
 #include "Misc/MessageDialog.h"
 #include "MovieScene.h"
-#include "Sequencer.h"
 #include "SequencerFilterBarContext.h"
 #include "SSequencer.h"
 #include "ToolMenu.h"
@@ -249,7 +249,7 @@ void FSequencerTrackFilterMenu::FillGroupFilterMenu(UToolMenu* const InMenu)
 		return;
 	}
 
-	FSequencer& Sequencer = FilterBar->GetSequencer();
+	const ISequencer& Sequencer = FilterBar->GetSequencer();
 
 	FToolMenuSection& GroupFilterOptionsSection = InMenu->FindOrAddSection(TEXT("GroupFilterOptions")
 		, FText(), FToolMenuInsert(NAME_None, EToolMenuInsertType::First));
@@ -738,7 +738,7 @@ void FSequencerTrackFilterMenu::OpenNodeGroupsManager()
 		return;
 	}
 
-	FSequencer& Sequencer = FilterBar->GetSequencer();
+	const ISequencer& Sequencer = FilterBar->GetSequencer();
 
 	const TSharedPtr<SSequencer> SequencerWidget = StaticCastSharedRef<SSequencer>(Sequencer.GetSequencerWidget());
 	if (!SequencerWidget.IsValid())
@@ -757,12 +757,6 @@ void FSequencerTrackFilterMenu::OnOpenTextExpressionHelp()
 		return;
 	}
 
-	const TSharedPtr<SSequencer> SequencerWidget = StaticCastSharedRef<SSequencer>(FilterBar->GetSequencer().GetSequencerWidget());
-	if (!SequencerWidget.IsValid())
-	{
-		return;
-	}
-
 	const TSharedPtr<SSequencerFilterBar> FilterBarWidget = FilterBar->GetWidget();
 	if (!FilterBarWidget.IsValid())
 	{
diff --git a/Engine/Source/Editor/Sequencer/Private/Filters/SequencerFilterBar.cpp b/Engine/Source/Editor/Sequencer/Private/Filters/SequencerFilterBar.cpp
--- a/Engine/Source/Editor/Sequencer/Private/Filters/SequencerFilterBar.cpp
+++ b/Engine/Source/Editor/Sequencer/Private/Filters/SequencerFilterBar.cpp
@@ -28,7 +28,6 @@
 #include "MVVM/ViewModels/CategoryModel.h"
 #include "MVVM/Views/SOutlinerView.h"
 #include "Sequencer.h"
-#include "SequencerLog.h"
 #include "SSequencer.h"
 #include "Stats/StatsMisc.h"
 #include "UObject/UObjectIterator.h"
@@ -185,7 +184,10 @@ void FSequencerFilterBar::BindCommands()
 
 	CommandList->MapAction(TrackFilterCommands.DisableAllFilters,
 		FUIAction(
-			FExecuteAction::CreateSP(this, &FSequencerFilterBar::EnableAllFilters, false, TArray<FString>()),
+			FExecuteAction::CreateLambda([this]()
+				{
+					EnableAllFilters(false, {});
+				}),
 			FCanExecuteAction::CreateSP(this, &FSequencerFilterBar::HasAnyFilterEnabled)
 		));
 
@@ -271,7 +273,7 @@ void FSequencerFilterBar::CreateCustomTextFiltersFromConfig()
 	}
 }
 
-FSequencer& FSequencerFilterBar::GetSequencer() const
+ISequencer& FSequencerFilterBar::GetSequencer() const
 {
 	return Sequencer;
 }
@@ -349,7 +351,7 @@ void FSequencerFilterBar::SetTextFilterString(const FString& InText)
 	}
 }
 
-bool FSequencerFilterBar::DoesTextFilterStringContainExpressionPair(const FSequencerTextFilterExpressionContext& InExpression) const
+bool FSequencerFilterBar::DoesTextFilterStringContainExpressionPair(const ISequencerTextFilterExpressionContext& InExpression) const
 {
 	return TextFilter->DoesTextFilterStringContainExpressionPair(InExpression);
 }
@@ -482,7 +484,7 @@ bool FSequencerFilterBar::HasAnyFiltersEnabled() const
 	return HasEnabledCommonFilters() || HasEnabledCustomTextFilters();
 }
 
-bool FSequencerFilterBar::IsFilterActiveByDisplayName(const FString InFilterName) const
+bool FSequencerFilterBar::IsFilterActiveByDisplayName(const FString& InFilterName) const
 {
 	if (const TSharedPtr<FSequencerTrackFilter> Filter = FindFilterByDisplayName(InFilterName))
 	{
@@ -491,7 +493,7 @@ bool FSequencerFilterBar::IsFilterActiveByDisplayName(const FString InFilterName
 	return false;
 }
 
-bool FSequencerFilterBar::IsFilterEnabledByDisplayName(const FString InFilterName) const
+bool FSequencerFilterBar::IsFilterEnabledByDisplayName(const FString& InFilterName) const
 {
 	if (const TSharedPtr<FSequencerTrackFilter> Filter = FindFilterByDisplayName(InFilterName))
 	{
@@ -500,7 +502,7 @@ bool FSequencerFilterBar::IsFilterEnabledByDisplayName(const FString InFilterNam
 	return false;
 }
 
-bool FSequencerFilterBar::SetFilterActiveByDisplayName(const FString InFilterName, const bool bInActive, const bool bInRequestFilterUpdate)
+bool FSequencerFilterBar::SetFilterActiveByDisplayName(const FString& InFilterName, const bool bInActive, const bool bInRequestFilterUpdate)
 {
 	if (const TSharedPtr<FSequencerTrackFilter> Filter = FindFilterByDisplayName(InFilterName))
 	{
@@ -515,7 +517,7 @@ bool FSequencerFilterBar::SetFilterActiveByDisplayName(const FString InFilterNam
 	return false;
 }
 
-bool FSequencerFilterBar::SetFilterEnabledByDisplayName(const FString InFilterName, const bool bInEnabled, const bool bInRequestFilterUpdate)
+bool FSequencerFilterBar::SetFilterEnabledByDisplayName(const FString& InFilterName, const bool bInEnabled, const bool bInRequestFilterUpdate)
 {
 	if (const TSharedPtr<FSequencerTrackFilter> Filter = FindFilterByDisplayName(InFilterName))
 	{
@@ -622,7 +624,7 @@ bool FSequencerFilterBar::SetFilterActive(const TSharedRef<FSequencerTrackFilter
 	return bSuccess;
 }
 
-void FSequencerFilterBar::EnableAllFilters(const bool bInEnable, const TArray<FString> InExceptionFilterNames)
+void FSequencerFilterBar::EnableAllFilters(const bool bInEnable, const TArray<FString>& InExceptionFilterNames)
 {
 	TArray<TSharedRef<FSequencerTrackFilter>> ExceptionFilters;
 	TArray<TSharedRef<FSequencerTrackFilter_CustomText>> ExceptionCustomTextFilters;
@@ -643,7 +645,7 @@ void FSequencerFilterBar::EnableAllFilters(const bool bInEnable, const TArray<FS
 	EnableCustomTextFilters(bInEnable, ExceptionCustomTextFilters);
 }
 
-void FSequencerFilterBar::ActivateCommonFilters(const bool bInActivate, const TArray<FString> InExceptionFilterNames)
+void FSequencerFilterBar::ActivateCommonFilters(const bool bInActivate, const TArray<FString>& InExceptionFilterNames)
 {
 	TArray<TSharedRef<FSequencerTrackFilter>> ExceptionFilters;
 
@@ -704,7 +706,7 @@ void FSequencerFilterBar::ActivateCommonFilters(const bool bInActivate
 	RequestFilterUpdate();
 }
 
-bool FSequencerFilterBar::AreAllEnabledFiltersActive(const bool bInActive, const TArray<FString> InExceptionFilterNames) const
+bool FSequencerFilterBar::AreAllEnabledFiltersActive(const bool bInActive, const TArray<FString>& InExceptionFilterNames) const
 {
 	const TArray<TSharedRef<FSequencerTrackFilter>> EnabledFilters = GetEnabledFilters();
 	for (const TSharedRef<FSequencerTrackFilter>& Filter : EnabledFilters)
@@ -739,7 +741,7 @@ bool FSequencerFilterBar::AreAllEnabledFiltersActive(const bool bInActive, const
 	return true;
 }
 
-void FSequencerFilterBar::ActivateAllEnabledFilters(const bool bInActivate, const TArray<FString> InExceptionFilterNames)
+void FSequencerFilterBar::ActivateAllEnabledFilters(const bool bInActivate, const TArray<FString>& InExceptionFilterNames)
 {
 	const TArray<TSharedRef<FSequencerTrackFilter>> EnabledFilters = GetEnabledFilters();
 	for (const TSharedRef<FSequencerTrackFilter>& Filter : EnabledFilters)
@@ -1397,6 +1399,22 @@ FSequencerFilterData& FSequencerFilterBar::GetFilterData()
 	return FilterData;
 }
 
+const FTextFilterExpressionEvaluator& FSequencerFilterBar::GetTextFilterExpressionEvaluator() const
+{
+	return GetTextFilter()->GetTextFilterExpressionEvaluator();
+}
+
+TArray<TSharedRef<ISequencerTextFilterExpressionContext>> FSequencerFilterBar::GetTextFilterExpressionContexts() const
+{
+	TArray<TSharedRef<ISequencerTextFilterExpressionContext>> TextExpressions;
+	Algo::Transform(TextFilter->GetTextFilterExpressionContexts(), TextExpressions
+		, [](TSharedRef<FSequencerTextFilterExpressionContext> InExpressionContext)
+		{
+			return InExpressionContext;
+		});
+	return TextExpressions;
+}
+
 bool FSequencerFilterBar::FilterNodesRecursive(const bool bInHasActiveFilter, const TViewModelPtr<IOutlinerExtension>& InStartNode)
 {
 	/**
diff --git a/Engine/Source/Editor/Sequencer/Private/Filters/SequencerFilterBar.h b/Engine/Source/Editor/Sequencer/Private/Filters/SequencerFilterBar.h
--- a/Engine/Source/Editor/Sequencer/Private/Filters/SequencerFilterBar.h
+++ b/Engine/Source/Editor/Sequencer/Private/Filters/SequencerFilterBar.h
@@ -95,33 +95,33 @@ public:
 
 	bool HasAnyFiltersEnabled() const;
 
-	//~ Begin ISequencerTrackFilters
+	//~ Begin ISequencerFilterBar
 
 	virtual FName GetIdentifier() const override;
 
-	virtual FSequencer& GetSequencer() const override;
+	virtual ISequencer& GetSequencer() const override;
 
 	virtual TSharedPtr<FUICommandList> GetCommandList() const override;
 
 	virtual FString GetTextFilterString() const override;
     virtual void SetTextFilterString(const FString& InText) override;
 
-	virtual bool DoesTextFilterStringContainExpressionPair(const FSequencerTextFilterExpressionContext& InExpression) const override;
+	virtual bool DoesTextFilterStringContainExpressionPair(const ISequencerTextFilterExpressionContext& InExpression) const override;
 
 	virtual void RequestFilterUpdate() override;
 
-	virtual void EnableAllFilters(const bool bInEnable, const TArray<FString> InExceptionFilterNames) override;
+	virtual void EnableAllFilters(const bool bInEnable, const TArray<FString>& InExceptionFilterNames) override;
 
-	virtual void ActivateCommonFilters(const bool bInActivate, const TArray<FString> InExceptionFilterNames) override;
+	virtual void ActivateCommonFilters(const bool bInActivate, const TArray<FString>& InExceptionFilterNames) override;
 
-	virtual bool AreAllEnabledFiltersActive(const bool bInActive, const TArray<FString> InExceptionFilterNames) const;
-	virtual void ActivateAllEnabledFilters(const bool bInActivate, const TArray<FString> InExceptionFilterNames) override;
+	virtual bool AreAllEnabledFiltersActive(const bool bInActive, const TArray<FString>& InExceptionFilterNames) const;
+	virtual void ActivateAllEnabledFilters(const bool bInActivate, const TArray<FString>& InExceptionFilterNames) override;
 	virtual void ToggleActivateAllEnabledFilters();
 
-	virtual bool IsFilterActiveByDisplayName(const FString InFilterName) const override;
-	virtual bool IsFilterEnabledByDisplayName(const FString InFilterName) const override;
-	virtual bool SetFilterActiveByDisplayName(const FString InFilterName, const bool bInActive, const bool bInRequestFilterUpdate = true) override;
-	virtual bool SetFilterEnabledByDisplayName(const FString InFilterName, const bool bInEnabled, const bool bInRequestFilterUpdate = true) override;
+	virtual bool IsFilterActiveByDisplayName(const FString& InFilterName) const override;
+	virtual bool IsFilterEnabledByDisplayName(const FString& InFilterName) const override;
+	virtual bool SetFilterActiveByDisplayName(const FString& InFilterName, const bool bInActive, const bool bInRequestFilterUpdate = true) override;
+	virtual bool SetFilterEnabledByDisplayName(const FString& InFilterName, const bool bInEnabled, const bool bInRequestFilterUpdate = true) override;
 
 	virtual TArray<FText> GetFilterDisplayNames() const override;
 	virtual TArray<FText> GetCustomTextFilterNames() const override;
@@ -129,6 +129,10 @@ public:
 	virtual int32 GetTotalDisplayNodeCount() const override;
 	virtual int32 GetFilteredDisplayNodeCount() const override;
 
+	//~ End ISequencerFilterBar
+
+	//~ End ISequencerTrackFilters
+
 	virtual void HideSelectedTracks() override;
 	virtual void IsolateSelectedTracks() override;
 
@@ -140,6 +144,9 @@ public:
 
 	virtual FSequencerFilterData& GetFilterData() override;
 
+	virtual const FTextFilterExpressionEvaluator& GetTextFilterExpressionEvaluator() const override;
+	virtual TArray<TSharedRef<ISequencerTextFilterExpressionContext>> GetTextFilterExpressionContexts() const override;
+
 	//~ End ISequencerTrackFilters
 
 	/** Active Filter Functions */
diff --git a/Engine/Source/Editor/Sequencer/Private/Filters/TextExpressions/SequencerTextFilterExpression_Group.cpp b/Engine/Source/Editor/Sequencer/Private/Filters/TextExpressions/SequencerTextFilterExpression_Group.cpp
--- a/Engine/Source/Editor/Sequencer/Private/Filters/TextExpressions/SequencerTextFilterExpression_Group.cpp
+++ b/Engine/Source/Editor/Sequencer/Private/Filters/TextExpressions/SequencerTextFilterExpression_Group.cpp
@@ -2,8 +2,8 @@
 
 #include "SequencerTextFilterExpression_Group.h"
 #include "Filters/Filters/SequencerTrackFilter_Group.h"
+#include "ISequencer.h"
 #include "MovieScene.h"
-#include "Sequencer.h"
 
 using namespace UE::Sequencer;
 
@@ -39,7 +39,7 @@ bool FSequencerTextFilterExpression_Group::TestComplexExpression(const FName& In
 		return true;
 	}
 
-	const FSequencer& Sequencer = FilterInterface.GetSequencer();
+	const ISequencer& Sequencer = FilterInterface.GetSequencer();
 
 	const UMovieSceneSequence* const FocusedSequence = Sequencer.GetFocusedMovieSceneSequence();
 	if (!IsValid(FocusedSequence))
diff --git a/Engine/Source/Editor/Sequencer/Private/Filters/TextExpressions/SequencerTextFilterExpression_Tag.cpp b/Engine/Source/Editor/Sequencer/Private/Filters/TextExpressions/SequencerTextFilterExpression_Tag.cpp
--- a/Engine/Source/Editor/Sequencer/Private/Filters/TextExpressions/SequencerTextFilterExpression_Tag.cpp
+++ b/Engine/Source/Editor/Sequencer/Private/Filters/TextExpressions/SequencerTextFilterExpression_Tag.cpp
@@ -1,9 +1,9 @@
 ﻿// Copyright Epic Games, Inc. All Rights Reserved.
 
 #include "SequencerTextFilterExpression_Tag.h"
+#include "ISequencer.h"
 #include "MovieScene.h"
 #include "MVVM/ViewModels/ObjectBindingModel.h"
-#include "Sequencer.h"
 
 using namespace UE::Sequencer;
 
@@ -48,7 +48,7 @@ bool FSequencerTextFilterExpression_Tag::TestComplexExpression(const FName& InKe
 
 	if (const TViewModelPtr<FObjectBindingModel> BindingModel = FilterItem->FindAncestorOfType<FObjectBindingModel>())
 	{
-		FSequencer& Sequencer = FilterInterface.GetSequencer();
+		const ISequencer& Sequencer = FilterInterface.GetSequencer();
 
 		const FMovieSceneObjectBindingID ObjectBindingID(UE::MovieScene::FFixedObjectBindingID(BindingModel->GetObjectGuid(), Sequencer.GetFocusedTemplateID()));
 		const TMap<FName, FMovieSceneObjectBindingIDs>& AllTaggedBindings = MovieScene->AllTaggedBindings();
diff --git a/Engine/Source/Editor/Sequencer/Private/Filters/TextExpressions/SequencerTextFilterExpression_Time.cpp b/Engine/Source/Editor/Sequencer/Private/Filters/TextExpressions/SequencerTextFilterExpression_Time.cpp
--- a/Engine/Source/Editor/Sequencer/Private/Filters/TextExpressions/SequencerTextFilterExpression_Time.cpp
+++ b/Engine/Source/Editor/Sequencer/Private/Filters/TextExpressions/SequencerTextFilterExpression_Time.cpp
@@ -3,9 +3,10 @@
 #include "SequencerTextFilterExpression_Time.h"
 #include "Filters/SequencerFilterBar.h"
 #include "IKeyArea.h"
+#include "ISequencer.h"
 #include "MovieScene.h"
 #include "MVVM/ViewModels/CategoryModel.h"
-#include "Sequencer.h"
+#include "MVVM/ViewModels/ChannelModel.h"
 
 using namespace UE::Sequencer;
 
@@ -50,7 +51,7 @@ bool FSequencerTextFilterExpression_Time::TestComplexExpression(const FName& InK
 		return true;
 	}
 
-	UMovieScene* const FocusedMovieScene = GetFocusedGetMovieScene();
+	const UMovieScene* const FocusedMovieScene = GetFocusedGetMovieScene();
 	if (!IsValid(FocusedMovieScene))
 	{
 		return true;
@@ -62,7 +63,7 @@ bool FSequencerTextFilterExpression_Time::TestComplexExpression(const FName& InK
 		return false;
 	}
 
-	FSequencer& Sequencer = FilterInterface.GetSequencer();
+	ISequencer& Sequencer = FilterInterface.GetSequencer();
 
 	// Assume the value is a specified as a numeric frame number
 	FTextFilterString ValueToCheck = InValue;
diff --git a/Engine/Source/Editor/Sequencer/Private/Filters/TextExpressions/SequencerTextFilterExpression_Time.h b/Engine/Source/Editor/Sequencer/Private/Filters/TextExpressions/SequencerTextFilterExpression_Time.h
--- a/Engine/Source/Editor/Sequencer/Private/Filters/TextExpressions/SequencerTextFilterExpression_Time.h
+++ b/Engine/Source/Editor/Sequencer/Private/Filters/TextExpressions/SequencerTextFilterExpression_Time.h
@@ -4,8 +4,6 @@
 
 #include "Filters/SequencerTextFilterExpressionContext.h"
 
-class IKeyArea;
-
 class FSequencerTextFilterExpression_Time : public FSequencerTextFilterExpressionContext
 {
 public:
diff --git a/Engine/Source/Editor/Sequencer/Private/Filters/Widgets/SFilterExpressionHelpDialog.cpp b/Engine/Source/Editor/Sequencer/Private/Filters/Widgets/SFilterExpressionHelpDialog.cpp
--- a/Engine/Source/Editor/Sequencer/Private/Filters/Widgets/SFilterExpressionHelpDialog.cpp
+++ b/Engine/Source/Editor/Sequencer/Private/Filters/Widgets/SFilterExpressionHelpDialog.cpp
@@ -2,6 +2,8 @@
 
 #include "SFilterExpressionHelpDialog.h"
 #include "Filters/SequencerTextFilterExpressionContext.h"
+#include "Framework/Application/SlateApplication.h"
+#include "Interfaces/IMainFrameModule.h"
 #include "Styling/StyleColors.h"
 #include "Widgets/Images/SImage.h"
 #include "Widgets/Input/SHyperlink.h"
@@ -11,17 +13,75 @@
 
 #define LOCTEXT_NAMESPACE "SFilterExpressionHelpDialog"
 
+FFilterExpressionHelpDialogConfig::FFilterExpressionHelpDialogConfig()
+{
+	DialogTitle = LOCTEXT("DialogTitle", "Text Filter Expression Help");
+	DocumentationLink = TEXT("https://dev.epicgames.com/documentation/en-us/unreal-engine/advanced-search-syntax-in-unreal-engine");
+}
+
 const FSlateColor SFilterExpressionHelpDialog::KeyColor = FStyleColors::AccentBlue;
 const FSlateColor SFilterExpressionHelpDialog::ValueColor = FStyleColors::AccentOrange;
 
-void SFilterExpressionHelpDialog::Construct(const FArguments& InArgs)
+TMap<FName, TSharedPtr<SFilterExpressionHelpDialog>> SFilterExpressionHelpDialog::DialogInstance;
+
+void SFilterExpressionHelpDialog::Open(FFilterExpressionHelpDialogConfig&& InConfig)
+{
+	const FName IdentifierName = InConfig.IdentifierName;
+
+	if (DialogInstance.Contains(IdentifierName)
+		&& DialogInstance[IdentifierName].IsValid()
+		&& DialogInstance[IdentifierName]->IsVisible())
+	{
+		DialogInstance[IdentifierName]->BringToFront();
+		return;
+	}
+
+	DialogInstance.Add(IdentifierName, SNew(SFilterExpressionHelpDialog, MoveTemp(InConfig)));
+
+	DialogInstance[IdentifierName]->GetOnWindowClosedEvent().AddLambda([IdentifierName](const TSharedRef<SWindow>& InWindow)
+		{
+			DialogInstance[IdentifierName].Reset();
+		});
+
+	TSharedPtr<SWindow> ParentWindow;
+	if (FModuleManager::Get().IsModuleLoaded(TEXT("MainFrame")))
+	{
+		IMainFrameModule& MainFrame = FModuleManager::LoadModuleChecked<IMainFrameModule>(TEXT("MainFrame"));
+		ParentWindow = MainFrame.GetParentWindow();
+	}
+
+	if (ParentWindow.IsValid())
+	{
+		FSlateApplication::Get().AddWindowAsNativeChild(DialogInstance[IdentifierName].ToSharedRef(), ParentWindow.ToSharedRef());
+	}
+	else
+	{
+		FSlateApplication::Get().AddWindow(DialogInstance[IdentifierName].ToSharedRef());
+	}
+}
+
+bool SFilterExpressionHelpDialog::IsOpen(const FName InName)
+{
+	return DialogInstance.Contains(InName) && DialogInstance[InName].IsValid();
+}
+
+void SFilterExpressionHelpDialog::CloseWindow(const FName InName)
+{
+	if (DialogInstance.Contains(InName) && DialogInstance[InName].IsValid())
+	{
+		DialogInstance[InName]->RequestDestroyWindow();
+		DialogInstance[InName].Reset();
+
+		DialogInstance.Remove(InName);
+	}
+}
+
+void SFilterExpressionHelpDialog::Construct(const FArguments& InArgs, FFilterExpressionHelpDialogConfig&& InConfig)
 {
-	HeaderText = InArgs._HeaderText;
-	DocumentationLink = InArgs._DocumentationLink;
-	TextFilterExpressionContexts = InArgs._TextFilterExpressionContexts;
+	Config = MoveTemp(InConfig);
 
 	SWindow::Construct(SWindow::FArguments()
-		.Title(InArgs._DialogTitle)
+		.Title(Config.DialogTitle)
 		.AutoCenter(EAutoCenter::PrimaryWorkArea)
 		.SizingRule(ESizingRule::Autosized)
 		.HasCloseButton(true)
@@ -40,11 +100,11 @@ void SFilterExpressionHelpDialog::Construct(const FArguments& InArgs)
 		.Padding(5.f)
 		[
 			SNew(SBox)
-			.MaxDesiredWidth(MaxDesiredWidth)
-			.MaxDesiredHeight(MaxDesiredHeight)
+			.MaxDesiredWidth(Config.MaxDesiredWidth)
+			.MaxDesiredHeight(Config.MaxDesiredHeight)
 			[
 				SNew(SBorder)
-				.BorderImage(FAppStyle::Get().GetBrush(TEXT("Brushes.Panel")))
+				.BorderImage(FAppStyle::Get().GetBrush(TEXT("Brushes.Recessed")))
 				[
 					ConstructExpressionWidgetList()
 				]
@@ -63,10 +123,10 @@ TSharedRef<SWidget> SFilterExpressionHelpDialog::ConstructDialogHeader()
 		[
 			SNew(STextBlock)
 			.Font(FCoreStyle::GetDefaultFontStyle(TEXT("Bold"), 11))
-			.Text(HeaderText)
+			.Text(LOCTEXT("HeaderText", "Text Filter Expressions"))
 		];
 
-	if (!DocumentationLink.IsEmpty())
+	if (!Config.DocumentationLink.IsEmpty())
 	{
 		HorizontalBox->AddSlot()
 			.AutoWidth()
@@ -87,7 +147,7 @@ TSharedRef<SWidget> SFilterExpressionHelpDialog::ConstructDialogHeader()
 			[
 				SNew(SHyperlink)
 				.Text(LOCTEXT("DocumentationLink", "Documentation"))
-				.ToolTipText(FText::Format(LOCTEXT("NavigateToDocumentation", "Open the online documentation ({0})"), FText::FromString(DocumentationLink)))
+				.ToolTipText(FText::Format(LOCTEXT("NavigateToDocumentation", "Open the online documentation ({0})"), FText::FromString(Config.DocumentationLink)))
 				.Style(FAppStyle::Get(), TEXT("HoverOnlyHyperlink"))
 				.TextStyle(FAppStyle::Get(), TEXT("SmallText"))
 				.OnNavigate(this, &SFilterExpressionHelpDialog::OpenDocumentationLink)
@@ -103,7 +163,7 @@ TSharedRef<SWidget> SFilterExpressionHelpDialog::ConstructExpressionWidgetList()
 
 	bool bIsFirst = true;
 
-	for (const TSharedRef<FSequencerTextFilterExpressionContext>& ExpressionContext : TextFilterExpressionContexts)
+	for (const TSharedRef<ISequencerTextFilterExpressionContext>& ExpressionContext : Config.TextFilterExpressionContexts)
 	{
 		if (bIsFirst)
 		{
@@ -129,7 +189,7 @@ TSharedRef<SWidget> SFilterExpressionHelpDialog::ConstructExpressionWidgetList()
 	return Container;
 }
 
-TSharedRef<SWidget> SFilterExpressionHelpDialog::ConstructExpressionWidget(const TSharedPtr<FSequencerTextFilterExpressionContext>& InExpressionContext)
+TSharedRef<SWidget> SFilterExpressionHelpDialog::ConstructExpressionWidget(const TSharedPtr<ISequencerTextFilterExpressionContext>& InExpressionContext)
 {
 	return SNew(SVerticalBox)
 		+ SVerticalBox::Slot()
@@ -240,9 +300,9 @@ TSharedRef<SWidget> SFilterExpressionHelpDialog::ConstructValueWidget(const ESeq
 
 void SFilterExpressionHelpDialog::OpenDocumentationLink() const
 {
-	if (!DocumentationLink.IsEmpty())
+	if (!Config.DocumentationLink.IsEmpty())
 	{
-		FPlatformProcess::LaunchURL(*DocumentationLink, nullptr, nullptr);
+		FPlatformProcess::LaunchURL(*Config.DocumentationLink, nullptr, nullptr);
 	}
 }
 
diff --git a/Engine/Source/Editor/Sequencer/Private/Filters/Widgets/SFilterExpressionHelpDialog.h b/Engine/Source/Editor/Sequencer/Private/Filters/Widgets/SFilterExpressionHelpDialog.h
--- a/Engine/Source/Editor/Sequencer/Private/Filters/Widgets/SFilterExpressionHelpDialog.h
+++ b/Engine/Source/Editor/Sequencer/Private/Filters/Widgets/SFilterExpressionHelpDialog.h
@@ -4,42 +4,51 @@
 
 #include "Widgets/SWindow.h"
 
-class FSequencerTextFilterExpressionContext;
+class ISequencerTextFilterExpressionContext;
 enum class ESequencerTextFilterValueType : uint8;
 
+struct FFilterExpressionHelpDialogConfig
+{
+	static constexpr float DefaultMaxDesiredWidth = 460.f;
+	static constexpr float DefaultMaxDesiredHeight = 560.f;
+
+	FFilterExpressionHelpDialogConfig();
+
+	FName IdentifierName;
+
+	FText DialogTitle;
+	FString DocumentationLink;
+
+	TArray<TSharedRef<ISequencerTextFilterExpressionContext>> TextFilterExpressionContexts;
+
+	float MaxDesiredWidth = DefaultMaxDesiredWidth;
+	float MaxDesiredHeight = DefaultMaxDesiredHeight;
+};
+
 class SFilterExpressionHelpDialog : public SWindow
 {
 public:
-	static constexpr float MaxDesiredWidth = 460.f;
-	static constexpr float MaxDesiredHeight = 560.f;
-
-	SLATE_BEGIN_ARGS(SFilterExpressionHelpDialog)
-		: _DialogTitle(NSLOCTEXT("SFilterExpressionHelpDialog", "DialogTitle", "Text Filter Expression Help"))
-		, _HeaderText(NSLOCTEXT("SFilterExpressionHelpDialog", "HeaderText", "Text Filter Expressions"))
-		, _DocumentationLink(TEXT("https://dev.epicgames.com/documentation/en-us/unreal-engine/advanced-search-syntax-in-unreal-engine"))
-	{}
-		SLATE_ARGUMENT(FText, DialogTitle)
-		SLATE_ARGUMENT(FText, HeaderText)
-		SLATE_ARGUMENT(FString, DocumentationLink)
-		SLATE_ARGUMENT(TArray<TSharedRef<FSequencerTextFilterExpressionContext>>, TextFilterExpressionContexts)
-	SLATE_END_ARGS()
-
-	void Construct(const FArguments& InArgs);
+	static void Open(FFilterExpressionHelpDialogConfig&& InConfig);
+
+	static bool IsOpen(const FName InName);
+
+	static void CloseWindow(const FName InName);
+
+	void Construct(const FArguments& InArgs, FFilterExpressionHelpDialogConfig&& InConfig);
 
 protected:
 	static const FSlateColor KeyColor;
 	static const FSlateColor ValueColor;
 
+	static TMap<FName, TSharedPtr<SFilterExpressionHelpDialog>> DialogInstance;
+
 	TSharedRef<SWidget> ConstructDialogHeader();
 	TSharedRef<SWidget> ConstructExpressionWidgetList();
-	TSharedRef<SWidget> ConstructExpressionWidget(const TSharedPtr<FSequencerTextFilterExpressionContext>& InExpressionContext);
+	TSharedRef<SWidget> ConstructExpressionWidget(const TSharedPtr<ISequencerTextFilterExpressionContext>& InExpressionContext);
 	TSharedRef<SWidget> ConstructKeysWidget(const TSet<FName>& InKeys);
 	TSharedRef<SWidget> ConstructValueWidget(const ESequencerTextFilterValueType InValueType);
 
 	void OpenDocumentationLink() const;
 
-	FText HeaderText;
-	FString DocumentationLink;
-
-	TArray<TSharedRef<FSequencerTextFilterExpressionContext>> TextFilterExpressionContexts;
+	FFilterExpressionHelpDialogConfig Config;
 };
diff --git a/Engine/Source/Editor/Sequencer/Private/Filters/Widgets/SSequencerFilter.cpp b/Engine/Source/Editor/Sequencer/Private/Filters/Widgets/SSequencerFilter.cpp
--- a/Engine/Source/Editor/Sequencer/Private/Filters/Widgets/SSequencerFilter.cpp
+++ b/Engine/Source/Editor/Sequencer/Private/Filters/Widgets/SSequencerFilter.cpp
@@ -1,20 +1,25 @@
 ﻿// Copyright Epic Games, Inc. All Rights Reserved.
 
 #include "Filters/Widgets/SSequencerFilter.h"
-#include "Filters/Menus/SequencerTrackFilterContextMenu.h"
-#include "Filters/SequencerFilterBar.h"
 #include "Filters/Widgets/SSequencerFilterCheckBox.h"
 
 #define LOCTEXT_NAMESPACE "SSequencerFilter"
 
-void SSequencerFilter::Construct(const FArguments& InArgs
-	, const TSharedRef<FSequencerFilterBar>& InFilterBar
-	, const TSharedRef<FSequencerTrackFilter>& InFilter)
+void SSequencerFilter::Construct(const FArguments& InArgs)
 {
-	WeakFilterBar = InFilterBar;
-	WeakFilter = InFilter;
+	IsFilterActiveDelegate = InArgs._OnIsFilterActive;
 
-	ContextMenu = MakeShared<FSequencerTrackFilterContextMenu>();
+	FilterToggleDelegate = InArgs._OnFilterToggle;
+	CtrlClickDelegate = InArgs._OnCtrlClick;
+	AltClickDelegate = InArgs._OnAltClick;
+	MiddleClickDelegate = InArgs._OnMiddleClick;
+	DoubleClickDelegate = InArgs._OnDoubleClick;
+
+	DisplayName = InArgs._DisplayName;
+	ToolTipText = InArgs._ToolTipText;
+	BlockColor = InArgs._BlockColor;
+
+	GetMenuContentDelegate = InArgs._OnGetMenuContent;
 
 	TSharedPtr<SWidget> ContentWidget;
 	FName BrushName;
@@ -40,7 +45,7 @@ void SSequencerFilter::Construct(const FArguments& InArgs
 	[
 		SAssignNew(ToggleButtonPtr, SSequencerFilterCheckBox)
 		.Style(FAppStyle::Get(), BrushName)
-		.ToolTipText(TAttribute<FText>::Create(TAttribute<FText>::FGetter::CreateSP(InFilter, &FSequencerTrackFilter::GetToolTipText)))
+		.ToolTipText(ToolTipText)
 		.IsChecked(this, &SSequencerFilter::IsChecked)
 		.OnCheckStateChanged(this, &SSequencerFilter::OnFilterToggled)
 		.CheckBoxContentUsesAutoWidth(false)
@@ -61,23 +66,11 @@ TSharedRef<SWidget> SSequencerFilter::ConstructBasicFilterWidget()
 	return SNew(STextBlock)
 		.Margin(0.f)
 		.TextStyle(FAppStyle::Get(), TEXT("SmallText"))
-		.Text(this, &SSequencerFilter::GetFilterDisplayName);
+		.Text(DisplayName);
 }
 
 TSharedRef<SWidget> SSequencerFilter::ConstructDefaultFilterWidget()
 {
-	const TSharedPtr<FSequencerFilterBar> FilterBar = WeakFilterBar.Pin();
-	if (!FilterBar.IsValid())
-	{
-		return SNullWidget::NullWidget;
-	}
-
-	const TSharedPtr<FSequencerTrackFilter> Filter = WeakFilter.Pin();
-	if (!Filter.IsValid())
-	{
-		return SNullWidget::NullWidget;
-	}
-
 	return SNew(SBorder)
 		.Padding(1.f)
 		.BorderImage(FAppStyle::Get().GetBrush(TEXT("FilterBar.FilterBackground")))
@@ -98,95 +91,84 @@ TSharedRef<SWidget> SSequencerFilter::ConstructDefaultFilterWidget()
 			[
 				SNew(STextBlock)
 				.Font(FCoreStyle::GetDefaultFontStyle(TEXT("Regular"), 8))
-				.Text(this, &SSequencerFilter::GetFilterDisplayName)
+				.Text(DisplayName)
 				.IsEnabled(this, &SSequencerFilter::IsButtonEnabled)
 			]
 		];
 }
 
-const TSharedPtr<FSequencerTrackFilter> SSequencerFilter::GetFilter() const
-{
-	return WeakFilter.IsValid() ? WeakFilter.Pin() : nullptr;
-}
-
 bool SSequencerFilter::IsActive() const
 {
-	const TSharedPtr<FSequencerFilterBar> FilterBar = WeakFilterBar.Pin();
-	if (!FilterBar.IsValid())
-	{
-		return false;
-	}
-
-	const TSharedPtr<FSequencerTrackFilter> Filter = WeakFilter.Pin();
-	if (!Filter.IsValid())
+	if (IsFilterActiveDelegate.IsBound())
 	{
-		return false;
+		return IsFilterActiveDelegate.Execute();
 	}
-
-	return FilterBar->IsFilterActive(Filter.ToSharedRef());
+	return false;
 }
 
 void SSequencerFilter::OnFilterToggled(const ECheckBoxState NewState)
 {
-	const TSharedPtr<FSequencerFilterBar> FilterBar = WeakFilterBar.Pin();
-	if (!FilterBar.IsValid())
-	{
-		return;
-	}
-
-	const TSharedPtr<FSequencerTrackFilter> Filter = WeakFilter.Pin();
-	if (!Filter.IsValid())
+	if (FilterToggleDelegate.IsBound())
 	{
-		return;
+		FilterToggleDelegate.Execute(NewState);
 	}
-
-	const bool bNewActive = NewState == ECheckBoxState::Checked;
-	FilterBar->SetFilterActive(Filter.ToSharedRef(), bNewActive, true);
 }
 
 FReply SSequencerFilter::OnFilterCtrlClick()
 {
-	ActivateAllButThis(false);
+	if (CtrlClickDelegate.IsBound())
+	{
+		CtrlClickDelegate.Execute();
+
+		return FReply::Handled();
+	}
 
-	return FReply::Handled();
+	return FReply::Unhandled();
 }
 
 FReply SSequencerFilter::OnFilterAltClick()
 {
-	ActivateAllButThis(true);
+	if (AltClickDelegate.IsBound())
+	{
+		AltClickDelegate.Execute();
 
-	return FReply::Handled();
+		return FReply::Handled();
+	}
+
+	return FReply::Unhandled();
 }
 
 FReply SSequencerFilter::OnFilterMiddleButtonClick()
 {
-	const TSharedPtr<FSequencerFilterBar> FilterBar = WeakFilterBar.Pin();
-	if (!FilterBar.IsValid())
+	if (MiddleClickDelegate.IsBound())
 	{
-		return FReply::Handled();
-	}
+		MiddleClickDelegate.Execute();
 
-	const TSharedPtr<FSequencerTrackFilter> Filter = WeakFilter.Pin();
-	if (!Filter.IsValid())
-	{
 		return FReply::Handled();
 	}
 
-	FilterBar->SetFilterEnabled(Filter.ToSharedRef(), false, true);
-
-	return FReply::Handled();
+	return FReply::Unhandled();
 }
 
 FReply SSequencerFilter::OnFilterDoubleClick()
 {
-	ActivateAllButThis(false);
+	if (DoubleClickDelegate.IsBound())
+	{
+		DoubleClickDelegate.Execute();
+
+		return FReply::Handled();
+	}
 
-	return FReply::Handled();
+	return FReply::Unhandled();
 }
 
 TSharedRef<SWidget> SSequencerFilter::GetRightClickMenuContent()
 {
-	return ContextMenu->CreateMenuWidget(SharedThis(this));
+	if (GetMenuContentDelegate.IsBound())
+	{
+		return GetMenuContentDelegate.Execute();
+	}
+	return SNullWidget::NullWidget;
 }
 
 ECheckBoxState SSequencerFilter::IsChecked() const
@@ -196,12 +178,7 @@ ECheckBoxState SSequencerFilter::IsChecked() const
 
 FSlateColor SSequencerFilter::GetFilterImageColorAndOpacity() const
 {
-	const TSharedPtr<FSequencerTrackFilter> Filter = WeakFilter.Pin();
-	if (!Filter.IsValid() || !IsActive())
-	{
-		return FAppStyle::Get().GetSlateColor(TEXT("Colors.Recessed"));;
-	}
-	return Filter->GetColor();
+	return BlockColor.Get(FSlateColor());
 }
 
 EVisibility SSequencerFilter::GetFilterOverlayVisibility() const
@@ -214,44 +191,13 @@ FMargin SSequencerFilter::GetFilterNamePadding() const
 	return ToggleButtonPtr->IsPressed() ? FMargin(3, 1, 3, 0) : FMargin(3, 0, 3, 0);
 }
 
-FText SSequencerFilter::GetFilterDisplayName() const
-{
-	return WeakFilter.IsValid() ? WeakFilter.Pin()->GetDisplayName() : FText::GetEmpty();
-}
-
 bool SSequencerFilter::IsButtonEnabled() const
 {
-	const TSharedPtr<FSequencerFilterBar> FilterBar = WeakFilterBar.Pin();
-	if (!FilterBar.IsValid())
+	if (IsFilterActiveDelegate.IsBound())
 	{
-		return false;
+		return IsFilterActiveDelegate.Execute();
 	}
-
-	const TSharedPtr<FSequencerTrackFilter> Filter = WeakFilter.Pin();
-	if (!Filter.IsValid())
-	{
-		return false;
-	}
-
-	return FilterBar->IsFilterActive(Filter.ToSharedRef());
-}
-
-void SSequencerFilter::ActivateAllButThis(const bool bInActive)
-{
-	const TSharedPtr<FSequencerFilterBar> FilterBar = WeakFilterBar.Pin();
-	if (!FilterBar.IsValid())
-	{
-		return;
-	}
-
-	const TSharedPtr<FSequencerTrackFilter> Filter = WeakFilter.Pin();
-	if (!Filter.IsValid())
-	{
-		return;
-	}
-
-	FilterBar->ActivateAllEnabledFilters(bInActive, {});
-	FilterBar->SetFilterActive(Filter.ToSharedRef(), !bInActive, true);
+	return true;
 }
 
 #undef LOCTEXT_NAMESPACE
diff --git a/Engine/Source/Editor/Sequencer/Private/Filters/Widgets/SSequencerFilter.h b/Engine/Source/Editor/Sequencer/Private/Filters/Widgets/SSequencerFilter.h
--- a/Engine/Source/Editor/Sequencer/Private/Filters/Widgets/SSequencerFilter.h
+++ b/Engine/Source/Editor/Sequencer/Private/Filters/Widgets/SSequencerFilter.h
@@ -2,27 +2,42 @@
 
 #pragma once
 
-#include "Filters/SequencerTrackFilterBase.h"
-#include "SSequencerFilterBar.h"
+#include "Filters/SBasicFilterBar.h"
 
-class FSequencerTrackFilterContextMenu;
 class SSequencerFilterCheckBox;
 
+/**
+ * Generic Sequencer filter class used by Sequencer and Navigation Tool filter bars.
+ */
 class SSequencerFilter : public SCompoundWidget
 {
 public:
-	DECLARE_DELEGATE_OneParam(FOnSequencerFilterRequest, const TSharedRef<SSequencerFilter>& /*InFilter*/);
+	DECLARE_DELEGATE_RetVal(bool, FOnIsFilterActive);
+	DECLARE_DELEGATE_OneParam(FOnFilterToggle, const ECheckBoxState /*InNewState*/);
+
+	DECLARE_DELEGATE_RetVal(TSharedRef<SWidget>, FOnGetMenuContent);
 
 	SLATE_BEGIN_ARGS(SSequencerFilter) {}
+
 		/** Determines how each individual filter pill looks like */
 		SLATE_ARGUMENT(EFilterPillStyle, FilterPillStyle)
-	SLATE_END_ARGS()
 
-	void Construct(const FArguments& InArgs
-		, const TSharedRef<FSequencerFilterBar>& InFilterBar
-		, const TSharedRef<FSequencerTrackFilter>& InFilter);
+		SLATE_EVENT(FOnIsFilterActive, OnIsFilterActive)
+		SLATE_EVENT(FOnFilterToggle, OnFilterToggle)
+		SLATE_EVENT(FSimpleDelegate, OnCtrlClick)
+		SLATE_EVENT(FSimpleDelegate, OnAltClick)
+		SLATE_EVENT(FSimpleDelegate, OnMiddleClick)
+		SLATE_EVENT(FSimpleDelegate, OnDoubleClick)
+
+		SLATE_ATTRIBUTE(FText, DisplayName)
+		SLATE_ATTRIBUTE(FText, ToolTipText)
+		SLATE_ATTRIBUTE(FSlateColor, BlockColor)
+
+		SLATE_EVENT(FOnGetMenuContent, OnGetMenuContent)
 
-	const TSharedPtr<FSequencerTrackFilter> GetFilter() const;
+	SLATE_END_ARGS()
+
+	void Construct(const FArguments& InArgs);
 
 protected:
 	TSharedRef<SWidget> ConstructBasicFilterWidget();
@@ -47,17 +62,21 @@ protected:
 
 	FMargin GetFilterNamePadding() const;
 
-	FText GetFilterDisplayName() const;
-
 	bool IsButtonEnabled() const;
 
-	void ActivateAllButThis(const bool bInActive);
+	FOnIsFilterActive IsFilterActiveDelegate;
 
-	TWeakPtr<FSequencerFilterBar> WeakFilterBar;
+	FOnFilterToggle FilterToggleDelegate;
+	FSimpleDelegate CtrlClickDelegate;
+	FSimpleDelegate AltClickDelegate;
+	FSimpleDelegate MiddleClickDelegate;
+	FSimpleDelegate DoubleClickDelegate;
 
-	TWeakPtr<FSequencerTrackFilter> WeakFilter;
+	TAttribute<FText> DisplayName;
+	TAttribute<FText> ToolTipText;
+	TAttribute<FSlateColor> BlockColor;
 
-	TSharedPtr<SSequencerFilterCheckBox> ToggleButtonPtr;
+	FOnGetMenuContent GetMenuContentDelegate;
 
-	TSharedPtr<FSequencerTrackFilterContextMenu> ContextMenu;
+	TSharedPtr<SSequencerFilterCheckBox> ToggleButtonPtr;
 };
diff --git a/Engine/Source/Editor/Sequencer/Private/Filters/Widgets/SSequencerFilterBar.cpp b/Engine/Source/Editor/Sequencer/Private/Filters/Widgets/SSequencerFilterBar.cpp
--- a/Engine/Source/Editor/Sequencer/Private/Filters/Widgets/SSequencerFilterBar.cpp
+++ b/Engine/Source/Editor/Sequencer/Private/Filters/Widgets/SSequencerFilterBar.cpp
@@ -2,21 +2,17 @@
 
 #include "Filters/Widgets/SSequencerFilterBar.h"
 #include "Filters/Filters/SequencerTrackFilter_CustomText.h"
-#include "Filters/Filters/SequencerTrackFilter_Text.h"
 #include "Filters/Menus/SequencerFilterBarContextMenu.h"
+#include "Filters/Menus/SequencerTrackFilterContextMenu.h"
 #include "Filters/SequencerFilterBar.h"
 #include "Filters/Widgets/SFilterBarClippingHorizontalBox.h"
 #include "Filters/Widgets/SFilterExpressionHelpDialog.h"
-#include "Interfaces/IMainFrameModule.h"
 #include "MovieScene.h"
 #include "Sequencer.h"
 #include "SequencerLog.h"
 #include "SSequencerFilter.h"
 #include "SSequencerCustomTextFilterDialog.h"
 #include "Widgets/SSequencerSearchBox.h"
-#include "Widgets/SWindow.h"
-
-using namespace UE::Sequencer;
 
 #define LOCTEXT_NAMESPACE "SSequencerFilterBar"
 
@@ -164,9 +160,9 @@ void SSequencerFilterBar::SetLayout(const EFilterBarLayout InFilterBarLayout)
 	HorizontalContainerWidget->ClearChildren();
 	VerticalContainerWidget->ClearChildren();
 
-	for (const TSharedRef<SSequencerFilter>& FilterWidget : FilterWidgets)
+	for (const TPair<TSharedRef<FSequencerTrackFilter>, TSharedRef<SSequencerFilter>>& FilterPair : FilterWidgets)
 	{
-		AddWidgetToLayout(FilterWidget);
+		AddWidgetToLayout(FilterPair.Value);
 	}
 }
 
@@ -232,101 +228,159 @@ void SSequencerFilterBar::RemoveWidgetFromLayout(const TSharedRef<SWidget>& InWi
 	}
 }
 
-TSharedPtr<SSequencerFilter> SSequencerFilterBar::FindFilterWidget(const TSharedRef<FSequencerTrackFilter>& InFilter) const
+void SSequencerFilterBar::CreateAndAddFilterWidget(const TSharedRef<FSequencerTrackFilter>& InFilter)
 {
-	for (const TSharedRef<SSequencerFilter>& FilterWidget : FilterWidgets)
+	const TSharedPtr<FSequencerFilterBar> FilterBar = WeakFilterBar.Pin();
+	if (!FilterBar.IsValid())
 	{
-		if (FilterWidget->GetFilter() == InFilter)
-		{
-			return FilterWidget;
-		}
+		return;
 	}
-	return nullptr;
+
+	const TSharedRef<FSequencerFilterBar> FilterBarRef = FilterBar.ToSharedRef();
+
+	const TSharedRef<SSequencerFilter> NewFilterWidget = SNew(SSequencerFilter)
+		.FilterPillStyle(FilterPillStyle)
+		.DisplayName(this, &SSequencerFilterBar::GetFilterDisplayName, InFilter)
+		.ToolTipText(InFilter, &FSequencerTrackFilter::GetToolTipText)
+		.BlockColor(this, &SSequencerFilterBar::GetFilterBlockColor, InFilter)
+		.OnIsFilterActive(FilterBarRef, &FSequencerFilterBar::IsFilterActive, InFilter)
+		.OnFilterToggle(this, &SSequencerFilterBar::OnFilterToggle, InFilter)
+		.OnAltClick(this, &SSequencerFilterBar::OnFilterAltClick, InFilter)
+		.OnMiddleClick(this, &SSequencerFilterBar::OnFilterMiddleClick, InFilter)
+		.OnDoubleClick(this, &SSequencerFilterBar::OnFilterDoubleClick, InFilter)
+		.OnGetMenuContent(this, &SSequencerFilterBar::OnGetMenuContent, InFilter);
+
+	AddFilterWidget(InFilter, NewFilterWidget);
 }
 
-void SSequencerFilterBar::CreateAndAddFilterWidget(const TSharedRef<FSequencerTrackFilter>& InFilter)
+FText SSequencerFilterBar::GetFilterDisplayName(const TSharedRef<FSequencerTrackFilter> InFilter) const
+{
+	return InFilter->GetDisplayName();
+}
+
+FSlateColor SSequencerFilterBar::GetFilterBlockColor(const TSharedRef<FSequencerTrackFilter> InFilter) const
 {
 	const TSharedPtr<FSequencerFilterBar> FilterBar = WeakFilterBar.Pin();
 	if (!FilterBar.IsValid())
 	{
-		return;
+		return FSlateColor();
+	}
+
+	if (!FilterBar->IsFilterActive(InFilter))
+	{
+		return FAppStyle::Get().GetSlateColor(TEXT("Colors.Recessed"));;
 	}
 
-	const TSharedRef<SSequencerFilter> NewFilter = SNew(SSequencerFilter, FilterBar.ToSharedRef(), InFilter)
-		.FilterPillStyle(FilterPillStyle);
+	return InFilter->GetColor();
+}
+
+void SSequencerFilterBar::OnFilterToggle(const ECheckBoxState InNewState, const TSharedRef<FSequencerTrackFilter> InFilter)
+{
+	const TSharedPtr<FSequencerFilterBar> FilterBar = WeakFilterBar.Pin();
+	if (!FilterBar.IsValid())
+	{
+		return;
+	}
 
-	AddFilterWidget(NewFilter);
+	const bool bNewActive = InNewState == ECheckBoxState::Checked;
+	FilterBar->SetFilterActive(InFilter, bNewActive, true);
 }
 
-void SSequencerFilterBar::AddFilterWidget(const TSharedRef<SSequencerFilter>& InFilterWidget)
+void SSequencerFilterBar::OnFilterCtrlClick(const TSharedRef<FSequencerTrackFilter> InFilter)
 {
-	FilterWidgets.Add(InFilterWidget);
+	ActivateAllButThis(false, InFilter);
+}
 
-	AddWidgetToLayout(InFilterWidget);
+void SSequencerFilterBar::OnFilterAltClick(const TSharedRef<FSequencerTrackFilter> InFilter)
+{
+	ActivateAllButThis(true, InFilter);
 }
 
-void SSequencerFilterBar::RemoveFilterWidget(const TSharedRef<FSequencerTrackFilter>& InFilter, bool ExecuteOnFilterChanged)
+void SSequencerFilterBar::OnFilterMiddleClick(const TSharedRef<FSequencerTrackFilter> InFilter)
 {
-	TSharedPtr<SSequencerFilter> FilterToRemove;
-	for (const TSharedRef<SSequencerFilter>& Filter : FilterWidgets)
+	const TSharedPtr<FSequencerFilterBar> FilterBar = WeakFilterBar.Pin();
+	if (!FilterBar.IsValid())
 	{
-		if (Filter->GetFilter() == InFilter)
-		{
-			FilterToRemove = Filter;
-			break;
-		}
+		return;
 	}
 
-	if (FilterToRemove.IsValid())
+	FilterBar->SetFilterEnabled(InFilter, false, true);
+}
+
+void SSequencerFilterBar::OnFilterDoubleClick(const TSharedRef<FSequencerTrackFilter> InFilter)
+{
+	ActivateAllButThis(false, InFilter);
+}
+
+TSharedRef<SWidget> SSequencerFilterBar::OnGetMenuContent(const TSharedRef<FSequencerTrackFilter> InFilter)
+{
+	const TSharedRef<FSequencerTrackFilterContextMenu> FilterContextMenu = MakeShared<FSequencerTrackFilterContextMenu>();
+	return FilterContextMenu->CreateMenuWidget(InFilter);
+}
+
+void SSequencerFilterBar::ActivateAllButThis(const bool bInActive, const TSharedRef<FSequencerTrackFilter> InFilter)
+{
+	const TSharedPtr<FSequencerFilterBar> FilterBar = WeakFilterBar.Pin();
+	if (!FilterBar.IsValid())
 	{
-		if (ExecuteOnFilterChanged)
-		{
-			RemoveFilterWidgetAndUpdate(FilterToRemove.ToSharedRef());
-		}
-		else
-		{
-			RemoveFilterWidget(FilterToRemove.ToSharedRef());
-		}
+		return;
 	}
+
+	FilterBar->ActivateAllEnabledFilters(bInActive, {});
+	FilterBar->SetFilterActive(InFilter, !bInActive, true);
 }
 
-void SSequencerFilterBar::RemoveFilterWidget(const TSharedRef<SSequencerFilter>& InFilterWidget)
+void SSequencerFilterBar::AddFilterWidget(const TSharedRef<FSequencerTrackFilter>& InFilter, const TSharedRef<SSequencerFilter>& InFilterWidget)
 {
-	FilterWidgets.Remove(InFilterWidget);
+	FilterWidgets.Add(InFilter, InFilterWidget);
 
-	RemoveWidgetFromLayout(InFilterWidget);
+	AddWidgetToLayout(InFilterWidget);
+}
+
+void SSequencerFilterBar::RemoveFilterWidget(const TSharedRef<FSequencerTrackFilter>& InFilter)
+{
+	if (!FilterWidgets.Contains(InFilter))
+	{
+		return;
+	}
+
+	RemoveWidgetFromLayout(FilterWidgets[InFilter]);
+
+	FilterWidgets.Remove(InFilter);
 }
 
 void SSequencerFilterBar::RemoveAllFilterWidgets()
 {
-	for (const TSharedRef<SSequencerFilter>& FilterWidget : FilterWidgets)
+	for (const TPair<TSharedRef<FSequencerTrackFilter>, TSharedRef<SSequencerFilter>>& FilterPair : FilterWidgets)
 	{
-		RemoveWidgetFromLayout(FilterWidget);
+		RemoveWidgetFromLayout(FilterPair.Value);
 	}
 
 	FilterWidgets.Empty();
 }
 
-void SSequencerFilterBar::RemoveAllFilterWidgetsButThis(const TSharedRef<SSequencerFilter>& InFilterWidget)
+void SSequencerFilterBar::RemoveAllFilterWidgetsButThis(const TSharedRef<FSequencerTrackFilter>& InFilter)
 {
-	for (const TSharedRef<SSequencerFilter>& FilterWidget : FilterWidgets)
+	if (!FilterWidgets.Contains(InFilter))
+	{
+		return;
+	}
+
+	for (const TPair<TSharedRef<FSequencerTrackFilter>, TSharedRef<SSequencerFilter>>& FilterPair : FilterWidgets)
 	{
-		if (FilterWidget == InFilterWidget)
+		if (FilterPair.Key == InFilter)
 		{
 			continue;
 		}
 
-		RemoveWidgetFromLayout(FilterWidget);
+		RemoveWidgetFromLayout(FilterPair.Value);
 	}
 
-	FilterWidgets.Empty();
+	const TSharedRef<SSequencerFilter> PreviousWidget = FilterWidgets[InFilter];
 
-	AddFilterWidget(InFilterWidget);
-}
+	FilterWidgets.Empty();
 
-void SSequencerFilterBar::RemoveFilterWidgetAndUpdate(const TSharedRef<SSequencerFilter>& InFilterWidget)
-{
-	RemoveFilterWidget(InFilterWidget);
+	AddFilterWidget(InFilter, PreviousWidget);
 }
 
 void SSequencerFilterBar::OnEnableAllGroupFilters(bool bEnableAll)
@@ -410,8 +464,7 @@ void SSequencerFilterBar::OnFiltersChanged(const ESequencerFilterChange InChange
 	case ESequencerFilterChange::Enable:
 	case ESequencerFilterChange::Activate:
 		{
-			const TSharedPtr<SSequencerFilter> FilterWidget = FindFilterWidget(InFilter);
-			if (!FilterWidget.IsValid())
+			if (!FilterWidgets.Contains(InFilter))
 			{
 				CreateAndAddFilterWidget(InFilter);
 			}
@@ -452,37 +505,12 @@ void SSequencerFilterBar::OnOpenTextExpressionHelp()
 		return;
 	}
 
-	if (TextExpressionHelpDialog.IsValid())
-	{
-		TextExpressionHelpDialog->BringToFront();
-	}
-	else
-	{
-		TextExpressionHelpDialog = SNew(SFilterExpressionHelpDialog)
-			.DialogTitle(LOCTEXT("SequencerCustomTextFilterHelp", "Sequencer Custom Text Filter Help"))
-			.TextFilterExpressionContexts(FilterBar->GetTextFilter()->GetTextFilterExpressionContexts());
+	FFilterExpressionHelpDialogConfig Config;
+	Config.IdentifierName = TEXT("SequencerCustomTextFilterHelp");
+	Config.DialogTitle = LOCTEXT("SequencerCustomTextFilterHelp", "Sequencer Custom Text Filter Help");
+	Config.TextFilterExpressionContexts = FilterBar->GetTextFilterExpressionContexts();
 
-		TextExpressionHelpDialog->GetOnWindowClosedEvent().AddLambda([this](const TSharedRef<SWindow>& InWindow)
-			{
-				TextExpressionHelpDialog.Reset();
-			});
-
-		TSharedPtr<SWindow> ParentWindow;
-		if (FModuleManager::Get().IsModuleLoaded(TEXT("MainFrame")))
-		{
-			IMainFrameModule& MainFrame = FModuleManager::LoadModuleChecked<IMainFrameModule>(TEXT("MainFrame"));
-			ParentWindow = MainFrame.GetParentWindow();
-		}
-
-		if (ParentWindow.IsValid())
-		{
-			FSlateApplication::Get().AddWindowAsNativeChild(TextExpressionHelpDialog.ToSharedRef(), ParentWindow.ToSharedRef());
-		}
-		else
-		{
-			FSlateApplication::Get().AddWindow(TextExpressionHelpDialog.ToSharedRef());
-		}
-	}
+	SFilterExpressionHelpDialog::Open(MoveTemp(Config));
 }
 void SSequencerFilterBar::SaveCurrentFilterSetAsCustomTextFilter()
 {
@@ -539,8 +567,7 @@ void SSequencerFilterBar::CreateFilterWidgetsFromConfig()
 			return;
 		}
 
-		const TSharedPtr<SSequencerFilter> FilterWidget = FindFilterWidget(InFilter);
-		if (!FilterWidget.IsValid())
+		if (!FilterWidgets.Contains(InFilter))
 		{
 			CreateAndAddFilterWidget(InFilter);
 		}
diff --git a/Engine/Source/Editor/Sequencer/Private/Filters/Widgets/SSequencerFilterBar.h b/Engine/Source/Editor/Sequencer/Private/Filters/Widgets/SSequencerFilterBar.h
--- a/Engine/Source/Editor/Sequencer/Private/Filters/Widgets/SSequencerFilterBar.h
+++ b/Engine/Source/Editor/Sequencer/Private/Filters/Widgets/SSequencerFilterBar.h
@@ -93,20 +93,14 @@ protected:
 	void AddWidgetToLayout(const TSharedRef<SWidget>& InWidget);
 	void RemoveWidgetFromLayout(const TSharedRef<SWidget>& InWidget);
 
-	TSharedPtr<SSequencerFilter> FindFilterWidget(const TSharedRef<FSequencerTrackFilter>& InFilter) const;
-
 	void CreateAndAddFilterWidget(const TSharedRef<FSequencerTrackFilter>& InFilter);
-	void AddFilterWidget(const TSharedRef<SSequencerFilter>& InFilterWidget);
-
-	void RemoveFilterWidget(const TSharedRef<FSequencerTrackFilter>& InFilter, const bool ExecuteOnFilterChanged = true);
-	void RemoveFilterWidget(const TSharedRef<SSequencerFilter>& InFilterWidget);
+	void AddFilterWidget(const TSharedRef<FSequencerTrackFilter>& InFilter, const TSharedRef<SSequencerFilter>& InFilterWidget);
+	void RemoveFilterWidget(const TSharedRef<FSequencerTrackFilter>& InFilter);
+	void RemoveAllFilterWidgets();
+	void RemoveAllFilterWidgetsButThis(const TSharedRef<FSequencerTrackFilter>& InFilter);
 
 	UWorld* GetWorld() const;
 
-	void RemoveAllFilterWidgets();
-	void RemoveAllFilterWidgetsButThis(const TSharedRef<SSequencerFilter>& InFilterWidget);
-	void RemoveFilterWidgetAndUpdate(const TSharedRef<SSequencerFilter>& InFilterWidget);
-
 	void OnEnableAllGroupFilters(bool bEnableAll);
 	void OnNodeGroupFilterClicked(UMovieSceneNodeGroup* NodeGroup);
 
@@ -116,6 +110,19 @@ protected:
 
 	TSharedRef<SWidget> OnWrapButtonClicked();
 
+	FText GetFilterDisplayName(const TSharedRef<FSequencerTrackFilter> InFilter) const;
+	FSlateColor GetFilterBlockColor(const TSharedRef<FSequencerTrackFilter> InFilter) const;
+
+	void OnFilterToggle(const ECheckBoxState InNewState, const TSharedRef<FSequencerTrackFilter> InFilter);
+	void OnFilterCtrlClick(const TSharedRef<FSequencerTrackFilter> InFilter);
+	void OnFilterAltClick(const TSharedRef<FSequencerTrackFilter> InFilter);
+	void OnFilterMiddleClick(const TSharedRef<FSequencerTrackFilter> InFilter);
+	void OnFilterDoubleClick(const TSharedRef<FSequencerTrackFilter> InFilter);
+
+	TSharedRef<SWidget> OnGetMenuContent(const TSharedRef<FSequencerTrackFilter> InFilter);
+
+	void ActivateAllButThis(const bool bInActive, const TSharedRef<FSequencerTrackFilter> InFilter);
+
 	TWeakPtr<FSequencerFilterBar> WeakFilterBar;
 
 	TWeakPtr<SSequencerSearchBox> WeakSearchBox;
@@ -128,7 +135,7 @@ protected:
 	bool bCanChangeOrientation = true;
 	EFilterPillStyle FilterPillStyle = EFilterPillStyle::Default;
 
-	TArray<TSharedRef<SSequencerFilter>> FilterWidgets;
+	TMap<TSharedRef<FSequencerTrackFilter>, TSharedRef<SSequencerFilter>> FilterWidgets;
 
 	TSharedPtr<SFilterExpressionHelpDialog> TextExpressionHelpDialog;
 
diff --git a/Engine/Source/Editor/Sequencer/Private/SSequencer.cpp b/Engine/Source/Editor/Sequencer/Private/SSequencer.cpp
--- a/Engine/Source/Editor/Sequencer/Private/SSequencer.cpp
+++ b/Engine/Source/Editor/Sequencer/Private/SSequencer.cpp
@@ -1323,7 +1323,7 @@ TSharedRef<SWidget> SSequencer::ConstructSearchAndFilterRow()
 
 	if (!SearchBox.IsValid())
 	{
-		SearchBox = SNew(SSequencerSearchBox, SharedThis(this))
+		SearchBox = SNew(SSequencerSearchBox, FilterBar)
 			.AddMetaData<FTagMetaData>(FTagMetaData(TEXT("SequencerFilterSearch")))
 			.HintText(LOCTEXT("FilterSearch", "Search..."))
 			.ToolTipText(LOCTEXT("FilterSearchHint", "Type here to search (pressing enter selects the results)"))
diff --git a/Engine/Source/Editor/Sequencer/Private/Widgets/SSequencerSearchBox.cpp b/Engine/Source/Editor/Sequencer/Private/Widgets/SSequencerSearchBox.cpp
--- a/Engine/Source/Editor/Sequencer/Private/Widgets/SSequencerSearchBox.cpp
+++ b/Engine/Source/Editor/Sequencer/Private/Widgets/SSequencerSearchBox.cpp
@@ -1,11 +1,11 @@
 ﻿// Copyright Epic Games, Inc. All Rights Reserved.
 
 #include "SSequencerSearchBox.h"
-#include "Framework/Application/SlateApplication.h"
-#include "SSequencer.h"
+#include "Filters/ISequencerFilterBar.h"
 #include "Filters/Filters/SequencerTrackFilter_Text.h"
-#include "Filters/SequencerFilterBar.h"
 #include "Filters/Widgets/SSequencerFilterSuggestionListRow.h"
+#include "Framework/Application/SlateApplication.h"
+#include "SequencerSettings.h"
 #include "Widgets/Input/SMenuAnchor.h"
 #include "Widgets/Layout/SBorder.h"
 #include "Widgets/Layout/SBox.h"
@@ -43,9 +43,9 @@ struct FSequencerSearchCategoryKeyMapFuncs : BaseKeyFuncs<ValueType, FText, /*bI
 	}
 };
 
-void SSequencerSearchBox::Construct(const FArguments& InArgs, const TWeakPtr<SSequencer>& InWeakSequencerWidget)
+void SSequencerSearchBox::Construct(const FArguments& InArgs, const TWeakPtr<ISequencerFilterBar>& InWeakFilterBar)
 {
-	WeakSequencerWidget = InWeakSequencerWidget;
+	WeakFilterBar = InWeakFilterBar;
 
 	OnTextChanged = InArgs._OnTextChanged;
 	OnTextCommitted = InArgs._OnTextCommitted;
@@ -450,19 +450,13 @@ void SSequencerSearchBox::SetOnSaveSearchHandler(SFilterSearchBox::FOnSaveSearch
 
 void SSequencerSearchBox::ExtractSearchFilterTerms(const FText& InSearchText, FParsedExpression& OutParsedPair)
 {
-	const TSharedPtr<SSequencer> SequencerWidget = WeakSequencerWidget.Pin();
-	if (!SequencerWidget.IsValid())
-	{
-		return;
-	}
-
-	const TSharedPtr<FSequencerFilterBar> FilterBar = SequencerWidget->GetFilterBar();
+	const TSharedPtr<ISequencerFilterBar> FilterBar = WeakFilterBar.Pin();
 	if (!FilterBar.IsValid())
 	{
 		return;
 	}
 
-	const FTextFilterExpressionEvaluator& TextFilterExpressionEvaluator = FilterBar->GetTextFilter()->GetTextFilterExpressionEvaluator();
+	const FTextFilterExpressionEvaluator& TextFilterExpressionEvaluator = FilterBar->GetTextFilterExpressionEvaluator();
 	const TArray<FExpressionToken>& ExpressionTokens = TextFilterExpressionEvaluator.GetFilterExpressionTokens();
 	if (ExpressionTokens.Num() == 0)
 	{
@@ -573,21 +567,15 @@ FText SSequencerSearchBox::DefaultSuggestionChosenImpl(const FText& InSearchText
 
 void SSequencerSearchBox::DefaultKeySuggestions(const FString& InKeyValue, TArray<FSequencerFilterSuggestion>& OutPossibleSuggestions)
 {
-	const TSharedPtr<SSequencer> SequencerWidget = WeakSequencerWidget.Pin();
-	if (!SequencerWidget.IsValid())
-	{
-		return;
-	}
-
-	const TSharedPtr<FSequencerFilterBar> FilterBar = SequencerWidget->GetFilterBar();
+	const TSharedPtr<ISequencerFilterBar> FilterBar = WeakFilterBar.Pin();
 	if (!FilterBar.IsValid())
 	{
 		return;
 	}
 
-	const TArray<TSharedRef<FSequencerTextFilterExpressionContext>>& FilterExpressions = FilterBar->GetTextFilter()->GetTextFilterExpressionContexts();
+	const TArray<TSharedRef<ISequencerTextFilterExpressionContext>>& FilterExpressions = FilterBar->GetTextFilterExpressionContexts();
 
-	for (const TSharedRef<FSequencerTextFilterExpressionContext>& Expression : FilterExpressions)
+	for (const TSharedRef<ISequencerTextFilterExpressionContext>& Expression : FilterExpressions)
 	{
 		const TArray<FName> Keys = Expression->GetKeys().Array();
 		for (const FName Key : Keys)
@@ -611,24 +599,18 @@ void SSequencerSearchBox::DefaultKeySuggestions(const FString& InKeyValue, TArra
 
 void SSequencerSearchBox::DefaultValueSuggestions(const FString& InKeyValue, TArray<FSequencerFilterSuggestion>& OutPossibleSuggestions)
 {
-	const TSharedPtr<SSequencer> SequencerWidget = WeakSequencerWidget.Pin();
-	if (!SequencerWidget.IsValid())
-	{
-		return;
-	}
-
-	const TSharedPtr<FSequencerFilterBar> FilterBar = SequencerWidget->GetFilterBar();
+	const TSharedPtr<ISequencerFilterBar> FilterBar = WeakFilterBar.Pin();
 	if (!FilterBar.IsValid())
 	{
 		return;
 	}
 
-	const TArray<TSharedRef<FSequencerTextFilterExpressionContext>>& FilterExpressions = FilterBar->GetTextFilter()->GetTextFilterExpressionContexts();
+	const TArray<TSharedRef<ISequencerTextFilterExpressionContext>> FilterExpressions = FilterBar->GetTextFilterExpressionContexts();
 
 	FSequencerTextFilterKeyword FoundKeyword;
-	TSharedPtr<FSequencerTextFilterExpressionContext> FoundExpression;
+	TSharedPtr<ISequencerTextFilterExpressionContext> FoundExpression;
 
-	for (const TSharedRef<FSequencerTextFilterExpressionContext>& Expression : FilterExpressions)
+	for (const TSharedRef<ISequencerTextFilterExpressionContext>& Expression : FilterExpressions)
 	{
 		for (const FName Key : Expression->GetKeys())
 		{
@@ -693,19 +675,13 @@ int32 SSequencerSearchBox::FindTokenIndex(const TArray<FExpressionToken>& Expres
 
 const FTextFilterExpressionEvaluator* SSequencerSearchBox::GetTextFilterExpressionEvaluator() const
 {
-	const TSharedPtr<SSequencer> SequencerWidget = WeakSequencerWidget.Pin();
-	if (!SequencerWidget.IsValid())
-	{
-		return nullptr;
-	}
-
-	const TSharedPtr<FSequencerFilterBar> FilterBar = SequencerWidget->GetFilterBar();
+	const TSharedPtr<ISequencerFilterBar> FilterBar = WeakFilterBar.Pin();
 	if (!FilterBar.IsValid())
 	{
 		return nullptr;
 	}
 
-	return &FilterBar->GetTextFilter()->GetTextFilterExpressionEvaluator();
+	return &FilterBar->GetTextFilterExpressionEvaluator();
 }
 
 #undef LOCTEXT_NAMESPACE
diff --git a/Engine/Source/Editor/Sequencer/Private/Widgets/SSequencerSearchBox.h b/Engine/Source/Editor/Sequencer/Private/Widgets/SSequencerSearchBox.h
--- a/Engine/Source/Editor/Sequencer/Private/Widgets/SSequencerSearchBox.h
+++ b/Engine/Source/Editor/Sequencer/Private/Widgets/SSequencerSearchBox.h
@@ -10,9 +10,9 @@
 class FExpressionToken;
 class FSequencerFilterSuggestionListEntryBase;
 class FTextFilterExpressionEvaluator;
+class ISequencerFilterBar;
 class ITableRow;
 class SMenuAnchor;
-class SSequencer;
 class STableViewBase;
 struct FFocusEvent;
 struct FGeometry;
@@ -91,7 +91,7 @@ public:
 	SLATE_END_ARGS()
 
 	/** Constructs this widget with InArgs */
-	void Construct(const FArguments& InArgs, const TWeakPtr<SSequencer>& InWeakSequencerWidget);
+	void Construct(const FArguments& InArgs, const TWeakPtr<ISequencerFilterBar>& InWeakFilterBar);
 
 	/** Get the text string currently being edited */
 	FText GetText() const;
@@ -171,7 +171,7 @@ private:
 
 	const FTextFilterExpressionEvaluator* GetTextFilterExpressionEvaluator() const;
 
-	TWeakPtr<SSequencer> WeakSequencerWidget;
+	TWeakPtr<ISequencerFilterBar> WeakFilterBar;
 
 	TSharedPtr<SFilterSearchBox> SearchBox;
 
diff --git a/Engine/Source/Editor/Sequencer/Public/Filters/ISequencerFilterBar.h b/Engine/Source/Editor/Sequencer/Public/Filters/ISequencerFilterBar.h
new file mode 100644
--- /dev/null
+++ b/Engine/Source/Editor/Sequencer/Public/Filters/ISequencerFilterBar.h
@@ -0,0 +1,62 @@
+﻿// Copyright Epic Games, Inc. All Rights Reserved.
+
+#pragma once
+
+#include "Framework/Commands/UICommandList.h"
+
+class FSequencerTrackFilter_CustomText;
+class FString;
+class FText;
+class FTextFilterExpressionEvaluator;
+class ISequencer;
+class ISequencerTextFilterExpressionContext;
+
+/** Generic interface for all Sequencer filter bar implementations. */
+class ISequencerFilterBar
+{
+public:
+	virtual ~ISequencerFilterBar() = default;
+
+	SEQUENCER_API virtual FName GetIdentifier() const = 0;
+
+	SEQUENCER_API virtual ISequencer& GetSequencer() const = 0;
+
+	SEQUENCER_API virtual TSharedPtr<FUICommandList> GetCommandList() const = 0;
+
+	SEQUENCER_API virtual FString GetTextFilterString() const = 0;
+	SEQUENCER_API virtual void SetTextFilterString(const FString& InText) = 0;
+
+	/** Returns true if the current filter bar text filter string contains the specified text expression.
+	 * The text expression must have key, operator, and value tokens. */
+	SEQUENCER_API virtual bool DoesTextFilterStringContainExpressionPair(const ISequencerTextFilterExpressionContext& InExpression) const = 0;
+
+	SEQUENCER_API virtual bool HasAnyFilterActive(const bool bCheckTextFilter = true
+		, const bool bInCheckHideIsolateFilter = true
+		, const bool bInCheckCommonFilters = true
+		, const bool bInCheckInternalFilters = true
+		, const bool bInCheckCustomTextFilters = true) const = 0;
+
+	SEQUENCER_API virtual bool HasAnyFilterEnabled() const = 0;
+
+	SEQUENCER_API virtual void RequestFilterUpdate() = 0;
+
+	SEQUENCER_API virtual void EnableAllFilters(const bool bInEnable, const TArray<FString>& InExceptionFilterNames) = 0;
+
+	SEQUENCER_API virtual void ActivateCommonFilters(const bool bInActivate, const TArray<FString>& InExceptionFilterNames) = 0;
+
+	SEQUENCER_API virtual void ActivateAllEnabledFilters(const bool bInActivate, const TArray<FString>& InExceptionFilterNames) = 0;
+
+	SEQUENCER_API virtual bool IsFilterActiveByDisplayName(const FString& InFilterName) const = 0;
+	SEQUENCER_API virtual bool IsFilterEnabledByDisplayName(const FString& InFilterName) const = 0;
+	SEQUENCER_API virtual bool SetFilterActiveByDisplayName(const FString& InFilterName, const bool bInActive, const bool bInRequestFilterUpdate = true) = 0;
+	SEQUENCER_API virtual bool SetFilterEnabledByDisplayName(const FString& InFilterName, const bool bInEnabled, const bool bInRequestFilterUpdate = true) = 0;
+
+	SEQUENCER_API virtual TArray<FText> GetFilterDisplayNames() const = 0;
+	SEQUENCER_API virtual TArray<FText> GetCustomTextFilterNames() const = 0;
+
+	SEQUENCER_API virtual int32 GetTotalDisplayNodeCount() const = 0;
+	SEQUENCER_API virtual int32 GetFilteredDisplayNodeCount() const  = 0;
+
+	SEQUENCER_API virtual const FTextFilterExpressionEvaluator& GetTextFilterExpressionEvaluator() const = 0;
+	SEQUENCER_API virtual TArray<TSharedRef<ISequencerTextFilterExpressionContext>> GetTextFilterExpressionContexts() const = 0;
+};
diff --git a/Engine/Source/Editor/Sequencer/Public/Filters/ISequencerTextFilterExpressionContext.h b/Engine/Source/Editor/Sequencer/Public/Filters/ISequencerTextFilterExpressionContext.h
new file mode 100644
--- /dev/null
+++ b/Engine/Source/Editor/Sequencer/Public/Filters/ISequencerTextFilterExpressionContext.h
@@ -0,0 +1,35 @@
+﻿// Copyright Epic Games, Inc. All Rights Reserved.
+
+#pragma once
+
+#include "Misc/TextFilterExpressionEvaluator.h"
+
+enum class SEQUENCER_API ESequencerTextFilterValueType : uint8
+{
+	String,
+	Boolean,
+	Integer
+};
+
+struct SEQUENCER_API FSequencerTextFilterKeyword
+{
+	FString Keyword;
+	FText Description;
+};
+
+/** Extends the ITextFilterExpressionContext interface to add support for suggestions. */
+class ISequencerTextFilterExpressionContext : public ITextFilterExpressionContext
+{ 
+public:
+	//~ Begin ISequencerTextFilterExpressionContext
+
+	SEQUENCER_API virtual TSet<FName> GetKeys() const = 0;
+
+	SEQUENCER_API virtual ESequencerTextFilterValueType GetValueType() const = 0;
+	SEQUENCER_API virtual TArray<FSequencerTextFilterKeyword> GetValueKeywords() const { return {}; }
+
+	SEQUENCER_API virtual FText GetDescription() const = 0;
+	SEQUENCER_API virtual FText GetCategory() const { return FText::GetEmpty(); }
+
+	//~ End ISequencerTextFilterExpressionContext
+};
diff --git a/Engine/Source/Editor/Sequencer/Public/Filters/ISequencerTrackFilters.h b/Engine/Source/Editor/Sequencer/Public/Filters/ISequencerTrackFilters.h
--- a/Engine/Source/Editor/Sequencer/Public/Filters/ISequencerTrackFilters.h
+++ b/Engine/Source/Editor/Sequencer/Public/Filters/ISequencerTrackFilters.h
@@ -2,63 +2,19 @@
 
 #pragma once
 
-#include "Framework/Commands/UICommandList.h"
+#include "ISequencerFilterBar.h"
 
-class FSequencerTextFilterExpressionContext;
 class FSequencerTrackFilter_CustomText;
-class FSequencer;
-class FString;
-class FText;
 struct FSequencerFilterData;
 
-class ISequencerTrackFilters
+class ISequencerTrackFilters : public ISequencerFilterBar
 {
 public:
-	virtual ~ISequencerTrackFilters() = default;
-
-	SEQUENCER_API virtual FName GetIdentifier() const = 0;
-
-	SEQUENCER_API virtual FSequencer& GetSequencer() const = 0;
-
-	SEQUENCER_API virtual TSharedPtr<FUICommandList> GetCommandList() const = 0;
-
-	SEQUENCER_API virtual FString GetTextFilterString() const = 0;
-	SEQUENCER_API virtual void SetTextFilterString(const FString& InText) = 0;
-
-	/** Returns true if the current filter bar text filter string contains the specified text expression.
-	 * The text expression must have key, operator, and value tokens. */
-	SEQUENCER_API virtual bool DoesTextFilterStringContainExpressionPair(const FSequencerTextFilterExpressionContext& InExpression) const = 0;
-
-	SEQUENCER_API virtual bool HasAnyFilterActive(const bool bCheckTextFilter = true
-		, const bool bInCheckHideIsolateFilter = true
-		, const bool bInCheckCommonFilters = true
-		, const bool bInCheckInternalFilters = true
-		, const bool bInCheckCustomTextFilters = true) const = 0;
-
-	SEQUENCER_API virtual bool HasAnyFilterEnabled() const = 0;
-
-	SEQUENCER_API virtual void RequestFilterUpdate() = 0;
+	virtual ~ISequencerTrackFilters() override = default;
 
 	SEQUENCER_API virtual bool AddCustomTextFilter(const TSharedRef<FSequencerTrackFilter_CustomText>& InFilter, const bool bInAddToConfig) = 0;
 	SEQUENCER_API virtual bool RemoveCustomTextFilter(const TSharedRef<FSequencerTrackFilter_CustomText>& InFilter, const bool bInAddToConfig) = 0;
 
-	SEQUENCER_API virtual void EnableAllFilters(const bool bInEnable, const TArray<FString> InExceptionFilterNames) = 0;
-
-	SEQUENCER_API virtual void ActivateCommonFilters(const bool bInActivate, const TArray<FString> InExceptionFilterNames) = 0;
-
-	SEQUENCER_API virtual void ActivateAllEnabledFilters(const bool bInActivate, const TArray<FString> InExceptionFilterNames) = 0;
-
-	SEQUENCER_API virtual bool IsFilterActiveByDisplayName(const FString InFilterName) const = 0;
-	SEQUENCER_API virtual bool IsFilterEnabledByDisplayName(const FString InFilterName) const = 0;
-	SEQUENCER_API virtual bool SetFilterActiveByDisplayName(const FString InFilterName, const bool bInActive, const bool bInRequestFilterUpdate = true) = 0;
-	SEQUENCER_API virtual bool SetFilterEnabledByDisplayName(const FString InFilterName, const bool bInEnabled, const bool bInRequestFilterUpdate = true) = 0;
-
-	SEQUENCER_API virtual TArray<FText> GetFilterDisplayNames() const = 0;
-	SEQUENCER_API virtual TArray<FText> GetCustomTextFilterNames() const = 0;
-
-	SEQUENCER_API virtual int32 GetTotalDisplayNodeCount() const = 0;
-	SEQUENCER_API virtual int32 GetFilteredDisplayNodeCount() const  = 0;
-
 	SEQUENCER_API virtual void HideSelectedTracks() = 0;
 	SEQUENCER_API virtual void IsolateSelectedTracks() = 0;
 
diff --git a/Engine/Source/Editor/Sequencer/Public/Filters/SequencerFilterBase.h b/Engine/Source/Editor/Sequencer/Public/Filters/SequencerFilterBase.h
new file mode 100644
--- /dev/null
+++ b/Engine/Source/Editor/Sequencer/Public/Filters/SequencerFilterBase.h
@@ -0,0 +1,115 @@
+﻿// Copyright Epic Games, Inc. All Rights Reserved.
+
+#pragma once
+
+#include "Filters/FilterBase.h"
+#include "Filters/ISequencerFilterBar.h"
+#include "Templates/SharedPointer.h"
+
+class UMovieSceneSequence;
+struct FSlateBrush;
+
+/** Sequencer base filter for all Sequencer filters (Sequencer, Navigation Tool, etc.) */
+template<typename InFilterType>
+class FSequencerFilterBase
+	: public FFilterBase<InFilterType>
+	, public TSharedFromThis<FSequencerFilterBase<InFilterType>>
+{
+public:
+	SEQUENCER_API static FText BuildTooltipTextForCommand(const FText& InBaseText, const TSharedPtr<FUICommandInfo>& InCommand)
+	{
+		const TSharedRef<const FInputChord> FirstValidChord = InCommand->GetFirstValidChord();
+		if (FirstValidChord->IsValidChord())
+		{
+			return FText::Format(NSLOCTEXT("Sequencer", "TrackFilterTooltipText", "{0} ({1})")
+				, InBaseText, FirstValidChord->GetInputText());
+		}
+		return InBaseText;
+	}
+
+	FSequencerFilterBase(ISequencerFilterBar& InOutFilterInterface, TSharedPtr<FFilterCategory>&& InCategory = nullptr)
+		: FFilterBase<InFilterType>(InCategory)
+		, FilterInterface(InOutFilterInterface)
+	{}
+
+	//~ Begin IFilter
+	SEQUENCER_API virtual bool PassesFilter(InFilterType InItem) const override { return true; };
+	//~ End IFilter
+
+	//~ Begin FFilterBase
+
+	SEQUENCER_API virtual FText GetDisplayName() const override { return FText::GetEmpty(); }
+
+	SEQUENCER_API virtual FText GetToolTipText() const override
+	{
+		if (const TSharedPtr<FUICommandInfo> ToggleCommand = GetToggleCommand())
+		{
+			return BuildTooltipTextForCommand(GetDefaultToolTipText(), ToggleCommand);
+		}
+		return GetDefaultToolTipText();
+	}
+
+	SEQUENCER_API virtual FLinearColor GetColor() const override { return FLinearColor::White; }
+
+	virtual void ModifyContextMenu(FMenuBuilder& MenuBuilder) override {}
+
+	virtual void SaveSettings(const FString& InIniFilename, const FString& InIniSection, const FString& InSettingsString) const override {}
+	virtual void LoadSettings(const FString& InIniFilename, const FString& InIniSection, const FString& InSettingsString) override {}
+
+	virtual bool IsInverseFilter() const override { return false; }
+
+	virtual void ActiveStateChanged(const bool bInActive) override {}
+
+	//~ End FFilterBase
+
+	SEQUENCER_API ISequencer& GetSequencer() const { return FilterInterface.GetSequencer(); }
+
+	SEQUENCER_API virtual ISequencerFilterBar& GetFilterInterface() const { return FilterInterface; }
+
+	SEQUENCER_API virtual FText GetDefaultToolTipText() const { return FText(); }
+	SEQUENCER_API virtual FSlateIcon GetIcon() const { return FSlateIcon(); }
+	SEQUENCER_API virtual bool IsCustomTextFilter() const { return false; }
+
+	virtual void BindCommands()
+	{
+		if (const TSharedPtr<FUICommandInfo> ToggleCommand = GetToggleCommand())
+		{
+			MapToggleAction(ToggleCommand);
+		}
+	}
+
+	SEQUENCER_API virtual TSharedPtr<FUICommandInfo> GetToggleCommand() const { return nullptr; }
+
+	SEQUENCER_API bool CanToggleFilter() const
+	{
+		const FString FilterName = GetDisplayName().ToString();
+		return FilterInterface.IsFilterActiveByDisplayName(FilterName);
+	}
+
+	SEQUENCER_API void ToggleFilter()
+	{
+		const FString FilterName = GetDisplayName().ToString();
+		const bool bNewState = !FilterInterface.IsFilterActiveByDisplayName(FilterName);
+		FilterInterface.SetFilterActiveByDisplayName(FilterName, bNewState);
+	}
+
+	SEQUENCER_API virtual bool SupportsSequence(UMovieSceneSequence* const InSequence) const { return true; }
+
+protected:
+	void MapToggleAction(const TSharedPtr<FUICommandInfo>& InCommand)
+	{
+		FilterInterface.GetCommandList()->MapAction(
+			InCommand,
+			FExecuteAction::CreateSP(this, &FSequencerFilterBase::ToggleFilter),
+			FCanExecuteAction(),
+			FIsActionChecked::CreateSP(this, &FSequencerFilterBase::CanToggleFilter));
+	}
+
+	ISequencerFilterBar& FilterInterface;
+
+private:
+	//~ Begin IFilter
+	/** Hide behind private to force use of GetIcon() instead */
+	virtual FName GetIconName() const override { return FName(); }
+	//~ End IFilter
+};
diff --git a/Engine/Source/Editor/Sequencer/Public/Filters/SequencerTextFilterExpressionContext.h b/Engine/Source/Editor/Sequencer/Public/Filters/SequencerTextFilterExpressionContext.h
--- a/Engine/Source/Editor/Sequencer/Public/Filters/SequencerTextFilterExpressionContext.h
+++ b/Engine/Source/Editor/Sequencer/Public/Filters/SequencerTextFilterExpressionContext.h
@@ -2,6 +2,7 @@
 
 #pragma once
 
+#include "Filters/ISequencerTextFilterExpressionContext.h"
 #include "Filters/SequencerTrackFilterBase.h"
 #include "Misc/TextFilterExpressionEvaluator.h"
 
@@ -9,39 +10,14 @@ class UMovieScene;
 class UMovieSceneSequence;
 class UMovieSceneTrack;
 
-enum class SEQUENCER_API ESequencerTextFilterValueType : uint8
-{
-	String,
-	Boolean,
-	Integer
-};
-
-struct SEQUENCER_API FSequencerTextFilterKeyword
-{
-	FString Keyword;
-	FText Description;
-};
-
 /** Text expression context to test the given asset data against the current text filter */
-class FSequencerTextFilterExpressionContext : public ITextFilterExpressionContext
+class FSequencerTextFilterExpressionContext : public ISequencerTextFilterExpressionContext
 { 
 public:
 	SEQUENCER_API FSequencerTextFilterExpressionContext(ISequencerTrackFilters& InFilterInterface);
 
 	void SetFilterItem(FSequencerTrackFilterType InFilterItem, UMovieSceneTrack* const InTrackObject);
 
-	//~ Begin FSequencerTextFilterExpressionContext
-
-	SEQUENCER_API virtual TSet<FName> GetKeys() const = 0;
-
-	SEQUENCER_API virtual ESequencerTextFilterValueType GetValueType() const = 0;
-	SEQUENCER_API virtual TArray<FSequencerTextFilterKeyword> GetValueKeywords() const { return {}; }
-
-	SEQUENCER_API virtual FText GetDescription() const = 0;
-	SEQUENCER_API virtual FText GetCategory() const { return FText::GetEmpty(); }
-
-	//~ End FSequencerTextFilterExpressionContext
-
 	//~ Begin ITextFilterExpressionContext
 
 	SEQUENCER_API virtual bool TestBasicStringExpression(const FTextFilterString& InValue
diff --git a/Engine/Source/Editor/Sequencer/Public/Filters/SequencerTrackFilterBase.h b/Engine/Source/Editor/Sequencer/Public/Filters/SequencerTrackFilterBase.h
--- a/Engine/Source/Editor/Sequencer/Public/Filters/SequencerTrackFilterBase.h
+++ b/Engine/Source/Editor/Sequencer/Public/Filters/SequencerTrackFilterBase.h
@@ -2,18 +2,18 @@
 
 #pragma once
 
-#include "GameFramework/Actor.h"
 #include "Filters/FilterBase.h"
 #include "Filters/ISequencerTrackFilters.h"
+#include "GameFramework/Actor.h"
 #include "Misc/IFilter.h"
 #include "MovieSceneTrack.h"
 #include "MovieSceneTrackEditor.h"
+#include "SequencerFilterBase.h"
 #include "SequencerFilterData.h"
 #include "MVVM/ViewModelPtr.h"
 #include "Textures/SlateIcon.h"
 
 class ISequencer;
-class FMenuBuilder;
 class UMovieScene;
 class UMovieSceneSequence;
 struct FSequencerFilterData;
@@ -34,8 +34,7 @@ using FSequencerTrackFilterType = UE::Sequencer::FViewModelPtr;
 //////////////////////////////////////////////////////////////////////////
 //
 
-class FSequencerTrackFilter : public FFilterBase<FSequencerTrackFilterType>
-	, public TSharedFromThis<FSequencerTrackFilter>
+class FSequencerTrackFilter : public FSequencerFilterBase<FSequencerTrackFilterType>
 {
 public:
 	SEQUENCER_API static bool SupportsLevelSequence(UMovieSceneSequence* const InSequence);
@@ -50,70 +49,25 @@ public:
 		return Support == ETrackSupport::Supported;
 	}
 
-	SEQUENCER_API static FText BuildTooltipTextForCommand(const FText& InBaseText, const TSharedPtr<FUICommandInfo>& InCommand);
-
 	SEQUENCER_API FSequencerTrackFilter(ISequencerTrackFilters& InOutFilterInterface, TSharedPtr<FFilterCategory>&& InCategory = nullptr);
 
-	//~ Begin IFilter
-	SEQUENCER_API virtual bool PassesFilter(FSequencerTrackFilterType InItem) const override { return true; };
-	//~ End IFilter
-
-	//~ Begin FFilterBase
-
-	virtual FText GetDisplayName() const override { return FText::GetEmpty(); }
-	SEQUENCER_API virtual FText GetToolTipText() const override;
-	virtual FLinearColor GetColor() const override { return FLinearColor::White; }
-	virtual bool IsInverseFilter() const override { return false; }
-
-	virtual void ActiveStateChanged(bool bActive) override {}
-
-	virtual void ModifyContextMenu(FMenuBuilder& MenuBuilder) override {}
-
-	virtual void SaveSettings(const FString& IniFilename, const FString& IniSection, const FString& SettingsString) const override {}
-	virtual void LoadSettings(const FString& IniFilename, const FString& IniSection, const FString& SettingsString) override {}
-
-	//~ End FFilterBase
-
-	//~ Begin FSequencerTrackFilter
+	//~ Begin FSequencerFilterBase
 
-	SEQUENCER_API virtual FText GetDefaultToolTipText() const { return FText(); }
-	SEQUENCER_API virtual FSlateIcon GetIcon() const { return FSlateIcon(); }
+	SEQUENCER_API virtual ISequencerTrackFilters& GetFilterInterface() const override;
 
-	SEQUENCER_API virtual bool IsCustomTextFilter() const { return false; }
+	//~ End FSequencerFilterBase
 
 	/** Returns whether this filter needs reevaluating any time track values have been modified, not just tree changes */
 	SEQUENCER_API virtual bool ShouldUpdateOnTrackValueChanged() const { return false; }
 
 	/** Returns whether the filter supports the sequence type */
-	SEQUENCER_API virtual bool SupportsSequence(UMovieSceneSequence* const InSequence) const;
-
-	SEQUENCER_API virtual void BindCommands();
-
-	SEQUENCER_API virtual TSharedPtr<FUICommandInfo> GetToggleCommand() const { return nullptr; }
-
-	virtual TSubclassOf<UMovieSceneTrack> GetTrackClass() const { return nullptr; };
-
-	//~ End FSequencerTrackFilter
+	SEQUENCER_API virtual bool SupportsSequence(UMovieSceneSequence* const InSequence) const override;
 
-	SEQUENCER_API ISequencerTrackFilters& GetFilterInterface() const;
-	SEQUENCER_API ISequencer& GetSequencer() const;
-
-	bool CanToggleFilter() const;
-	void ToggleFilter();
-
-	SEQUENCER_API void MapToggleAction(const TSharedPtr<FUICommandInfo>& InCommand);
+	SEQUENCER_API virtual TSubclassOf<UMovieSceneTrack> GetTrackClass() const { return nullptr; }
 
 protected:
 	UMovieSceneSequence* GetFocusedMovieSceneSequence() const;
 	UMovieScene* GetFocusedGetMovieScene() const;
-
-	ISequencerTrackFilters& FilterInterface;
-
-private:
-	//~ Begin IFilter
-	/** Hide behind private to force use of GetIcon() instead */
-	virtual FName GetIconName() const override { return FName(); }
-	//~ End IFilter
 };
 
 //////////////////////////////////////////////////////////////////////////
@@ -152,14 +106,14 @@ public:
 	//~ Begin IFilter
 	virtual bool PassesFilter(FSequencerTrackFilterType InItem) const override
 	{
-		FSequencerFilterData& FilterData = FilterInterface.GetFilterData();
+		FSequencerFilterData& FilterData = GetFilterInterface().GetFilterData();
 		const UMovieSceneTrack* const TrackObject = FilterData.ResolveMovieSceneTrackObject(InItem);
 		return IsValid(TrackObject) && TrackObject->IsA(InClassType::StaticClass());
 	}
 	//~ End IFilter
 
 	//~ Begin FSequencerTrackFilter
-	TSubclassOf<UMovieSceneTrack> GetTrackClass() const override
+	virtual TSubclassOf<UMovieSceneTrack> GetTrackClass() const override
 	{
 		return InClassType::StaticClass();
 	}
@@ -181,7 +135,7 @@ public:
 	//~ Begin IFilter
 	virtual bool PassesFilter(FSequencerTrackFilterType InItem) const override
 	{
-		FSequencerFilterData& FilterData = FilterInterface.GetFilterData();
+		FSequencerFilterData& FilterData = GetFilterInterface().GetFilterData();
 
 		UObject* const BoundObject = FilterData.ResolveTrackBoundObject(GetSequencer(), InItem);
 		if (!IsValid(BoundObject))
