Sequencer: Use localized GetCategory for property category names

diff --git a/Engine/Source/Editor/Sequencer/Private/MVVM/ViewModels/ObjectBindingModel.cpp b/Engine/Source/Editor/Sequencer/Private/MVVM/ViewModels/ObjectBindingModel.cpp
--- a/Engine/Source/Editor/Sequencer/Private/MVVM/ViewModels/ObjectBindingModel.cpp
+++ b/Engine/Source/Editor/Sequencer/Private/MVVM/ViewModels/ObjectBindingModel.cpp
@@ -845,7 +845,7 @@ void FObjectBindingModel::HandleAddTrackSubMenuNew(FMenuBuilder& AddTrackMenuBui
 			}
 			else
 			{
-				KeyableMenuData.MenuName = FObjectEditorUtils::GetCategoryFName(Property).ToString();
+				KeyableMenuData.MenuName = FObjectEditorUtils::GetCategory(Property);
 			}
 
 			PropertiesTraversed.Add(Property);
