Sequencer: ULevelSequenceEditorSubsystem::CopyFolders now recursively gathers tracks and objects to copy as well as folders.

This was already recursively copying tracks and objects when done through the UI. This now allows the same to function when calling it through python/editor utility widget.

diff --git a/Engine/Plugins/MovieScene/LevelSequenceEditor/Source/LevelSequenceEditor/Private/LevelSequenceEditorSubsystem.cpp b/Engine/Plugins/MovieScene/LevelSequenceEditor/Source/LevelSequenceEditor/Private/LevelSequenceEditorSubsystem.cpp
--- a/Engine/Plugins/MovieScene/LevelSequenceEditor/Source/LevelSequenceEditor/Private/LevelSequenceEditorSubsystem.cpp
+++ b/Engine/Plugins/MovieScene/LevelSequenceEditor/Source/LevelSequenceEditor/Private/LevelSequenceEditorSubsystem.cpp
@@ -1098,7 +1098,24 @@ void ULevelSequenceEditorSubsystem::SaveDefaultSpawnableState(const FMovieSceneB
 
 void ULevelSequenceEditorSubsystem::CopyFolders(const TArray<UMovieSceneFolder*>& Folders, FString& ExportedText)
 {
-	FSequencerUtilities::CopyFolders(Folders, ExportedText);
+	FString DummyText;
+	CopyFolders(Folders, ExportedText, DummyText, DummyText);
+}
+
+void ULevelSequenceEditorSubsystem::CopyFolders(const TArray<UMovieSceneFolder*>&Folders, FString &FoldersExportedText, FString &ObjectsExportedText, FString &TracksExportedText)
+{
+	TSharedPtr<ISequencer> Sequencer = GetActiveSequencer();
+	if (Sequencer == nullptr)
+	{
+		return;
+	}
+
+	FSequencerUtilities::CopyFolders(Sequencer.ToSharedRef(), Folders, FoldersExportedText, ObjectsExportedText, TracksExportedText);
+
+	FString ExportedText;
+	ExportedText += ObjectsExportedText;
+	ExportedText += TracksExportedText;
+	ExportedText += FoldersExportedText;
 
 	FPlatformApplicationMisc::ClipboardCopy(*ExportedText);
 }
diff --git a/Engine/Plugins/MovieScene/LevelSequenceEditor/Source/LevelSequenceEditor/Public/LevelSequenceEditorSubsystem.h b/Engine/Plugins/MovieScene/LevelSequenceEditor/Source/LevelSequenceEditor/Public/LevelSequenceEditorSubsystem.h
--- a/Engine/Plugins/MovieScene/LevelSequenceEditor/Source/LevelSequenceEditor/Public/LevelSequenceEditorSubsystem.h
+++ b/Engine/Plugins/MovieScene/LevelSequenceEditor/Source/LevelSequenceEditor/Public/LevelSequenceEditorSubsystem.h
@@ -158,7 +158,10 @@ public:
 	 * folders without relying on a single clipboard. 
 	 */
 	UFUNCTION(BlueprintCallable, Category = "Level Sequence Editor")
-	void CopyFolders(const TArray<UMovieSceneFolder*>& Folders, FString& ExportedText);
+	void CopyFolders(const TArray<UMovieSceneFolder*>& Folders, FString& FoldersExportedText, FString& ObjectsExportedText, FString& TracksExportedText);
+
+	UE_DEPRECATED(5.5, "CopyFolders now gathers objects and tracks within the folders. Please use CopyFolders that outputs ObjectsExportedText and TracksExportedText")
+	void CopyFolders(const TArray<UMovieSceneFolder*>& Folders, FString& FoldersExportedText);
 
 	/** 
 	 * Paste folders 
diff --git a/Engine/Source/Editor/Sequencer/Private/Sequencer.cpp b/Engine/Source/Editor/Sequencer/Private/Sequencer.cpp
--- a/Engine/Source/Editor/Sequencer/Private/Sequencer.cpp
+++ b/Engine/Source/Editor/Sequencer/Private/Sequencer.cpp
@@ -8424,9 +8424,9 @@ void FSequencer::CopySelectedTracks(TArray<TSharedPtr<FViewModel>>& TrackNodes,
 	}
 }
 
-void FSequencer::CopySelectedFolders(const TArray<UMovieSceneFolder*>& Folders, FString& ExportedText)
+void FSequencer::CopySelectedFolders(const TArray<UMovieSceneFolder*>& Folders, FString& FoldersExportedText, FString& TracksExportedText, FString& ObjectExportedText)
 {
-	FSequencerUtilities::CopyFolders(Folders, ExportedText);
+	FSequencerUtilities::CopyFolders(AsShared(), Folders, FoldersExportedText, TracksExportedText, ObjectExportedText);
 
 	// Make sure to clear the clipboard for the keys
 	GClipboardStack.Empty();
@@ -10521,11 +10521,6 @@ void GatherTracksAndObjectsToCopy(TSharedRef<UE::Sequencer::FViewModel> Node, TA
 	else if (FFolderModel* FolderNode = Node->CastThis<FFolderModel>())
 	{
 		FoldersToCopy.Add(FolderNode->GetFolder());
-
-		for (TSharedPtr<FViewModel> ChildNode : FolderNode->GetChildren())
-		{
-			GatherTracksAndObjectsToCopy(ChildNode.ToSharedRef(), TracksToCopy, ObjectsToCopy, FoldersToCopy);
-		}
 	}
 }
 
@@ -10575,7 +10570,7 @@ void FSequencer::CopySelection()
 
 		if (FoldersToCopy.Num())
 		{
-			CopySelectedFolders(FoldersToCopy, FoldersExportedText);
+			CopySelectedFolders(FoldersToCopy, FoldersExportedText, ObjectsExportedText, TracksExportedText);
 		}
 
 		FString ExportedText;
diff --git a/Engine/Source/Editor/Sequencer/Private/Sequencer.h b/Engine/Source/Editor/Sequencer/Private/Sequencer.h
--- a/Engine/Source/Editor/Sequencer/Private/Sequencer.h
+++ b/Engine/Source/Editor/Sequencer/Private/Sequencer.h
@@ -568,7 +568,7 @@ public:
 	/** Called when a user executes the copy track menu item */
 	void CopySelectedObjects(TArray<TSharedPtr<UE::Sequencer::FObjectBindingModel>>& ObjectNodes, const TArray<UMovieSceneFolder*>& Folders, /*out*/ FString& ExportedText);
 	void CopySelectedTracks(TArray<TSharedPtr<UE::Sequencer::FViewModel>>& TrackNodes, const TArray<UMovieSceneFolder*>& Folders, /*out*/ FString& ExportedText);
-	void CopySelectedFolders(const TArray<UMovieSceneFolder*>& Folders, /*out*/ FString& ExportedText);
+	void CopySelectedFolders(const TArray<UMovieSceneFolder*>& Folders, /*out*/ FString& FoldersExportedText, FString& TracksExportedText, FString& ObjectsExportedText);
 
 	/** Called when a user executes the paste track menu item */
 	bool CanPaste(const FString& TextToImport);
diff --git a/Engine/Source/Editor/Sequencer/Private/SequencerUtilities.cpp b/Engine/Source/Editor/Sequencer/Private/SequencerUtilities.cpp
--- a/Engine/Source/Editor/Sequencer/Private/SequencerUtilities.cpp
+++ b/Engine/Source/Editor/Sequencer/Private/SequencerUtilities.cpp
@@ -1707,6 +1707,61 @@ void FSequencerUtilities::CopyFolders(const TArray<UMovieSceneFolder*>& Folders,
 	ExportObjectsToText(Objects, /*out*/ ExportedText);
 }
 
+void GatherFolderContents(UMovieSceneFolder* Folder, TArray<UMovieSceneFolder*>& Folders, TArray<UMovieSceneTrack*>& Tracks, TArray<FMovieSceneBindingProxy>& Bindings)
+{
+	if (!Folder)
+	{
+		return;
+	}
+
+	Folders.AddUnique(Folder);
+
+	UMovieScene* MovieScene = CastChecked<UMovieScene>(Folder->GetOuter());
+
+	UMovieSceneSequence* Sequence = CastChecked<UMovieSceneSequence>(MovieScene->GetOuter());
+
+	for (const FGuid& ObjectBinding : Folder->GetChildObjectBindings())
+	{
+		Bindings.AddUnique(FMovieSceneBindingProxy(ObjectBinding, Sequence));
+	}
+
+	for (UMovieSceneTrack* ChildTrack : Folder->GetChildTracks())
+	{
+		Tracks.AddUnique(ChildTrack);
+	}
+
+	for (UMovieSceneFolder* ChildFolder : Folder->GetChildFolders())
+	{
+		if (ChildFolder)
+		{
+			GatherFolderContents(ChildFolder, Folders, Tracks, Bindings);
+		}
+	}
+}
+
+void FSequencerUtilities::CopyFolders(TSharedRef<ISequencer> Sequencer, const TArray<UMovieSceneFolder*>& InFolders, FString& FoldersExportedText, FString& TracksExportedText, FString& ObjectsExportedText)
+{
+	TArray<UMovieSceneFolder*> Folders;
+	TArray<UMovieSceneTrack*> Tracks;
+	TArray<FMovieSceneBindingProxy> Bindings;
+
+	for (UMovieSceneFolder* Folder : InFolders)
+	{
+		GatherFolderContents(Folder, Folders, Tracks, Bindings);
+	}
+
+	CopyTracks(Tracks, Folders, TracksExportedText);
+	CopyBindings(Sequencer, Bindings, Folders, ObjectsExportedText);
+
+	TArray<UObject*> Objects;
+	for (UMovieSceneFolder* Folder : Folders)
+	{
+		Objects.Add(Folder);
+	}
+
+	ExportObjectsToText(Objects, /*out*/ FoldersExportedText);
+}
+
 class FFolderObjectTextFactory : public FCustomizableTextObjectFactory
 {
 public:
diff --git a/Engine/Source/Editor/Sequencer/Public/SequencerUtilities.h b/Engine/Source/Editor/Sequencer/Public/SequencerUtilities.h
--- a/Engine/Source/Editor/Sequencer/Public/SequencerUtilities.h
+++ b/Engine/Source/Editor/Sequencer/Public/SequencerUtilities.h
@@ -281,10 +281,13 @@ struct SEQUENCER_API FSequencerUtilities
 	static FMovieScenePossessable* ConvertToCustomBinding(TSharedRef<ISequencer> Sequencer, FGuid BindingGuid, TSubclassOf<UMovieSceneCustomBinding> CustomBindingType, int32 BindingIndex = 0);
 
 	/** Copy/paste folders */
-	static void CopyFolders(const TArray<UMovieSceneFolder*>& Folders, FString& ExportedText);
+	static void CopyFolders(TSharedRef<ISequencer> Sequencer, const TArray<UMovieSceneFolder*>& Folders, FString& FoldersExportedText, FString& TracksExportedText, FString& ObjectsExportedText);
 	static bool PasteFolders(const FString& TextToImport, FMovieScenePasteFoldersParams PasteFoldersParams, TArray<UMovieSceneFolder*>& OutFolders, TArray<FNotificationInfo>& OutErrors);
 	static bool CanPasteFolders(const FString& TextToImport);
 
+	UE_DEPRECATED(5.5, "CopyFolders now gathers objects and tracks within the folders. Please use CopyFolders that outputs ObjectsExportedText and TracksExportedText")
+	static void CopyFolders(const TArray<UMovieSceneFolder*>& Folders, FString& ExportedText);
+
 	/** Copy/paste tracks */
 	static void CopyTracks(const TArray<UMovieSceneTrack*>& Tracks, const TArray<UMovieSceneFolder*>& InFolders, FString& ExportedText);
 	static bool PasteTracks(const FString& TextToImport, FMovieScenePasteTracksParams PasteTracksParams, TArray<UMovieSceneTrack*>& OutTracks, TArray<FNotificationInfo>& OutErrors);
