Sequencer: Follow up to 38925830 to look for relative path names using the current world

Copy bindings/paste bindings stores and restores bound object names relative to a context, previously the resolution context. The resolution context could have been the world or a level sequence actor in the world. The copy/paste could fail because it depended upon the same set up for when you do the copy and paste. For example, if would fail if you did the copy when there is a level sequence actor and then the paste when there is no level sequence actor. This fix eliminates that by always resolving the world.

diff --git a/Engine/Source/Editor/Sequencer/Private/Sequencer.cpp b/Engine/Source/Editor/Sequencer/Private/Sequencer.cpp
--- a/Engine/Source/Editor/Sequencer/Private/Sequencer.cpp
+++ b/Engine/Source/Editor/Sequencer/Private/Sequencer.cpp
@@ -8514,6 +8514,8 @@ bool FSequencer::PasteObjectBindings(const FString& TextToImport, UMovieSceneFol
 		}
 	}
 
+	UWorld* World = GCurrentLevelEditingViewportClient ? GCurrentLevelEditingViewportClient->GetWorld() : nullptr;
+
 	bool bDuplicateExistingActors = false;
 	TArray<FString> ObjectNames = FSequencerUtilities::GetPasteBindingsObjectNames(AsShared(), TextToImport);
 
@@ -8525,12 +8527,6 @@ bool FSequencer::PasteObjectBindings(const FString& TextToImport, UMovieSceneFol
 	{
 		FMovieScenePossessable& Possessable = MovieScene->GetPossessable(Index);
 
-		UObject* ResolutionContext = FSequencerUtilities::FindResolutionContext(AsShared()
-			, *MovieScene->GetTypedOuter<UMovieSceneSequence>()
-			, *MovieScene
-			, Possessable.GetParent()
-			, GetPlaybackContext());
-
 		FGuid ThisGuid = Possessable.GetGuid();
 
 		for (TWeakObjectPtr<> WeakObject : FindBoundObjects(ThisGuid, GetFocusedTemplateID()))
@@ -8539,7 +8535,7 @@ bool FSequencer::PasteObjectBindings(const FString& TextToImport, UMovieSceneFol
 			{
 				AActor* Actor = Cast<AActor>(WeakObject.Get());
 
-				if (Actor && ObjectNames.Contains(Actor->GetPathName(ResolutionContext)))
+				if (Actor && ObjectNames.Contains(Actor->GetPathName(World)))
 				{
 					FText DuplicateActorsMsg = FText::Format(LOCTEXT("DuplicateActorsForPastedBinding", "Attempting to paste a binding that is already bound to {0}.\nShould the existing actor be duplicated for the pasted binding?"), FText::FromString(Actor->GetActorLabel()));
 
