Sequencer: Always reset the hotspot when changing sequences

This fixes an issue where the cursor could get stuck on the movement tool when switching sequences

diff --git a/Engine/Source/Editor/Sequencer/Private/Sequencer.cpp b/Engine/Source/Editor/Sequencer/Private/Sequencer.cpp
--- a/Engine/Source/Editor/Sequencer/Private/Sequencer.cpp
+++ b/Engine/Source/Editor/Sequencer/Private/Sequencer.cpp
@@ -1413,6 +1413,9 @@ void FSequencer::FocusSequenceInstance(UMovieSceneSubSection& InSubSection)
 
 	ScrubPositionParent.Reset();
 
+	// Always reset the hotspot so that it can be regenerated when rebuilding
+	ViewModel->GetTrackArea()->SetHotspot(nullptr);
+
 	// Reset data that is only used for the previous movie scene
 	ResetPerMovieSceneData();
 	SequencerWidget->UpdateBreadcrumbs();
