Sequencer: Copy/paste keys in root time space

This allows you to copy/paste keys from one sequence to another in which their time spaces don't match. For example, if you have keys at 0, 10, 20 in one sequence and paste them into a sequence that is time scaled by 2, the keys will be placed at 0, 20, 40.

Tested the case above and vice versa, copying from within the time scaled by 2 sequence into the parent sequence.
Tested the normal case of non time scaled sequences and also pasting to the same sequence.

diff --git a/Engine/Source/Editor/Sequencer/Private/SequencerContextMenus.cpp b/Engine/Source/Editor/Sequencer/Private/SequencerContextMenus.cpp
--- a/Engine/Source/Editor/Sequencer/Private/SequencerContextMenus.cpp
+++ b/Engine/Source/Editor/Sequencer/Private/SequencerContextMenus.cpp
@@ -1889,6 +1889,7 @@ bool FPasteContextMenu::PasteInto(int32 DestinationIndex, FName KeyAreaName, TSe
 	FSequencerPasteEnvironment PasteEnvironment;
 	PasteEnvironment.TickResolution = Sequencer->GetFocusedTickResolution();
 	PasteEnvironment.CardinalTime = Args.PasteAtTime;
+	PasteEnvironment.TimeTransform = Sequencer->GetFocusedMovieSceneSequenceTransform().LinearTransform;
 	PasteEnvironment.OnKeyPasted = [&](FKeyHandle Handle, TSharedPtr<UE::Sequencer::FChannelModel> Channel){
 		NewSelection.Add(FSequencerSelectedKey(*Channel->GetSection(), Channel, Handle));
 	};
diff --git a/Engine/Source/Editor/Sequencer/Public/SequencerChannelTraits.h b/Engine/Source/Editor/Sequencer/Public/SequencerChannelTraits.h
--- a/Engine/Source/Editor/Sequencer/Public/SequencerChannelTraits.h
+++ b/Engine/Source/Editor/Sequencer/Public/SequencerChannelTraits.h
@@ -10,6 +10,7 @@
 #include "CoreTypes.h"
 #include "Curves/KeyHandle.h"
 #include "Evaluation/MovieSceneRootOverridePath.h"
+#include "Evaluation/MovieSceneTimeTransform.h"
 #include "ISequencer.h"
 #include "ISequencerChannelInterface.h"
 #include "Math/Range.h"
@@ -428,7 +429,13 @@ namespace Sequencer
 
 		auto ForEachKey = [Section, PasteAt, &NewRange, &ChannelData, &OutPastedKeys, &SrcEnvironment, &DstEnvironment](const FMovieSceneClipboardKey& Key)
 		{
-			FFrameNumber Time = (PasteAt + FFrameRate::TransformTime(Key.GetTime(), SrcEnvironment.TickResolution, DstEnvironment.TickResolution)).FloorToFrame();
+			FFrameTime SrcTime = Key.GetTime();
+			SrcTime = SrcTime * SrcEnvironment.TimeTransform.Inverse();
+
+			FFrameTime DstTime = FFrameRate::TransformTime(SrcTime, SrcEnvironment.TickResolution, DstEnvironment.TickResolution);
+			DstTime = DstTime * DstEnvironment.TimeTransform;
+
+			FFrameNumber Time = (PasteAt + DstTime).FloorToFrame();
 
 			NewRange = TRange<FFrameNumber>::Hull(NewRange, TRange<FFrameNumber>(Time));
 
diff --git a/Engine/Source/Runtime/MovieScene/Public/MovieSceneClipboard.h b/Engine/Source/Runtime/MovieScene/Public/MovieSceneClipboard.h
--- a/Engine/Source/Runtime/MovieScene/Public/MovieSceneClipboard.h
+++ b/Engine/Source/Runtime/MovieScene/Public/MovieSceneClipboard.h
@@ -5,6 +5,7 @@
 #include "Containers/Array.h"
 #include "Containers/Map.h"
 #include "CoreMinimal.h"
+#include "Evaluation/MovieSceneTimeTransform.h"
 #include "HAL/Platform.h"
 #include "Internationalization/Text.h"
 #include "Misc/AssertionMacros.h"
@@ -287,6 +288,9 @@ struct FMovieSceneClipboardEnvironment
 	/** The cardinal time for a copy-paste operation. */
 	FFrameTime CardinalTime;
 
+	/** The time transform for the copy-paste operation. */
+	FMovieSceneTimeTransform TimeTransform;
+
 	/** The date/time at which the copy operation was performed */
 	FDateTime DateTime;
 
