Sequencer: Copy/paste keys in root time space

This allows you to copy/paste keys from one sequence to another in which their time spaces don't match. For example, if you have keys at 0, 10, 20 in one sequence and paste them into a sequence that is time scaled by 2, the keys will be placed at 0, 20, 40.

Tested the case above and vice versa, copying from within the time scaled by 2 sequence into the parent sequence.
Tested the normal case of non time scaled sequences and also pasting to the same sequence.

diff --git a/Engine/Source/Editor/Sequencer/Private/Sequencer.cpp b/Engine/Source/Editor/Sequencer/Private/Sequencer.cpp
--- a/Engine/Source/Editor/Sequencer/Private/Sequencer.cpp
+++ b/Engine/Source/Editor/Sequencer/Private/Sequencer.cpp
@@ -10689,6 +10689,7 @@ void FSequencer::CopySelectedKeys()
 	TSharedRef<FMovieSceneClipboard> Clipboard = MakeShareable( new FMovieSceneClipboard(Builder.Commit(TOptional<FFrameNumber>())) );
 	
 	Clipboard->GetEnvironment().TickResolution = GetFocusedTickResolution();
+	Clipboard->GetEnvironment().TimeTransform = GetFocusedMovieSceneSequenceTransform().LinearTransform;
 
 	if (Clipboard->GetKeyTrackGroups().Num())
 	{
