Use "GetNumericTypeInterface" method to get the appropriate numeric type interface when displaying time information for a section in a sequence (subsequences, audio track times, etc.)

#rnx

diff --git a/Engine/Source/Editor/Sequencer/Private/SSequencerSection.cpp b/Engine/Source/Editor/Sequencer/Private/SSequencerSection.cpp
--- a/Engine/Source/Editor/Sequencer/Private/SSequencerSection.cpp
+++ b/Engine/Source/Editor/Sequencer/Private/SSequencerSection.cpp
@@ -32,7 +32,7 @@
 #include "MovieScene.h"
 #include "Fonts/FontCache.h"
 #include "Fonts/FontMeasure.h"
-#include "FrameNumberNumericInterface.h"
+#include "Widgets/Input/NumericTypeInterface.h"
 #include "Framework/Application/SlateApplication.h"
 #include "MovieSceneTimeHelpers.h"
 #include "Tracks/MovieScenePropertyTrack.h"
@@ -1331,7 +1331,7 @@ TSharedPtr<FTrackAreaViewModel> SSequencerSection::GetTrackAreaViewModel() const
 	return nullptr;
 }
 
-void DrawFrameTimeHint(FSequencerSectionPainter& InPainter, const FFrameTime& CurrentTime, const FFrameTime& FrameTime, const FFrameNumberInterface* FrameNumberInterface, int32 LayerId)
+void DrawFrameTimeHint(FSequencerSectionPainter& InPainter, const FFrameTime& CurrentTime, const FFrameTime& FrameTime, const TSharedPtr<INumericTypeInterface<double>> FrameNumberInterface, int32 LayerId)
 {
 	FString FrameTimeString;
 	if (FrameNumberInterface)
@@ -1554,25 +1554,8 @@ int32 SSequencerSection::OnPaint( const FPaintArgs& Args, const FGeometry& Allot
 	TOptional<FFrameTime> SectionTime = SectionInterface->GetSectionTime(Painter);
 	if (SectionTime.IsSet())
 	{
-		const FFrameRate DisplayRate = GetSequencer().GetFocusedDisplayRate();
-		const FFrameRate TickResolution = GetSequencer().GetFocusedTickResolution();
-
-		// Get the desired frame display format and zero padding from
-		// the sequencer settings, if possible.
-		TAttribute<EFrameNumberDisplayFormats> DisplayFormatAttr(EFrameNumberDisplayFormats::Frames);
-		TAttribute<uint8> ZeroPadFrameNumbersAttr(0u);
-		if (const USequencerSettings* SequencerSettings = GetSequencer().GetSequencerSettings())
-		{
-			DisplayFormatAttr.Set(SequencerSettings->GetTimeDisplayFormat());
-			ZeroPadFrameNumbersAttr.Set(SequencerSettings->GetZeroPadFrames());
-		}
-
-		const TAttribute<FFrameRate> TickResolutionAttr(TickResolution);
-		const TAttribute<FFrameRate> DisplayRateAttr(DisplayRate);
-
-		const FFrameNumberInterface FrameNumberInterface(DisplayFormatAttr, ZeroPadFrameNumbersAttr, TickResolutionAttr, DisplayRateAttr);
-
-		DrawFrameTimeHint(Painter, GetSequencer().GetLocalTime().Time, SectionTime.GetValue(), &FrameNumberInterface, LayerId);
+		TSharedPtr<INumericTypeInterface<double>> NumericTypeInterface = GetSequencer().GetNumericTypeInterface(ENumericIntent::Position);
+		DrawFrameTimeHint(Painter, GetSequencer().GetLocalTime().Time, SectionTime.GetValue(), NumericTypeInterface, LayerId);
 	}
 
 	++LayerId;
