[Sequencer] Update Tracks Names and Reorganize Tracks Order

diff --git a/Engine/Plugins/Animation/ControlRig/Source/ControlRigEditor/Private/Sequencer/ControlRigParameterTrackEditor.cpp b/Engine/Plugins/Animation/ControlRig/Source/ControlRigEditor/Private/Sequencer/ControlRigParameterTrackEditor.cpp
--- a/Engine/Plugins/Animation/ControlRig/Source/ControlRigEditor/Private/Sequencer/ControlRigParameterTrackEditor.cpp
+++ b/Engine/Plugins/Animation/ControlRig/Source/ControlRigEditor/Private/Sequencer/ControlRigParameterTrackEditor.cpp
@@ -331,6 +331,11 @@ FControlRigParameterTrackEditor::FControlRigParameterTrackEditor(TSharedRef<ISeq
 	}
 }
 
+FText FControlRigParameterTrackEditor::GetDisplayName() const
+{
+	return LOCTEXT("ControlRigParameterTrackEditor_DisplayName", "Control Rig Parameter");
+}
+
 void FControlRigParameterTrackEditor::OnInitialize()
 {
 	if (GetSequencer().IsValid() == false)
diff --git a/Engine/Plugins/Animation/ControlRig/Source/ControlRigEditor/Private/Sequencer/ControlRigParameterTrackEditor.h b/Engine/Plugins/Animation/ControlRig/Source/ControlRigEditor/Private/Sequencer/ControlRigParameterTrackEditor.h
--- a/Engine/Plugins/Animation/ControlRig/Source/ControlRigEditor/Private/Sequencer/ControlRigParameterTrackEditor.h
+++ b/Engine/Plugins/Animation/ControlRig/Source/ControlRigEditor/Private/Sequencer/ControlRigParameterTrackEditor.h
@@ -109,6 +109,7 @@ public:
 public:
 
 	// ISequencerTrackEditor interface
+	virtual FText GetDisplayName() const override;
 	virtual void OnInitialize() override;
 	virtual void OnRelease() override;
 	virtual void BuildObjectBindingContextMenu(FMenuBuilder& MenuBuilder, const TArray<FGuid>& ObjectBindings, const UClass* ObjectClass) override;
diff --git a/Engine/Plugins/Animation/LiveLink/Source/LiveLinkSequencer/Private/Sequencer/LiveLinkPropertyTrackEditor.cpp b/Engine/Plugins/Animation/LiveLink/Source/LiveLinkSequencer/Private/Sequencer/LiveLinkPropertyTrackEditor.cpp
--- a/Engine/Plugins/Animation/LiveLink/Source/LiveLinkSequencer/Private/Sequencer/LiveLinkPropertyTrackEditor.cpp
+++ b/Engine/Plugins/Animation/LiveLink/Source/LiveLinkSequencer/Private/Sequencer/LiveLinkPropertyTrackEditor.cpp
@@ -185,6 +185,11 @@ TSharedRef<ISequencerTrackEditor> FLiveLinkPropertyTrackEditor::CreateTrackEdito
 	return MakeShareable(new FLiveLinkPropertyTrackEditor(InSequencer));
 }
 
+FText FLiveLinkPropertyTrackEditor::GetDisplayName() const
+{
+	return LOCTEXT("LiveLinkPropertyTrackEditor_DisplayName", "Live Link");
+}
+
 void FLiveLinkPropertyTrackEditor::BuildTrackContextMenu(FMenuBuilder& MenuBuilder, UMovieSceneTrack* Track)
 {
 
diff --git a/Engine/Plugins/Animation/LiveLink/Source/LiveLinkSequencer/Private/Sequencer/LiveLinkPropertyTrackEditor.h b/Engine/Plugins/Animation/LiveLink/Source/LiveLinkSequencer/Private/Sequencer/LiveLinkPropertyTrackEditor.h
--- a/Engine/Plugins/Animation/LiveLink/Source/LiveLinkSequencer/Private/Sequencer/LiveLinkPropertyTrackEditor.h
+++ b/Engine/Plugins/Animation/LiveLink/Source/LiveLinkSequencer/Private/Sequencer/LiveLinkPropertyTrackEditor.h
@@ -29,6 +29,8 @@ public:
 public:
 
 	// ISequencerTrackEditor interface
+	virtual FText GetDisplayName() const override;
+
 	virtual TSharedRef<ISequencerSection> MakeSectionInterface(UMovieSceneSection& SectionObject, UMovieSceneTrack& Track, FGuid ObjectBinding) override;
 
 	virtual void BuildAddTrackMenu(FMenuBuilder& MenuBuilder) override;
diff --git a/Engine/Plugins/Compositing/Composure/Source/ComposureEditor/Private/Sequencer/ComposureExportTrackEditor.cpp b/Engine/Plugins/Compositing/Composure/Source/ComposureEditor/Private/Sequencer/ComposureExportTrackEditor.cpp
--- a/Engine/Plugins/Compositing/Composure/Source/ComposureEditor/Private/Sequencer/ComposureExportTrackEditor.cpp
+++ b/Engine/Plugins/Compositing/Composure/Source/ComposureEditor/Private/Sequencer/ComposureExportTrackEditor.cpp
@@ -72,6 +72,11 @@ void FComposureExportTrackEditor::BuildTrackContextMenu(FMenuBuilder& MenuBuilde
 	MenuBuilder.AddSubMenu(LOCTEXT("Properties_MenuText", "Properties"), FText(), FNewMenuDelegate::CreateLambda(PopulateSubMenu));
 }
 
+FText FComposureExportTrackEditor::GetDisplayName() const
+{
+	return LOCTEXT("ComposureExportTrackEditor_DisplayName", "Composure Export");
+}
+
 void FComposureExportTrackEditor::BuildObjectBindingTrackMenu(FMenuBuilder& MenuBuilder, const TArray<FGuid>& ObjectBindings, const UClass* ObjectClass)
 {
 	TSharedPtr<ISequencer> SequencerPtr = GetSequencer();
diff --git a/Engine/Plugins/Compositing/Composure/Source/ComposureEditor/Private/Sequencer/ComposureExportTrackEditor.h b/Engine/Plugins/Compositing/Composure/Source/ComposureEditor/Private/Sequencer/ComposureExportTrackEditor.h
--- a/Engine/Plugins/Compositing/Composure/Source/ComposureEditor/Private/Sequencer/ComposureExportTrackEditor.h
+++ b/Engine/Plugins/Compositing/Composure/Source/ComposureEditor/Private/Sequencer/ComposureExportTrackEditor.h
@@ -14,6 +14,7 @@ public:
 
 private:
 
+	virtual FText GetDisplayName() const override;
 	virtual void BuildObjectBindingTrackMenu(FMenuBuilder& MenuBuilder, const TArray<FGuid>& ObjectBindings, const UClass* ObjectClass) override;
 	virtual void BuildTrackContextMenu(FMenuBuilder& MenuBuilder, UMovieSceneTrack* Track) override;
 	virtual bool SupportsType(TSubclassOf<UMovieSceneTrack> TrackClass) const override;
diff --git a/Engine/Plugins/Compositing/Composure/Source/ComposureEditor/Private/Sequencer/ComposurePostMoveSettingsPropertyTrackEditor.cpp b/Engine/Plugins/Compositing/Composure/Source/ComposureEditor/Private/Sequencer/ComposurePostMoveSettingsPropertyTrackEditor.cpp
--- a/Engine/Plugins/Compositing/Composure/Source/ComposureEditor/Private/Sequencer/ComposurePostMoveSettingsPropertyTrackEditor.cpp
+++ b/Engine/Plugins/Compositing/Composure/Source/ComposureEditor/Private/Sequencer/ComposurePostMoveSettingsPropertyTrackEditor.cpp
@@ -17,6 +17,10 @@ TSharedRef<ISequencerTrackEditor> FComposurePostMoveSettingsPropertyTrackEditor:
 	return MakeShareable(new FComposurePostMoveSettingsPropertyTrackEditor(InSequencer));
 }
 
+FText FComposurePostMoveSettingsPropertyTrackEditor::GetDisplayName() const
+{
+	return LOCTEXT("ComposurePostMoveSettingsPropertyTrackEditor_DisplayName", "Composure Post Move Settings Property");
+}
 
 void FComposurePostMoveSettingsPropertyTrackEditor::BuildTrackContextMenu(FMenuBuilder& MenuBuilder, UMovieSceneTrack* Track)
 {
diff --git a/Engine/Plugins/Compositing/Composure/Source/ComposureEditor/Private/Sequencer/ComposurePostMoveSettingsPropertyTrackEditor.h b/Engine/Plugins/Compositing/Composure/Source/ComposureEditor/Private/Sequencer/ComposurePostMoveSettingsPropertyTrackEditor.h
--- a/Engine/Plugins/Compositing/Composure/Source/ComposureEditor/Private/Sequencer/ComposurePostMoveSettingsPropertyTrackEditor.h
+++ b/Engine/Plugins/Compositing/Composure/Source/ComposureEditor/Private/Sequencer/ComposurePostMoveSettingsPropertyTrackEditor.h
@@ -26,6 +26,7 @@ public:
 
 	//~ FMovieSceneTrackEditor interface
 
+	virtual FText GetDisplayName() const override;
 	virtual void BuildTrackContextMenu(FMenuBuilder& MenuBuilder, UMovieSceneTrack* Track) override;
 
 	/**
diff --git a/Engine/Plugins/Experimental/Avalanche/Source/AvalancheSequencer/Private/ObjectBinding/AvaSequencerObjectBinding.cpp b/Engine/Plugins/Experimental/Avalanche/Source/AvalancheSequencer/Private/ObjectBinding/AvaSequencerObjectBinding.cpp
--- a/Engine/Plugins/Experimental/Avalanche/Source/AvalancheSequencer/Private/ObjectBinding/AvaSequencerObjectBinding.cpp
+++ b/Engine/Plugins/Experimental/Avalanche/Source/AvalancheSequencer/Private/ObjectBinding/AvaSequencerObjectBinding.cpp
@@ -25,6 +25,11 @@ FAvaSequencerObjectBinding::FAvaSequencerObjectBinding(TSharedRef<ISequencer> In
 {
 }
 
+FText FAvaSequencerObjectBinding::GetDisplayName() const
+{
+	return LOCTEXT("ActorTrackEditor_DisplayName", "Motion Design Actor");
+}
+
 void FAvaSequencerObjectBinding::BuildSequencerAddMenu(FMenuBuilder& OutMenuBuilder)
 {
 	OutMenuBuilder.AddSubMenu(
diff --git a/Engine/Plugins/Experimental/Avalanche/Source/AvalancheSequencer/Private/ObjectBinding/AvaSequencerObjectBinding.h b/Engine/Plugins/Experimental/Avalanche/Source/AvalancheSequencer/Private/ObjectBinding/AvaSequencerObjectBinding.h
--- a/Engine/Plugins/Experimental/Avalanche/Source/AvalancheSequencer/Private/ObjectBinding/AvaSequencerObjectBinding.h
+++ b/Engine/Plugins/Experimental/Avalanche/Source/AvalancheSequencer/Private/ObjectBinding/AvaSequencerObjectBinding.h
@@ -15,6 +15,7 @@ public:
 	FAvaSequencerObjectBinding(TSharedRef<ISequencer> InSequencer);
 
 	//~ Begin ISequencerEditorObjectBinding
+	virtual FText GetDisplayName() const override;
 	virtual void BuildSequencerAddMenu(FMenuBuilder& OutMenuBuilder) override;
 	virtual bool SupportsSequence(UMovieSceneSequence* InSequence) const override;
 	//~ End ISequencerEditorObjectBinding
diff --git a/Engine/Plugins/Experimental/Avalanche/Source/AvalancheShapesEditor/Private/TrackEditors/AvaShapeRectCornerTrackEditor.cpp b/Engine/Plugins/Experimental/Avalanche/Source/AvalancheShapesEditor/Private/TrackEditors/AvaShapeRectCornerTrackEditor.cpp
--- a/Engine/Plugins/Experimental/Avalanche/Source/AvalancheShapesEditor/Private/TrackEditors/AvaShapeRectCornerTrackEditor.cpp
+++ b/Engine/Plugins/Experimental/Avalanche/Source/AvalancheShapesEditor/Private/TrackEditors/AvaShapeRectCornerTrackEditor.cpp
@@ -4,11 +4,18 @@
 #include "Channels/MovieSceneByteChannel.h"
 #include "Channels/MovieSceneFloatChannel.h"
 
+#define LOCTEXT_NAMESPACE "AvaShapeRectCornerTrackEditor"
+
 TSharedRef<ISequencerTrackEditor> FAvaShapeRectCornerTrackEditor::CreateTrackEditor(TSharedRef<ISequencer> InOwningSequencer)
 {
 	return MakeShared<FAvaShapeRectCornerTrackEditor>(InOwningSequencer);
 }
 
+FText FAvaShapeRectCornerTrackEditor::GetDisplayName() const
+{
+	return LOCTEXT("AvaShapeRectCornerTrackEditor_DisplayName", "Shape Rect Corner");
+}
+
 void FAvaShapeRectCornerTrackEditor::GenerateKeysFromPropertyChanged(const FPropertyChangedParams& InPropertyChangedParams
 	, UMovieSceneSection* SectionToKey
 	, FGeneratedTrackKeys& OutGeneratedKeys)
@@ -39,3 +46,5 @@ void FAvaShapeRectCornerTrackEditor::GenerateKeysFromPropertyChanged(const FProp
 		, static_cast<uint8>(CornerSettings.Type)
 		, bKeyType));
 }
+
+#undef LOCTEXT_NAMESPACE
diff --git a/Engine/Plugins/Experimental/Avalanche/Source/AvalancheShapesEditor/Private/TrackEditors/AvaShapeRectCornerTrackEditor.h b/Engine/Plugins/Experimental/Avalanche/Source/AvalancheShapesEditor/Private/TrackEditors/AvaShapeRectCornerTrackEditor.h
--- a/Engine/Plugins/Experimental/Avalanche/Source/AvalancheShapesEditor/Private/TrackEditors/AvaShapeRectCornerTrackEditor.h
+++ b/Engine/Plugins/Experimental/Avalanche/Source/AvalancheShapesEditor/Private/TrackEditors/AvaShapeRectCornerTrackEditor.h
@@ -35,6 +35,8 @@ public:
 	
 protected:
 	//~ Begin FPropertyTrackEditor
+	virtual FText GetDisplayName() const override;
+
 	virtual void GenerateKeysFromPropertyChanged(const FPropertyChangedParams& InPropertyChangedParams
 		, UMovieSceneSection* InSectionToKey
 		, FGeneratedTrackKeys& OutGeneratedKeys) override;
diff --git a/Engine/Plugins/Experimental/ChaosCaching/Source/ChaosCachingEditor/Private/Sequencer/ChaosCacheTrackEditor.cpp b/Engine/Plugins/Experimental/ChaosCaching/Source/ChaosCachingEditor/Private/Sequencer/ChaosCacheTrackEditor.cpp
--- a/Engine/Plugins/Experimental/ChaosCaching/Source/ChaosCachingEditor/Private/Sequencer/ChaosCacheTrackEditor.cpp
+++ b/Engine/Plugins/Experimental/ChaosCaching/Source/ChaosCachingEditor/Private/Sequencer/ChaosCacheTrackEditor.cpp
@@ -268,6 +268,11 @@ TSharedRef<ISequencerSection> FChaosCacheTrackEditor::MakeSectionInterface(UMovi
 	return MakeShareable(new FChaosCacheSection(SectionObject, GetSequencer()));
 }
 
+FText FChaosCacheTrackEditor::GetDisplayName() const
+{
+	return LOCTEXT("ChaosTrackEditor_DisplayName", "Chaos");
+}
+
 void FChaosCacheTrackEditor::BuildObjectBindingTrackMenu(FMenuBuilder& MenuBuilder, const TArray<FGuid>& ObjectBindings, const UClass* ObjectClass)
 {
 	if (ObjectClass->IsChildOf(AChaosCacheManager::StaticClass()) || ObjectClass->IsChildOf(AActor::StaticClass()))
diff --git a/Engine/Plugins/Experimental/ChaosCaching/Source/ChaosCachingEditor/Public/Sequencer/ChaosCacheTrackEditor.h b/Engine/Plugins/Experimental/ChaosCaching/Source/ChaosCachingEditor/Public/Sequencer/ChaosCacheTrackEditor.h
--- a/Engine/Plugins/Experimental/ChaosCaching/Source/ChaosCachingEditor/Public/Sequencer/ChaosCacheTrackEditor.h
+++ b/Engine/Plugins/Experimental/ChaosCaching/Source/ChaosCachingEditor/Public/Sequencer/ChaosCacheTrackEditor.h
@@ -34,6 +34,7 @@ public:
 public:
 
 	// ~Begin ISequencerTrackEditor interface
+	virtual FText GetDisplayName() const override;
 	virtual void BuildObjectBindingTrackMenu(class FMenuBuilder& MenuBuilder, const TArray<FGuid>& ObjectBindings, const UClass* ObjectClass) override;
 	virtual TSharedRef<ISequencerSection> MakeSectionInterface(UMovieSceneSection& SectionObject, UMovieSceneTrack& Track, FGuid ObjectBinding) override;
 	virtual bool SupportsSequence(UMovieSceneSequence* InSequence) const override;
diff --git a/Engine/Plugins/Experimental/ClonerEffector/Source/ClonerEffectorEditor/Private/Cloner/Sequencer/MovieSceneClonerTrackEditor.cpp b/Engine/Plugins/Experimental/ClonerEffector/Source/ClonerEffectorEditor/Private/Cloner/Sequencer/MovieSceneClonerTrackEditor.cpp
--- a/Engine/Plugins/Experimental/ClonerEffector/Source/ClonerEffectorEditor/Private/Cloner/Sequencer/MovieSceneClonerTrackEditor.cpp
+++ b/Engine/Plugins/Experimental/ClonerEffector/Source/ClonerEffectorEditor/Private/Cloner/Sequencer/MovieSceneClonerTrackEditor.cpp
@@ -19,6 +19,11 @@ FMovieSceneClonerTrackEditor::~FMovieSceneClonerTrackEditor()
 	OnClonerTrackExists.RemoveAll(this);
 }
 
+FText FMovieSceneClonerTrackEditor::GetDisplayName() const
+{
+	return LOCTEXT("ClonerTrackEditor_DisplayName", "Cloner");
+}
+
 bool FMovieSceneClonerTrackEditor::SupportsType(TSubclassOf<UMovieSceneTrack> InTrackClass) const
 {
 	return InTrackClass == UMovieSceneNiagaraCacheTrack::StaticClass()
diff --git a/Engine/Plugins/Experimental/ClonerEffector/Source/ClonerEffectorEditor/Private/Cloner/Sequencer/MovieSceneClonerTrackEditor.h b/Engine/Plugins/Experimental/ClonerEffector/Source/ClonerEffectorEditor/Private/Cloner/Sequencer/MovieSceneClonerTrackEditor.h
--- a/Engine/Plugins/Experimental/ClonerEffector/Source/ClonerEffectorEditor/Private/Cloner/Sequencer/MovieSceneClonerTrackEditor.h
+++ b/Engine/Plugins/Experimental/ClonerEffector/Source/ClonerEffectorEditor/Private/Cloner/Sequencer/MovieSceneClonerTrackEditor.h
@@ -32,6 +32,7 @@ public:
 
 private:
 	//~ Begin FMovieSceneTrackEditor
+	virtual FText GetDisplayName() const override;
 	virtual bool SupportsType(TSubclassOf<UMovieSceneTrack> InTrackClass) const override;
 	//~ End FMovieSceneTrackEditor
 
diff --git a/Engine/Plugins/Experimental/GeometryCollectionPlugin/Source/GeometryCollectionSequencer/Private/GeometryCollectionTrackEditor.cpp b/Engine/Plugins/Experimental/GeometryCollectionPlugin/Source/GeometryCollectionSequencer/Private/GeometryCollectionTrackEditor.cpp
--- a/Engine/Plugins/Experimental/GeometryCollectionPlugin/Source/GeometryCollectionSequencer/Private/GeometryCollectionTrackEditor.cpp
+++ b/Engine/Plugins/Experimental/GeometryCollectionPlugin/Source/GeometryCollectionSequencer/Private/GeometryCollectionTrackEditor.cpp
@@ -299,6 +299,11 @@ TSharedRef<ISequencerSection> FGeometryCollectionTrackEditor::MakeSectionInterfa
 	return MakeShareable( new FGeometryCollectionTrackSection(SectionObject, GetSequencer()) );
 }
 
+FText FGeometryCollectionTrackEditor::GetDisplayName() const
+{
+	return LOCTEXT("GeometryCollectionTrackEditor_DisplayName", "Geometry Collection");
+}
+
 void FGeometryCollectionTrackEditor::BuildObjectBindingTrackMenu(FMenuBuilder& MenuBuilder, const TArray<FGuid>& ObjectBindings, const UClass* ObjectClass)
 {
 	if (ObjectClass->IsChildOf(UGeometryCollectionComponent::StaticClass()))
diff --git a/Engine/Plugins/Experimental/GeometryCollectionPlugin/Source/GeometryCollectionSequencer/Public/GeometryCollectionTrackEditor.h b/Engine/Plugins/Experimental/GeometryCollectionPlugin/Source/GeometryCollectionSequencer/Public/GeometryCollectionTrackEditor.h
--- a/Engine/Plugins/Experimental/GeometryCollectionPlugin/Source/GeometryCollectionSequencer/Public/GeometryCollectionTrackEditor.h
+++ b/Engine/Plugins/Experimental/GeometryCollectionPlugin/Source/GeometryCollectionSequencer/Public/GeometryCollectionTrackEditor.h
@@ -46,6 +46,7 @@ public:
 public:
 
 	// ISequencerTrackEditor interface
+	virtual FText GetDisplayName() const override;
 	virtual void BuildObjectBindingTrackMenu(FMenuBuilder& MenuBuilder, const TArray<FGuid>& ObjectBindings, const UClass* ObjectClass) override;
 	virtual TSharedRef<ISequencerSection> MakeSectionInterface( UMovieSceneSection& SectionObject, UMovieSceneTrack& Track, FGuid ObjectBinding ) override;
 	virtual bool SupportsSequence(UMovieSceneSequence* InSequence) const override;
diff --git a/Engine/Plugins/Experimental/PropertyAnimatorCore/Source/PropertyAnimatorCoreEditor/Private/Sequencer/MovieSceneAnimatorTrackEditor.cpp b/Engine/Plugins/Experimental/PropertyAnimatorCore/Source/PropertyAnimatorCoreEditor/Private/Sequencer/MovieSceneAnimatorTrackEditor.cpp
--- a/Engine/Plugins/Experimental/PropertyAnimatorCore/Source/PropertyAnimatorCoreEditor/Private/Sequencer/MovieSceneAnimatorTrackEditor.cpp
+++ b/Engine/Plugins/Experimental/PropertyAnimatorCore/Source/PropertyAnimatorCoreEditor/Private/Sequencer/MovieSceneAnimatorTrackEditor.cpp
@@ -17,6 +17,11 @@ FMovieSceneAnimatorTrackEditor::~FMovieSceneAnimatorTrackEditor()
 	OnGetAnimatorTrackCount.RemoveAll(this);
 }
 
+FText FMovieSceneAnimatorTrackEditor::GetDisplayName() const
+{
+	return LOCTEXT("AnimatorTrackEditor_DisplayName", "Animator");
+}
+
 void FMovieSceneAnimatorTrackEditor::BuildAddTrackMenu(FMenuBuilder& InMenuBuilder)
 {
 	// Empty, do not allow the creation of new tracks, only allow through sequencer time source
diff --git a/Engine/Plugins/Experimental/PropertyAnimatorCore/Source/PropertyAnimatorCoreEditor/Private/Sequencer/MovieSceneAnimatorTrackEditor.h b/Engine/Plugins/Experimental/PropertyAnimatorCore/Source/PropertyAnimatorCoreEditor/Private/Sequencer/MovieSceneAnimatorTrackEditor.h
--- a/Engine/Plugins/Experimental/PropertyAnimatorCore/Source/PropertyAnimatorCoreEditor/Private/Sequencer/MovieSceneAnimatorTrackEditor.h
+++ b/Engine/Plugins/Experimental/PropertyAnimatorCore/Source/PropertyAnimatorCoreEditor/Private/Sequencer/MovieSceneAnimatorTrackEditor.h
@@ -32,6 +32,7 @@ public:
 
 private:
 	//~ Begin FMovieSceneTrackEditor
+	virtual FText GetDisplayName() const override;
 	virtual void BuildAddTrackMenu(FMenuBuilder& InMenuBuilder) override;
 	virtual TSharedPtr<SWidget> BuildOutlinerEditWidget(const FGuid& InObjectBinding, UMovieSceneTrack* InTrack, const FBuildEditWidgetParams& InParams) override;
 	virtual bool SupportsSequence(UMovieSceneSequence* InSequence) const override;
diff --git a/Engine/Plugins/FX/Niagara/Source/NiagaraEditor/Private/Sequencer/LevelSequence/NiagaraSystemTrackEditor.cpp b/Engine/Plugins/FX/Niagara/Source/NiagaraEditor/Private/Sequencer/LevelSequence/NiagaraSystemTrackEditor.cpp
--- a/Engine/Plugins/FX/Niagara/Source/NiagaraEditor/Private/Sequencer/LevelSequence/NiagaraSystemTrackEditor.cpp
+++ b/Engine/Plugins/FX/Niagara/Source/NiagaraEditor/Private/Sequencer/LevelSequence/NiagaraSystemTrackEditor.cpp
@@ -31,6 +31,11 @@ FNiagaraSystemTrackEditor::FNiagaraSystemTrackEditor(TSharedRef<ISequencer> InSe
 {
 }
 
+FText FNiagaraSystemTrackEditor::GetDisplayName() const
+{
+	return LOCTEXT("NiagaraSystemTrackEditor_DisplayName", "Niagara System");
+}
+
 TSharedRef<ISequencerSection> FNiagaraSystemTrackEditor::MakeSectionInterface(UMovieSceneSection& SectionObject, UMovieSceneTrack& Track, FGuid ObjectBinding)
 {
 	checkf(SectionObject.GetClass()->IsChildOf<UMovieSceneNiagaraSystemSpawnSection>(), TEXT("Unsupported section."));
diff --git a/Engine/Plugins/FX/Niagara/Source/NiagaraEditor/Private/Sequencer/LevelSequence/NiagaraSystemTrackEditor.h b/Engine/Plugins/FX/Niagara/Source/NiagaraEditor/Private/Sequencer/LevelSequence/NiagaraSystemTrackEditor.h
--- a/Engine/Plugins/FX/Niagara/Source/NiagaraEditor/Private/Sequencer/LevelSequence/NiagaraSystemTrackEditor.h
+++ b/Engine/Plugins/FX/Niagara/Source/NiagaraEditor/Private/Sequencer/LevelSequence/NiagaraSystemTrackEditor.h
@@ -22,6 +22,11 @@ public:
 	{
 	}
 
+	virtual FText GetDisplayName() const override
+	{
+		return NSLOCTEXT("Sequencer", "NiagaraSystemParameterTrackEditor_DisplayName", "Niagara System Parameter");
+	}
+
 	virtual bool SupportsType(TSubclassOf<UMovieSceneTrack> Type) const override
 	{
 		return Type == MovieSceneTrackType::StaticClass();
@@ -38,6 +43,7 @@ public:
 	FNiagaraSystemTrackEditor(TSharedRef<ISequencer> InSequencer);
 
 	//~ ISequencerTrackEditor interface
+	virtual FText GetDisplayName() const override;
 	virtual TSharedRef<ISequencerSection> MakeSectionInterface(UMovieSceneSection& SectionObject, UMovieSceneTrack& Track, FGuid ObjectBinding) override;
 	virtual bool SupportsType(TSubclassOf<UMovieSceneTrack> Type) const override;
 	virtual void BuildObjectBindingTrackMenu(FMenuBuilder& MenuBuilder, const TArray<FGuid>& ObjectBindings, const UClass* ObjectClass) override;
diff --git a/Engine/Plugins/FX/Niagara/Source/NiagaraEditor/Private/Sequencer/NiagaraSequence/NiagaraEmitterTrackEditor.cpp b/Engine/Plugins/FX/Niagara/Source/NiagaraEditor/Private/Sequencer/NiagaraSequence/NiagaraEmitterTrackEditor.cpp
--- a/Engine/Plugins/FX/Niagara/Source/NiagaraEditor/Private/Sequencer/NiagaraSequence/NiagaraEmitterTrackEditor.cpp
+++ b/Engine/Plugins/FX/Niagara/Source/NiagaraEditor/Private/Sequencer/NiagaraSequence/NiagaraEmitterTrackEditor.cpp
@@ -281,6 +281,11 @@ TSharedRef<ISequencerTrackEditor> FNiagaraEmitterTrackEditor::CreateTrackEditor(
 	return MakeShareable(new FNiagaraEmitterTrackEditor(InSequencer));
 }
 
+FText FNiagaraEmitterTrackEditor::GetDisplayName() const
+{
+	return LOCTEXT("NiagaraEmitterTrackEditor_DisplayName", "Niagara Emitter");
+}
+
 bool FNiagaraEmitterTrackEditor::SupportsType(TSubclassOf<UMovieSceneTrack> TrackClass) const 
 {
 	if (TrackClass == UMovieSceneNiagaraEmitterTrack::StaticClass())
diff --git a/Engine/Plugins/FX/Niagara/Source/NiagaraEditor/Private/Sequencer/NiagaraSequence/NiagaraEmitterTrackEditor.h b/Engine/Plugins/FX/Niagara/Source/NiagaraEditor/Private/Sequencer/NiagaraSequence/NiagaraEmitterTrackEditor.h
--- a/Engine/Plugins/FX/Niagara/Source/NiagaraEditor/Private/Sequencer/NiagaraSequence/NiagaraEmitterTrackEditor.h
+++ b/Engine/Plugins/FX/Niagara/Source/NiagaraEditor/Private/Sequencer/NiagaraSequence/NiagaraEmitterTrackEditor.h
@@ -19,6 +19,7 @@ public:
 	static TSharedRef<ISequencerTrackEditor> CreateTrackEditor(TSharedRef<ISequencer> InSequencer);
 
 	//~ FMovieSceneTrackEditor interface.
+	virtual FText GetDisplayName() const override;
 	virtual bool SupportsType(TSubclassOf<UMovieSceneTrack> TrackClass) const override;
 	virtual TSharedRef<ISequencerSection> MakeSectionInterface(UMovieSceneSection& SectionObject, UMovieSceneTrack& Track, FGuid ObjectBinding) override;
 	virtual bool HandleAssetAdded(UObject* Asset, const FGuid& TargetObjectGuid) override;
diff --git a/Engine/Plugins/FX/NiagaraSimCaching/Source/NiagaraSimCachingEditor/Private/Sequencer/NiagaraCacheTrackEditor.cpp b/Engine/Plugins/FX/NiagaraSimCaching/Source/NiagaraSimCachingEditor/Private/Sequencer/NiagaraCacheTrackEditor.cpp
--- a/Engine/Plugins/FX/NiagaraSimCaching/Source/NiagaraSimCachingEditor/Private/Sequencer/NiagaraCacheTrackEditor.cpp
+++ b/Engine/Plugins/FX/NiagaraSimCaching/Source/NiagaraSimCachingEditor/Private/Sequencer/NiagaraCacheTrackEditor.cpp
@@ -279,6 +279,11 @@ TSharedRef<ISequencerSection> FNiagaraCacheTrackEditor::MakeSectionInterface(UMo
 	return MakeShareable(new FNiagaraCacheSection(SectionObject, GetSequencer()));
 }
 
+FText FNiagaraCacheTrackEditor::GetDisplayName() const
+{
+	return LOCTEXT("NiagaraCacheTrackEditor_DisplayName", "Niagara");
+}
+
 void FNiagaraCacheTrackEditor::BuildObjectBindingTrackMenu(FMenuBuilder& MenuBuilder, const TArray<FGuid>& ObjectBindings, const UClass* ObjectClass)
 {
 	if (ObjectClass->IsChildOf(UNiagaraComponent::StaticClass()) || ObjectClass->IsChildOf(AActor::StaticClass()))
diff --git a/Engine/Plugins/FX/NiagaraSimCaching/Source/NiagaraSimCachingEditor/Public/Sequencer/NiagaraCacheTrackEditor.h b/Engine/Plugins/FX/NiagaraSimCaching/Source/NiagaraSimCachingEditor/Public/Sequencer/NiagaraCacheTrackEditor.h
--- a/Engine/Plugins/FX/NiagaraSimCaching/Source/NiagaraSimCachingEditor/Public/Sequencer/NiagaraCacheTrackEditor.h
+++ b/Engine/Plugins/FX/NiagaraSimCaching/Source/NiagaraSimCachingEditor/Public/Sequencer/NiagaraCacheTrackEditor.h
@@ -35,6 +35,7 @@ public:
 public:
 
 	// ~Begin ISequencerTrackEditor interface
+	virtual FText GetDisplayName() const override;
 	virtual void BuildObjectBindingTrackMenu(FMenuBuilder& MenuBuilder, const TArray<FGuid>& ObjectBindings, const UClass* ObjectClass) override;
 	virtual void BuildTrackContextMenu(FMenuBuilder& MenuBuilder, UMovieSceneTrack* Track) override; 
 	virtual TSharedRef<ISequencerSection> MakeSectionInterface(UMovieSceneSection& SectionObject, UMovieSceneTrack& Track, FGuid ObjectBinding) override;
diff --git a/Engine/Plugins/Media/MediaCompositing/Source/MediaCompositingEditor/Private/Sequencer/MediaPlayerPropertyTrackEditor.cpp b/Engine/Plugins/Media/MediaCompositing/Source/MediaCompositingEditor/Private/Sequencer/MediaPlayerPropertyTrackEditor.cpp
--- a/Engine/Plugins/Media/MediaCompositing/Source/MediaCompositingEditor/Private/Sequencer/MediaPlayerPropertyTrackEditor.cpp
+++ b/Engine/Plugins/Media/MediaCompositing/Source/MediaCompositingEditor/Private/Sequencer/MediaPlayerPropertyTrackEditor.cpp
@@ -131,6 +131,11 @@ void FMediaPlayerPropertyTrackEditor::OnAnimatedPropertyChanged(const FPropertyC
 	}
 }
 
+FText FMediaPlayerPropertyTrackEditor::GetDisplayName() const
+{
+	return LOCTEXT("MediaPlayerPropertyTrackEditor_DisplayName", "Media Player Property");
+}
+
 UMovieSceneTrack* FMediaPlayerPropertyTrackEditor::AddTrack(UMovieScene* FocusedMovieScene, const FGuid& ObjectHandle, TSubclassOf<class UMovieSceneTrack> TrackClass, FName UniqueTypeName)
 {
 	UMovieSceneTrack* NewTrack = FPropertyTrackEditor::AddTrack(FocusedMovieScene, ObjectHandle, TrackClass, UniqueTypeName);
diff --git a/Engine/Plugins/Media/MediaCompositing/Source/MediaCompositingEditor/Private/Sequencer/MediaPlayerPropertyTrackEditor.h b/Engine/Plugins/Media/MediaCompositing/Source/MediaCompositingEditor/Private/Sequencer/MediaPlayerPropertyTrackEditor.h
--- a/Engine/Plugins/Media/MediaCompositing/Source/MediaCompositingEditor/Private/Sequencer/MediaPlayerPropertyTrackEditor.h
+++ b/Engine/Plugins/Media/MediaCompositing/Source/MediaCompositingEditor/Private/Sequencer/MediaPlayerPropertyTrackEditor.h
@@ -36,6 +36,7 @@ public:
 
 private:
 
+	virtual FText GetDisplayName() const override;
 	virtual UMovieSceneTrack* AddTrack(UMovieScene* FocusedMovieScene, const FGuid& ObjectHandle, TSubclassOf<class UMovieSceneTrack> TrackClass, FName UniqueTypeName) override;
 	virtual TSharedPtr<SWidget> BuildOutlinerEditWidget(const FGuid& ObjectBinding, UMovieSceneTrack* Track, const FBuildEditWidgetParams& Params) override;
 	virtual bool SupportsSequence(UMovieSceneSequence* InSequence) const override;
diff --git a/Engine/Plugins/Media/MediaCompositing/Source/MediaCompositingEditor/Private/Sequencer/MediaTrackEditor.cpp b/Engine/Plugins/Media/MediaCompositing/Source/MediaCompositingEditor/Private/Sequencer/MediaTrackEditor.cpp
--- a/Engine/Plugins/Media/MediaCompositing/Source/MediaCompositingEditor/Private/Sequencer/MediaTrackEditor.cpp
+++ b/Engine/Plugins/Media/MediaCompositing/Source/MediaCompositingEditor/Private/Sequencer/MediaTrackEditor.cpp
@@ -64,6 +64,11 @@ FMediaTrackEditor::~FMediaTrackEditor()
 /* FMovieSceneTrackEditor interface
  *****************************************************************************/
 
+FText FMediaTrackEditor::GetDisplayName() const
+{
+	return LOCTEXT("MediaTrackEditor_DisplayName", "Media");
+}
+
 UMovieSceneTrack* FMediaTrackEditor::AddTrack(UMovieScene* FocusedMovieScene, const FGuid& ObjectHandle, TSubclassOf<class UMovieSceneTrack> TrackClass, FName UniqueTypeName)
 {
 	UMovieSceneTrack* Track = FocusedMovieScene->AddTrack(TrackClass, ObjectHandle);
diff --git a/Engine/Plugins/Media/MediaCompositing/Source/MediaCompositingEditor/Public/Sequencer/MediaTrackEditor.h b/Engine/Plugins/Media/MediaCompositing/Source/MediaCompositingEditor/Public/Sequencer/MediaTrackEditor.h
--- a/Engine/Plugins/Media/MediaCompositing/Source/MediaCompositingEditor/Public/Sequencer/MediaTrackEditor.h
+++ b/Engine/Plugins/Media/MediaCompositing/Source/MediaCompositingEditor/Public/Sequencer/MediaTrackEditor.h
@@ -62,6 +62,7 @@ public:
 
 	//~ FMovieSceneTrackEditor interface
 
+	virtual FText GetDisplayName() const override;
 	virtual UMovieSceneTrack* AddTrack(UMovieScene* FocusedMovieScene, const FGuid& ObjectHandle, TSubclassOf<class UMovieSceneTrack> TrackClass, FName UniqueTypeName) override;
 	virtual void BuildAddTrackMenu(FMenuBuilder& MenuBuilder) override;
 	virtual TSharedPtr<SWidget> BuildOutlinerEditWidget(const FGuid& ObjectBinding, UMovieSceneTrack* Track, const FBuildEditWidgetParams& Params) override;
diff --git a/Engine/Plugins/Media/MediaPlate/Source/MediaPlateEditor/Private/Sequencer/MediaPlateTrackEditor.cpp b/Engine/Plugins/Media/MediaPlate/Source/MediaPlateEditor/Private/Sequencer/MediaPlateTrackEditor.cpp
--- a/Engine/Plugins/Media/MediaPlate/Source/MediaPlateEditor/Private/Sequencer/MediaPlateTrackEditor.cpp
+++ b/Engine/Plugins/Media/MediaPlate/Source/MediaPlateEditor/Private/Sequencer/MediaPlateTrackEditor.cpp
@@ -45,6 +45,11 @@ FMediaPlateTrackEditor::~FMediaPlateTrackEditor()
 	FMediaTrackEditor::OnBuildOutlinerEditWidget.RemoveAll(this);
 }
 
+FText FMediaPlateTrackEditor::GetDisplayName() const
+{
+	return LOCTEXT("MediaPlateTrackEditor_DisplayName", "Media Plate");
+}
+
 void FMediaPlateTrackEditor::BuildObjectBindingContextMenu(FMenuBuilder& MenuBuilder, const TArray<FGuid>& ObjectBindings, const UClass* ObjectClass)
 {
 	if ((ObjectClass != nullptr) && (ObjectClass->IsChildOf(AMediaPlate::StaticClass())))
diff --git a/Engine/Plugins/Media/MediaPlate/Source/MediaPlateEditor/Private/Sequencer/MediaPlateTrackEditor.h b/Engine/Plugins/Media/MediaPlate/Source/MediaPlateEditor/Private/Sequencer/MediaPlateTrackEditor.h
--- a/Engine/Plugins/Media/MediaPlate/Source/MediaPlateEditor/Private/Sequencer/MediaPlateTrackEditor.h
+++ b/Engine/Plugins/Media/MediaPlate/Source/MediaPlateEditor/Private/Sequencer/MediaPlateTrackEditor.h
@@ -52,6 +52,7 @@ public:
 public:
 
 	//~ FMovieSceneTrackEditor interface
+	virtual FText GetDisplayName() const override;
 	virtual void BuildObjectBindingContextMenu(FMenuBuilder& MenuBuilder, const TArray<FGuid>& ObjectBindings, const UClass* ObjectClass) override;
 	virtual void BuildObjectBindingTrackMenu(FMenuBuilder& MenuBuilder, const TArray<FGuid>& ObjectBindings, const UClass* ObjectClass) override;
 	virtual bool SupportsType(TSubclassOf<UMovieSceneTrack> TrackClass) const override;
diff --git a/Engine/Plugins/MovieScene/LevelSequenceEditor/Source/LevelSequenceEditor/Private/Misc/LevelSequenceEditorActorBinding.cpp b/Engine/Plugins/MovieScene/LevelSequenceEditor/Source/LevelSequenceEditor/Private/Misc/LevelSequenceEditorActorBinding.cpp
--- a/Engine/Plugins/MovieScene/LevelSequenceEditor/Source/LevelSequenceEditor/Private/Misc/LevelSequenceEditorActorBinding.cpp
+++ b/Engine/Plugins/MovieScene/LevelSequenceEditor/Source/LevelSequenceEditor/Private/Misc/LevelSequenceEditorActorBinding.cpp
@@ -24,10 +24,15 @@ FLevelSequenceEditorActorBinding::FLevelSequenceEditorActorBinding(TSharedRef<IS
 {
 }
 
+FText FLevelSequenceEditorActorBinding::GetDisplayName() const
+{
+	return LOCTEXT("ActorTrackEditor_DisplayName", "Actor");
+}
+
 void FLevelSequenceEditorActorBinding::BuildSequencerAddMenu(FMenuBuilder& MenuBuilder)
 {
 	MenuBuilder.AddSubMenu(
-		LOCTEXT("AddActor_Label", "Actor To Sequencer"),
+		LOCTEXT("AddActor_Label", "Add Actor Track"),
 		LOCTEXT("AddActor_ToolTip", "Allow sequencer to possess an actor that already exists in the current level"),
 		FNewMenuDelegate::CreateRaw(this, &FLevelSequenceEditorActorBinding::AddPossessActorMenuExtensions),
 		false /*bInOpenSubMenuOnClick*/,
diff --git a/Engine/Plugins/MovieScene/LevelSequenceEditor/Source/LevelSequenceEditor/Private/Misc/LevelSequenceEditorActorBinding.h b/Engine/Plugins/MovieScene/LevelSequenceEditor/Source/LevelSequenceEditor/Private/Misc/LevelSequenceEditorActorBinding.h
--- a/Engine/Plugins/MovieScene/LevelSequenceEditor/Source/LevelSequenceEditor/Private/Misc/LevelSequenceEditorActorBinding.h
+++ b/Engine/Plugins/MovieScene/LevelSequenceEditor/Source/LevelSequenceEditor/Private/Misc/LevelSequenceEditorActorBinding.h
@@ -15,6 +15,7 @@ public:
 	FLevelSequenceEditorActorBinding(TSharedRef<ISequencer> InSequencer);
 
 	// ISequencerEditorObjectBinding interface
+	virtual FText GetDisplayName() const override;
 	virtual void BuildSequencerAddMenu(FMenuBuilder& MenuBuilder) override;
 	virtual bool SupportsSequence(UMovieSceneSequence* InSequence) const override;
 
diff --git a/Engine/Plugins/MovieScene/TemplateSequence/Source/TemplateSequenceEditor/Private/TrackEditors/TemplateSequenceTrackEditor.cpp b/Engine/Plugins/MovieScene/TemplateSequence/Source/TemplateSequenceEditor/Private/TrackEditors/TemplateSequenceTrackEditor.cpp
--- a/Engine/Plugins/MovieScene/TemplateSequence/Source/TemplateSequenceEditor/Private/TrackEditors/TemplateSequenceTrackEditor.cpp
+++ b/Engine/Plugins/MovieScene/TemplateSequence/Source/TemplateSequenceEditor/Private/TrackEditors/TemplateSequenceTrackEditor.cpp
@@ -40,6 +40,11 @@ TSharedRef<ISequencerTrackEditor> FTemplateSequenceTrackEditor::CreateTrackEdito
 	return MakeShareable(new FTemplateSequenceTrackEditor(InSequencer));
 }
 
+FText FTemplateSequenceTrackEditor::GetDisplayName() const
+{
+	return LOCTEXT("TemplateSequenceTrackEditor_DisplayName", "Template Sequence");
+}
+
 bool FTemplateSequenceTrackEditor::SupportsType(TSubclassOf<UMovieSceneTrack> Type) const
 {
 	return Type == UTemplateSequenceTrack::StaticClass();
diff --git a/Engine/Plugins/MovieScene/TemplateSequence/Source/TemplateSequenceEditor/Private/TrackEditors/TemplateSequenceTrackEditor.h b/Engine/Plugins/MovieScene/TemplateSequence/Source/TemplateSequenceEditor/Private/TrackEditors/TemplateSequenceTrackEditor.h
--- a/Engine/Plugins/MovieScene/TemplateSequence/Source/TemplateSequenceEditor/Private/TrackEditors/TemplateSequenceTrackEditor.h
+++ b/Engine/Plugins/MovieScene/TemplateSequence/Source/TemplateSequenceEditor/Private/TrackEditors/TemplateSequenceTrackEditor.h
@@ -24,6 +24,7 @@ public:
 
 public:
 	// ISequencerTrackEditor interface
+	virtual FText GetDisplayName() const override;
 	virtual bool SupportsType(TSubclassOf<class UMovieSceneTrack> TrackClass) const override;
 	virtual bool SupportsSequence(UMovieSceneSequence* InSequence) const override;
 	virtual void BuildObjectBindingTrackMenu(FMenuBuilder& MenuBuilder, const TArray<FGuid>& ObjectBindings, const UClass* ObjectClass) override;
diff --git a/Engine/Plugins/Runtime/GeometryCache/Source/GeometryCacheSequencer/Classes/GeometryCacheTrackEditor.cpp b/Engine/Plugins/Runtime/GeometryCache/Source/GeometryCacheSequencer/Classes/GeometryCacheTrackEditor.cpp
--- a/Engine/Plugins/Runtime/GeometryCache/Source/GeometryCacheSequencer/Classes/GeometryCacheTrackEditor.cpp
+++ b/Engine/Plugins/Runtime/GeometryCache/Source/GeometryCacheSequencer/Classes/GeometryCacheTrackEditor.cpp
@@ -299,6 +299,11 @@ TSharedRef<ISequencerSection> FGeometryCacheTrackEditor::MakeSectionInterface(UM
 	return MakeShareable(new FGeometryCacheSection(SectionObject, GetSequencer()));
 }
 
+FText FGeometryCacheTrackEditor::GetDisplayName() const
+{
+	return LOCTEXT("GeometryCacheTrackEditor_DisplayName", "Geometry");
+}
+
 void FGeometryCacheTrackEditor::BuildObjectBindingTrackMenu(FMenuBuilder& MenuBuilder, const TArray<FGuid>& ObjectBindings, const UClass* ObjectClass)
 {
 	if (ObjectClass->IsChildOf(UGeometryCacheComponent::StaticClass()) || ObjectClass->IsChildOf(AActor::StaticClass()))
diff --git a/Engine/Plugins/Runtime/GeometryCache/Source/GeometryCacheSequencer/Classes/GeometryCacheTrackEditor.h b/Engine/Plugins/Runtime/GeometryCache/Source/GeometryCacheSequencer/Classes/GeometryCacheTrackEditor.h
--- a/Engine/Plugins/Runtime/GeometryCache/Source/GeometryCacheSequencer/Classes/GeometryCacheTrackEditor.h
+++ b/Engine/Plugins/Runtime/GeometryCache/Source/GeometryCacheSequencer/Classes/GeometryCacheTrackEditor.h
@@ -41,6 +41,7 @@ public:
 public:
 
 	// ISequencerTrackEditor interface
+	virtual FText GetDisplayName() const override;
 	virtual void BuildObjectBindingTrackMenu(FMenuBuilder& MenuBuilder, const TArray<FGuid>& ObjectBindings, const UClass* ObjectClass) override;
 	virtual TSharedRef<ISequencerSection> MakeSectionInterface( UMovieSceneSection& SectionObject, UMovieSceneTrack& Track, FGuid ObjectBinding ) override;
 	virtual bool SupportsSequence(UMovieSceneSequence* InSequence) const override;
diff --git a/Engine/Plugins/Runtime/HairStrands/Source/HairStrandsEditor/Private/GroomCacheTrackEditor.cpp b/Engine/Plugins/Runtime/HairStrands/Source/HairStrandsEditor/Private/GroomCacheTrackEditor.cpp
--- a/Engine/Plugins/Runtime/HairStrands/Source/HairStrandsEditor/Private/GroomCacheTrackEditor.cpp
+++ b/Engine/Plugins/Runtime/HairStrands/Source/HairStrandsEditor/Private/GroomCacheTrackEditor.cpp
@@ -288,6 +288,11 @@ TSharedRef<ISequencerSection> FGroomCacheTrackEditor::MakeSectionInterface(UMovi
 	return MakeShareable(new FGroomCacheSection(SectionObject, GetSequencer()));
 }
 
+FText FGroomCacheTrackEditor::GetDisplayName() const
+{
+	return LOCTEXT("GroomCacheTrackEditor_DisplayName", "Groom");
+}
+
 void FGroomCacheTrackEditor::BuildObjectBindingTrackMenu(FMenuBuilder& MenuBuilder, const TArray<FGuid>& ObjectBindings, const UClass* ObjectClass)
 {
 	if (ObjectClass->IsChildOf(UGroomComponent::StaticClass()) || ObjectClass->IsChildOf(AActor::StaticClass()))
diff --git a/Engine/Plugins/Runtime/HairStrands/Source/HairStrandsEditor/Private/GroomCacheTrackEditor.h b/Engine/Plugins/Runtime/HairStrands/Source/HairStrandsEditor/Private/GroomCacheTrackEditor.h
--- a/Engine/Plugins/Runtime/HairStrands/Source/HairStrandsEditor/Private/GroomCacheTrackEditor.h
+++ b/Engine/Plugins/Runtime/HairStrands/Source/HairStrandsEditor/Private/GroomCacheTrackEditor.h
@@ -34,6 +34,7 @@ public:
 public:
 
 	// ISequencerTrackEditor interface
+	virtual FText GetDisplayName() const override;
 	virtual void BuildObjectBindingTrackMenu(class FMenuBuilder& MenuBuilder, const TArray<FGuid>& ObjectBindings, const UClass* ObjectClass) override;
 	virtual TSharedRef<ISequencerSection> MakeSectionInterface(UMovieSceneSection& SectionObject, UMovieSceneTrack& Track, FGuid ObjectBinding) override;
 	virtual bool SupportsSequence(UMovieSceneSequence* InSequence) const override;
diff --git a/Engine/Plugins/VirtualProduction/LensComponent/Source/LensComponentEditor/Private/MovieScene/LensComponentTrackEditor.cpp b/Engine/Plugins/VirtualProduction/LensComponent/Source/LensComponentEditor/Private/MovieScene/LensComponentTrackEditor.cpp
--- a/Engine/Plugins/VirtualProduction/LensComponent/Source/LensComponentEditor/Private/MovieScene/LensComponentTrackEditor.cpp
+++ b/Engine/Plugins/VirtualProduction/LensComponent/Source/LensComponentEditor/Private/MovieScene/LensComponentTrackEditor.cpp
@@ -4,11 +4,18 @@
 
 #include "LevelSequence.h"
 
+#define LOCTEXT_NAMESPACE "LensComponentTrackEditor"
+
 TSharedRef<ISequencerTrackEditor> FLensComponentTrackEditor::CreateTrackEditor(TSharedRef<ISequencer> InSequencer)
 {
 	return MakeShared<FLensComponentTrackEditor>(InSequencer);
 }
 
+FText FLensComponentTrackEditor::GetDisplayName() const
+{
+	return LOCTEXT("LensComponentTrackEditor_DisplayName", "Lens");
+}
+
 bool FLensComponentTrackEditor::SupportsSequence(UMovieSceneSequence* InSequence) const
 {
 	if (!InSequence || InSequence->IsTrackSupported(UMovieSceneLensComponentTrack::StaticClass()) == ETrackSupport::NotSupported)
@@ -23,3 +30,5 @@ bool FLensComponentTrackEditor::SupportsType(TSubclassOf<UMovieSceneTrack> Type)
 {
 	return (Type == UMovieSceneLensComponentTrack::StaticClass());
 }
+
+#undef LOCTEXT_NAMESPACE
diff --git a/Engine/Plugins/VirtualProduction/LensComponent/Source/LensComponentEditor/Private/MovieScene/LensComponentTrackEditor.h b/Engine/Plugins/VirtualProduction/LensComponent/Source/LensComponentEditor/Private/MovieScene/LensComponentTrackEditor.h
--- a/Engine/Plugins/VirtualProduction/LensComponent/Source/LensComponentEditor/Private/MovieScene/LensComponentTrackEditor.h
+++ b/Engine/Plugins/VirtualProduction/LensComponent/Source/LensComponentEditor/Private/MovieScene/LensComponentTrackEditor.h
@@ -24,6 +24,7 @@ public:
 
 public:
 	//~ Begin ISequencerTrackEditor interface
+	virtual FText GetDisplayName() const override;
 	virtual bool SupportsSequence(UMovieSceneSequence* InSequence) const override;
 	virtual bool SupportsType(TSubclassOf<UMovieSceneTrack> Type) const override;
 	//~ End ISequencerTrackEditor interface
diff --git a/Engine/Plugins/VirtualProduction/LiveLinkCamera/Source/LiveLinkCameraRecording/Private/LiveLinkCameraControllerTrackEditor.cpp b/Engine/Plugins/VirtualProduction/LiveLinkCamera/Source/LiveLinkCameraRecording/Private/LiveLinkCameraControllerTrackEditor.cpp
--- a/Engine/Plugins/VirtualProduction/LiveLinkCamera/Source/LiveLinkCameraRecording/Private/LiveLinkCameraControllerTrackEditor.cpp
+++ b/Engine/Plugins/VirtualProduction/LiveLinkCamera/Source/LiveLinkCameraRecording/Private/LiveLinkCameraControllerTrackEditor.cpp
@@ -4,12 +4,18 @@
 
 #include "LevelSequence.h"
 
+#define LOCTEXT_NAMESPACE "LiveLinkCameraControllerTrackEditor"
 
 TSharedRef<ISequencerTrackEditor> FLiveLinkCameraControllerTrackEditor::CreateTrackEditor(TSharedRef<ISequencer> InSequencer)
 {
 	return MakeShared<FLiveLinkCameraControllerTrackEditor>(InSequencer);
 }
 
+FText FLiveLinkCameraControllerTrackEditor::GetDisplayName() const
+{
+	return LOCTEXT("LiveLinkCameraControllerTrackEditor_DisplayName", "Live Link Camera Controller");
+}
+
 bool FLiveLinkCameraControllerTrackEditor::SupportsSequence(UMovieSceneSequence* InSequence) const
 {
 	ETrackSupport TrackSupported = InSequence ? InSequence->IsTrackSupported(UMovieSceneLiveLinkCameraControllerTrack::StaticClass()) : ETrackSupport::Default;
@@ -26,3 +32,5 @@ bool FLiveLinkCameraControllerTrackEditor::SupportsType(TSubclassOf<UMovieSceneT
 {
 	return (Type == UMovieSceneLiveLinkCameraControllerTrack::StaticClass());
 }
+
+#undef LOCTEXT_NAMESPACE
diff --git a/Engine/Plugins/VirtualProduction/LiveLinkCamera/Source/LiveLinkCameraRecording/Private/LiveLinkCameraControllerTrackEditor.h b/Engine/Plugins/VirtualProduction/LiveLinkCamera/Source/LiveLinkCameraRecording/Private/LiveLinkCameraControllerTrackEditor.h
--- a/Engine/Plugins/VirtualProduction/LiveLinkCamera/Source/LiveLinkCameraRecording/Private/LiveLinkCameraControllerTrackEditor.h
+++ b/Engine/Plugins/VirtualProduction/LiveLinkCamera/Source/LiveLinkCameraRecording/Private/LiveLinkCameraControllerTrackEditor.h
@@ -26,6 +26,7 @@ public:
 public:
 
 	//~ Begin ISequencerTrackEditor interface
+	virtual FText GetDisplayName() const override;
 	virtual bool SupportsSequence(UMovieSceneSequence* InSequence) const override;
 	virtual bool SupportsType(TSubclassOf<UMovieSceneTrack> Type) const override;
 	//~ End ISequencerTrackEditor interface
diff --git a/Engine/Plugins/VirtualProduction/Takes/Source/TakeSequencer/Private/TakeTrackEditor.cpp b/Engine/Plugins/VirtualProduction/Takes/Source/TakeSequencer/Private/TakeTrackEditor.cpp
--- a/Engine/Plugins/VirtualProduction/Takes/Source/TakeSequencer/Private/TakeTrackEditor.cpp
+++ b/Engine/Plugins/VirtualProduction/Takes/Source/TakeSequencer/Private/TakeTrackEditor.cpp
@@ -49,6 +49,10 @@ bool FTakeTrackEditor::SupportsType(TSubclassOf<UMovieSceneTrack> Type) const
 	return Type == UMovieSceneTakeTrack::StaticClass();
 }
 
+FText FTakeTrackEditor::GetDisplayName() const
+{
+	return LOCTEXT("TakeTrackEditor_DisplayName", "Take");
+}
 
 TSharedRef<ISequencerSection> FTakeTrackEditor::MakeSectionInterface(UMovieSceneSection& SectionObject, UMovieSceneTrack& Track, FGuid ObjectBinding)
 {
diff --git a/Engine/Plugins/VirtualProduction/Takes/Source/TakeSequencer/Private/TakeTrackEditor.h b/Engine/Plugins/VirtualProduction/Takes/Source/TakeSequencer/Private/TakeTrackEditor.h
--- a/Engine/Plugins/VirtualProduction/Takes/Source/TakeSequencer/Private/TakeTrackEditor.h
+++ b/Engine/Plugins/VirtualProduction/Takes/Source/TakeSequencer/Private/TakeTrackEditor.h
@@ -45,6 +45,7 @@ public:
 public:
 		
 	// ISequencerTrackEditor interface
+	virtual FText GetDisplayName() const override;
 	virtual TSharedRef<ISequencerSection> MakeSectionInterface( UMovieSceneSection& SectionObject, UMovieSceneTrack& Track, FGuid ObjectBinding ) override;
 	virtual bool SupportsType( TSubclassOf<UMovieSceneTrack> Type ) const override;
 	virtual const FSlateBrush* GetIconBrush() const override;
diff --git a/Engine/Source/Editor/MovieSceneTools/Private/TrackEditors/ActorPickerTrackEditor.cpp b/Engine/Source/Editor/MovieSceneTools/Private/TrackEditors/ActorPickerTrackEditor.cpp
--- a/Engine/Source/Editor/MovieSceneTools/Private/TrackEditors/ActorPickerTrackEditor.cpp
+++ b/Engine/Source/Editor/MovieSceneTools/Private/TrackEditors/ActorPickerTrackEditor.cpp
@@ -50,6 +50,11 @@ void FActorPickerTrackEditor::PickActorInteractive(const TArray<FGuid>& ObjectBi
 	}
 }
 
+FText FActorPickerTrackEditor::GetDisplayName() const
+{
+	return LOCTEXT("ActorPickerTrackEditor_DisplayName", "Actor");
+}
+
 void FActorPickerTrackEditor::ShowActorSubMenu(FMenuBuilder& MenuBuilder, TArray<FGuid> ObjectBindings, UMovieSceneSection* Section)
 {
 	struct Local
diff --git a/Engine/Source/Editor/MovieSceneTools/Private/TrackEditors/ActorPickerTrackEditor.h b/Engine/Source/Editor/MovieSceneTools/Private/TrackEditors/ActorPickerTrackEditor.h
--- a/Engine/Source/Editor/MovieSceneTools/Private/TrackEditors/ActorPickerTrackEditor.h
+++ b/Engine/Source/Editor/MovieSceneTools/Private/TrackEditors/ActorPickerTrackEditor.h
@@ -32,6 +32,10 @@ public:
 
 	FActorPickerTrackEditor(TSharedRef<ISequencer> InSequencer);
 
+	//~ Begin ISequencerTrackEditor
+	virtual FText GetDisplayName() const override;
+	//~ End ISequencerTrackEditor
+
 	/** Is this actor pickable? */
 	virtual bool IsActorPickable( const AActor* const ParentActor, FGuid ObjectBinding, UMovieSceneSection* InSection) { return false; }
 
diff --git a/Engine/Source/Editor/MovieSceneTools/Private/TrackEditors/AudioTrackEditor.cpp b/Engine/Source/Editor/MovieSceneTools/Private/TrackEditors/AudioTrackEditor.cpp
--- a/Engine/Source/Editor/MovieSceneTools/Private/TrackEditors/AudioTrackEditor.cpp
+++ b/Engine/Source/Editor/MovieSceneTools/Private/TrackEditors/AudioTrackEditor.cpp
@@ -1051,6 +1051,11 @@ FAudioTrackEditor::~FAudioTrackEditor()
 {
 }
 
+FText FAudioTrackEditor::GetDisplayName() const
+{
+	return LOCTEXT("AudioTrackEditor_DisplayName", "Audio");
+}
+
 void FAudioTrackEditor::OnInitialize()
 {
 	RegisterMovieSceneChangedDelegate();
diff --git a/Engine/Source/Editor/MovieSceneTools/Private/TrackEditors/BindingLifetimeTrackEditor.cpp b/Engine/Source/Editor/MovieSceneTools/Private/TrackEditors/BindingLifetimeTrackEditor.cpp
--- a/Engine/Source/Editor/MovieSceneTools/Private/TrackEditors/BindingLifetimeTrackEditor.cpp
+++ b/Engine/Source/Editor/MovieSceneTools/Private/TrackEditors/BindingLifetimeTrackEditor.cpp
@@ -163,6 +163,10 @@ void FBindingLifetimeTrackEditor::CreateNewSection(UMovieSceneTrack* Track, bool
 	}
 }
 
+FText FBindingLifetimeTrackEditor::GetDisplayName() const
+{
+	return LOCTEXT("BindingLifetimeTrackEditor_DisplayName", "Binding Lifetime");
+}
 
 UMovieSceneTrack* FBindingLifetimeTrackEditor::AddTrack(UMovieScene* FocusedMovieScene, const FGuid& ObjectHandle, TSubclassOf<UMovieSceneTrack> TrackClass, FName UniqueTypeName)
 {
diff --git a/Engine/Source/Editor/MovieSceneTools/Private/TrackEditors/CVarTrackEditor.cpp b/Engine/Source/Editor/MovieSceneTools/Private/TrackEditors/CVarTrackEditor.cpp
--- a/Engine/Source/Editor/MovieSceneTools/Private/TrackEditors/CVarTrackEditor.cpp
+++ b/Engine/Source/Editor/MovieSceneTools/Private/TrackEditors/CVarTrackEditor.cpp
@@ -28,6 +28,10 @@ TSharedRef<ISequencerTrackEditor> FCVarTrackEditor::CreateTrackEditor(TSharedRef
 	return MakeShareable(new FCVarTrackEditor(InSequencer));
 }
 
+FText FCVarTrackEditor::GetDisplayName() const
+{
+	return LOCTEXT("CVarTrackEditor_DisplayName", "CVar");
+}
 
 // ISequencerTrackEditor interface
 void FCVarTrackEditor::BuildAddTrackMenu(FMenuBuilder& MenuBuilder)
diff --git a/Engine/Source/Editor/MovieSceneTools/Private/TrackEditors/CVarTrackEditor.h b/Engine/Source/Editor/MovieSceneTools/Private/TrackEditors/CVarTrackEditor.h
--- a/Engine/Source/Editor/MovieSceneTools/Private/TrackEditors/CVarTrackEditor.h
+++ b/Engine/Source/Editor/MovieSceneTools/Private/TrackEditors/CVarTrackEditor.h
@@ -47,7 +47,7 @@ public:
 
 	// ISequencerTrackEditor interface
 
-
+	virtual FText GetDisplayName() const override;
 	virtual void BuildAddTrackMenu(FMenuBuilder& MenuBuilder) override;
 	virtual TSharedPtr<SWidget> BuildOutlinerEditWidget(const FGuid& ObjectBinding, UMovieSceneTrack* Track, const FBuildEditWidgetParams& Params) override;
 	virtual TSharedRef<ISequencerSection> MakeSectionInterface(UMovieSceneSection& SectionObject, UMovieSceneTrack& Track, FGuid ObjectBinding) override;
diff --git a/Engine/Source/Editor/MovieSceneTools/Private/TrackEditors/CameraCutTrackEditor.cpp b/Engine/Source/Editor/MovieSceneTools/Private/TrackEditors/CameraCutTrackEditor.cpp
--- a/Engine/Source/Editor/MovieSceneTools/Private/TrackEditors/CameraCutTrackEditor.cpp
+++ b/Engine/Source/Editor/MovieSceneTools/Private/TrackEditors/CameraCutTrackEditor.cpp
@@ -101,6 +101,11 @@ TSharedRef<ISequencerTrackEditor> FCameraCutTrackEditor::CreateTrackEditor(TShar
 /* ISequencerTrackEditor interface
  *****************************************************************************/
 
+FText FCameraCutTrackEditor::GetDisplayName() const
+{
+	return LOCTEXT("CameraCutTrackEditor_DisplayName", "Camera Cut");
+}
+
 void FCameraCutTrackEditor::BindCommands(TSharedRef<FUICommandList> SequencerCommandBindings)
 {
 	const FCameraCutTrackCommands& Commands = FCameraCutTrackCommands::Get();
diff --git a/Engine/Source/Editor/MovieSceneTools/Private/TrackEditors/CameraCutTrackEditor.h b/Engine/Source/Editor/MovieSceneTools/Private/TrackEditors/CameraCutTrackEditor.h
--- a/Engine/Source/Editor/MovieSceneTools/Private/TrackEditors/CameraCutTrackEditor.h
+++ b/Engine/Source/Editor/MovieSceneTools/Private/TrackEditors/CameraCutTrackEditor.h
@@ -53,6 +53,7 @@ public:
 
 	// ISequencerTrackEditor interface
 
+	virtual FText GetDisplayName() const override;
 	virtual void BindCommands(TSharedRef<FUICommandList> SequencerCommandBindings) override;
 	virtual void BuildAddTrackMenu(FMenuBuilder& MenuBuilder) override;
 	virtual void BuildTrackContextMenu(FMenuBuilder& MenuBuilder, UMovieSceneTrack* Track) override;
diff --git a/Engine/Source/Editor/MovieSceneTools/Private/TrackEditors/CameraShakeSourceShakeTrackEditor.cpp b/Engine/Source/Editor/MovieSceneTools/Private/TrackEditors/CameraShakeSourceShakeTrackEditor.cpp
--- a/Engine/Source/Editor/MovieSceneTools/Private/TrackEditors/CameraShakeSourceShakeTrackEditor.cpp
+++ b/Engine/Source/Editor/MovieSceneTools/Private/TrackEditors/CameraShakeSourceShakeTrackEditor.cpp
@@ -264,6 +264,11 @@ TSharedRef<ISequencerSection> FCameraShakeSourceShakeTrackEditor::MakeSectionInt
 	return MakeShareable(new FSequencerSection(SectionObject));
 }
 
+FText FCameraShakeSourceShakeTrackEditor::GetDisplayName() const
+{
+	return LOCTEXT("CameraShakeSourceShakeTrackEditor_DisplayName", "Camera Shake Source");
+}
+
 UMovieSceneTrack* FCameraShakeSourceShakeTrackEditor::AddTrack(UMovieScene* FocusedMovieScene, const FGuid& ObjectHandle, TSubclassOf<UMovieSceneTrack> TrackClass, FName UniqueTypeName)
 {
 	UMovieSceneTrack* NewTrack = FMovieSceneTrackEditor::AddTrack(FocusedMovieScene, ObjectHandle, TrackClass, UniqueTypeName);
diff --git a/Engine/Source/Editor/MovieSceneTools/Private/TrackEditors/CameraShakeSourceShakeTrackEditor.h b/Engine/Source/Editor/MovieSceneTools/Private/TrackEditors/CameraShakeSourceShakeTrackEditor.h
--- a/Engine/Source/Editor/MovieSceneTools/Private/TrackEditors/CameraShakeSourceShakeTrackEditor.h
+++ b/Engine/Source/Editor/MovieSceneTools/Private/TrackEditors/CameraShakeSourceShakeTrackEditor.h
@@ -35,6 +35,7 @@ public:
 	virtual ~FCameraShakeSourceShakeTrackEditor() {}
 
 	// ISequencerTrackEditor interface
+	virtual FText GetDisplayName() const override;
 	virtual UMovieSceneTrack* AddTrack(UMovieScene* FocusedMovieScene, const FGuid& ObjectHandle, TSubclassOf<class UMovieSceneTrack> TrackClass, FName UniqueTypeName) override;
 	virtual void BuildObjectBindingTrackMenu(FMenuBuilder& MenuBuilder, const TArray<FGuid>& ObjectBindings, const UClass* ObjectClass) override;
 	virtual TSharedRef<ISequencerSection> MakeSectionInterface(UMovieSceneSection& SectionObject, UMovieSceneTrack& Track, FGuid ObjectBinding) override;
diff --git a/Engine/Source/Editor/MovieSceneTools/Private/TrackEditors/CameraShakeTrackEditor.cpp b/Engine/Source/Editor/MovieSceneTools/Private/TrackEditors/CameraShakeTrackEditor.cpp
--- a/Engine/Source/Editor/MovieSceneTools/Private/TrackEditors/CameraShakeTrackEditor.cpp
+++ b/Engine/Source/Editor/MovieSceneTools/Private/TrackEditors/CameraShakeTrackEditor.cpp
@@ -132,6 +132,11 @@ bool FCameraShakeTrackEditor::HandleAssetAdded(UObject* Asset, const FGuid& Targ
 	return false;
 }
 
+FText FCameraShakeTrackEditor::GetDisplayName() const
+{
+	return LOCTEXT("CameraShakeTrackEditor_DisplayName", "Camera Shake");
+}
+
 void FCameraShakeTrackEditor::BuildObjectBindingTrackMenu(FMenuBuilder& MenuBuilder, const TArray<FGuid>& ObjectBindings, const UClass* ObjectClass)
 {
 	// only offer this track if we can find a camera component
diff --git a/Engine/Source/Editor/MovieSceneTools/Private/TrackEditors/CameraShakeTrackEditor.h b/Engine/Source/Editor/MovieSceneTools/Private/TrackEditors/CameraShakeTrackEditor.h
--- a/Engine/Source/Editor/MovieSceneTools/Private/TrackEditors/CameraShakeTrackEditor.h
+++ b/Engine/Source/Editor/MovieSceneTools/Private/TrackEditors/CameraShakeTrackEditor.h
@@ -51,6 +51,7 @@ public:
 public:
 
 	// ISequencerTrackEditor interface
+	virtual FText GetDisplayName() const override;
 	virtual void BuildObjectBindingTrackMenu(FMenuBuilder& MenuBuilder, const TArray<FGuid>& ObjectBindings, const UClass* ObjectClass) override;
 	virtual bool HandleAssetAdded(UObject* Asset, const FGuid& TargetObjectGuid) override;
 	virtual TSharedRef<ISequencerSection> MakeSectionInterface(UMovieSceneSection& SectionObject, UMovieSceneTrack& Track, FGuid ObjectBinding) override;
diff --git a/Engine/Source/Editor/MovieSceneTools/Private/TrackEditors/CinematicShotTrackEditor.cpp b/Engine/Source/Editor/MovieSceneTools/Private/TrackEditors/CinematicShotTrackEditor.cpp
--- a/Engine/Source/Editor/MovieSceneTools/Private/TrackEditors/CinematicShotTrackEditor.cpp
+++ b/Engine/Source/Editor/MovieSceneTools/Private/TrackEditors/CinematicShotTrackEditor.cpp
@@ -66,6 +66,11 @@ TWeakObjectPtr<AActor> FCinematicShotTrackEditor::GetCinematicShotCamera() const
 /* ISequencerTrackEditor interface
  *****************************************************************************/
 
+FText FCinematicShotTrackEditor::GetDisplayName() const
+{
+	return LOCTEXT("ShotTrackEditor_DisplayName", "Shot");
+}
+
 TSharedPtr<SWidget> FCinematicShotTrackEditor::BuildOutlinerColumnWidget(const FBuildColumnWidgetParams& Params, const FName& ColumnName)
 {
 	using namespace UE::Sequencer;
diff --git a/Engine/Source/Editor/MovieSceneTools/Private/TrackEditors/CustomPrimitiveDataTrackEditor.cpp b/Engine/Source/Editor/MovieSceneTools/Private/TrackEditors/CustomPrimitiveDataTrackEditor.cpp
--- a/Engine/Source/Editor/MovieSceneTools/Private/TrackEditors/CustomPrimitiveDataTrackEditor.cpp
+++ b/Engine/Source/Editor/MovieSceneTools/Private/TrackEditors/CustomPrimitiveDataTrackEditor.cpp
@@ -37,6 +37,11 @@ TSharedRef<ISequencerSection> FCustomPrimitiveDataTrackEditor::MakeSectionInterf
 	return MakeShareable(new FParameterSection(*CustomPrimitiveDataSection, GetSequencer()));
 }
 
+FText FCustomPrimitiveDataTrackEditor::GetDisplayName() const
+{
+	return LOCTEXT("CustomPrimitiveDataTrackEditor_DisplayName", "Custom Primitive Data");
+}
+
 TSharedPtr<SWidget> FCustomPrimitiveDataTrackEditor::BuildOutlinerEditWidget(const FGuid& ObjectBinding, UMovieSceneTrack* Track, const FBuildEditWidgetParams& Params)
 {
 	UMovieSceneCustomPrimitiveDataTrack* CPDTrack = Cast<UMovieSceneCustomPrimitiveDataTrack>(Track);
diff --git a/Engine/Source/Editor/MovieSceneTools/Private/TrackEditors/DataLayerTrackEditor.cpp b/Engine/Source/Editor/MovieSceneTools/Private/TrackEditors/DataLayerTrackEditor.cpp
--- a/Engine/Source/Editor/MovieSceneTools/Private/TrackEditors/DataLayerTrackEditor.cpp
+++ b/Engine/Source/Editor/MovieSceneTools/Private/TrackEditors/DataLayerTrackEditor.cpp
@@ -275,6 +275,11 @@ const FSlateBrush* FDataLayerTrackEditor::GetIconBrush() const
 	return FAppStyle::Get().GetBrush("Sequencer.Tracks.DataLayer");
 }
 
+FText FDataLayerTrackEditor::GetDisplayName() const
+{
+	return LOCTEXT("DataLayerTrackEditor_DisplayName", "Data Layer");
+}
+
 TSharedRef<ISequencerSection> FDataLayerTrackEditor::MakeSectionInterface(UMovieSceneSection& SectionObject, UMovieSceneTrack& Track, FGuid ObjectBinding)
 {
 	UMovieSceneDataLayerSection* DataLayerSection = Cast<UMovieSceneDataLayerSection>(&SectionObject);
diff --git a/Engine/Source/Editor/MovieSceneTools/Private/TrackEditors/DataLayerTrackEditor.h b/Engine/Source/Editor/MovieSceneTools/Private/TrackEditors/DataLayerTrackEditor.h
--- a/Engine/Source/Editor/MovieSceneTools/Private/TrackEditors/DataLayerTrackEditor.h
+++ b/Engine/Source/Editor/MovieSceneTools/Private/TrackEditors/DataLayerTrackEditor.h
@@ -25,6 +25,7 @@ public:
 public:
 
 	// ISequencerTrackEditor interface
+	virtual FText GetDisplayName() const override;
 	virtual TSharedRef<ISequencerSection> MakeSectionInterface(UMovieSceneSection& SectionObject, UMovieSceneTrack& Track, FGuid ObjectBinding) override;
 	virtual bool SupportsSequence(UMovieSceneSequence* InSequence) const override;
 	virtual bool SupportsType(TSubclassOf<UMovieSceneTrack> Type) const override;
diff --git a/Engine/Source/Editor/MovieSceneTools/Private/TrackEditors/EventTrackEditor.cpp b/Engine/Source/Editor/MovieSceneTools/Private/TrackEditors/EventTrackEditor.cpp
--- a/Engine/Source/Editor/MovieSceneTools/Private/TrackEditors/EventTrackEditor.cpp
+++ b/Engine/Source/Editor/MovieSceneTools/Private/TrackEditors/EventTrackEditor.cpp
@@ -88,6 +88,11 @@ void FEventTrackEditor::AddEventSubMenu(FMenuBuilder& MenuBuilder, TArray<FGuid>
 	);
 }
 
+FText FEventTrackEditor::GetDisplayName() const
+{
+	return LOCTEXT("EventTrackEditor_DisplayName", "Event");
+}
+
 void FEventTrackEditor::BuildAddTrackMenu(FMenuBuilder& MenuBuilder)
 {
 	UMovieSceneSequence*       RootMovieSceneSequence = GetSequencer()->GetRootMovieSceneSequence();
diff --git a/Engine/Source/Editor/MovieSceneTools/Private/TrackEditors/EventTrackEditor.h b/Engine/Source/Editor/MovieSceneTools/Private/TrackEditors/EventTrackEditor.h
--- a/Engine/Source/Editor/MovieSceneTools/Private/TrackEditors/EventTrackEditor.h
+++ b/Engine/Source/Editor/MovieSceneTools/Private/TrackEditors/EventTrackEditor.h
@@ -42,6 +42,7 @@ public:
 
 	// ISequencerTrackEditor interface
 
+	virtual FText GetDisplayName() const override;
 	virtual void BuildAddTrackMenu(FMenuBuilder& MenuBuilder) override;
 	virtual void BuildObjectBindingTrackMenu(FMenuBuilder& MenuBuilder, const TArray<FGuid>& ObjectBindings, const UClass* ObjectClass) override;
 	virtual bool SupportsType( TSubclassOf<UMovieSceneTrack> Type ) const override;
diff --git a/Engine/Source/Editor/MovieSceneTools/Private/TrackEditors/FadeTrackEditor.cpp b/Engine/Source/Editor/MovieSceneTools/Private/TrackEditors/FadeTrackEditor.cpp
--- a/Engine/Source/Editor/MovieSceneTools/Private/TrackEditors/FadeTrackEditor.cpp
+++ b/Engine/Source/Editor/MovieSceneTools/Private/TrackEditors/FadeTrackEditor.cpp
@@ -126,6 +126,11 @@ FFadeTrackEditor::FFadeTrackEditor(TSharedRef<ISequencer> InSequencer)
 /* ISequencerTrackEditor interface
  *****************************************************************************/
 
+FText FFadeTrackEditor::GetDisplayName() const
+{
+	return LOCTEXT("FadeTrackEditor_DisplayName", "Fade");
+}
+
 TSharedRef<ISequencerSection> FFadeTrackEditor::MakeSectionInterface(UMovieSceneSection& SectionObject, UMovieSceneTrack& Track, FGuid ObjectBinding)
 {
 	return MakeShareable(new FFadeSection(SectionObject));
diff --git a/Engine/Source/Editor/MovieSceneTools/Private/TrackEditors/FadeTrackEditor.h b/Engine/Source/Editor/MovieSceneTools/Private/TrackEditors/FadeTrackEditor.h
--- a/Engine/Source/Editor/MovieSceneTools/Private/TrackEditors/FadeTrackEditor.h
+++ b/Engine/Source/Editor/MovieSceneTools/Private/TrackEditors/FadeTrackEditor.h
@@ -45,6 +45,7 @@ public:
 
 	// ISequencerTrackEditor interface
 
+	virtual FText GetDisplayName() const override;
 	virtual TSharedRef<ISequencerSection> MakeSectionInterface(UMovieSceneSection& SectionObject, UMovieSceneTrack& Track, FGuid ObjectBinding) override;
 	virtual void BuildAddTrackMenu(FMenuBuilder& MenuBuilder) override;
 	virtual bool SupportsSequence(UMovieSceneSequence* InSequence) const override;
diff --git a/Engine/Source/Editor/MovieSceneTools/Private/TrackEditors/LevelVisibilityTrackEditor.cpp b/Engine/Source/Editor/MovieSceneTools/Private/TrackEditors/LevelVisibilityTrackEditor.cpp
--- a/Engine/Source/Editor/MovieSceneTools/Private/TrackEditors/LevelVisibilityTrackEditor.cpp
+++ b/Engine/Source/Editor/MovieSceneTools/Private/TrackEditors/LevelVisibilityTrackEditor.cpp
@@ -40,6 +40,11 @@ const FSlateBrush* FLevelVisibilityTrackEditor::GetIconBrush() const
 	return FAppStyle::GetBrush("Sequencer.Tracks.LevelVisibility");
 }
 
+FText FLevelVisibilityTrackEditor::GetDisplayName() const
+{
+	return LOCTEXT("LevelVisibilityTrackEditor_DisplayName", "Level Visibility");
+}
+
 TSharedRef<ISequencerSection> FLevelVisibilityTrackEditor::MakeSectionInterface( UMovieSceneSection& SectionObject, UMovieSceneTrack& Track, FGuid ObjectBinding )
 {
 	UMovieSceneLevelVisibilitySection* LevelVisibilitySection = Cast<UMovieSceneLevelVisibilitySection>(&SectionObject);
diff --git a/Engine/Source/Editor/MovieSceneTools/Private/TrackEditors/LevelVisibilityTrackEditor.h b/Engine/Source/Editor/MovieSceneTools/Private/TrackEditors/LevelVisibilityTrackEditor.h
--- a/Engine/Source/Editor/MovieSceneTools/Private/TrackEditors/LevelVisibilityTrackEditor.h
+++ b/Engine/Source/Editor/MovieSceneTools/Private/TrackEditors/LevelVisibilityTrackEditor.h
@@ -46,6 +46,7 @@ public:
 
 	// ISequencerTrackEditor interface
 
+	virtual FText GetDisplayName() const override;
 	virtual TSharedRef<ISequencerSection> MakeSectionInterface( UMovieSceneSection& SectionObject, UMovieSceneTrack& Track, FGuid ObjectBinding ) override;
 	virtual bool SupportsSequence(UMovieSceneSequence* InSequence) const override;
 	virtual bool SupportsType( TSubclassOf<UMovieSceneTrack> Type ) const override;
diff --git a/Engine/Source/Editor/MovieSceneTools/Private/TrackEditors/MaterialParameterCollectionTrackEditor.cpp b/Engine/Source/Editor/MovieSceneTools/Private/TrackEditors/MaterialParameterCollectionTrackEditor.cpp
--- a/Engine/Source/Editor/MovieSceneTools/Private/TrackEditors/MaterialParameterCollectionTrackEditor.cpp
+++ b/Engine/Source/Editor/MovieSceneTools/Private/TrackEditors/MaterialParameterCollectionTrackEditor.cpp
@@ -260,6 +260,11 @@ const FSlateBrush* FMaterialParameterCollectionTrackEditor::GetIconBrush() const
 	return FSlateIconFinder::FindIconForClass(UMaterialParameterCollection::StaticClass()).GetIcon();
 }
 
+FText FMaterialParameterCollectionTrackEditor::GetDisplayName() const
+{
+	return LOCTEXT("MaterialParameterCollectionTrackEditor_DisplayName", "Material Parameter Collection");
+}
+
 TSharedPtr<SWidget> FMaterialParameterCollectionTrackEditor::BuildOutlinerEditWidget(const FGuid& ObjectBinding, UMovieSceneTrack* Track, const FBuildEditWidgetParams& Params)
 {
 	UMovieSceneMaterialParameterCollectionTrack* MPCTrack = Cast<UMovieSceneMaterialParameterCollectionTrack>(Track);
diff --git a/Engine/Source/Editor/MovieSceneTools/Private/TrackEditors/MaterialParameterCollectionTrackEditor.h b/Engine/Source/Editor/MovieSceneTools/Private/TrackEditors/MaterialParameterCollectionTrackEditor.h
--- a/Engine/Source/Editor/MovieSceneTools/Private/TrackEditors/MaterialParameterCollectionTrackEditor.h
+++ b/Engine/Source/Editor/MovieSceneTools/Private/TrackEditors/MaterialParameterCollectionTrackEditor.h
@@ -42,6 +42,7 @@ public:
 public:
 
 	//~ ISequencerTrackEditor interface
+	virtual FText GetDisplayName() const override;
 	virtual TSharedPtr<SWidget> BuildOutlinerEditWidget(const FGuid& ObjectBinding, UMovieSceneTrack* Track, const FBuildEditWidgetParams& Params) override;
 	virtual TSharedRef<ISequencerSection> MakeSectionInterface(UMovieSceneSection& SectionObject, UMovieSceneTrack& Track, FGuid ObjectBinding) override;
 	virtual bool SupportsType(TSubclassOf<UMovieSceneTrack> Type) const override;
diff --git a/Engine/Source/Editor/MovieSceneTools/Private/TrackEditors/MaterialTrackEditor.cpp b/Engine/Source/Editor/MovieSceneTools/Private/TrackEditors/MaterialTrackEditor.cpp
--- a/Engine/Source/Editor/MovieSceneTools/Private/TrackEditors/MaterialTrackEditor.cpp
+++ b/Engine/Source/Editor/MovieSceneTools/Private/TrackEditors/MaterialTrackEditor.cpp
@@ -45,6 +45,10 @@ TSharedRef<ISequencerSection> FMaterialTrackEditor::MakeSectionInterface( UMovie
 	}
 }
 
+FText FMaterialTrackEditor::GetDisplayName() const
+{
+	return LOCTEXT("MaterialTrackEditor_DisplayName", "Material");
+}
 
 TSharedPtr<SWidget> FMaterialTrackEditor::BuildOutlinerEditWidget( const FGuid& ObjectBinding, UMovieSceneTrack* Track, const FBuildEditWidgetParams& Params )
 {
diff --git a/Engine/Source/Editor/MovieSceneTools/Private/TrackEditors/ObjectPropertyTrackEditor.cpp b/Engine/Source/Editor/MovieSceneTools/Private/TrackEditors/ObjectPropertyTrackEditor.cpp
--- a/Engine/Source/Editor/MovieSceneTools/Private/TrackEditors/ObjectPropertyTrackEditor.cpp
+++ b/Engine/Source/Editor/MovieSceneTools/Private/TrackEditors/ObjectPropertyTrackEditor.cpp
@@ -35,6 +35,11 @@ TSharedRef<ISequencerTrackEditor> FObjectPropertyTrackEditor::CreateTrackEditor(
 	return MakeShared<FObjectPropertyTrackEditor>(OwningSequencer);
 }
 
+FText FObjectPropertyTrackEditor::GetDisplayName() const
+{
+	return LOCTEXT("ObjectPropertyTrackEditor_DisplayName", "Object Property");
+}
+
 void FObjectPropertyTrackEditor::GenerateKeysFromPropertyChanged(const FPropertyChangedParams& PropertyChangedParams, UMovieSceneSection* SectionToKey, FGeneratedTrackKeys& OutGeneratedKeys)
 {
 	UObject* Existing = PropertyChangedParams.GetPropertyValue<UObject*>();
diff --git a/Engine/Source/Editor/MovieSceneTools/Private/TrackEditors/ObjectPropertyTrackEditor.h b/Engine/Source/Editor/MovieSceneTools/Private/TrackEditors/ObjectPropertyTrackEditor.h
--- a/Engine/Source/Editor/MovieSceneTools/Private/TrackEditors/ObjectPropertyTrackEditor.h
+++ b/Engine/Source/Editor/MovieSceneTools/Private/TrackEditors/ObjectPropertyTrackEditor.h
@@ -40,6 +40,7 @@ public:
 public:
 
 	//~ ISequencerTrackEditor interface
+	virtual FText GetDisplayName() const override;
 	virtual void GenerateKeysFromPropertyChanged(const FPropertyChangedParams& PropertyChangedParams, UMovieSceneSection* SectionToKey, FGeneratedTrackKeys& OutGeneratedKeys) override;
 	virtual void InitializeNewTrack(UMovieSceneObjectPropertyTrack* NewTrack, FPropertyChangedParams PropertyChangedParams) override;
 };
diff --git a/Engine/Source/Editor/MovieSceneTools/Private/TrackEditors/ParticleParameterTrackEditor.cpp b/Engine/Source/Editor/MovieSceneTools/Private/TrackEditors/ParticleParameterTrackEditor.cpp
--- a/Engine/Source/Editor/MovieSceneTools/Private/TrackEditors/ParticleParameterTrackEditor.cpp
+++ b/Engine/Source/Editor/MovieSceneTools/Private/TrackEditors/ParticleParameterTrackEditor.cpp
@@ -34,6 +34,10 @@ TSharedRef<ISequencerSection> FParticleParameterTrackEditor::MakeSectionInterfac
 	return MakeShareable(new FParameterSection( *ParameterSection, GetSequencer() ));
 }
 
+FText FParticleParameterTrackEditor::GetDisplayName() const
+{
+	return LOCTEXT("ParticleParameterTrackEditor_DisplayName", "Particle Parameter");
+}
 
 TSharedPtr<SWidget> FParticleParameterTrackEditor::BuildOutlinerEditWidget( const FGuid& ObjectBinding, UMovieSceneTrack* Track, const FBuildEditWidgetParams& Params )
 {
diff --git a/Engine/Source/Editor/MovieSceneTools/Private/TrackEditors/ParticleParameterTrackEditor.h b/Engine/Source/Editor/MovieSceneTools/Private/TrackEditors/ParticleParameterTrackEditor.h
--- a/Engine/Source/Editor/MovieSceneTools/Private/TrackEditors/ParticleParameterTrackEditor.h
+++ b/Engine/Source/Editor/MovieSceneTools/Private/TrackEditors/ParticleParameterTrackEditor.h
@@ -42,6 +42,7 @@ public:
 
 	// ISequencerTrackEditor interface
 
+	virtual FText GetDisplayName() const override;
 	virtual TSharedPtr<SWidget> BuildOutlinerEditWidget( const FGuid& ObjectBinding, UMovieSceneTrack* Track, const FBuildEditWidgetParams& Params ) override;
 	virtual TSharedRef<ISequencerSection> MakeSectionInterface( UMovieSceneSection& SectionObject, UMovieSceneTrack& Track, FGuid ObjectBinding ) override;
 	virtual bool SupportsType( TSubclassOf<UMovieSceneTrack> Type ) const override;
diff --git a/Engine/Source/Editor/MovieSceneTools/Private/TrackEditors/ParticleTrackEditor.cpp b/Engine/Source/Editor/MovieSceneTools/Private/TrackEditors/ParticleTrackEditor.cpp
--- a/Engine/Source/Editor/MovieSceneTools/Private/TrackEditors/ParticleTrackEditor.cpp
+++ b/Engine/Source/Editor/MovieSceneTools/Private/TrackEditors/ParticleTrackEditor.cpp
@@ -235,6 +235,10 @@ TSharedRef<ISequencerSection> FParticleTrackEditor::MakeSectionInterface( UMovie
 	return MakeShareable( new FParticleSection( SectionObject, OwningSequencer.ToSharedRef() ) );
 }
 
+FText FParticleTrackEditor::GetDisplayName() const
+{
+	return LOCTEXT("ParticleTrackEditor_DisplayName", "Particle");
+}
 
 void FParticleTrackEditor::BuildObjectBindingTrackMenu(FMenuBuilder& MenuBuilder, const TArray<FGuid>& ObjectBindings, const UClass* ObjectClass)
 {
diff --git a/Engine/Source/Editor/MovieSceneTools/Private/TrackEditors/ParticleTrackEditor.h b/Engine/Source/Editor/MovieSceneTools/Private/TrackEditors/ParticleTrackEditor.h
--- a/Engine/Source/Editor/MovieSceneTools/Private/TrackEditors/ParticleTrackEditor.h
+++ b/Engine/Source/Editor/MovieSceneTools/Private/TrackEditors/ParticleTrackEditor.h
@@ -49,6 +49,7 @@ public:
 
 	// ISequencerTrackEditor interface
 
+	virtual FText GetDisplayName() const override;
 	virtual void BuildObjectBindingTrackMenu(FMenuBuilder& MenuBuilder, const TArray<FGuid>& ObjectBindings, const UClass* ObjectClass) override;
 	virtual TSharedRef<ISequencerSection> MakeSectionInterface( UMovieSceneSection& SectionObject, UMovieSceneTrack& Track, FGuid ObjectBinding ) override;
 	virtual bool SupportsType( TSubclassOf<UMovieSceneTrack> Type ) const override;
diff --git a/Engine/Source/Editor/MovieSceneTools/Private/TrackEditors/PrimitiveMaterialTrackEditor.cpp b/Engine/Source/Editor/MovieSceneTools/Private/TrackEditors/PrimitiveMaterialTrackEditor.cpp
--- a/Engine/Source/Editor/MovieSceneTools/Private/TrackEditors/PrimitiveMaterialTrackEditor.cpp
+++ b/Engine/Source/Editor/MovieSceneTools/Private/TrackEditors/PrimitiveMaterialTrackEditor.cpp
@@ -24,6 +24,11 @@ TSharedRef<ISequencerTrackEditor> FPrimitiveMaterialTrackEditor::CreateTrackEdit
 	return MakeShared<FPrimitiveMaterialTrackEditor>(OwningSequencer);
 }
 
+FText FPrimitiveMaterialTrackEditor::GetDisplayName() const
+{
+	return LOCTEXT("PrimitiveMaterialTrackEditor_DisplayName", "Primitive Material");
+}
+
 void FPrimitiveMaterialTrackEditor::ExtendObjectBindingTrackMenu(TSharedRef<FExtender> Extender, const TArray<FGuid>& ObjectBindings, const UClass* ObjectClass)
 {
 	if (ObjectClass->IsChildOf(UPrimitiveComponent::StaticClass()))
diff --git a/Engine/Source/Editor/MovieSceneTools/Private/TrackEditors/PrimitiveMaterialTrackEditor.h b/Engine/Source/Editor/MovieSceneTools/Private/TrackEditors/PrimitiveMaterialTrackEditor.h
--- a/Engine/Source/Editor/MovieSceneTools/Private/TrackEditors/PrimitiveMaterialTrackEditor.h
+++ b/Engine/Source/Editor/MovieSceneTools/Private/TrackEditors/PrimitiveMaterialTrackEditor.h
@@ -20,6 +20,8 @@ public:
 public:
 
 	//~ ISequencerTrackEditor interface
+	virtual FText GetDisplayName() const override;
+
 	virtual void ExtendObjectBindingTrackMenu(TSharedRef<FExtender> Extender, const TArray<FGuid>& ObjectBindings, const UClass* ObjectClass) override;
 	virtual void BuildTrackContextMenu(FMenuBuilder& MenuBuilder, UMovieSceneTrack* Track) override;
 
diff --git a/Engine/Source/Editor/MovieSceneTools/Private/TrackEditors/PropertyTrackEditors/ActorReferencePropertyTrackEditor.cpp b/Engine/Source/Editor/MovieSceneTools/Private/TrackEditors/PropertyTrackEditors/ActorReferencePropertyTrackEditor.cpp
--- a/Engine/Source/Editor/MovieSceneTools/Private/TrackEditors/PropertyTrackEditors/ActorReferencePropertyTrackEditor.cpp
+++ b/Engine/Source/Editor/MovieSceneTools/Private/TrackEditors/PropertyTrackEditors/ActorReferencePropertyTrackEditor.cpp
@@ -17,12 +17,18 @@ class ISequencerTrackEditor;
 class UMovieSceneSection;
 class UObject;
 
+#define LOCTEXT_NAMESPACE "ActorReferencePropertyTrackEditor"
 
 TSharedRef<ISequencerTrackEditor> FActorReferencePropertyTrackEditor::CreateTrackEditor( TSharedRef<ISequencer> OwningSequencer )
 {
 	return MakeShareable(new FActorReferencePropertyTrackEditor(OwningSequencer));
 }
 
+FText FActorReferencePropertyTrackEditor::GetDisplayName() const
+{
+	return LOCTEXT("ActorReferencePropertyTrackEditor_DisplayName", "Actor Reference Property");
+}
+
 void FActorReferencePropertyTrackEditor::GenerateKeysFromPropertyChanged( const FPropertyChangedParams& PropertyChangedParams, UMovieSceneSection* SectionToKey, FGeneratedTrackKeys& OutGeneratedKeys )
 {
 	// Care is taken here to ensure that GetPropertyValue is templated on UObject* which causes it to use the correct instantiation of GetPropertyValueImpl
@@ -53,3 +59,5 @@ void FActorReferencePropertyTrackEditor::GenerateKeysFromPropertyChanged( const
 		}
 	}
 }
+
+#undef LOCTEXT_NAMESPACE
diff --git a/Engine/Source/Editor/MovieSceneTools/Private/TrackEditors/PropertyTrackEditors/ActorReferencePropertyTrackEditor.h b/Engine/Source/Editor/MovieSceneTools/Private/TrackEditors/PropertyTrackEditors/ActorReferencePropertyTrackEditor.h
--- a/Engine/Source/Editor/MovieSceneTools/Private/TrackEditors/PropertyTrackEditors/ActorReferencePropertyTrackEditor.h
+++ b/Engine/Source/Editor/MovieSceneTools/Private/TrackEditors/PropertyTrackEditors/ActorReferencePropertyTrackEditor.h
@@ -63,5 +63,6 @@ public:
 protected:
 
 	//~ FPropertyTrackEditor interface
+	virtual FText GetDisplayName() const override;
 	virtual void GenerateKeysFromPropertyChanged(const FPropertyChangedParams& PropertyChangedParams, UMovieSceneSection* SectionToKey, FGeneratedTrackKeys& OutGeneratedKeys) override;
 };
diff --git a/Engine/Source/Editor/MovieSceneTools/Private/TrackEditors/PropertyTrackEditors/BoolPropertyTrackEditor.cpp b/Engine/Source/Editor/MovieSceneTools/Private/TrackEditors/PropertyTrackEditors/BoolPropertyTrackEditor.cpp
--- a/Engine/Source/Editor/MovieSceneTools/Private/TrackEditors/PropertyTrackEditors/BoolPropertyTrackEditor.cpp
+++ b/Engine/Source/Editor/MovieSceneTools/Private/TrackEditors/PropertyTrackEditors/BoolPropertyTrackEditor.cpp
@@ -14,21 +14,27 @@ class ISequencerTrackEditor;
 class UMovieSceneSection;
 class UMovieSceneTrack;
 
+#define LOCTEXT_NAMESPACE "BoolPropertyTrackEditor"
 
 TSharedRef<ISequencerTrackEditor> FBoolPropertyTrackEditor::CreateTrackEditor(TSharedRef<ISequencer> OwningSequencer)
 {
 	return MakeShareable(new FBoolPropertyTrackEditor(OwningSequencer));
 }
 
+FText FBoolPropertyTrackEditor::GetDisplayName() const
+{
+	return LOCTEXT("BoolPropertyTrackEditor_DisplayName", "Bool Property");
+}
 
 TSharedRef<ISequencerSection> FBoolPropertyTrackEditor::MakeSectionInterface(UMovieSceneSection& SectionObject, UMovieSceneTrack& Track, FGuid ObjectBinding)
 {
 	return MakeShared<FBoolPropertySection>(SectionObject);
 }
 
-
 void FBoolPropertyTrackEditor::GenerateKeysFromPropertyChanged( const FPropertyChangedParams& PropertyChangedParams, UMovieSceneSection* SectionToKey, FGeneratedTrackKeys& OutGeneratedKeys )
 {
 	const bool KeyedValue = PropertyChangedParams.GetPropertyValue<bool>();
 	OutGeneratedKeys.Add(FMovieSceneChannelValueSetter::Create<FMovieSceneBoolChannel>(0, KeyedValue, true));
 }
+
+#undef LOCTEXT_NAMESPACE
diff --git a/Engine/Source/Editor/MovieSceneTools/Private/TrackEditors/PropertyTrackEditors/BytePropertyTrackEditor.cpp b/Engine/Source/Editor/MovieSceneTools/Private/TrackEditors/PropertyTrackEditors/BytePropertyTrackEditor.cpp
--- a/Engine/Source/Editor/MovieSceneTools/Private/TrackEditors/PropertyTrackEditors/BytePropertyTrackEditor.cpp
+++ b/Engine/Source/Editor/MovieSceneTools/Private/TrackEditors/PropertyTrackEditors/BytePropertyTrackEditor.cpp
@@ -22,6 +22,7 @@ class UMovieScene;
 class UMovieSceneSection;
 struct FGuid;
 
+#define LOCTEXT_NAMESPACE "BytePropertyTrackEditor"
 
 TSharedRef<ISequencerTrackEditor> FBytePropertyTrackEditor::CreateTrackEditor( TSharedRef<ISequencer> OwningSequencer )
 {
@@ -75,6 +76,10 @@ UEnum* GetEnumForByteTrack(TSharedPtr<ISequencer> Sequencer, const FGuid& OwnerO
 	return TrackEnum;
 }
 
+FText FBytePropertyTrackEditor::GetDisplayName() const
+{
+	return LOCTEXT("BytePropertyTrackEditor_DisplayName", "Byte Property");
+}
 
 void FBytePropertyTrackEditor::GenerateKeysFromPropertyChanged( const FPropertyChangedParams& PropertyChangedParams, UMovieSceneSection* SectionToKey, FGeneratedTrackKeys& OutGeneratedKeys )
 {
@@ -96,3 +101,5 @@ void FBytePropertyTrackEditor::InitializeNewTrack(UMovieSceneByteTrack* NewTrack
 	}
 
 }
+
+#undef LOCTEXT_NAMESPACE
diff --git a/Engine/Source/Editor/MovieSceneTools/Private/TrackEditors/PropertyTrackEditors/BytePropertyTrackEditor.h b/Engine/Source/Editor/MovieSceneTools/Private/TrackEditors/PropertyTrackEditors/BytePropertyTrackEditor.h
--- a/Engine/Source/Editor/MovieSceneTools/Private/TrackEditors/PropertyTrackEditors/BytePropertyTrackEditor.h
+++ b/Engine/Source/Editor/MovieSceneTools/Private/TrackEditors/PropertyTrackEditors/BytePropertyTrackEditor.h
@@ -65,6 +65,7 @@ protected:
 
 	//~ FPropertyTrackEditor interface
 
+	virtual FText GetDisplayName() const override;
 	virtual void GenerateKeysFromPropertyChanged(const FPropertyChangedParams& PropertyChangedParams, UMovieSceneSection* SectionToKey, FGeneratedTrackKeys& OutGeneratedKeys) override;
 	virtual void InitializeNewTrack(UMovieSceneByteTrack* NewTrack, FPropertyChangedParams PropertyChangedParams) override;
 
diff --git a/Engine/Source/Editor/MovieSceneTools/Private/TrackEditors/PropertyTrackEditors/ColorPropertyTrackEditor.cpp b/Engine/Source/Editor/MovieSceneTools/Private/TrackEditors/PropertyTrackEditors/ColorPropertyTrackEditor.cpp
--- a/Engine/Source/Editor/MovieSceneTools/Private/TrackEditors/PropertyTrackEditors/ColorPropertyTrackEditor.cpp
+++ b/Engine/Source/Editor/MovieSceneTools/Private/TrackEditors/PropertyTrackEditors/ColorPropertyTrackEditor.cpp
@@ -7,6 +7,8 @@
 #include "Evaluation/MovieScenePropertyTemplate.h"
 #include "EntitySystem/Interrogation/MovieSceneInterrogationLinker.h"
 
+#define LOCTEXT_NAMESPACE "ColorPropertyTrackEditor"
+
 FName FColorPropertyTrackEditor::RedName( "R" );
 FName FColorPropertyTrackEditor::GreenName( "G" );
 FName FColorPropertyTrackEditor::BlueName( "B" );
@@ -18,6 +20,10 @@ TSharedRef<ISequencerTrackEditor> FColorPropertyTrackEditor::CreateTrackEditor(T
 	return MakeShareable(new FColorPropertyTrackEditor(InSequencer));
 }
 
+FText FColorPropertyTrackEditor::GetDisplayName() const
+{
+	return LOCTEXT("ColorPropertyTrackEditor_DisplayName", "Color Property");
+}
 
 TSharedRef<ISequencerSection> FColorPropertyTrackEditor::MakeSectionInterface(UMovieSceneSection& SectionObject, UMovieSceneTrack& Track, FGuid ObjectBinding)
 {
@@ -110,3 +116,5 @@ bool FColorPropertyTrackEditor::ModifyGeneratedKeysByCurrentAndWeight(UObject *O
 
 	return false;
 }
+
+#undef LOCTEXT_NAMESPACE
diff --git a/Engine/Source/Editor/MovieSceneTools/Private/TrackEditors/PropertyTrackEditors/ColorPropertyTrackEditor.h b/Engine/Source/Editor/MovieSceneTools/Private/TrackEditors/PropertyTrackEditors/ColorPropertyTrackEditor.h
--- a/Engine/Source/Editor/MovieSceneTools/Private/TrackEditors/PropertyTrackEditors/ColorPropertyTrackEditor.h
+++ b/Engine/Source/Editor/MovieSceneTools/Private/TrackEditors/PropertyTrackEditors/ColorPropertyTrackEditor.h
@@ -51,6 +51,7 @@ public:
 
 	//~ ISequencerTrackEditor interface
 
+	virtual FText GetDisplayName() const override;
 	virtual TSharedRef<ISequencerSection> MakeSectionInterface(UMovieSceneSection& SectionObject, UMovieSceneTrack& Track, FGuid ObjectBinding) override;
 
 protected:
diff --git a/Engine/Source/Editor/MovieSceneTools/Private/TrackEditors/PropertyTrackEditors/DoublePropertyTrackEditor.cpp b/Engine/Source/Editor/MovieSceneTools/Private/TrackEditors/PropertyTrackEditors/DoublePropertyTrackEditor.cpp
--- a/Engine/Source/Editor/MovieSceneTools/Private/TrackEditors/PropertyTrackEditors/DoublePropertyTrackEditor.cpp
+++ b/Engine/Source/Editor/MovieSceneTools/Private/TrackEditors/PropertyTrackEditors/DoublePropertyTrackEditor.cpp
@@ -9,11 +9,17 @@
 #include "Systems/MovieScenePropertyInstantiator.h"
 #include "MovieSceneTracksComponentTypes.h"
 
+#define LOCTEXT_NAMESPACE "DoublePropertyTrackEditor"
+
 TSharedRef<ISequencerTrackEditor> FDoublePropertyTrackEditor::CreateTrackEditor( TSharedRef<ISequencer> OwningSequencer )
 {
 	return MakeShareable(new FDoublePropertyTrackEditor(OwningSequencer));
 }
 
+FText FDoublePropertyTrackEditor::GetDisplayName() const
+{
+	return LOCTEXT("DoublePropertyTrackEditor_DisplayName", "Double Property");
+}
 
 void FDoublePropertyTrackEditor::GenerateKeysFromPropertyChanged( const FPropertyChangedParams& PropertyChangedParams, UMovieSceneSection* SectionToKey, FGeneratedTrackKeys& OutGeneratedKeys )
 {
@@ -51,3 +57,5 @@ double FDoublePropertyTrackEditor::RecomposeDouble(double InCurrentValue, UObjec
 
 	return Recomposed;
 }
+
+#undef LOCTEXT_NAMESPACE
diff --git a/Engine/Source/Editor/MovieSceneTools/Private/TrackEditors/PropertyTrackEditors/DoublePropertyTrackEditor.h b/Engine/Source/Editor/MovieSceneTools/Private/TrackEditors/PropertyTrackEditors/DoublePropertyTrackEditor.h
--- a/Engine/Source/Editor/MovieSceneTools/Private/TrackEditors/PropertyTrackEditors/DoublePropertyTrackEditor.h
+++ b/Engine/Source/Editor/MovieSceneTools/Private/TrackEditors/PropertyTrackEditors/DoublePropertyTrackEditor.h
@@ -45,6 +45,8 @@ public:
 	static TSharedRef<ISequencerTrackEditor> CreateTrackEditor(TSharedRef<ISequencer> OwningSequencer);
 protected:
 
+	virtual FText GetDisplayName() const override;
+
 	virtual void GenerateKeysFromPropertyChanged(const FPropertyChangedParams& PropertyChangedParams, UMovieSceneSection* SectionToKey, FGeneratedTrackKeys& OutGeneratedKeys) override;
 
 private:
diff --git a/Engine/Source/Editor/MovieSceneTools/Private/TrackEditors/PropertyTrackEditors/EulerTransformPropertyTrackEditor.cpp b/Engine/Source/Editor/MovieSceneTools/Private/TrackEditors/PropertyTrackEditors/EulerTransformPropertyTrackEditor.cpp
--- a/Engine/Source/Editor/MovieSceneTools/Private/TrackEditors/PropertyTrackEditors/EulerTransformPropertyTrackEditor.cpp
+++ b/Engine/Source/Editor/MovieSceneTools/Private/TrackEditors/PropertyTrackEditors/EulerTransformPropertyTrackEditor.cpp
@@ -12,11 +12,17 @@
 #include "MovieSceneTracksComponentTypes.h"
 #include "Systems/MovieScenePropertyInstantiator.h"
 
+#define LOCTEXT_NAMESPACE "EulerTransformPropertyTrackEditor"
+
 TSharedRef<ISequencerTrackEditor> FEulerTransformPropertyTrackEditor::CreateTrackEditor( TSharedRef<ISequencer> InSequencer )
 {
 	return MakeShareable( new FEulerTransformPropertyTrackEditor( InSequencer ) );
 }
 
+FText FEulerTransformPropertyTrackEditor::GetDisplayName() const
+{
+	return LOCTEXT("EulerTransformPropertyTrackEditor_DisplayName", "Euler Transform Property");
+}
 
 TSharedRef<ISequencerSection> FEulerTransformPropertyTrackEditor::MakeSectionInterface(UMovieSceneSection& SectionObject, UMovieSceneTrack& Track, FGuid ObjectBinding)
 {
@@ -41,7 +47,7 @@ TSharedPtr<SWidget> FEulerTransformPropertyTrackEditor::BuildOutlinerEditWidget(
 		return MenuBuilder.MakeWidget();
 	};
 
-	return UE::Sequencer::MakeAddButton(NSLOCTEXT("FEulerTransformPropertyTrackEditor", "AddSection", "Section"), FOnGetContent::CreateLambda(SubMenuCallback), Params.ViewModel);
+	return UE::Sequencer::MakeAddButton(LOCTEXT("AddSection", "Section"), FOnGetContent::CreateLambda(SubMenuCallback), Params.ViewModel);
 }
 
 
@@ -252,3 +258,5 @@ FEulerTransform FEulerTransformPropertyTrackEditor::RecomposeTransform(const FEu
 
 	return InTransform;
 }
+
+#undef LOCTEXT_NAMESPACE
diff --git a/Engine/Source/Editor/MovieSceneTools/Private/TrackEditors/PropertyTrackEditors/EulerTransformPropertyTrackEditor.h b/Engine/Source/Editor/MovieSceneTools/Private/TrackEditors/PropertyTrackEditors/EulerTransformPropertyTrackEditor.h
--- a/Engine/Source/Editor/MovieSceneTools/Private/TrackEditors/PropertyTrackEditors/EulerTransformPropertyTrackEditor.h
+++ b/Engine/Source/Editor/MovieSceneTools/Private/TrackEditors/PropertyTrackEditors/EulerTransformPropertyTrackEditor.h
@@ -53,7 +53,7 @@ public:
 	static TSharedRef<ISequencerTrackEditor> CreateTrackEditor(TSharedRef<ISequencer> OwningSequencer);
 
 	//~ ISequencerTrackEditor interface
-
+	virtual FText GetDisplayName() const override;
 	virtual TSharedRef<ISequencerSection> MakeSectionInterface(UMovieSceneSection& SectionObject, UMovieSceneTrack& Track, FGuid ObjectBinding) override;
 	virtual TSharedPtr<SWidget> BuildOutlinerEditWidget(const FGuid& ObjectBinding, UMovieSceneTrack* Track, const FBuildEditWidgetParams& Params) override;
 protected:
diff --git a/Engine/Source/Editor/MovieSceneTools/Private/TrackEditors/PropertyTrackEditors/FloatPropertyTrackEditor.cpp b/Engine/Source/Editor/MovieSceneTools/Private/TrackEditors/PropertyTrackEditors/FloatPropertyTrackEditor.cpp
--- a/Engine/Source/Editor/MovieSceneTools/Private/TrackEditors/PropertyTrackEditors/FloatPropertyTrackEditor.cpp
+++ b/Engine/Source/Editor/MovieSceneTools/Private/TrackEditors/PropertyTrackEditors/FloatPropertyTrackEditor.cpp
@@ -9,11 +9,17 @@
 #include "Systems/MovieScenePropertyInstantiator.h"
 #include "MovieSceneTracksComponentTypes.h"
 
+#define LOCTEXT_NAMESPACE "FloatPropertyTrackEditor"
+
 TSharedRef<ISequencerTrackEditor> FFloatPropertyTrackEditor::CreateTrackEditor( TSharedRef<ISequencer> OwningSequencer )
 {
 	return MakeShareable(new FFloatPropertyTrackEditor(OwningSequencer));
 }
 
+FText FFloatPropertyTrackEditor::GetDisplayName() const
+{
+	return LOCTEXT("FloatPropertyTrackEditor_DisplayName", "Float Property");
+}
 
 void FFloatPropertyTrackEditor::GenerateKeysFromPropertyChanged( const FPropertyChangedParams& PropertyChangedParams, UMovieSceneSection* SectionToKey, FGeneratedTrackKeys& OutGeneratedKeys )
 {
@@ -51,3 +57,5 @@ float FFloatPropertyTrackEditor::RecomposeFloat(float InCurrentValue, UObject* A
 
 	return Recomposed;
 }
+
+#undef LOCTEXT_NAMESPACE
diff --git a/Engine/Source/Editor/MovieSceneTools/Private/TrackEditors/PropertyTrackEditors/FloatPropertyTrackEditor.h b/Engine/Source/Editor/MovieSceneTools/Private/TrackEditors/PropertyTrackEditors/FloatPropertyTrackEditor.h
--- a/Engine/Source/Editor/MovieSceneTools/Private/TrackEditors/PropertyTrackEditors/FloatPropertyTrackEditor.h
+++ b/Engine/Source/Editor/MovieSceneTools/Private/TrackEditors/PropertyTrackEditors/FloatPropertyTrackEditor.h
@@ -43,8 +43,10 @@ public:
 	 * @return The new instance of this class
 	 */
 	static TSharedRef<ISequencerTrackEditor> CreateTrackEditor(TSharedRef<ISequencer> OwningSequencer);
+
 protected:
 
+	virtual FText GetDisplayName() const override;
 	virtual void GenerateKeysFromPropertyChanged(const FPropertyChangedParams& PropertyChangedParams, UMovieSceneSection* SectionToKey, FGeneratedTrackKeys& OutGeneratedKeys) override;
 
 private:
diff --git a/Engine/Source/Editor/MovieSceneTools/Private/TrackEditors/PropertyTrackEditors/IntegerPropertyTrackEditor.cpp b/Engine/Source/Editor/MovieSceneTools/Private/TrackEditors/PropertyTrackEditors/IntegerPropertyTrackEditor.cpp
--- a/Engine/Source/Editor/MovieSceneTools/Private/TrackEditors/PropertyTrackEditors/IntegerPropertyTrackEditor.cpp
+++ b/Engine/Source/Editor/MovieSceneTools/Private/TrackEditors/PropertyTrackEditors/IntegerPropertyTrackEditor.cpp
@@ -4,12 +4,17 @@
 #include "EntitySystem/Interrogation/MovieSceneInterrogationLinker.h"
 #include "Tracks/MovieSceneIntegerTrack.h"
 
+#define LOCTEXT_NAMESPACE "IntegerPropertyTrackEditor"
 
 TSharedRef<ISequencerTrackEditor> FIntegerPropertyTrackEditor::CreateTrackEditor( TSharedRef<ISequencer> OwningSequencer )
 {
 	return MakeShareable(new FIntegerPropertyTrackEditor(OwningSequencer));
 }
 
+FText FIntegerPropertyTrackEditor::GetDisplayName() const
+{
+	return LOCTEXT("IntegerPropertyTrackEditor_DisplayName", "Integer Property");
+}
 
 void FIntegerPropertyTrackEditor::GenerateKeysFromPropertyChanged( const FPropertyChangedParams& PropertyChangedParams, UMovieSceneSection* SectionToKey, FGeneratedTrackKeys& OutGeneratedKeys )
 {
@@ -48,3 +53,5 @@ bool FIntegerPropertyTrackEditor::ModifyGeneratedKeysByCurrentAndWeight(UObject
 
 	return false;
 }
+
+#undef LOCTEXT_NAMESPACE
diff --git a/Engine/Source/Editor/MovieSceneTools/Private/TrackEditors/PropertyTrackEditors/IntegerPropertyTrackEditor.h b/Engine/Source/Editor/MovieSceneTools/Private/TrackEditors/PropertyTrackEditors/IntegerPropertyTrackEditor.h
--- a/Engine/Source/Editor/MovieSceneTools/Private/TrackEditors/PropertyTrackEditors/IntegerPropertyTrackEditor.h
+++ b/Engine/Source/Editor/MovieSceneTools/Private/TrackEditors/PropertyTrackEditors/IntegerPropertyTrackEditor.h
@@ -47,6 +47,7 @@ public:
 
 protected:
 
+	virtual FText GetDisplayName() const override;
 	virtual void GenerateKeysFromPropertyChanged(const FPropertyChangedParams& PropertyChangedParams, UMovieSceneSection* SectionToKey, FGeneratedTrackKeys& OutGeneratedKeys) override;
 	virtual bool ModifyGeneratedKeysByCurrentAndWeight(UObject* Object, UMovieSceneTrack *Track, UMovieSceneSection* SectionToKey, FFrameNumber Time, FGeneratedTrackKeys& GeneratedTotalKeys, float Weight) const override;
 
diff --git a/Engine/Source/Editor/MovieSceneTools/Private/TrackEditors/PropertyTrackEditors/RotatorPropertyTrackEditor.cpp b/Engine/Source/Editor/MovieSceneTools/Private/TrackEditors/PropertyTrackEditors/RotatorPropertyTrackEditor.cpp
--- a/Engine/Source/Editor/MovieSceneTools/Private/TrackEditors/PropertyTrackEditors/RotatorPropertyTrackEditor.cpp
+++ b/Engine/Source/Editor/MovieSceneTools/Private/TrackEditors/PropertyTrackEditors/RotatorPropertyTrackEditor.cpp
@@ -4,11 +4,18 @@
 
 #include "MovieSceneTracksComponentTypes.h"
 
+#define LOCTEXT_NAMESPACE "RotatorPropertyTrackEditor"
+
 TSharedRef<ISequencerTrackEditor> FRotatorPropertyTrackEditor::CreateTrackEditor(TSharedRef<ISequencer> InOwningSequencer)
 {
 	return MakeShared<FRotatorPropertyTrackEditor>(InOwningSequencer);
 }
 
+FText FRotatorPropertyTrackEditor::GetDisplayName() const
+{
+	return LOCTEXT("RotatorPropertyTrackEditor_DisplayName", "Rotator Property");
+}
+
 void FRotatorPropertyTrackEditor::GenerateKeysFromPropertyChanged(const FPropertyChangedParams& InPropertyChangedParams, UMovieSceneSection* InSectionToKey, FGeneratedTrackKeys& OutGeneratedKeys)
 {
 	using namespace  UE::MovieScene;
@@ -39,3 +46,5 @@ void FRotatorPropertyTrackEditor::GenerateKeysFromPropertyChanged(const FPropert
 		, CurrentRotator.Roll
 		, bKeyRoll));
 }
+
+#undef LOCTEXT_NAMESPACE
diff --git a/Engine/Source/Editor/MovieSceneTools/Private/TrackEditors/PropertyTrackEditors/RotatorPropertyTrackEditor.h b/Engine/Source/Editor/MovieSceneTools/Private/TrackEditors/PropertyTrackEditors/RotatorPropertyTrackEditor.h
--- a/Engine/Source/Editor/MovieSceneTools/Private/TrackEditors/PropertyTrackEditors/RotatorPropertyTrackEditor.h
+++ b/Engine/Source/Editor/MovieSceneTools/Private/TrackEditors/PropertyTrackEditors/RotatorPropertyTrackEditor.h
@@ -31,6 +31,7 @@ public:
 
 protected:
 	//~ Begin FPropertyTrackEditor
+	virtual FText GetDisplayName() const override;
 	virtual void GenerateKeysFromPropertyChanged(const FPropertyChangedParams& InPropertyChangedParams, UMovieSceneSection* InSectionToKey, FGeneratedTrackKeys& OutGeneratedKeys) override;
 	//~ End FPropertyTrackEditor
 };
diff --git a/Engine/Source/Editor/MovieSceneTools/Private/TrackEditors/PropertyTrackEditors/StringPropertyTrackEditor.cpp b/Engine/Source/Editor/MovieSceneTools/Private/TrackEditors/PropertyTrackEditors/StringPropertyTrackEditor.cpp
--- a/Engine/Source/Editor/MovieSceneTools/Private/TrackEditors/PropertyTrackEditors/StringPropertyTrackEditor.cpp
+++ b/Engine/Source/Editor/MovieSceneTools/Private/TrackEditors/PropertyTrackEditors/StringPropertyTrackEditor.cpp
@@ -16,14 +16,22 @@ class ISequencer;
 class ISequencerTrackEditor;
 class UMovieSceneSection;
 
+#define LOCTEXT_NAMESPACE "StringPropertyTrackEditor"
 
 TSharedRef<ISequencerTrackEditor> FStringPropertyTrackEditor::CreateTrackEditor( TSharedRef<ISequencer> OwningSequencer )
 {
 	return MakeShareable(new FStringPropertyTrackEditor(OwningSequencer));
 }
 
+FText FStringPropertyTrackEditor::GetDisplayName() const
+{
+	return LOCTEXT("StringPropertyTrackEditor_DisplayName", "String Property");
+}
+
 void FStringPropertyTrackEditor::GenerateKeysFromPropertyChanged( const FPropertyChangedParams& PropertyChangedParams, UMovieSceneSection* SectionToKey, FGeneratedTrackKeys& OutGeneratedKeys )
 {
 	FString KeyedValue = PropertyChangedParams.GetPropertyValue<FString>();
 	OutGeneratedKeys.Add(FMovieSceneChannelValueSetter::Create<FMovieSceneStringChannel>(0, MoveTemp(KeyedValue), true));
 }
+
+#undef LOCTEXT_NAMESPACE
diff --git a/Engine/Source/Editor/MovieSceneTools/Private/TrackEditors/PropertyTrackEditors/StringPropertyTrackEditor.h b/Engine/Source/Editor/MovieSceneTools/Private/TrackEditors/PropertyTrackEditors/StringPropertyTrackEditor.h
--- a/Engine/Source/Editor/MovieSceneTools/Private/TrackEditors/PropertyTrackEditors/StringPropertyTrackEditor.h
+++ b/Engine/Source/Editor/MovieSceneTools/Private/TrackEditors/PropertyTrackEditors/StringPropertyTrackEditor.h
@@ -56,6 +56,6 @@ public:
 protected:
 
 	//~ FPropertyTrackEditor interface
-
+	virtual FText GetDisplayName() const override;
 	virtual void GenerateKeysFromPropertyChanged(const FPropertyChangedParams& PropertyChangedParams, UMovieSceneSection* SectionToKey, FGeneratedTrackKeys& OutGeneratedKeys) override;
 };
diff --git a/Engine/Source/Editor/MovieSceneTools/Private/TrackEditors/PropertyTrackEditors/TransformPropertyTrackEditor.cpp b/Engine/Source/Editor/MovieSceneTools/Private/TrackEditors/PropertyTrackEditors/TransformPropertyTrackEditor.cpp
--- a/Engine/Source/Editor/MovieSceneTools/Private/TrackEditors/PropertyTrackEditors/TransformPropertyTrackEditor.cpp
+++ b/Engine/Source/Editor/MovieSceneTools/Private/TrackEditors/PropertyTrackEditors/TransformPropertyTrackEditor.cpp
@@ -13,11 +13,17 @@
 #include "Systems/MovieScenePropertyInstantiator.h"
 #include "MovieSceneTracksComponentTypes.h"
 
+#define LOCTEXT_NAMESPACE "TransformPropertyTrackEditor"
+
 TSharedRef<ISequencerTrackEditor> FTransformPropertyTrackEditor::CreateTrackEditor( TSharedRef<ISequencer> InSequencer )
 {
 	return MakeShareable( new FTransformPropertyTrackEditor( InSequencer ) );
 }
 
+FText FTransformPropertyTrackEditor::GetDisplayName() const
+{
+	return LOCTEXT("TransformPropertyTrackEditor_DisplayName", "Transform Property");
+}
 
 TSharedRef<ISequencerSection> FTransformPropertyTrackEditor::MakeSectionInterface(UMovieSceneSection& SectionObject, UMovieSceneTrack& Track, FGuid ObjectBinding)
 {
@@ -89,3 +95,5 @@ UE::MovieScene::FIntermediate3DTransform FTransformPropertyTrackEditor::Recompos
 
 	return InTransformData;
 }
+
+#undef LOCTEXT_NAMESPACE
diff --git a/Engine/Source/Editor/MovieSceneTools/Private/TrackEditors/PropertyTrackEditors/TransformPropertyTrackEditor.h b/Engine/Source/Editor/MovieSceneTools/Private/TrackEditors/PropertyTrackEditors/TransformPropertyTrackEditor.h
--- a/Engine/Source/Editor/MovieSceneTools/Private/TrackEditors/PropertyTrackEditors/TransformPropertyTrackEditor.h
+++ b/Engine/Source/Editor/MovieSceneTools/Private/TrackEditors/PropertyTrackEditors/TransformPropertyTrackEditor.h
@@ -51,6 +51,7 @@ public:
 
 	//~ ISequencerTrackEditor interface
 
+	virtual FText GetDisplayName() const override;
 	virtual TSharedRef<ISequencerSection> MakeSectionInterface(UMovieSceneSection& SectionObject, UMovieSceneTrack& Track, FGuid ObjectBinding) override;
 	virtual TSharedPtr<SWidget> BuildOutlinerEditWidget(const FGuid& ObjectBinding, UMovieSceneTrack* Track, const FBuildEditWidgetParams& Params) override;
 protected:
diff --git a/Engine/Source/Editor/MovieSceneTools/Private/TrackEditors/PropertyTrackEditors/VectorPropertyTrackEditor.cpp b/Engine/Source/Editor/MovieSceneTools/Private/TrackEditors/PropertyTrackEditors/VectorPropertyTrackEditor.cpp
--- a/Engine/Source/Editor/MovieSceneTools/Private/TrackEditors/PropertyTrackEditors/VectorPropertyTrackEditor.cpp
+++ b/Engine/Source/Editor/MovieSceneTools/Private/TrackEditors/PropertyTrackEditors/VectorPropertyTrackEditor.cpp
@@ -8,6 +8,7 @@
 #include "Evaluation/MovieScenePropertyTemplate.h"
 #include "EntitySystem/Interrogation/MovieSceneInterrogationLinker.h"
 
+#define LOCTEXT_NAMESPACE "FloatVectorPropertyTrackEditor"
 
 FName FFloatVectorPropertyTrackEditor::XName( "X" );
 FName FFloatVectorPropertyTrackEditor::YName( "Y" );
@@ -19,6 +20,11 @@ TSharedRef<ISequencerTrackEditor> FFloatVectorPropertyTrackEditor::CreateTrackEd
 	return MakeShareable( new FFloatVectorPropertyTrackEditor( InSequencer ) );
 }
 
+FText FFloatVectorPropertyTrackEditor::GetDisplayName() const
+{
+	return LOCTEXT("FloatVectorPropertyTrackEditor_DisplayName", "Float Vector Property");
+}
+
 void FFloatVectorPropertyTrackEditor::GenerateKeysFromPropertyChanged( const FPropertyChangedParams& PropertyChangedParams, UMovieSceneSection* SectionToKey, FGeneratedTrackKeys& OutGeneratedKeys )
 {
 	const FStructProperty* StructProp = CastField<const FStructProperty>( PropertyChangedParams.PropertyPath.GetLeafMostProperty().Property.Get() );
@@ -168,6 +174,11 @@ TSharedRef<ISequencerTrackEditor> FDoubleVectorPropertyTrackEditor::CreateTrackE
 	return MakeShareable( new FDoubleVectorPropertyTrackEditor( InSequencer ) );
 }
 
+FText FDoubleVectorPropertyTrackEditor::GetDisplayName() const
+{
+	return LOCTEXT("DoubleVectorPropertyTrackEditor_DisplayName", "Double Vector Property");
+}
+
 void FDoubleVectorPropertyTrackEditor::GenerateKeysFromPropertyChanged( const FPropertyChangedParams& PropertyChangedParams, UMovieSceneSection* SectionToKey, FGeneratedTrackKeys& OutGeneratedKeys )
 {
 	const FStructProperty* StructProp = CastField<const FStructProperty>( PropertyChangedParams.PropertyPath.GetLeafMostProperty().Property.Get() );
@@ -313,3 +324,5 @@ bool FDoubleVectorPropertyTrackEditor::ModifyGeneratedKeysByCurrentAndWeight(UOb
 
 	return false;
 }
+
+#undef LOCTEXT_NAMESPACE
diff --git a/Engine/Source/Editor/MovieSceneTools/Private/TrackEditors/PropertyTrackEditors/VectorPropertyTrackEditor.h b/Engine/Source/Editor/MovieSceneTools/Private/TrackEditors/PropertyTrackEditors/VectorPropertyTrackEditor.h
--- a/Engine/Source/Editor/MovieSceneTools/Private/TrackEditors/PropertyTrackEditors/VectorPropertyTrackEditor.h
+++ b/Engine/Source/Editor/MovieSceneTools/Private/TrackEditors/PropertyTrackEditors/VectorPropertyTrackEditor.h
@@ -53,6 +53,7 @@ public:
 protected:
 	//~ FPropertyTrackEditor interface
 
+	virtual FText GetDisplayName() const override;
 	virtual void GenerateKeysFromPropertyChanged(const FPropertyChangedParams& PropertyChangedParams, UMovieSceneSection* SectionToKey, FGeneratedTrackKeys& OutGeneratedKeys) override;
 	virtual void InitializeNewTrack(UMovieSceneFloatVectorTrack* NewTrack, FPropertyChangedParams PropertyChangedParams) override;
 	virtual bool ModifyGeneratedKeysByCurrentAndWeight(UObject* Object, UMovieSceneTrack *Track, UMovieSceneSection* SectionToKey, FFrameNumber Time, FGeneratedTrackKeys& GeneratedTotalKeys, float Weight) const override;
@@ -108,6 +109,7 @@ public:
 protected:
 	//~ FPropertyTrackEditor interface
 
+	virtual FText GetDisplayName() const override;
 	virtual void GenerateKeysFromPropertyChanged(const FPropertyChangedParams& PropertyChangedParams, UMovieSceneSection* SectionToKey, FGeneratedTrackKeys& OutGeneratedKeys) override;
 	virtual void InitializeNewTrack(UMovieSceneDoubleVectorTrack* NewTrack, FPropertyChangedParams PropertyChangedParams) override;
 	virtual bool ModifyGeneratedKeysByCurrentAndWeight(UObject* Object, UMovieSceneTrack *Track, UMovieSceneSection* SectionToKey, FFrameNumber Time, FGeneratedTrackKeys& GeneratedTotalKeys, float Weight) const override;
diff --git a/Engine/Source/Editor/MovieSceneTools/Private/TrackEditors/PropertyTrackEditors/VisibilityPropertyTrackEditor.cpp b/Engine/Source/Editor/MovieSceneTools/Private/TrackEditors/PropertyTrackEditors/VisibilityPropertyTrackEditor.cpp
--- a/Engine/Source/Editor/MovieSceneTools/Private/TrackEditors/PropertyTrackEditors/VisibilityPropertyTrackEditor.cpp
+++ b/Engine/Source/Editor/MovieSceneTools/Private/TrackEditors/PropertyTrackEditors/VisibilityPropertyTrackEditor.cpp
@@ -14,12 +14,17 @@ class ISequencerTrackEditor;
 class UMovieSceneSection;
 class UMovieSceneTrack;
 
+#define LOCTEXT_NAMESPACE "VisibilityPropertyTrackEditor"
 
 TSharedRef<ISequencerTrackEditor> FVisibilityPropertyTrackEditor::CreateTrackEditor( TSharedRef<ISequencer> OwningSequencer )
 {
 	return MakeShareable(new FVisibilityPropertyTrackEditor(OwningSequencer));
 }
 
+FText FVisibilityPropertyTrackEditor::GetDisplayName() const
+{
+	return LOCTEXT("VisibilityPropertyTrackEditor_DisplayName", "Visibility Property");
+}
 
 TSharedRef<ISequencerSection> FVisibilityPropertyTrackEditor::MakeSectionInterface(UMovieSceneSection& SectionObject, UMovieSceneTrack& Track, FGuid ObjectBinding)
 {
@@ -33,3 +38,5 @@ void FVisibilityPropertyTrackEditor::GenerateKeysFromPropertyChanged( const FPro
 	bool KeyedValue = !PropertyChangedParams.GetPropertyValue<bool>();
 	OutGeneratedKeys.Add(FMovieSceneChannelValueSetter::Create<FMovieSceneBoolChannel>(0, KeyedValue, true));
 }
+
+#undef LOCTEXT_NAMESPACE
diff --git a/Engine/Source/Editor/MovieSceneTools/Private/TrackEditors/PropertyTrackEditors/VisibilityPropertyTrackEditor.h b/Engine/Source/Editor/MovieSceneTools/Private/TrackEditors/PropertyTrackEditors/VisibilityPropertyTrackEditor.h
--- a/Engine/Source/Editor/MovieSceneTools/Private/TrackEditors/PropertyTrackEditors/VisibilityPropertyTrackEditor.h
+++ b/Engine/Source/Editor/MovieSceneTools/Private/TrackEditors/PropertyTrackEditors/VisibilityPropertyTrackEditor.h
@@ -58,6 +58,7 @@ protected:
 
 	//~ FPropertyTrackEditor interface
 
+	virtual FText GetDisplayName() const override;
 	virtual TSharedRef<ISequencerSection> MakeSectionInterface(UMovieSceneSection& SectionObject, UMovieSceneTrack& Track, FGuid ObjectBinding) override;
 	virtual void GenerateKeysFromPropertyChanged(const FPropertyChangedParams& PropertyChangedParams, UMovieSceneSection* SectionToKey, FGeneratedTrackKeys& OutGeneratedKeys) override;
 	virtual bool ForCustomizedUseOnly() { return true; }
diff --git a/Engine/Source/Editor/MovieSceneTools/Private/TrackEditors/SkeletalAnimationTrackEditor.cpp b/Engine/Source/Editor/MovieSceneTools/Private/TrackEditors/SkeletalAnimationTrackEditor.cpp
--- a/Engine/Source/Editor/MovieSceneTools/Private/TrackEditors/SkeletalAnimationTrackEditor.cpp
+++ b/Engine/Source/Editor/MovieSceneTools/Private/TrackEditors/SkeletalAnimationTrackEditor.cpp
@@ -1430,6 +1430,11 @@ bool FSkeletalAnimationTrackEditor::HandleAssetAdded(UObject* Asset, const FGuid
 	return false;
 }
 
+FText FSkeletalAnimationTrackEditor::GetDisplayName() const
+{
+	return LOCTEXT("SkeletalAnimationTrackEditor_DisplayName", "Skeletal Animation");
+}
+
 void FSkeletalAnimationTrackEditor::BuildObjectBindingContextMenu(FMenuBuilder& MenuBuilder, const TArray<FGuid>& ObjectBindings, const UClass* ObjectClass)
 {
 	if (ObjectClass->IsChildOf(USkeletalMeshComponent::StaticClass()) || ObjectClass->IsChildOf(AActor::StaticClass()) || ObjectClass->IsChildOf(UChildActorComponent::StaticClass()))
diff --git a/Engine/Source/Editor/MovieSceneTools/Private/TrackEditors/SkeletalAnimationTrackEditor.h b/Engine/Source/Editor/MovieSceneTools/Private/TrackEditors/SkeletalAnimationTrackEditor.h
--- a/Engine/Source/Editor/MovieSceneTools/Private/TrackEditors/SkeletalAnimationTrackEditor.h
+++ b/Engine/Source/Editor/MovieSceneTools/Private/TrackEditors/SkeletalAnimationTrackEditor.h
@@ -69,6 +69,7 @@ public:
 
 	// ISequencerTrackEditor interface
 
+	virtual FText GetDisplayName() const override;
 	virtual void BuildObjectBindingContextMenu(FMenuBuilder& MenuBuilder, const TArray<FGuid>& ObjectBindings, const UClass* ObjectClass) override;
 	virtual void BuildObjectBindingTrackMenu(FMenuBuilder& MenuBuilder, const TArray<FGuid>& ObjectBindings, const UClass* ObjectClass) override;
 	virtual bool HandleAssetAdded(UObject* Asset, const FGuid& TargetObjectGuid) override;
diff --git a/Engine/Source/Editor/MovieSceneTools/Private/TrackEditors/SlomoTrackEditor.cpp b/Engine/Source/Editor/MovieSceneTools/Private/TrackEditors/SlomoTrackEditor.cpp
--- a/Engine/Source/Editor/MovieSceneTools/Private/TrackEditors/SlomoTrackEditor.cpp
+++ b/Engine/Source/Editor/MovieSceneTools/Private/TrackEditors/SlomoTrackEditor.cpp
@@ -49,6 +49,11 @@ FSlomoTrackEditor::FSlomoTrackEditor(TSharedRef<ISequencer> InSequencer)
 /* ISequencerTrackEditor interface
  *****************************************************************************/
 
+FText FSlomoTrackEditor::GetDisplayName() const
+{
+	return LOCTEXT("TimeDilationTrackEditor_DisplayName", "Time Dilation");
+}
+
 void FSlomoTrackEditor::BuildAddTrackMenu(FMenuBuilder& MenuBuilder)
 {
 	MenuBuilder.AddMenuEntry(
diff --git a/Engine/Source/Editor/MovieSceneTools/Private/TrackEditors/SlomoTrackEditor.h b/Engine/Source/Editor/MovieSceneTools/Private/TrackEditors/SlomoTrackEditor.h
--- a/Engine/Source/Editor/MovieSceneTools/Private/TrackEditors/SlomoTrackEditor.h
+++ b/Engine/Source/Editor/MovieSceneTools/Private/TrackEditors/SlomoTrackEditor.h
@@ -42,6 +42,7 @@ public:
 
 	// ISequencerTrackEditor interface
 
+	virtual FText GetDisplayName() const override;
 	virtual void BuildAddTrackMenu(FMenuBuilder& MenuBuilder) override;
 	virtual bool SupportsSequence(UMovieSceneSequence* InSequence) const override;
 	virtual bool SupportsType(TSubclassOf<UMovieSceneTrack> Type) const override;
diff --git a/Engine/Source/Editor/MovieSceneTools/Private/TrackEditors/SubTrackEditor.cpp b/Engine/Source/Editor/MovieSceneTools/Private/TrackEditors/SubTrackEditor.cpp
--- a/Engine/Source/Editor/MovieSceneTools/Private/TrackEditors/SubTrackEditor.cpp
+++ b/Engine/Source/Editor/MovieSceneTools/Private/TrackEditors/SubTrackEditor.cpp
@@ -367,6 +367,11 @@ double FSubTrackEditor::UnwindChannel(const double& OldValue, double NewValue)
 	return NewValue;
 }
 
+FText FSubTrackEditor::GetDisplayName() const
+{
+	return LOCTEXT("SubsequenceTrackEditor_DisplayName", "Subsequence");
+}
+
 void FSubTrackEditor::ProcessKeyOperation(FFrameNumber InKeyTime, const UE::Sequencer::FKeyOperation& Operation, ISequencer& InSequencer, TArray<UE::Sequencer::FAddKeyResult>* OutResults)
 {
 	using namespace UE::Sequencer;
diff --git a/Engine/Source/Editor/MovieSceneTools/Private/TrackEditors/TimeWarpTrackEditor.cpp b/Engine/Source/Editor/MovieSceneTools/Private/TrackEditors/TimeWarpTrackEditor.cpp
--- a/Engine/Source/Editor/MovieSceneTools/Private/TrackEditors/TimeWarpTrackEditor.cpp
+++ b/Engine/Source/Editor/MovieSceneTools/Private/TrackEditors/TimeWarpTrackEditor.cpp
@@ -346,6 +346,11 @@ void FTimeWarpTrackEditor::ProcessKeyOperation(FFrameNumber InKeyTime, const UE:
 	Operation.ApplyDefault(InKeyTime, InSequencer, OutResults);
 }
 
+FText FTimeWarpTrackEditor::GetDisplayName() const
+{
+	return LOCTEXT("TimeWarpTrackEditor_DisplayName", "Time Warp");
+}
+
 void FTimeWarpTrackEditor::BuildAddTrackMenu(FMenuBuilder& MenuBuilder)
 {
 	UMovieScene* FocusedMovieScene = GetFocusedMovieScene();
diff --git a/Engine/Source/Editor/MovieSceneTools/Private/TrackEditors/TransformTrackEditor.cpp b/Engine/Source/Editor/MovieSceneTools/Private/TrackEditors/TransformTrackEditor.cpp
--- a/Engine/Source/Editor/MovieSceneTools/Private/TrackEditors/TransformTrackEditor.cpp
+++ b/Engine/Source/Editor/MovieSceneTools/Private/TrackEditors/TransformTrackEditor.cpp
@@ -1244,6 +1244,11 @@ FTransformData F3DTransformTrackEditor::RecomposeTransform(const FTransformData&
 	return FTransformData(CurrentTransform.GetLocation(), CurrentTransform.GetRotation().Rotator(), CurrentTransform.GetScale3D());
 }
 
+FText F3DTransformTrackEditor::GetDisplayName() const
+{
+	return LOCTEXT("3DTransformTrackEditor_DisplayName", "3D Transform");
+}
+
 void F3DTransformTrackEditor::ProcessKeyOperation(FFrameNumber InKeyTime, const UE::Sequencer::FKeyOperation& Operation, ISequencer& InSequencer, TArray<UE::Sequencer::FAddKeyResult>* OutResults)
 {
 	using namespace UE::Sequencer;
diff --git a/Engine/Source/Editor/MovieSceneTools/Private/TrackEditors/TransformTrackEditor.h b/Engine/Source/Editor/MovieSceneTools/Private/TrackEditors/TransformTrackEditor.h
--- a/Engine/Source/Editor/MovieSceneTools/Private/TrackEditors/TransformTrackEditor.h
+++ b/Engine/Source/Editor/MovieSceneTools/Private/TrackEditors/TransformTrackEditor.h
@@ -54,6 +54,7 @@ public:
 public:
 
 	// ISequencerTrackEditor interface
+	virtual FText GetDisplayName() const override;
 	virtual void ProcessKeyOperation(FFrameNumber InKeyTime, const UE::Sequencer::FKeyOperation& Operation, ISequencer& InSequencer, TArray<UE::Sequencer::FAddKeyResult>* OutResults = nullptr) override;
 	virtual void BuildObjectBindingColumnWidgets(TFunctionRef<TSharedRef<SHorizontalBox>()> GetEditBox, const UE::Sequencer::TViewModelPtr<UE::Sequencer::FObjectBindingModel>& ObjectBinding, const UE::Sequencer::FCreateOutlinerViewParams& InParams, const FName& InColumnName) override;
 	virtual void BuildObjectBindingTrackMenu(FMenuBuilder& MenuBuilder, const TArray<FGuid>& ObjectBindinsg, const UClass* ObjectClass) override;
diff --git a/Engine/Source/Editor/MovieSceneTools/Public/KeyframeTrackEditor.h b/Engine/Source/Editor/MovieSceneTools/Public/KeyframeTrackEditor.h
--- a/Engine/Source/Editor/MovieSceneTools/Public/KeyframeTrackEditor.h
+++ b/Engine/Source/Editor/MovieSceneTools/Public/KeyframeTrackEditor.h
@@ -473,6 +473,11 @@ public:
 
 	// ISequencerTrackEditor interface
 
+	virtual FText GetDisplayName() const override
+	{
+		return NSLOCTEXT("Sequencer", "KeyframeTrackEditor_DisplayName", "Keyframe");
+	}
+
 	virtual bool SupportsType( TSubclassOf<UMovieSceneTrack> Type ) const override
 	{
 		return Type == TrackType::StaticClass();
diff --git a/Engine/Source/Editor/MovieSceneTools/Public/PropertyTrackEditor.h b/Engine/Source/Editor/MovieSceneTools/Public/PropertyTrackEditor.h
--- a/Engine/Source/Editor/MovieSceneTools/Public/PropertyTrackEditor.h
+++ b/Engine/Source/Editor/MovieSceneTools/Public/PropertyTrackEditor.h
@@ -16,6 +16,8 @@
 #include "AnimatedPropertyKey.h"
 #include "Tracks/MovieScenePropertyTrack.h"
 
+#define LOCTEXT_NAMESPACE "PropertyTrackEditor"
+
 /**
 * Tools for animatable property types such as floats ands vectors
 */
@@ -65,6 +67,11 @@ public:
 
 	//~ ISequencerTrackEditor interface
 
+	virtual FText GetDisplayName() const override
+	{
+		return LOCTEXT("PropertyTrackEditor_DisplayName", "Property");
+	}
+
     virtual bool SupportsSequence(UMovieSceneSequence* InSequence) const override
     {
         return true;
@@ -113,7 +120,7 @@ protected:
 				const FProperty* ArrayInnerProperty = Info.Property.Get();
 				if (ArrayInnerProperty && Info.ArrayIndex != INDEX_NONE)
 				{
-					DisplayText = FText::Format(NSLOCTEXT("PropertyTrackEditor", "DisplayTextArrayFormat", "{0} ({1}[{2}])"),
+					DisplayText = FText::Format(LOCTEXT("DisplayTextArrayFormat", "{0} ({1}[{2}])"),
 						ChangedProperty->GetDisplayNameText(),
 						(ParentArrayProperty ? ParentArrayProperty : ArrayInnerProperty)->GetDisplayNameText(),
 						FText::AsNumber(Info.ArrayIndex)
@@ -129,7 +136,7 @@ protected:
 					const FStructProperty* ParentStructProperty = PropertyIndex > 0 ? CastField<FStructProperty>(PropertyChangedParams.PropertyPath.GetPropertyInfo(PropertyIndex - 1).Property.Get()) : nullptr;
 					if (ParentStructProperty)
 					{
-						DisplayText = FText::Format(NSLOCTEXT("PropertyTrackEditor", "DisplayTextStructFormat", "{0} ({1})"),
+						DisplayText = FText::Format(LOCTEXT("DisplayTextStructFormat", "{0} ({1})"),
 							ChangedProperty->GetDisplayNameText(),
 							ParentStructProperty->GetDisplayNameText()
 						);
@@ -258,3 +265,5 @@ private:
 	/** An array of property type names which are being watched for changes. */
 	TArray<FAnimatedPropertyKey> WatchedProperties;
 };
+
+#undef LOCTEXT_NAMESPACE
diff --git a/Engine/Source/Editor/MovieSceneTools/Public/TrackEditors/AudioTrackEditor.h b/Engine/Source/Editor/MovieSceneTools/Public/TrackEditors/AudioTrackEditor.h
--- a/Engine/Source/Editor/MovieSceneTools/Public/TrackEditors/AudioTrackEditor.h
+++ b/Engine/Source/Editor/MovieSceneTools/Public/TrackEditors/AudioTrackEditor.h
@@ -52,6 +52,7 @@ public:
 
 	// ISequencerTrackEditor interface
 
+	virtual FText GetDisplayName() const override;
 	virtual void OnInitialize() override;
 	virtual void OnRelease() override;
 	virtual void BuildAddTrackMenu(FMenuBuilder& MenuBuilder) override;
diff --git a/Engine/Source/Editor/MovieSceneTools/Public/TrackEditors/BindingLifetimeTrackEditor.h b/Engine/Source/Editor/MovieSceneTools/Public/TrackEditors/BindingLifetimeTrackEditor.h
--- a/Engine/Source/Editor/MovieSceneTools/Public/TrackEditors/BindingLifetimeTrackEditor.h
+++ b/Engine/Source/Editor/MovieSceneTools/Public/TrackEditors/BindingLifetimeTrackEditor.h
@@ -55,6 +55,8 @@ public:
 
 	// ISequencerTrackEditor interface
 
+	virtual FText GetDisplayName() const override;
+
 	virtual UMovieSceneTrack* AddTrack(UMovieScene* FocusedMovieScene, const FGuid& ObjectHandle, TSubclassOf<class UMovieSceneTrack> TrackClass, FName UniqueTypeName) override;
 	virtual void BuildObjectBindingTrackMenu(FMenuBuilder& MenuBuilder, const TArray<FGuid>& ObjectBindings, const UClass* ObjectClass) override;
 	
diff --git a/Engine/Source/Editor/MovieSceneTools/Public/TrackEditors/CinematicShotTrackEditor.h b/Engine/Source/Editor/MovieSceneTools/Public/TrackEditors/CinematicShotTrackEditor.h
--- a/Engine/Source/Editor/MovieSceneTools/Public/TrackEditors/CinematicShotTrackEditor.h
+++ b/Engine/Source/Editor/MovieSceneTools/Public/TrackEditors/CinematicShotTrackEditor.h
@@ -50,6 +50,7 @@ public:
 public:
 
 	// ISequencerTrackEditor interface
+	virtual FText GetDisplayName() const override;
 	virtual TSharedPtr<SWidget> BuildOutlinerColumnWidget(const FBuildColumnWidgetParams& Params, const FName& ColumnName) override;
 	virtual TSharedRef<ISequencerSection> MakeSectionInterface(UMovieSceneSection& SectionObject, UMovieSceneTrack& Track, FGuid ObjectBinding) override;
 	virtual bool SupportsSequence(UMovieSceneSequence* InSequence) const override;
diff --git a/Engine/Source/Editor/MovieSceneTools/Public/TrackEditors/CustomPrimitiveDataTrackEditor.h b/Engine/Source/Editor/MovieSceneTools/Public/TrackEditors/CustomPrimitiveDataTrackEditor.h
--- a/Engine/Source/Editor/MovieSceneTools/Public/TrackEditors/CustomPrimitiveDataTrackEditor.h
+++ b/Engine/Source/Editor/MovieSceneTools/Public/TrackEditors/CustomPrimitiveDataTrackEditor.h
@@ -37,6 +37,7 @@ public:
 
 	// ISequencerTrackEditor interface
 
+	virtual FText GetDisplayName() const override;
 	virtual TSharedPtr<SWidget> BuildOutlinerEditWidget(const FGuid& ObjectBinding, UMovieSceneTrack* Track, const FBuildEditWidgetParams& Params) override;
 	virtual TSharedRef<ISequencerSection> MakeSectionInterface(UMovieSceneSection& SectionObject, UMovieSceneTrack& Track, FGuid ObjectBinding) override;
 	virtual bool SupportsType(TSubclassOf<UMovieSceneTrack> Type) const override;
diff --git a/Engine/Source/Editor/MovieSceneTools/Public/TrackEditors/MaterialTrackEditor.h b/Engine/Source/Editor/MovieSceneTools/Public/TrackEditors/MaterialTrackEditor.h
--- a/Engine/Source/Editor/MovieSceneTools/Public/TrackEditors/MaterialTrackEditor.h
+++ b/Engine/Source/Editor/MovieSceneTools/Public/TrackEditors/MaterialTrackEditor.h
@@ -37,6 +37,7 @@ public:
 
 	// ISequencerTrackEditor interface
 
+	virtual FText GetDisplayName() const override;
 	virtual TSharedPtr<SWidget> BuildOutlinerEditWidget( const FGuid& ObjectBinding, UMovieSceneTrack* Track, const FBuildEditWidgetParams& Params ) override;
 	virtual TSharedRef<ISequencerSection> MakeSectionInterface( UMovieSceneSection& SectionObject, UMovieSceneTrack& Track, FGuid ObjectBinding ) override;
 
diff --git a/Engine/Source/Editor/MovieSceneTools/Public/TrackEditors/PropertyTrackEditors/BoolPropertyTrackEditor.h b/Engine/Source/Editor/MovieSceneTools/Public/TrackEditors/PropertyTrackEditors/BoolPropertyTrackEditor.h
--- a/Engine/Source/Editor/MovieSceneTools/Public/TrackEditors/PropertyTrackEditors/BoolPropertyTrackEditor.h
+++ b/Engine/Source/Editor/MovieSceneTools/Public/TrackEditors/PropertyTrackEditors/BoolPropertyTrackEditor.h
@@ -65,6 +65,7 @@ public:
 protected:
 
 	//~ FPropertyTrackEditor interface
+	virtual FText GetDisplayName() const override;
 	virtual TSharedRef<ISequencerSection> MakeSectionInterface(UMovieSceneSection& SectionObject, UMovieSceneTrack& Track, FGuid ObjectBinding) override;
 	virtual void GenerateKeysFromPropertyChanged( const FPropertyChangedParams& PropertyChangedParams, UMovieSceneSection* SectionToKey, FGeneratedTrackKeys& OutGeneratedKeys ) override;
 };
diff --git a/Engine/Source/Editor/MovieSceneTools/Public/TrackEditors/SubTrackEditor.h b/Engine/Source/Editor/MovieSceneTools/Public/TrackEditors/SubTrackEditor.h
--- a/Engine/Source/Editor/MovieSceneTools/Public/TrackEditors/SubTrackEditor.h
+++ b/Engine/Source/Editor/MovieSceneTools/Public/TrackEditors/SubTrackEditor.h
@@ -50,6 +50,7 @@ public:
 public:
 
 	// ISequencerTrackEditor interface
+	virtual FText GetDisplayName() const override;
 	virtual void ProcessKeyOperation(FFrameNumber InKeyTime, const UE::Sequencer::FKeyOperation& Operation, ISequencer& InSequencer, TArray<UE::Sequencer::FAddKeyResult>* OutResults = nullptr) override;
 	virtual void BuildAddTrackMenu(FMenuBuilder& MenuBuilder) override;
 	virtual TSharedPtr<SWidget> BuildOutlinerEditWidget(const FGuid& ObjectBinding, UMovieSceneTrack* Track, const FBuildEditWidgetParams& Params) override;
diff --git a/Engine/Source/Editor/MovieSceneTools/Public/TrackEditors/TimeWarpTrackEditor.h b/Engine/Source/Editor/MovieSceneTools/Public/TrackEditors/TimeWarpTrackEditor.h
--- a/Engine/Source/Editor/MovieSceneTools/Public/TrackEditors/TimeWarpTrackEditor.h
+++ b/Engine/Source/Editor/MovieSceneTools/Public/TrackEditors/TimeWarpTrackEditor.h
@@ -72,11 +72,12 @@ public:
 
 	static TSharedPtr<UE::Sequencer::FTrackModel> CreateTrackModel(UMovieSceneTrack* Track);
 
+	MOVIESCENETOOLS_API void HandleAddTimeWarpTrack(TSubclassOf<UMovieSceneTimeWarpGetter> ClassType);
+
 private:
+	virtual FText GetDisplayName() const override;
 	virtual void BuildAddTrackMenu(FMenuBuilder& MenuBuilder) override;
 	virtual void ProcessKeyOperation(FFrameNumber InKeyTime, const UE::Sequencer::FKeyOperation& Operation, ISequencer& InSequencer, TArray<UE::Sequencer::FAddKeyResult>* OutResults = nullptr) override;
 	virtual TSharedRef<ISequencerSection> MakeSectionInterface(UMovieSceneSection& SectionObject, UMovieSceneTrack& Track, FGuid ObjectBinding) override;
 	virtual bool SupportsSequence(UMovieSceneSequence* InSequence) const override;
-
-	void HandleAddTimeWarpTrack(TSubclassOf<UMovieSceneTimeWarpGetter> ClassType);
 };
diff --git a/Engine/Source/Editor/Sequencer/Private/MVVM/ViewModels/SequencerOutlinerViewModel.cpp b/Engine/Source/Editor/Sequencer/Private/MVVM/ViewModels/SequencerOutlinerViewModel.cpp
--- a/Engine/Source/Editor/Sequencer/Private/MVVM/ViewModels/SequencerOutlinerViewModel.cpp
+++ b/Engine/Source/Editor/Sequencer/Private/MVVM/ViewModels/SequencerOutlinerViewModel.cpp
@@ -66,7 +66,7 @@ void FSequencerOutlinerViewModel::BuildContextMenu(FMenuBuilder& MenuBuilder)
 	// let track editors & object bindings populate the menu
 
 	// Always create the section so that we afford extension
-	MenuBuilder.BeginSection("ObjectBindings");
+	MenuBuilder.BeginSection("AAObjectBindings");
 	if (Sequencer.IsValid())
 	{
 		Sequencer->BuildAddObjectBindingsMenu(MenuBuilder);
diff --git a/Engine/Source/Editor/Sequencer/Private/MovieSceneTrackEditor.cpp b/Engine/Source/Editor/Sequencer/Private/MovieSceneTrackEditor.cpp
--- a/Engine/Source/Editor/Sequencer/Private/MovieSceneTrackEditor.cpp
+++ b/Engine/Source/Editor/Sequencer/Private/MovieSceneTrackEditor.cpp
@@ -205,6 +205,11 @@ const TSharedPtr<ISequencer> FMovieSceneTrackEditor::GetSequencer() const
 	return Sequencer.Pin();
 }
 
+FText FMovieSceneTrackEditor::GetDisplayName() const
+{
+	return FText::GetEmpty();
+}
+
 UMovieSceneTrack* FMovieSceneTrackEditor::AddTrack(UMovieScene* FocusedMovieScene, const FGuid& ObjectHandle, TSubclassOf<class UMovieSceneTrack> TrackClass, FName UniqueTypeName)
 {
 	return FocusedMovieScene->AddTrack(TrackClass, ObjectHandle);
diff --git a/Engine/Source/Editor/Sequencer/Private/SSequencer.cpp b/Engine/Source/Editor/Sequencer/Private/SSequencer.cpp
--- a/Engine/Source/Editor/Sequencer/Private/SSequencer.cpp
+++ b/Engine/Source/Editor/Sequencer/Private/SSequencer.cpp
@@ -1770,7 +1770,7 @@ TSharedRef<SWidget> SSequencer::MakeAddMenu()
 	if (SequencerPtr.Pin()->GetHostCapabilities().bSupportsAddFromContentBrowser)
 	{
 		MenuBuilder.AddMenuEntry(
-			LOCTEXT("SelectedFromContentBrowser", "Selected from Content Browser"),
+			LOCTEXT("SelectedFromContentBrowser", "Selection from Content Browser"),
 			LOCTEXT("SelectedFromContentBrowserToolTip", "Add selected content from the content browser"),
 			FSlateIcon(FAppStyle::GetAppStyleSetName(), "Icons.Use"),
 			FUIAction(
diff --git a/Engine/Source/Editor/Sequencer/Private/Sequencer.cpp b/Engine/Source/Editor/Sequencer/Private/Sequencer.cpp
--- a/Engine/Source/Editor/Sequencer/Private/Sequencer.cpp
+++ b/Engine/Source/Editor/Sequencer/Private/Sequencer.cpp
@@ -183,6 +183,8 @@
 #include "Bindings/MovieSceneSpawnableActorBinding.h"
 #include "Bindings/MovieSceneReplaceableActorBinding.h"
 #include "Filters/SequencerFilterBar.h"
+#include "TrackEditors/TimeWarpTrackEditor.h"
+#include "Variants/MovieSceneTimeWarpGetter.h"
 
 #define LOCTEXT_NAMESPACE "Sequencer"
 
@@ -9320,6 +9322,22 @@ void FSequencer::AddFolder()
 	NotifyMovieSceneDataChanged( EMovieSceneDataChangeType::MovieSceneStructureItemAdded );
 }
 
+void FSequencer::AddTimeWarp()
+{
+	const TSharedPtr<ISequencerTrackEditor>* const TimeWarpEditorPtr = TrackEditors.FindByPredicate(
+	[](const TSharedPtr<ISequencerTrackEditor>& InTrackEditor)
+	{
+		return InTrackEditor->GetDisplayName().ToString().Equals(TEXT("Time Warp"));
+	});
+	if (TimeWarpEditorPtr)
+	{
+		if (const TSharedPtr<FTimeWarpTrackEditor> TimeWarpEditor = StaticCastSharedPtr<FTimeWarpTrackEditor>(*TimeWarpEditorPtr))
+		{
+			TimeWarpEditor->HandleAddTimeWarpTrack(UMovieSceneTimeWarpGetter::StaticClass());
+		}
+	}
+}
+
 void FixSortingOrders(UMovieSceneTrack* InTrack, const TArray<UMovieSceneTrack*>& Tracks)
 {
 	// Fix the sorting orders if InTrack is going to conflict with any of the given Tracks (ie. same class and same name)
@@ -11986,32 +12004,60 @@ void FSequencer::BuildAddTrackMenu(class FMenuBuilder& MenuBuilder)
 		MenuBuilder.AddMenuEntry(
 			LOCTEXT("LoadRecording", "Load Recorded Data"),
 			LOCTEXT("LoadRecordingDataTooltip", "Load in saved data from a previous recording."),
-			FSlateIcon(FAppStyle::GetAppStyleSetName(), "ContentBrowser.AssetTreeFolderOpen"),
+			FSlateIcon(FAppStyle::GetAppStyleSetName(), TEXT("ContentBrowser.AssetTreeFolderOpen")),
 			FUIAction(FExecuteAction::CreateRaw(this, &FSequencer::OnLoadRecordedData)));
 	}
 
 	MenuBuilder.AddMenuEntry(
-		LOCTEXT( "AddFolder", "Add Folder" ),
-		LOCTEXT( "AddFolderToolTip", "Adds a new folder." ),
-		FSlateIcon( FAppStyle::GetAppStyleSetName(), "ContentBrowser.AssetTreeFolderOpen" ),
-		FUIAction( FExecuteAction::CreateRaw( this, &FSequencer::AddFolder ) ) );
+		LOCTEXT("AddFolder", "Add Folder"),
+		LOCTEXT("AddFolderToolTip", "Adds a new folder track."),
+		FSlateIcon(FAppStyle::GetAppStyleSetName(), TEXT("ContentBrowser.AssetTreeFolderOpen")),
+		FUIAction(FExecuteAction::CreateRaw(this, &FSequencer::AddFolder)));
 
-	for (int32 i = 0; i < TrackEditors.Num(); ++i)
+	MenuBuilder.AddMenuEntry(
+		LOCTEXT("TimeWarp", "Add Time Warp Track"),
+		LOCTEXT("TimeWarpToolTip", "Adds a new time warp track."),
+		FSlateIcon(FAppStyle::GetAppStyleSetName(), TEXT("Sequencer.Tracks.TimeWarp")),
+		FUIAction(FExecuteAction::CreateRaw(this, &FSequencer::AddTimeWarp)));
+
+	MenuBuilder.AddMenuSeparator(NAME_None);
+
+	TArray<TSharedPtr<ISequencerTrackEditor>> SortedTrackEditors = TrackEditors;
+
+	// UX design calls for some tracks to be forced to the top of the list
+	const TArray<FString> TrackEditorsToIgnore = { TEXT("Time Warp") };
+	SortedTrackEditors.RemoveAll([&TrackEditorsToIgnore](const TSharedPtr<ISequencerTrackEditor>& InTrackEditor)
+		{
+			return TrackEditorsToIgnore.Contains(InTrackEditor->GetDisplayName().ToString());
+		});
+
+	SortedTrackEditors.Sort([](const TSharedPtr<ISequencerTrackEditor>& InA, const TSharedPtr<ISequencerTrackEditor>& InB)
+		{
+			return InA->GetDisplayName().CompareTo(InB->GetDisplayName()) < 0;
+		});
+
+	for (const TSharedPtr<ISequencerTrackEditor>& TrackEditor : SortedTrackEditors)
 	{
-		if (TrackEditors[i]->SupportsSequence(GetFocusedMovieSceneSequence()))
+		if (TrackEditor->SupportsSequence(GetFocusedMovieSceneSequence()))
 		{
-			TrackEditors[i]->BuildAddTrackMenu(MenuBuilder);
+			TrackEditor->BuildAddTrackMenu(MenuBuilder);
 		}
 	}
 }
 
-void FSequencer::BuildAddObjectBindingsMenu(class FMenuBuilder& MenuBuilder)
+void FSequencer::BuildAddObjectBindingsMenu(FMenuBuilder& MenuBuilder)
 {
-	for (int32 i = 0; i < ObjectBindings.Num(); ++i)
+	TArray<TSharedPtr<ISequencerEditorObjectBinding>> SortedObjectBindings = ObjectBindings;
+	SortedObjectBindings.Sort([](const TSharedPtr<ISequencerEditorObjectBinding>& InA, const TSharedPtr<ISequencerEditorObjectBinding>& InB) -> bool
+		{
+			return InA->GetDisplayName().CompareTo(InB->GetDisplayName()) < 0;
+		});
+
+	for (const TSharedPtr<ISequencerEditorObjectBinding>& ObjectBinding : SortedObjectBindings)
 	{
-		if (ObjectBindings[i]->SupportsSequence(GetFocusedMovieSceneSequence()))
+		if (ObjectBinding->SupportsSequence(GetFocusedMovieSceneSequence()))
 		{
-			ObjectBindings[i]->BuildSequencerAddMenu(MenuBuilder);
+			ObjectBinding->BuildSequencerAddMenu(MenuBuilder);
 		}
 	}
 }
diff --git a/Engine/Source/Editor/Sequencer/Private/Sequencer.h b/Engine/Source/Editor/Sequencer/Private/Sequencer.h
--- a/Engine/Source/Editor/Sequencer/Private/Sequencer.h
+++ b/Engine/Source/Editor/Sequencer/Private/Sequencer.h
@@ -1100,6 +1100,9 @@ public:
 	/** Handles adding a new folder to the outliner tree. */
 	void AddFolder();
 
+	/** Adds a time warp track to the outliner tree. */
+	void AddTimeWarp();
+
 private:
 
 	/** Updates viewport clients' actor locks if they relate to sequencer cameras */
diff --git a/Engine/Source/Editor/Sequencer/Public/ISequencerEditorObjectBinding.h b/Engine/Source/Editor/Sequencer/Public/ISequencerEditorObjectBinding.h
--- a/Engine/Source/Editor/Sequencer/Public/ISequencerEditorObjectBinding.h
+++ b/Engine/Source/Editor/Sequencer/Public/ISequencerEditorObjectBinding.h
@@ -14,6 +14,9 @@ class ISequencerEditorObjectBinding
 {
 public:
 
+	/** The display name text to be used for UX */
+	virtual FText GetDisplayName() const = 0;
+
 	/**
 	 * Builds up the sequencer's "Add" menu.
 	 *
diff --git a/Engine/Source/Editor/Sequencer/Public/ISequencerTrackEditor.h b/Engine/Source/Editor/Sequencer/Public/ISequencerTrackEditor.h
--- a/Engine/Source/Editor/Sequencer/Public/ISequencerTrackEditor.h
+++ b/Engine/Source/Editor/Sequencer/Public/ISequencerTrackEditor.h
@@ -103,6 +103,9 @@ class ISequencerTrackEditor
 {
 public:
 
+	/** The display name text to use in the UX. Currently only used for sorting. */
+	virtual FText GetDisplayName() const = 0;
+
 	/**
 	 * Add keys for the following sections based on an external value if possible
 	 *
diff --git a/Engine/Source/Editor/Sequencer/Public/MovieSceneTrackEditor.h b/Engine/Source/Editor/Sequencer/Public/MovieSceneTrackEditor.h
--- a/Engine/Source/Editor/Sequencer/Public/MovieSceneTrackEditor.h
+++ b/Engine/Source/Editor/Sequencer/Public/MovieSceneTrackEditor.h
@@ -171,6 +171,8 @@ public:
 
 	// ISequencerTrackEditor interface
 
+	virtual FText GetDisplayName() const override;
+
 	virtual UMovieSceneTrack* AddTrack(UMovieScene* FocusedMovieScene, const FGuid& ObjectHandle, TSubclassOf<class UMovieSceneTrack> TrackClass, FName UniqueTypeName) override;
 
 	virtual TSharedRef<ISequencerSection> MakeSectionInterface(UMovieSceneSection& SectionObject, UMovieSceneTrack& Track, FGuid ObjectBinding) override;
diff --git a/Engine/Source/Editor/UMGEditor/Private/Animation/MarginTrackEditor.cpp b/Engine/Source/Editor/UMGEditor/Private/Animation/MarginTrackEditor.cpp
--- a/Engine/Source/Editor/UMGEditor/Private/Animation/MarginTrackEditor.cpp
+++ b/Engine/Source/Editor/UMGEditor/Private/Animation/MarginTrackEditor.cpp
@@ -16,6 +16,8 @@ class ISequencer;
 class ISequencerTrackEditor;
 class UMovieSceneSection;
 
+#define LOCTEXT_NAMESPACE "FMarginTrackEditor"
+
 FName FMarginTrackEditor::LeftName( "Left" );
 FName FMarginTrackEditor::TopName( "Top" );
 FName FMarginTrackEditor::RightName( "Right" );
@@ -26,6 +28,11 @@ TSharedRef<ISequencerTrackEditor> FMarginTrackEditor::CreateTrackEditor( TShared
 	return MakeShareable( new FMarginTrackEditor( InSequencer ) );
 }
 
+FText FMarginTrackEditor::GetDisplayName() const
+{
+	return LOCTEXT("MarginTrackEditor_DisplayName", "Margin");
+}
+
 void FMarginTrackEditor::GenerateKeysFromPropertyChanged( const FPropertyChangedParams& PropertyChangedParams, UMovieSceneSection* SectionToKey, FGeneratedTrackKeys& OutGeneratedKeys)
 {
 	FPropertyPath StructPath = PropertyChangedParams.StructPathToKey;
@@ -43,3 +50,5 @@ void FMarginTrackEditor::GenerateKeysFromPropertyChanged( const FPropertyChanged
 	OutGeneratedKeys.Add(FMovieSceneChannelValueSetter::Create<FMovieSceneFloatChannel>(2, Margin.Right,  bKeyRight));
 	OutGeneratedKeys.Add(FMovieSceneChannelValueSetter::Create<FMovieSceneFloatChannel>(3, Margin.Bottom, bKeyBottom));
 }
+
+#undef LOCTEXT_NAMESPACE
diff --git a/Engine/Source/Editor/UMGEditor/Private/Animation/MarginTrackEditor.h b/Engine/Source/Editor/UMGEditor/Private/Animation/MarginTrackEditor.h
--- a/Engine/Source/Editor/UMGEditor/Private/Animation/MarginTrackEditor.h
+++ b/Engine/Source/Editor/UMGEditor/Private/Animation/MarginTrackEditor.h
@@ -56,7 +56,7 @@ public:
 protected:
 
 	//~ FPropertyTrackEditor Interface
-
+	virtual FText GetDisplayName() const override;
 	virtual void GenerateKeysFromPropertyChanged( const FPropertyChangedParams& PropertyChangedParams, UMovieSceneSection* SectionToKey, FGeneratedTrackKeys& OutGeneratedKeys) override;
 
 private:
diff --git a/Engine/Source/Editor/UMGEditor/Private/Animation/Sequencer2DTransformTrackEditor.cpp b/Engine/Source/Editor/UMGEditor/Private/Animation/Sequencer2DTransformTrackEditor.cpp
--- a/Engine/Source/Editor/UMGEditor/Private/Animation/Sequencer2DTransformTrackEditor.cpp
+++ b/Engine/Source/Editor/UMGEditor/Private/Animation/Sequencer2DTransformTrackEditor.cpp
@@ -39,6 +39,8 @@
 class ISequencerTrackEditor;
 struct FMovieSceneChannel;
 
+#define LOCTEXT_NAMESPACE "2DTransformTrackEditor"
+
 FName F2DTransformTrackEditor::TranslationName( "Translation" );
 FName F2DTransformTrackEditor::ScaleName( "Scale" );
 FName F2DTransformTrackEditor::ShearName( "Shear" );
@@ -51,6 +53,10 @@ TSharedRef<ISequencerTrackEditor> F2DTransformTrackEditor::CreateTrackEditor( TS
 	return MakeShareable( new F2DTransformTrackEditor( InSequencer ) );
 }
 
+FText F2DTransformTrackEditor::GetDisplayName() const
+{
+	return LOCTEXT("2DTransformTrackEditor_DisplayName", "2D Transform");
+}
 
 TSharedRef<ISequencerSection> F2DTransformTrackEditor::MakeSectionInterface(UMovieSceneSection& SectionObject, UMovieSceneTrack& Track, FGuid ObjectBinding)
 {
@@ -234,4 +240,6 @@ void F2DTransformTrackEditor::ProcessKeyOperation(UObject* ObjectToKey, TArrayVi
 			}
 		}
 	}
-}
\ No newline at end of file
+}
+
+#undef LOCTEXT_NAMESPACE
diff --git a/Engine/Source/Editor/UMGEditor/Private/Animation/Sequencer2DTransformTrackEditor.h b/Engine/Source/Editor/UMGEditor/Private/Animation/Sequencer2DTransformTrackEditor.h
--- a/Engine/Source/Editor/UMGEditor/Private/Animation/Sequencer2DTransformTrackEditor.h
+++ b/Engine/Source/Editor/UMGEditor/Private/Animation/Sequencer2DTransformTrackEditor.h
@@ -63,7 +63,7 @@ public:
 protected:
 
 	// ISequencerTrackEditor interface
-
+	virtual FText GetDisplayName() const override;
 	virtual TSharedRef<ISequencerSection> MakeSectionInterface( UMovieSceneSection& SectionObject, UMovieSceneTrack& Track, FGuid ObjectBinding ) override;
 	virtual void ProcessKeyOperation(FFrameNumber InKeyTime, const UE::Sequencer::FKeyOperation& Operation, ISequencer& InSequencer, TArray<UE::Sequencer::FAddKeyResult>* OutResults = nullptr) override;
 	// FPropertyTrackEditor interface
