Curve Editor: Inject auto-keying buttons from sequencer to curve editor's toolbar

diff --git a/Engine/Source/Editor/CurveEditor/Private/CurveEditor.cpp b/Engine/Source/Editor/CurveEditor/Private/CurveEditor.cpp
--- a/Engine/Source/Editor/CurveEditor/Private/CurveEditor.cpp
+++ b/Engine/Source/Editor/CurveEditor/Private/CurveEditor.cpp
@@ -108,6 +108,7 @@ void FCurveEditor::InitCurveEditor(const FCurveEditorInitParams& InInitParams)
 	// we have to manually reach out to the module and get a list of extensions to create an instance of them.
 	// If none of your extensions are showing up, it's because you forgot to call this function after construction
 	// We're not allowed to use SharedThis(...) in a Constructor so it must exist as a separate function call.
+	EditorExtensions.Append(InInitParams.AdditionalEditorExtensions);
 	TArrayView<const FOnCreateCurveEditorExtension> Extensions = CurveEditorModule.GetEditorExtensions();
 	for (int32 DelegateIndex = 0; DelegateIndex < Extensions.Num(); ++DelegateIndex)
 	{
@@ -126,6 +127,7 @@ void FCurveEditor::InitCurveEditor(const FCurveEditorInitParams& InInitParams)
 		// We call a delegate and have the delegate create the instance to cover cross-module
 		AddTool(Tools[DelegateIndex].Execute(SharedThis(this)));
 	}
+	
 	SuspendBroadcastCount = 0;
 	// Listen to global undo so we can fix up our selection state for keys that no longer exist.
 	GEditor->RegisterForUndo(this);
diff --git a/Engine/Source/Editor/CurveEditor/Public/ICurveEditorModule.h b/Engine/Source/Editor/CurveEditor/Public/ICurveEditorModule.h
--- a/Engine/Source/Editor/CurveEditor/Public/ICurveEditorModule.h
+++ b/Engine/Source/Editor/CurveEditor/Public/ICurveEditorModule.h
@@ -33,6 +33,9 @@ struct FCurveEditorInitParams
 {
 	FCurveEditorInitParams()
 	{}
+	
+	/** Extensions you want to inject for this curve editor. */
+	TArray<TSharedRef<ICurveEditorExtension>> AdditionalEditorExtensions;
 };
 
 /**
diff --git a/Engine/Source/Editor/Sequencer/Private/MVVM/CurveEditorExtension.cpp b/Engine/Source/Editor/Sequencer/Private/MVVM/CurveEditorExtension.cpp
--- a/Engine/Source/Editor/Sequencer/Private/MVVM/CurveEditorExtension.cpp
+++ b/Engine/Source/Editor/Sequencer/Private/MVVM/CurveEditorExtension.cpp
@@ -10,6 +10,7 @@
 #include "IStructureDetailsView.h"
 #include "MVVM/ViewModels/SequencerEditorViewModel.h"
 #include "CurveEditorAxis.h"
+#include "ICurveEditorExtension.h"
 #include "SCurveEditorView.h"
 #include "SCurveEditorPanel.h"
 #include "SCurveEditorToolProperties.h"
@@ -18,6 +19,7 @@
 #include "STemporarilyFocusedSpinBox.h"
 #include "Sequencer.h"
 #include "SequencerCommands.h"
+#include "Menus/SequencerToolbarUtils.h"
 #include "Toolkits/IToolkitHost.h"
 #include "Tree/SCurveEditorTree.h"
 #include "Tree/SCurveEditorTreeFilterStatusBar.h"
@@ -224,10 +226,37 @@ struct FSequencerCurveEditorBounds : ICurveEditorBounds
 	TWeakPtr<FSequencer> WeakSequencer;
 };
 
+class FSequencerCurveEditorToolbarExtender : public ICurveEditorExtension
+{
+	TWeakPtr<FSequencer> WeakSequencer;
+public:
+
+	explicit FSequencerCurveEditorToolbarExtender(TWeakPtr<FSequencer> InWeakSequencer) : WeakSequencer(MoveTemp(InWeakSequencer)) {}
+	
+	virtual void BindCommands(TSharedRef<FUICommandList> CommandBindings) override {}
+	virtual TSharedPtr<FExtender> MakeToolbarExtender(const TSharedRef<FUICommandList>& InCommandList) override
+	{
+		TSharedRef<FExtender> Extender = MakeShared<FExtender>();
+		Extender->AddToolBarExtension("Adjustment", EExtensionHook::After, InCommandList,
+			FToolBarExtensionDelegate::CreateLambda([this](FToolBarBuilder& ToolbarBuilder)
+			{
+				TSharedPtr<FSequencer> SequenerPin = WeakSequencer.Pin();
+				
+				ToolbarBuilder.BeginSection("Keying");
+				ToolbarBuilder.PushCommandList(SequenerPin->GetCommandBindings().ToSharedRef());
+				AppendSequencerToolbarEntries(SequenerPin, ToolbarBuilder);
+				ToolbarBuilder.PopCommandList();
+				ToolbarBuilder.EndSection();
+			}));
+		return Extender;
+	}
+};
+
 const FName FCurveEditorExtension::CurveEditorTabName = FName(TEXT("SequencerGraphEditor"));
 
 FCurveEditorExtension::FCurveEditorExtension()
 {
+	
 }
 
 void FCurveEditorExtension::OnCreated(TSharedRef<FViewModel> InWeakOwner)
@@ -260,6 +289,7 @@ void FCurveEditorExtension::CreateCurveEditor(const FTimeSliderArgs& TimeSliderA
 	// Create the curve editor;
 	{
 		FCurveEditorInitParams CurveEditorInitParams;
+		CurveEditorInitParams.AdditionalEditorExtensions = { MakeShared<FSequencerCurveEditorToolbarExtender>(Sequencer.ToWeakPtr()) };
 
 		USequencerSettings* SequencerSettings = Sequencer->GetSequencerSettings();
 
diff --git a/Engine/Source/Editor/Sequencer/Private/Menus/SequencerToolbarUtils.cpp b/Engine/Source/Editor/Sequencer/Private/Menus/SequencerToolbarUtils.cpp
new file mode 100644
--- /dev/null
+++ b/Engine/Source/Editor/Sequencer/Private/Menus/SequencerToolbarUtils.cpp
@@ -0,0 +1,393 @@
+﻿// Copyright Epic Games, Inc. All Rights Reserved.
+
+#include "SequencerToolbarUtils.h"
+
+#include "IVREditorModule.h"
+#include "LevelSequence.h"
+#include "Misc/Attribute.h"
+#include "Misc/Optional.h"
+#include "Sequencer.h"
+#include "SequencerCommands.h"
+#include "SequencerSettings.h"
+#include "Templates/SharedPointer.h"
+#include "Textures/SlateIcon.h"
+
+#define LOCTEXT_NAMESPACE "Sequencer"
+
+namespace UE::Sequencer
+{
+const FName SequencerToolbarStyleName("SequencerToolbar");
+	
+namespace ToolbarUtilsDetail
+{
+static TSharedRef<SWidget> MakeKeyGroupMenu(const TWeakPtr<FSequencer> InWeakSequencer)
+{
+	const bool bShouldCloseWindowAfterMenuSelection = true;
+	FMenuBuilder MenuBuilder(bShouldCloseWindowAfterMenuSelection, InWeakSequencer.Pin()->GetCommandBindings());
+
+	if (InWeakSequencer.Pin()->IsLevelEditorSequencer())
+	{
+		MenuBuilder.AddMenuEntry(FSequencerCommands::Get().SetKeyAll);
+		MenuBuilder.AddMenuEntry(FSequencerCommands::Get().SetKeyGroup);
+		MenuBuilder.AddMenuEntry(FSequencerCommands::Get().SetKeyChanged);
+	}
+
+	// Interpolation
+	MenuBuilder.BeginSection("SequencerInterpolation", LOCTEXT("KeyInterpolationMenu", "Default Key Interpolation"));
+	{
+		MenuBuilder.AddMenuEntry(
+			LOCTEXT("SetKeyInterpolationSmartAuto", "Cubic (Smart Auto)"),
+			LOCTEXT("SetKeyInterpolationSmartAutoTooltip", "Set key interpolation to smart auto"),
+			FSlateIcon(FAppStyle::GetAppStyleSetName(), "Sequencer.IconKeySmartAuto"),
+			FUIAction(
+				FExecuteAction::CreateLambda([InWeakSequencer] { InWeakSequencer.Pin()->SetKeyInterpolation(EMovieSceneKeyInterpolation::SmartAuto); }),
+				FCanExecuteAction(),
+				FIsActionChecked::CreateLambda([InWeakSequencer] { return InWeakSequencer.Pin()->GetKeyInterpolation() == EMovieSceneKeyInterpolation::SmartAuto; })),
+			NAME_None,
+			EUserInterfaceActionType::RadioButton
+		);
+		MenuBuilder.AddMenuEntry(
+			LOCTEXT("SetKeyInterpolationAuto", "Cubic (Auto)"),
+			LOCTEXT("SetKeyInterpolationAutoTooltip", "Set key interpolation to auto"),
+			FSlateIcon(FAppStyle::GetAppStyleSetName(), "Sequencer.IconKeyAuto"),
+			FUIAction(
+				FExecuteAction::CreateLambda([InWeakSequencer] { InWeakSequencer.Pin()->SetKeyInterpolation(EMovieSceneKeyInterpolation::Auto); }),
+				FCanExecuteAction(),
+				FIsActionChecked::CreateLambda([InWeakSequencer] { return InWeakSequencer.Pin()->GetKeyInterpolation() == EMovieSceneKeyInterpolation::Auto; })),
+			NAME_None,
+			EUserInterfaceActionType::RadioButton
+		);
+
+		MenuBuilder.AddMenuEntry(
+			LOCTEXT("SetKeyInterpolationUser", "Cubic (User)"),
+			LOCTEXT("SetKeyInterpolationUserTooltip", "Set key interpolation to user"),
+			FSlateIcon(FAppStyle::GetAppStyleSetName(), "Sequencer.IconKeyUser"),
+			FUIAction(
+				FExecuteAction::CreateLambda([InWeakSequencer] { InWeakSequencer.Pin()->SetKeyInterpolation(EMovieSceneKeyInterpolation::User); }),
+				FCanExecuteAction(),
+				FIsActionChecked::CreateLambda([InWeakSequencer] { return InWeakSequencer.Pin()->GetKeyInterpolation() == EMovieSceneKeyInterpolation::User; })),
+			NAME_None,
+			EUserInterfaceActionType::RadioButton
+		);
+
+		MenuBuilder.AddMenuEntry(
+			LOCTEXT("SetKeyInterpolationBreak", "Cubic (Break)"),
+			LOCTEXT("SetKeyInterpolationBreakTooltip", "Set key interpolation to break"),
+			FSlateIcon(FAppStyle::GetAppStyleSetName(), "Sequencer.IconKeyBreak"),
+			FUIAction(
+				FExecuteAction::CreateLambda([InWeakSequencer] { InWeakSequencer.Pin()->SetKeyInterpolation(EMovieSceneKeyInterpolation::Break); }),
+				FCanExecuteAction(),
+				FIsActionChecked::CreateLambda([InWeakSequencer] { return InWeakSequencer.Pin()->GetKeyInterpolation() == EMovieSceneKeyInterpolation::Break; })),
+			NAME_None,
+			EUserInterfaceActionType::RadioButton
+		);
+
+		MenuBuilder.AddMenuEntry(
+			LOCTEXT("SetKeyInterpolationLinear", "Linear"),
+			LOCTEXT("SetKeyInterpolationLinearTooltip", "Set key interpolation to linear"),
+			FSlateIcon(FAppStyle::GetAppStyleSetName(), "Sequencer.IconKeyLinear"),
+			FUIAction(
+				FExecuteAction::CreateLambda([InWeakSequencer] { InWeakSequencer.Pin()->SetKeyInterpolation(EMovieSceneKeyInterpolation::Linear); }),
+				FCanExecuteAction(),
+				FIsActionChecked::CreateLambda([InWeakSequencer] { return InWeakSequencer.Pin()->GetKeyInterpolation() == EMovieSceneKeyInterpolation::Linear; })),
+			NAME_None,
+			EUserInterfaceActionType::RadioButton
+		);
+
+		MenuBuilder.AddMenuEntry(
+			LOCTEXT("SetKeyInterpolationConstant", "Constant"),
+			LOCTEXT("SetKeyInterpolationConstantTooltip", "Set key interpolation to constant"),
+			FSlateIcon(FAppStyle::GetAppStyleSetName(), "Sequencer.IconKeyConstant"),
+			FUIAction(
+				FExecuteAction::CreateLambda([InWeakSequencer] { InWeakSequencer.Pin()->SetKeyInterpolation(EMovieSceneKeyInterpolation::Constant); }),
+				FCanExecuteAction(),
+				FIsActionChecked::CreateLambda([InWeakSequencer] { return InWeakSequencer.Pin()->GetKeyInterpolation() == EMovieSceneKeyInterpolation::Constant; })),
+			NAME_None,
+			EUserInterfaceActionType::RadioButton
+		);
+	}
+	MenuBuilder.EndSection(); // SequencerInterpolation
+
+	return MenuBuilder.MakeWidget();
+}
+
+template<typename TCallback>
+static void GetKeyGroupMenuEntryArgs(const TWeakPtr<FSequencer>& InWeakSequencer, TCallback&& InConsumer)
+{
+	TAttribute<FSlateIcon> KeyGroupModeIcon;
+	KeyGroupModeIcon.Bind(TAttribute<FSlateIcon>::FGetter::CreateLambda([InWeakSequencer] 
+	{
+		switch (InWeakSequencer.Pin()->GetKeyGroupMode())
+		{
+		case EKeyGroupMode::KeyAll:
+			return FSequencerCommands::Get().SetKeyAll->GetIcon();
+		case EKeyGroupMode::KeyGroup:
+			return FSequencerCommands::Get().SetKeyGroup->GetIcon();
+		default: // EKeyGroupMode::KeyChanged
+			return FSequencerCommands::Get().SetKeyChanged->GetIcon();
+		}
+	}));
+
+	TAttribute<FText> KeyGroupModeToolTip;
+	KeyGroupModeToolTip.Bind(TAttribute<FText>::FGetter::CreateLambda([InWeakSequencer] 
+	{
+		switch (InWeakSequencer.Pin()->GetKeyGroupMode())
+		{
+		case EKeyGroupMode::KeyAll:
+			return FSequencerCommands::Get().SetKeyAll->GetDescription();
+		case EKeyGroupMode::KeyGroup:
+			return FSequencerCommands::Get().SetKeyGroup->GetDescription();
+		default: // EKeyGroupMode::KeyChanged
+			return FSequencerCommands::Get().SetKeyChanged->GetDescription();
+		}
+	}));
+
+	InConsumer(
+		"KeyGroup", FUIAction(), FOnGetContent::CreateStatic(&MakeKeyGroupMenu, InWeakSequencer),
+		LOCTEXT("KeyGroup", "Key All"), KeyGroupModeToolTip, KeyGroupModeIcon
+		);
+}
+}
+	
+FToolMenuEntry MakeKeyGroupMenuEntry_ToolMenus(const TWeakPtr<FSequencer>& InWeakSequencer)
+{
+	FToolMenuEntry KeyGroupEntry;
+	ToolbarUtilsDetail::GetKeyGroupMenuEntryArgs(InWeakSequencer,
+		[&KeyGroupEntry](const FName InName, const FUIAction& InAction, const FOnGetContent& InMenuContentGenerator,
+		const TAttribute<FText>& InLabelOverride, const TAttribute<FText>& ToolTipAttr, const TAttribute<FSlateIcon>& IconAttr)
+	{
+		KeyGroupEntry = FToolMenuEntry::InitComboButton(
+			InName, InAction, InMenuContentGenerator, InLabelOverride, ToolTipAttr, IconAttr
+		);
+		KeyGroupEntry.StyleNameOverride = SequencerToolbarStyleName;
+	});
+	return KeyGroupEntry;
+}
+	
+namespace ToolbarUtilsDetail
+{
+static TSharedRef<SWidget> MakeAutoChangeMenu(const TWeakPtr<FSequencer> InWeakSequencer)
+{
+	FMenuBuilder MenuBuilder(false, InWeakSequencer.Pin()->GetCommandBindings());
+
+	MenuBuilder.AddMenuEntry(FSequencerCommands::Get().SetAutoKey);
+
+	if (InWeakSequencer.Pin()->IsLevelEditorSequencer())
+	{
+		MenuBuilder.AddMenuEntry(FSequencerCommands::Get().SetAutoTrack);
+	}
+
+	if (IVREditorModule::Get().IsVREditorModeActive()
+		|| (InWeakSequencer.Pin()->IsLevelEditorSequencer() && !ExactCast<ULevelSequence>(InWeakSequencer.Pin()->GetFocusedMovieSceneSequence())))
+	{
+		MenuBuilder.AddMenuEntry(FSequencerCommands::Get().SetAutoChangeAll);
+	}
+
+	MenuBuilder.AddMenuEntry(FSequencerCommands::Get().SetAutoChangeNone);
+
+	return MenuBuilder.MakeWidget();
+}
+
+template<typename TCallback, typename TCallback2>
+static void GetAutoKeyMenuEntryArgs(const TSharedPtr<FSequencer>& InSequencer, TCallback&& InConsumer, TCallback2&& InConsumer2)
+{
+	const TWeakPtr<FSequencer> WeakSequencer = InSequencer.ToWeakPtr();
+	if (IVREditorModule::Get().IsVREditorModeActive()
+		|| (InSequencer->IsLevelEditorSequencer() && ExactCast<ULevelSequence>(InSequencer->GetFocusedMovieSceneSequence()) == nullptr))
+	{
+		TAttribute<FSlateIcon> AutoChangeModeIcon;
+		AutoChangeModeIcon.Bind(TAttribute<FSlateIcon>::FGetter::CreateLambda([WeakSequencer] 
+		{
+			switch ( WeakSequencer.Pin()->GetAutoChangeMode() )
+			{
+			case EAutoChangeMode::AutoKey:
+				return FSequencerCommands::Get().SetAutoKey->GetIcon();
+			case EAutoChangeMode::AutoTrack:
+				return FSequencerCommands::Get().SetAutoTrack->GetIcon();
+			case EAutoChangeMode::All:
+				return FSequencerCommands::Get().SetAutoChangeAll->GetIcon();
+			default: // EAutoChangeMode::None
+				return FSequencerCommands::Get().SetAutoChangeNone->GetIcon();
+			}
+		} ) );
+
+		TAttribute<FText> AutoChangeModeToolTip;
+		AutoChangeModeToolTip.Bind( TAttribute<FText>::FGetter::CreateLambda([WeakSequencer] 
+		{
+			switch ( WeakSequencer.Pin()->GetAutoChangeMode() )
+			{
+			case EAutoChangeMode::AutoKey:
+				return FSequencerCommands::Get().SetAutoKey->GetDescription();
+			case EAutoChangeMode::AutoTrack:
+				return FSequencerCommands::Get().SetAutoTrack->GetDescription();
+			case EAutoChangeMode::All:
+				return FSequencerCommands::Get().SetAutoChangeAll->GetDescription();
+			default: // EAutoChangeMode::None
+				return FSequencerCommands::Get().SetAutoChangeNone->GetDescription();
+			}
+		} ) );
+		
+		InConsumer(
+			"AutoChange", FUIAction(), FOnGetContent::CreateStatic(&ToolbarUtilsDetail::MakeAutoChangeMenu, WeakSequencer),
+			LOCTEXT("AutoChangeMode", "Auto-Change Mode"), AutoChangeModeToolTip, AutoChangeModeIcon
+		);
+	}
+	else	
+	{
+		InConsumer2(FSequencerCommands::Get().ToggleAutoKeyEnabled);
+	}
+}
+}
+	
+FToolMenuEntry MakeAutoKeyMenuEntry(const TSharedPtr<FSequencer>& InSequencer)
+{
+	FToolMenuEntry KeyGroupEntry;
+	ToolbarUtilsDetail::GetAutoKeyMenuEntryArgs(InSequencer,
+		[&KeyGroupEntry](const FName InName, const FUIAction& InAction, const FOnGetContent& InMenuContentGenerator,
+		const TAttribute<FText>& InLabelOverride, const TAttribute<FText>& ToolTipAttr, const TAttribute<FSlateIcon>& IconAttr)
+		{
+		KeyGroupEntry = FToolMenuEntry::InitComboButton(
+			InName, InAction, InMenuContentGenerator, InLabelOverride, ToolTipAttr, IconAttr
+		);
+		KeyGroupEntry.StyleNameOverride = SequencerToolbarStyleName;
+		},
+		[&KeyGroupEntry](TSharedPtr<FUICommandInfo> ToggleAutoKeyEnabled)
+		{
+			KeyGroupEntry = FToolMenuEntry::InitMenuEntry(ToggleAutoKeyEnabled);
+		});
+	KeyGroupEntry.StyleNameOverride = SequencerToolbarStyleName;
+	return KeyGroupEntry;
+}
+
+namespace ToolbarUtilsDetail
+{
+static TSharedRef<SWidget> MakeAllowEditsMenu(const TWeakPtr<FSequencer> InWeakSequencer)
+{
+	const bool bShouldCloseWindowAfterMenuSelection = true;
+	FMenuBuilder MenuBuilder(bShouldCloseWindowAfterMenuSelection, InWeakSequencer.Pin()->GetCommandBindings());
+
+	MenuBuilder.AddMenuEntry(FSequencerCommands::Get().AllowAllEdits);
+	MenuBuilder.AddMenuEntry(FSequencerCommands::Get().AllowSequencerEditsOnly);
+	MenuBuilder.AddMenuEntry(FSequencerCommands::Get().AllowLevelEditsOnly);
+
+	return MenuBuilder.MakeWidget();
+}
+
+template<typename TCallback>
+static bool GetAllowEditsModeMenuArgs(const TSharedPtr<FSequencer>& InSequencer, TCallback&& InConsumer)
+{
+	const TWeakPtr<FSequencer> WeakSequencer = InSequencer.ToWeakPtr();
+	if (InSequencer->IsLevelEditorSequencer())
+	{
+		TAttribute<FSlateIcon> AllowEditsModeIcon;
+		AllowEditsModeIcon.Bind(TAttribute<FSlateIcon>::FGetter::CreateLambda([WeakSequencer] 
+		{
+			switch ( WeakSequencer.Pin()->GetAllowEditsMode() )
+			{
+			case EAllowEditsMode::AllEdits:
+				return FSequencerCommands::Get().AllowAllEdits->GetIcon();
+			case EAllowEditsMode::AllowSequencerEditsOnly:
+				return FSequencerCommands::Get().AllowSequencerEditsOnly->GetIcon();
+			default: // EAllowEditsMode::AllowLevelEditsOnly
+				return FSequencerCommands::Get().AllowLevelEditsOnly->GetIcon();
+			}
+		} ) );
+
+		TAttribute<FText> AllowEditsModeToolTip;
+		AllowEditsModeToolTip.Bind( TAttribute<FText>::FGetter::CreateLambda([WeakSequencer] 
+		{
+			switch ( WeakSequencer.Pin()->GetAllowEditsMode() )
+			{
+			case EAllowEditsMode::AllEdits:
+				return FSequencerCommands::Get().AllowAllEdits->GetDescription();
+			case EAllowEditsMode::AllowSequencerEditsOnly:
+				return FSequencerCommands::Get().AllowSequencerEditsOnly->GetDescription();
+			default: // EAllowEditsMode::AllowLevelEditsOnly
+				return FSequencerCommands::Get().AllowLevelEditsOnly->GetDescription();
+			}
+		} ) );
+
+		InConsumer(
+			"AllowEditsMode",
+			FUIAction(),
+			FOnGetContent::CreateStatic(&ToolbarUtilsDetail::MakeAllowEditsMenu, WeakSequencer),
+			LOCTEXT("AllowEditsMode", "Allow Edits"),
+			AllowEditsModeToolTip,
+			AllowEditsModeIcon
+		);
+		return true;
+	}
+	return false;
+}
+}
+	
+TOptional<FToolMenuEntry> MakeAllowEditsModeMenuEntry(const TSharedPtr<FSequencer>& InSequencer)
+{
+	FToolMenuEntry KeyGroupEntry;
+	const bool bMade = ToolbarUtilsDetail::GetAllowEditsModeMenuArgs(InSequencer,
+		[&KeyGroupEntry](const FName InName, const FUIAction& InAction, const FOnGetContent& InMenuContentGenerator,
+		const TAttribute<FText>& InLabelOverride, const TAttribute<FText>& ToolTipAttr, const TAttribute<FSlateIcon>& IconAttr)
+	{
+		KeyGroupEntry = FToolMenuEntry::InitComboButton(
+			InName, InAction, InMenuContentGenerator, InLabelOverride, ToolTipAttr, IconAttr
+		);
+		KeyGroupEntry.StyleNameOverride = SequencerToolbarStyleName;
+	});
+	if (bMade)
+	{
+		return KeyGroupEntry;
+	}
+	return {};
+}
+
+void AppendSequencerToolbarEntries(const TSharedPtr<FSequencer>& InSequencer, FToolBarBuilder& InToolbarBuilder, const EToolbarItemFlags InFlags)
+{
+	const TWeakPtr<FSequencer> WeakSequencer = InSequencer.ToWeakPtr();
+
+	InToolbarBuilder.BeginStyleOverride(SequencerToolbarStyleName);
+	
+	if (EnumHasAnyFlags(InFlags, EToolbarItemFlags::KeyGroup))
+	{
+		ToolbarUtilsDetail::GetKeyGroupMenuEntryArgs(WeakSequencer,
+			[&InToolbarBuilder](const FName, const FUIAction& InAction, const FOnGetContent& InMenuContentGenerator,
+			const TAttribute<FText>& InLabelOverride, const TAttribute<FText>& ToolTipAttr, const TAttribute<FSlateIcon>& IconAttr)
+		{
+			InToolbarBuilder.AddComboButton(InAction, InMenuContentGenerator, InLabelOverride, ToolTipAttr, IconAttr);
+		});
+	}
+	
+	if (EnumHasAnyFlags(InFlags, EToolbarItemFlags::AutoKey))
+	{
+		ToolbarUtilsDetail::GetAutoKeyMenuEntryArgs(InSequencer,
+			[&InToolbarBuilder](const FName, const FUIAction& InAction, const FOnGetContent& InMenuContentGenerator,
+			const TAttribute<FText>& InLabelOverride, const TAttribute<FText>& ToolTipAttr, const TAttribute<FSlateIcon>& IconAttr)
+			{
+				FButtonArgs ButtonArgs;
+				ButtonArgs.Action = InAction;
+				ButtonArgs.OnGetMenuContent = InMenuContentGenerator;
+				ButtonArgs.LabelOverride = InLabelOverride;
+				ButtonArgs.ToolTipOverride = ToolTipAttr;
+				ButtonArgs.IconOverride = IconAttr;
+				InToolbarBuilder.AddToolBarButton(ButtonArgs);
+			},
+			[&InToolbarBuilder](TSharedPtr<FUICommandInfo> ToggleAutoKeyEnabled)
+			{
+				InToolbarBuilder.AddToolBarButton(ToggleAutoKeyEnabled);
+			});
+	}
+	
+	if (EnumHasAnyFlags(InFlags, EToolbarItemFlags::AllowEditsMode))
+	{
+		ToolbarUtilsDetail::GetAllowEditsModeMenuArgs(InSequencer,
+			[&InToolbarBuilder](const FName, const FUIAction& InAction, const FOnGetContent& InMenuContentGenerator,
+			const TAttribute<FText>& InLabelOverride, const TAttribute<FText>& ToolTipAttr, const TAttribute<FSlateIcon>& IconAttr)
+		{
+			InToolbarBuilder.AddComboButton(InAction, InMenuContentGenerator, InLabelOverride, ToolTipAttr, IconAttr);
+		});
+	}
+	
+	InToolbarBuilder.EndStyleOverride();
+}
+}
+
+#undef LOCTEXT_NAMESPACE
\ No newline at end of file
diff --git a/Engine/Source/Editor/Sequencer/Private/Menus/SequencerToolbarUtils.h b/Engine/Source/Editor/Sequencer/Private/Menus/SequencerToolbarUtils.h
new file mode 100644
--- /dev/null
+++ b/Engine/Source/Editor/Sequencer/Private/Menus/SequencerToolbarUtils.h
@@ -0,0 +1,47 @@
+﻿// Copyright Epic Games, Inc. All Rights Reserved.
+
+#pragma once
+
+#include "HAL/Platform.h"
+#include "Misc/EnumClassFlags.h"
+#include "Templates//SharedPointerFwd.h"
+#include "UObject/NameTypes.h"
+
+class FSequencer;
+class FToolBarBuilder;
+struct FButtonArgs;
+struct FToolMenuEntry;
+template<typename OptionalType>struct TOptional;
+
+namespace UE::Sequencer
+{
+extern const FName SequencerToolbarStyleName;
+
+/** Makes the combo button for changing the key group settings. */
+FToolMenuEntry MakeKeyGroupMenuEntry_ToolMenus(const TWeakPtr<FSequencer>& InWeakSequencer);
+	
+/** Makes the button for toggling auto keying. */
+FToolMenuEntry MakeAutoKeyMenuEntry(const TSharedPtr<FSequencer>& InSequencer);
+	
+/** Makes the combo button for changing the types of edits driving auto keying. */
+TOptional<FToolMenuEntry> MakeAllowEditsModeMenuEntry(const TSharedPtr<FSequencer>& InSequencer);
+
+enum class EToolbarItemFlags : uint32
+{
+	None = 0,
+	
+	KeyGroup = 1 << 0,
+	AutoKey = 1 << 1,
+	AllowEditsMode = 1 << 2,
+	
+	All = KeyGroup | AutoKey | AllowEditsMode
+};
+ENUM_CLASS_FLAGS(EToolbarItemFlags);
+/**
+ * Appends the flagged items to InToolbarBuilder.
+ * This function is effectively an adapter for converting the UToolMenus API to FToolBarBuilder. It can e.g. be used to inject items to CurveEditor.
+ */
+void AppendSequencerToolbarEntries(
+	const TSharedPtr<FSequencer>& InSequencer, FToolBarBuilder& InToolbarBuilder, const EToolbarItemFlags InFlags = EToolbarItemFlags::All
+	);
+}
diff --git a/Engine/Source/Editor/Sequencer/Private/SSequencer.cpp b/Engine/Source/Editor/Sequencer/Private/SSequencer.cpp
--- a/Engine/Source/Editor/Sequencer/Private/SSequencer.cpp
+++ b/Engine/Source/Editor/Sequencer/Private/SSequencer.cpp
@@ -141,6 +141,7 @@
 #include "Widgets/SSequencerSearchBox.h"
 #include "IContentBrowserSingleton.h"
 #include "ContentBrowserModule.h"
+#include "Menus/SequencerToolbarUtils.h"
 #include "Settings/EditorExperimentalSettings.h"
 
 #define LOCTEXT_NAMESPACE "Sequencer"
@@ -1439,7 +1440,7 @@ TSharedRef<SWidget> SSequencer::MakeToolBar()
 
 void SSequencer::PopulateToolBar(UToolMenu* InMenu)
 {
-	const FName SequencerToolbarStyleName = "SequencerToolbar";
+	const FName SequencerToolbarStyleName = UE::Sequencer::SequencerToolbarStyleName;
 	
 	USequencerToolMenuContext* ContextObject = InMenu->FindContext<USequencerToolMenuContext>();
 	if (!ContextObject)
@@ -1566,138 +1567,12 @@ void SSequencer::PopulateToolBar(UToolMenu* InMenu)
 		Section.AddEntry(PlaybackOptionsEntry);
 
 		Section.AddSeparator(NAME_None);
-
-		TAttribute<FSlateIcon> KeyGroupModeIcon;
-		KeyGroupModeIcon.Bind(TAttribute<FSlateIcon>::FGetter::CreateLambda([WeakSequencer] 
-		{
-			switch (WeakSequencer.Pin()->GetKeyGroupMode())
-			{
-			case EKeyGroupMode::KeyAll:
-				return FSequencerCommands::Get().SetKeyAll->GetIcon();
-			case EKeyGroupMode::KeyGroup:
-				return FSequencerCommands::Get().SetKeyGroup->GetIcon();
-			default: // EKeyGroupMode::KeyChanged
-				return FSequencerCommands::Get().SetKeyChanged->GetIcon();
-			}
-		}));
-
-		TAttribute<FText> KeyGroupModeToolTip;
-		KeyGroupModeToolTip.Bind(TAttribute<FText>::FGetter::CreateLambda([WeakSequencer] 
-		{
-			switch (WeakSequencer.Pin()->GetKeyGroupMode())
-			{
-			case EKeyGroupMode::KeyAll:
-				return FSequencerCommands::Get().SetKeyAll->GetDescription();
-			case EKeyGroupMode::KeyGroup:
-				return FSequencerCommands::Get().SetKeyGroup->GetDescription();
-			default: // EKeyGroupMode::KeyChanged
-				return FSequencerCommands::Get().SetKeyChanged->GetDescription();
-			}
-		}));
-
-		FToolMenuEntry KeyGroupEntry = FToolMenuEntry::InitComboButton(
-			"KeyGroup",
-			FUIAction(),
-			FOnGetContent::CreateSP(SequencerWidget, &SSequencer::MakeKeyGroupMenu),
-			LOCTEXT("KeyGroup", "Key All"),
-			KeyGroupModeToolTip,
-			KeyGroupModeIcon
-		);
-		KeyGroupEntry.StyleNameOverride = SequencerToolbarStyleName;
-		Section.AddEntry(KeyGroupEntry);
-	
-		if (IVREditorModule::Get().IsVREditorModeActive() || (Sequencer->IsLevelEditorSequencer() && ExactCast<ULevelSequence>(Sequencer->GetFocusedMovieSceneSequence()) == nullptr))
-		{
-			TAttribute<FSlateIcon> AutoChangeModeIcon;
-			AutoChangeModeIcon.Bind(TAttribute<FSlateIcon>::FGetter::CreateLambda([WeakSequencer] 
-			{
-				switch ( WeakSequencer.Pin()->GetAutoChangeMode() )
-				{
-				case EAutoChangeMode::AutoKey:
-					return FSequencerCommands::Get().SetAutoKey->GetIcon();
-				case EAutoChangeMode::AutoTrack:
-					return FSequencerCommands::Get().SetAutoTrack->GetIcon();
-				case EAutoChangeMode::All:
-					return FSequencerCommands::Get().SetAutoChangeAll->GetIcon();
-				default: // EAutoChangeMode::None
-					return FSequencerCommands::Get().SetAutoChangeNone->GetIcon();
-				}
-			} ) );
-
-			TAttribute<FText> AutoChangeModeToolTip;
-			AutoChangeModeToolTip.Bind( TAttribute<FText>::FGetter::CreateLambda([WeakSequencer] 
-			{
-				switch ( WeakSequencer.Pin()->GetAutoChangeMode() )
-				{
-				case EAutoChangeMode::AutoKey:
-					return FSequencerCommands::Get().SetAutoKey->GetDescription();
-				case EAutoChangeMode::AutoTrack:
-					return FSequencerCommands::Get().SetAutoTrack->GetDescription();
-				case EAutoChangeMode::All:
-					return FSequencerCommands::Get().SetAutoChangeAll->GetDescription();
-				default: // EAutoChangeMode::None
-					return FSequencerCommands::Get().SetAutoChangeNone->GetDescription();
-				}
-			} ) );
-			
-			FToolMenuEntry AutoChangeEntry = FToolMenuEntry::InitComboButton(
-				"AutoChange",
-				FUIAction(),
-				FOnGetContent::CreateSP(SequencerWidget, &SSequencer::MakeAutoChangeMenu),
-				LOCTEXT("AutoChangeMode", "Auto-Change Mode"),
-				AutoChangeModeToolTip,
-				AutoChangeModeIcon
-			);
-			AutoChangeEntry.StyleNameOverride = SequencerToolbarStyleName;
-			Section.AddEntry(AutoChangeEntry);
-		}
-		else	
+		
+		Section.AddEntry(UE::Sequencer::MakeKeyGroupMenuEntry_ToolMenus(WeakSequencer));
+		Section.AddEntry(UE::Sequencer::MakeAutoKeyMenuEntry(Sequencer));
+		if (const TOptional<FToolMenuEntry> AllowEditsEntry = UE::Sequencer::MakeAllowEditsModeMenuEntry(Sequencer))
 		{
-			FToolMenuEntry ToggleAutoKeyEntry = FToolMenuEntry::InitToolBarButton(FSequencerCommands::Get().ToggleAutoKeyEnabled);
-			ToggleAutoKeyEntry.StyleNameOverride = SequencerToolbarStyleName;
-			Section.AddEntry(ToggleAutoKeyEntry);
-		}
-
-		if (Sequencer->IsLevelEditorSequencer())
-		{
-			TAttribute<FSlateIcon> AllowEditsModeIcon;
-			AllowEditsModeIcon.Bind(TAttribute<FSlateIcon>::FGetter::CreateLambda([WeakSequencer] 
-			{
-				switch ( WeakSequencer.Pin()->GetAllowEditsMode() )
-				{
-				case EAllowEditsMode::AllEdits:
-					return FSequencerCommands::Get().AllowAllEdits->GetIcon();
-				case EAllowEditsMode::AllowSequencerEditsOnly:
-					return FSequencerCommands::Get().AllowSequencerEditsOnly->GetIcon();
-				default: // EAllowEditsMode::AllowLevelEditsOnly
-					return FSequencerCommands::Get().AllowLevelEditsOnly->GetIcon();
-				}
-			} ) );
-
-			TAttribute<FText> AllowEditsModeToolTip;
-			AllowEditsModeToolTip.Bind( TAttribute<FText>::FGetter::CreateLambda([WeakSequencer] 
-			{
-				switch ( WeakSequencer.Pin()->GetAllowEditsMode() )
-				{
-				case EAllowEditsMode::AllEdits:
-					return FSequencerCommands::Get().AllowAllEdits->GetDescription();
-				case EAllowEditsMode::AllowSequencerEditsOnly:
-					return FSequencerCommands::Get().AllowSequencerEditsOnly->GetDescription();
-				default: // EAllowEditsMode::AllowLevelEditsOnly
-					return FSequencerCommands::Get().AllowLevelEditsOnly->GetDescription();
-				}
-			} ) );
-
-			FToolMenuEntry AllowEditsModeEntry = FToolMenuEntry::InitComboButton(
-				"AllowEditsMode",
-				FUIAction(),
-				FOnGetContent::CreateSP(SequencerWidget, &SSequencer::MakeAllowEditsMenu),
-				LOCTEXT("AllowEditsMode", "Allow Edits"),
-				AllowEditsModeToolTip,
-				AllowEditsModeIcon
-			);
-			AllowEditsModeEntry.StyleNameOverride = SequencerToolbarStyleName;
-			Section.AddEntry(AllowEditsModeEntry);
+			Section.AddEntry(*AllowEditsEntry);
 		}
 		
 		Section.AddSeparator(NAME_None);
@@ -2633,133 +2508,6 @@ TSharedRef<SWidget> SSequencer::MakeSnapMenu()
 	return MenuBuilder.MakeWidget();
 }
 
-
-TSharedRef<SWidget> SSequencer::MakeAutoChangeMenu()
-{
-	FMenuBuilder MenuBuilder(false, SequencerPtr.Pin()->GetCommandBindings());
-
-	MenuBuilder.AddMenuEntry(FSequencerCommands::Get().SetAutoKey);
-
-	if (SequencerPtr.Pin()->IsLevelEditorSequencer())
-	{
-		MenuBuilder.AddMenuEntry(FSequencerCommands::Get().SetAutoTrack);
-	}
-
-	if (IVREditorModule::Get().IsVREditorModeActive() || (SequencerPtr.Pin()->IsLevelEditorSequencer() && ExactCast<ULevelSequence>(SequencerPtr.Pin()->GetFocusedMovieSceneSequence()) == nullptr))
-	{
-		MenuBuilder.AddMenuEntry(FSequencerCommands::Get().SetAutoChangeAll);
-	}
-
-	MenuBuilder.AddMenuEntry(FSequencerCommands::Get().SetAutoChangeNone);
-
-	return MenuBuilder.MakeWidget();
-}
-
-TSharedRef<SWidget> SSequencer::MakeAllowEditsMenu()
-{
-	const bool bShouldCloseWindowAfterMenuSelection = true;
-	FMenuBuilder MenuBuilder(bShouldCloseWindowAfterMenuSelection, SequencerPtr.Pin()->GetCommandBindings());
-
-	MenuBuilder.AddMenuEntry(FSequencerCommands::Get().AllowAllEdits);
-	MenuBuilder.AddMenuEntry(FSequencerCommands::Get().AllowSequencerEditsOnly);
-	MenuBuilder.AddMenuEntry(FSequencerCommands::Get().AllowLevelEditsOnly);
-
-	return MenuBuilder.MakeWidget();
-
-}
-
-TSharedRef<SWidget> SSequencer::MakeKeyGroupMenu()
-{
-	const bool bShouldCloseWindowAfterMenuSelection = true;
-	FMenuBuilder MenuBuilder(bShouldCloseWindowAfterMenuSelection, SequencerPtr.Pin()->GetCommandBindings());
-
-	if (SequencerPtr.Pin()->IsLevelEditorSequencer())
-	{
-		MenuBuilder.AddMenuEntry(FSequencerCommands::Get().SetKeyAll);
-		MenuBuilder.AddMenuEntry(FSequencerCommands::Get().SetKeyGroup);
-		MenuBuilder.AddMenuEntry(FSequencerCommands::Get().SetKeyChanged);
-	}
-
-	// Interpolation
-	MenuBuilder.BeginSection("SequencerInterpolation", LOCTEXT("KeyInterpolationMenu", "Default Key Interpolation"));
-	{
-		MenuBuilder.AddMenuEntry(
-			LOCTEXT("SetKeyInterpolationSmartAuto", "Cubic (Smart Auto)"),
-			LOCTEXT("SetKeyInterpolationSmartAutoTooltip", "Set key interpolation to smart auto"),
-			FSlateIcon(FAppStyle::GetAppStyleSetName(), "Sequencer.IconKeySmartAuto"),
-			FUIAction(
-				FExecuteAction::CreateLambda([this] { SequencerPtr.Pin()->SetKeyInterpolation(EMovieSceneKeyInterpolation::SmartAuto); }),
-				FCanExecuteAction(),
-				FIsActionChecked::CreateLambda([this] { return SequencerPtr.Pin()->GetKeyInterpolation() == EMovieSceneKeyInterpolation::SmartAuto; })),
-			NAME_None,
-			EUserInterfaceActionType::RadioButton
-		);
-		MenuBuilder.AddMenuEntry(
-			LOCTEXT("SetKeyInterpolationAuto", "Cubic (Auto)"),
-			LOCTEXT("SetKeyInterpolationAutoTooltip", "Set key interpolation to auto"),
-			FSlateIcon(FAppStyle::GetAppStyleSetName(), "Sequencer.IconKeyAuto"),
-			FUIAction(
-				FExecuteAction::CreateLambda([this] { SequencerPtr.Pin()->SetKeyInterpolation(EMovieSceneKeyInterpolation::Auto); }),
-				FCanExecuteAction(),
-				FIsActionChecked::CreateLambda([this] { return SequencerPtr.Pin()->GetKeyInterpolation() == EMovieSceneKeyInterpolation::Auto; })),
-			NAME_None,
-			EUserInterfaceActionType::RadioButton
-		);
-
-		MenuBuilder.AddMenuEntry(
-			LOCTEXT("SetKeyInterpolationUser", "Cubic (User)"),
-			LOCTEXT("SetKeyInterpolationUserTooltip", "Set key interpolation to user"),
-			FSlateIcon(FAppStyle::GetAppStyleSetName(), "Sequencer.IconKeyUser"),
-			FUIAction(
-				FExecuteAction::CreateLambda([this] { SequencerPtr.Pin()->SetKeyInterpolation(EMovieSceneKeyInterpolation::User); }),
-				FCanExecuteAction(),
-				FIsActionChecked::CreateLambda([this] { return SequencerPtr.Pin()->GetKeyInterpolation() == EMovieSceneKeyInterpolation::User; })),
-			NAME_None,
-			EUserInterfaceActionType::RadioButton
-		);
-
-		MenuBuilder.AddMenuEntry(
-			LOCTEXT("SetKeyInterpolationBreak", "Cubic (Break)"),
-			LOCTEXT("SetKeyInterpolationBreakTooltip", "Set key interpolation to break"),
-			FSlateIcon(FAppStyle::GetAppStyleSetName(), "Sequencer.IconKeyBreak"),
-			FUIAction(
-				FExecuteAction::CreateLambda([this] { SequencerPtr.Pin()->SetKeyInterpolation(EMovieSceneKeyInterpolation::Break); }),
-				FCanExecuteAction(),
-				FIsActionChecked::CreateLambda([this] { return SequencerPtr.Pin()->GetKeyInterpolation() == EMovieSceneKeyInterpolation::Break; })),
-			NAME_None,
-			EUserInterfaceActionType::RadioButton
-		);
-
-		MenuBuilder.AddMenuEntry(
-			LOCTEXT("SetKeyInterpolationLinear", "Linear"),
-			LOCTEXT("SetKeyInterpolationLinearTooltip", "Set key interpolation to linear"),
-			FSlateIcon(FAppStyle::GetAppStyleSetName(), "Sequencer.IconKeyLinear"),
-			FUIAction(
-				FExecuteAction::CreateLambda([this] { SequencerPtr.Pin()->SetKeyInterpolation(EMovieSceneKeyInterpolation::Linear); }),
-				FCanExecuteAction(),
-				FIsActionChecked::CreateLambda([this] { return SequencerPtr.Pin()->GetKeyInterpolation() == EMovieSceneKeyInterpolation::Linear; })),
-			NAME_None,
-			EUserInterfaceActionType::RadioButton
-		);
-
-		MenuBuilder.AddMenuEntry(
-			LOCTEXT("SetKeyInterpolationConstant", "Constant"),
-			LOCTEXT("SetKeyInterpolationConstantTooltip", "Set key interpolation to constant"),
-			FSlateIcon(FAppStyle::GetAppStyleSetName(), "Sequencer.IconKeyConstant"),
-			FUIAction(
-				FExecuteAction::CreateLambda([this] { SequencerPtr.Pin()->SetKeyInterpolation(EMovieSceneKeyInterpolation::Constant); }),
-				FCanExecuteAction(),
-				FIsActionChecked::CreateLambda([this] { return SequencerPtr.Pin()->GetKeyInterpolation() == EMovieSceneKeyInterpolation::Constant; })),
-			NAME_None,
-			EUserInterfaceActionType::RadioButton
-		);
-	}
-	MenuBuilder.EndSection(); // SequencerInterpolation
-
-	return MenuBuilder.MakeWidget();
-
-}
-
 TSharedRef<SWidget> SSequencer::MakeTimeRange(const TSharedRef<SWidget>& InnerContent, bool bShowWorkingRange, bool bShowViewRange, bool bShowPlaybackRange)
 {
 	ISequencerWidgetsModule& SequencerWidgets = FModuleManager::Get().LoadModuleChecked<ISequencerWidgetsModule>( "SequencerWidgets" );
diff --git a/Engine/Source/Editor/Sequencer/Private/SSequencer.h b/Engine/Source/Editor/Sequencer/Private/SSequencer.h
--- a/Engine/Source/Editor/Sequencer/Private/SSequencer.h
+++ b/Engine/Source/Editor/Sequencer/Private/SSequencer.h
@@ -422,15 +422,6 @@ private:
 	/** Makes the snapping menu for the toolbar. */
 	TSharedRef<SWidget> MakeSnapMenu();
 
-	/** Makes the auto-change menu for the toolbar. */
-	TSharedRef<SWidget> MakeAutoChangeMenu();
-
-	/** Makes the allow edits menu for the toolbar. */
-	TSharedRef<SWidget> MakeAllowEditsMenu();
-
-	/** Makes the key group menu for the toolbar. */
-	TSharedRef<SWidget> MakeKeyGroupMenu();
-
 	void OpenTaggedBindingManager();
 
 	/** Makes the advanced menu for the toolbar. */
