fixed Linux asan issue

#rnx

diff --git a/Engine/Source/Editor/Sequencer/Private/LevelEditorSequencerIntegration.cpp b/Engine/Source/Editor/Sequencer/Private/LevelEditorSequencerIntegration.cpp
--- a/Engine/Source/Editor/Sequencer/Private/LevelEditorSequencerIntegration.cpp
+++ b/Engine/Source/Editor/Sequencer/Private/LevelEditorSequencerIntegration.cpp
@@ -302,7 +302,7 @@ void FLevelEditorSequencerIntegration::Initialize(const FLevelEditorSequencerInt
 		});
 	}
 
-	ActivateDetailHandler(Options);
+	BindDetailHandler(Options);
 
 	{
 		FLevelEditorModule& LevelEditorModule = FModuleManager::Get().GetModuleChecked<FLevelEditorModule>("LevelEditor");
@@ -925,12 +925,8 @@ void FLevelEditorSequencerIntegration::MakeBrowseToSelectedActorSubMenu(UToolMen
 	}
 }
 
-void FLevelEditorSequencerIntegration::ActivateDetailHandler(const FLevelEditorSequencerIntegrationOptions& Options)
+void FLevelEditorSequencerIntegration::ActivateDetailHandler()
 {
-	static FName DetailHandlerName("DetailHandler");
-
-	AcquiredResources.Release(DetailHandlerName);
-
 	// Add sequencer detail keyframe handler
 	FPropertyEditorModule& EditModule = FModuleManager::Get().GetModuleChecked<FPropertyEditorModule>("PropertyEditor");
 
@@ -943,18 +939,26 @@ void FLevelEditorSequencerIntegration::ActivateDetailHandler(const FLevelEditorS
 			DetailsView->SetIsPropertyReadOnlyDelegate(FIsPropertyReadOnly::CreateRaw(this, &FLevelEditorSequencerIntegration::IsPropertyReadOnly));
 		}
 	}
+}
 
-	FDelegateHandle OnPropertyEditorOpenedHandle = EditModule.OnPropertyEditorOpened().AddLambda([this, Options]() { ActivateDetailHandler(Options); });
+void FLevelEditorSequencerIntegration::BindDetailHandler(const FLevelEditorSequencerIntegrationOptions& Options)
+{
+	static const FName DetailHandlerName("DetailHandler");
+	// NOTE: this should already have been done in Initialize
+	AcquiredResources.Release(DetailHandlerName);
 
-	auto DeactivateDetailKeyframeHandler =
-		[this, OnPropertyEditorOpenedHandle]()
-		{
-			FPropertyEditorModule* EditModulePtr = FModuleManager::Get().GetModulePtr<FPropertyEditorModule>("PropertyEditor");
-			if (!EditModulePtr)
-			{
-				return;
-			}
+	// bind keyframe handler
+	FPropertyEditorModule& EditModule = FModuleManager::Get().GetModuleChecked<FPropertyEditorModule>("PropertyEditor");
+	FDelegateHandle OnPropertyEditorOpenedHandle = EditModule.OnPropertyEditorOpened().AddLambda([this]()
+	{
+		ActivateDetailHandler();
+	});
 
+	// unbind keyframe handler
+	auto DeactivateDetailKeyframeHandler = [this, OnPropertyEditorOpenedHandle]()
+	{		
+		if (FPropertyEditorModule* EditModulePtr = FModuleManager::Get().GetModulePtr<FPropertyEditorModule>("PropertyEditor"))
+		{
 			EditModulePtr->OnPropertyEditorOpened().Remove(OnPropertyEditorOpenedHandle);
 
 			for (const FName& DetailsTabIdentifier : DetailsTabIdentifiers)
@@ -970,33 +974,30 @@ void FLevelEditorSequencerIntegration::ActivateDetailHandler(const FLevelEditorS
 					DetailsView->GetIsPropertyReadOnlyDelegate().Unbind();
 				}
 			}
-		};
-
+		}
+	};
 	AcquiredResources.Add(DetailHandlerName, DeactivateDetailKeyframeHandler);
 
-	FName DetailHandlerRefreshName("DetailHandlerRefresh");
-	auto RefreshDetailHandler =
-		[]()
+	static const FName DetailHandlerRefreshName("DetailHandlerRefresh");
+	// NOTE: this should already have been done in Initialize
+	AcquiredResources.Release(DetailHandlerRefreshName);
+	
+	if (Options.bForceRefreshDetails)
+	{
+		auto RefreshDetailHandler = []()
 		{
-			FPropertyEditorModule* EditModulePtr = FModuleManager::Get().GetModulePtr<FPropertyEditorModule>("PropertyEditor");
-			if (!EditModulePtr)
-			{
-				return;
-			}
-
-			for (const FName& DetailsTabIdentifier : DetailsTabIdentifiers)
+			if (FPropertyEditorModule* EditModulePtr = FModuleManager::Get().GetModulePtr<FPropertyEditorModule>("PropertyEditor"))
 			{
-				TSharedPtr<IDetailsView> DetailsView = EditModulePtr->FindDetailView(DetailsTabIdentifier);
-				if (DetailsView.IsValid())
+				for (const FName& DetailsTabIdentifier : DetailsTabIdentifiers)
 				{
-					DetailsView->ForceRefresh();
+					TSharedPtr<IDetailsView> DetailsView = EditModulePtr->FindDetailView(DetailsTabIdentifier);
+					if (DetailsView.IsValid())
+					{
+						DetailsView->ForceRefresh();
+					}
 				}
 			}
 		};
-
-
-	if (Options.bForceRefreshDetails)
-	{
 		AcquiredResources.Add(DetailHandlerRefreshName, RefreshDetailHandler);
 	}
 }
diff --git a/Engine/Source/Editor/Sequencer/Public/LevelEditorSequencerIntegration.h b/Engine/Source/Editor/Sequencer/Public/LevelEditorSequencerIntegration.h
--- a/Engine/Source/Editor/Sequencer/Public/LevelEditorSequencerIntegration.h
+++ b/Engine/Source/Editor/Sequencer/Public/LevelEditorSequencerIntegration.h
@@ -219,7 +219,8 @@ private:
 	void ActivateSequencerEditorMode();
 	void DeactivateSequencerEditorMode();
 	void AddLevelViewportMenuExtender();
-	void ActivateDetailHandler(const FLevelEditorSequencerIntegrationOptions& Options);
+	void BindDetailHandler(const FLevelEditorSequencerIntegrationOptions& Options);
+	void ActivateDetailHandler();
 	void AttachOutlinerColumn();
 	void DetachOutlinerColumn();
 	void ActivateRealtimeViewports();
