Fix compilation error

diff --git a/Engine/Source/Editor/Sequencer/Private/Menus/SequencerToolbarUtils.cpp b/Engine/Source/Editor/Sequencer/Private/Menus/SequencerToolbarUtils.cpp
--- a/Engine/Source/Editor/Sequencer/Private/Menus/SequencerToolbarUtils.cpp
+++ b/Engine/Source/Editor/Sequencer/Private/Menus/SequencerToolbarUtils.cpp
@@ -16,7 +16,7 @@
 
 namespace UE::Sequencer
 {
-const FName SequencerToolbarStyleName("SequencerToolbar");
+const FName GSequencerToolbarStyleName("SequencerToolbar");
 	
 namespace ToolbarUtilsDetail
 {
@@ -159,7 +159,7 @@ FToolMenuEntry MakeKeyGroupMenuEntry_ToolMenus(const TWeakPtr<FSequencer>& InWea
 		KeyGroupEntry = FToolMenuEntry::InitComboButton(
 			InName, InAction, InMenuContentGenerator, InLabelOverride, ToolTipAttr, IconAttr
 		);
-		KeyGroupEntry.StyleNameOverride = SequencerToolbarStyleName;
+		KeyGroupEntry.StyleNameOverride = GSequencerToolbarStyleName;
 	});
 	return KeyGroupEntry;
 }
@@ -249,13 +249,13 @@ FToolMenuEntry MakeAutoKeyMenuEntry(const TSharedPtr<FSequencer>& InSequencer)
 		KeyGroupEntry = FToolMenuEntry::InitComboButton(
 			InName, InAction, InMenuContentGenerator, InLabelOverride, ToolTipAttr, IconAttr
 		);
-		KeyGroupEntry.StyleNameOverride = SequencerToolbarStyleName;
+		KeyGroupEntry.StyleNameOverride = GSequencerToolbarStyleName;
 		},
 		[&KeyGroupEntry](TSharedPtr<FUICommandInfo> ToggleAutoKeyEnabled)
 		{
 			KeyGroupEntry = FToolMenuEntry::InitMenuEntry(ToggleAutoKeyEnabled);
 		});
-	KeyGroupEntry.StyleNameOverride = SequencerToolbarStyleName;
+	KeyGroupEntry.StyleNameOverride = GSequencerToolbarStyleName;
 	return KeyGroupEntry;
 }
 
@@ -331,7 +331,7 @@ TOptional<FToolMenuEntry> MakeAllowEditsModeMenuEntry(const TSharedPtr<FSequence
 		KeyGroupEntry = FToolMenuEntry::InitComboButton(
 			InName, InAction, InMenuContentGenerator, InLabelOverride, ToolTipAttr, IconAttr
 		);
-		KeyGroupEntry.StyleNameOverride = SequencerToolbarStyleName;
+		KeyGroupEntry.StyleNameOverride = GSequencerToolbarStyleName;
 	});
 	if (bMade)
 	{
@@ -344,7 +344,7 @@ void AppendSequencerToolbarEntries(const TSharedPtr<FSequencer>& InSequencer, FT
 {
 	const TWeakPtr<FSequencer> WeakSequencer = InSequencer.ToWeakPtr();
 
-	InToolbarBuilder.BeginStyleOverride(SequencerToolbarStyleName);
+	InToolbarBuilder.BeginStyleOverride(GSequencerToolbarStyleName);
 	
 	if (EnumHasAnyFlags(InFlags, EToolbarItemFlags::KeyGroup))
 	{
diff --git a/Engine/Source/Editor/Sequencer/Private/Menus/SequencerToolbarUtils.h b/Engine/Source/Editor/Sequencer/Private/Menus/SequencerToolbarUtils.h
--- a/Engine/Source/Editor/Sequencer/Private/Menus/SequencerToolbarUtils.h
+++ b/Engine/Source/Editor/Sequencer/Private/Menus/SequencerToolbarUtils.h
@@ -15,7 +15,7 @@ template<typename OptionalType>struct TOptional;
 
 namespace UE::Sequencer
 {
-extern const FName SequencerToolbarStyleName;
+extern const FName GSequencerToolbarStyleName;
 
 /** Makes the combo button for changing the key group settings. */
 FToolMenuEntry MakeKeyGroupMenuEntry_ToolMenus(const TWeakPtr<FSequencer>& InWeakSequencer);
diff --git a/Engine/Source/Editor/Sequencer/Private/SSequencer.cpp b/Engine/Source/Editor/Sequencer/Private/SSequencer.cpp
--- a/Engine/Source/Editor/Sequencer/Private/SSequencer.cpp
+++ b/Engine/Source/Editor/Sequencer/Private/SSequencer.cpp
@@ -1440,7 +1440,7 @@ TSharedRef<SWidget> SSequencer::MakeToolBar()
 
 void SSequencer::PopulateToolBar(UToolMenu* InMenu)
 {
-	const FName SequencerToolbarStyleName = UE::Sequencer::SequencerToolbarStyleName;
+	const FName SequencerToolbarStyleName = UE::Sequencer::GSequencerToolbarStyleName;
 	
 	USequencerToolMenuContext* ContextObject = InMenu->FindContext<USequencerToolMenuContext>();
 	if (!ContextObject)
