Sequencer: Fix cycling through time formats. Skip past Hidden time formats

diff --git a/Engine/Source/Editor/Sequencer/Private/Sequencer.cpp b/Engine/Source/Editor/Sequencer/Private/Sequencer.cpp
--- a/Engine/Source/Editor/Sequencer/Private/Sequencer.cpp
+++ b/Engine/Source/Editor/Sequencer/Private/Sequencer.cpp
@@ -11337,12 +11337,28 @@ void FSequencer::BindCommands()
 	SequencerCommandBindings->MapAction(
 		Commands.ChangeTimeDisplayFormat,
 		FExecuteAction::CreateLambda( [this]{
-			EFrameNumberDisplayFormats NextFormat = (EFrameNumberDisplayFormats)((uint8)Settings->GetTimeDisplayFormat() + 1);
-			if (NextFormat == EFrameNumberDisplayFormats::MAX_Count)
+			
+			const UEnum* FrameNumberDisplayEnum = StaticEnum<EFrameNumberDisplayFormats>();
+			check(FrameNumberDisplayEnum);
+
+			uint8 NextIndex = (uint8)Settings->GetTimeDisplayFormat() + 1;
+			if (NextIndex >= FrameNumberDisplayEnum->NumEnums() - 1)
 			{
-				NextFormat = EFrameNumberDisplayFormats::NonDropFrameTimecode;
+				NextIndex = 0;
 			}
 
+			while (FrameNumberDisplayEnum->HasMetaData(TEXT("Hidden"), NextIndex))
+			{
+				NextIndex++;
+
+				if (NextIndex >= FrameNumberDisplayEnum->NumEnums() - 1)
+				{
+					NextIndex = 0;
+				}
+			}
+
+			EFrameNumberDisplayFormats NextFormat = (EFrameNumberDisplayFormats)(NextIndex);
+
 			// If the next framerate in the list is drop format timecode and we're not in a play rate that supports drop format timecode,
 			// then we will skip over it.
 			bool bCanShowDropFrameTimecode = FTimecode::UseDropFormatTimecode(GetFocusedDisplayRate());
