Sequencer: Fix crash generating CurveEditor toolbar key combo button in Avalanche due to command being null.

diff --git a/Engine/Source/Editor/Sequencer/Private/Menus/SequencerToolbarUtils.cpp b/Engine/Source/Editor/Sequencer/Private/Menus/SequencerToolbarUtils.cpp
--- a/Engine/Source/Editor/Sequencer/Private/Menus/SequencerToolbarUtils.cpp
+++ b/Engine/Source/Editor/Sequencer/Private/Menus/SequencerToolbarUtils.cpp
@@ -362,13 +362,7 @@ void AppendSequencerToolbarEntries(const TSharedPtr<FSequencer>& InSequencer, FT
 			[&InToolbarBuilder](const FName, const FUIAction& InAction, const FOnGetContent& InMenuContentGenerator,
 			const TAttribute<FText>& InLabelOverride, const TAttribute<FText>& ToolTipAttr, const TAttribute<FSlateIcon>& IconAttr)
 			{
-				FButtonArgs ButtonArgs;
-				ButtonArgs.Action = InAction;
-				ButtonArgs.OnGetMenuContent = InMenuContentGenerator;
-				ButtonArgs.LabelOverride = InLabelOverride;
-				ButtonArgs.ToolTipOverride = ToolTipAttr;
-				ButtonArgs.IconOverride = IconAttr;
-				InToolbarBuilder.AddToolBarButton(ButtonArgs);
+				InToolbarBuilder.AddComboButton(InAction, InMenuContentGenerator, InLabelOverride, ToolTipAttr, IconAttr);
 			},
 			[&InToolbarBuilder](TSharedPtr<FUICommandInfo> ToggleAutoKeyEnabled)
 			{
