Sequencer: Sequencer.EnableRelevantThumbnails is now true by default. This will result in level sequences automatically capturing the middle frame as relevant thumbnail when saved.

diff --git a/Engine/Source/Editor/Sequencer/Private/Misc/ConsoleVariables.cpp b/Engine/Source/Editor/Sequencer/Private/Misc/ConsoleVariables.cpp
--- a/Engine/Source/Editor/Sequencer/Private/Misc/ConsoleVariables.cpp
+++ b/Engine/Source/Editor/Sequencer/Private/Misc/ConsoleVariables.cpp
@@ -6,7 +6,7 @@ namespace UE::Sequencer
 {
 	TAutoConsoleVariable<bool> CVarEnableRelevantThumbnails(
 		TEXT("Sequencer.EnableRelevantThumbnails"),
-		false,
+		true,
 		TEXT("If enabled, upon saving a sequence, tries to find a suitable frame, scrubs to it, renders a thumbnail, and saves it as thumbnail.")
 	);
 }
\ No newline at end of file
