[Sequencer] Filtering:
- Add the ability for UMovieSceneSequence types to indicate which track filters they support

diff --git a/Engine/Source/Editor/Sequencer/Private/Filters/Filters/SequencerTrackFilter_Condition.cpp b/Engine/Source/Editor/Sequencer/Private/Filters/Filters/SequencerTrackFilter_Condition.cpp
--- a/Engine/Source/Editor/Sequencer/Private/Filters/Filters/SequencerTrackFilter_Condition.cpp
+++ b/Engine/Source/Editor/Sequencer/Private/Filters/Filters/SequencerTrackFilter_Condition.cpp
@@ -1,6 +1,6 @@
 // Copyright Epic Games, Inc. All Rights Reserved.
 
-#include "SequencerTrackFilter_Condition.h"
+#include "Filters/Filters/SequencerTrackFilter_Condition.h"
 #include "Filters/SequencerFilterBar.h"
 #include "Filters/SequencerTrackFilterCommands.h"
 #include "MVVM/ViewModels/ObjectBindingModel.h"
@@ -42,7 +42,7 @@ FSlateIcon FSequencerTrackFilter_Condition::GetIcon() const
 
 FString FSequencerTrackFilter_Condition::GetName() const
 {
-	return TEXT("Condition");
+	return StaticName();
 }
 
 bool FSequencerTrackFilter_Condition::PassesFilter(FSequencerTrackFilterType InItem) const
diff --git a/Engine/Source/Editor/Sequencer/Private/Filters/Filters/SequencerTrackFilter_Condition.h b/Engine/Source/Editor/Sequencer/Private/Filters/Filters/SequencerTrackFilter_Condition.h
deleted file mode 100644
diff --git a/Engine/Source/Editor/Sequencer/Private/Filters/Filters/SequencerTrackFilter_Group.cpp b/Engine/Source/Editor/Sequencer/Private/Filters/Filters/SequencerTrackFilter_Group.cpp
--- a/Engine/Source/Editor/Sequencer/Private/Filters/Filters/SequencerTrackFilter_Group.cpp
+++ b/Engine/Source/Editor/Sequencer/Private/Filters/Filters/SequencerTrackFilter_Group.cpp
@@ -1,6 +1,6 @@
 ﻿// Copyright Epic Games, Inc. All Rights Reserved.
 
-#include "SequencerTrackFilter_Group.h"
+#include "Filters/Filters/SequencerTrackFilter_Group.h"
 #include "Filters/SequencerTrackFilterCommands.h"
 #include "MovieScene.h"
 #include "MVVM/Extensions/IOutlinerExtension.h"
@@ -47,7 +47,7 @@ FSlateIcon FSequencerTrackFilter_Group::GetIcon() const
 
 FString FSequencerTrackFilter_Group::GetName() const
 {
-	return TEXT("Group");
+	return StaticName();
 }
 
 bool FSequencerTrackFilter_Group::PassesFilter(FSequencerTrackFilterType InItem) const
diff --git a/Engine/Source/Editor/Sequencer/Private/Filters/Filters/SequencerTrackFilter_Group.h b/Engine/Source/Editor/Sequencer/Private/Filters/Filters/SequencerTrackFilter_Group.h
deleted file mode 100644
diff --git a/Engine/Source/Editor/Sequencer/Private/Filters/Filters/SequencerTrackFilter_HideIsolate.cpp b/Engine/Source/Editor/Sequencer/Private/Filters/Filters/SequencerTrackFilter_HideIsolate.cpp
--- a/Engine/Source/Editor/Sequencer/Private/Filters/Filters/SequencerTrackFilter_HideIsolate.cpp
+++ b/Engine/Source/Editor/Sequencer/Private/Filters/Filters/SequencerTrackFilter_HideIsolate.cpp
@@ -1,6 +1,6 @@
 ﻿// Copyright Epic Games, Inc. All Rights Reserved.
 
-#include "SequencerTrackFilter_HideIsolate.h"
+#include "Filters/Filters/SequencerTrackFilter_HideIsolate.h"
 #include "Filters/SequencerTrackFilterCommands.h"
 #include "MVVM/ViewModels/CategoryModel.h"
 #include "MVVM/ViewModels/TrackModel.h"
@@ -306,7 +306,7 @@ FSlateIcon FSequencerTrackFilter_HideIsolate::GetIcon() const
 
 FString FSequencerTrackFilter_HideIsolate::GetName() const
 {
-	return TEXT("HideIsolate");
+	return StaticName();
 }
 
 bool FSequencerTrackFilter_HideIsolate::PassesFilter(FSequencerTrackFilterType InItem) const
diff --git a/Engine/Source/Editor/Sequencer/Private/Filters/Filters/SequencerTrackFilter_HideIsolate.h b/Engine/Source/Editor/Sequencer/Private/Filters/Filters/SequencerTrackFilter_HideIsolate.h
deleted file mode 100644
diff --git a/Engine/Source/Editor/Sequencer/Private/Filters/Filters/SequencerTrackFilter_Keyed.cpp b/Engine/Source/Editor/Sequencer/Private/Filters/Filters/SequencerTrackFilter_Keyed.cpp
--- a/Engine/Source/Editor/Sequencer/Private/Filters/Filters/SequencerTrackFilter_Keyed.cpp
+++ b/Engine/Source/Editor/Sequencer/Private/Filters/Filters/SequencerTrackFilter_Keyed.cpp
@@ -1,6 +1,6 @@
 ﻿// Copyright Epic Games, Inc. All Rights Reserved.
 
-#include "SequencerTrackFilter_Keyed.h"
+#include "Filters/Filters/SequencerTrackFilter_Keyed.h"
 #include "Filters/SequencerFilterBar.h"
 #include "Filters/SequencerTrackFilterCommands.h"
 #include "Framework/Commands/Commands.h"
@@ -70,7 +70,7 @@ FSlateIcon FSequencerTrackFilter_Keyed::GetIcon() const
 
 FString FSequencerTrackFilter_Keyed::GetName() const
 {
-	return TEXT("Keyed");
+	return StaticName();
 }
 
 bool FSequencerTrackFilter_Keyed::DoesTrackExtensionHaveKeys(const TWeakViewModelPtr<ITrackExtension>& InTrack)
diff --git a/Engine/Source/Editor/Sequencer/Private/Filters/Filters/SequencerTrackFilter_Keyed.h b/Engine/Source/Editor/Sequencer/Private/Filters/Filters/SequencerTrackFilter_Keyed.h
deleted file mode 100644
diff --git a/Engine/Source/Editor/Sequencer/Private/Filters/Filters/SequencerTrackFilter_Level.cpp b/Engine/Source/Editor/Sequencer/Private/Filters/Filters/SequencerTrackFilter_Level.cpp
--- a/Engine/Source/Editor/Sequencer/Private/Filters/Filters/SequencerTrackFilter_Level.cpp
+++ b/Engine/Source/Editor/Sequencer/Private/Filters/Filters/SequencerTrackFilter_Level.cpp
@@ -1,6 +1,6 @@
 ﻿// Copyright Epic Games, Inc. All Rights Reserved.
 
-#include "SequencerTrackFilter_Level.h"
+#include "Filters/Filters/SequencerTrackFilter_Level.h"
 #include "Engine/Level.h"
 #include "Engine/World.h"
 #include "Filters/SequencerTrackFilterCommands.h"
@@ -54,7 +54,7 @@ FSlateIcon FSequencerTrackFilter_Level::GetIcon() const
 
 FString FSequencerTrackFilter_Level::GetName() const
 {
-	return TEXT("Level");
+	return StaticName();
 }
 
 bool FSequencerTrackFilter_Level::PassesFilter(FSequencerTrackFilterType InItem) const
diff --git a/Engine/Source/Editor/Sequencer/Private/Filters/Filters/SequencerTrackFilter_Level.h b/Engine/Source/Editor/Sequencer/Private/Filters/Filters/SequencerTrackFilter_Level.h
deleted file mode 100644
diff --git a/Engine/Source/Editor/Sequencer/Private/Filters/Filters/SequencerTrackFilter_Modified.cpp b/Engine/Source/Editor/Sequencer/Private/Filters/Filters/SequencerTrackFilter_Modified.cpp
--- a/Engine/Source/Editor/Sequencer/Private/Filters/Filters/SequencerTrackFilter_Modified.cpp
+++ b/Engine/Source/Editor/Sequencer/Private/Filters/Filters/SequencerTrackFilter_Modified.cpp
@@ -1,6 +1,6 @@
 ﻿// Copyright Epic Games, Inc. All Rights Reserved.
 
-#include "SequencerTrackFilter_Modified.h"
+#include "Filters/Filters/SequencerTrackFilter_Modified.h"
 #include "Filters/SequencerFilterBar.h"
 #include "Filters/SequencerTrackFilterCommands.h"
 #include "UObject/Package.h"
@@ -39,7 +39,7 @@ FSlateIcon FSequencerTrackFilter_Modified::GetIcon() const
 
 FString FSequencerTrackFilter_Modified::GetName() const
 {
-	return TEXT("Modified");
+	return StaticName();
 }
 
 bool FSequencerTrackFilter_Modified::PassesFilter(FSequencerTrackFilterType InItem) const
diff --git a/Engine/Source/Editor/Sequencer/Private/Filters/Filters/SequencerTrackFilter_Modified.h b/Engine/Source/Editor/Sequencer/Private/Filters/Filters/SequencerTrackFilter_Modified.h
deleted file mode 100644
diff --git a/Engine/Source/Editor/Sequencer/Private/Filters/Filters/SequencerTrackFilter_Selected.cpp b/Engine/Source/Editor/Sequencer/Private/Filters/Filters/SequencerTrackFilter_Selected.cpp
--- a/Engine/Source/Editor/Sequencer/Private/Filters/Filters/SequencerTrackFilter_Selected.cpp
+++ b/Engine/Source/Editor/Sequencer/Private/Filters/Filters/SequencerTrackFilter_Selected.cpp
@@ -1,6 +1,6 @@
 ﻿// Copyright Epic Games, Inc. All Rights Reserved.
 
-#include "SequencerTrackFilter_Selected.h"
+#include "Filters/Filters/SequencerTrackFilter_Selected.h"
 #include "EditorModeManager.h"
 #include "Elements/Framework/TypedElementSelectionSet.h"
 #include "Filters/SequencerFilterBar.h"
@@ -100,7 +100,7 @@ FSlateIcon FSequencerTrackFilter_Selected::GetIcon() const
 
 FString FSequencerTrackFilter_Selected::GetName() const
 {
-	return TEXT("Selected");
+	return StaticName();
 }
 
 bool FSequencerTrackFilter_Selected::PassesFilter(FSequencerTrackFilterType InItem) const
diff --git a/Engine/Source/Editor/Sequencer/Private/Filters/Filters/SequencerTrackFilter_Selected.h b/Engine/Source/Editor/Sequencer/Private/Filters/Filters/SequencerTrackFilter_Selected.h
deleted file mode 100644
diff --git a/Engine/Source/Editor/Sequencer/Private/Filters/Filters/SequencerTrackFilter_Text.cpp b/Engine/Source/Editor/Sequencer/Private/Filters/Filters/SequencerTrackFilter_Text.cpp
--- a/Engine/Source/Editor/Sequencer/Private/Filters/Filters/SequencerTrackFilter_Text.cpp
+++ b/Engine/Source/Editor/Sequencer/Private/Filters/Filters/SequencerTrackFilter_Text.cpp
@@ -96,7 +96,7 @@ FText FSequencerTrackFilter_Text::GetToolTipText() const
 
 FString FSequencerTrackFilter_Text::GetName() const
 {
-	return TEXT("TextFilter");
+	return StaticName();
 }
 
 bool FSequencerTrackFilter_Text::PassesFilter(FSequencerTrackFilterType InItem) const
diff --git a/Engine/Source/Editor/Sequencer/Private/Filters/Filters/SequencerTrackFilter_Text.h b/Engine/Source/Editor/Sequencer/Private/Filters/Filters/SequencerTrackFilter_Text.h
--- a/Engine/Source/Editor/Sequencer/Private/Filters/Filters/SequencerTrackFilter_Text.h
+++ b/Engine/Source/Editor/Sequencer/Private/Filters/Filters/SequencerTrackFilter_Text.h
@@ -8,6 +8,8 @@
 class FSequencerTrackFilter_Text : public FSequencerTrackFilter
 {
 public:
+	static FString StaticName() { return TEXT("Text"); }
+
 	FSequencerTrackFilter_Text(ISequencerTrackFilters& InFilterInterface);
 
 	//~ Begin FFilterBase
diff --git a/Engine/Source/Editor/Sequencer/Private/Filters/Filters/SequencerTrackFilter_TimeWarp.cpp b/Engine/Source/Editor/Sequencer/Private/Filters/Filters/SequencerTrackFilter_TimeWarp.cpp
--- a/Engine/Source/Editor/Sequencer/Private/Filters/Filters/SequencerTrackFilter_TimeWarp.cpp
+++ b/Engine/Source/Editor/Sequencer/Private/Filters/Filters/SequencerTrackFilter_TimeWarp.cpp
@@ -1,6 +1,6 @@
 // Copyright Epic Games, Inc. All Rights Reserved.
 
-#include "SequencerTrackFilter_TimeWarp.h"
+#include "Filters/Filters/SequencerTrackFilter_TimeWarp.h"
 #include "Filters/SequencerFilterBar.h"
 #include "Filters/SequencerTrackFilterCommands.h"
 #include "MovieSceneSection.h"
@@ -47,7 +47,7 @@ FSlateIcon FSequencerTrackFilter_TimeWarp::GetIcon() const
 
 FString FSequencerTrackFilter_TimeWarp::GetName() const
 {
-	return TEXT("TimeWarp");
+	return StaticName();
 }
 
 bool FSequencerTrackFilter_TimeWarp::PassesFilter(FSequencerTrackFilterType InItem) const
diff --git a/Engine/Source/Editor/Sequencer/Private/Filters/Filters/SequencerTrackFilter_TimeWarp.h b/Engine/Source/Editor/Sequencer/Private/Filters/Filters/SequencerTrackFilter_TimeWarp.h
deleted file mode 100644
diff --git a/Engine/Source/Editor/Sequencer/Private/Filters/Filters/SequencerTrackFilter_Unbound.cpp b/Engine/Source/Editor/Sequencer/Private/Filters/Filters/SequencerTrackFilter_Unbound.cpp
--- a/Engine/Source/Editor/Sequencer/Private/Filters/Filters/SequencerTrackFilter_Unbound.cpp
+++ b/Engine/Source/Editor/Sequencer/Private/Filters/Filters/SequencerTrackFilter_Unbound.cpp
@@ -1,6 +1,6 @@
 ﻿// Copyright Epic Games, Inc. All Rights Reserved.
 
-#include "SequencerTrackFilter_Unbound.h"
+#include "Filters/Filters/SequencerTrackFilter_Unbound.h"
 #include "Filters/SequencerFilterBar.h"
 #include "Filters/SequencerTrackFilterCommands.h"
 #include "MVVM/ViewModels/ObjectBindingModel.h"
@@ -42,7 +42,7 @@ FSlateIcon FSequencerTrackFilter_Unbound::GetIcon() const
 
 FString FSequencerTrackFilter_Unbound::GetName() const
 {
-	return TEXT("Unbound");
+	return StaticName();
 }
 
 bool FSequencerTrackFilter_Unbound::PassesFilter(FSequencerTrackFilterType InItem) const
diff --git a/Engine/Source/Editor/Sequencer/Private/Filters/Filters/SequencerTrackFilter_Unbound.h b/Engine/Source/Editor/Sequencer/Private/Filters/Filters/SequencerTrackFilter_Unbound.h
deleted file mode 100644
diff --git a/Engine/Source/Editor/Sequencer/Private/Filters/Filters/SequencerTrackFilters.cpp b/Engine/Source/Editor/Sequencer/Private/Filters/Filters/SequencerTrackFilters.cpp
new file mode 100644
--- /dev/null
+++ b/Engine/Source/Editor/Sequencer/Private/Filters/Filters/SequencerTrackFilters.cpp
@@ -0,0 +1,437 @@
+﻿// Copyright Epic Games, Inc. All Rights Reserved.
+
+#include "Filters/Filters/SequencerTrackFilters.h"
+#include "Filters/SequencerTrackFilterBase.h"
+#include "Filters/SequencerTrackFilterCommands.h"
+#include "Misc/IFilter.h"
+#include "MovieSceneFolder.h"
+#include "Particles/ParticleSystem.h"
+#include "Styling/AppStyle.h"
+#include "Styling/SlateIconFinder.h"
+
+using namespace UE::Sequencer;
+
+#define LOCTEXT_NAMESPACE "SequencerTrackFilters"
+
+FSequencerTrackFilter_Audio::FSequencerTrackFilter_Audio(ISequencerTrackFilters& InFilterInterface, TSharedPtr<FFilterCategory> InCategory)
+	: FSequencerTrackFilter_ClassType<UMovieSceneAudioTrack>(InFilterInterface, InCategory)
+{
+}
+
+FText FSequencerTrackFilter_Audio::GetDisplayName() const
+{
+	return LOCTEXT("SequencerTrackFilter_Audio", "Audio");
+}
+
+FSlateIcon FSequencerTrackFilter_Audio::GetIcon() const
+{
+	return FSlateIcon(FAppStyle::GetAppStyleSetName(), TEXT("Sequencer.Tracks.Audio"));
+}
+
+FText FSequencerTrackFilter_Audio::GetDefaultToolTipText() const
+{
+	return LOCTEXT("SequencerTrackFilter_AudioToolTip", "Show only Audio tracks");
+}
+
+TSharedPtr<FUICommandInfo> FSequencerTrackFilter_Audio::GetToggleCommand() const
+{
+	return FSequencerTrackFilterCommands::Get().ToggleFilter_Audio;
+}
+
+bool FSequencerTrackFilter_Audio::SupportsSequence(UMovieSceneSequence* const InSequence) const
+{
+	return IsSequenceTrackSupported<UMovieSceneAudioTrack>(InSequence)
+		|| SupportsLevelSequence(InSequence)
+		|| SupportsUMGSequence(InSequence);
+}
+
+//////////////////////////////////////////////////////////////////////////
+//
+
+FSequencerTrackFilter_Event::FSequencerTrackFilter_Event(ISequencerTrackFilters& InFilterInterface, TSharedPtr<FFilterCategory> InCategory)
+	: FSequencerTrackFilter_ClassType<UMovieSceneEventTrack>(InFilterInterface, InCategory)
+{
+}
+
+FText FSequencerTrackFilter_Event::GetDisplayName() const
+{
+	return LOCTEXT("SequencerTrackFilter_Event", "Event");
+}
+
+FSlateIcon FSequencerTrackFilter_Event::GetIcon() const
+{
+	return FSlateIcon(FAppStyle::GetAppStyleSetName(), TEXT("Sequencer.Tracks.Event"));
+}
+
+FText FSequencerTrackFilter_Event::GetDefaultToolTipText() const
+{
+	return LOCTEXT("SequencerTrackFilter_EventToolTip", "Show only Event tracks");
+}
+
+TSharedPtr<FUICommandInfo> FSequencerTrackFilter_Event::GetToggleCommand() const
+{
+	return FSequencerTrackFilterCommands::Get().ToggleFilter_Event;
+}
+
+bool FSequencerTrackFilter_Event::SupportsSequence(UMovieSceneSequence* const InSequence) const
+{
+	return FSequencerTrackFilter::SupportsSequence(InSequence)
+		|| SupportsUMGSequence(InSequence);
+}
+
+//////////////////////////////////////////////////////////////////////////
+//
+
+FSequencerTrackFilter_LevelVisibility::FSequencerTrackFilter_LevelVisibility(ISequencerTrackFilters& InFilterInterface, TSharedPtr<FFilterCategory> InCategory)
+	: FSequencerTrackFilter_ClassType<UMovieSceneLevelVisibilityTrack>(InFilterInterface, InCategory)
+{
+}
+
+FText FSequencerTrackFilter_LevelVisibility::GetDisplayName() const
+{
+	return LOCTEXT("SequencerTrackFilter_LevelVisibility", "Level Visibility");
+}
+
+FSlateIcon FSequencerTrackFilter_LevelVisibility::GetIcon() const
+{
+	return FSlateIcon(FAppStyle::GetAppStyleSetName(), TEXT("Sequencer.Tracks.LevelVisibility"));
+}
+
+FText FSequencerTrackFilter_LevelVisibility::GetDefaultToolTipText() const
+{
+	return LOCTEXT("SequencerTrackFilter_LevelVisibilityToolTip", "Show only Level Visibility tracks");
+}
+
+TSharedPtr<FUICommandInfo> FSequencerTrackFilter_LevelVisibility::GetToggleCommand() const
+{
+	return FSequencerTrackFilterCommands::Get().ToggleFilter_LevelVisibility;
+}
+
+//////////////////////////////////////////////////////////////////////////
+//
+
+FSequencerTrackFilter_Particle::FSequencerTrackFilter_Particle(ISequencerTrackFilters& InFilterInterface, TSharedPtr<FFilterCategory> InCategory)
+	: FSequencerTrackFilter_ClassType<UMovieSceneParticleTrack>(InFilterInterface, InCategory)
+{
+}
+
+FText FSequencerTrackFilter_Particle::GetDisplayName() const
+{
+	return LOCTEXT("SequencerTrackFilter_Particle", "Particle System");
+}
+
+FSlateIcon FSequencerTrackFilter_Particle::GetIcon() const
+{
+	return FSlateIconFinder::FindIconForClass(UParticleSystem::StaticClass());
+}
+
+FText FSequencerTrackFilter_Particle::GetDefaultToolTipText() const
+{
+	return LOCTEXT("SequencerTrackFilter_ParticleToolTip", "Show only Particle System tracks");
+}
+
+TSharedPtr<FUICommandInfo> FSequencerTrackFilter_Particle::GetToggleCommand() const
+{
+	return FSequencerTrackFilterCommands::Get().ToggleFilter_Particle;
+}
+
+//////////////////////////////////////////////////////////////////////////
+//
+
+FSequencerTrackFilter_CinematicShot::FSequencerTrackFilter_CinematicShot(ISequencerTrackFilters& InFilterInterface, TSharedPtr<FFilterCategory> InCategory)
+	: FSequencerTrackFilter_ClassType<UMovieSceneCinematicShotTrack>(InFilterInterface, InCategory)
+{
+}
+
+FText FSequencerTrackFilter_CinematicShot::GetDisplayName() const
+{
+	return LOCTEXT("SequencerTrackFilter_CinematicShot", "Shot");
+}
+
+FSlateIcon FSequencerTrackFilter_CinematicShot::GetIcon() const
+{
+	return FSlateIcon(FAppStyle::GetAppStyleSetName(), TEXT("Sequencer.Tracks.CinematicShot"));
+}
+
+FText FSequencerTrackFilter_CinematicShot::GetDefaultToolTipText() const
+{
+	return LOCTEXT("SequencerTrackFilter_CinematicShotToolTip", "Show only Shot tracks");
+}
+
+TSharedPtr<FUICommandInfo> FSequencerTrackFilter_CinematicShot::GetToggleCommand() const
+{
+	return FSequencerTrackFilterCommands::Get().ToggleFilter_CinematicShot;
+}
+
+//////////////////////////////////////////////////////////////////////////
+//
+
+FSequencerTrackFilter_Subsequence::FSequencerTrackFilter_Subsequence(ISequencerTrackFilters& InFilterInterface, TSharedPtr<FFilterCategory> InCategory)
+	: FSequencerTrackFilter_ClassType<UMovieSceneSubTrack>(InFilterInterface, InCategory)
+{
+}
+
+bool FSequencerTrackFilter_Subsequence::PassesFilter(FSequencerTrackFilterType InItem) const
+{
+	FSequencerFilterData& FilterData = GetFilterInterface().GetFilterData();
+	const UMovieSceneTrack* const TrackObject = FilterData.ResolveMovieSceneTrackObject(InItem);
+	return IsValid(TrackObject)
+		&& TrackObject->IsA(UMovieSceneSubTrack::StaticClass())
+		&& !TrackObject->IsA(UMovieSceneCinematicShotTrack::StaticClass());
+}
+
+FText FSequencerTrackFilter_Subsequence::GetDisplayName() const
+{
+	return LOCTEXT("SequencerTrackFilter_Subsequence", "Subsequence");
+}
+
+FSlateIcon FSequencerTrackFilter_Subsequence::GetIcon() const
+{
+	return FSlateIcon(FAppStyle::GetAppStyleSetName(), TEXT("Sequencer.Tracks.Sub"));
+}
+
+FText FSequencerTrackFilter_Subsequence::GetDefaultToolTipText() const
+{
+	return LOCTEXT("SequencerTrackFilter_SubsequenceToolTip", "Show only Subsequence tracks");
+}
+
+TSharedPtr<FUICommandInfo> FSequencerTrackFilter_Subsequence::GetToggleCommand() const
+{
+	return FSequencerTrackFilterCommands::Get().ToggleFilter_Subsequence;
+}
+
+//////////////////////////////////////////////////////////////////////////
+//
+
+FSequencerTrackFilter_SkeletalMesh::FSequencerTrackFilter_SkeletalMesh(ISequencerTrackFilters& InFilterInterface, TSharedPtr<FFilterCategory> InCategory)
+	: FSequencerTrackFilter_ComponentType<USkeletalMeshComponent>(InFilterInterface, InCategory)
+{
+}
+
+FText FSequencerTrackFilter_SkeletalMesh::GetDisplayName() const
+{
+	return LOCTEXT("SequencerTrackFilter_SkeletalMesh", "Skeletal Mesh");
+}
+
+FSlateIcon FSequencerTrackFilter_SkeletalMesh::GetIcon() const
+{
+	return FSlateIconFinder::FindIconForClass(USkeletalMeshComponent::StaticClass());
+}
+
+FText FSequencerTrackFilter_SkeletalMesh::GetDefaultToolTipText() const
+{
+	return LOCTEXT("SequencerTrackFilter_SkeletalMeshToolTip", "Show only Skeletal Mesh tracks");
+}
+
+TSharedPtr<FUICommandInfo> FSequencerTrackFilter_SkeletalMesh::GetToggleCommand() const
+{
+	return FSequencerTrackFilterCommands::Get().ToggleFilter_SkeletalMesh;
+}
+
+//////////////////////////////////////////////////////////////////////////
+//
+
+FSequencerTrackFilter_Camera::FSequencerTrackFilter_Camera(ISequencerTrackFilters& InFilterInterface, TSharedPtr<FFilterCategory> InCategory)
+	: FSequencerTrackFilter_ComponentType<UCameraComponent>(InFilterInterface, InCategory)
+{
+}
+
+FText FSequencerTrackFilter_Camera::GetDisplayName() const
+{
+	return LOCTEXT("SequencerTrackFilter_Camera", "Camera");
+}
+
+FSlateIcon FSequencerTrackFilter_Camera::GetIcon() const
+{
+	return FSlateIconFinder::FindIconForClass(UCameraComponent::StaticClass());
+}
+
+FText FSequencerTrackFilter_Camera::GetDefaultToolTipText() const
+{
+	return LOCTEXT("SequencerTrackFilter_CameraToolTip", "Show only Camera tracks");
+}
+
+TSharedPtr<FUICommandInfo> FSequencerTrackFilter_Camera::GetToggleCommand() const
+{
+	return FSequencerTrackFilterCommands::Get().ToggleFilter_Camera;
+}
+
+//////////////////////////////////////////////////////////////////////////
+//
+
+FSequencerTrackFilter_Light::FSequencerTrackFilter_Light(ISequencerTrackFilters& InFilterInterface, TSharedPtr<FFilterCategory> InCategory)
+	: FSequencerTrackFilter_ComponentType<ULightComponentBase>(InFilterInterface, InCategory)
+{
+}
+
+FText FSequencerTrackFilter_Light::GetDisplayName() const
+{
+	return LOCTEXT("SequencerTrackFilter_Light", "Light");
+}
+
+FSlateIcon FSequencerTrackFilter_Light::GetIcon() const
+{
+	return FSlateIcon(FAppStyle::GetAppStyleSetName(), TEXT("ClassIcon.Light"));
+}
+
+FText FSequencerTrackFilter_Light::GetDefaultToolTipText() const
+{
+	return LOCTEXT("SequencerTrackFilter_LightToolTip", "Show only Light tracks");
+}
+
+TSharedPtr<FUICommandInfo> FSequencerTrackFilter_Light::GetToggleCommand() const
+{
+	return FSequencerTrackFilterCommands::Get().ToggleFilter_Light;
+}
+
+//////////////////////////////////////////////////////////////////////////
+//
+
+FSequencerTrackFilter_CameraCut::FSequencerTrackFilter_CameraCut(ISequencerTrackFilters& InFilterInterface, TSharedPtr<FFilterCategory> InCategory)
+	: FSequencerTrackFilter_ClassType<UMovieSceneCameraCutTrack>(InFilterInterface, InCategory)
+{
+}
+
+FText FSequencerTrackFilter_CameraCut::GetDisplayName() const
+{
+	return LOCTEXT("SequencerTrackFilter_CameraCut", "Camera Cut");
+}
+
+FSlateIcon FSequencerTrackFilter_CameraCut::GetIcon() const
+{
+	return FSlateIcon(FAppStyle::GetAppStyleSetName(), TEXT("Sequencer.Tracks.CameraCut"));
+}
+
+FText FSequencerTrackFilter_CameraCut::GetDefaultToolTipText() const
+{
+	return LOCTEXT("SequencerTrackFilter_CameraCutToolTip", "Show only Camera Cut tracks");
+}
+
+TSharedPtr<FUICommandInfo> FSequencerTrackFilter_CameraCut::GetToggleCommand() const
+{
+	return FSequencerTrackFilterCommands::Get().ToggleFilter_CameraCut;
+}
+
+//////////////////////////////////////////////////////////////////////////
+//
+
+FSequencerTrackFilter_Fade::FSequencerTrackFilter_Fade(ISequencerTrackFilters& InFilterInterface, TSharedPtr<FFilterCategory> InCategory)
+	: FSequencerTrackFilter_ClassType<UMovieSceneFadeTrack>(InFilterInterface, InCategory)
+{
+}
+
+FText FSequencerTrackFilter_Fade::GetDisplayName() const
+{
+	return LOCTEXT("SequencerTrackFilter_Fade", "Fade");
+}
+
+FSlateIcon FSequencerTrackFilter_Fade::GetIcon() const
+{
+	return FSlateIcon(FAppStyle::GetAppStyleSetName(), TEXT("Sequencer.Tracks.Fade"));
+}
+
+FText FSequencerTrackFilter_Fade::GetDefaultToolTipText() const
+{
+	return LOCTEXT("SequencerTrackFilter_FadeToolTip", "Show only Fade tracks");
+}
+
+TSharedPtr<FUICommandInfo> FSequencerTrackFilter_Fade::GetToggleCommand() const
+{
+	return FSequencerTrackFilterCommands::Get().ToggleFilter_Fade;
+}
+
+//////////////////////////////////////////////////////////////////////////
+//
+
+FSequencerTrackFilter_DataLayer::FSequencerTrackFilter_DataLayer(ISequencerTrackFilters& InFilterInterface, TSharedPtr<FFilterCategory> InCategory)
+	: FSequencerTrackFilter_ClassType<UMovieSceneDataLayerTrack>(InFilterInterface, InCategory)
+{
+}
+
+FText FSequencerTrackFilter_DataLayer::GetDisplayName() const
+{
+	return LOCTEXT("SequencerTrackFilter_DataLayer", "Data Layer");
+}
+
+FSlateIcon FSequencerTrackFilter_DataLayer::GetIcon() const
+{
+	return FSlateIcon(FAppStyle::GetAppStyleSetName(), TEXT("Sequencer.Tracks.DataLayer"));
+}
+
+FText FSequencerTrackFilter_DataLayer::GetDefaultToolTipText() const
+{
+	return LOCTEXT("SequencerTrackFilter_DataLayerToolTip", "Show only Data Layer tracks");
+}
+
+TSharedPtr<FUICommandInfo> FSequencerTrackFilter_DataLayer::GetToggleCommand() const
+{
+	return FSequencerTrackFilterCommands::Get().ToggleFilter_DataLayer;
+}
+
+//////////////////////////////////////////////////////////////////////////
+//
+
+FSequencerTrackFilter_TimeDilation::FSequencerTrackFilter_TimeDilation(ISequencerTrackFilters& InFilterInterface, TSharedPtr<FFilterCategory> InCategory)
+	: FSequencerTrackFilter_ClassType<UMovieSceneSlomoTrack>(InFilterInterface, InCategory)
+{
+}
+
+FText FSequencerTrackFilter_TimeDilation::GetDisplayName() const
+{
+	return LOCTEXT("SequencerTrackFilter_TimeDilation", "Time Dilation");
+}
+
+FSlateIcon FSequencerTrackFilter_TimeDilation::GetIcon() const
+{
+	return FSlateIcon(FAppStyle::GetAppStyleSetName(), TEXT("Sequencer.Tracks.Slomo"));
+}
+
+FText FSequencerTrackFilter_TimeDilation::GetDefaultToolTipText() const
+{
+	return LOCTEXT("SequencerTrackFilter_TimeDilationToolTip", "Show only Time Dilation tracks");
+}
+
+TSharedPtr<FUICommandInfo> FSequencerTrackFilter_TimeDilation::GetToggleCommand() const
+{
+	return FSequencerTrackFilterCommands::Get().ToggleFilter_TimeDilation;
+}
+
+//////////////////////////////////////////////////////////////////////////
+//
+
+FSequencerTrackFilter_Folder::FSequencerTrackFilter_Folder(ISequencerTrackFilters& InFilterInterface, TSharedPtr<FFilterCategory> InCategory)
+	: FSequencerTrackFilter_ModelType<UE::Sequencer::FFolderModel>(InFilterInterface, InCategory)
+{
+}
+
+FText FSequencerTrackFilter_Folder::GetDisplayName() const
+{
+	return LOCTEXT("SequencerTrackFilter_Folder", "Folder");
+}
+
+FSlateIcon FSequencerTrackFilter_Folder::GetIcon() const
+{
+	return FSlateIcon(FAppStyle::GetAppStyleSetName(), TEXT("ContentBrowser.AssetTreeFolderClosed"));
+}
+
+FText FSequencerTrackFilter_Folder::GetDefaultToolTipText() const
+{
+	return LOCTEXT("SequencerTrackFilter_FolderToolTip", "Show only Folder tracks");
+}
+
+TSharedPtr<FUICommandInfo> FSequencerTrackFilter_Folder::GetToggleCommand() const
+{
+	return FSequencerTrackFilterCommands::Get().ToggleFilter_Folder;
+}
+
+TSubclassOf<UMovieSceneTrack> FSequencerTrackFilter_Folder::GetTrackClass() const
+{
+	return UMovieSceneFolder::StaticClass();
+}
+
+bool FSequencerTrackFilter_Folder::SupportsSequence(UMovieSceneSequence* const InSequence) const
+{
+	return true;
+}
+
+#undef LOCTEXT_NAMESPACE
diff --git a/Engine/Source/Editor/Sequencer/Private/Filters/Filters/SequencerTrackFilters.h b/Engine/Source/Editor/Sequencer/Private/Filters/Filters/SequencerTrackFilters.h
deleted file mode 100644
diff --git a/Engine/Source/Editor/Sequencer/Private/Filters/Menus/SequencerTrackFilterMenu.cpp b/Engine/Source/Editor/Sequencer/Private/Filters/Menus/SequencerTrackFilterMenu.cpp
--- a/Engine/Source/Editor/Sequencer/Private/Filters/Menus/SequencerTrackFilterMenu.cpp
+++ b/Engine/Source/Editor/Sequencer/Private/Filters/Menus/SequencerTrackFilterMenu.cpp
@@ -3,6 +3,8 @@
 #include "SequencerTrackFilterMenu.h"
 #include "Filters/SequencerFilterBar.h"
 #include "Filters/Filters/SequencerTrackFilter_CustomText.h"
+#include "Filters/Filters/SequencerTrackFilter_Level.h"
+#include "Filters/Filters/SequencerTrackFilter_Group.h"
 #include "Filters/Widgets/SSequencerFilterBar.h"
 #include "Filters/Widgets/SSequencerCustomTextFilterDialog.h"
 #include "ISequencer.h"
@@ -167,7 +169,8 @@ void FSequencerTrackFilterMenu::PopulateOtherFilterSections(UToolMenu& InMenu)
 
 	const UMovieSceneSequence* const FocusedSequence = FilterBar->GetSequencer().GetFocusedMovieSceneSequence();
 
-	if (FocusedSequence->IsA(ULevelSequence::StaticClass()))
+	if (FocusedSequence->IsA(ULevelSequence::StaticClass())
+		&& FocusedSequence->IsFilterSupported(FSequencerTrackFilter_Level::StaticName()))
 	{
 		Section.AddSubMenu(TEXT("LevelFilters"),
 			LOCTEXT("LevelFilters", "Level Filters"),
@@ -184,19 +187,22 @@ void FSequencerTrackFilterMenu::PopulateOtherFilterSections(UToolMenu& InMenu)
 			false);
 	}
 
-	Section.AddSubMenu(TEXT("GroupFilters"),
-		LOCTEXT("GroupFilters", "Group Filters"),
-		LOCTEXT("GroupFiltersToolTip", "Filter tracks by group"),
-		FNewToolMenuDelegate::CreateSP(this, &FSequencerTrackFilterMenu::FillGroupFilterMenu),
-		FUIAction(
-			FExecuteAction::CreateSP(this, &FSequencerTrackFilterMenu::ToggleAllGroupFilters),
-			FCanExecuteAction(),
-			FGetActionCheckState::CreateSP(this, &FSequencerTrackFilterMenu::GetAllGroupsCheckState)
-		),
-		EUserInterfaceActionType::ToggleButton,
-		false,
-		FSlateIcon(FAppStyle::Get().GetStyleSetName(), TEXT("Icons.GroupActors")),
-		false);
+	if (FocusedSequence->IsFilterSupported(FSequencerTrackFilter_Group::StaticName()))
+	{
+		Section.AddSubMenu(TEXT("GroupFilters"),
+			LOCTEXT("GroupFilters", "Group Filters"),
+			LOCTEXT("GroupFiltersToolTip", "Filter tracks by group"),
+			FNewToolMenuDelegate::CreateSP(this, &FSequencerTrackFilterMenu::FillGroupFilterMenu),
+			FUIAction(
+				FExecuteAction::CreateSP(this, &FSequencerTrackFilterMenu::ToggleAllGroupFilters),
+				FCanExecuteAction(),
+				FGetActionCheckState::CreateSP(this, &FSequencerTrackFilterMenu::GetAllGroupsCheckState)
+			),
+			EUserInterfaceActionType::ToggleButton,
+			false,
+			FSlateIcon(FAppStyle::Get().GetStyleSetName(), TEXT("Icons.GroupActors")),
+			false);
+	}
 }
 
 void FSequencerTrackFilterMenu::FillLevelFilterMenu(UToolMenu* const InMenu)
@@ -360,12 +366,9 @@ void FSequencerTrackFilterMenu::FillFiltersMenuCategory(FToolMenuSection& InOutS
 
 	for (const TSharedRef<FSequencerTrackFilter>& Filter : FilterBar->GetCommonFilters())
 	{
-		if (Filter->GetCategory() != InMenuCategory)
-		{
-			continue;
-		}
-
-		if (!Filter->SupportsSequence(FocusedSequence))
+		if (!FocusedSequence->IsFilterSupported(Filter->GetName())
+			|| Filter->GetCategory() != InMenuCategory
+			|| !Filter->SupportsSequence(FocusedSequence))
 		{
 			continue;
 		}
diff --git a/Engine/Source/Editor/Sequencer/Private/Filters/SequencerFilterBar.cpp b/Engine/Source/Editor/Sequencer/Private/Filters/SequencerFilterBar.cpp
--- a/Engine/Source/Editor/Sequencer/Private/Filters/SequencerFilterBar.cpp
+++ b/Engine/Source/Editor/Sequencer/Private/Filters/SequencerFilterBar.cpp
@@ -1279,18 +1279,25 @@ bool FSequencerFilterBar::IsGroupFilterActive(const FString& InGroupName) const
 
 bool FSequencerFilterBar::PassesAnyCommonFilter(const TViewModelPtr<IOutlinerExtension>& InNode)
 {
+	const UMovieSceneSequence* const FocusedSequence = Sequencer.GetFocusedMovieSceneSequence();
+	if (!IsValid(FocusedSequence))
+	{
+		return true;
+	}
+
 	bool bPassedAnyFilters = false;
 	bool bAnyFilterActive = false;
 
 	// Only one common filter needs to pass for this node to be included in the filtered set
-	CommonFilters->ForEachFilter([this, &InNode, &bPassedAnyFilters, &bAnyFilterActive]
+	CommonFilters->ForEachFilter([this, &InNode, &bPassedAnyFilters, &bAnyFilterActive, FocusedSequence]
 		(const TSharedRef<FSequencerTrackFilter>& InFilter)
 		{
 			const FString FilterName = InFilter->GetDisplayName().ToString();
-			if (IsFilterActive(InFilter))
+			if (FocusedSequence->IsFilterSupported(FilterName)
+				&& IsFilterActive(InFilter))
 			{
 				bAnyFilterActive = true;
-				
+			
 				if (InFilter->PassesFilter(InNode))
 				{
 					bPassedAnyFilters = true;
@@ -1312,12 +1319,19 @@ bool FSequencerFilterBar::PassesAnyCommonFilter(const TViewModelPtr<IOutlinerExt
 
 bool FSequencerFilterBar::PassesAllInternalFilters(const TViewModelPtr<IOutlinerExtension>& InNode)
 {
+	const UMovieSceneSequence* const FocusedSequence = Sequencer.GetFocusedMovieSceneSequence();
+	if (!IsValid(FocusedSequence))
+	{
+		return true;
+	}
+
 	bool bPassedAllFilters = true;
 
-	InternalFilters->ForEachFilter([this, &InNode, &bPassedAllFilters]
+	InternalFilters->ForEachFilter([this, &InNode, &bPassedAllFilters, FocusedSequence]
 		(const TSharedRef<FSequencerTrackFilter>& InFilter)
 		{
-			if (!InFilter->PassesFilter(InNode))
+			if (FocusedSequence->IsFilterSupported(InFilter->GetName())
+				&& !InFilter->PassesFilter(InNode))
 			{
 				bPassedAllFilters = false;
 				return false; // Stop processing filters
diff --git a/Engine/Source/Editor/Sequencer/Private/Filters/SequencerTrackFilterCommands.cpp b/Engine/Source/Editor/Sequencer/Private/Filters/SequencerTrackFilterCommands.cpp
--- a/Engine/Source/Editor/Sequencer/Private/Filters/SequencerTrackFilterCommands.cpp
+++ b/Engine/Source/Editor/Sequencer/Private/Filters/SequencerTrackFilterCommands.cpp
@@ -1,6 +1,6 @@
 // Copyright Epic Games, Inc. All Rights Reserved.
 
-#include "SequencerTrackFilterCommands.h"
+#include "Filters/SequencerTrackFilterCommands.h"
 #include "Styling/AppStyle.h"
 
 #define LOCTEXT_NAMESPACE "SequencerTrackFilterCommands"
diff --git a/Engine/Source/Editor/Sequencer/Public/Filters/Filters/SequencerTrackFilter_Condition.h b/Engine/Source/Editor/Sequencer/Public/Filters/Filters/SequencerTrackFilter_Condition.h
new file mode 100644
--- /dev/null
+++ b/Engine/Source/Editor/Sequencer/Public/Filters/Filters/SequencerTrackFilter_Condition.h
@@ -0,0 +1,29 @@
+// Copyright Epic Games, Inc. All Rights Reserved.
+
+#pragma once
+
+#include "Filters/SequencerTrackFilterBase.h"
+
+class FSequencerTrackFilter_Condition : public FSequencerTrackFilter
+{
+public:
+	static FString StaticName() { return TEXT("Condition"); }
+
+	FSequencerTrackFilter_Condition(ISequencerTrackFilters& InFilterInterface, TSharedPtr<FFilterCategory> InCategory = nullptr);
+
+	//~ Begin FSequencerTrackFilter
+	virtual bool ShouldUpdateOnTrackValueChanged() const override;
+	virtual FText GetDefaultToolTipText() const override;
+	virtual TSharedPtr<FUICommandInfo> GetToggleCommand() const override;
+	//~ End FSequencerTrackFilter
+
+	//~ Begin FFilterBase
+	virtual FText GetDisplayName() const override;
+	virtual FSlateIcon GetIcon() const override;
+	//~ End FFilterBase
+
+	//~ Begin IFilter
+	virtual FString GetName() const override;
+	virtual bool PassesFilter(FSequencerTrackFilterType InItem) const override;
+	//~ End IFilter
+};
diff --git a/Engine/Source/Editor/Sequencer/Public/Filters/Filters/SequencerTrackFilter_Group.h b/Engine/Source/Editor/Sequencer/Public/Filters/Filters/SequencerTrackFilter_Group.h
new file mode 100644
--- /dev/null
+++ b/Engine/Source/Editor/Sequencer/Public/Filters/Filters/SequencerTrackFilter_Group.h
@@ -0,0 +1,45 @@
+﻿// Copyright Epic Games, Inc. All Rights Reserved.
+
+#pragma once
+
+#include "Filters/SequencerTrackFilterBase.h"
+
+class ISequencer;
+class UMovieSceneNodeGroup;
+
+class FSequencerTrackFilter_Group : public FSequencerTrackFilter
+{
+public:
+	static FString StaticName() { return TEXT("Group"); }
+
+	static void ForEachMovieSceneNodeGroup(UMovieScene* const InMovieScene
+		, FSequencerTrackFilterType InItem
+		, const TFunctionRef<bool(const UE::Sequencer::TViewModelPtr<UE::Sequencer::IOutlinerExtension>& InParent, UMovieSceneNodeGroup*)>& InFunction);
+
+	FSequencerTrackFilter_Group(ISequencerTrackFilters& InFilterInterface, TSharedPtr<FFilterCategory> InCategory = nullptr);
+	virtual ~FSequencerTrackFilter_Group() override;
+
+	//~ Begin FSequencerTrackFilter
+	virtual FText GetDefaultToolTipText() const override;
+	virtual TSharedPtr<FUICommandInfo> GetToggleCommand() const override;
+	//~ End FSequencerTrackFilter
+
+	//~ Begin FFilterBase
+	virtual FText GetDisplayName() const override;
+	virtual FSlateIcon GetIcon() const override;
+	//~ End FFilterBase
+
+	//~ Begin IFilter
+	virtual FString GetName() const override;
+	virtual bool PassesFilter(FSequencerTrackFilterType InItem) const override;
+	//~ End IFilter
+
+	void UpdateMovieScene(UMovieScene* const InMovieScene);
+
+	bool HasActiveGroupFilter() const;
+
+private:
+	void HandleGroupsChanged();
+
+	TWeakObjectPtr<UMovieScene> MovieSceneWeak;
+};
diff --git a/Engine/Source/Editor/Sequencer/Public/Filters/Filters/SequencerTrackFilter_HideIsolate.h b/Engine/Source/Editor/Sequencer/Public/Filters/Filters/SequencerTrackFilter_HideIsolate.h
new file mode 100644
--- /dev/null
+++ b/Engine/Source/Editor/Sequencer/Public/Filters/Filters/SequencerTrackFilter_HideIsolate.h
@@ -0,0 +1,61 @@
+﻿// Copyright Epic Games, Inc. All Rights Reserved.
+
+#pragma once
+
+#include "Filters/SequencerTrackFilterBase.h"
+
+class FSequencerTrackFilter_HideIsolate : public FSequencerTrackFilter
+{
+public:
+	static FString StaticName() { return TEXT("HideIsolate"); }
+
+	FSequencerTrackFilter_HideIsolate(ISequencerTrackFilters& InFilterInterface, TSharedPtr<FFilterCategory> InCategory = nullptr);
+
+	void ResetFilter();
+
+	TSet<UE::Sequencer::TWeakViewModelPtr<UE::Sequencer::IOutlinerExtension>> GetHiddenTracks() const;
+	TSet<UE::Sequencer::TWeakViewModelPtr<UE::Sequencer::IOutlinerExtension>> GetIsolatedTracks() const;
+
+	void HideTracks(const TSet<UE::Sequencer::TWeakViewModelPtr<UE::Sequencer::IOutlinerExtension>>& InTracks, const bool bInAddToExisting);
+	void UnhideTracks(const TSet<UE::Sequencer::TWeakViewModelPtr<UE::Sequencer::IOutlinerExtension>>& InTracks);
+
+	void IsolateTracks(const TSet<UE::Sequencer::TWeakViewModelPtr<UE::Sequencer::IOutlinerExtension>>& InTracks, const bool bInAddToExisting);
+	void UnisolateTracks(const TSet<UE::Sequencer::TWeakViewModelPtr<UE::Sequencer::IOutlinerExtension>>& InTracks);
+
+	void IsolateCategoryGroupTracks(const TSet<UE::Sequencer::TWeakViewModelPtr<UE::Sequencer::IOutlinerExtension>>& InTracks
+		, const TSet<FName>& InCategoryNames
+		, const bool bInAddToExisting);
+
+	void ShowAllTracks();
+
+	bool HasHiddenTracks() const;
+	bool HasIsolatedTracks() const;
+	bool HasHiddenOrIsolatedTracks() const;
+
+	bool IsTrackHidden(const UE::Sequencer::TViewModelPtr<UE::Sequencer::IOutlinerExtension>& InTrack) const;
+	bool IsTrackIsolated(const UE::Sequencer::TViewModelPtr<UE::Sequencer::IOutlinerExtension>& InTrack) const;
+
+	void EmptyHiddenTracks(const bool bInBroadcastChange = true);
+	void EmptyIsolatedTracks(const bool bInBroadcastChange = true);
+
+	//~ Begin FSequencerTrackFilter
+	virtual void BindCommands() override;
+	//~ End FSequencerTrackFilter
+
+	//~ Begin FFilterBase
+	virtual FText GetDisplayName() const override;
+	virtual FText GetToolTipText() const override;
+	virtual FSlateIcon GetIcon() const override;
+	//~ End FFilterBase
+
+	//~ Begin IFilter
+	virtual FString GetName() const override;
+	virtual bool PassesFilter(FSequencerTrackFilterType InItem) const override;
+	//~ End IFilter
+
+	bool IsActive() const;
+
+private:
+	TSet<UE::Sequencer::TWeakViewModelPtr<UE::Sequencer::IOutlinerExtension>> HiddenTracks;
+	TSet<UE::Sequencer::TWeakViewModelPtr<UE::Sequencer::IOutlinerExtension>> IsolatedTracks;
+};
diff --git a/Engine/Source/Editor/Sequencer/Public/Filters/Filters/SequencerTrackFilter_Keyed.h b/Engine/Source/Editor/Sequencer/Public/Filters/Filters/SequencerTrackFilter_Keyed.h
new file mode 100644
--- /dev/null
+++ b/Engine/Source/Editor/Sequencer/Public/Filters/Filters/SequencerTrackFilter_Keyed.h
@@ -0,0 +1,32 @@
+﻿// Copyright Epic Games, Inc. All Rights Reserved.
+
+#pragma once
+
+#include "Filters/SequencerTrackFilterBase.h"
+
+class FSequencerTrackFilter_Keyed : public FSequencerTrackFilter
+{
+public:
+	static FString StaticName() { return TEXT("Keyed"); }
+
+	static bool DoesTrackExtensionHaveKeys(const UE::Sequencer::TWeakViewModelPtr<UE::Sequencer::ITrackExtension>& InTrack);
+
+	FSequencerTrackFilter_Keyed(ISequencerTrackFilters& InFilterInterface, TSharedPtr<FFilterCategory> InCategory = nullptr);
+
+	//~ Begin FSequencerTrackFilter
+	virtual bool ShouldUpdateOnTrackValueChanged() const override;
+	virtual FText GetDefaultToolTipText() const override;
+	virtual TSharedPtr<FUICommandInfo> GetToggleCommand() const override;
+	virtual bool SupportsSequence(UMovieSceneSequence* const InSequence) const override;
+	//~ End FSequencerTrackFilter
+
+	//~ Begin FFilterBase
+	virtual FText GetDisplayName() const override;
+	virtual FSlateIcon GetIcon() const override;
+	//~ End FFilterBase
+
+	//~ Begin IFilter
+	virtual FString GetName() const override;
+	virtual bool PassesFilter(FSequencerTrackFilterType InItem) const override;
+	//~ End IFilter
+};
diff --git a/Engine/Source/Editor/Sequencer/Public/Filters/Filters/SequencerTrackFilter_Level.h b/Engine/Source/Editor/Sequencer/Public/Filters/Filters/SequencerTrackFilter_Level.h
new file mode 100644
--- /dev/null
+++ b/Engine/Source/Editor/Sequencer/Public/Filters/Filters/SequencerTrackFilter_Level.h
@@ -0,0 +1,55 @@
+﻿// Copyright Epic Games, Inc. All Rights Reserved.
+
+#pragma once
+
+#include "Filters/SequencerTrackFilterBase.h"
+
+class FSequencerTrackFilter_Level : public FSequencerTrackFilter
+{
+public:
+	static FString StaticName() { return TEXT("Level"); }
+
+	FSequencerTrackFilter_Level(ISequencerTrackFilters& InFilterInterface, TSharedPtr<FFilterCategory> InCategory = nullptr);
+	virtual ~FSequencerTrackFilter_Level() override;
+
+	//~ Begin FSequencerTrackFilter
+	virtual FText GetDefaultToolTipText() const override;
+	virtual TSharedPtr<FUICommandInfo> GetToggleCommand() const override;
+	//~ End FSequencerTrackFilter
+
+	//~ Begin FFilterBase
+	virtual FText GetDisplayName() const override;
+	virtual FText GetToolTipText() const override;
+	virtual FSlateIcon GetIcon() const override;
+	//~ End FFilterBase
+
+	//~ Begin IFilter
+	virtual FString GetName() const override;
+	virtual bool PassesFilter(FSequencerTrackFilterType InItem) const override;
+	//~ End IFilter
+
+	bool IsActive() const;
+
+	void UpdateWorld(UWorld* const InWorld);
+	void ResetFilter();
+
+	const TSet<FString>& GetAllWorldLevels() const;
+
+	bool HasHiddenLevels() const;
+	bool HasAllLevelsHidden() const;
+	const TSet<FString>& GetHiddenLevels() const;
+	bool IsLevelHidden(const FString& InLevelName) const;
+	void HideLevel(const FString& InLevelName);
+	void UnhideLevel(const FString& InLevelName);
+
+	void HideAllLevels(const bool bInHide);
+	bool CanHideAllLevels(const bool bInHide) const;
+
+private:
+	void HandleLevelsChanged();
+
+	TSet<FString> HiddenLevels;
+
+	TWeakObjectPtr<UWorld> CachedWorld;
+	TSet<FString> AllWorldLevels;
+};
diff --git a/Engine/Source/Editor/Sequencer/Public/Filters/Filters/SequencerTrackFilter_Modified.h b/Engine/Source/Editor/Sequencer/Public/Filters/Filters/SequencerTrackFilter_Modified.h
new file mode 100644
--- /dev/null
+++ b/Engine/Source/Editor/Sequencer/Public/Filters/Filters/SequencerTrackFilter_Modified.h
@@ -0,0 +1,31 @@
+﻿// Copyright Epic Games, Inc. All Rights Reserved.
+
+#pragma once
+
+#include "Filters/SequencerTrackFilterBase.h"
+
+class FSequencerTrackFilter_Modified : public FSequencerTrackFilter
+{
+public:
+	static FString StaticName() { return TEXT("Modified"); }
+	
+	FSequencerTrackFilter_Modified(ISequencerTrackFilters& InFilterInterface, TSharedPtr<FFilterCategory> InCategory = nullptr);
+
+	//~ Begin FSequencerTrackFilter
+	virtual bool ShouldUpdateOnTrackValueChanged() const override;
+	virtual FText GetDefaultToolTipText() const override;
+	virtual TSharedPtr<FUICommandInfo> GetToggleCommand() const override;
+	//~ End FSequencerTrackFilter
+
+	//~ Begin FFilterBase
+	virtual FText GetDisplayName() const override;
+	virtual FSlateIcon GetIcon() const override;
+	//~ End FFilterBase
+
+	//~ Begin IFilter
+	virtual FString GetName() const override;
+	virtual bool PassesFilter(FSequencerTrackFilterType InItem) const override;
+	//~ End IFilter
+
+	void ToggleShowOnlyModifiedTracks();
+};
diff --git a/Engine/Source/Editor/Sequencer/Public/Filters/Filters/SequencerTrackFilter_Selected.h b/Engine/Source/Editor/Sequencer/Public/Filters/Filters/SequencerTrackFilter_Selected.h
new file mode 100644
--- /dev/null
+++ b/Engine/Source/Editor/Sequencer/Public/Filters/Filters/SequencerTrackFilter_Selected.h
@@ -0,0 +1,47 @@
+﻿// Copyright Epic Games, Inc. All Rights Reserved.
+
+#pragma once
+
+#include "Filters/SequencerTrackFilterBase.h"
+
+class FEditorModeTools;
+class ILevelEditor;
+
+class FSequencerTrackFilter_Selected : public FSequencerTrackFilter
+{
+public:
+	static FString StaticName() { return TEXT("Selected"); }
+
+	FSequencerTrackFilter_Selected(ISequencerTrackFilters& InFilterInterface, TSharedPtr<FFilterCategory> InCategory = nullptr);
+	virtual ~FSequencerTrackFilter_Selected() override;
+
+	//~ Begin FSequencerTrackFilter
+	virtual bool ShouldUpdateOnTrackValueChanged() const override;
+	virtual FText GetDefaultToolTipText() const override;
+	virtual TSharedPtr<FUICommandInfo> GetToggleCommand() const override;
+	virtual void ActiveStateChanged(const bool bInActive) override;
+	//~ End FSequencerTrackFilter
+
+	//~ Begin FFilterBase
+	virtual FText GetDisplayName() const override;
+	virtual FSlateIcon GetIcon() const override;
+	//~ End FFilterBase
+
+	//~ Begin IFilter
+	virtual FString GetName() const override;
+	virtual bool PassesFilter(FSequencerTrackFilterType InItem) const override;
+	//~ End IFilter
+
+	void ToggleShowOnlySelectedTracks();
+
+protected:
+	void BindSelectionChanged();
+	void UnbindSelectionChanged();
+
+	TSharedPtr<ILevelEditor> GetLevelEditor() const;
+	FEditorModeTools* GetEditorModeManager() const;
+
+	void OnSelectionChanged(UObject* const InObject);
+
+	FDelegateHandle OnSelectionChangedHandle;
+};
diff --git a/Engine/Source/Editor/Sequencer/Public/Filters/Filters/SequencerTrackFilter_TimeWarp.h b/Engine/Source/Editor/Sequencer/Public/Filters/Filters/SequencerTrackFilter_TimeWarp.h
new file mode 100644
--- /dev/null
+++ b/Engine/Source/Editor/Sequencer/Public/Filters/Filters/SequencerTrackFilter_TimeWarp.h
@@ -0,0 +1,29 @@
+// Copyright Epic Games, Inc. All Rights Reserved.
+
+#pragma once
+
+#include "Filters/SequencerTrackFilterBase.h"
+
+class FSequencerTrackFilter_TimeWarp : public FSequencerTrackFilter
+{
+public:
+	static FString StaticName() { return TEXT("TimeWarp"); }
+
+	FSequencerTrackFilter_TimeWarp(ISequencerTrackFilters& InFilterInterface, TSharedPtr<FFilterCategory> InCategory = nullptr);
+
+	//~ Begin FSequencerTrackFilter
+	virtual bool ShouldUpdateOnTrackValueChanged() const override;
+	virtual FText GetDefaultToolTipText() const override;
+	virtual TSharedPtr<FUICommandInfo> GetToggleCommand() const override;
+	//~ End FSequencerTrackFilter
+
+	//~ Begin FFilterBase
+	virtual FText GetDisplayName() const override;
+	virtual FSlateIcon GetIcon() const override;
+	//~ End FFilterBase
+
+	//~ Begin IFilter
+	virtual FString GetName() const override;
+	virtual bool PassesFilter(FSequencerTrackFilterType InItem) const override;
+	//~ End IFilter
+};
diff --git a/Engine/Source/Editor/Sequencer/Public/Filters/Filters/SequencerTrackFilter_Unbound.h b/Engine/Source/Editor/Sequencer/Public/Filters/Filters/SequencerTrackFilter_Unbound.h
new file mode 100644
--- /dev/null
+++ b/Engine/Source/Editor/Sequencer/Public/Filters/Filters/SequencerTrackFilter_Unbound.h
@@ -0,0 +1,31 @@
+﻿// Copyright Epic Games, Inc. All Rights Reserved.
+
+#pragma once
+
+#include "Filters/SequencerTrackFilterBase.h"
+
+class FSequencerTrackFilter_Unbound : public FSequencerTrackFilter
+{
+public:
+	static FString StaticName() { return TEXT("Unbound"); }
+
+	FSequencerTrackFilter_Unbound(ISequencerTrackFilters& InFilterInterface, TSharedPtr<FFilterCategory> InCategory = nullptr);
+
+	//~ Begin FSequencerTrackFilter
+	virtual bool ShouldUpdateOnTrackValueChanged() const override;
+	virtual FText GetDefaultToolTipText() const override;
+	virtual TSharedPtr<FUICommandInfo> GetToggleCommand() const override;
+	//~ End FSequencerTrackFilter
+
+	//~ Begin FFilterBase
+	virtual FText GetDisplayName() const override;
+	virtual FSlateIcon GetIcon() const override;
+	//~ End FFilterBase
+
+	//~ Begin IFilter
+	virtual FString GetName() const override;
+	virtual bool PassesFilter(FSequencerTrackFilterType InItem) const override;
+	//~ End IFilter
+
+	void ToggleShowOnlyModifiedTracks();
+};
diff --git a/Engine/Source/Editor/Sequencer/Public/Filters/Filters/SequencerTrackFilters.h b/Engine/Source/Editor/Sequencer/Public/Filters/Filters/SequencerTrackFilters.h
new file mode 100644
--- /dev/null
+++ b/Engine/Source/Editor/Sequencer/Public/Filters/Filters/SequencerTrackFilters.h
@@ -0,0 +1,371 @@
+// Copyright Epic Games, Inc. All Rights Reserved.
+
+#pragma once
+
+#include "Camera/CameraComponent.h"
+#include "Components/LightComponentBase.h"
+#include "Components/SkeletalMeshComponent.h"
+#include "Filters/SequencerTrackFilterBase.h"
+#include "MVVM/ViewModels/FolderModel.h"
+#include "Tracks/MovieSceneAudioTrack.h"
+#include "Tracks/MovieSceneCameraCutTrack.h"
+#include "Tracks/MovieSceneCinematicShotTrack.h"
+#include "Tracks/MovieSceneDataLayerTrack.h"
+#include "Tracks/MovieSceneEventTrack.h"
+#include "Tracks/MovieSceneFadeTrack.h"
+#include "Tracks/MovieSceneLevelVisibilityTrack.h"
+#include "Tracks/MovieSceneParticleTrack.h"
+#include "Tracks/MovieSceneSlomoTrack.h"
+#include "Tracks/MovieSceneSubTrack.h"
+
+class FSequencerTrackFilter_Audio : public FSequencerTrackFilter_ClassType<UMovieSceneAudioTrack>
+{
+public:
+	static FString StaticName() { return TEXT("Audio"); }
+
+	FSequencerTrackFilter_Audio(ISequencerTrackFilters& InFilterInterface, TSharedPtr<FFilterCategory> InCategory = nullptr);
+
+	//~ Begin IFilter
+	virtual FString GetName() const override { return StaticName(); }
+	//~ End IFilter
+
+	//~ Begin FFilterBase
+	virtual FText GetDisplayName() const override;
+	virtual FSlateIcon GetIcon() const override;
+	//~ End FFilterBase
+
+	//~ Begin FSequencerTrackFilter
+	virtual FText GetDefaultToolTipText() const override;
+	virtual TSharedPtr<FUICommandInfo> GetToggleCommand() const override;
+	virtual bool SupportsSequence(UMovieSceneSequence* const InSequence) const override;
+	//~ End FSequencerTrackFilter
+};
+
+//////////////////////////////////////////////////////////////////////////
+//
+
+class FSequencerTrackFilter_Event : public FSequencerTrackFilter_ClassType<UMovieSceneEventTrack>
+{
+public:
+	static FString StaticName() { return TEXT("Event"); }
+
+	FSequencerTrackFilter_Event(ISequencerTrackFilters& InFilterInterface, TSharedPtr<FFilterCategory> InCategory = nullptr);
+
+	//~ Begin IFilter
+	virtual FString GetName() const override { return StaticName(); }
+	//~ End IFilter
+
+	//~ Begin FFilterBase
+	virtual FText GetDisplayName() const override;
+	virtual FSlateIcon GetIcon() const override;
+	//~ End FFilterBase
+
+	//~ Begin FSequencerTrackFilter
+	virtual FText GetDefaultToolTipText() const override;
+	virtual TSharedPtr<FUICommandInfo> GetToggleCommand() const override;
+	virtual bool SupportsSequence(UMovieSceneSequence* const InSequence) const override;
+	//~ End FSequencerTrackFilter
+};
+
+//////////////////////////////////////////////////////////////////////////
+//
+
+class FSequencerTrackFilter_LevelVisibility : public FSequencerTrackFilter_ClassType<UMovieSceneLevelVisibilityTrack>
+{
+public:
+	static FString StaticName() { return TEXT("LevelVisibility"); }
+
+	FSequencerTrackFilter_LevelVisibility(ISequencerTrackFilters& InFilterInterface, TSharedPtr<FFilterCategory> InCategory = nullptr);
+
+	//~ Begin IFilter
+	virtual FString GetName() const override { return StaticName(); }
+	//~ End IFilter
+
+	//~ Begin FFilterBase
+	virtual FText GetDisplayName() const override;
+	virtual FSlateIcon GetIcon() const override;
+	//~ End FFilterBase
+
+	//~ Begin FSequencerTrackFilter
+	virtual FText GetDefaultToolTipText() const override;
+	virtual TSharedPtr<FUICommandInfo> GetToggleCommand() const override;
+	//~ End FSequencerTrackFilter
+};
+
+//////////////////////////////////////////////////////////////////////////
+//
+
+class FSequencerTrackFilter_Particle : public FSequencerTrackFilter_ClassType<UMovieSceneParticleTrack>
+{
+public:
+	static FString StaticName() { return TEXT("Particle"); }
+
+	FSequencerTrackFilter_Particle(ISequencerTrackFilters& InFilterInterface, TSharedPtr<FFilterCategory> InCategory = nullptr);
+
+	//~ Begin IFilter
+	virtual FString GetName() const override { return StaticName(); }
+	//~ End IFilter
+
+	//~ Begin FFilterBase
+	virtual FText GetDisplayName() const override;
+	virtual FSlateIcon GetIcon() const override;
+	//~ End FFilterBase
+
+	//~ Begin FSequencerTrackFilter
+	virtual FText GetDefaultToolTipText() const override;
+	virtual TSharedPtr<FUICommandInfo> GetToggleCommand() const override;
+	//~ End FSequencerTrackFilter
+};
+
+//////////////////////////////////////////////////////////////////////////
+//
+
+class FSequencerTrackFilter_CinematicShot : public FSequencerTrackFilter_ClassType<UMovieSceneCinematicShotTrack>
+{
+public:
+	static FString StaticName() { return TEXT("CinematicShot"); }
+
+	FSequencerTrackFilter_CinematicShot(ISequencerTrackFilters& InFilterInterface, TSharedPtr<FFilterCategory> InCategory = nullptr);
+
+	//~ Begin IFilter
+	virtual FString GetName() const override { return StaticName(); }
+	//~ End IFilter
+
+	//~ Begin FFilterBase
+	virtual FText GetDisplayName() const override;
+	virtual FSlateIcon GetIcon() const override;
+	//~ End FFilterBase
+
+	//~ Begin FSequencerTrackFilter
+	virtual FText GetDefaultToolTipText() const override;
+	virtual TSharedPtr<FUICommandInfo> GetToggleCommand() const override;
+	//~ End FSequencerTrackFilter
+};
+
+//////////////////////////////////////////////////////////////////////////
+//
+
+class FSequencerTrackFilter_Subsequence : public FSequencerTrackFilter_ClassType<UMovieSceneSubTrack>
+{
+public:
+	static FString StaticName() { return TEXT("SubSequence"); }
+
+	FSequencerTrackFilter_Subsequence(ISequencerTrackFilters& InFilterInterface, TSharedPtr<FFilterCategory> InCategory = nullptr);
+
+	//~ Begin IFilter
+	virtual FString GetName() const override { return StaticName(); }
+	virtual bool PassesFilter(FSequencerTrackFilterType InItem) const override;
+	//~ End IFilter
+
+	//~ Begin FFilterBase
+	virtual FText GetDisplayName() const override;
+	virtual FSlateIcon GetIcon() const override;
+	//~ End FFilterBase
+
+	//~ Begin FSequencerTrackFilter
+	virtual FText GetDefaultToolTipText() const override;
+	virtual TSharedPtr<FUICommandInfo> GetToggleCommand() const override;
+	//~ End FSequencerTrackFilter
+};
+
+//////////////////////////////////////////////////////////////////////////
+//
+
+class FSequencerTrackFilter_SkeletalMesh : public FSequencerTrackFilter_ComponentType<USkeletalMeshComponent>
+{
+public:
+	static FString StaticName() { return TEXT("SkeletalMesh"); }
+
+	FSequencerTrackFilter_SkeletalMesh(ISequencerTrackFilters& InFilterInterface, TSharedPtr<FFilterCategory> InCategory = nullptr);
+
+	//~ Begin IFilter
+	virtual FString GetName() const override { return StaticName(); }
+	//~ End IFilter
+
+	//~ Begin FFilterBase
+	virtual FText GetDisplayName() const override;
+	virtual FSlateIcon GetIcon() const override;
+	//~ End FFilterBase
+
+	//~ Begin FSequencerTrackFilter
+	virtual FText GetDefaultToolTipText() const override;
+	virtual TSharedPtr<FUICommandInfo> GetToggleCommand() const override;
+	//~ End FSequencerTrackFilter
+};
+
+//////////////////////////////////////////////////////////////////////////
+//
+
+class FSequencerTrackFilter_Camera : public FSequencerTrackFilter_ComponentType<UCameraComponent>
+{
+public:
+	static FString StaticName() { return TEXT("Camera"); }
+
+	FSequencerTrackFilter_Camera(ISequencerTrackFilters& InFilterInterface, TSharedPtr<FFilterCategory> InCategory = nullptr);
+
+	//~ Begin IFilter
+	virtual FString GetName() const override { return StaticName(); }
+	//~ End IFilter
+
+	//~ Begin FFilterBase
+	virtual FText GetDisplayName() const override;
+	virtual FSlateIcon GetIcon() const override;
+	//~ End FFilterBase
+
+	//~ Begin FSequencerTrackFilter
+	virtual FText GetDefaultToolTipText() const override;
+	virtual TSharedPtr<FUICommandInfo> GetToggleCommand() const override;
+	//~ End FSequencerTrackFilter
+};
+
+//////////////////////////////////////////////////////////////////////////
+//
+
+class FSequencerTrackFilter_Light : public FSequencerTrackFilter_ComponentType<ULightComponentBase>
+{
+public:
+	static FString StaticName() { return TEXT("Light"); }
+
+	FSequencerTrackFilter_Light(ISequencerTrackFilters& InFilterInterface, TSharedPtr<FFilterCategory> InCategory = nullptr);
+
+	//~ Begin IFilter
+	virtual FString GetName() const override { return StaticName(); }
+	//~ End IFilter
+
+	//~ Begin FFilterBase
+	virtual FText GetDisplayName() const override;
+	virtual FSlateIcon GetIcon() const override;
+	//~ End FFilterBase
+
+	//~ Begin FSequencerTrackFilter
+	virtual FText GetDefaultToolTipText() const override;
+	virtual TSharedPtr<FUICommandInfo> GetToggleCommand() const override;
+	//~ End FSequencerTrackFilter
+};
+
+//////////////////////////////////////////////////////////////////////////
+//
+
+class FSequencerTrackFilter_CameraCut : public FSequencerTrackFilter_ClassType<UMovieSceneCameraCutTrack>
+{
+public:
+	static FString StaticName() { return TEXT("CameraCut"); }
+
+	FSequencerTrackFilter_CameraCut(ISequencerTrackFilters& InFilterInterface, TSharedPtr<FFilterCategory> InCategory = nullptr);
+
+	//~ Begin IFilter
+	virtual FString GetName() const override { return StaticName(); }
+	//~ End IFilter
+
+	//~ Begin FFilterBase
+	virtual FText GetDisplayName() const override;
+	virtual FSlateIcon GetIcon() const override;
+	//~ End FFilterBase
+
+	//~ Begin FSequencerTrackFilter
+	virtual FText GetDefaultToolTipText() const override;
+	virtual TSharedPtr<FUICommandInfo> GetToggleCommand() const override;
+	//~ End FSequencerTrackFilter
+};
+
+//////////////////////////////////////////////////////////////////////////
+//
+
+class FSequencerTrackFilter_Fade : public FSequencerTrackFilter_ClassType<UMovieSceneFadeTrack>
+{
+public:
+	static FString StaticName() { return TEXT("Fade"); }
+
+	FSequencerTrackFilter_Fade(ISequencerTrackFilters& InFilterInterface, TSharedPtr<FFilterCategory> InCategory = nullptr);
+
+	//~ Begin IFilter
+	virtual FString GetName() const override { return StaticName(); }
+	//~ End IFilter
+
+	//~ Begin FFilterBase
+	virtual FText GetDisplayName() const override;
+	virtual FSlateIcon GetIcon() const override;
+	//~ End FFilterBase
+
+	//~ Begin FSequencerTrackFilter
+	virtual FText GetDefaultToolTipText() const override;
+	virtual TSharedPtr<FUICommandInfo> GetToggleCommand() const override;
+	//~ End FSequencerTrackFilter
+};
+
+//////////////////////////////////////////////////////////////////////////
+//
+
+class FSequencerTrackFilter_DataLayer : public FSequencerTrackFilter_ClassType<UMovieSceneDataLayerTrack>
+{
+public:
+	static FString StaticName() { return TEXT("DataLayer"); }
+
+	FSequencerTrackFilter_DataLayer(ISequencerTrackFilters& InFilterInterface, TSharedPtr<FFilterCategory> InCategory = nullptr);
+
+	//~ Begin IFilter
+	virtual FString GetName() const override { return StaticName(); }
+	//~ End IFilter
+
+	//~ Begin FFilterBase
+	virtual FText GetDisplayName() const override;
+	virtual FSlateIcon GetIcon() const override;
+	//~ End FFilterBase
+
+	//~ Begin FSequencerTrackFilter
+	virtual FText GetDefaultToolTipText() const override;
+	virtual TSharedPtr<FUICommandInfo> GetToggleCommand() const override;
+	//~ End FSequencerTrackFilter
+};
+
+//////////////////////////////////////////////////////////////////////////
+//
+
+class FSequencerTrackFilter_TimeDilation : public FSequencerTrackFilter_ClassType<UMovieSceneSlomoTrack>
+{
+public:
+	static FString StaticName() { return TEXT("TimeDilation"); }
+
+	FSequencerTrackFilter_TimeDilation(ISequencerTrackFilters& InFilterInterface, TSharedPtr<FFilterCategory> InCategory = nullptr);
+
+	//~ Begin IFilter
+	virtual FString GetName() const override { return StaticName(); }
+	//~ End IFilter
+
+	//~ Begin FFilterBase
+	virtual FText GetDisplayName() const override;
+	virtual FSlateIcon GetIcon() const override;
+	//~ End FFilterBase
+
+	//~ Begin FSequencerTrackFilter
+	virtual FText GetDefaultToolTipText() const override;
+	virtual TSharedPtr<FUICommandInfo> GetToggleCommand() const override;
+	//~ End FSequencerTrackFilter
+};
+
+//////////////////////////////////////////////////////////////////////////
+//
+
+class FSequencerTrackFilter_Folder : public FSequencerTrackFilter_ModelType<UE::Sequencer::FFolderModel>
+{
+public:
+	static FString StaticName() { return TEXT("Folder"); }
+
+	FSequencerTrackFilter_Folder(ISequencerTrackFilters& InFilterInterface, TSharedPtr<FFilterCategory> InCategory = nullptr);
+
+	//~ Begin IFilter
+	virtual FString GetName() const override { return StaticName(); }
+	//~ End IFilter
+
+	//~ Begin FFilterBase
+	virtual FText GetDisplayName() const override;
+	virtual FSlateIcon GetIcon() const override;
+	//~ End FFilterBase
+
+	//~ Begin FSequencerTrackFilter
+	virtual FText GetDefaultToolTipText() const override;
+	virtual TSharedPtr<FUICommandInfo> GetToggleCommand() const override;
+	virtual TSubclassOf<UMovieSceneTrack> GetTrackClass() const override;
+	virtual bool SupportsSequence(UMovieSceneSequence* const InSequence) const override;
+	//~ End FSequencerTrackFilter
+};
diff --git a/Engine/Source/Runtime/MovieScene/Private/MovieSceneSequence.cpp b/Engine/Source/Runtime/MovieScene/Private/MovieSceneSequence.cpp
--- a/Engine/Source/Runtime/MovieScene/Private/MovieSceneSequence.cpp
+++ b/Engine/Source/Runtime/MovieScene/Private/MovieSceneSequence.cpp
@@ -431,4 +431,9 @@ ETrackSupport UMovieSceneSequence::IsTrackSupported(TSubclassOf<class UMovieScen
 	return IsTrackSupportedImpl(InTrackClass); 
 }
 
+bool UMovieSceneSequence::IsFilterSupported(const FString& InFilterName) const
+{
+	return IsFilterSupportedImpl(InFilterName); 
+}
+
 #endif
diff --git a/Engine/Source/Runtime/MovieScene/Public/MovieSceneSequence.h b/Engine/Source/Runtime/MovieScene/Public/MovieSceneSequence.h
--- a/Engine/Source/Runtime/MovieScene/Public/MovieSceneSequence.h
+++ b/Engine/Source/Runtime/MovieScene/Public/MovieSceneSequence.h
@@ -441,5 +441,13 @@ public:
 	/** Subclasses can implement this to determine whether they support a particular track type */
 	virtual ETrackSupport IsTrackSupportedImpl(TSubclassOf<UMovieSceneTrack> InTrackClass) const { return ETrackSupport::Default; }
 
+	/*
+	 * Sequences can determine whether they support a particular filter type
+	 */
+	MOVIESCENE_API bool IsFilterSupported(const FString& InFilterName) const;
+
+	/** Subclasses can implement this to determine whether they support a particular filter type */
+	virtual bool IsFilterSupportedImpl(const FString& InFilterName) const { return true; }
+
 #endif
 };
