[Sequencer] Add "Align to Playhead" action and Sequencer toolbar dropdown menu entry

diff --git a/Engine/Source/Editor/Sequencer/Private/SSequencer.cpp b/Engine/Source/Editor/Sequencer/Private/SSequencer.cpp
--- a/Engine/Source/Editor/Sequencer/Private/SSequencer.cpp
+++ b/Engine/Source/Editor/Sequencer/Private/SSequencer.cpp
@@ -1703,6 +1703,7 @@ TSharedRef<SWidget> SSequencer::MakeActionsMenu()
 	MenuBuilder.AddMenuEntry(FSequencerCommands::Get().TranslateRight);
 	MenuBuilder.AddMenuEntry(FSequencerCommands::Get().TrimOrExtendSectionLeft);
 	MenuBuilder.AddMenuEntry(FSequencerCommands::Get().TrimOrExtendSectionRight);
+	MenuBuilder.AddMenuEntry(FSequencerCommands::Get().AlignSelectionToPlayhead);
 
 	MenuBuilder.EndSection();
 
diff --git a/Engine/Source/Editor/Sequencer/Private/Sequencer.cpp b/Engine/Source/Editor/Sequencer/Private/Sequencer.cpp
--- a/Engine/Source/Editor/Sequencer/Private/Sequencer.cpp
+++ b/Engine/Source/Editor/Sequencer/Private/Sequencer.cpp
@@ -174,6 +174,7 @@
 #include "Misc/Thumbnail/ThumbnailCaptureUtils.h"
 #include "ToolMenus.h"
 #include "SequencerToolMenuContext.h"
+#include "Tools/SequencerSelectionAlignmentUtils.h"
 
 #include "EngineModule.h"
 #include "IViewportSelectableObject.h"
@@ -11967,7 +11968,12 @@ void FSequencer::BindCommands()
 
 	SequencerCommandBindings->MapAction(
 		Commands.ToggleSidebarDrawerDock,
-		FExecuteAction::CreateSP( this, &FSequencer::ToggleSidebarDrawerDocked));
+		FExecuteAction::CreateSP(this, &FSequencer::ToggleSidebarDrawerDocked));
+
+	SequencerCommandBindings->MapAction(
+		Commands.AlignSelectionToPlayhead,
+		FExecuteAction::CreateSP(this, &FSequencer::AlignSelectionToPlayhead),
+		FCanExecuteAction::CreateSP(this, &FSequencer::CanAlignSelectionToPlayhead));
 
 	// If this sequencer supports a curve editor, let's add bindings for it.
 	FCurveEditorExtension* CurveEditorExtension = ViewModel->CastDynamic<FCurveEditorExtension>();
@@ -12710,4 +12716,14 @@ void FSequencer::SortSelectedOutlinerItems(const bool bInSortBySelectOrder, cons
 	SequencerHelpers::SortOutlinerItems(*this, OutlinerItems, bInSortBySelectOrder, bInDescending);
 }
 
+bool FSequencer::CanAlignSelectionToPlayhead()
+{
+	return FSequencerSelectionAlignmentUtils::CanAlignSelection(*this);
+}
+
+void FSequencer::AlignSelectionToPlayhead()
+{
+	FSequencerSelectionAlignmentUtils::AlignSelectionToPlayhead(*this);
+}
+
 #undef LOCTEXT_NAMESPACE
diff --git a/Engine/Source/Editor/Sequencer/Private/Sequencer.h b/Engine/Source/Editor/Sequencer/Private/Sequencer.h
--- a/Engine/Source/Editor/Sequencer/Private/Sequencer.h
+++ b/Engine/Source/Editor/Sequencer/Private/Sequencer.h
@@ -1278,6 +1278,9 @@ private:
 
 	void SortSelectedOutlinerItems(const bool bInSortBySelectOrder, const bool bInDescending);
 
+	bool CanAlignSelectionToPlayhead();
+	void AlignSelectionToPlayhead();
+
 	/** User-supplied settings object for this sequencer */
 	TObjectPtr<USequencerSettings> Settings;
 
diff --git a/Engine/Source/Editor/Sequencer/Private/SequencerCommands.cpp b/Engine/Source/Editor/Sequencer/Private/SequencerCommands.cpp
--- a/Engine/Source/Editor/Sequencer/Private/SequencerCommands.cpp
+++ b/Engine/Source/Editor/Sequencer/Private/SequencerCommands.cpp
@@ -200,6 +200,8 @@ void FSequencerCommands::RegisterCommands()
 	UI_COMMAND(ToggleSidebarVisible, "Sidebar Details", "Toggles the visibility of the Sequencer sidebar.", EUserInterfaceActionType::ToggleButton, FInputChord());
 	UI_COMMAND(ToggleSidebarSelectionDrawerOpen, "Toggle Sidebar Selection Drawer Open", "Toggles the visible state of the sidebar selection drawer. Undocks and closes the sidebar drawer if docked.", EUserInterfaceActionType::Button, FInputChord(EModifierKey::Control, EKeys::E));
 	UI_COMMAND(ToggleSidebarDrawerDock, "Toggle Sidebar Drawer Docked", "Toggles the docked state of the sidebar current drawer. Undocks the docked sidebar drawer if docked or docks the sidebar drawer if there is one open and no currently docked drawer.", EUserInterfaceActionType::Button, FInputChord(EModifierKey::Alt, EKeys::D));
+
+	UI_COMMAND(AlignSelectionToPlayhead, "Align Selection To Playhead", "Aligns layer bars, sections, or key selections to the current playhead location", EUserInterfaceActionType::Button, FInputChord());
 }
 
 #undef LOCTEXT_NAMESPACE
diff --git a/Engine/Source/Editor/Sequencer/Private/Tools/SequencerSelectionAlignmentUtils.cpp b/Engine/Source/Editor/Sequencer/Private/Tools/SequencerSelectionAlignmentUtils.cpp
new file mode 100644
--- /dev/null
+++ b/Engine/Source/Editor/Sequencer/Private/Tools/SequencerSelectionAlignmentUtils.cpp
@@ -0,0 +1,272 @@
+﻿// Copyright Epic Games, Inc. All Rights Reserved.
+
+#include "Tools/SequencerSelectionAlignmentUtils.h"
+#include "Containers/Set.h"
+#include "Curves/KeyHandle.h"
+#include "IKeyArea.h"
+#include "ISequencer.h"
+#include "Misc/FrameNumber.h"
+#include "Misc/Optional.h"
+#include "MovieSceneSection.h"
+#include "MVVM/Extensions/ITrackAreaExtension.h"
+#include "MVVM/Selection/Selection.h"
+#include "MVVM/ViewModels/LayerBarModel.h"
+#include "MVVM/ViewModels/SectionModel.h"
+#include "ScopedTransaction.h"
+#include "Sequencer.h"
+
+namespace UE::Sequencer
+{
+
+TSet<TViewModelPtr<FLayerBarModel>> GatherAllSelectedLayerBars(const TSharedRef<FSequencerSelection>& InSequencerSelection, TOptional<FFrameNumber>& OutEarliestTime)
+{
+	TSet<TViewModelPtr<FLayerBarModel>> LayerBarModels;
+
+	const int32 OutlinerCount = InSequencerSelection->Outliner.Num();
+	const int32 TrackAreaCount = InSequencerSelection->TrackArea.Num();
+
+	LayerBarModels.Reserve(OutlinerCount + TrackAreaCount);
+
+	for (const FViewModelPtr& ViewModel : InSequencerSelection->Outliner)
+	{
+		const TViewModelPtr<ITrackAreaExtension> Track = ViewModel.ImplicitCast();
+		if (!Track.IsValid())
+		{
+			continue;
+		}
+
+		for (const FViewModelPtr& TrackAreaModel : Track->GetTopLevelChildTrackAreaModels())
+		{
+			if (const TViewModelPtr<FLayerBarModel> LayerBarModel = TrackAreaModel.ImplicitCast())
+			{
+				LayerBarModels.Add(LayerBarModel);
+			}
+		}
+
+		for (const TViewModelPtr<FLayerBarModel>& LayerBarModel : Track->GetTrackAreaModelListAs<FLayerBarModel>())
+		{
+			LayerBarModels.Add(LayerBarModel);
+		}
+	}
+
+	for (const FViewModelPtr& ViewModel : InSequencerSelection->TrackArea)
+	{
+		if (const TViewModelPtr<FLayerBarModel> LayerBarModel = ViewModel.ImplicitCast())
+		{
+			LayerBarModels.Add(LayerBarModel);
+		}
+	}
+
+	// Disallow operations on all layer bars that have a descendant that is selected
+	for (const TViewModelPtr<FLayerBarModel>& InLayerBarModel : LayerBarModels)
+	{
+		const TViewModelPtr<IOutlinerExtension> OutlinerExtension = InLayerBarModel->GetLinkedOutlinerItem();
+		if (!OutlinerExtension.IsValid())
+		{
+			continue;
+		}
+
+		for (const TViewModelPtr<IOutlinerExtension>& ChildOutlinerExtension : OutlinerExtension.AsModel()->GetDescendantsOfType<IOutlinerExtension>())
+		{
+			if (InSequencerSelection->Outliner.IsSelected(ChildOutlinerExtension))
+			{
+				return {};
+			}
+		}
+
+		const FFrameNumber ThisTime = InLayerBarModel->ComputeRange().GetLowerBoundValue();
+		if (!OutEarliestTime.IsSet() || ThisTime < OutEarliestTime.GetValue())
+		{
+			OutEarliestTime = ThisTime;
+		}
+	}
+
+	return MoveTemp(LayerBarModels);
+}
+
+TSet<TViewModelPtr<FSectionModel>> GatherAllSelectedSections(const TSharedRef<FSequencerSelection>& InSequencerSelection, TOptional<FFrameNumber>& OutEarliestTime)
+{
+	TSet<TViewModelPtr<FSectionModel>> SectionModels;
+
+	SectionModels.Reserve(InSequencerSelection->TrackArea.Num());
+
+	for (const FViewModelPtr& ViewModel : InSequencerSelection->TrackArea)
+	{
+		if (const TViewModelPtr<FSectionModel> SectionModel = ViewModel.ImplicitCast())
+		{
+			SectionModels.Add(SectionModel);
+
+			const FFrameNumber ThisTime = SectionModel->GetLayerBarRange().GetLowerBoundValue();
+			if (!OutEarliestTime.IsSet() || ThisTime < OutEarliestTime.GetValue())
+			{
+				OutEarliestTime = ThisTime;
+			}
+		}
+	}
+
+	return MoveTemp(SectionModels);
+}
+
+TSet<FKeyHandle> GatherAllSelectedKeyframes(const TSharedRef<FSequencerSelection>& InSequencerSelection, TOptional<FFrameNumber>& OutEarliestTime)
+{
+	TSet<FKeyHandle> OutKeyframes;
+
+	for (const FKeyHandle KeyHandle : InSequencerSelection->KeySelection)
+	{
+		if (const TSharedPtr<FChannelModel> ChannelModel = InSequencerSelection->KeySelection.GetModelForKey(KeyHandle))
+		{
+			if (const TSharedPtr<IKeyArea> KeyArea = ChannelModel->GetKeyArea())
+			{
+				UMovieSceneSection* const Section = ChannelModel->GetSection();
+				if (IsValid(Section) && !Section->IsReadOnly())
+				{
+					OutKeyframes.Add(KeyHandle);
+
+					const FFrameNumber ThisTime = KeyArea->GetKeyTime(KeyHandle);
+					if (!OutEarliestTime.IsSet() || ThisTime < OutEarliestTime.GetValue())
+					{
+						OutEarliestTime = ThisTime;
+					}
+				}
+			}
+		}
+	}
+
+	return MoveTemp(OutKeyframes);
+}
+
+} // namespace UE::Sequencer
+
+using namespace UE::Sequencer;
+
+bool FSequencerSelectionAlignmentUtils::CanAlignSelection(const ISequencer& InSequencer)
+{
+	const TSharedPtr<FSequencerEditorViewModel> ViewModel = InSequencer.GetViewModel();
+	if (!ViewModel.IsValid())
+	{
+		return false;
+	}
+
+	const TSharedPtr<FSequencerSelection> SequencerSelection = ViewModel->GetSelection();
+	if (!SequencerSelection.IsValid())
+	{
+		return false;
+	}
+
+	const TSharedRef<FSequencerSelection> SequencerSelectionRef = SequencerSelection.ToSharedRef();
+
+	TOptional<FFrameNumber> EarliestLayerBarTime;
+	const TSet<TViewModelPtr<FLayerBarModel>> LayerBarModels = GatherAllSelectedLayerBars(SequencerSelectionRef, EarliestLayerBarTime);
+
+	TOptional<FFrameNumber> EarliestSectionTime;
+	const TSet<TViewModelPtr<FSectionModel>> SectionModels = GatherAllSelectedSections(SequencerSelectionRef, EarliestSectionTime);
+
+	TOptional<FFrameNumber> EarliestKeyframeTime;
+	const TSet<FKeyHandle> KeyframeModels = GatherAllSelectedKeyframes(SequencerSelectionRef, EarliestKeyframeTime);
+
+	const bool bIsLayerBarOrSection = EarliestLayerBarTime.IsSet() || EarliestSectionTime.IsSet();
+	const bool bNoLayerBarOrSection = !EarliestLayerBarTime.IsSet() && !EarliestSectionTime.IsSet();
+	const int32 KeyFrameCount = KeyframeModels.Num();
+
+	// To avoid having to do a bunch of extra processing to support aligning layers bars AND keys at the same time,
+	// we will only allow layer bar OR key selection alignments since layer bars are effected by keys.
+	return (bIsLayerBarOrSection && KeyFrameCount == 0)
+		|| (bNoLayerBarOrSection && KeyFrameCount > 0);
+}
+
+void FSequencerSelectionAlignmentUtils::AlignSelectionToTime(const ISequencer& InSequencer, const FQualifiedFrameTime& InFrameTime, const bool bInTransact)
+{
+	const TSharedPtr<FSequencerEditorViewModel> ViewModel = InSequencer.GetViewModel();
+	if (!ViewModel.IsValid())
+	{
+		return;
+	}
+
+	const TSharedPtr<FSequencerSelection> SequencerSelection = ViewModel->GetSelection();
+	if (!SequencerSelection.IsValid())
+	{
+		return;
+	}
+
+	const TSharedRef<FSequencerSelection> SequencerSelectionRef = SequencerSelection.ToSharedRef();
+
+	TOptional<FFrameNumber> EarliestLayerBarTime;
+	const TSet<TViewModelPtr<FLayerBarModel>> LayerBarModels = GatherAllSelectedLayerBars(SequencerSelectionRef, EarliestLayerBarTime);
+
+	TOptional<FFrameNumber> EarliestSectionTime;
+	const TSet<TViewModelPtr<FSectionModel>> SectionModels = GatherAllSelectedSections(SequencerSelectionRef, EarliestSectionTime);
+
+	const FFrameNumber FrameNumber = InFrameTime.Time.FrameNumber;
+	const bool bIsLayerBarOrSection = EarliestLayerBarTime.IsSet() || EarliestSectionTime.IsSet();
+	const bool bNoLayerBarOrSection = !EarliestLayerBarTime.IsSet() && !EarliestSectionTime.IsSet();
+	const int32 KeyFrameCount = SequencerSelection->KeySelection.Num();
+
+	const FScopedTransaction Transaction(NSLOCTEXT("Sequencer", "AlignToPlayhead", "Align to Playhead"), bInTransact);
+
+	// Layer Bars and Sections
+	if (bIsLayerBarOrSection && KeyFrameCount == 0)
+	{
+		TOptional<FFrameNumber> EarliestTime;
+
+		if (EarliestLayerBarTime.IsSet() && EarliestSectionTime.IsSet())
+		{
+			EarliestTime = (EarliestLayerBarTime.GetValue() < EarliestSectionTime.GetValue())
+				? EarliestLayerBarTime.GetValue() : EarliestSectionTime.GetValue();
+		}
+		else if (EarliestLayerBarTime.IsSet())
+		{
+			EarliestTime = EarliestLayerBarTime.GetValue();
+		}
+		else if (EarliestSectionTime.IsSet())
+		{
+			EarliestTime = EarliestSectionTime.GetValue();
+		}
+
+		if (EarliestTime.IsSet())
+		{
+			const FFrameNumber RelativeDistance = FrameNumber - EarliestTime.GetValue();
+
+			for (const TViewModelPtr<FLayerBarModel>& LayerBarModel : LayerBarModels)
+			{
+				LayerBarModel->Offset(RelativeDistance);
+			}
+
+			for (const TViewModelPtr<FSectionModel>& SectionModel : SectionModels)
+            {
+            	SectionModel->OffsetLayerBar(RelativeDistance);
+            }
+		}
+	}
+	// Keyframes
+	else if (bNoLayerBarOrSection && KeyFrameCount > 0)
+	{
+		TOptional<FFrameNumber> EarliestKeyframeTime;
+		const TSet<FKeyHandle> Keyframes = GatherAllSelectedKeyframes(SequencerSelectionRef, EarliestKeyframeTime);
+		if (EarliestKeyframeTime.IsSet())
+		{
+			const FFrameNumber RelativeDistance = FrameNumber - EarliestKeyframeTime.GetValue();
+
+			for (const FKeyHandle& Key : Keyframes)
+			{
+				if (const TSharedPtr<FChannelModel> ChannelModel = SequencerSelection->KeySelection.GetModelForKey(Key))
+				{
+					if (const TSharedPtr<IKeyArea> KeyArea = ChannelModel->GetKeyArea())
+					{
+						UMovieSceneSection* const Section = ChannelModel->GetSection();
+						if (IsValid(Section) && !Section->IsLocked() && Section->TryModify())
+						{
+							const FFrameNumber NewTime = RelativeDistance + KeyArea->GetKeyTime(Key);
+							KeyArea->SetKeyTime(Key, NewTime);
+							Section->ExpandToFrame(NewTime);
+						}
+					}
+				}
+			}
+		}
+	}
+}
+
+void FSequencerSelectionAlignmentUtils::AlignSelectionToPlayhead(const ISequencer& InSequencer, const bool bInTransact)
+{
+	AlignSelectionToTime(InSequencer, InSequencer.GetGlobalTime(), bInTransact);
+}
diff --git a/Engine/Source/Editor/Sequencer/Public/SequencerCommands.h b/Engine/Source/Editor/Sequencer/Public/SequencerCommands.h
--- a/Engine/Source/Editor/Sequencer/Public/SequencerCommands.h
+++ b/Engine/Source/Editor/Sequencer/Public/SequencerCommands.h
@@ -484,6 +484,9 @@ public:
 	/** Toggles the docked state of the current sidebar drawer. Undocks the docked sidebar drawer if docked or docks the sidebar drawer if there is one open and no currently docked drawer. */
 	TSharedPtr< FUICommandInfo > ToggleSidebarDrawerDock;
 
+	/** Aligns layer bars, sections, or key selections to the current playhead location */
+	TSharedPtr< FUICommandInfo > AlignSelectionToPlayhead;
+
 	/**
 	 * Initialize commands
 	 */
diff --git a/Engine/Source/Editor/Sequencer/Public/Tools/SequencerSelectionAlignmentUtils.h b/Engine/Source/Editor/Sequencer/Public/Tools/SequencerSelectionAlignmentUtils.h
new file mode 100644
--- /dev/null
+++ b/Engine/Source/Editor/Sequencer/Public/Tools/SequencerSelectionAlignmentUtils.h
@@ -0,0 +1,21 @@
+﻿// Copyright Epic Games, Inc. All Rights Reserved.
+
+#pragma once
+
+#include "HAL/Platform.h"
+
+class ISequencer;
+struct FQualifiedFrameTime;
+
+class FSequencerSelectionAlignmentUtils
+{
+public:
+	/** Returns true if there is a valid layer bar or key selection. False if there are both layer bars and keys selected. */
+	SEQUENCER_API static bool CanAlignSelection(const ISequencer& InSequencer);
+
+	/** Aligns the Sequencer selection to a specified time, optionally transacting. */
+	SEQUENCER_API static void AlignSelectionToTime(const ISequencer& InSequencer, const FQualifiedFrameTime& InFrameTime, const bool bInTransact = true);
+
+	/** Aligns the Sequencer selection to the current playhead location, optionally transacting. */
+	SEQUENCER_API static void AlignSelectionToPlayhead(const ISequencer& InSequencer, const bool bInTransact = true);
+};
