Sequencer: Fix regression from 38446563 where the SequencerKeyEdit menu section name was removed and the SequencerKeys menu section name was changed to Keys. The tool menu was also missing the extender.

This was preventing the SequencerKeyEdit menu with the interpolation types to appear (registered in BuiltInChannelEditors)

diff --git a/Engine/Source/Editor/Sequencer/Private/SequencerContextMenus.cpp b/Engine/Source/Editor/Sequencer/Private/SequencerContextMenus.cpp
--- a/Engine/Source/Editor/Sequencer/Private/SequencerContextMenus.cpp
+++ b/Engine/Source/Editor/Sequencer/Private/SequencerContextMenus.cpp
@@ -177,6 +177,7 @@ void FKeyContextMenu::PopulateMenu(FMenuBuilder& MenuBuilder, TSharedPtr<FExtend
 	if (!UToolMenus::Get()->IsMenuRegistered("Sequencer.KeyContextMenu"))
 	{
 		UToolMenu* KeyContextMenu = UToolMenus::Get()->RegisterMenu("Sequencer.KeyContextMenu", NAME_None, EMultiBoxType::Menu);
+		KeyContextMenu->bSearchable = false;
 
 		KeyContextMenu->AddDynamicSection("Edit", FNewToolMenuDelegate::CreateStatic([](UToolMenu* InMenu){
 			USequencerToolMenuContext* ContextObject = InMenu->FindContext<USequencerToolMenuContext>();
@@ -184,7 +185,7 @@ void FKeyContextMenu::PopulateMenu(FMenuBuilder& MenuBuilder, TSharedPtr<FExtend
 
 			if (Sequencer && HotspotCast<FKeyHotspot>(Sequencer->GetViewModel()->GetHotspot()))
 			{
-				FToolMenuSection& Section = InMenu->AddSection(NAME_None, LOCTEXT("EditMenu", "Edit"));
+				FToolMenuSection& Section = InMenu->AddSection("SequencerKeyEdit", LOCTEXT("EditMenu", "Edit"));
 
 				Section.AddMenuEntry(FGenericCommands::Get().Cut);
 				Section.AddMenuEntry(FGenericCommands::Get().Copy);
@@ -192,10 +193,11 @@ void FKeyContextMenu::PopulateMenu(FMenuBuilder& MenuBuilder, TSharedPtr<FExtend
 			}
 		}));
 
-		FToolMenuSection& KeysSection = KeyContextMenu->AddSection("Keys", LOCTEXT("KeysMenu", "Keys"));
+		FToolMenuSection& KeysSection = KeyContextMenu->AddSection("SequencerKeys", LOCTEXT("KeysMenu", "Keys"));
 		KeysSection.AddMenuEntry(FSequencerCommands::Get().SetKeyTime);
 		KeysSection.AddMenuEntry(FSequencerCommands::Get().Rekey);
 		KeysSection.AddMenuEntry(FSequencerCommands::Get().SnapToFrame);
+
 		FToolMenuSection& KeysRemovalSection = KeyContextMenu->AddSection("KeyRemoval");
 		KeysRemovalSection.AddSeparator(NAME_None);
 		KeysRemovalSection.AddMenuEntry(FSequencerCommands::Get().DeleteKeys);
@@ -206,6 +208,7 @@ void FKeyContextMenu::PopulateMenu(FMenuBuilder& MenuBuilder, TSharedPtr<FExtend
 
 	FToolMenuContext MenuContext(ContextObject);
 	MenuContext.AppendCommandList(Sequencer->GetCommandBindings(ESequencerCommandBindings::Sequencer));
+	MenuContext.AddExtender(MenuExtender);
 	MenuBuilder.AddWidget(UToolMenus::Get()->GenerateWidget("Sequencer.KeyContextMenu", MenuContext), FText(), true, false);
 }
 
