Add a scoped FViewModelHierarchyOperation before calling ResetImplicitScopedModifyDefer() to prevent possible N^2 operation with a large number of tracks.

diff --git a/Engine/Source/Editor/Sequencer/Private/Sequencer.cpp b/Engine/Source/Editor/Sequencer/Private/Sequencer.cpp
--- a/Engine/Source/Editor/Sequencer/Private/Sequencer.cpp
+++ b/Engine/Source/Editor/Sequencer/Private/Sequencer.cpp
@@ -971,8 +971,11 @@ void FSequencer::Tick(float InDeltaTime)
 		SharedModelData->ReportLatentHierarchicalOperations();
 	}
 
-	UMovieSceneSignedObject::ResetImplicitScopedModifyDefer();
-
+	{
+		FViewModelHierarchyOperation Operation(GetViewModel()->GetRootSequenceModel()->GetSharedData());
+		UMovieSceneSignedObject::ResetImplicitScopedModifyDefer();
+	}
+	
 	if (bNeedsInvalidateCachedData)
 	{
 		InvalidateCachedData();
