Sequencer: Follow up to IsFilterSupported

With the current logic, even if a deried sequence indicated that it supports a certain type of filter, it wouldn't show up because the filter would return false for SupportsSequence.

Tested that a WidgetAnimation could return true for IsFilterSupported for Camera Cuts and the filter appears in the UMG Animation editor.

#rnx

diff --git a/Engine/Source/Editor/Sequencer/Private/Filters/Menus/SequencerTrackFilterMenu.cpp b/Engine/Source/Editor/Sequencer/Private/Filters/Menus/SequencerTrackFilterMenu.cpp
--- a/Engine/Source/Editor/Sequencer/Private/Filters/Menus/SequencerTrackFilterMenu.cpp
+++ b/Engine/Source/Editor/Sequencer/Private/Filters/Menus/SequencerTrackFilterMenu.cpp
@@ -1,4 +1,4 @@
-﻿// Copyright Epic Games, Inc. All Rights Reserved.
+// Copyright Epic Games, Inc. All Rights Reserved.
 
 #include "SequencerTrackFilterMenu.h"
 #include "Filters/SequencerFilterBar.h"
@@ -366,9 +366,12 @@ void FSequencerTrackFilterMenu::FillFiltersMenuCategory(FToolMenuSection& InOutS
 
 	for (const TSharedRef<FSequencerTrackFilter>& Filter : FilterBar->GetCommonFilters())
 	{
-		if (!FocusedSequence->IsFilterSupported(Filter->GetName())
-			|| Filter->GetCategory() != InMenuCategory
-			|| !Filter->SupportsSequence(FocusedSequence))
+		if (Filter->GetCategory() != InMenuCategory)
+		{
+			continue;
+		}
+
+		if (!FocusedSequence->IsFilterSupported(Filter->GetName()) && !Filter->SupportsSequence(FocusedSequence))
 		{
 			continue;
 		}
diff --git a/Engine/Source/Runtime/MovieScene/Public/MovieSceneSequence.h b/Engine/Source/Runtime/MovieScene/Public/MovieSceneSequence.h
--- a/Engine/Source/Runtime/MovieScene/Public/MovieSceneSequence.h
+++ b/Engine/Source/Runtime/MovieScene/Public/MovieSceneSequence.h
@@ -447,7 +447,7 @@ public:
 	MOVIESCENE_API bool IsFilterSupported(const FString& InFilterName) const;
 
 	/** Subclasses can implement this to determine whether they support a particular filter type */
-	virtual bool IsFilterSupportedImpl(const FString& InFilterName) const { return true; }
+	virtual bool IsFilterSupportedImpl(const FString& InFilterName) const { return false; }
 
 #endif
 };
