Sequencer: Null ptr check in UpdateActorBindingsData when existing PIE with an actor sequence

diff --git a/Engine/Source/Editor/Sequencer/Private/LevelEditorSequencerIntegration.cpp b/Engine/Source/Editor/Sequencer/Private/LevelEditorSequencerIntegration.cpp
--- a/Engine/Source/Editor/Sequencer/Private/LevelEditorSequencerIntegration.cpp
+++ b/Engine/Source/Editor/Sequencer/Private/LevelEditorSequencerIntegration.cpp
@@ -1685,7 +1685,11 @@ void FLevelEditorSequencerBindingData::UpdateActorBindingsData(TWeakPtr<FSequenc
 		const FMovieSceneSequenceHierarchy*        Hierarchy    = RootInstance.GetCompiledDataManager()->FindHierarchy(RootInstance.GetCompiledDataID());
 
 		UMovieSceneSequence* RootSequence = Pinned->GetRootMovieSceneSequence();
-		AddActorsToBindingsMapRecursive(*Pinned, RootSequence, MovieSceneSequenceID::Root, Hierarchy, ActorBindingsMap);
+		
+		if (RootSequence)
+		{
+			AddActorsToBindingsMapRecursive(*Pinned, RootSequence, MovieSceneSequenceID::Root, Hierarchy, ActorBindingsMap);
+		}
 
 		bActorBindingsDirty = false;
 
