Sequencer: Fixed crash when adding time warp to a sequence

diff --git a/Engine/Source/Editor/MovieSceneTools/Private/TrackEditors/TimeWarpTrackEditor.cpp b/Engine/Source/Editor/MovieSceneTools/Private/TrackEditors/TimeWarpTrackEditor.cpp
--- a/Engine/Source/Editor/MovieSceneTools/Private/TrackEditors/TimeWarpTrackEditor.cpp
+++ b/Engine/Source/Editor/MovieSceneTools/Private/TrackEditors/TimeWarpTrackEditor.cpp
@@ -351,7 +351,7 @@ FText FTimeWarpTrackEditor::GetDisplayName() const
 	return LOCTEXT("TimeWarpTrackEditor_DisplayName", "Time Warp");
 }
 
-void FTimeWarpTrackEditor::BuildAddTrackMenu(FMenuBuilder& MenuBuilder)
+void FTimeWarpTrackEditor::BuildPinnedAddTrackMenu(FMenuBuilder& MenuBuilder)
 {
 	UMovieScene* FocusedMovieScene = GetFocusedMovieScene();
 	if (FocusedMovieScene == nullptr)
diff --git a/Engine/Source/Editor/MovieSceneTools/Public/TrackEditors/TimeWarpTrackEditor.h b/Engine/Source/Editor/MovieSceneTools/Public/TrackEditors/TimeWarpTrackEditor.h
--- a/Engine/Source/Editor/MovieSceneTools/Public/TrackEditors/TimeWarpTrackEditor.h
+++ b/Engine/Source/Editor/MovieSceneTools/Public/TrackEditors/TimeWarpTrackEditor.h
@@ -76,7 +76,7 @@ public:
 
 private:
 	virtual FText GetDisplayName() const override;
-	virtual void BuildAddTrackMenu(FMenuBuilder& MenuBuilder) override;
+	virtual void BuildPinnedAddTrackMenu(FMenuBuilder& MenuBuilder) override;
 	virtual void ProcessKeyOperation(FFrameNumber InKeyTime, const UE::Sequencer::FKeyOperation& Operation, ISequencer& InSequencer, TArray<UE::Sequencer::FAddKeyResult>* OutResults = nullptr) override;
 	virtual TSharedRef<ISequencerSection> MakeSectionInterface(UMovieSceneSection& SectionObject, UMovieSceneTrack& Track, FGuid ObjectBinding) override;
 	virtual bool SupportsSequence(UMovieSceneSequence* InSequence) const override;
diff --git a/Engine/Source/Editor/Sequencer/Private/MovieSceneTrackEditor.cpp b/Engine/Source/Editor/Sequencer/Private/MovieSceneTrackEditor.cpp
--- a/Engine/Source/Editor/Sequencer/Private/MovieSceneTrackEditor.cpp
+++ b/Engine/Source/Editor/Sequencer/Private/MovieSceneTrackEditor.cpp
@@ -218,6 +218,10 @@ void FMovieSceneTrackEditor::BindCommands(TSharedRef<FUICommandList> SequencerCo
 { 
 }
 
+void FMovieSceneTrackEditor::BuildPinnedAddTrackMenu(FMenuBuilder& MenuBuilder)
+{
+}
+
 void FMovieSceneTrackEditor::BuildAddTrackMenu(FMenuBuilder& MenuBuilder) 
 { 
 }
diff --git a/Engine/Source/Editor/Sequencer/Private/Sequencer.cpp b/Engine/Source/Editor/Sequencer/Private/Sequencer.cpp
--- a/Engine/Source/Editor/Sequencer/Private/Sequencer.cpp
+++ b/Engine/Source/Editor/Sequencer/Private/Sequencer.cpp
@@ -9326,22 +9326,6 @@ void FSequencer::AddFolder()
 	NotifyMovieSceneDataChanged( EMovieSceneDataChangeType::MovieSceneStructureItemAdded );
 }
 
-void FSequencer::AddTimeWarp()
-{
-	const TSharedPtr<ISequencerTrackEditor>* const TimeWarpEditorPtr = TrackEditors.FindByPredicate(
-	[](const TSharedPtr<ISequencerTrackEditor>& InTrackEditor)
-	{
-		return InTrackEditor->GetDisplayName().ToString().Equals(TEXT("Time Warp"));
-	});
-	if (TimeWarpEditorPtr)
-	{
-		if (const TSharedPtr<FTimeWarpTrackEditor> TimeWarpEditor = StaticCastSharedPtr<FTimeWarpTrackEditor>(*TimeWarpEditorPtr))
-		{
-			TimeWarpEditor->HandleAddTimeWarpTrack(UMovieSceneTimeWarpGetter::StaticClass());
-		}
-	}
-}
-
 void FixSortingOrders(UMovieSceneTrack* InTrack, const TArray<UMovieSceneTrack*>& Tracks)
 {
 	// Fix the sorting orders if InTrack is going to conflict with any of the given Tracks (ie. same class and same name)
@@ -12039,34 +12023,33 @@ void FSequencer::BuildAddTrackMenu(class FMenuBuilder& MenuBuilder)
 		FSlateIcon(FAppStyle::GetAppStyleSetName(), TEXT("ContentBrowser.AssetTreeFolderOpen")),
 		FUIAction(FExecuteAction::CreateRaw(this, &FSequencer::AddFolder)));
 
-	MenuBuilder.AddMenuEntry(
-		LOCTEXT("TimeWarp", "Add Time Warp Track"),
-		LOCTEXT("TimeWarpToolTip", "Adds a new time warp track."),
-		FSlateIcon(FAppStyle::GetAppStyleSetName(), TEXT("Sequencer.Tracks.TimeWarp")),
-		FUIAction(FExecuteAction::CreateRaw(this, &FSequencer::AddTimeWarp)));
-
-	MenuBuilder.AddMenuSeparator(NAME_None);
-
-	TArray<TSharedPtr<ISequencerTrackEditor>> SortedTrackEditors = TrackEditors;
+	TArray<TSharedPtr<ISequencerTrackEditor>> SortedAndFilteredTrackEditors = TrackEditors;
 
-	// UX design calls for some tracks to be forced to the top of the list
-	const TArray<FString> TrackEditorsToIgnore = { TEXT("Time Warp") };
-	SortedTrackEditors.RemoveAll([&TrackEditorsToIgnore](const TSharedPtr<ISequencerTrackEditor>& InTrackEditor)
+	// Filter unsupported
+	for (int32 Index = SortedAndFilteredTrackEditors.Num()-1; Index >= 0; --Index)
+	{
+		if (!SortedAndFilteredTrackEditors[Index]->SupportsSequence(GetFocusedMovieSceneSequence()))
 		{
-			return TrackEditorsToIgnore.Contains(InTrackEditor->GetDisplayName().ToString());
-		});
+			SortedAndFilteredTrackEditors.RemoveAtSwap(Index, 1, EAllowShrinking::No);
+		}
+	}
 
-	SortedTrackEditors.Sort([](const TSharedPtr<ISequencerTrackEditor>& InA, const TSharedPtr<ISequencerTrackEditor>& InB)
+	// Sort by name
+	SortedAndFilteredTrackEditors.Sort([](const TSharedPtr<ISequencerTrackEditor>& InA, const TSharedPtr<ISequencerTrackEditor>& InB)
 		{
 			return InA->GetDisplayName().CompareTo(InB->GetDisplayName()) < 0;
 		});
 
-	for (const TSharedPtr<ISequencerTrackEditor>& TrackEditor : SortedTrackEditors)
+	for (TSharedPtr<ISequencerTrackEditor> TrackEditor : SortedAndFilteredTrackEditors)
 	{
-		if (TrackEditor->SupportsSequence(GetFocusedMovieSceneSequence()))
-		{
-			TrackEditor->BuildAddTrackMenu(MenuBuilder);
-		}
+		TrackEditor->BuildPinnedAddTrackMenu(MenuBuilder);
+	}
+
+	MenuBuilder.AddMenuSeparator(NAME_None);
+
+	for (const TSharedPtr<ISequencerTrackEditor>& TrackEditor : SortedAndFilteredTrackEditors)
+	{
+		TrackEditor->BuildAddTrackMenu(MenuBuilder);
 	}
 }
 
diff --git a/Engine/Source/Editor/Sequencer/Private/Sequencer.h b/Engine/Source/Editor/Sequencer/Private/Sequencer.h
--- a/Engine/Source/Editor/Sequencer/Private/Sequencer.h
+++ b/Engine/Source/Editor/Sequencer/Private/Sequencer.h
@@ -1100,9 +1100,6 @@ public:
 	/** Handles adding a new folder to the outliner tree. */
 	void AddFolder();
 
-	/** Adds a time warp track to the outliner tree. */
-	void AddTimeWarp();
-
 private:
 
 	/** Updates viewport clients' actor locks if they relate to sequencer cameras */
diff --git a/Engine/Source/Editor/Sequencer/Public/ISequencerTrackEditor.h b/Engine/Source/Editor/Sequencer/Public/ISequencerTrackEditor.h
--- a/Engine/Source/Editor/Sequencer/Public/ISequencerTrackEditor.h
+++ b/Engine/Source/Editor/Sequencer/Public/ISequencerTrackEditor.h
@@ -128,6 +128,13 @@ public:
 	*/
 	virtual void BindCommands(TSharedRef<FUICommandList> SequencerCommandBindings) = 0;
 
+	/**
+	 * Builds up the sequencer's "Add Track" menu. for menu items pinned to the top
+	 *
+	 * @param MenuBuilder The menu builder to change.
+	 */
+	virtual void BuildPinnedAddTrackMenu(FMenuBuilder& MenuBuilder) = 0;
+
 	/**
 	 * Builds up the sequencer's "Add Track" menu.
 	 *
diff --git a/Engine/Source/Editor/Sequencer/Public/MovieSceneTrackEditor.h b/Engine/Source/Editor/Sequencer/Public/MovieSceneTrackEditor.h
--- a/Engine/Source/Editor/Sequencer/Public/MovieSceneTrackEditor.h
+++ b/Engine/Source/Editor/Sequencer/Public/MovieSceneTrackEditor.h
@@ -178,6 +178,7 @@ public:
 	virtual TSharedRef<ISequencerSection> MakeSectionInterface(UMovieSceneSection& SectionObject, UMovieSceneTrack& Track, FGuid ObjectBinding) override;
 
 	virtual void BindCommands(TSharedRef<FUICommandList> SequencerCommandBindings) override;
+	virtual void BuildPinnedAddTrackMenu(FMenuBuilder& MenuBuilder) override;
 	virtual void BuildAddTrackMenu(FMenuBuilder& MenuBuilder) override;
 	virtual void BuildObjectBindingColumnWidgets(TFunctionRef<TSharedRef<SHorizontalBox>()> GetEditBox, const UE::Sequencer::TViewModelPtr<UE::Sequencer::FObjectBindingModel>& ObjectBinding, const UE::Sequencer::FCreateOutlinerViewParams& InParams, const FName& InColumnName) override;
 	virtual void BuildObjectBindingTrackMenu(FMenuBuilder& MenuBuilder, const TArray<FGuid>& ObjectBindings, const UClass* ObjectClass) override;
