[Sequencer]
- Fix filters checking IsFilterSupported() to use filter name instead of display name
- Fix issues with Sequence->IsFilterSupported() and Filter->SupportsSequence()
- Only add filters that are supported on init instead of checking supported constantly

diff --git a/Engine/Plugins/Animation/ControlRig/Source/ControlRigEditor/Private/Sequencer/ControlRigSequencerFilter.cpp b/Engine/Plugins/Animation/ControlRig/Source/ControlRigEditor/Private/Sequencer/ControlRigSequencerFilter.cpp
--- a/Engine/Plugins/Animation/ControlRig/Source/ControlRigEditor/Private/Sequencer/ControlRigSequencerFilter.cpp
+++ b/Engine/Plugins/Animation/ControlRig/Source/ControlRigEditor/Private/Sequencer/ControlRigSequencerFilter.cpp
@@ -97,6 +97,11 @@ public:
 		return FSequencerTrackFilter_ControlRigControlsCommands::Get().ToggleFilter_ControlRigControls;
 	}
 
+	virtual bool SupportsSequence(UMovieSceneSequence* const InSequence) const override
+	{
+		return IsSequenceTrackSupported<UMovieSceneControlRigParameterTrack>(InSequence);
+	}
+
 	//~ End FSequencerTrackFilter
 
 private:
@@ -225,6 +230,11 @@ public:
 		return FSequencerTrackFilter_ControlRigSelectedControlsCommands::Get().ToggleFilter_ControlRigSelectedControls;
 	}
 
+	virtual bool SupportsSequence(UMovieSceneSequence* const InSequence) const override
+	{
+		return IsSequenceTrackSupported<UMovieSceneControlRigParameterTrack>(InSequence);
+	}
+
 	//~ End FSequencerTrackFilter
 
 	static URigHierarchy* GetControlRigHierarchyFromTrackObject(UMovieSceneTrack* const InTrackObject)
diff --git a/Engine/Plugins/Animation/LiveLink/Source/LiveLinkEditor/Private/Sequencer/LiveLinkSequencerFilters.cpp b/Engine/Plugins/Animation/LiveLink/Source/LiveLinkEditor/Private/Sequencer/LiveLinkSequencerFilters.cpp
--- a/Engine/Plugins/Animation/LiveLink/Source/LiveLinkEditor/Private/Sequencer/LiveLinkSequencerFilters.cpp
+++ b/Engine/Plugins/Animation/LiveLink/Source/LiveLinkEditor/Private/Sequencer/LiveLinkSequencerFilters.cpp
@@ -77,7 +77,7 @@ public:
 
 	virtual bool SupportsSequence(UMovieSceneSequence* const InSequence) const override
 	{
-		return SupportsLevelSequence(InSequence) || IsSequenceTrackSupported<UMovieSceneLiveLinkTrack>(InSequence);
+		return IsSequenceTrackSupported<UMovieSceneLiveLinkTrack>(InSequence);
 	}
 
 	//~ End FSequencerTrackFilter
diff --git a/Engine/Plugins/Experimental/DaySequence/Source/DaySequence/Private/DaySequence.cpp b/Engine/Plugins/Experimental/DaySequence/Source/DaySequence/Private/DaySequence.cpp
--- a/Engine/Plugins/Experimental/DaySequence/Source/DaySequence/Private/DaySequence.cpp
+++ b/Engine/Plugins/Experimental/DaySequence/Source/DaySequence/Private/DaySequence.cpp
@@ -206,6 +206,15 @@ ETrackSupport UDaySequence::IsTrackSupportedImpl(TSubclassOf<class UMovieSceneTr
 	return Super::IsTrackSupportedImpl(InTrackClass);
 }
 
+bool UDaySequence::IsFilterSupportedImpl(const FString& InFilterName) const
+{
+	static const TArray<FString> SupportedFilters = {
+		TEXT("Audio"),
+		TEXT("Folder")
+	};
+	return SupportedFilters.Contains(InFilterName);
+}
+
 void UDaySequence::GetAssetRegistryTags(FAssetRegistryTagsContext Context) const
 {
 #if WITH_EDITORONLY_DATA
diff --git a/Engine/Plugins/Experimental/DaySequence/Source/DaySequence/Public/DaySequence.h b/Engine/Plugins/Experimental/DaySequence/Source/DaySequence/Public/DaySequence.h
--- a/Engine/Plugins/Experimental/DaySequence/Source/DaySequence/Public/DaySequence.h
+++ b/Engine/Plugins/Experimental/DaySequence/Source/DaySequence/Public/DaySequence.h
@@ -63,6 +63,7 @@ public:
 
 #if WITH_EDITOR
 	virtual ETrackSupport IsTrackSupportedImpl(TSubclassOf<class UMovieSceneTrack> InTrackClass) const override;
+	virtual bool IsFilterSupportedImpl(const FString& InFilterName) const override;
 	virtual void GetAssetRegistryTags(FAssetRegistryTagsContext Context) const override;
 
 	DECLARE_DELEGATE_RetVal_OneParam(void, FPostDuplicateEvent, UDaySequence*);
diff --git a/Engine/Plugins/FX/Niagara/Source/NiagaraEditor/Private/Sequencer/NiagaraSequence/NiagaraSequencerFilters.cpp b/Engine/Plugins/FX/Niagara/Source/NiagaraEditor/Private/Sequencer/NiagaraSequence/NiagaraSequencerFilters.cpp
--- a/Engine/Plugins/FX/Niagara/Source/NiagaraEditor/Private/Sequencer/NiagaraSequence/NiagaraSequencerFilters.cpp
+++ b/Engine/Plugins/FX/Niagara/Source/NiagaraEditor/Private/Sequencer/NiagaraSequence/NiagaraSequencerFilters.cpp
@@ -102,6 +102,11 @@ public:
 		return FSequencerTrackFilter_NiagaraFilterCommands::Get().ToggleFilter_Niagara;
 	}
 
+	virtual bool SupportsSequence(UMovieSceneSequence* const InSequence) const override
+	{
+		return IsSequenceTrackSupported<UMovieSceneNiagaraTrack>(InSequence);
+	}
+
 	//~ End FSequencerTrackFilter
 
 private:
diff --git a/Engine/Plugins/Media/MediaCompositing/Source/MediaCompositingEditor/Private/Sequencer/MediaSequencerFilters.cpp b/Engine/Plugins/Media/MediaCompositing/Source/MediaCompositingEditor/Private/Sequencer/MediaSequencerFilters.cpp
--- a/Engine/Plugins/Media/MediaCompositing/Source/MediaCompositingEditor/Private/Sequencer/MediaSequencerFilters.cpp
+++ b/Engine/Plugins/Media/MediaCompositing/Source/MediaCompositingEditor/Private/Sequencer/MediaSequencerFilters.cpp
@@ -74,6 +74,11 @@ public:
 		return FSequencerTrackFilter_MediaFilterCommands::Get().ToggleFilter_Media;
 	}
 
+	virtual bool SupportsSequence(UMovieSceneSequence* const InSequence) const override
+	{
+		return IsSequenceTrackSupported<UMovieSceneMediaTrack>(InSequence);
+	}
+
 	//~ End FSequencerTrackFilter
 
 private:
diff --git a/Engine/Source/Editor/Sequencer/Private/Filters/Filters/SequencerTrackFilterBase.cpp b/Engine/Source/Editor/Sequencer/Private/Filters/Filters/SequencerTrackFilterBase.cpp
--- a/Engine/Source/Editor/Sequencer/Private/Filters/Filters/SequencerTrackFilterBase.cpp
+++ b/Engine/Source/Editor/Sequencer/Private/Filters/Filters/SequencerTrackFilterBase.cpp
@@ -3,7 +3,6 @@
 #include "Filters/SequencerTrackFilterBase.h"
 #include "Filters/SequencerFilterBar.h"
 #include "ISequencer.h"
-#include "LevelSequence.h"
 #include "MovieSceneSequence.h"
 #include "MVVM/ViewModelPtr.h"
 
@@ -16,7 +15,7 @@ FSequencerTrackFilter::FSequencerTrackFilter(ISequencerTrackFilters& InOutFilter
 
 bool FSequencerTrackFilter::SupportsSequence(UMovieSceneSequence* const InSequence) const
 {
-	return SupportsLevelSequence(InSequence);
+	return false;
 }
 
 ISequencerTrackFilters& FSequencerTrackFilter::GetFilterInterface() const
@@ -34,19 +33,3 @@ UMovieScene* FSequencerTrackFilter::GetFocusedGetMovieScene() const
 	const UMovieSceneSequence* const FocusedMovieSceneSequence = GetFocusedMovieSceneSequence();
 	return IsValid(FocusedMovieSceneSequence) ? FocusedMovieSceneSequence->GetMovieScene() : nullptr;
 }
-
-bool FSequencerTrackFilter::SupportsLevelSequence(UMovieSceneSequence* const InSequence)
-{
-	const UClass* const LevelSequenceClass = ULevelSequence::StaticClass();
-	return IsValid(InSequence)
-		&& IsValid(LevelSequenceClass)
-		&& InSequence->GetClass()->IsChildOf(LevelSequenceClass);
-}
-
-bool FSequencerTrackFilter::SupportsUMGSequence(UMovieSceneSequence* const InSequence)
-{
-	static UClass* const WidgetAnimationClass = FindObject<UClass>(nullptr, TEXT("/Script/UMG.WidgetAnimation"), true);
-	return IsValid(InSequence)
-		&& IsValid(WidgetAnimationClass)
-		&& InSequence->GetClass()->IsChildOf(WidgetAnimationClass);
-}
diff --git a/Engine/Source/Editor/Sequencer/Private/Filters/Filters/SequencerTrackFilter_Keyed.cpp b/Engine/Source/Editor/Sequencer/Private/Filters/Filters/SequencerTrackFilter_Keyed.cpp
--- a/Engine/Source/Editor/Sequencer/Private/Filters/Filters/SequencerTrackFilter_Keyed.cpp
+++ b/Engine/Source/Editor/Sequencer/Private/Filters/Filters/SequencerTrackFilter_Keyed.cpp
@@ -32,11 +32,6 @@ TSharedPtr<FUICommandInfo> FSequencerTrackFilter_Keyed::GetToggleCommand() const
 	return FSequencerTrackFilterCommands::Get().ToggleFilter_Keyed;
 }
 
-bool FSequencerTrackFilter_Keyed::SupportsSequence(UMovieSceneSequence* const InSequence) const
-{
-	return FSequencerTrackFilter::SupportsSequence(InSequence) || SupportsUMGSequence(InSequence);
-}
-
 bool FSequencerTrackFilter_Keyed::PassesFilter(FSequencerTrackFilterType InItem) const
 {
 	if (const TViewModelPtr<FCategoryGroupModel> CategoryGroupModel = InItem.ImplicitCast())
diff --git a/Engine/Source/Editor/Sequencer/Private/Filters/Filters/SequencerTrackFilters.cpp b/Engine/Source/Editor/Sequencer/Private/Filters/Filters/SequencerTrackFilters.cpp
--- a/Engine/Source/Editor/Sequencer/Private/Filters/Filters/SequencerTrackFilters.cpp
+++ b/Engine/Source/Editor/Sequencer/Private/Filters/Filters/SequencerTrackFilters.cpp
@@ -4,7 +4,6 @@
 #include "Filters/SequencerTrackFilterBase.h"
 #include "Filters/SequencerTrackFilterCommands.h"
 #include "Misc/IFilter.h"
-#include "MovieSceneFolder.h"
 #include "Particles/ParticleSystem.h"
 #include "Styling/AppStyle.h"
 #include "Styling/SlateIconFinder.h"
@@ -40,9 +39,7 @@ TSharedPtr<FUICommandInfo> FSequencerTrackFilter_Audio::GetToggleCommand() const
 
 bool FSequencerTrackFilter_Audio::SupportsSequence(UMovieSceneSequence* const InSequence) const
 {
-	return IsSequenceTrackSupported<UMovieSceneAudioTrack>(InSequence)
-		|| SupportsLevelSequence(InSequence)
-		|| SupportsUMGSequence(InSequence);
+	return IsSequenceTrackSupported<UMovieSceneAudioTrack>(InSequence);
 }
 
 //////////////////////////////////////////////////////////////////////////
@@ -75,8 +72,7 @@ TSharedPtr<FUICommandInfo> FSequencerTrackFilter_Event::GetToggleCommand() const
 
 bool FSequencerTrackFilter_Event::SupportsSequence(UMovieSceneSequence* const InSequence) const
 {
-	return FSequencerTrackFilter::SupportsSequence(InSequence)
-		|| SupportsUMGSequence(InSequence);
+	return IsSequenceTrackSupported<UMovieSceneEventTrack>(InSequence);
 }
 
 //////////////////////////////////////////////////////////////////////////
@@ -107,6 +103,11 @@ TSharedPtr<FUICommandInfo> FSequencerTrackFilter_LevelVisibility::GetToggleComma
 	return FSequencerTrackFilterCommands::Get().ToggleFilter_LevelVisibility;
 }
 
+bool FSequencerTrackFilter_LevelVisibility::SupportsSequence(UMovieSceneSequence* const InSequence) const
+{
+	return IsSequenceTrackSupported<UMovieSceneLevelVisibilityTrack>(InSequence);
+}
+
 //////////////////////////////////////////////////////////////////////////
 //
 
@@ -135,6 +136,11 @@ TSharedPtr<FUICommandInfo> FSequencerTrackFilter_Particle::GetToggleCommand() co
 	return FSequencerTrackFilterCommands::Get().ToggleFilter_Particle;
 }
 
+bool FSequencerTrackFilter_Particle::SupportsSequence(UMovieSceneSequence* const InSequence) const
+{
+	return IsSequenceTrackSupported<UMovieSceneParticleTrack>(InSequence);
+}
+
 //////////////////////////////////////////////////////////////////////////
 //
 
@@ -163,6 +169,11 @@ TSharedPtr<FUICommandInfo> FSequencerTrackFilter_CinematicShot::GetToggleCommand
 	return FSequencerTrackFilterCommands::Get().ToggleFilter_CinematicShot;
 }
 
+bool FSequencerTrackFilter_CinematicShot::SupportsSequence(UMovieSceneSequence* const InSequence) const
+{
+	return IsSequenceTrackSupported<UMovieSceneCinematicShotTrack>(InSequence);
+}
+
 //////////////////////////////////////////////////////////////////////////
 //
 
@@ -200,6 +211,11 @@ TSharedPtr<FUICommandInfo> FSequencerTrackFilter_Subsequence::GetToggleCommand()
 	return FSequencerTrackFilterCommands::Get().ToggleFilter_Subsequence;
 }
 
+bool FSequencerTrackFilter_Subsequence::SupportsSequence(UMovieSceneSequence* const InSequence) const
+{
+	return IsSequenceTrackSupported<UMovieSceneSubTrack>(InSequence);
+}
+
 //////////////////////////////////////////////////////////////////////////
 //
 
@@ -228,6 +244,11 @@ TSharedPtr<FUICommandInfo> FSequencerTrackFilter_SkeletalMesh::GetToggleCommand(
 	return FSequencerTrackFilterCommands::Get().ToggleFilter_SkeletalMesh;
 }
 
+bool FSequencerTrackFilter_SkeletalMesh::SupportsSequence(UMovieSceneSequence* const InSequence) const
+{
+	return IsSequenceTrackSupported<USkeletalMeshComponent>(InSequence);
+}
+
 //////////////////////////////////////////////////////////////////////////
 //
 
@@ -256,6 +277,11 @@ TSharedPtr<FUICommandInfo> FSequencerTrackFilter_Camera::GetToggleCommand() cons
 	return FSequencerTrackFilterCommands::Get().ToggleFilter_Camera;
 }
 
+bool FSequencerTrackFilter_Camera::SupportsSequence(UMovieSceneSequence* const InSequence) const
+{
+	return IsSequenceTrackSupported<UCameraComponent>(InSequence);
+}
+
 //////////////////////////////////////////////////////////////////////////
 //
 
@@ -284,6 +310,11 @@ TSharedPtr<FUICommandInfo> FSequencerTrackFilter_Light::GetToggleCommand() const
 	return FSequencerTrackFilterCommands::Get().ToggleFilter_Light;
 }
 
+bool FSequencerTrackFilter_Light::SupportsSequence(UMovieSceneSequence* const InSequence) const
+{
+	return IsSequenceTrackSupported<UCameraComponent>(InSequence);
+}
+
 //////////////////////////////////////////////////////////////////////////
 //
 
@@ -312,6 +343,11 @@ TSharedPtr<FUICommandInfo> FSequencerTrackFilter_CameraCut::GetToggleCommand() c
 	return FSequencerTrackFilterCommands::Get().ToggleFilter_CameraCut;
 }
 
+bool FSequencerTrackFilter_CameraCut::SupportsSequence(UMovieSceneSequence* const InSequence) const
+{
+	return IsSequenceTrackSupported<UMovieSceneCameraCutTrack>(InSequence);
+}
+
 //////////////////////////////////////////////////////////////////////////
 //
 
@@ -340,6 +376,11 @@ TSharedPtr<FUICommandInfo> FSequencerTrackFilter_Fade::GetToggleCommand() const
 	return FSequencerTrackFilterCommands::Get().ToggleFilter_Fade;
 }
 
+bool FSequencerTrackFilter_Fade::SupportsSequence(UMovieSceneSequence* const InSequence) const
+{
+	return IsSequenceTrackSupported<UMovieSceneCameraCutTrack>(InSequence);
+}
+
 //////////////////////////////////////////////////////////////////////////
 //
 
@@ -368,6 +409,11 @@ TSharedPtr<FUICommandInfo> FSequencerTrackFilter_DataLayer::GetToggleCommand() c
 	return FSequencerTrackFilterCommands::Get().ToggleFilter_DataLayer;
 }
 
+bool FSequencerTrackFilter_DataLayer::SupportsSequence(UMovieSceneSequence* const InSequence) const
+{
+	return IsSequenceTrackSupported<UMovieSceneDataLayerTrack>(InSequence);
+}
+
 //////////////////////////////////////////////////////////////////////////
 //
 
@@ -396,6 +442,11 @@ TSharedPtr<FUICommandInfo> FSequencerTrackFilter_TimeDilation::GetToggleCommand(
 	return FSequencerTrackFilterCommands::Get().ToggleFilter_TimeDilation;
 }
 
+bool FSequencerTrackFilter_TimeDilation::SupportsSequence(UMovieSceneSequence* const InSequence) const
+{
+	return IsSequenceTrackSupported<UMovieSceneDataLayerTrack>(InSequence);
+}
+
 //////////////////////////////////////////////////////////////////////////
 //
 
@@ -424,11 +475,6 @@ TSharedPtr<FUICommandInfo> FSequencerTrackFilter_Folder::GetToggleCommand() cons
 	return FSequencerTrackFilterCommands::Get().ToggleFilter_Folder;
 }
 
-TSubclassOf<UMovieSceneTrack> FSequencerTrackFilter_Folder::GetTrackClass() const
-{
-	return UMovieSceneFolder::StaticClass();
-}
-
 bool FSequencerTrackFilter_Folder::SupportsSequence(UMovieSceneSequence* const InSequence) const
 {
 	return true;
diff --git a/Engine/Source/Editor/Sequencer/Private/Filters/Menus/SequencerTrackFilterMenu.cpp b/Engine/Source/Editor/Sequencer/Private/Filters/Menus/SequencerTrackFilterMenu.cpp
--- a/Engine/Source/Editor/Sequencer/Private/Filters/Menus/SequencerTrackFilterMenu.cpp
+++ b/Engine/Source/Editor/Sequencer/Private/Filters/Menus/SequencerTrackFilterMenu.cpp
@@ -1,4 +1,4 @@
-﻿// Copyright Epic Games, Inc. All Rights Reserved.
+// Copyright Epic Games, Inc. All Rights Reserved.
 
 #include "SequencerTrackFilterMenu.h"
 #include "Filters/SequencerFilterBar.h"
@@ -167,10 +167,7 @@ void FSequencerTrackFilterMenu::PopulateOtherFilterSections(UToolMenu& InMenu)
 
 	FToolMenuSection& Section = InMenu.FindOrAddSection(TEXT("OtherFilters"), LOCTEXT("OtherFiltersHeading", "Other Filters"));
 
-	const UMovieSceneSequence* const FocusedSequence = FilterBar->GetSequencer().GetFocusedMovieSceneSequence();
-
-	if (FocusedSequence->IsA(ULevelSequence::StaticClass())
-		&& FocusedSequence->IsFilterSupported(FSequencerTrackFilter_Level::StaticName()))
+	if (FilterBar->IsFilterSupported(FSequencerTrackFilter_Level::StaticName()))
 	{
 		Section.AddSubMenu(TEXT("LevelFilters"),
 			LOCTEXT("LevelFilters", "Level Filters"),
@@ -187,7 +184,7 @@ void FSequencerTrackFilterMenu::PopulateOtherFilterSections(UToolMenu& InMenu)
 			false);
 	}
 
-	if (FocusedSequence->IsFilterSupported(FSequencerTrackFilter_Group::StaticName()))
+	if (FilterBar->IsFilterSupported(FSequencerTrackFilter_Group::StaticName()))
 	{
 		Section.AddSubMenu(TEXT("GroupFilters"),
 			LOCTEXT("GroupFilters", "Group Filters"),
@@ -266,7 +263,7 @@ void FSequencerTrackFilterMenu::FillGroupFilterMenu(UToolMenu* const InMenu)
 		FSlateIcon(FAppStyle::GetAppStyleSetName(), TEXT("Icons.GroupActors")),
 		FUIAction(FExecuteAction::CreateSP(this, &FSequencerTrackFilterMenu::OpenNodeGroupsManager)));
 
-	const UMovieSceneSequence* const FocusedMovieSequence = Sequencer.GetFocusedMovieSceneSequence();
+	const UMovieSceneSequence* const FocusedMovieSequence = Sequencer.GetRootMovieSceneSequence();
 	if (!IsValid(FocusedMovieSequence))
 	{
 		return;
@@ -362,13 +359,9 @@ void FSequencerTrackFilterMenu::FillFiltersMenuCategory(FToolMenuSection& InOutS
 		return;
 	}
 
-	UMovieSceneSequence* const FocusedSequence = FilterBar->GetSequencer().GetFocusedMovieSceneSequence();
-
 	for (const TSharedRef<FSequencerTrackFilter>& Filter : FilterBar->GetCommonFilters())
 	{
-		if (!FocusedSequence->IsFilterSupported(Filter->GetName())
-			|| Filter->GetCategory() != InMenuCategory
-			|| !Filter->SupportsSequence(FocusedSequence))
+		if (Filter->GetCategory() != InMenuCategory)
 		{
 			continue;
 		}
@@ -652,13 +645,13 @@ ECheckBoxState FSequencerTrackFilterMenu::GetAllGroupsCheckState() const
 		return ECheckBoxState::Unchecked;
 	}
 
-	const UMovieSceneSequence* const FocusedMovieSequence = FilterBar->GetSequencer().GetFocusedMovieSceneSequence();
-	if (!IsValid(FocusedMovieSequence))
+	const UMovieSceneSequence* const RootSequence = FilterBar->GetSequencer().GetRootMovieSceneSequence();
+	if (!IsValid(RootSequence))
 	{
 		return ECheckBoxState::Unchecked;
 	}
 
-	UMovieScene* const FocusedMovieScene = FocusedMovieSequence->GetMovieScene();
+	UMovieScene* const FocusedMovieScene = RootSequence->GetMovieScene();
 	if (!IsValid(FocusedMovieScene))
 	{
 		return ECheckBoxState::Unchecked;
diff --git a/Engine/Source/Editor/Sequencer/Private/Filters/SequencerFilterBar.cpp b/Engine/Source/Editor/Sequencer/Private/Filters/SequencerFilterBar.cpp
--- a/Engine/Source/Editor/Sequencer/Private/Filters/SequencerFilterBar.cpp
+++ b/Engine/Source/Editor/Sequencer/Private/Filters/SequencerFilterBar.cpp
@@ -101,39 +101,48 @@ TSharedPtr<ICustomTextFilter<FSequencerTrackFilterType>> FSequencerFilterBar::Cr
 
 void FSequencerFilterBar::CreateDefaultFilters()
 {
+	auto AddFilterIfSupported = [this]
+		(const TSharedPtr<FSequencerTrackFilterCollection>& InFilterCollection, const TSharedRef<FSequencerTrackFilter>& InFilter)
+	{
+		if (IsFilterSupported(InFilter))
+		{
+			InFilterCollection->Add(InFilter);
+		}
+	};
+
 	// Add internal filters that won't be saved to config
 	InternalFilters->RemoveAll();
 
-	InternalFilters->Add(LevelFilter);
-	InternalFilters->Add(GroupFilter);
+	AddFilterIfSupported(InternalFilters, LevelFilter);
+	AddFilterIfSupported(InternalFilters, GroupFilter);
 
 	// Add class type category filters
 	CommonFilters->RemoveAll();
 
-	CommonFilters->Add(MakeShared<FSequencerTrackFilter_Audio>(*this, ClassTypeCategory));
-	CommonFilters->Add(MakeShared<FSequencerTrackFilter_CameraCut>(*this, ClassTypeCategory));
-	CommonFilters->Add(MakeShared<FSequencerTrackFilter_DataLayer>(*this, ClassTypeCategory));
-	CommonFilters->Add(MakeShared<FSequencerTrackFilter_Event>(*this, ClassTypeCategory));
-	CommonFilters->Add(MakeShared<FSequencerTrackFilter_Fade>(*this, ClassTypeCategory));
-	CommonFilters->Add(MakeShared<FSequencerTrackFilter_Folder>(*this, ClassTypeCategory));
-	CommonFilters->Add(MakeShared<FSequencerTrackFilter_LevelVisibility>(*this, ClassTypeCategory));
-	CommonFilters->Add(MakeShared<FSequencerTrackFilter_Particle>(*this, ClassTypeCategory));
-	CommonFilters->Add(MakeShared<FSequencerTrackFilter_CinematicShot>(*this, ClassTypeCategory));
-	CommonFilters->Add(MakeShared<FSequencerTrackFilter_Subsequence>(*this, ClassTypeCategory));
-	CommonFilters->Add(MakeShared<FSequencerTrackFilter_TimeDilation>(*this, ClassTypeCategory));
+	AddFilterIfSupported(CommonFilters, MakeShared<FSequencerTrackFilter_Audio>(*this, ClassTypeCategory));
+	AddFilterIfSupported(CommonFilters, MakeShared<FSequencerTrackFilter_CameraCut>(*this, ClassTypeCategory));
+	AddFilterIfSupported(CommonFilters, MakeShared<FSequencerTrackFilter_DataLayer>(*this, ClassTypeCategory));
+	AddFilterIfSupported(CommonFilters, MakeShared<FSequencerTrackFilter_Event>(*this, ClassTypeCategory));
+	AddFilterIfSupported(CommonFilters, MakeShared<FSequencerTrackFilter_Fade>(*this, ClassTypeCategory));
+	AddFilterIfSupported(CommonFilters, MakeShared<FSequencerTrackFilter_Folder>(*this, ClassTypeCategory));
+	AddFilterIfSupported(CommonFilters, MakeShared<FSequencerTrackFilter_LevelVisibility>(*this, ClassTypeCategory));
+	AddFilterIfSupported(CommonFilters, MakeShared<FSequencerTrackFilter_Particle>(*this, ClassTypeCategory));
+	AddFilterIfSupported(CommonFilters, MakeShared<FSequencerTrackFilter_CinematicShot>(*this, ClassTypeCategory));
+	AddFilterIfSupported(CommonFilters, MakeShared<FSequencerTrackFilter_Subsequence>(*this, ClassTypeCategory));
+	AddFilterIfSupported(CommonFilters, MakeShared<FSequencerTrackFilter_TimeDilation>(*this, ClassTypeCategory));
 
 	// Add component type category filters
-	CommonFilters->Add(MakeShared<FSequencerTrackFilter_Camera>(*this, ComponentTypeCategory));
-	CommonFilters->Add(MakeShared<FSequencerTrackFilter_Light>(*this, ComponentTypeCategory));
-	CommonFilters->Add(MakeShared<FSequencerTrackFilter_SkeletalMesh>(*this, ComponentTypeCategory));
+	AddFilterIfSupported(CommonFilters, MakeShared<FSequencerTrackFilter_Camera>(*this, ComponentTypeCategory));
+	AddFilterIfSupported(CommonFilters, MakeShared<FSequencerTrackFilter_Light>(*this, ComponentTypeCategory));
+	AddFilterIfSupported(CommonFilters, MakeShared<FSequencerTrackFilter_SkeletalMesh>(*this, ComponentTypeCategory));
 
 	// Add misc category filters
-	CommonFilters->Add(MakeShared<FSequencerTrackFilter_Keyed>(*this, MiscCategory));
-	//CommonFilters->Add(ModifiedFilter); // Disabling until clear direction on what this should do
-	CommonFilters->Add(SelectedFilter);
-	CommonFilters->Add(MakeShared<FSequencerTrackFilter_Unbound>(*this, MiscCategory));
-	CommonFilters->Add(MakeShared<FSequencerTrackFilter_Condition>(*this, MiscCategory));
-	CommonFilters->Add(MakeShared<FSequencerTrackFilter_TimeWarp>(*this, MiscCategory));
+	AddFilterIfSupported(CommonFilters, MakeShared<FSequencerTrackFilter_Keyed>(*this, MiscCategory));
+	//AddFilterIfSupported(CommonFilters, ModifiedFilter); // Disabling until clear direction on what this should do
+	AddFilterIfSupported(CommonFilters, SelectedFilter);
+	AddFilterIfSupported(CommonFilters, MakeShared<FSequencerTrackFilter_Unbound>(*this, MiscCategory));
+	AddFilterIfSupported(CommonFilters, MakeShared<FSequencerTrackFilter_Condition>(*this, MiscCategory));
+	AddFilterIfSupported(CommonFilters, MakeShared<FSequencerTrackFilter_TimeWarp>(*this, MiscCategory));
 
 	// Add global user-defined filters
 	for (TObjectIterator<USequencerTrackFilterExtension> ExtensionIt(RF_NoFlags); ExtensionIt; ++ExtensionIt)
@@ -148,12 +157,19 @@ void FSequencerFilterBar::CreateDefaultFilters()
 
 			for (const TSharedRef<FSequencerTrackFilter>& ExtendedTrackFilter : ExtendedTrackFilters)
 			{
-				CommonFilters->Add(ExtendedTrackFilter);
+				AddFilterIfSupported(CommonFilters, ExtendedTrackFilter);
 			}
 		}
 	}
 
 	CommonFilters->Sort();
+
+	CommonFilters->OnChanged().Broadcast();
+	InternalFilters->OnChanged().Broadcast();
+	TextFilter->OnChanged().Broadcast();
+	LevelFilter->OnChanged().Broadcast();
+	HideIsolateFilter->OnChanged().Broadcast();
+	SelectedFilter->OnChanged().Broadcast();
 }
 
 void FSequencerFilterBar::BindCommands()
@@ -207,10 +223,7 @@ void FSequencerFilterBar::BindCommands()
 	const TArray<TSharedRef<FSequencerTrackFilter>> AllFilters = GetFilterList(true);
 	for (const TSharedRef<FSequencerTrackFilter>& Filter : AllFilters)
 	{
-		if (Filter->SupportsSequence(FocusedSequence))
-		{
-			Filter->BindCommands();
-		}
+		Filter->BindCommands();
 	}
 
 	// Add bindings for curve editor if supported
@@ -237,15 +250,12 @@ void FSequencerFilterBar::BindCommands()
 				// Add the specific track filter toggle commands
 				for (const TSharedRef<FSequencerTrackFilter>& Filter : AllFilters)
 				{
-					if (Filter->SupportsSequence(FocusedSequence))
-					{
-						const TSharedPtr<FUICommandList>& FilterCommandList = Filter->GetFilterInterface().GetCommandList();
-						const TSharedPtr<FUICommandInfo>& FilterCommand = Filter->GetToggleCommand();
+					const TSharedPtr<FUICommandList>& FilterCommandList = Filter->GetFilterInterface().GetCommandList();
+					const TSharedPtr<FUICommandInfo>& FilterCommand = Filter->GetToggleCommand();
 
-						if (FilterCommand.IsValid() && FilterCommandList->IsActionMapped(FilterCommand))
-						{
-							CurveEditorSharedBindings->MapAction(FilterCommand, *FilterCommandList->GetActionForCommand(FilterCommand));
-						}
+					if (FilterCommand.IsValid() && FilterCommandList->IsActionMapped(FilterCommand))
+					{
+						CurveEditorSharedBindings->MapAction(FilterCommand, *FilterCommandList->GetActionForCommand(FilterCommand));
 					}
 				}
 
@@ -456,8 +466,7 @@ TSharedPtr<FSequencerTrackFilter> FSequencerFilterBar::FindFilterByDisplayName(c
 				return false;
 			}
 			return true;
-		}
-		, false);
+		});
 
 	return OutFilter;
 }
@@ -545,8 +554,7 @@ bool FSequencerFilterBar::AnyCommonFilterActive() const
 				return false;
 			}
 			return true;
-		}
-		, false);
+		});
 
 	return bOutActiveFilter;
 }
@@ -695,7 +703,6 @@ void FSequencerFilterBar::ActivateCommonFilters(const bool bInActivate
 
 			return true;
 		}
-		, true
 		, InMatchCategories);
 
 	if (bNeedsSave)
@@ -792,8 +799,7 @@ TArray<TSharedRef<FSequencerTrackFilter>> FSequencerFilterBar::GetActiveFilters(
 				OutFilters.Add(InFilter);
 			}
 			return true;
-		}
-		, true);
+		});
 
 	return OutFilters;
 }
@@ -811,8 +817,7 @@ bool FSequencerFilterBar::HasEnabledCommonFilters() const
 				return false;
 			}
 			return true;
-		}
-		, true);
+		});
 
 	if (bOutReturn)
 	{
@@ -828,8 +833,7 @@ bool FSequencerFilterBar::HasEnabledCommonFilters() const
 				return false;
 			}
 			return true;
-		}
-		, false);
+		});
 
 	return bOutReturn;
 }
@@ -930,7 +934,6 @@ void FSequencerFilterBar::EnableFilters(const bool bInEnable
 			}
 			return true;
 		}
-		, true
 		, InMatchCategories);
 
 	SequencerSettings->SaveConfig();
@@ -955,8 +958,7 @@ TArray<TSharedRef<FSequencerTrackFilter>> FSequencerFilterBar::GetEnabledFilters
 				OutFilters.Add(InFilter);
 			}
 			return true;
-		}
-		, true);
+		});
 
 	return OutFilters;
 }
@@ -1203,13 +1205,6 @@ TSharedRef<FFilterCategory> FSequencerFilterBar::GetMiscCategory() const
 	return MiscCategory;
 }
 
-void FSequencerFilterBar::ForEachFilter(const TFunctionRef<bool(const TSharedRef<FSequencerTrackFilter>&)>& InFunction
-	, const bool bInCheckSupportsSequence
-	, const TArray<TSharedRef<FFilterCategory>>& InCategories) const
-{
-	CommonFilters->ForEachFilter(InFunction, bInCheckSupportsSequence, InCategories);
-}
-
 bool FSequencerFilterBar::HasActiveLevelFilter() const
 {
 	return LevelFilter->HasHiddenLevels();
@@ -1279,7 +1274,7 @@ bool FSequencerFilterBar::IsGroupFilterActive(const FString& InGroupName) const
 
 bool FSequencerFilterBar::PassesAnyCommonFilter(const TViewModelPtr<IOutlinerExtension>& InNode)
 {
-	const UMovieSceneSequence* const FocusedSequence = Sequencer.GetFocusedMovieSceneSequence();
+	UMovieSceneSequence* const FocusedSequence = Sequencer.GetFocusedMovieSceneSequence();
 	if (!IsValid(FocusedSequence))
 	{
 		return true;
@@ -1292,9 +1287,7 @@ bool FSequencerFilterBar::PassesAnyCommonFilter(const TViewModelPtr<IOutlinerExt
 	CommonFilters->ForEachFilter([this, &InNode, &bPassedAnyFilters, &bAnyFilterActive, FocusedSequence]
 		(const TSharedRef<FSequencerTrackFilter>& InFilter)
 		{
-			const FString FilterName = InFilter->GetDisplayName().ToString();
-			if (FocusedSequence->IsFilterSupported(FilterName)
-				&& IsFilterActive(InFilter))
+			if (IsFilterActive(InFilter))
 			{
 				bAnyFilterActive = true;
 			
@@ -1306,8 +1299,7 @@ bool FSequencerFilterBar::PassesAnyCommonFilter(const TViewModelPtr<IOutlinerExt
 			}
 
 			return true;
-		}
-		, true);
+		});
 
 	if (!bAnyFilterActive)
 	{
@@ -1319,7 +1311,7 @@ bool FSequencerFilterBar::PassesAnyCommonFilter(const TViewModelPtr<IOutlinerExt
 
 bool FSequencerFilterBar::PassesAllInternalFilters(const TViewModelPtr<IOutlinerExtension>& InNode)
 {
-	const UMovieSceneSequence* const FocusedSequence = Sequencer.GetFocusedMovieSceneSequence();
+	UMovieSceneSequence* const FocusedSequence = Sequencer.GetFocusedMovieSceneSequence();
 	if (!IsValid(FocusedSequence))
 	{
 		return true;
@@ -1330,15 +1322,13 @@ bool FSequencerFilterBar::PassesAllInternalFilters(const TViewModelPtr<IOutliner
 	InternalFilters->ForEachFilter([this, &InNode, &bPassedAllFilters, FocusedSequence]
 		(const TSharedRef<FSequencerTrackFilter>& InFilter)
 		{
-			if (FocusedSequence->IsFilterSupported(InFilter->GetName())
-				&& !InFilter->PassesFilter(InNode))
+			if (!InFilter->PassesFilter(InNode))
 			{
 				bPassedAllFilters = false;
 				return false; // Stop processing filters
 			}
 			return true;
-		}
-		, true);
+		});
 
 	return bPassedAllFilters;
 }
@@ -1711,4 +1701,29 @@ TSharedPtr<SSequencerFilterBar> FSequencerFilterBar::GetWidget() const
 	return SequencerWidget->GetFilterBarWidget();
 }
 
+bool FSequencerFilterBar::IsFilterSupported(const TSharedRef<FSequencerTrackFilter>& InFilter) const
+{
+	UMovieSceneSequence* const RootMovieSceneSequence = Sequencer.GetRootMovieSceneSequence();
+	if (!IsValid(RootMovieSceneSequence))
+	{
+		return false;
+	}
+
+	const FString FilterName = InFilter->GetName();
+	const bool bFilterSupportsSequence = InFilter->SupportsSequence(RootMovieSceneSequence);
+	const bool bSequenceSupportsFilter = RootMovieSceneSequence->IsFilterSupported(FilterName);
+	return bFilterSupportsSequence || bSequenceSupportsFilter;
+}
+
+bool FSequencerFilterBar::IsFilterSupported(const FString& InFilterName) const
+{
+	const TArray<TSharedRef<FSequencerTrackFilter>> FilterList = GetFilterList();
+	const TSharedRef<FSequencerTrackFilter>* const FoundFilter = FilterList.FindByPredicate(
+		[InFilterName](const TSharedRef<FSequencerTrackFilter>& InFilter)
+		{
+			return InFilter->GetName().Equals(InFilterName, ESearchCase::IgnoreCase);
+		});
+	return FoundFilter ? IsFilterSupported(*FoundFilter) : false;
+}
+
 #undef LOCTEXT_NAMESPACE
diff --git a/Engine/Source/Editor/Sequencer/Private/Filters/SequencerFilterBar.h b/Engine/Source/Editor/Sequencer/Private/Filters/SequencerFilterBar.h
--- a/Engine/Source/Editor/Sequencer/Private/Filters/SequencerFilterBar.h
+++ b/Engine/Source/Editor/Sequencer/Private/Filters/SequencerFilterBar.h
@@ -205,10 +205,6 @@ public:
 	TSharedRef<FFilterCategory> GetComponentTypeCategory() const;
 	TSharedRef<FFilterCategory> GetMiscCategory() const;
 
-	void ForEachFilter(const TFunctionRef<bool(const TSharedRef<FSequencerTrackFilter>&)>& InFunction
-		, const bool bInCheckSupportsSequence
-		, const TArray<TSharedRef<FFilterCategory>>& InCategories = {}) const;
-
 	/** Level Filter Functions */
 
 	bool HasActiveLevelFilter() const;
@@ -244,6 +240,9 @@ public:
 	/** Attempts to get the filter bar widget from the Sequencer widget */
 	TSharedPtr<SSequencerFilterBar> GetWidget() const;
 
+	bool IsFilterSupported(const TSharedRef<FSequencerTrackFilter>& InFilter) const;
+	bool IsFilterSupported(const FString& InFilterName) const;
+
 protected:
 	void CreateDefaultFilters();
 
diff --git a/Engine/Source/Editor/Sequencer/Private/Filters/SequencerTrackFilterCollection.cpp b/Engine/Source/Editor/Sequencer/Private/Filters/SequencerTrackFilterCollection.cpp
--- a/Engine/Source/Editor/Sequencer/Private/Filters/SequencerTrackFilterCollection.cpp
+++ b/Engine/Source/Editor/Sequencer/Private/Filters/SequencerTrackFilterCollection.cpp
@@ -23,8 +23,7 @@ bool FSequencerTrackFilterCollection::ContainsFilter(const TSharedRef<FSequencer
 				return false;
 			}
 			return true;
-		}
-		, false);
+		});
 
 	return bOutContains;
 }
@@ -120,8 +119,7 @@ TArray<FText> FSequencerTrackFilterCollection::GetFilterDisplayNames() const
 		{
 			OutDisplayNames.Add(InTrackFilter->GetDisplayName());
 			return true;
-		}
-		, false);
+		});
 
 	return OutDisplayNames;
 }
@@ -136,7 +134,6 @@ TArray<TSharedRef<FSequencerTrackFilter>> FSequencerTrackFilterCollection::GetAl
 			OutFilters.Add(InTrackFilter);
 			return true;
 		}
-		, false
 		, InCategories);
 
 	return OutFilters;
@@ -160,8 +157,7 @@ TSet<TSharedRef<FFilterCategory>> FSequencerTrackFilterCollection::GetCategories
 				}
 			}
 			return true;
-		}
-		, false);
+		});
 
 	return OutCategories;
 }
@@ -178,20 +174,18 @@ TArray<TSharedRef<FSequencerTrackFilter>> FSequencerTrackFilterCollection::GetCa
 				OutFilters.Add(InTrackFilter);
 			}
 			return true;
-		}
-		, false);
+		});
 
 	return OutFilters;
 }
 
 void FSequencerTrackFilterCollection::ForEachFilter(const TFunctionRef<bool(const TSharedRef<FSequencerTrackFilter>&)>& InFunction
-	, const bool bInCheckSupportsSequence
 	, const TArray<TSharedRef<FFilterCategory>>& InCategories) const
 {
 	using namespace UE::Sequencer;
 
-	UMovieSceneSequence* const FocusedSequence = FilterInterface.GetSequencer().GetFocusedMovieSceneSequence();
-	if (!IsValid(FocusedSequence))
+	UMovieSceneSequence* const RootSequence = FilterInterface.GetSequencer().GetRootMovieSceneSequence();
+	if (!IsValid(RootSequence))
 	{
 		return;
 	}
@@ -199,10 +193,17 @@ void FSequencerTrackFilterCollection::ForEachFilter(const TFunctionRef<bool(cons
 	for (const TSharedPtr<IFilter<TViewModelPtr<FViewModel>>>& Filter : ChildFilters)
 	{
 		const TSharedPtr<FSequencerTrackFilter> FilterCasted = StaticCastSharedPtr<FSequencerTrackFilter>(Filter);
-		if (FilterCasted.IsValid()
-			&& (!bInCheckSupportsSequence || FilterCasted->SupportsSequence(FocusedSequence))
-			&& (InCategories.IsEmpty() || InCategories.Contains(FilterCasted->GetCategory()))
-			&& !InFunction(FilterCasted.ToSharedRef()))
+		if (!FilterCasted.IsValid())
+		{
+			continue;
+		}
+
+		if (!InCategories.IsEmpty() && InCategories.Contains(FilterCasted->GetCategory()))
+		{
+			continue;
+		}
+
+		if (!InFunction(FilterCasted.ToSharedRef()))
 		{
 			return;
 		}
diff --git a/Engine/Source/Editor/Sequencer/Public/Filters/Filters/SequencerTrackFilter_Keyed.h b/Engine/Source/Editor/Sequencer/Public/Filters/Filters/SequencerTrackFilter_Keyed.h
--- a/Engine/Source/Editor/Sequencer/Public/Filters/Filters/SequencerTrackFilter_Keyed.h
+++ b/Engine/Source/Editor/Sequencer/Public/Filters/Filters/SequencerTrackFilter_Keyed.h
@@ -17,7 +17,6 @@ public:
 	virtual bool ShouldUpdateOnTrackValueChanged() const override;
 	virtual FText GetDefaultToolTipText() const override;
 	virtual TSharedPtr<FUICommandInfo> GetToggleCommand() const override;
-	virtual bool SupportsSequence(UMovieSceneSequence* const InSequence) const override;
 	//~ End FSequencerTrackFilter
 
 	//~ Begin FFilterBase
diff --git a/Engine/Source/Editor/Sequencer/Public/Filters/Filters/SequencerTrackFilters.h b/Engine/Source/Editor/Sequencer/Public/Filters/Filters/SequencerTrackFilters.h
--- a/Engine/Source/Editor/Sequencer/Public/Filters/Filters/SequencerTrackFilters.h
+++ b/Engine/Source/Editor/Sequencer/Public/Filters/Filters/SequencerTrackFilters.h
@@ -89,6 +89,7 @@ public:
 	//~ Begin FSequencerTrackFilter
 	virtual FText GetDefaultToolTipText() const override;
 	virtual TSharedPtr<FUICommandInfo> GetToggleCommand() const override;
+	virtual bool SupportsSequence(UMovieSceneSequence* const InSequence) const override;
 	//~ End FSequencerTrackFilter
 };
 
@@ -114,6 +115,7 @@ public:
 	//~ Begin FSequencerTrackFilter
 	virtual FText GetDefaultToolTipText() const override;
 	virtual TSharedPtr<FUICommandInfo> GetToggleCommand() const override;
+	virtual bool SupportsSequence(UMovieSceneSequence* const InSequence) const override;
 	//~ End FSequencerTrackFilter
 };
 
@@ -139,6 +141,7 @@ public:
 	//~ Begin FSequencerTrackFilter
 	virtual FText GetDefaultToolTipText() const override;
 	virtual TSharedPtr<FUICommandInfo> GetToggleCommand() const override;
+	virtual bool SupportsSequence(UMovieSceneSequence* const InSequence) const override;
 	//~ End FSequencerTrackFilter
 };
 
@@ -165,6 +168,7 @@ public:
 	//~ Begin FSequencerTrackFilter
 	virtual FText GetDefaultToolTipText() const override;
 	virtual TSharedPtr<FUICommandInfo> GetToggleCommand() const override;
+	virtual bool SupportsSequence(UMovieSceneSequence* const InSequence) const override;
 	//~ End FSequencerTrackFilter
 };
 
@@ -190,6 +194,7 @@ public:
 	//~ Begin FSequencerTrackFilter
 	virtual FText GetDefaultToolTipText() const override;
 	virtual TSharedPtr<FUICommandInfo> GetToggleCommand() const override;
+	virtual bool SupportsSequence(UMovieSceneSequence* const InSequence) const override;
 	//~ End FSequencerTrackFilter
 };
 
@@ -215,6 +220,7 @@ public:
 	//~ Begin FSequencerTrackFilter
 	virtual FText GetDefaultToolTipText() const override;
 	virtual TSharedPtr<FUICommandInfo> GetToggleCommand() const override;
+	virtual bool SupportsSequence(UMovieSceneSequence* const InSequence) const override;
 	//~ End FSequencerTrackFilter
 };
 
@@ -240,6 +246,7 @@ public:
 	//~ Begin FSequencerTrackFilter
 	virtual FText GetDefaultToolTipText() const override;
 	virtual TSharedPtr<FUICommandInfo> GetToggleCommand() const override;
+	virtual bool SupportsSequence(UMovieSceneSequence* const InSequence) const override;
 	//~ End FSequencerTrackFilter
 };
 
@@ -265,6 +272,7 @@ public:
 	//~ Begin FSequencerTrackFilter
 	virtual FText GetDefaultToolTipText() const override;
 	virtual TSharedPtr<FUICommandInfo> GetToggleCommand() const override;
+	virtual bool SupportsSequence(UMovieSceneSequence* const InSequence) const override;
 	//~ End FSequencerTrackFilter
 };
 
@@ -290,6 +298,7 @@ public:
 	//~ Begin FSequencerTrackFilter
 	virtual FText GetDefaultToolTipText() const override;
 	virtual TSharedPtr<FUICommandInfo> GetToggleCommand() const override;
+	virtual bool SupportsSequence(UMovieSceneSequence* const InSequence) const override;
 	//~ End FSequencerTrackFilter
 };
 
@@ -315,6 +324,7 @@ public:
 	//~ Begin FSequencerTrackFilter
 	virtual FText GetDefaultToolTipText() const override;
 	virtual TSharedPtr<FUICommandInfo> GetToggleCommand() const override;
+	virtual bool SupportsSequence(UMovieSceneSequence* const InSequence) const override;
 	//~ End FSequencerTrackFilter
 };
 
@@ -340,6 +350,7 @@ public:
 	//~ Begin FSequencerTrackFilter
 	virtual FText GetDefaultToolTipText() const override;
 	virtual TSharedPtr<FUICommandInfo> GetToggleCommand() const override;
+	virtual bool SupportsSequence(UMovieSceneSequence* const InSequence) const override;
 	//~ End FSequencerTrackFilter
 };
 
@@ -365,7 +376,6 @@ public:
 	//~ Begin FSequencerTrackFilter
 	virtual FText GetDefaultToolTipText() const override;
 	virtual TSharedPtr<FUICommandInfo> GetToggleCommand() const override;
-	virtual TSubclassOf<UMovieSceneTrack> GetTrackClass() const override;
 	virtual bool SupportsSequence(UMovieSceneSequence* const InSequence) const override;
 	//~ End FSequencerTrackFilter
 };
diff --git a/Engine/Source/Editor/Sequencer/Public/Filters/SequencerFilterBase.h b/Engine/Source/Editor/Sequencer/Public/Filters/SequencerFilterBase.h
--- a/Engine/Source/Editor/Sequencer/Public/Filters/SequencerFilterBase.h
+++ b/Engine/Source/Editor/Sequencer/Public/Filters/SequencerFilterBase.h
@@ -5,8 +5,8 @@
 #include "Filters/FilterBase.h"
 #include "Filters/ISequencerFilterBar.h"
 #include "Templates/SharedPointer.h"
+#include "MovieSceneSequence.h"
 
-class UMovieSceneSequence;
 struct FSlateBrush;
 
 /** Sequencer base filter for all Sequencer filters (Sequencer, Navigation Tool, etc.) */
@@ -16,6 +16,30 @@ class FSequencerFilterBase
 	, public TSharedFromThis<FSequencerFilterBase<InFilterType>>
 {
 public:
+	SEQUENCER_API static bool SupportsLevelSequence(UMovieSceneSequence* const InSequence)
+	{
+		static UClass* const LevelSequenceClass = FindObject<UClass>(nullptr, TEXT("/Script/LevelSequence.LevelSequence"), true);
+		return IsValid(InSequence)
+			&& IsValid(LevelSequenceClass)
+			&& InSequence->GetClass()->IsChildOf(LevelSequenceClass);
+	}
+
+	SEQUENCER_API static bool SupportsUMGSequence(UMovieSceneSequence* const InSequence)
+	{
+		static UClass* const WidgetAnimationClass = FindObject<UClass>(nullptr, TEXT("/Script/UMG.WidgetAnimation"), true);
+		return IsValid(InSequence)
+			&& IsValid(WidgetAnimationClass)
+			&& InSequence->GetClass()->IsChildOf(WidgetAnimationClass);
+	}
+
+	SEQUENCER_API static bool SupportsDaySequence(UMovieSceneSequence* const InSequence)
+	{
+		static UClass* const DaySequenceClass = FindObject<UClass>(nullptr, TEXT("/Script/DaySequence.DaySequence"), true);
+		return IsValid(InSequence)
+			&& IsValid(DaySequenceClass)
+			&& InSequence->GetClass()->IsChildOf(DaySequenceClass);
+	}
+
 	SEQUENCER_API static FText BuildTooltipTextForCommand(const FText& InBaseText, const TSharedPtr<FUICommandInfo>& InCommand)
 	{
 		const TSharedRef<const FInputChord> FirstValidChord = InCommand->GetFirstValidChord();
diff --git a/Engine/Source/Editor/Sequencer/Public/Filters/SequencerTrackFilterBase.h b/Engine/Source/Editor/Sequencer/Public/Filters/SequencerTrackFilterBase.h
--- a/Engine/Source/Editor/Sequencer/Public/Filters/SequencerTrackFilterBase.h
+++ b/Engine/Source/Editor/Sequencer/Public/Filters/SequencerTrackFilterBase.h
@@ -37,9 +37,6 @@ using FSequencerTrackFilterType = UE::Sequencer::FViewModelPtr;
 class FSequencerTrackFilter : public FSequencerFilterBase<FSequencerTrackFilterType>
 {
 public:
-	SEQUENCER_API static bool SupportsLevelSequence(UMovieSceneSequence* const InSequence);
-	SEQUENCER_API static bool SupportsUMGSequence(UMovieSceneSequence* const InSequence);
-
 	template<typename InTrackClassType>
 	static bool IsSequenceTrackSupported(UMovieSceneSequence* const InSequence)
 	{
diff --git a/Engine/Source/Editor/Sequencer/Public/Filters/SequencerTrackFilterCollection.h b/Engine/Source/Editor/Sequencer/Public/Filters/SequencerTrackFilterCollection.h
--- a/Engine/Source/Editor/Sequencer/Public/Filters/SequencerTrackFilterCollection.h
+++ b/Engine/Source/Editor/Sequencer/Public/Filters/SequencerTrackFilterCollection.h
@@ -67,7 +67,6 @@ public:
 	TArray<TSharedRef<FSequencerTrackFilter>> GetCategoryFilters(const TSharedRef<FFilterCategory>& InCategory) const;
 
 	void ForEachFilter(const TFunctionRef<bool(const TSharedRef<FSequencerTrackFilter>&)>& InFunction
-		, const bool bInCheckSupportsSequence
 		, const TArray<TSharedRef<FFilterCategory>>& InCategories = {}) const;
 
 protected:
diff --git a/Engine/Source/Runtime/LevelSequence/Private/LevelSequence.cpp b/Engine/Source/Runtime/LevelSequence/Private/LevelSequence.cpp
--- a/Engine/Source/Runtime/LevelSequence/Private/LevelSequence.cpp
+++ b/Engine/Source/Runtime/LevelSequence/Private/LevelSequence.cpp
@@ -164,6 +164,40 @@ ETrackSupport ULevelSequence::IsTrackSupportedImpl(TSubclassOf<class UMovieScene
 	return Super::IsTrackSupportedImpl(InTrackClass);
 }
 
+bool ULevelSequence::IsFilterSupportedImpl(const FString& InFilterName) const
+{
+	static const TArray<FString> SupportedFilters = {
+		TEXT("Audio"),
+		TEXT("Camera"),
+		TEXT("CameraCut"),
+		TEXT("CinematicShot"),
+		TEXT("Condition"),
+		TEXT("ControlRigControl"),
+		TEXT("DataLayer"),
+		TEXT("Event"),
+		TEXT("Fade"),
+		TEXT("Folder"),
+		TEXT("Keyed"),
+		TEXT("Group"),
+		TEXT("Level"),
+		TEXT("LevelVisibility"),
+		TEXT("Light"),
+		TEXT("LiveLink"),
+		TEXT("Media"),
+		TEXT("Niagara"),
+		TEXT("Particle"),
+		TEXT("Selected"),
+		TEXT("SelectedAnimLayersFilter"),
+		TEXT("SelectedControlRigControl"),
+		TEXT("SkeletalMesh"),
+		TEXT("Subsequence"),
+		TEXT("TimeDilation"),
+		TEXT("TimeWarp"),
+		TEXT("Unbound")
+	};
+	return SupportedFilters.Contains(InFilterName);
+}
+
 void ULevelSequence::GetAssetRegistryTags(TArray<FAssetRegistryTag>& OutTags) const
 {
 	PRAGMA_DISABLE_DEPRECATION_WARNINGS;
diff --git a/Engine/Source/Runtime/LevelSequence/Public/LevelSequence.h b/Engine/Source/Runtime/LevelSequence/Public/LevelSequence.h
--- a/Engine/Source/Runtime/LevelSequence/Public/LevelSequence.h
+++ b/Engine/Source/Runtime/LevelSequence/Public/LevelSequence.h
@@ -79,6 +79,7 @@ public:
 
 #if WITH_EDITOR
 	LEVELSEQUENCE_API virtual ETrackSupport IsTrackSupportedImpl(TSubclassOf<class UMovieSceneTrack> InTrackClass) const override;
+	LEVELSEQUENCE_API virtual bool IsFilterSupportedImpl(const FString& InFilterName) const override;
 	LEVELSEQUENCE_API virtual void GetAssetRegistryTagMetadata(TMap<FName, FAssetRegistryTagMetadata>& OutMetadata) const override;
 	LEVELSEQUENCE_API virtual void GetAssetRegistryTags(FAssetRegistryTagsContext Context) const override;
 	UE_DEPRECATED(5.4, "Implement the version that takes FAssetRegistryTagsContext instead.")
diff --git a/Engine/Source/Runtime/UMG/Private/Animation/WidgetAnimation.cpp b/Engine/Source/Runtime/UMG/Private/Animation/WidgetAnimation.cpp
--- a/Engine/Source/Runtime/UMG/Private/Animation/WidgetAnimation.cpp
+++ b/Engine/Source/Runtime/UMG/Private/Animation/WidgetAnimation.cpp
@@ -126,6 +126,21 @@ ETrackSupport UWidgetAnimation::IsTrackSupportedImpl(TSubclassOf<class UMovieSce
 	return Super::IsTrackSupportedImpl(InTrackClass);
 }
 
+bool UWidgetAnimation::IsFilterSupportedImpl(const FString& InFilterName) const
+{
+	static const TArray<FString> SupportedFilters = {
+		TEXT("Audio"),
+		TEXT("Event"),
+		TEXT("Keyed"),
+		TEXT("Folder"),
+		TEXT("Group"),
+		TEXT("TimeDilation"),
+		TEXT("TimeWarp"),
+		TEXT("Unbound")
+	};
+	return SupportedFilters.Contains(InFilterName);
+}
+
 const FSlateBrush* UWidgetAnimation::GetCustomBrushForBinding(FGuid BindingID) const
 {
 	for (const FWidgetAnimationBinding& Binding : AnimationBindings)
diff --git a/Engine/Source/Runtime/UMG/Public/Animation/WidgetAnimation.h b/Engine/Source/Runtime/UMG/Public/Animation/WidgetAnimation.h
--- a/Engine/Source/Runtime/UMG/Public/Animation/WidgetAnimation.h
+++ b/Engine/Source/Runtime/UMG/Public/Animation/WidgetAnimation.h
@@ -108,6 +108,7 @@ public:
 #if WITH_EDITOR
 	virtual void IterateDynamicBindings(TFunction<void(const FGuid&, FMovieSceneDynamicBinding&)> InCallback) override;
 	virtual ETrackSupport IsTrackSupportedImpl(TSubclassOf<class UMovieSceneTrack> InTrackClass) const override;
+	virtual bool IsFilterSupportedImpl(const FString& InFilterName) const override;
 	virtual const FSlateBrush* GetCustomBrushForBinding(FGuid BindingID) const override;
 
 	DECLARE_MULTICAST_DELEGATE_OneParam(FFixupWidgetDynamicBindingsEvent, UWidgetAnimation*);
