Sequencer: Support for cancelllng a paste operation so that it doesn't attempt to paste anything else.

This is used for when pasting object bindings, which prompts the user whether to duplicate the bound actors. If the user chooses to cancel, nothing else should be pasted

diff --git a/Engine/Source/Editor/Sequencer/Private/Sequencer.cpp b/Engine/Source/Editor/Sequencer/Private/Sequencer.cpp
--- a/Engine/Source/Editor/Sequencer/Private/Sequencer.cpp
+++ b/Engine/Source/Editor/Sequencer/Private/Sequencer.cpp
@@ -8483,18 +8483,36 @@ bool FSequencer::DoPaste(bool bClearSelection)
 	}
 	if (EnumHasAnyFlags(PasteSupport, ESequencerPasteSupport::ObjectBindings)) 
 	{
-		bAnythingPasted |= PasteObjectBindings(TextToImport, ParentFolder, PastedFolders, PasteErrors, bClearSelection);
+		EPasteResult PasteObjectBindingsResult = PasteObjectBindings(TextToImport, ParentFolder, PastedFolders, PasteErrors, bClearSelection);
+		if (PasteObjectBindingsResult == EPasteResult::Cancelled)
+		{
+			return false;
+		}
+
+		bAnythingPasted |= PasteObjectBindingsResult == EPasteResult::Success;
 	}
 	if (EnumHasAnyFlags(PasteSupport, ESequencerPasteSupport::Tracks)) 
 	{
-		bAnythingPasted |= PasteTracks(TextToImport, ParentFolder, PastedFolders, PasteErrors, bClearSelection);
+		EPasteResult PasteTracksResult = PasteTracks(TextToImport, ParentFolder, PastedFolders, PasteErrors, bClearSelection);
+		if (PasteTracksResult == EPasteResult::Cancelled)
+		{
+			return false;
+		}
+
+		bAnythingPasted |= PasteTracksResult == EPasteResult::Success;
 	}
 
 	if (!bAnythingPasted)
 	{
 		if (EnumHasAnyFlags(PasteSupport, ESequencerPasteSupport::Sections))
 		{
-			bAnythingPasted |= PasteSections(TextToImport, PasteErrors);
+			EPasteResult PasteSectionsResult = PasteSections(TextToImport, PasteErrors);
+			if (PasteSectionsResult == EPasteResult::Cancelled)
+			{
+				return false;
+			}
+
+			bAnythingPasted |= PasteSectionsResult == EPasteResult::Success;
 		}
 	}
 
@@ -8510,7 +8528,7 @@ bool FSequencer::DoPaste(bool bClearSelection)
 	return bAnythingPasted;
 }
 
-bool FSequencer::PasteObjectBindings(const FString& TextToImport, UMovieSceneFolder* InParentFolder, const TArray<UMovieSceneFolder*>& InFolders, TArray<FNotificationInfo>& PasteErrors, bool bClearSelection)
+FSequencer::EPasteResult FSequencer::PasteObjectBindings(const FString& TextToImport, UMovieSceneFolder* InParentFolder, const TArray<UMovieSceneFolder*>& InFolders, TArray<FNotificationInfo>& PasteErrors, bool bClearSelection)
 {
 	using namespace UE::Sequencer;
 
@@ -8561,7 +8579,7 @@ bool FSequencer::PasteObjectBindings(const FString& TextToImport, UMovieSceneFol
 						break;
 
 					case EAppReturnType::Cancel:
-						return false;
+						return EPasteResult::Cancelled;
 					}
 					bWasPrompted = true;				
 					break;
@@ -8571,10 +8589,14 @@ bool FSequencer::PasteObjectBindings(const FString& TextToImport, UMovieSceneFol
 	}
 
 	TArray<FMovieSceneBindingProxy> OutBindings;
-	return FSequencerUtilities::PasteBindings(TextToImport, AsShared(), FMovieScenePasteBindingsParams(TargetBindings, InParentFolder, InFolders, bDuplicateExistingActors), OutBindings, PasteErrors);
+	if (FSequencerUtilities::PasteBindings(TextToImport, AsShared(), FMovieScenePasteBindingsParams(TargetBindings, InParentFolder, InFolders, bDuplicateExistingActors), OutBindings, PasteErrors))
+	{
+		return EPasteResult::Success;
+	}
+	return EPasteResult::Failure;
 }
 
-bool FSequencer::PasteTracks(const FString& TextToImport, UMovieSceneFolder* InParentFolder, const TArray<UMovieSceneFolder*>& InFolders, TArray<FNotificationInfo>& PasteErrors, bool bClearSelection)
+FSequencer::EPasteResult FSequencer::PasteTracks(const FString& TextToImport, UMovieSceneFolder* InParentFolder, const TArray<UMovieSceneFolder*>& InFolders, TArray<FNotificationInfo>& PasteErrors, bool bClearSelection)
 {
 	using namespace UE::Sequencer;
 
@@ -8590,15 +8612,16 @@ bool FSequencer::PasteTracks(const FString& TextToImport, UMovieSceneFolder* InP
 	TArray<UMovieSceneTrack*> OutTracks;
 	bool bSuccess = FSequencerUtilities::PasteTracks(TextToImport, FMovieScenePasteTracksParams(GetFocusedMovieSceneSequence(), TargetBindings, InParentFolder, InFolders), OutTracks, PasteErrors);
 
-	if(bSuccess)
+	if (bSuccess)
 	{
 		NotifyMovieSceneDataChanged(EMovieSceneDataChangeType::MovieSceneStructureItemAdded);
+		return EPasteResult::Success;
 	}
 	
-	return bSuccess;
+	return EPasteResult::Failure;
 }
 
-bool FSequencer::PasteSections(const FString& TextToImport, TArray<FNotificationInfo>& PasteErrors)
+FSequencer::EPasteResult FSequencer::PasteSections(const FString& TextToImport, TArray<FNotificationInfo>& PasteErrors)
 {
 	using namespace UE::Sequencer;
 
@@ -8648,7 +8671,7 @@ bool FSequencer::PasteSections(const FString& TextToImport, TArray<FNotification
 	{
 		FNotificationInfo Info(LOCTEXT("PasteSections_NoSelectedTracks", "Can't paste section. No selected tracks to paste sections onto"));
 		PasteErrors.Add(Info);
-		return false;
+		return EPasteResult::Failure;
 	}
 
 	TArray<UMovieSceneTrack*> Tracks;
@@ -8704,7 +8727,7 @@ bool FSequencer::PasteSections(const FString& TextToImport, TArray<FNotification
 	TArray<UMovieSceneSection*> NewSections;
 	if (!FSequencerUtilities::PasteSections(TextToImport, FMovieScenePasteSectionsParams(Tracks, TrackRowIndices, GetLocalTime().Time), NewSections, PasteErrors))
 	{
-		return false;
+		return EPasteResult::Failure;
 	}
 
 	FSectionModelStorageExtension* SectionModelStorage = ViewModel->GetRootModel()->CastDynamic<FSectionModelStorageExtension>();
@@ -8728,7 +8751,7 @@ bool FSequencer::PasteSections(const FString& TextToImport, TArray<FNotification
 
 	NotifyMovieSceneDataChanged(EMovieSceneDataChangeType::MovieSceneStructureItemAdded);
 
-	return true;
+	return EPasteResult::Success;
 }
 
 bool FSequencer::CanPaste(const FString& TextToImport)
diff --git a/Engine/Source/Editor/Sequencer/Private/Sequencer.h b/Engine/Source/Editor/Sequencer/Private/Sequencer.h
--- a/Engine/Source/Editor/Sequencer/Private/Sequencer.h
+++ b/Engine/Source/Editor/Sequencer/Private/Sequencer.h
@@ -577,10 +577,17 @@ public:
 	 * Attempts to paste from the clipboard
 	 * @return Whether the paste event was handled
 	 */
+	enum class EPasteResult : uint8
+	{
+		Success,
+		Failure,
+		Cancelled
+	};
+
 	bool DoPaste(bool bClearSelection = false);
-	bool PasteTracks(const FString& TextToImport, UMovieSceneFolder* ParentFolder, const TArray<UMovieSceneFolder*>& InFolders, TArray<FNotificationInfo>& PasteErrors, bool bClearSelection = false);
-	bool PasteSections(const FString& TextToImport, TArray<FNotificationInfo>& PasteErrors);
-	bool PasteObjectBindings(const FString& TextToImport, UMovieSceneFolder* ParentFolder, const TArray<UMovieSceneFolder*>& InFolders, TArray<FNotificationInfo>& PasteErrors, bool bClearSelection = false);
+	EPasteResult PasteTracks(const FString& TextToImport, UMovieSceneFolder* ParentFolder, const TArray<UMovieSceneFolder*>& InFolders, TArray<FNotificationInfo>& PasteErrors, bool bClearSelection = false);
+	EPasteResult PasteSections(const FString& TextToImport, TArray<FNotificationInfo>& PasteErrors);
+	EPasteResult PasteObjectBindings(const FString& TextToImport, UMovieSceneFolder* ParentFolder, const TArray<UMovieSceneFolder*>& InFolders, TArray<FNotificationInfo>& PasteErrors, bool bClearSelection = false);
 
 	/** Called when a user executes the locked node menu item */
 	void ToggleNodeLocked();
