[Sequencer] Change "Align to Playhead" action to not use relative location for multiple layer bars and sections. Keys will remain as relative for multiple selections.

diff --git a/Engine/Source/Editor/Sequencer/Private/Tools/SequencerSelectionAlignmentUtils.cpp b/Engine/Source/Editor/Sequencer/Private/Tools/SequencerSelectionAlignmentUtils.cpp
--- a/Engine/Source/Editor/Sequencer/Private/Tools/SequencerSelectionAlignmentUtils.cpp
+++ b/Engine/Source/Editor/Sequencer/Private/Tools/SequencerSelectionAlignmentUtils.cpp
@@ -204,40 +204,23 @@ void FSequencerSelectionAlignmentUtils::AlignSelectionToTime(const ISequencer& I
 	const FScopedTransaction Transaction(NSLOCTEXT("Sequencer", "AlignToPlayhead", "Align to Playhead"), bInTransact);
 
 	// Layer Bars and Sections
+	// Set all to the time, regardless of the relative offsets of each selection
 	if (bIsLayerBarOrSection && KeyFrameCount == 0)
 	{
-		TOptional<FFrameNumber> EarliestTime;
-
-		if (EarliestLayerBarTime.IsSet() && EarliestSectionTime.IsSet())
-		{
-			EarliestTime = (EarliestLayerBarTime.GetValue() < EarliestSectionTime.GetValue())
-				? EarliestLayerBarTime.GetValue() : EarliestSectionTime.GetValue();
-		}
-		else if (EarliestLayerBarTime.IsSet())
+		for (const TViewModelPtr<FLayerBarModel>& LayerBarModel : LayerBarModels)
 		{
-			EarliestTime = EarliestLayerBarTime.GetValue();
-		}
-		else if (EarliestSectionTime.IsSet())
-		{
-			EarliestTime = EarliestSectionTime.GetValue();
+			const FFrameNumber RelativeDistance = FrameNumber - LayerBarModel->ComputeRange().GetLowerBoundValue();
+			LayerBarModel->Offset(RelativeDistance);
 		}
 
-		if (EarliestTime.IsSet())
-		{
-			const FFrameNumber RelativeDistance = FrameNumber - EarliestTime.GetValue();
-
-			for (const TViewModelPtr<FLayerBarModel>& LayerBarModel : LayerBarModels)
-			{
-				LayerBarModel->Offset(RelativeDistance);
-			}
-
-			for (const TViewModelPtr<FSectionModel>& SectionModel : SectionModels)
-            {
-            	SectionModel->OffsetLayerBar(RelativeDistance);
-            }
-		}
+		for (const TViewModelPtr<FSectionModel>& SectionModel : SectionModels)
+        {
+			const FFrameNumber RelativeDistance = FrameNumber - SectionModel->GetLayerBarRange().GetLowerBoundValue();
+            SectionModel->OffsetLayerBar(RelativeDistance);
+        }
 	}
 	// Keyframes
+	// Find the earliest time in the selection and move the selection, while maintaining relative offsets between each selection
 	else if (bNoLayerBarOrSection && KeyFrameCount > 0)
 	{
 		TOptional<FFrameNumber> EarliestKeyframeTime;
