Sequencer: Fix issue where some tracks don't show up as options to add sometimes because we weren't looking for the object class through its template.

diff --git a/Engine/Source/Editor/Sequencer/Private/MVVM/ViewModels/PossessableModel.cpp b/Engine/Source/Editor/Sequencer/Private/MVVM/ViewModels/PossessableModel.cpp
--- a/Engine/Source/Editor/Sequencer/Private/MVVM/ViewModels/PossessableModel.cpp
+++ b/Engine/Source/Editor/Sequencer/Private/MVVM/ViewModels/PossessableModel.cpp
@@ -160,7 +160,13 @@ const UClass* FPossessableModel::FindObjectClass() const
 	FMovieScenePossessable* Possessable = MovieScene ? MovieScene->FindPossessable(ObjectBindingID) : nullptr;
 
 	const UClass* Class = UObject::StaticClass();
-	if (Possessable && Possessable->GetPossessedObjectClass() != nullptr)
+	TSharedPtr<ISequencer> Sequencer = OwnerModel->GetSequencer();
+
+	if (UObject* Template = MovieSceneHelpers::GetObjectTemplate(OwnerModel->GetSequence(), ObjectBindingID, Sequencer->GetSharedPlaybackState()))
+	{
+		Class = Template->GetClass();
+	}
+	else if (Possessable && Possessable->GetPossessedObjectClass() != nullptr)
 	{
 		Class = Possessable->GetPossessedObjectClass();
 	}
