Sequencer: Preference replaceable when alt key is held down. This already happens with dragging in actors- this is for consistency.

diff --git a/Engine/Source/Editor/Sequencer/Private/SSequencer.cpp b/Engine/Source/Editor/Sequencer/Private/SSequencer.cpp
--- a/Engine/Source/Editor/Sequencer/Private/SSequencer.cpp
+++ b/Engine/Source/Editor/Sequencer/Private/SSequencer.cpp
@@ -3120,9 +3120,11 @@ void SSequencer::OnAssetsDropped(TSharedPtr<FAssetDragDropOp> DragDropOp)
 					continue;
 				}
 
+				bool bPreferenceReplaceable = FSlateApplication::Get().GetModifierKeys().IsAltDown();
+
 				UE::Sequencer::FCreateBindingParams Params;
 				Params.BindingNameOverride = CurObject->GetName();
-				Params.bSpawnable = true;
+				Params.bSpawnable = bPreferenceReplaceable ? false : true;
 				Params.bReplaceable = true;
 				Params.ActorFactory = DragDropOp->GetActorFactory();
 
