Sequencer: Fix static analysis warnings.

diff --git a/Engine/Source/Editor/Sequencer/Private/MVVM/ViewModels/PossessableModel.cpp b/Engine/Source/Editor/Sequencer/Private/MVVM/ViewModels/PossessableModel.cpp
--- a/Engine/Source/Editor/Sequencer/Private/MVVM/ViewModels/PossessableModel.cpp
+++ b/Engine/Source/Editor/Sequencer/Private/MVVM/ViewModels/PossessableModel.cpp
@@ -156,21 +156,25 @@ const FSlateBrush* FPossessableModel::GetIconOverlayBrush() const
 
 const UClass* FPossessableModel::FindObjectClass() const
 {
-	UMovieScene*            MovieScene  = OwnerModel ? OwnerModel->GetMovieScene() : nullptr;
-	FMovieScenePossessable* Possessable = MovieScene ? MovieScene->FindPossessable(ObjectBindingID) : nullptr;
-
 	const UClass* Class = UObject::StaticClass();
-	TSharedPtr<ISequencer> Sequencer = OwnerModel->GetSequencer();
-
-	if (UObject* Template = MovieSceneHelpers::GetObjectTemplate(OwnerModel->GetSequence(), ObjectBindingID, Sequencer->GetSharedPlaybackState()))
+	if (OwnerModel)
 	{
-		Class = Template->GetClass();
-	}
-	else if (Possessable && Possessable->GetPossessedObjectClass() != nullptr)
-	{
-		Class = Possessable->GetPossessedObjectClass();
+		if (UMovieScene* MovieScene = OwnerModel->GetMovieScene())
+		{
+			FMovieScenePossessable* Possessable = MovieScene->FindPossessable(ObjectBindingID);
+			if (TSharedPtr<ISequencer> Sequencer = OwnerModel->GetSequencer())
+			{
+				if (UObject* Template = MovieSceneHelpers::GetObjectTemplate(OwnerModel->GetSequence(), ObjectBindingID, Sequencer->GetSharedPlaybackState()))
+				{
+					Class = Template->GetClass();
+				}
+				else if (Possessable && Possessable->GetPossessedObjectClass() != nullptr)
+				{
+					Class = Possessable->GetPossessedObjectClass();
+				}
+			}
+		}
 	}
-
 	return Class;
 }
 
