Sequencer: Combine snapping interval settings into a toggle "Force Whole Frames" which is independent of GetIsSnapEnabled

Removed bSnapKeyTimesToInterval, bSnapSectionTimesToInterval, bSnapPlayTimeToInterval

Additional fixes:
- Shift should temporarily enable snapping in key/section dragging. This already happens when holding down shift and dragging the playhead.
- Fixed dragging frame markers not snapping to whole frames when dragging (they do now when Force Whole Frames is enabled) because the difference wasn't snapped to a whole frame
- Fixed shift-dragging playback range not snapping to whole frames when dragging (they do now when Force Whole Frames is enabled) because the difference wasn't snapped to a whole frame

diff --git a/Engine/Source/Editor/MovieSceneTools/Private/TrackEditors/TimeWarpTrackEditor.cpp b/Engine/Source/Editor/MovieSceneTools/Private/TrackEditors/TimeWarpTrackEditor.cpp
--- a/Engine/Source/Editor/MovieSceneTools/Private/TrackEditors/TimeWarpTrackEditor.cpp
+++ b/Engine/Source/Editor/MovieSceneTools/Private/TrackEditors/TimeWarpTrackEditor.cpp
@@ -75,18 +75,17 @@ struct FScrubberHotspot : ITrackAreaHotspot
 				if (TSharedPtr<ISequencer> Sequencer = WeakSequencer.Pin())
 				{
 					USequencerSettings* Settings = Sequencer->GetSequencerSettings();
-					if (Settings->GetIsSnapEnabled() || MouseEvent.IsShiftDown())
+					if (Settings->GetForceWholeFrames())
 					{
-						if (Settings->GetSnapPlayTimeToInterval())
-						{
-							// Set the style of the scrub handle
-							ScrubTime = FFrameRate::Snap(ScrubTime, Sequencer->GetFocusedTickResolution(), Sequencer->GetFocusedDisplayRate());
-						}
+						ScrubTime = FFrameRate::Snap(ScrubTime, Sequencer->GetFocusedTickResolution(), Sequencer->GetFocusedDisplayRate());
+					}
 
-						ENearestKeyOption NearestKeyOption = MouseEvent.IsShiftDown()
-							? ENearestKeyOption::NKO_SearchKeys | ENearestKeyOption::NKO_SearchSections | ENearestKeyOption::NKO_SearchMarkers
-							: ENearestKeyOption::NKO_None;
+					ENearestKeyOption NearestKeyOption = MouseEvent.IsShiftDown()
+						? ENearestKeyOption::NKO_SearchKeys | ENearestKeyOption::NKO_SearchSections | ENearestKeyOption::NKO_SearchMarkers
+						: ENearestKeyOption::NKO_None;
 
+					if (Settings->GetIsSnapEnabled() || MouseEvent.IsShiftDown())
+					{
 						if (Settings->GetSnapPlayTimeToKeys())
 						{
 							EnumAddFlags(NearestKeyOption, ENearestKeyOption::NKO_SearchKeys);
@@ -429,4 +428,4 @@ void FTimeWarpTrackEditor::HandleAddTimeWarpTrack(TSubclassOf<UMovieSceneTimeWar
 	SequencerPtr->OnAddTrack(NewTrack, FGuid());
 }
 
-#undef LOCTEXT_NAMESPACE
\ No newline at end of file
+#undef LOCTEXT_NAMESPACE
diff --git a/Engine/Source/Editor/Sequencer/Private/MVVM/Views/SChannelView.cpp b/Engine/Source/Editor/Sequencer/Private/MVVM/Views/SChannelView.cpp
--- a/Engine/Source/Editor/Sequencer/Private/MVVM/Views/SChannelView.cpp
+++ b/Engine/Source/Editor/Sequencer/Private/MVVM/Views/SChannelView.cpp
@@ -394,7 +394,7 @@ FReply SChannelView::OnMouseButtonUp(const FGeometry& MyGeometry, const FPointer
 
 		// Snap keys on mouse up since we want to create keys at the exact mouse position (ie. to keep the newly created keys under the mouse 
 		// while dragging) but obey snapping rules if necessary
-		if (Sequencer->GetSequencerSettings()->GetIsSnapEnabled() && Sequencer->GetSequencerSettings()->GetSnapKeyTimesToInterval())
+		if (Sequencer->GetSequencerSettings()->GetForceWholeFrames())
 		{
 			Sequencer->SnapToFrame();
 
diff --git a/Engine/Source/Editor/Sequencer/Private/MVVM/Views/SSequencerTrackAreaView.cpp b/Engine/Source/Editor/Sequencer/Private/MVVM/Views/SSequencerTrackAreaView.cpp
--- a/Engine/Source/Editor/Sequencer/Private/MVVM/Views/SSequencerTrackAreaView.cpp
+++ b/Engine/Source/Editor/Sequencer/Private/MVVM/Views/SSequencerTrackAreaView.cpp
@@ -130,7 +130,7 @@ FReply SSequencerTrackAreaView::OnDragOver(const FGeometry& MyGeometry, const FD
 
 			FVector2D LocalPos = MyGeometry.AbsoluteToLocal(DragDropEvent.GetScreenSpacePosition());
 			FFrameNumber DropFrameNumber = TimeToPixel->PixelToFrame(LocalPos.X).FrameNumber;
-			if (Sequencer->GetSequencerSettings()->GetIsSnapEnabled() && Sequencer->GetSequencerSettings()->GetSnapPlayTimeToInterval())
+			if (Sequencer->GetSequencerSettings()->GetForceWholeFrames())
 			{
 				DropFrameNumber = FFrameRate::Snap(DropFrameNumber, Sequencer->GetFocusedTickResolution(), Sequencer->GetFocusedDisplayRate()).FrameNumber;
 			}
@@ -183,7 +183,7 @@ FReply SSequencerTrackAreaView::OnDrop(const FGeometry& MyGeometry, const FDragD
 
 			FVector2D LocalPos = MyGeometry.AbsoluteToLocal(DragDropEvent.GetScreenSpacePosition());
 			FFrameNumber DropFrameNumber = TimeToPixel->PixelToFrame(LocalPos.X).FrameNumber;
-			if (Sequencer->GetSequencerSettings()->GetIsSnapEnabled() && Sequencer->GetSequencerSettings()->GetSnapPlayTimeToInterval())
+			if (Sequencer->GetSequencerSettings()->GetForceWholeFrames())
 			{
 				DropFrameNumber = FFrameRate::Snap(DropFrameNumber, Sequencer->GetFocusedTickResolution(), Sequencer->GetFocusedDisplayRate()).FrameNumber;
 			}
diff --git a/Engine/Source/Editor/Sequencer/Private/SSequencer.cpp b/Engine/Source/Editor/Sequencer/Private/SSequencer.cpp
--- a/Engine/Source/Editor/Sequencer/Private/SSequencer.cpp
+++ b/Engine/Source/Editor/Sequencer/Private/SSequencer.cpp
@@ -2497,29 +2497,16 @@ TSharedRef<SWidget> SSequencer::MakeSnapMenu()
 {
 	FMenuBuilder MenuBuilder( false, SequencerPtr.Pin()->GetCommandBindings() );
 
-	MenuBuilder.BeginSection( "KeySnapping", LOCTEXT( "SnappingMenuKeyHeader", "Key Snapping" ) );
+	MenuBuilder.BeginSection( "KeyAndSectionsSnapping", LOCTEXT( "SnappingMenuKeyAndSectionsHeader", "Key and Sections Snapping" ) );
 	{
-		MenuBuilder.AddMenuEntry( FSequencerCommands::Get().ToggleSnapKeyTimesToInterval );
 		MenuBuilder.AddMenuEntry( FSequencerCommands::Get().ToggleSnapKeyTimesToKeys );
-	}
-	MenuBuilder.EndSection();
-
-	MenuBuilder.BeginSection("SectionSnapping", LOCTEXT("SnappingMenuSectionHeader", "Section Snapping"));
-	{
-		MenuBuilder.AddMenuEntry(FSequencerCommands::Get().ToggleSnapSectionTimesToInterval);
-		MenuBuilder.AddMenuEntry(FSequencerCommands::Get().ToggleSnapSectionTimesToSections);
-	}
-	MenuBuilder.EndSection();
-
-	MenuBuilder.BeginSection("KeyAndSectionSnapping", LOCTEXT("SnappingMenuKeysAndSectionHeader", "Keys and Sections Snapping"));
-	{
-		MenuBuilder.AddMenuEntry(FSequencerCommands::Get().ToggleSnapKeysAndSectionsToPlayRange);
+		MenuBuilder.AddMenuEntry( FSequencerCommands::Get().ToggleSnapSectionTimesToSections );
+		MenuBuilder.AddMenuEntry( FSequencerCommands::Get().ToggleSnapKeysAndSectionsToPlayRange );
 	}
 	MenuBuilder.EndSection();
 
 	MenuBuilder.BeginSection( "PlayHeadSnapping", LOCTEXT( "SnappingMenuPlayHeadHeader", "Play Head Snapping" ) );
 	{
-		MenuBuilder.AddMenuEntry( FSequencerCommands::Get().ToggleSnapPlayTimeToInterval );
 		MenuBuilder.AddMenuEntry( FSequencerCommands::Get().ToggleSnapPlayTimeToKeys );
 		MenuBuilder.AddMenuEntry( FSequencerCommands::Get().ToggleSnapPlayTimeToSections );
 		MenuBuilder.AddMenuEntry( FSequencerCommands::Get().ToggleSnapPlayTimeToMarkers );
@@ -2534,6 +2521,11 @@ TSharedRef<SWidget> SSequencer::MakeSnapMenu()
 	}
 	MenuBuilder.EndSection();
 
+	MenuBuilder.BeginSection( "IntervalSnapping", LOCTEXT("SnappingMenuIntervalHeader", "Force Whole Frames"));
+	{
+		MenuBuilder.AddMenuEntry(FSequencerCommands::Get().ToggleForceWholeFrames);
+	}
+	MenuBuilder.EndSection();
 	return MenuBuilder.MakeWidget();
 }
 
@@ -3511,7 +3503,7 @@ FPasteContextMenuArgs SSequencer::GeneratePasteArgs(FFrameNumber PasteAtTime, TS
 	using namespace UE::Sequencer;
 
 	TSharedPtr<FSequencer> Sequencer = SequencerPtr.Pin();
-	if (GetSequencerSettings()->GetIsSnapEnabled())
+	if (GetSequencerSettings()->GetForceWholeFrames())
 	{
 		FFrameRate TickResolution = Sequencer->GetFocusedTickResolution();
 		FFrameRate DisplayRate    = Sequencer->GetFocusedDisplayRate();
diff --git a/Engine/Source/Editor/Sequencer/Private/Sequencer.cpp b/Engine/Source/Editor/Sequencer/Private/Sequencer.cpp
--- a/Engine/Source/Editor/Sequencer/Private/Sequencer.cpp
+++ b/Engine/Source/Editor/Sequencer/Private/Sequencer.cpp
@@ -3441,7 +3441,7 @@ void FSequencer::SetLocalTime( FFrameTime NewTime, ESnapTimeMode SnapTimeMode, b
 	}
 
 	// Perform snapping
-	if ((SnapTimeMode & ESnapTimeMode::STM_Interval) && Settings->GetIsSnapEnabled())
+	if ((SnapTimeMode & ESnapTimeMode::STM_Interval) && Settings->GetForceWholeFrames())
 	{
 		FFrameRate LocalDisplayRate = GetFocusedDisplayRate();
 
@@ -3583,7 +3583,7 @@ void FSequencer::SnapSequencerTime(FFrameTime& ScrubTime)
 		ScrubTime = UE::MovieScene::ClampToDiscreteRange(ScrubTime, PlaybackRange);
 	}
 
-	if (GetSequencerSettings()->GetIsSnapEnabled())
+	if (GetSequencerSettings()->GetForceWholeFrames())
 	{
 		FFrameRate TickResolution = GetFocusedTickResolution();
 		FFrameRate DisplayRate = GetFocusedDisplayRate();
@@ -5111,7 +5111,7 @@ FReply FSequencer::OnJumpToStart()
 FReply FSequencer::OnJumpToEnd()
 {
 	SetPlaybackStatus(EMovieScenePlayerStatus::Stepping);
-	const bool bInsetDisplayFrame = ScrubStyle == ESequencerScrubberStyle::FrameBlock && Settings->GetSnapPlayTimeToInterval() && Settings->GetIsSnapEnabled();
+	const bool bInsetDisplayFrame = ScrubStyle == ESequencerScrubberStyle::FrameBlock && Settings->GetForceWholeFrames();
 
 	FFrameRate LocalResolution = GetFocusedTickResolution();
 	FFrameRate DisplayRate = GetFocusedDisplayRate();
@@ -9679,7 +9679,7 @@ void FSequencer::Pause()
 	SetPlaybackStatus(EMovieScenePlayerStatus::Stopped);
 
 	// When stopping a sequence, we always evaluate a non-empty range if possible. This ensures accurate paused motion blur effects.
-	if (Settings->GetIsSnapEnabled())
+	if (Settings->GetForceWholeFrames())
 	{
 		FQualifiedFrameTime LocalTime          = GetUnwarpedLocalTime();
 		FFrameRate          FocusedDisplayRate = GetFocusedDisplayRate();
@@ -11397,24 +11397,12 @@ void FSequencer::BindCommands()
 		FCanExecuteAction::CreateLambda( []{ return true; } ),
 		FIsActionChecked::CreateLambda( [this]{ return Settings->GetIsSnapEnabled(); } ) );
 
-	SequencerCommandBindings->MapAction(
-		Commands.ToggleSnapKeyTimesToInterval,
-		FExecuteAction::CreateLambda( [this]{ Settings->SetSnapKeyTimesToInterval( !Settings->GetSnapKeyTimesToInterval() ); } ),
-		FCanExecuteAction::CreateLambda( []{ return true; } ),
-		FIsActionChecked::CreateLambda( [this]{ return Settings->GetSnapKeyTimesToInterval(); } ) );
-
 	SequencerCommandBindings->MapAction(
 		Commands.ToggleSnapKeyTimesToKeys,
 		FExecuteAction::CreateLambda( [this]{ Settings->SetSnapKeyTimesToKeys( !Settings->GetSnapKeyTimesToKeys() ); } ),
 		FCanExecuteAction::CreateLambda( []{ return true; } ),
 		FIsActionChecked::CreateLambda( [this]{ return Settings->GetSnapKeyTimesToKeys(); } ) );
 
-	SequencerCommandBindings->MapAction(
-		Commands.ToggleSnapSectionTimesToInterval,
-		FExecuteAction::CreateLambda( [this]{ Settings->SetSnapSectionTimesToInterval( !Settings->GetSnapSectionTimesToInterval() ); } ),
-		FCanExecuteAction::CreateLambda( []{ return true; } ),
-		FIsActionChecked::CreateLambda( [this]{ return Settings->GetSnapSectionTimesToInterval(); } ) );
-
 	SequencerCommandBindings->MapAction(
 		Commands.ToggleSnapSectionTimesToSections,
 		FExecuteAction::CreateLambda( [this]{ Settings->SetSnapSectionTimesToSections( !Settings->GetSnapSectionTimesToSections() ); } ),
@@ -11445,12 +11433,6 @@ void FSequencer::BindCommands()
 		FCanExecuteAction::CreateLambda( []{ return true; } ),
 		FIsActionChecked::CreateLambda( [this]{ return Settings->GetSnapPlayTimeToMarkers(); } ) );
 
-	SequencerCommandBindings->MapAction(
-		Commands.ToggleSnapPlayTimeToInterval,
-		FExecuteAction::CreateLambda( [this]{ Settings->SetSnapPlayTimeToInterval( !Settings->GetSnapPlayTimeToInterval() ); } ),
-		FCanExecuteAction::CreateLambda( []{ return true; } ),
-		FIsActionChecked::CreateLambda( [this]{ return Settings->GetSnapPlayTimeToInterval(); } ) );
-
 	SequencerCommandBindings->MapAction(
 		Commands.ToggleSnapPlayTimeToPressedKey,
 		FExecuteAction::CreateLambda( [this]{ Settings->SetSnapPlayTimeToPressedKey( !Settings->GetSnapPlayTimeToPressedKey() ); } ),
@@ -11469,6 +11451,12 @@ void FSequencer::BindCommands()
 		FCanExecuteAction::CreateLambda( []{ return true; } ),
 		FIsActionChecked::CreateLambda( [this]{ return Settings->GetSnapCurveValueToInterval(); } ) );
 
+	SequencerCommandBindings->MapAction(
+		Commands.ToggleForceWholeFrames,
+		FExecuteAction::CreateLambda([this] { Settings->SetForceWholeFrames(!Settings->GetForceWholeFrames()); }),
+		FCanExecuteAction::CreateLambda([] { return true; }),
+		FIsActionChecked::CreateLambda([this] { return Settings->GetForceWholeFrames(); }));
+
 	SequencerCommandBindings->MapAction(
 		Commands.ToggleShowCurveEditor,
 		FExecuteAction::CreateLambda( [this]{ SetShowCurveEditor(!GetCurveEditorIsVisible()); } ),
diff --git a/Engine/Source/Editor/Sequencer/Private/SequencerCommands.cpp b/Engine/Source/Editor/Sequencer/Private/SequencerCommands.cpp
--- a/Engine/Source/Editor/Sequencer/Private/SequencerCommands.cpp
+++ b/Engine/Source/Editor/Sequencer/Private/SequencerCommands.cpp
@@ -133,22 +133,20 @@ void FSequencerCommands::RegisterCommands()
 	UI_COMMAND( ToggleShowInfoButton, "Info Button", "Enables and disables showing the information button in the playback controls", EUserInterfaceActionType::ToggleButton, FInputChord());
 	UI_COMMAND( ToggleIsSnapEnabled, "Enable Snapping", "Enables and disables snapping", EUserInterfaceActionType::ToggleButton, FInputChord() );
 
-	UI_COMMAND( ToggleSnapKeyTimesToInterval, "Snap to the Interval", "Snap keys to the time snapping interval", EUserInterfaceActionType::ToggleButton, FInputChord() );
-	UI_COMMAND( ToggleSnapKeyTimesToKeys, "Snap to Keys and Sections", "Snap keys to other keys and sections in this sequence", EUserInterfaceActionType::ToggleButton, FInputChord() );
-
-	UI_COMMAND( ToggleSnapSectionTimesToInterval, "Snap to the Interval", "Snap sections to the time snapping interval", EUserInterfaceActionType::ToggleButton, FInputChord() );
-	UI_COMMAND( ToggleSnapSectionTimesToSections, "Snap to Keys and Sections", "Snap sections to other keys and sections in this sequence", EUserInterfaceActionType::ToggleButton, FInputChord() );
-	UI_COMMAND( ToggleSnapKeysAndSectionsToPlayRange, "Snap Keys and Sections to the Playback Range", "When checked, keys and sections will be snapped to the playback range bounds", EUserInterfaceActionType::ToggleButton, FInputChord());
+	UI_COMMAND( ToggleSnapKeyTimesToKeys, "Snap Keys to Keys and Sections", "Snap keys to other keys and sections in this sequence", EUserInterfaceActionType::ToggleButton, FInputChord() );
+	UI_COMMAND( ToggleSnapSectionTimesToSections, "Snap Sections to Keys and Sections", "Snap sections to other keys and sections in this sequence", EUserInterfaceActionType::ToggleButton, FInputChord() );
+	UI_COMMAND( ToggleSnapKeysAndSectionsToPlayRange, "Snap Keys and Sections to the Playback Range", "When checked, keys and sections will be snapped to the playback range bounds", EUserInterfaceActionType::ToggleButton, FInputChord() );
 
 	UI_COMMAND( ToggleSnapPlayTimeToKeys, "Snap to Keys While Scrubbing", "Snap the playhead to keys of the selected track while scrubbing", EUserInterfaceActionType::ToggleButton, FInputChord() );
 	UI_COMMAND( ToggleSnapPlayTimeToSections, "Snap to Sections While Scrubbing", "Snap the playhead to section bounds while scrubbing", EUserInterfaceActionType::ToggleButton, FInputChord() );
 	UI_COMMAND( ToggleSnapPlayTimeToMarkers, "Snap to Markers While Scrubbing", "Snap the playhead to markers while scrubbing", EUserInterfaceActionType::ToggleButton, FInputChord() );
-	UI_COMMAND( ToggleSnapPlayTimeToInterval, "Snap to the Interval While Scrubbing", "Snap the playhead to the time snapping interval while scrubbing", EUserInterfaceActionType::ToggleButton, FInputChord() );
 	UI_COMMAND( ToggleSnapPlayTimeToPressedKey, "Snap to the Pressed Key", "Snap the playhead to the pressed key", EUserInterfaceActionType::ToggleButton, FInputChord() );
 	UI_COMMAND( ToggleSnapPlayTimeToDraggedKey, "Snap to the Dragged Key", "Snap the playhead to the dragged key", EUserInterfaceActionType::ToggleButton, FInputChord() );
 
 	UI_COMMAND( ToggleSnapCurveValueToInterval, "Snap Curve Key Values", "Snap curve keys to the value snapping interval", EUserInterfaceActionType::ToggleButton, FInputChord() );
 
+	UI_COMMAND( ToggleForceWholeFrames, "Force Whole Frames", "Snap keys/sections/playhead to the interval which forces whole frames", EUserInterfaceActionType::ToggleButton, FInputChord() );
+
 	UI_COMMAND( FindInContentBrowser, "Find in Content Browser", "Find the viewed sequence asset in the content browser", EUserInterfaceActionType::Button, FInputChord() );
 	UI_COMMAND( BrowseToObject, "Browse to Object", "Browse to the selected section's object", EUserInterfaceActionType::Button, FInputChord(EModifierKey::Control, EKeys::B) );
 	UI_COMMAND( ToggleLayerBars, "Layer Bars", "Show/hide the layer bars to edit keyframes in bulk", EUserInterfaceActionType::ToggleButton, FInputChord() );
diff --git a/Engine/Source/Editor/Sequencer/Private/SequencerHotspots.cpp b/Engine/Source/Editor/Sequencer/Private/SequencerHotspots.cpp
--- a/Engine/Source/Editor/Sequencer/Private/SequencerHotspots.cpp
+++ b/Engine/Source/Editor/Sequencer/Private/SequencerHotspots.cpp
@@ -371,7 +371,7 @@ TSharedPtr<ISequencerEditToolDragOperation> FKeyBarHotspot::InitiateDrag(const F
 			AllLinearKeys = AllKeys.Array();
 
 			SnapField = FSequencerSnapField(*InSequencer, *this);
-			SnapField.SetSnapToInterval(InSequencer->GetSequencerSettings()->GetSnapKeyTimesToInterval());
+			SnapField.SetSnapToInterval(InSequencer->GetSequencerSettings()->GetForceWholeFrames());
 		}
 
 		void OnBeginDrag(const FPointerEvent& MouseEvent, FVector2D LocalMousePos, const FVirtualTrackArea& VirtualTrackArea) override
diff --git a/Engine/Source/Editor/Sequencer/Private/SequencerSettings.cpp b/Engine/Source/Editor/Sequencer/Private/SequencerSettings.cpp
--- a/Engine/Source/Editor/Sequencer/Private/SequencerSettings.cpp
+++ b/Engine/Source/Editor/Sequencer/Private/SequencerSettings.cpp
@@ -19,16 +19,14 @@ USequencerSettings::USequencerSettings( const FObjectInitializer& ObjectInitiali
 	bCreateSpawnableCameras = true;
 	bShowRangeSlider = false;
 	bIsSnapEnabled = true;
-	bSnapKeyTimesToInterval = true;
 	bSnapKeyTimesToKeys = true;
-	bSnapSectionTimesToInterval = true;
 	bSnapSectionTimesToSections = true;
 	bSnapPlayTimeToKeys = false;
 	bSnapPlayTimeToSections = false;
 	bSnapPlayTimeToMarkers = false;
-	bSnapPlayTimeToInterval = true;
 	bSnapPlayTimeToPressedKey = true;
 	bSnapPlayTimeToDraggedKey = true;
+	bForceWholeFrames = true;
 	CurveValueSnapInterval = 0.1f;
 	GridSpacing = TOptional<float>();
 	bSnapCurveValueToInterval = false;
@@ -205,20 +203,6 @@ void USequencerSettings::SetIsSnapEnabled(bool InbIsSnapEnabled)
 	}
 }
 
-bool USequencerSettings::GetSnapKeyTimesToInterval() const
-{
-	return bSnapKeyTimesToInterval;
-}
-
-void USequencerSettings::SetSnapKeyTimesToInterval(bool InbSnapKeyTimesToInterval)
-{
-	if ( bSnapKeyTimesToInterval != InbSnapKeyTimesToInterval )
-	{
-		bSnapKeyTimesToInterval = InbSnapKeyTimesToInterval;
-		SaveConfig();
-	}
-}
-
 bool USequencerSettings::GetSnapKeyTimesToKeys() const
 {
 	return bSnapKeyTimesToKeys;
@@ -246,20 +230,6 @@ void USequencerSettings::SetLeftMouseDragDoesMarque(bool bDoMarque)
 	}
 }
 
-bool USequencerSettings::GetSnapSectionTimesToInterval() const
-{
-	return bSnapSectionTimesToInterval;
-}
-
-void USequencerSettings::SetSnapSectionTimesToInterval(bool InbSnapSectionTimesToInterval)
-{
-	if ( bSnapSectionTimesToInterval != InbSnapSectionTimesToInterval )
-	{
-		bSnapSectionTimesToInterval = InbSnapSectionTimesToInterval;
-		SaveConfig();
-	}
-}
-
 bool USequencerSettings::GetSnapSectionTimesToSections() const
 {
 	return bSnapSectionTimesToSections;
@@ -330,20 +300,6 @@ void USequencerSettings::SetSnapPlayTimeToMarkers(bool InbSnapPlayTimeToMarkers)
 	}
 }
 
-bool USequencerSettings::GetSnapPlayTimeToInterval() const
-{
-	return bSnapPlayTimeToInterval;
-}
-
-void USequencerSettings::SetSnapPlayTimeToInterval(bool InbSnapPlayTimeToInterval)
-{
-	if ( bSnapPlayTimeToInterval != InbSnapPlayTimeToInterval )
-	{
-		bSnapPlayTimeToInterval = InbSnapPlayTimeToInterval;
-		SaveConfig();
-	}
-}
-
 bool USequencerSettings::GetSnapPlayTimeToPressedKey() const
 {
 	return bSnapPlayTimeToPressedKey;
@@ -414,6 +370,20 @@ void USequencerSettings::SetSnapCurveValueToInterval( bool InbSnapCurveValueToIn
 	}
 }
 
+bool USequencerSettings::GetForceWholeFrames() const
+{
+	return bForceWholeFrames;
+}
+
+void USequencerSettings::SetForceWholeFrames(bool InbForceWholeFrames)
+{
+	if (bForceWholeFrames != InbForceWholeFrames)
+	{
+		bForceWholeFrames = InbForceWholeFrames;
+		SaveConfig();
+	}
+}
+
 bool USequencerSettings::GetShowSelectedNodesOnly() const
 {
 	return bShowSelectedNodesOnly;
diff --git a/Engine/Source/Editor/Sequencer/Private/SequencerTimeSliderController.cpp b/Engine/Source/Editor/Sequencer/Private/SequencerTimeSliderController.cpp
--- a/Engine/Source/Editor/Sequencer/Private/SequencerTimeSliderController.cpp
+++ b/Engine/Source/Editor/Sequencer/Private/SequencerTimeSliderController.cpp
@@ -109,6 +109,31 @@ FSequencerTimeSliderController::~FSequencerTimeSliderController()
 	}
 }
 
+FFrameTime FSequencerTimeSliderController::SnapSequencerTime(FFrameTime InTime) const
+{
+	TSharedPtr<FSequencer> Sequencer = WeakSequencer.Pin();
+	if (!Sequencer.IsValid())
+	{
+		return InTime;
+	}
+
+	if (Sequencer->GetSequencerSettings()->GetForceWholeFrames())
+	{
+		if (Sequencer->GetScrubStyle() == ESequencerScrubberStyle::FrameBlock)
+		{
+			// Floor to the display frame
+			InTime = ConvertFrameTime(ConvertFrameTime(InTime, GetTickResolution(), GetDisplayRate()).FloorToFrame(), GetDisplayRate(), GetTickResolution());
+		}
+		else
+		{
+			// Snap (round) to display rate
+			InTime = FFrameRate::Snap(InTime, GetTickResolution(), GetDisplayRate());
+		}
+	}
+
+	return InTime;
+}
+
 FFrameTime FSequencerTimeSliderController::ComputeScrubTimeFromMouse(const FGeometry& Geometry, const FPointerEvent& MouseEvent, FScrubRangeToScreen RangeToScreen) const
 {
 	FVector2D           ScreenSpacePosition = MouseEvent.GetScreenSpacePosition();
@@ -122,23 +147,10 @@ FFrameTime FSequencerTimeSliderController::ComputeScrubTimeFromMouse(const FGeom
 		return ScrubTime;
 	}
 	
-	if ( Sequencer->GetSequencerSettings()->GetIsSnapEnabled() || MouseEvent.IsShiftDown() )
-	{
-		if (Sequencer->GetSequencerSettings()->GetSnapPlayTimeToInterval())
-		{
-			// Set the style of the scrub handle
-			if (Sequencer->GetScrubStyle() == ESequencerScrubberStyle::FrameBlock)
-			{
-				// Floor to the display frame
-				ScrubTime = ConvertFrameTime(ConvertFrameTime(ScrubTime, GetTickResolution(), GetDisplayRate()).FloorToFrame(), GetDisplayRate(), GetTickResolution());
-			}
-			else
-			{
-				// Snap (round) to display rate
-				ScrubTime = FFrameRate::Snap(ScrubTime, GetTickResolution(), GetDisplayRate());
-			}
-		}
+	ScrubTime = SnapSequencerTime(ScrubTime);
 
+	if (Sequencer->GetSequencerSettings()->GetIsSnapEnabled() || MouseEvent.IsShiftDown())
+	{
 		// SnapTimeToNearestKey will return ScrubTime unmodified if there is no key within range.
 		ScrubTime = SnapTimeToNearestKey(MouseEvent, RangeToScreen, CursorPos.X, ScrubTime);
 	}
@@ -157,7 +169,7 @@ FFrameTime FSequencerTimeSliderController::ComputeFrameTimeFromMouse(const FGeom
 		return MouseValue * GetTickResolution();
 	}
 
-	if (CheckSnapping && Sequencer->GetSequencerSettings()->GetIsSnapEnabled())
+	if (CheckSnapping && Sequencer->GetSequencerSettings()->GetForceWholeFrames())
 	{
 		FFrameNumber        SnappedFrameNumber = (MouseValue * GetDisplayRate()).FloorToFrame();
 		FQualifiedFrameTime RoundedPlayFrame   = FQualifiedFrameTime(SnappedFrameNumber, GetDisplayRate());
@@ -1390,7 +1402,7 @@ void FSequencerTimeSliderController::CommitScrubPosition( FFrameTime NewValue, b
 		FFrameTime LowerBound = (ViewRange.GetLowerBoundValue() * TickResolution).CeilToFrame();
 		FFrameTime UpperBound = (ViewRange.GetUpperBoundValue() * TickResolution).FloorToFrame();
 
-		if (Sequencer->GetSequencerSettings()->GetIsSnapEnabled() && Sequencer->GetSequencerSettings()->GetSnapPlayTimeToInterval())
+		if (Sequencer->GetSequencerSettings()->GetForceWholeFrames())
 		{
 			LowerBound = FFrameRate::Snap(LowerBound, TickResolution, DisplayRate);
 			UpperBound = FFrameRate::Snap(UpperBound, TickResolution, DisplayRate);
@@ -2356,6 +2368,8 @@ void FSequencerTimeSliderController::SetPlaybackRangeStart(FFrameNumber NewStart
 {
 	TRange<FFrameNumber> PlaybackRange = TimeSliderArgs.PlaybackRange.Get();
 
+	NewStart = SnapSequencerTime(NewStart).FrameNumber;
+
 	if (NewStart <= UE::MovieScene::DiscreteExclusiveUpper(PlaybackRange))
 	{
 		TimeSliderArgs.OnPlaybackRangeChanged.ExecuteIfBound(TRange<FFrameNumber>(NewStart, PlaybackRange.GetUpperBound()));
@@ -2366,6 +2380,8 @@ void FSequencerTimeSliderController::SetPlaybackRangeEnd(FFrameNumber NewEnd)
 {
 	TRange<FFrameNumber> PlaybackRange = TimeSliderArgs.PlaybackRange.Get();
 
+	NewEnd = SnapSequencerTime(NewEnd).FrameNumber;
+
 	if (NewEnd >= UE::MovieScene::DiscreteInclusiveLower(PlaybackRange))
 	{
 		TimeSliderArgs.OnPlaybackRangeChanged.ExecuteIfBound(TRange<FFrameNumber>(PlaybackRange.GetLowerBound(), TRangeBound<FFrameNumber>::Exclusive(NewEnd)));
@@ -2376,6 +2392,8 @@ void FSequencerTimeSliderController::SetSelectionRangeStart(FFrameNumber NewStar
 {
 	TRange<FFrameNumber> SelectionRange = TimeSliderArgs.SelectionRange.Get();
 
+	NewStart = SnapSequencerTime(NewStart).FrameNumber;
+
 	if (SelectionRange.IsEmpty())
 	{
 		TimeSliderArgs.OnSelectionRangeChanged.ExecuteIfBound(TRange<FFrameNumber>(NewStart, NewStart + 1));
@@ -2390,6 +2408,8 @@ void FSequencerTimeSliderController::SetSelectionRangeEnd(FFrameNumber NewEnd)
 {
 	TRange<FFrameNumber> SelectionRange = TimeSliderArgs.SelectionRange.Get();
 
+	NewEnd = SnapSequencerTime(NewEnd).FrameNumber;
+
 	if (SelectionRange.IsEmpty())
 	{
 		TimeSliderArgs.OnSelectionRangeChanged.ExecuteIfBound(TRange<FFrameNumber>(NewEnd - 1, NewEnd));
@@ -2439,7 +2459,7 @@ void FSequencerTimeSliderController::SetMark(FFrameNumber DiffFrame)
 	for (TMap<int32, FFrameNumber>::TConstIterator It = DragMarkMap.CreateConstIterator(); It; ++It)
 	{
 		int32 MarkIndex = It.Key();
-		FFrameNumber FrameNumber = It.Value() + DiffFrame;
+		FFrameNumber FrameNumber = SnapSequencerTime(It.Value() + DiffFrame).FrameNumber;
 		TimeSliderArgs.OnSetMarkedFrame.ExecuteIfBound(MarkIndex, FrameNumber);
 	}
 }
diff --git a/Engine/Source/Editor/Sequencer/Private/Tools/DragOperation_Stretch.cpp b/Engine/Source/Editor/Sequencer/Private/Tools/DragOperation_Stretch.cpp
--- a/Engine/Source/Editor/Sequencer/Private/Tools/DragOperation_Stretch.cpp
+++ b/Engine/Source/Editor/Sequencer/Private/Tools/DragOperation_Stretch.cpp
@@ -127,6 +127,11 @@ void FEditToolDragOperation_Stretch::OnDrag(const FPointerEvent& MouseEvent, FVe
 		}
 	}
 
+	if (Sequencer->GetSequencerSettings()->GetForceWholeFrames())
+	{
+		CurrentTime = FFrameRate::Snap(CurrentTime, Sequencer->GetFocusedTickResolution(), Sequencer->GetFocusedDisplayRate());
+	}
+
 	FStretchScreenParameters ScreenParameters;
 	ScreenParameters.DragStartPosition   = DragStartPosition;
 	ScreenParameters.CurrentDragPosition = CurrentTime;
@@ -168,6 +173,11 @@ void FEditToolDragOperation_Stretch::OnEndDrag( const FPointerEvent& MouseEvent,
 		}
 	}
 
+	if (Sequencer->GetSequencerSettings()->GetForceWholeFrames())
+	{
+		CurrentTime = FFrameRate::Snap(CurrentTime, Sequencer->GetFocusedTickResolution(), Sequencer->GetFocusedDisplayRate());
+	}
+
 	FStretchScreenParameters ScreenParameters;
 	ScreenParameters.DragStartPosition   = DragStartPosition;
 	ScreenParameters.CurrentDragPosition = VirtualTrackArea.PixelToFrame(LocalMousePos.X);
diff --git a/Engine/Source/Editor/Sequencer/Private/Tools/EditToolDragOperations.cpp b/Engine/Source/Editor/Sequencer/Private/Tools/EditToolDragOperations.cpp
--- a/Engine/Source/Editor/Sequencer/Private/Tools/EditToolDragOperations.cpp
+++ b/Engine/Source/Editor/Sequencer/Private/Tools/EditToolDragOperations.cpp
@@ -178,7 +178,7 @@ void FResizeSection::OnBeginDrag(const FPointerEvent& MouseEvent, FVector2D Loca
 	TSet<FSequencerSelectedKey> EmptyKeySet;
 	FInvalidKeyAndSectionSnappingCandidates SnapCandidates(EmptyKeySet, Sections);
 	SnapField = FSequencerSnapField(Sequencer, SnapCandidates, ESequencerEntity::Section | ESequencerEntity::Key);
-	SnapField.GetValue().SetSnapToInterval(Sequencer.GetSequencerSettings()->GetSnapSectionTimesToInterval());
+	SnapField.GetValue().SetSnapToInterval(Sequencer.GetSequencerSettings()->GetForceWholeFrames());
 	SnapField.GetValue().SetSnapToLikeTypes(Sequencer.GetSequencerSettings()->GetSnapSectionTimesToSections());
 
 	SectionInitTimes.Empty();
@@ -293,7 +293,7 @@ void FResizeSection::OnDrag(const FPointerEvent& MouseEvent, FVector2D LocalMous
 	FFrameNumber DeltaTime = (VirtualTrackArea.PixelToFrame(LocalMousePos.X) - MouseDownTime).RoundToFrame();
 
 	// Snapping
-	if ( Settings->GetIsSnapEnabled() )
+	if ((Settings->GetIsSnapEnabled() || MouseEvent.IsShiftDown()) && Settings->GetSnapSectionTimesToSections())
 	{
 		TArray<FFrameTime> SectionTimes;
 		for (UMovieSceneSection* Section : Sections)
@@ -304,12 +304,7 @@ void FResizeSection::OnDrag(const FPointerEvent& MouseEvent, FVector2D LocalMous
 		float SnapThresholdPx = VirtualTrackArea.PixelToSeconds(PixelSnapWidth) - VirtualTrackArea.PixelToSeconds(0.f);
 		int32 SnapThreshold   = ( SnapThresholdPx * TickResolution ).FloorToFrame().Value;
 
-		TOptional<FSequencerSnapField::FSnapResult> SnappedTime;
-
-		if (Settings->GetSnapSectionTimesToSections())
-		{
-			SnappedTime = SnapField->Snap(SectionTimes, SnapThreshold);
-		}
+		TOptional<FSequencerSnapField::FSnapResult> SnappedTime = SnapField->Snap(SectionTimes, SnapThreshold);
 
 		if (SnappedTime.IsSet())
 		{
@@ -318,8 +313,7 @@ void FResizeSection::OnDrag(const FPointerEvent& MouseEvent, FVector2D LocalMous
 		}
 	}
 
-
-	auto GetMovementMaximums = [this](const UMovieSceneSection* Section, TOptional<FFrameNumber>& LeftMovementMaximum, TOptional<FFrameNumber>& RightMovementMaximum)
+	auto GetMovementMaximums = [this](const FPointerEvent& MouseEvent, const UMovieSceneSection* Section, TOptional<FFrameNumber>& LeftMovementMaximum, TOptional<FFrameNumber>& RightMovementMaximum)
 	{
 		// We'll calculate this section's borders and clamp the possible delta time to be less than that
 
@@ -330,7 +324,7 @@ void FResizeSection::OnDrag(const FPointerEvent& MouseEvent, FVector2D LocalMous
 			RightMovementMaximum = UE::MovieScene::DiscreteExclusiveUpper(SectionBoundaries);
 		}
 
-		if (Settings->GetIsSnapEnabled() && Settings->GetSnapKeysAndSectionsToPlayRange() && !Settings->ShouldKeepPlayRangeInSectionBounds())
+		if ((Settings->GetIsSnapEnabled() || MouseEvent.IsShiftDown()) && Settings->GetSnapKeysAndSectionsToPlayRange() && !Settings->ShouldKeepPlayRangeInSectionBounds())
 		{
 			if (!LeftMovementMaximum.IsSet() || LeftMovementMaximum.GetValue() < Sequencer.GetPlaybackRange().GetLowerBoundValue())
 			{
@@ -355,7 +349,7 @@ void FResizeSection::OnDrag(const FPointerEvent& MouseEvent, FVector2D LocalMous
 
 			TOptional<FFrameNumber> LeftMovementMaximum;
 			TOptional<FFrameNumber> RightMovementMaximum;
-			GetMovementMaximums(Data.MovieSection, LeftMovementMaximum, RightMovementMaximum);
+			GetMovementMaximums(MouseEvent, Data.MovieSection, LeftMovementMaximum, RightMovementMaximum);
 
 			FFrameNumber StartPosition  = bDraggingByEnd ? UE::MovieScene::DiscreteExclusiveUpper(Data.InitialRange) : UE::MovieScene::DiscreteInclusiveLower(Data.InitialRange);
 
@@ -453,7 +447,7 @@ void FResizeSection::OnDrag(const FPointerEvent& MouseEvent, FVector2D LocalMous
 
 		TOptional<FFrameNumber> LeftMovementMaximum;
 		TOptional<FFrameNumber> RightMovementMaximum;
-		GetMovementMaximums(Section, LeftMovementMaximum, RightMovementMaximum);
+		GetMovementMaximums(MouseEvent, Section, LeftMovementMaximum, RightMovementMaximum);
 
 		TSharedPtr<ISequencerSection> SectionInterface = SectionHandle->GetSectionInterface();
 
@@ -463,7 +457,7 @@ void FResizeSection::OnDrag(const FPointerEvent& MouseEvent, FVector2D LocalMous
 		{
 			FFrameNumber MinFrame = Section->HasStartFrame() ? Section->GetInclusiveStartFrame() : TNumericLimits<int32>::Lowest();
 
-			if ( Settings->GetIsSnapEnabled() && Settings->GetSnapSectionTimesToInterval())
+			if (Settings->GetForceWholeFrames())
 			{
 				int32 IntervalSnapThreshold = FMath::RoundToInt( ( TickResolution / DisplayRate ).AsDecimal() );
 				MinFrame = MinFrame + IntervalSnapThreshold;
@@ -487,7 +481,7 @@ void FResizeSection::OnDrag(const FPointerEvent& MouseEvent, FVector2D LocalMous
 		{
 			FFrameNumber MaxFrame = Section->HasEndFrame() ? Section->GetExclusiveEndFrame()-1 : TNumericLimits<int32>::Max();
 
-			if ( Settings->GetIsSnapEnabled() && Settings->GetSnapSectionTimesToInterval())
+			if (Settings->GetForceWholeFrames())
 			{
 				int32 IntervalSnapThreshold = FMath::RoundToInt( ( TickResolution / DisplayRate ).AsDecimal() );
 				MaxFrame = MaxFrame - IntervalSnapThreshold;
@@ -615,7 +609,7 @@ void FManipulateSectionEasing::OnBeginDrag(const FPointerEvent& MouseEvent, FVec
 		// Construct a snap field of all section bounds
 		ISnapCandidate SnapCandidates;
 		SnapField = FSequencerSnapField(Sequencer, SnapCandidates, ESequencerEntity::Section);
-		SnapField.GetValue().SetSnapToInterval(Sequencer.GetSequencerSettings()->GetSnapSectionTimesToInterval());
+		SnapField.GetValue().SetSnapToInterval(Sequencer.GetSequencerSettings()->GetForceWholeFrames());
 	}
 
 	InitValue = bEaseIn ? Section->Easing.GetEaseInDuration() : Section->Easing.GetEaseOutDuration();
@@ -640,7 +634,7 @@ void FManipulateSectionEasing::OnDrag(const FPointerEvent& MouseEvent, FVector2D
 	}
 
 	// Snapping
-	if (Settings->GetIsSnapEnabled())
+	if ((Settings->GetIsSnapEnabled() || MouseEvent.IsShiftDown()) && Settings->GetSnapSectionTimesToSections())
 	{
 		TArray<FFrameTime> SnapTimes;
 		if (bEaseIn)
@@ -657,12 +651,7 @@ void FManipulateSectionEasing::OnDrag(const FPointerEvent& MouseEvent, FVector2D
 		float SnapThresholdPx = VirtualTrackArea.PixelToSeconds(PixelSnapWidth) - VirtualTrackArea.PixelToSeconds(0.f);
 		int32 SnapThreshold   = ( SnapThresholdPx * TickResolution ).FloorToFrame().Value;
 
-		TOptional<FSequencerSnapField::FSnapResult> SnappedTime;
-
-		if (Settings->GetSnapSectionTimesToSections())
-		{
-			SnappedTime = SnapField->Snap(SnapTimes, SnapThreshold);
-		}
+		TOptional<FSequencerSnapField::FSnapResult> SnappedTime = SnapField->Snap(SnapTimes, SnapThreshold);
 
 		if (SnappedTime.IsSet())
 		{
@@ -844,56 +833,53 @@ void FMoveKeysAndSections::OnDrag(const FPointerEvent& MouseEvent, FVector2D Loc
 	FFrameTime MouseTime = VirtualTrackArea.PixelToFrame(LocalMousePos.X);
 
 	// Calculate snapping first which modifies our MouseTime to reflect where it would have to be for the closest snap to work.
-	if (Settings->GetIsSnapEnabled())
-	{
-		FFrameTime SnapThreshold = VirtualTrackArea.PixelDeltaToFrame(PixelSnapWidth);
+	FFrameTime SnapThreshold = VirtualTrackArea.PixelDeltaToFrame(PixelSnapWidth);
 
-		// The edge of each bounded section as well as each individual key is a valid marker to try and snap to intervals/sections/etc.
-		// We take our stored offsets and add them to our current time to figure out where on the timeline the are currently.
-		TArray<FFrameTime> ValidSnapMarkers;
+	// The edge of each bounded section as well as each individual key is a valid marker to try and snap to intervals/sections/etc.
+	// We take our stored offsets and add them to our current time to figure out where on the timeline the are currently.
+	TArray<FFrameTime> ValidSnapMarkers;
 
-		// If they have both keys and settings selected then we snap to the interval if either one of them is enabled, otherwise respect the individual setting.
-		const bool bSnapToInterval = (KeysAsArray.Num() > 0 && Settings->GetSnapKeyTimesToInterval()) || (Sections.Num() > 0 && Settings->GetSnapSectionTimesToInterval());
-		const bool bSnapToLikeTypes = (KeysAsArray.Num() > 0 && Settings->GetSnapKeyTimesToKeys()) || (Sections.Num() > 0 && Settings->GetSnapSectionTimesToSections());
+	// If they have both keys and settings selected then we snap to the interval if either one of them is enabled, otherwise respect the individual setting.
+	const bool bSnapToInterval = Settings->GetForceWholeFrames() && (KeysAsArray.Num() > 0 || Sections.Num());
+	const bool bSnapToLikeTypes = (Settings->GetIsSnapEnabled() || MouseEvent.IsShiftDown()) && ( (KeysAsArray.Num() > 0 && Settings->GetSnapKeyTimesToKeys()) || (Sections.Num() > 0 && Settings->GetSnapSectionTimesToSections()) );
 
-		SnapField.GetValue().SetSnapToInterval(bSnapToInterval);
-		SnapField.GetValue().SetSnapToLikeTypes(bSnapToLikeTypes);
+	SnapField.GetValue().SetSnapToInterval(bSnapToInterval);
+	SnapField.GetValue().SetSnapToLikeTypes(bSnapToLikeTypes);
 
-		// RelativeSnapOffsets contains both our sections and our keys, and we add them all as potential things that can snap to stuff.
-		if (bSnapToLikeTypes || bSnapToInterval)
+	// RelativeSnapOffsets contains both our sections and our keys, and we add them all as potential things that can snap to stuff.
+	if (bSnapToLikeTypes || bSnapToInterval)
+	{
+		ValidSnapMarkers.SetNumUninitialized(RelativeSnapOffsets.Num());
+		for (int32 Index = 0; Index < RelativeSnapOffsets.Num(); ++Index)
 		{
-			ValidSnapMarkers.SetNumUninitialized(RelativeSnapOffsets.Num());
-			for (int32 Index = 0; Index < RelativeSnapOffsets.Num(); ++Index)
-			{
-				ValidSnapMarkers[Index] = (RelativeSnapOffsets[Index] + MouseTime);
-			}
+			ValidSnapMarkers[Index] = (RelativeSnapOffsets[Index] + MouseTime);
 		}
+	}
 
-		// Now we'll try and snap all of these points to the closest valid snap marker (which may be a section or interval)
-		TOptional<FSequencerSnapField::FSnapResult> SnappedTime;
+	// Now we'll try and snap all of these points to the closest valid snap marker (which may be a section or interval)
+	TOptional<FSequencerSnapField::FSnapResult> SnappedTime;
 
-		if (bSnapToLikeTypes || bSnapToInterval)
-		{
-			// This may or may not set the SnappedTime depending on if there are any sections within the threshold.
-			SnappedTime = SnapField->Snap(ValidSnapMarkers, SnapThreshold);
-		}
+	if (bSnapToLikeTypes || bSnapToInterval)
+	{
+		// This may or may not set the SnappedTime depending on if there are any sections within the threshold.
+		SnappedTime = SnapField->Snap(ValidSnapMarkers, SnapThreshold);
+	}
 
-		// If they actually snapped to something (snapping may be on but settings might dictate nothing to snap to) add the difference
-		// to our current MouseTime so that MouseTime reflects the amount needed to move to get to the whole snap point.
-		if (SnappedTime.IsSet())
-		{
-			// Add the snapped amount onto the mouse time so the resulting delta brings us in alignment.
-			MouseTime += (SnappedTime->SnappedTime - SnappedTime->OriginalTime);
-		}
+	// If they actually snapped to something (snapping may be on but settings might dictate nothing to snap to) add the difference
+	// to our current MouseTime so that MouseTime reflects the amount needed to move to get to the whole snap point.
+	if (SnappedTime.IsSet())
+	{
+		// Add the snapped amount onto the mouse time so the resulting delta brings us in alignment.
+		MouseTime += (SnappedTime->SnappedTime - SnappedTime->OriginalTime);
 	}
 
-	if (Settings->GetIsSnapEnabled() && Settings->GetSnapKeysAndSectionsToPlayRange() && !Settings->ShouldKeepPlayRangeInSectionBounds())
+	if ((Settings->GetIsSnapEnabled() || MouseEvent.IsShiftDown()) && Settings->GetSnapKeysAndSectionsToPlayRange() && !Settings->ShouldKeepPlayRangeInSectionBounds())
 	{
 		MouseTime = UE::MovieScene::ClampToDiscreteRange(MouseTime, Sequencer.GetPlaybackRange());
 	}
 
 	// We'll calculate a DeltaX based on limits on movement (snapping, section collision) and then use them on keys and sections below.
-	TOptional<FFrameNumber> MaxDeltaX = GetMovementDeltaX(MouseTime);
+	TOptional<FFrameNumber> MaxDeltaX = GetMovementDeltaX(MouseEvent, MouseTime);
 
 	FFrameNumber MouseDeltaTime = (MouseTime - MouseTimePrev).FloorToFrame();
 	MouseTimePrev = MouseTimePrev + MaxDeltaX.Get(MouseDeltaTime);
@@ -1020,7 +1006,7 @@ void FMoveKeysAndSections::ModifyNonSelectedSections()
 	}
 }
 
-TOptional<FFrameNumber> FMoveKeysAndSections::GetMovementDeltaX(FFrameTime MouseTime)
+TOptional<FFrameNumber> FMoveKeysAndSections::GetMovementDeltaX(const FPointerEvent& MouseEvent, FFrameTime MouseTime)
 {
 	TOptional<FFrameNumber> DeltaX;
 
@@ -1048,7 +1034,7 @@ TOptional<FFrameNumber> FMoveKeysAndSections::GetMovementDeltaX(FFrameTime Mouse
 			RightMovementMaximum = UE::MovieScene::DiscreteExclusiveUpper(SectionBoundaries);
 		}
 		
-		if (Settings->GetIsSnapEnabled() && Settings->GetSnapKeysAndSectionsToPlayRange() && !Settings->ShouldKeepPlayRangeInSectionBounds())
+		if ((Settings->GetIsSnapEnabled() || MouseEvent.IsShiftDown()) && Settings->GetSnapKeysAndSectionsToPlayRange() && !Settings->ShouldKeepPlayRangeInSectionBounds())
 		{
 			if (!LeftMovementMaximum.IsSet() || LeftMovementMaximum.GetValue() < Sequencer.GetPlaybackRange().GetLowerBoundValue())
 			{
@@ -1094,7 +1080,7 @@ TOptional<FFrameNumber> FMoveKeysAndSections::GetMovementDeltaX(FFrameTime Mouse
 		}
 	}
 
-	if (Settings->GetIsSnapEnabled() && Settings->GetSnapKeysAndSectionsToPlayRange() && !Settings->ShouldKeepPlayRangeInSectionBounds())
+	if ((Settings->GetIsSnapEnabled() || MouseEvent.IsShiftDown()) && Settings->GetSnapKeysAndSectionsToPlayRange() && !Settings->ShouldKeepPlayRangeInSectionBounds())
 	{
 		TArray<FFrameNumber> CurrentKeyTimes;
 		CurrentKeyTimes.SetNum(KeysAsArray.Num());
diff --git a/Engine/Source/Editor/Sequencer/Private/Tools/EditToolDragOperations.h b/Engine/Source/Editor/Sequencer/Private/Tools/EditToolDragOperations.h
--- a/Engine/Source/Editor/Sequencer/Private/Tools/EditToolDragOperations.h
+++ b/Engine/Source/Editor/Sequencer/Private/Tools/EditToolDragOperations.h
@@ -178,7 +178,7 @@ public:
 
 protected:
 	/** Calculate the possible horizontal movement we can, constrained by sections running into things. */
-	TOptional<FFrameNumber> GetMovementDeltaX(FFrameTime MouseTime);
+	TOptional<FFrameNumber> GetMovementDeltaX(const FPointerEvent& MouseEvent, FFrameTime MouseTime);
 
 	/** Move selected sections, if any. */
 	bool HandleSectionMovement(FFrameTime MouseTime, FVector2D VirtualMousePos, FVector2D LocalMousePos, TOptional<FFrameNumber> MaxDeltaX, FFrameNumber DesiredDeltaX);
diff --git a/Engine/Source/Editor/Sequencer/Public/SequencerCommands.h b/Engine/Source/Editor/Sequencer/Public/SequencerCommands.h
--- a/Engine/Source/Editor/Sequencer/Public/SequencerCommands.h
+++ b/Engine/Source/Editor/Sequencer/Public/SequencerCommands.h
@@ -340,15 +340,9 @@ public:
 	/** Turns snapping on and off. */
 	TSharedPtr< FUICommandInfo > ToggleIsSnapEnabled;
 
-	/** Toggles whether or not keys should snap to the selected interval. */
-	TSharedPtr< FUICommandInfo > ToggleSnapKeyTimesToInterval;
-
 	/** Toggles whether or not keys should snap to other keys in the section. */
 	TSharedPtr< FUICommandInfo > ToggleSnapKeyTimesToKeys;
 
-	/** Toggles whether or not sections should snap to the selected interval. */
-	TSharedPtr< FUICommandInfo > ToggleSnapSectionTimesToInterval;
-
 	/** Toggles whether or not sections should snap to other sections. */
 	TSharedPtr< FUICommandInfo > ToggleSnapSectionTimesToSections;
 
@@ -364,9 +358,6 @@ public:
 	/** Toggles whether or not snap to markers while scrubbing. */
 	TSharedPtr< FUICommandInfo > ToggleSnapPlayTimeToMarkers;
 
-	/** Toggles whether or not the play time should snap to the selected interval. */
-	TSharedPtr< FUICommandInfo > ToggleSnapPlayTimeToInterval;
-
 	/** Toggles whether or not the play time should snap to the pressed key. */
 	TSharedPtr< FUICommandInfo > ToggleSnapPlayTimeToPressedKey;
 
@@ -376,6 +367,9 @@ public:
 	/** Toggles whether or not to snap curve values to the interval. */
 	TSharedPtr< FUICommandInfo > ToggleSnapCurveValueToInterval;
 
+	/** Toggles whether or not to snap keys/sections/playhead to the interval which forces whole frames. */
+	TSharedPtr< FUICommandInfo > ToggleForceWholeFrames;
+
 	/** Finds the viewed sequence asset in the content browser. */
 	TSharedPtr< FUICommandInfo > FindInContentBrowser;
 
diff --git a/Engine/Source/Editor/Sequencer/Public/SequencerSettings.h b/Engine/Source/Editor/Sequencer/Public/SequencerSettings.h
--- a/Engine/Source/Editor/Sequencer/Public/SequencerSettings.h
+++ b/Engine/Source/Editor/Sequencer/Public/SequencerSettings.h
@@ -183,21 +183,11 @@ public:
 	/** Sets whether or not snapping is enabled. */
 	void SetIsSnapEnabled(bool InbIsSnapEnabled);
 
-	/** Gets whether or not to snap key times to the interval. */
-	bool GetSnapKeyTimesToInterval() const;
-	/** Sets whether or not to snap keys to the interval. */
-	void SetSnapKeyTimesToInterval(bool InbSnapKeyTimesToInterval);
-
 	/** Gets whether or not to snap keys to other keys. */
 	bool GetSnapKeyTimesToKeys() const;
 	/** Sets whether or not to snap keys to other keys. */
 	void SetSnapKeyTimesToKeys(bool InbSnapKeyTimesToKeys);
 
-	/** Gets whether or not to snap sections to the interval. */
-	bool GetSnapSectionTimesToInterval() const;
-	/** Sets whether or not to snap sections to the interval. */
-	void SetSnapSectionTimesToInterval(bool InbSnapSectionTimesToInterval);
-
 	/** Gets whether or not to snap sections to other sections. */
 	bool GetSnapSectionTimesToSections() const;
 	/** sets whether or not to snap sections to other sections. */
@@ -223,11 +213,6 @@ public:
 	/** Sets whether or not to snap the play time to markers while scrubbing. */
 	void SetSnapPlayTimeToMarkers(bool InbSnapPlayTimeToMarkers);
 
-	/** Gets whether or not to snap the play time to the interval while scrubbing. */
-	bool GetSnapPlayTimeToInterval() const;
-	/** Sets whether or not to snap the play time to the interval while scrubbing. */
-	void SetSnapPlayTimeToInterval(bool InbSnapPlayTimeToInterval);
-
 	/** Gets whether or not to snap the play time to the pressed key. */
 	bool GetSnapPlayTimeToPressedKey() const;
 	/** Sets whether or not to snap the play time to the pressed key. */
@@ -243,6 +228,11 @@ public:
 	/** Sets the snapping interval for curve values. */
 	void SetCurveValueSnapInterval(float InCurveValueSnapInterval);
 
+	/** Gets whether or not to snap keys/sections/playhead to the interval which forces whole frames */
+	bool GetForceWholeFrames() const;
+	/** Sets whether or not to snap keys/sections/playhead to the interval which forces whole frames */
+	void SetForceWholeFrames(bool bInForceWholeFrames);
+
 	/** Gets the state for spacing between grid lines */
 	TOptional<float> GetGridSpacing() const;
 	/** Sets the grid line spacing state */
@@ -591,18 +581,10 @@ protected:
 	UPROPERTY( config, EditAnywhere, Category=Snapping )
 	bool bIsSnapEnabled;
 
-	/** Enable or disable snapping keys to the time snapping interval. */
-	UPROPERTY( config, EditAnywhere, Category=Snapping )
-	bool bSnapKeyTimesToInterval;
-
 	/** Enable or disable snapping keys to other keys. */
 	UPROPERTY( config, EditAnywhere, Category=Snapping )
 	bool bSnapKeyTimesToKeys;
 	
-	/** Enable or disable snapping sections to the time snapping interval. */
-	UPROPERTY( config, EditAnywhere, Category=Snapping )
-	bool bSnapSectionTimesToInterval;
-
 	/** Enable or disable snapping sections to other sections. */
 	UPROPERTY( config, EditAnywhere, Category=Snapping )
 	bool bSnapSectionTimesToSections;
@@ -623,10 +605,6 @@ protected:
 	UPROPERTY( config, EditAnywhere, Category=Snapping, meta = (DisplayName = "Snap Playhead to Markers"))
 	bool bSnapPlayTimeToMarkers;
 
-	/** Enable or disable snapping the playhead to the time snapping interval while scrubbing. */
-	UPROPERTY( config, EditAnywhere, Category=Snapping, meta = (DisplayName = "Snap Playhead to Interval"))
-	bool bSnapPlayTimeToInterval;
-
 	/** Enable or disable snapping the playhead to the pressed key. */
 	UPROPERTY( config, EditAnywhere, Category=Snapping, meta = (DisplayName = "Snap Playhead to Pressed Key"))
 	bool bSnapPlayTimeToPressedKey;
@@ -645,6 +623,10 @@ protected:
 	UPROPERTY( config, EditAnywhere, Category=Snapping )
 	bool bSnapCurveValueToInterval;
 
+	/** Enable or disable snapping keys/sections/playhead to the interval which forces whole frames */
+	UPROPERTY(config, EditAnywhere, Category = Snapping)
+	bool bForceWholeFrames;
+
 	/** Only show selected nodes in the tree view. */
 	UPROPERTY( config, EditAnywhere, Category=General )
 	bool bShowSelectedNodesOnly;
diff --git a/Engine/Source/Editor/Sequencer/Public/SequencerTimeSliderController.h b/Engine/Source/Editor/Sequencer/Public/SequencerTimeSliderController.h
--- a/Engine/Source/Editor/Sequencer/Public/SequencerTimeSliderController.h
+++ b/Engine/Source/Editor/Sequencer/Public/SequencerTimeSliderController.h
@@ -334,6 +334,7 @@ private:
 	void SetMark(FFrameNumber DiffFrame);
 
 	TSharedRef<SWidget> OpenSetPlaybackRangeMenu(const FGeometry& MyGeometry, const FPointerEvent& MouseEvent);
+	FFrameTime SnapSequencerTime(FFrameTime InTime) const;
 	FFrameTime ComputeScrubTimeFromMouse(const FGeometry& Geometry, const FPointerEvent& MouseEvent, FScrubRangeToScreen RangeToScreen) const;
 	FFrameTime ComputeFrameTimeFromMouse(const FGeometry& Geometry, FVector2D ScreenSpacePosition, FScrubRangeToScreen RangeToScreen, bool CheckSnapping = true) const;
 
