Sequencer: Fix frame count localization

#rnx

diff --git a/Engine/Source/Editor/Sequencer/Private/Sequencer.cpp b/Engine/Source/Editor/Sequencer/Private/Sequencer.cpp
--- a/Engine/Source/Editor/Sequencer/Private/Sequencer.cpp
+++ b/Engine/Source/Editor/Sequencer/Private/Sequencer.cpp
@@ -4668,7 +4668,7 @@ TSharedRef<SWidget> FSequencer::MakePlayTimeDisplay()
 			FFrameNumber FrameCount = FFrameRate::TransformTime((CurrentFrame - CurrentRange.GetLowerBoundValue() + 1).Value, TickResolution, DisplayRate).CeilToFrame();
 			FFrameNumber FrameDuration = FFrameRate::TransformTime(CurrentRange.Size<FFrameNumber>().Value, TickResolution, DisplayRate).CeilToFrame();
 
-			return FText::FromString(FString::Printf(TEXT("%d of %d"), FrameCount.Value, FrameDuration.Value));
+			return FText::Format(LOCTEXT("FrameCountTooltip", "{0} of {1}"), FrameCount.Value, FrameDuration.Value);
 		})
 		.TypeInterface(this, &FSequencer::GetNumericTypeInterface, UE::Sequencer::ENumericIntent::Position)
 		.Delta(this, &FSequencer::GetDisplayRateDeltaFrameCount)
