Sequencer: Improve/restore remapping of bindings when objects are pasted

This used to remap binding ids across all sections in the sequence but that caused problems so it was removed in 35145603.  Instead, it should only remap binding ids in sections that are newly created. For example, if you duplicate two bindings that are attached to each other, the duplicates should be attached to each other. However, if you duplicate a camera that is referenced by a camera cut track, the camera cut track should be switched to reference the duplicate camera. It used to when all sections were remapped.

diff --git a/Engine/Source/Editor/Sequencer/Private/SequencerUtilities.cpp b/Engine/Source/Editor/Sequencer/Private/SequencerUtilities.cpp
--- a/Engine/Source/Editor/Sequencer/Private/SequencerUtilities.cpp
+++ b/Engine/Source/Editor/Sequencer/Private/SequencerUtilities.cpp
@@ -3575,6 +3575,41 @@ bool FSequencerUtilities::PasteBindings(const FString& TextToImport, TSharedRef<
 		}
 	}
 
+	// Find all the sections that have been added and only remap bindings in those sections
+	TSet<UMovieSceneSection*> Sections;
+	for (auto BindingPasted : BindingsPasted)
+	{
+		if (FMovieSceneBinding* Binding = MovieScene->FindBinding(BindingPasted.GetObjectGuid()))
+		{
+			for (UMovieSceneTrack* Track : Binding->GetTracks())
+			{
+				for (UMovieSceneSection* Section : Track->GetAllSections())
+				{
+					Sections.Add(Section);
+				}
+			}
+		}
+	}
+
+	if (Sections.Num() != 0)
+	{
+		FMovieSceneSequenceIDRef FocusedGuid = Sequencer->GetFocusedTemplateID();
+
+		TMap<UE::MovieScene::FFixedObjectBindingID, UE::MovieScene::FFixedObjectBindingID> OldFixedToNewFixedMap;
+
+		TSharedRef<UE::MovieScene::FSharedPlaybackState> SharedPlaybackState = Sequencer->GetSharedPlaybackState();
+
+		for (TPair<FGuid, FGuid> GuidPair : OldToNewGuidMap)
+		{
+			OldFixedToNewFixedMap.Add(UE::MovieScene::FFixedObjectBindingID(GuidPair.Key, FocusedGuid), UE::MovieScene::FFixedObjectBindingID(GuidPair.Value, FocusedGuid));
+		}
+
+		for (UMovieSceneSection* Section : Sections)
+		{	
+			Section->OnBindingIDsUpdated(OldFixedToNewFixedMap, FocusedGuid, SharedPlaybackState);
+		}
+	}
+
 	for (auto BindingPasted : BindingsPasted)
 	{
 		OutBindings.Add(FMovieSceneBindingProxy(BindingPasted.GetObjectGuid(), Sequence));
