Sequencer: rewrite how the "add track" menu works for property tracks

This addresses two problems with the menus:

1. Directly-keyable structures (such as FVector or FColor, which have dedicated tracks) would show sub-menu items for their children properties even though this is unnecessary.
2. Keying struct properties for structs in arrays was broken.

Add some automated tests for animated structs in arrays.

#virtualized

diff --git a/Engine/Source/Editor/Sequencer/Private/MVVM/ViewModels/ObjectBindingModel.cpp b/Engine/Source/Editor/Sequencer/Private/MVVM/ViewModels/ObjectBindingModel.cpp
--- a/Engine/Source/Editor/Sequencer/Private/MVVM/ViewModels/ObjectBindingModel.cpp
+++ b/Engine/Source/Editor/Sequencer/Private/MVVM/ViewModels/ObjectBindingModel.cpp
@@ -56,20 +56,19 @@
 namespace UE::Sequencer
 {
 
-namespace
-{
+bool GSequencerObjectBindingShowNestedProperties = false;
+FAutoConsoleVariableRef CVarSequencerObjectBindingShowNestedProperties(
+	TEXT("Sequencer.ObjectBinding.ShowNestedProperties"),
+	GSequencerObjectBindingShowNestedProperties,
+	TEXT("(Default: false) When enabled, always show bound object properties as sub-menus reflecting the hierarchy of nested structures. When disabled, only do that for Level Sequences, make others use flat menus."),
+	ECVF_Default
+);
 
-struct PropertyMenuData
+namespace
 {
-	FString MenuName;
-	FPropertyPath PropertyPath;
-};
 
 void GetKeyablePropertyPaths(UClass* Class, void* ValuePtr, UStruct* PropertySource, FPropertyPath PropertyPath, FSequencer& Sequencer, TArray<FPropertyPath>& KeyablePropertyPaths)
 {
-	//@todo need to resolve this between UMG and the level editor sequencer
-	const bool bRecurseAllProperties = Sequencer.IsLevelEditorSequencer();
-
 	for (TFieldIterator<FProperty> PropertyIterator(PropertySource); PropertyIterator; ++PropertyIterator)
 	{
 		FProperty* Property = *PropertyIterator;
@@ -78,10 +77,11 @@ void GetKeyablePropertyPaths(UClass* Class, void* ValuePtr, UStruct* PropertySou
 		{
 			PropertyPath.AddProperty(FPropertyInfo(Property));
 
-			bool bIsPropertyKeyable = false;
-			FArrayProperty* ArrayProperty = CastField<FArrayProperty>(Property);
-			if (ArrayProperty)
+			if (FArrayProperty* ArrayProperty = CastField<FArrayProperty>(Property))
 			{
+				// If this is an array property, add property paths for each item in the array. If we don't know how to key that item,
+				// like for instance it's a custom struct for which we don't have a track editor, then we recurse into it and add
+				// property paths for its own properties.
 				FScriptArrayHelper ArrayHelper(ArrayProperty, ArrayProperty->ContainerPtrToValuePtr<void>(ValuePtr));
 				for (int32 Index = 0; Index < ArrayHelper.Num(); ++Index)
 				{
@@ -90,7 +90,6 @@ void GetKeyablePropertyPaths(UClass* Class, void* ValuePtr, UStruct* PropertySou
 					if (Sequencer.CanKeyProperty(FCanKeyPropertyParams(Class, PropertyPath)))
 					{
 						KeyablePropertyPaths.Add(PropertyPath);
-						bIsPropertyKeyable = true;
 					}
 					else if (FStructProperty* StructProperty = CastField<FStructProperty>(ArrayProperty->Inner))
 					{
@@ -100,22 +99,18 @@ void GetKeyablePropertyPaths(UClass* Class, void* ValuePtr, UStruct* PropertySou
 					PropertyPath = *PropertyPath.TrimPath(1);
 				}
 			}
-
-			if (!bIsPropertyKeyable)
+			else if (Sequencer.CanKeyProperty(FCanKeyPropertyParams(Class, PropertyPath)))
 			{
-				bIsPropertyKeyable = Sequencer.CanKeyProperty(FCanKeyPropertyParams(Class, PropertyPath));
-				if (bIsPropertyKeyable)
-				{
-					KeyablePropertyPaths.Add(PropertyPath);
-				}
+				// This is a property that we can key directly. That is: we have a track editor specifically for that property type,
+				// such as FVector or FLinearColor or other well known structs. This also includes custom/system-specific structs
+				// like FMargin, for which UMG registers a custom track.
+				KeyablePropertyPaths.Add(PropertyPath);
 			}
-
-			if (!bIsPropertyKeyable || bRecurseAllProperties)
+			else if (FStructProperty* StructProperty = CastField<FStructProperty>(Property))
 			{
-				if (FStructProperty* StructProperty = CastField<FStructProperty>(Property))
-				{
-					GetKeyablePropertyPaths(Class, StructProperty->ContainerPtrToValuePtr<void>(ValuePtr), StructProperty->Struct, PropertyPath, Sequencer, KeyablePropertyPaths);
-				}
+				// It's a struct property that we don't know how to key directly, so add property paths for its own properties.
+				// The user will have to key them individually.
+				GetKeyablePropertyPaths(Class, StructProperty->ContainerPtrToValuePtr<void>(ValuePtr), StructProperty->Struct, PropertyPath, Sequencer, KeyablePropertyPaths);
 			}
 
 			PropertyPath = *PropertyPath.TrimPath(1);
@@ -561,9 +556,6 @@ TSharedRef<SWidget> FObjectBindingModel::GetAddTrackMenuContent()
 	TSharedPtr<FSequencer> Sequencer = OwnerModel->GetSequencerImpl();
 	check(Sequencer);
 
-	//@todo need to resolve this between UMG and the level editor sequencer
-	const bool bUseSubMenus = Sequencer->IsLevelEditorSequencer();
-
 	UObject* BoundObject = Sequencer->FindSpawnedObjectOrTemplate(ObjectBindingID);
 
 	const UClass* MainSelectionObjectClass = FindObjectClass();
@@ -646,11 +638,30 @@ TSharedRef<SWidget> FObjectBindingModel::GetAddTrackMenuContent()
 		GetKeyablePropertyPaths(BoundObject->GetClass(), BoundObject, BoundObject->GetClass(), PropertyPath, *Sequencer, KeyablePropertyPaths);
 	}
 
-	// [Aspect Ratio]
-	// [PostProcess Settings] [Bloom1Tint] [X]
-	// [PostProcess Settings] [Bloom1Tint] [Y]
-	// [PostProcess Settings] [ColorGrading]
-	// [Ortho View]
+	AddPropertyMenuItems(AddTrackMenuBuilder, NumStartingBlocks, KeyablePropertyPaths, 0);
+
+	return AddTrackMenuBuilder.MakeWidget();
+}
+
+void FObjectBindingModel::AddPropertyMenuItems(FMenuBuilder& AddTrackMenuBuilder, int32 NumStartingBlocks, TArray<FPropertyPath> KeyablePropertyPaths, int32 PropertyNameIndexStart)
+{
+	// KeyablePropertyPaths contain a property path for each property, nested or not, that we can key. For instance:
+	//
+	// [MyFloat]                              (float, via the float property track)
+	// [SomeStruct] [MyColor]                 (SomeStruct isn't keyable so key its color, via the color property track)
+	// [SomeStruct] [OtherStruct] [MyInt]     (SomeStruct and OtherStruct aren't keyable so key the integer, via the int property track)
+	// [SomeStruct] [KnownStruct]			  (KnownStruct has a custom track to key its properties)
+
+
+	// If PropertyNameIndexStart is greater that zero, we are showing the sub-menu of a property path.
+	// That is, if we have property paths like this:
+	//
+	// [SomeStruct] [OtherStruct] [MyInt]
+	// [SomeStruct] [OtherStruct] [MyColor]
+	//
+	// ...and if we are showing the sub-menu for [OtherStruct]
+	//
+	// ...then PropertyNameIndexStart is 2, and we only need to show [MyInt] and [MyColor].
 
 	static const FString DefaultPropertyCategory = TEXT("Default");
 
@@ -675,19 +686,25 @@ TSharedRef<SWidget> FObjectBindingModel::GetAddTrackMenuContent()
 	};
 
 	bool bDefaultCategoryFound = false;
+	const bool bIsRootMenu = (PropertyNameIndexStart == 0);
 
 	// Create property menu data based on keyable property paths
-	TMap<FString, TArray<PropertyMenuData>> KeyablePropertyMenuData;
+	TMap<FString, TArray<FPropertyMenuData>> KeyablePropertyMenuData;
 	for (const FPropertyPath& KeyablePropertyPath : KeyablePropertyPaths)
 	{
-		const FProperty* Property = KeyablePropertyPath.GetRootProperty().Property.Get();
-		if (Property)
+		if (!ensure(KeyablePropertyPath.GetNumProperties() > PropertyNameIndexStart))
 		{
-			PropertyMenuData KeyableMenuData;
+			continue;
+		}
+
+		const FPropertyInfo& PropertyInfo = KeyablePropertyPath.GetPropertyInfo(PropertyNameIndexStart);
+		if (const FProperty* Property = PropertyInfo.Property.Get())
+		{
+			FPropertyMenuData KeyableMenuData;
 			KeyableMenuData.PropertyPath = KeyablePropertyPath;
-			if (KeyablePropertyPath.GetRootProperty().ArrayIndex != INDEX_NONE)
+			if (PropertyInfo.ArrayIndex != INDEX_NONE)
 			{
-				KeyableMenuData.MenuName = FText::Format(LOCTEXT("PropertyMenuTextFormat", "{0} [{1}]"), Property->GetDisplayNameText(), FText::AsNumber(KeyablePropertyPath.GetRootProperty().ArrayIndex)).ToString();
+				KeyableMenuData.MenuName = FText::Format(LOCTEXT("PropertyMenuTextFormat", "{0} [{1}]"), Property->GetDisplayNameText(), FText::AsNumber(PropertyInfo.ArrayIndex)).ToString();
 			}
 			else
 			{
@@ -708,23 +725,24 @@ TSharedRef<SWidget> FObjectBindingModel::GetAddTrackMenuContent()
 	KeyablePropertyMenuData.KeySort(FCategorySortPredicate());
 
 	// Always add an extension point for Properties section even if none are found (Components rely on this) 
-	if (!bDefaultCategoryFound)
+	if (!bDefaultCategoryFound && bIsRootMenu)
 	{
 		AddTrackMenuBuilder.BeginSection(SequencerMenuExtensionPoints::AddTrackMenu_PropertiesSection, LOCTEXT("PropertiesMenuHeader", "Properties"));
 		AddTrackMenuBuilder.EndSection();
 	}
 
 	// Add menu items
-	for (TPair<FString, TArray<PropertyMenuData>>& Pair : KeyablePropertyMenuData)
+	TSharedPtr<FSequencer> Sequencer = OwnerModel->GetSequencerImpl();
+	check(Sequencer);
+	const bool bUseSubMenus = GSequencerObjectBindingShowNestedProperties || Sequencer->IsLevelEditorSequencer();
+
+	for (TPair<FString, TArray<FPropertyMenuData>>& Pair : KeyablePropertyMenuData)
 	{
+		const FString CategoryText = Pair.Key;
+		TArray<FPropertyMenuData>& KeyablePropertySubMenuData = Pair.Value;
+		
 		// Sort on the property name
-		Pair.Value.Sort([](const PropertyMenuData& A, const PropertyMenuData& B)
-		{
-			int32 CompareResult = A.MenuName.Compare(B.MenuName);
-			return CompareResult < 0;
-		});
-
-		FString CategoryText = Pair.Key;
+		KeyablePropertySubMenuData.Sort();
 		
 		if (CategoryText == DefaultPropertyCategory)
 		{
@@ -735,27 +753,51 @@ TSharedRef<SWidget> FObjectBindingModel::GetAddTrackMenuContent()
 			AddTrackMenuBuilder.BeginSection(NAME_None, FText::FromString(CategoryText));
 		}
 	
-		for (int32 MenuDataIndex = 0; MenuDataIndex < Pair.Value.Num(); )
+		for (int32 MenuDataIndex = 0; MenuDataIndex < KeyablePropertySubMenuData.Num(); )
 		{
-			TArray<FPropertyPath> KeyableSubMenuPropertyPaths;
-
-			KeyableSubMenuPropertyPaths.Add(Pair.Value[MenuDataIndex].PropertyPath);
-
-			// If this menu data only has one property name, add the menu item
-			if (Pair.Value[MenuDataIndex].PropertyPath.GetNumProperties() == 1 || !bUseSubMenus)
+			// If this menu data only has one property name left in it, add the menu item
+			if (KeyablePropertySubMenuData[MenuDataIndex].PropertyPath.GetNumProperties() == PropertyNameIndexStart + 1)
+			{
+				AddPropertyMenuItem(AddTrackMenuBuilder, KeyablePropertySubMenuData[MenuDataIndex]);
+				++MenuDataIndex;
+			}
+			// If we don't want sub-menus, concatenate the property names left to handle, and add the menu item.
+			else if (!bUseSubMenus)
 			{
-				AddPropertyMenuItems(AddTrackMenuBuilder, KeyableSubMenuPropertyPaths, 0, -1);
+				TArray<FString> PropertyNames;
+				const FPropertyPath& CurPropertyPath(KeyablePropertySubMenuData[MenuDataIndex].PropertyPath);
+				for (int32 PropertyNameIndex = PropertyNameIndexStart; PropertyNameIndex < CurPropertyPath.GetNumProperties(); ++PropertyNameIndex)
+				{
+					const FPropertyInfo& CurPropertyInfo(CurPropertyPath.GetPropertyInfo(PropertyNameIndex));
+					if (CurPropertyInfo.ArrayIndex != INDEX_NONE)
+					{
+						PropertyNames.Add(FText::Format(
+									LOCTEXT("PropertyMenuTextFormat", "{0} [{1}]"), 
+									CurPropertyInfo.Property.Get()->GetDisplayNameText(), 
+									FText::AsNumber(CurPropertyInfo.ArrayIndex)).ToString());
+					}
+					else
+					{
+						PropertyNames.Add(CurPropertyInfo.Property.Get()->GetDisplayNameText().ToString());
+					}
+				}
+				KeyablePropertySubMenuData[MenuDataIndex].MenuName = FString::Join(PropertyNames, TEXT("."));
+
+				AddPropertyMenuItem(AddTrackMenuBuilder, KeyablePropertySubMenuData[MenuDataIndex]);
 				++MenuDataIndex;
 			}
 			// Otherwise, look to the next menu data to gather up new data
 			else
 			{
-				for (; MenuDataIndex < Pair.Value.Num()-1; )
+				TArray<FPropertyPath> KeyableSubMenuPropertyPaths;
+				KeyableSubMenuPropertyPaths.Add(KeyablePropertySubMenuData[MenuDataIndex].PropertyPath);
+
+				for (; MenuDataIndex < KeyablePropertySubMenuData.Num() - 1; )
 				{
-					if (Pair.Value[MenuDataIndex].MenuName == Pair.Value[MenuDataIndex+1].MenuName)
+					if (KeyablePropertySubMenuData[MenuDataIndex].MenuName == KeyablePropertySubMenuData[MenuDataIndex + 1].MenuName)
 					{	
 						++MenuDataIndex;
-						KeyableSubMenuPropertyPaths.Add(Pair.Value[MenuDataIndex].PropertyPath);
+						KeyableSubMenuPropertyPaths.Add(KeyablePropertySubMenuData[MenuDataIndex].PropertyPath);
 					}
 					else
 					{
@@ -764,9 +806,9 @@ TSharedRef<SWidget> FObjectBindingModel::GetAddTrackMenuContent()
 				}
 
 				AddTrackMenuBuilder.AddSubMenu(
-					FText::FromString(Pair.Value[MenuDataIndex].MenuName),
+					FText::FromString(KeyablePropertySubMenuData[MenuDataIndex].MenuName),
 					FText::GetEmpty(), 
-					FNewMenuDelegate::CreateSP(this, &FObjectBindingModel::HandleAddTrackSubMenuNew, KeyableSubMenuPropertyPaths, 0));
+					FNewMenuDelegate::CreateSP(this, &FObjectBindingModel::HandleAddTrackSubMenuNew, KeyableSubMenuPropertyPaths, PropertyNameIndexStart + 1));
 
 				++MenuDataIndex;
 			}
@@ -787,111 +829,11 @@ TSharedRef<SWidget> FObjectBindingModel::GetAddTrackMenuContent()
 
 		AddTrackMenuBuilder.AddWidget(EmptyTip, FText(), true, false);
 	}
-
-	return AddTrackMenuBuilder.MakeWidget();
 }
 
 void FObjectBindingModel::HandleAddTrackSubMenuNew(FMenuBuilder& AddTrackMenuBuilder, TArray<FPropertyPath> KeyablePropertyPaths, int32 PropertyNameIndexStart)
 {
-	// [PostProcessSettings] [Bloom1Tint] [X]
-	// [PostProcessSettings] [Bloom1Tint] [Y]
-	// [PostProcessSettings] [ColorGrading]
-
-	// Create property menu data based on keyable property paths
-	TArray<const FProperty*> PropertiesTraversed;
-	TArray<int32> ArrayIndicesTraversed;
-	TArray<PropertyMenuData> KeyablePropertyMenuData;
-	for (const FPropertyPath& KeyablePropertyPath : KeyablePropertyPaths)
-	{
-		PropertyMenuData KeyableMenuData;
-		KeyableMenuData.PropertyPath = KeyablePropertyPath;
-
-		// If the path is greater than 1, keep track of the actual properties (not channels) and only add these properties once since we can't do single channel keying of a property yet.
-		if (KeyablePropertyPath.GetNumProperties() > 1) //@todo
-		{
-			const FPropertyInfo& PropertyInfo = KeyablePropertyPath.GetPropertyInfo(1);
-			const FProperty* Property = PropertyInfo.Property.Get();
-
-			// Search for any array elements
-			int32 ArrayIndex = INDEX_NONE;
-			for (int32 PropertyInfoIndex = 0; PropertyInfoIndex < KeyablePropertyPath.GetNumProperties(); ++PropertyInfoIndex)
-			{
-				const FPropertyInfo& ArrayPropertyInfo = KeyablePropertyPath.GetPropertyInfo(PropertyInfoIndex);
-				if (ArrayPropertyInfo.ArrayIndex != INDEX_NONE)
-				{
-					ArrayIndex = ArrayPropertyInfo.ArrayIndex;
-					break;
-				}
-			}
-
-			bool bFound = false;
-			for (int32 TraversedIndex = 0; TraversedIndex < PropertiesTraversed.Num(); ++TraversedIndex)
-			{
-				if (PropertiesTraversed[TraversedIndex] == Property && ArrayIndicesTraversed[TraversedIndex] == ArrayIndex)
-				{
-					bFound = true;
-					break;
-				}
-			}
-
-			if (bFound)
-			{
-				continue;
-			}
-
-			if (ArrayIndex != INDEX_NONE)
-			{
-				KeyableMenuData.MenuName = FText::Format(LOCTEXT("ArrayElementFormat", "{0} [{1}]"), Property->GetDisplayNameText(), FText::AsNumber(ArrayIndex)).ToString();
-			}
-			else
-			{
-				KeyableMenuData.MenuName = FObjectEditorUtils::GetCategory(Property);
-			}
-
-			PropertiesTraversed.Add(Property);
-			ArrayIndicesTraversed.Add(ArrayIndex);
-		}
-		else
-		{
-			// No sub menus items, so skip
-			continue; 
-		}
-		KeyablePropertyMenuData.Add(KeyableMenuData);
-	}
-
-	// Sort on the menu name
-	KeyablePropertyMenuData.Sort([](const PropertyMenuData& A, const PropertyMenuData& B)
-	{
-		int32 CompareResult = A.MenuName.Compare(B.MenuName);
-		return CompareResult < 0;
-	});
-
-	// Add menu items
-	for (int32 MenuDataIndex = 0; MenuDataIndex < KeyablePropertyMenuData.Num(); )
-	{
-		TArray<FPropertyPath> KeyableSubMenuPropertyPaths;
-		KeyableSubMenuPropertyPaths.Add(KeyablePropertyMenuData[MenuDataIndex].PropertyPath);
-
-		for (; MenuDataIndex < KeyablePropertyMenuData.Num()-1; )
-		{
-			if (KeyablePropertyMenuData[MenuDataIndex].MenuName == KeyablePropertyMenuData[MenuDataIndex+1].MenuName)
-			{
-				++MenuDataIndex;
-				KeyableSubMenuPropertyPaths.Add(KeyablePropertyMenuData[MenuDataIndex].PropertyPath);
-			}
-			else
-			{
-				break;
-			}
-		}
-
-		AddTrackMenuBuilder.AddSubMenu(
-			FText::FromString(KeyablePropertyMenuData[MenuDataIndex].MenuName),
-			FText::GetEmpty(), 
-			FNewMenuDelegate::CreateSP(this, &FObjectBindingModel::AddPropertyMenuItems, KeyableSubMenuPropertyPaths, PropertyNameIndexStart + 1, PropertyNameIndexStart + 2));
-
-		++MenuDataIndex;
-	}
+	AddPropertyMenuItems(AddTrackMenuBuilder, 0, KeyablePropertyPaths, PropertyNameIndexStart);
 }
 
 void FObjectBindingModel::HandlePropertyMenuItemExecute(FPropertyPath PropertyPath)
@@ -932,54 +874,12 @@ void FObjectBindingModel::HandlePropertyMenuItemExecute(FPropertyPath PropertyPa
 	Sequencer->KeyProperty(KeyPropertyParams);
 }
 
-void FObjectBindingModel::AddPropertyMenuItems(FMenuBuilder& AddTrackMenuBuilder, TArray<FPropertyPath> KeyableProperties, int32 PropertyNameIndexStart, int32 PropertyNameIndexEnd)
+void FObjectBindingModel::AddPropertyMenuItem(FMenuBuilder& AddTrackMenuBuilder, const FPropertyMenuData& KeyablePropertyMenuData)
 {
-	TArray<PropertyMenuData> KeyablePropertyMenuData;
-
-	for (auto KeyableProperty : KeyableProperties)
-	{
-		TArray<FString> PropertyNames;
-		if (PropertyNameIndexEnd == -1)
-		{
-			PropertyNameIndexEnd = KeyableProperty.GetNumProperties();
-		}
-
-		//@todo
-		if (PropertyNameIndexStart >= KeyableProperty.GetNumProperties())
-		{
-			continue;
-		}
-
-		for (int32 PropertyNameIndex = PropertyNameIndexStart; PropertyNameIndex < PropertyNameIndexEnd; ++PropertyNameIndex)
-		{
-			PropertyNames.Add(KeyableProperty.GetPropertyInfo(PropertyNameIndex).Property.Get()->GetDisplayNameText().ToString());
-		}
-
-		PropertyMenuData KeyableMenuData;
-		{
-			KeyableMenuData.PropertyPath = KeyableProperty;
-			KeyableMenuData.MenuName = FString::Join( PropertyNames, TEXT( "." ) );
-		}
-
-		KeyablePropertyMenuData.Add(KeyableMenuData);
-	}
-
-	// Sort on the menu name
-	KeyablePropertyMenuData.Sort([](const PropertyMenuData& A, const PropertyMenuData& B)
-	{
-		int32 CompareResult = A.MenuName.Compare(B.MenuName);
-		return CompareResult < 0;
-	});
-
-	// Add menu items
-	for (int32 MenuDataIndex = 0; MenuDataIndex < KeyablePropertyMenuData.Num(); ++MenuDataIndex)
-	{
-		FUIAction AddTrackMenuAction(FExecuteAction::CreateSP(this, &FObjectBindingModel::HandlePropertyMenuItemExecute, KeyablePropertyMenuData[MenuDataIndex].PropertyPath));
-		AddTrackMenuBuilder.AddMenuEntry(FText::FromString(KeyablePropertyMenuData[MenuDataIndex].MenuName), FText(), FSlateIcon(), AddTrackMenuAction);
-	}
+	FUIAction AddTrackMenuAction(FExecuteAction::CreateSP(this, &FObjectBindingModel::HandlePropertyMenuItemExecute, KeyablePropertyMenuData.PropertyPath));
+	AddTrackMenuBuilder.AddMenuEntry(FText::FromString(KeyablePropertyMenuData.MenuName), FText(), FSlateIcon(), AddTrackMenuAction);
 }
 
-
 void FObjectBindingModel::BuildContextMenu(FMenuBuilder& MenuBuilder)
 {
 	TSharedPtr<FSequencerEditorViewModel> EditorViewModel = GetEditor();
diff --git a/Engine/Source/Editor/Sequencer/Public/MVVM/ViewModels/ObjectBindingModel.h b/Engine/Source/Editor/Sequencer/Public/MVVM/ViewModels/ObjectBindingModel.h
--- a/Engine/Source/Editor/Sequencer/Public/MVVM/ViewModels/ObjectBindingModel.h
+++ b/Engine/Source/Editor/Sequencer/Public/MVVM/ViewModels/ObjectBindingModel.h
@@ -132,7 +132,20 @@ protected:
 
 private:
 
-	void AddPropertyMenuItems(FMenuBuilder& AddTrackMenuBuilder, TArray<FPropertyPath> KeyableProperties, int32 PropertyNameIndexStart, int32 PropertyNameIndexEnd);
+	struct FPropertyMenuData
+	{
+		FString MenuName;
+		FPropertyPath PropertyPath;
+
+		bool operator< (const FPropertyMenuData& Other) const
+		{
+			int32 CompareResult = MenuName.Compare(Other.MenuName);
+			return CompareResult < 0;
+		}
+	};
+
+	void AddPropertyMenuItems(FMenuBuilder& AddTrackMenuBuilder, int32 NumStartingBlocks, TArray<FPropertyPath> KeyablePropertyPaths, int32 PropertyNameIndexStart);
+	void AddPropertyMenuItem(FMenuBuilder& AddTrackMenuBuilder, const FPropertyMenuData& KeyablePropertyMenuData);
 	void HandleAddTrackSubMenuNew(FMenuBuilder& AddTrackMenuBuilder, TArray<FPropertyPath> KeyablePropertyPaths, int32 PropertyNameIndexStart);
 	void HandlePropertyMenuItemExecute(FPropertyPath PropertyPath);
 
diff --git a/Engine/Source/Runtime/MovieScene/Private/Tests/MovieSceneTestObjects.h b/Engine/Source/Runtime/MovieScene/Private/Tests/MovieSceneTestObjects.h
--- a/Engine/Source/Runtime/MovieScene/Private/Tests/MovieSceneTestObjects.h
+++ b/Engine/Source/Runtime/MovieScene/Private/Tests/MovieSceneTestObjects.h
@@ -7,10 +7,12 @@
 #include "MovieSceneSection.h"
 #include "MovieSceneSequence.h"
 #include "Evaluation/MovieSceneEvalTemplate.h"
+#include "GameFramework/Actor.h"
 #include "Sections/MovieSceneSubSection.h"
 #include "Sections/MovieSceneHookSection.h"
 #include "Tracks/MovieSceneSubTrack.h"
 #include "Compilation/IMovieSceneTrackTemplateProducer.h"
+
 #include "MovieSceneTestObjects.generated.h"
 
 USTRUCT()
@@ -156,3 +158,43 @@ public:
 		return TriggerTimes;
 	}
 };
+
+USTRUCT(BlueprintType)
+struct FTestMovieSceneStruct
+{
+	GENERATED_BODY()
+
+	UPROPERTY(EditAnywhere, BlueprintReadWrite, Interp, Category="Test")
+	float First = 0.f;
+
+	UPROPERTY(EditAnywhere, BlueprintReadWrite, Interp, Category="Test")
+	float Second = 0.f;
+
+	UPROPERTY(EditAnywhere, BlueprintReadWrite, Interp, Category="Test")
+	FVector Vector;
+};
+
+UCLASS(MinimalAPI, BlueprintType)
+class ATestMovieSceneArrayPropertiesActor : public AActor
+{
+	GENERATED_BODY()
+
+public:
+
+	UPROPERTY(EditAnywhere, BlueprintReadWrite, Interp, Category="Test")
+	FVector TestVector;
+
+	UPROPERTY(EditAnywhere, BlueprintReadWrite, Interp, Category="Test")
+	TArray<float> MultipleFloats;
+
+	UPROPERTY(EditAnywhere, BlueprintReadWrite, Interp, Category="Test")
+	FTestMovieSceneStruct SingleStruct;
+
+	UPROPERTY(EditAnywhere, BlueprintReadWrite, Interp, Category="Test")
+	TArray<FTestMovieSceneStruct> MultipleStructs;
+
+	ATestMovieSceneArrayPropertiesActor(const FObjectInitializer& ObjInit)
+		: AActor(ObjInit)
+	{}
+};
+
