Sequencer: add support for animating properties inside an FInstancedPropertyBag or FInstancedStruct.

diff --git a/Engine/Source/Editor/Sequencer/Private/MVVM/ViewModels/ObjectBindingModel.cpp b/Engine/Source/Editor/Sequencer/Private/MVVM/ViewModels/ObjectBindingModel.cpp
--- a/Engine/Source/Editor/Sequencer/Private/MVVM/ViewModels/ObjectBindingModel.cpp
+++ b/Engine/Source/Editor/Sequencer/Private/MVVM/ViewModels/ObjectBindingModel.cpp
@@ -46,6 +46,7 @@
 #include "SequencerCommands.h"
 #include "SequencerNodeTree.h"
 #include "SequencerSettings.h"
+#include "StructUtils/PropertyBag.h"
 #include "MVVM/Views/ViewUtilities.h"
 #include "Styling/AppStyle.h"
 #include "Styling/SlateIconFinder.h"
@@ -73,7 +74,7 @@ void GetKeyablePropertyPaths(UClass* Class, void* ValuePtr, UStruct* PropertySou
 	{
 		FProperty* Property = *PropertyIterator;
 
-		if (Property && !Property->HasAnyPropertyFlags(CPF_Deprecated))
+		if (Property && !Property->HasAnyPropertyFlags(CPF_Deprecated) && !Property->GetBoolMetaData(TEXT("SequencerHideProperty")))
 		{
 			PropertyPath.AddProperty(FPropertyInfo(Property));
 
@@ -108,9 +109,29 @@ void GetKeyablePropertyPaths(UClass* Class, void* ValuePtr, UStruct* PropertySou
 			}
 			else if (FStructProperty* StructProperty = CastField<FStructProperty>(Property))
 			{
-				// It's a struct property that we don't know how to key directly, so add property paths for its own properties.
-				// The user will have to key them individually.
-				GetKeyablePropertyPaths(Class, StructProperty->ContainerPtrToValuePtr<void>(ValuePtr), StructProperty->Struct, PropertyPath, Sequencer, KeyablePropertyPaths);
+				if (StructProperty->Struct == FInstancedPropertyBag::StaticStruct() && StructProperty->GetBoolMetaData(TEXT("InterpBagProperties")))
+				{
+					// It's a property bag. Show the properties inside if we have been explicitly allowed to do so. Note that we are
+					// using a custom metadata ("InterpBagProperties") because the built-in "Interp" tag only works on BlueprintTypes.
+					FInstancedPropertyBag* PropertyBag = StructProperty->ContainerPtrToValuePtr<FInstancedPropertyBag>(ValuePtr);
+					FStructProperty* PropertyBagValueProperty = CastFieldChecked<FStructProperty>(StructProperty->Struct->FindPropertyByName(TEXT("Value")));
+					PropertyPath.AddProperty(FPropertyInfo(PropertyBagValueProperty));
+					const UPropertyBag* PropertyBagStruct = PropertyBag->GetPropertyBagStruct();
+					GetKeyablePropertyPaths(Class, PropertyBag, const_cast<UPropertyBag*>(PropertyBagStruct), PropertyPath, Sequencer, KeyablePropertyPaths);
+				}
+				else if (StructProperty->Struct == FInstancedStruct::StaticStruct() && StructProperty->GetBoolMetaData(TEXT("InterpStructProperties")))
+				{
+					// As above, but for an instanced struct.
+					FInstancedStruct* InstancedStruct = StructProperty->ContainerPtrToValuePtr<FInstancedStruct>(ValuePtr);
+					const UScriptStruct* InstancedStructType = InstancedStruct->GetScriptStruct();
+					GetKeyablePropertyPaths(Class, InstancedStruct, const_cast<UScriptStruct*>(InstancedStructType), PropertyPath, Sequencer, KeyablePropertyPaths);
+				}
+				else
+				{
+					// It's a struct property that we don't know how to key directly, so add property paths for its own properties.
+					// The user will have to key them individually.
+					GetKeyablePropertyPaths(Class, StructProperty->ContainerPtrToValuePtr<void>(ValuePtr), StructProperty->Struct, PropertyPath, Sequencer, KeyablePropertyPaths);
+				}
 			}
 
 			PropertyPath = *PropertyPath.TrimPath(1);
diff --git a/Engine/Source/Editor/Sequencer/Public/KeyPropertyParams.h b/Engine/Source/Editor/Sequencer/Public/KeyPropertyParams.h
--- a/Engine/Source/Editor/Sequencer/Public/KeyPropertyParams.h
+++ b/Engine/Source/Editor/Sequencer/Public/KeyPropertyParams.h
@@ -7,6 +7,7 @@
 #include "CoreMinimal.h"
 #include "Math/UnrealMathSSE.h"
 #include "PropertyPath.h"
+#include "StructUtils/InstancedStruct.h"
 #include "UObject/Field.h"
 #include "UObject/UnrealType.h"
 #include "UObject/WeakFieldPtr.h"
@@ -99,11 +100,11 @@ public:
 			const FPropertyInfo& PropertyInfo = PropertyPath.GetPropertyInfo(i);
 			if (const FProperty* Property = PropertyInfo.Property.Get())
 			{
-				int32 ArrayIndex = FMath::Max(0, PropertyInfo.ArrayIndex);
 				if (const FArrayProperty* ArrayProp = CastField<FArrayProperty>(Property))
 				{
 					// Sometimes property paths have the array property twice, first with no array index,
 					// then a second so we skip over this property if that's the case
+					int32 ArrayIndex = FMath::Max(0, PropertyInfo.ArrayIndex);
 					if (PropertyInfo.ArrayIndex == INDEX_NONE && i < PropertyPath.GetNumProperties()-1)
 					{
 						const FPropertyInfo& InnerPropertyInfo = PropertyPath.GetPropertyInfo(i+1);
@@ -122,9 +123,21 @@ public:
 					}
 					ContainerPtr = ParentArrayHelper.GetRawPtr(ArrayIndex);
 				}
-				else if (ArrayIndex >= 0 && ArrayIndex < Property->ArrayDim)
+				else if (const FStructProperty* StructProp = CastField<FStructProperty>(Property))
 				{
-					ContainerPtr = Property->ContainerPtrToValuePtr<ValueType>(ContainerPtr, ArrayIndex);
+					if (StructProp->Struct == FInstancedStruct::StaticStruct())
+					{
+						FInstancedStruct* InstancedStruct = Property->ContainerPtrToValuePtr<FInstancedStruct>(ContainerPtr);
+						ContainerPtr = InstancedStruct->GetMutableMemory();
+					}
+					else
+					{
+						ContainerPtr = Property->ContainerPtrToValuePtr<void>(ContainerPtr);
+					}
+				}
+				else
+				{
+					ContainerPtr = Property->ContainerPtrToValuePtr<void>(ContainerPtr);
 				}
 			}
 		}
@@ -158,4 +171,4 @@ private:
 };
 
 template<> SEQUENCER_API bool FPropertyChangedParams::GetPropertyValueImpl<bool>(void* Data, const FPropertyInfo& Info);
-template<> SEQUENCER_API UObject* FPropertyChangedParams::GetPropertyValueImpl<UObject*>(void* Data, const FPropertyInfo& Info);
\ No newline at end of file
+template<> SEQUENCER_API UObject* FPropertyChangedParams::GetPropertyValueImpl<UObject*>(void* Data, const FPropertyInfo& Info);
diff --git a/Engine/Source/Runtime/MovieScene/Private/MovieSceneCommonHelpers.cpp b/Engine/Source/Runtime/MovieScene/Private/MovieSceneCommonHelpers.cpp
--- a/Engine/Source/Runtime/MovieScene/Private/MovieSceneCommonHelpers.cpp
+++ b/Engine/Source/Runtime/MovieScene/Private/MovieSceneCommonHelpers.cpp
@@ -23,6 +23,7 @@
 #include "EntitySystem/MovieSceneSequenceInstance.h"
 #include "EntitySystem/MovieSceneEntitySystemLinker.h"
 #include "Conditions/MovieSceneGroupCondition.h"
+#include "StructUtils/InstancedStruct.h"
 
 bool MovieSceneHelpers::IsSectionKeyable(const UMovieSceneSection* Section)
 {
@@ -1358,8 +1359,21 @@ FTrackInstancePropertyBindings::FResolvedProperty FTrackInstancePropertyBindings
 
 		if (InPropertyNames.IsValidIndex(Index + 1))
 		{
-			void* StructContainer = StructProp->ContainerPtrToValuePtr<void>(BasePointer);
-			return ResolvePropertyRecursive(Bindings, StructContainer, StructProp->Struct, InPropertyNames, Index + 1);
+			// Instanced structs are technically just a memory buffer with no real sub-properties, but
+			// they do have sub-properties if we ask them about their "logical" struct type. Let's do that,
+			// which makes it possible to animate the properties inside.
+			if (StructProp->Struct == FInstancedStruct::StaticStruct())
+			{
+				FInstancedStruct* InstancedStruct = StructProp->ContainerPtrToValuePtr<FInstancedStruct>(BasePointer);
+				const UScriptStruct* InstancedStructType = InstancedStruct->GetScriptStruct();
+				uint8* InstancedStructMemory = InstancedStruct->GetMutableMemory();
+				return ResolvePropertyRecursive(Bindings, InstancedStructMemory, const_cast<UScriptStruct*>(InstancedStructType), InPropertyNames, Index + 1);
+			}
+			else
+			{
+				void* StructContainer = StructProp->ContainerPtrToValuePtr<void>(BasePointer);
+				return ResolvePropertyRecursive(Bindings, StructContainer, StructProp->Struct, InPropertyNames, Index + 1);
+			}
 		}
 		else
 		{
