Sequencer: Remove the word Toggle from menu items that are checkboxes

#rnx

diff --git a/Engine/Source/Editor/Sequencer/Private/SequencerCommands.cpp b/Engine/Source/Editor/Sequencer/Private/SequencerCommands.cpp
--- a/Engine/Source/Editor/Sequencer/Private/SequencerCommands.cpp
+++ b/Engine/Source/Editor/Sequencer/Private/SequencerCommands.cpp
@@ -193,11 +193,11 @@ void FSequencerCommands::RegisterCommands()
 
 	UI_COMMAND(RefreshUI, "Refresh UI", "Forcibly refresh the UI from source data.", EUserInterfaceActionType::Button, FInputChord(EKeys::F5));
 
-	UI_COMMAND(ToggleLimitViewportSelection, "Toggle Limit Viewport Selection", "Toggles viewport selectability between Sequencer only objects and all objects.", EUserInterfaceActionType::ToggleButton, FInputChord());
+	UI_COMMAND(ToggleLimitViewportSelection, "Limit Viewport Selection", "Toggles viewport selectability between Sequencer only objects and all objects.", EUserInterfaceActionType::ToggleButton, FInputChord());
 
 	UI_COMMAND(ToggleSidebarVisible, "Sidebar Details", "Toggles the visibility of the Sequencer sidebar.", EUserInterfaceActionType::ToggleButton, FInputChord());
-	UI_COMMAND(ToggleSidebarSelectionDrawerOpen, "Toggle Sidebar Selection Drawer Open", "Toggles the visible state of the sidebar selection drawer. Undocks and closes the sidebar drawer if docked.", EUserInterfaceActionType::Button, FInputChord(EModifierKey::Control, EKeys::E));
-	UI_COMMAND(ToggleSidebarDrawerDock, "Toggle Sidebar Drawer Docked", "Toggles the docked state of the sidebar current drawer. Undocks the docked sidebar drawer if docked or docks the sidebar drawer if there is one open and no currently docked drawer.", EUserInterfaceActionType::Button, FInputChord(EModifierKey::Alt, EKeys::D));
+	UI_COMMAND(ToggleSidebarSelectionDrawerOpen, "Sidebar Selection Drawer Open", "Toggles the visible state of the sidebar selection drawer. Undocks and closes the sidebar drawer if docked.", EUserInterfaceActionType::Button, FInputChord(EModifierKey::Control, EKeys::E));
+	UI_COMMAND(ToggleSidebarDrawerDock, "Sidebar Drawer Docked", "Toggles the docked state of the sidebar current drawer. Undocks the docked sidebar drawer if docked or docks the sidebar drawer if there is one open and no currently docked drawer.", EUserInterfaceActionType::Button, FInputChord(EModifierKey::Alt, EKeys::D));
 
 	UI_COMMAND(AlignSelectionToPlayhead, "Align Selection To Playhead", "Aligns layer bars, sections, or key selections to the current playhead location", EUserInterfaceActionType::Button, FInputChord());
 }
