[Sequencer]
- Process section models with new sort actions
- Process root level tracks/folders

diff --git a/Engine/Source/Editor/Sequencer/Private/SequencerCommonHelpers.cpp b/Engine/Source/Editor/Sequencer/Private/SequencerCommonHelpers.cpp
--- a/Engine/Source/Editor/Sequencer/Private/SequencerCommonHelpers.cpp
+++ b/Engine/Source/Editor/Sequencer/Private/SequencerCommonHelpers.cpp
@@ -616,13 +616,6 @@ void SequencerHelpers::SortOutlinerItems(FSequencer& Sequencer
 {
 	using namespace UE::Sequencer;
 
-	struct FSortItemCache
-	{
-		TViewModelPtr<FLayerBarModel> LayerBar;
-		TViewModelPtr<ISortableExtension> SortableItem;
-		TViewModelPtr<IOutlinerExtension> OutlinerItem;
-	};
-
 	if (InItems.IsEmpty())
 	{
 		return;
@@ -640,11 +633,92 @@ void SequencerHelpers::SortOutlinerItems(FSequencer& Sequencer
 		return;
 	}
 
-	const FSequencerSelection SequencerSelection = Sequencer.GetSelection();
+	/** Represents a layer bar. Holds cached pointers for element operations. */
+	struct FSortBarElement
+	{
+		static FSortBarElement FromTrack(FSequencer& InSequencer, const TViewModelPtr<ITrackAreaExtension>& InTrackAreaExtension)
+		{
+			for (const FViewModelPtr& TrackAreaModel : InTrackAreaExtension->GetTopLevelChildTrackAreaModels())
+			{
+				if (const TViewModelPtr<FLayerBarModel> BarModel = TrackAreaModel.ImplicitCast())
+				{
+					return FSortBarElement(InSequencer, BarModel);
+				}
+			}
+			for (const TViewModelPtr<ILayerBarExtension>& TrackAreaModel : InTrackAreaExtension->GetTrackAreaModelListAs<ILayerBarExtension>())
+			{
+				if (const TViewModelPtr<ILayerBarExtension> BarModel = TrackAreaModel.ImplicitCast())
+				{
+					return FSortBarElement(InSequencer, BarModel);
+				}
+			}
+			return FSortBarElement(InSequencer);
+		}
+	
+		FSortBarElement(FSequencer& InSequencer)
+			: Sequencer(InSequencer)
+		{}
+		FSortBarElement(FSequencer& InSequencer, const TViewModelPtr<FLayerBarModel>& InBarModel)
+			: Sequencer(InSequencer)
+		{
+			if (const TViewModelPtr<IOutlinerExtension> LinkedOutlinerItem = InBarModel->GetLinkedOutlinerItem())
+			{
+				if (const TViewModelPtr<ISortableExtension> SortableExtension = LinkedOutlinerItem.ImplicitCast())
+				{
+					const TViewModelPtr<ISortableExtension> ParentSortableItem = LinkedOutlinerItem.AsModel()->FindAncestorOfType<ISortableExtension>();
+
+					BarModel.Emplace<TViewModelPtr<FLayerBarModel>>(InBarModel);
+					SortableItem = SortableExtension;
+					ParentItem = ParentSortableItem.IsValid()? ParentSortableItem : CastViewModelChecked<ISortableExtension>(Sequencer.GetViewModel()->GetRootSequenceModel());
+					Range = InBarModel->ComputeRange();
+				}
+			}
+		}
+		FSortBarElement(FSequencer& InSequencer, const TViewModelPtr<ILayerBarExtension>& InBarModel)
+			: Sequencer(InSequencer)
+		{
+			if (const TViewModelPtr<FLinkedOutlinerExtension> LinkedOutlinerExtension = InBarModel.ImplicitCast())
+			{
+				if (const TViewModelPtr<IOutlinerExtension> LinkedOutlinerItem = LinkedOutlinerExtension->GetLinkedOutlinerItem())
+				{
+					if (const TViewModelPtr<ISortableExtension> SortableExtension = LinkedOutlinerItem.ImplicitCast())
+					{
+						const TViewModelPtr<ISortableExtension> ParentSortableItem = LinkedOutlinerItem.AsModel()->FindAncestorOfType<ISortableExtension>();
+
+						BarModel.Emplace<TViewModelPtr<ILayerBarExtension>>(InBarModel);
+						SortableItem = SortableExtension;
+						ParentItem = ParentSortableItem.IsValid() ? ParentSortableItem : CastViewModelChecked<ISortableExtension>(Sequencer.GetViewModel()->GetRootSequenceModel());
+						Range = InBarModel->GetLayerBarRange();
+					}
+				}
+			}
+		}
+
+		bool IsValid() const
+		{
+			if (const TViewModelPtr<FLayerBarModel>* LayerBarModel = BarModel.TryGet<TViewModelPtr<FLayerBarModel>>())
+			{
+				return LayerBarModel->IsValid();
+			}
+			if (const TViewModelPtr<ILayerBarExtension>* LayerBarExtension = BarModel.TryGet<TViewModelPtr<ILayerBarExtension>>())
+			{
+				return LayerBarExtension->IsValid();
+			}
+			return false;
+		}
+
+		FSequencer& Sequencer;
+
+		TVariant<TViewModelPtr<FLayerBarModel>, TViewModelPtr<ILayerBarExtension>> BarModel;
+
+		TViewModelPtr<ISortableExtension> SortableItem;
+		TViewModelPtr<ISortableExtension> ParentItem;
+		TRange<FFrameNumber> Range;
+	};
 
 	// Gather all the selected folder and layer bar models from InItems so we can sort them
-	TMap<TViewModelPtr<FFolderModel>, TArray<TViewModelPtr<FLayerBarModel>>> FolderPairs;
-	TMap<TViewModelPtr<IOutlinerExtension>, TArray<FSortItemCache>> LayerBarSortItems;
+	TMap<TViewModelPtr<FFolderModel>, TArray<FSortBarElement>> FolderPairs;
+	TMap<TViewModelPtr<ISortableExtension>, TArray<FSortBarElement>> SortItems;
 
 	for (const TViewModelPtr<IOutlinerExtension>& TrackItem : InItems)
 	{
@@ -653,52 +727,34 @@ void SequencerHelpers::SortOutlinerItems(FSequencer& Sequencer
 		{
 			if (const TViewModelPtr<FFolderModel> FolderModel = TrackItem.ImplicitCast())
 			{
-				TArray<TViewModelPtr<FLayerBarModel>> FolderLayerBars;
+				TArray<FSortBarElement> FolderElements;
 
 				for (const TViewModelPtr<IOutlinerExtension>& ChildTrack : FolderModel->GetChildrenOfType<IOutlinerExtension>())
 				{
 					if (const TViewModelPtr<ITrackAreaExtension> ChildTrackArea = ChildTrack.ImplicitCast())
 					{
-						for (const FViewModelPtr& TrackAreaModel : ChildTrackArea->GetTopLevelChildTrackAreaModels())
+						FSortBarElement NewElement = FSortBarElement::FromTrack(Sequencer, ChildTrackArea);
+						if (NewElement.IsValid())
 						{
-							if (const TViewModelPtr<FLayerBarModel> LayerBarModel = TrackAreaModel.ImplicitCast())
-							{
-								FolderLayerBars.Add(LayerBarModel);
-							}
+							FolderElements.Add(MoveTemp(NewElement));
 						}
 					}
 				}
 
-				if (!FolderLayerBars.IsEmpty())
+				if (!FolderElements.IsEmpty())
 				{
-					FolderPairs.Add(FolderModel, FolderLayerBars);
+					FolderPairs.Add(FolderModel, FolderElements);
 				}
 			}
 		}
 
 		// Handle track areas
-		if (const TSharedPtr<ITrackAreaExtension> TrackArea = TrackItem.ImplicitCast())
+		if (const TViewModelPtr<ITrackAreaExtension> TrackArea = TrackItem.ImplicitCast())
 		{
-			for (const FViewModelPtr& TrackAreaModel : TrackArea->GetTopLevelChildTrackAreaModels())
+			FSortBarElement NewElement = FSortBarElement::FromTrack(Sequencer, TrackArea);
+			if (NewElement.IsValid())
 			{
-				if (const TViewModelPtr<FLayerBarModel> LayerBarModel = TrackAreaModel.ImplicitCast())
-				{
-					if (const TViewModelPtr<IOutlinerExtension> OutlinerItem = LayerBarModel->GetLinkedOutlinerItem())
-					{
-						if (const TViewModelPtr<ISortableExtension> SortableItem = OutlinerItem.ImplicitCast())
-						{
-							if (const TViewModelPtr<IOutlinerExtension> ParentOutlinerItem = TrackItem.AsModel()->FindAncestorOfType<IOutlinerExtension>())
-							{
-								// Cache this item info and add it to the array in the map for sorting
-								FSortItemCache NewSortItem;
-								NewSortItem.LayerBar = LayerBarModel;
-								NewSortItem.SortableItem = SortableItem;
-								NewSortItem.OutlinerItem = OutlinerItem;
-								LayerBarSortItems.FindOrAdd(ParentOutlinerItem).Add(NewSortItem);
-							}
-						}
-					}
-				}
+				SortItems.FindOrAdd(NewElement.ParentItem).Add(NewElement);
 			}
 		}
 	}
@@ -714,48 +770,42 @@ void SequencerHelpers::SortOutlinerItems(FSequencer& Sequencer
 		MovieScene->Modify();
 	}
 
-	// Set the sort order for all selected folders
+	// Set the sort order for all selected folder children
 	if (!bInSortByItemOrder)
 	{
 		int32 SortOrder = 0;
-		for (TTuple<TViewModelPtr<FFolderModel>, TArray<TViewModelPtr<FLayerBarModel>>>& FolderModel : FolderPairs)
+		for (TPair<TViewModelPtr<FFolderModel>, TArray<FSortBarElement>>& FolderModel : FolderPairs)
 		{
-			FolderModel.Value.Sort([bInDescending](const TViewModelPtr<FLayerBarModel>& InLayerBarA, const TViewModelPtr<FLayerBarModel>& InLayerBarB) -> bool
+			FolderModel.Value.Sort([bInDescending](const FSortBarElement& InLayerBarA, const FSortBarElement& InLayerBarB) -> bool
 				{
 					if (bInDescending)
 					{
-						return InLayerBarA->ComputeRange().GetLowerBoundValue() > InLayerBarB->ComputeRange().GetLowerBoundValue();
+						return InLayerBarA.Range.GetLowerBoundValue() > InLayerBarB.Range.GetLowerBoundValue();
 					}
-					return InLayerBarA->ComputeRange().GetLowerBoundValue() < InLayerBarB->ComputeRange().GetLowerBoundValue();
+					return InLayerBarA.Range.GetLowerBoundValue() < InLayerBarB.Range.GetLowerBoundValue();
 				});
 
 			SortOrder = 0;
-			for (const TViewModelPtr<FLayerBarModel>& LayerBar : FolderModel.Value)
+			for (const FSortBarElement& Element : FolderModel.Value)
 			{
-				if (const TViewModelPtr<IOutlinerExtension> OutlinerItem = LayerBar->GetLinkedOutlinerItem())
-				{
-					if (const TViewModelPtr<ISortableExtension> SortableExtension = OutlinerItem.ImplicitCast())
-					{
-						SortableExtension->SetCustomOrder(SortOrder);
-						++SortOrder;
-					}
-				}
+				Element.SortableItem->SetCustomOrder(SortOrder);
+				++SortOrder;
 			}
 		}
 	}
 
-	// Set the sort order for all selected layer bars
-	for (TPair<TViewModelPtr<IOutlinerExtension>, TArray<FSortItemCache>>& LayerBarSortItemPair : LayerBarSortItems)
+	// Set the sort order for all items
+	for (TPair<TViewModelPtr<ISortableExtension>, TArray<FSortBarElement>>& LayerBarSortItemPair : SortItems)
 	{
-		TViewModelPtr<IOutlinerExtension>& ParentOutlinerItem = LayerBarSortItemPair.Key;
-		TArray<FSortItemCache>& SortItems = LayerBarSortItemPair.Value;
+		TViewModelPtr<ISortableExtension>& ParentSortableItem = LayerBarSortItemPair.Key;
+		TArray<FSortBarElement>& Elements = LayerBarSortItemPair.Value;
 
-		// Remove all layer models that have a descendant that is selected
-		SortItems.RemoveAll([&SequencerSelection](const FSortItemCache& InSortItem) -> bool
+		// Remove all layer models that have a descendant that is in the InItems list
+		Elements.RemoveAll([&InItems](const FSortBarElement& InElement) -> bool
 			{
-				for (const TViewModelPtr<IOutlinerExtension>& ChildOutlinerExtension : InSortItem.OutlinerItem.AsModel()->GetDescendantsOfType<IOutlinerExtension>())
+				for (const TViewModelPtr<IOutlinerExtension>& ChildOutlinerExtension : InElement.SortableItem.AsModel()->GetDescendantsOfType<IOutlinerExtension>())
 				{
-					if (SequencerSelection.Outliner.IsSelected(ChildOutlinerExtension))
+					if (InItems.Contains(ChildOutlinerExtension))
 					{
 						return true;
 					}
@@ -765,20 +815,20 @@ void SequencerHelpers::SortOutlinerItems(FSequencer& Sequencer
 
 		if (!bInSortByItemOrder)
 		{
-			SortItems.Sort([bInDescending](const FSortItemCache& InSortItemA, const FSortItemCache& InSortItemB) -> bool
+			Elements.Sort([bInDescending](const FSortBarElement& InSortItemA, const FSortBarElement& InSortItemB) -> bool
 				{
 					if (bInDescending)
 					{
-						return InSortItemA.LayerBar->ComputeRange().GetLowerBoundValue() > InSortItemB.LayerBar->ComputeRange().GetLowerBoundValue();
+						return InSortItemA.Range.GetLowerBoundValue() > InSortItemB.Range.GetLowerBoundValue();
 					}
-					return InSortItemA.LayerBar->ComputeRange().GetLowerBoundValue() < InSortItemB.LayerBar->ComputeRange().GetLowerBoundValue();
+					return InSortItemA.Range.GetLowerBoundValue() < InSortItemB.Range.GetLowerBoundValue();
 				});
 		}
 
 		// Save the current sortable children and begin the re-ordering operation.
 		// We will attempt to maintain any current custom ordering that exists.
 		TArray<TViewModelPtr<ISortableExtension>> SortableChildren;
-		ParentOutlinerItem.AsModel()->GetChildrenOfType<ISortableExtension>().ToArray(SortableChildren);
+		ParentSortableItem.AsModel()->GetChildrenOfType<ISortableExtension>().ToArray(SortableChildren);
 
 		// Sort to make sure we are in the correct order since the array can be out of order from what is actually displayed
 		SortableChildren.Sort([](const TViewModelPtr<ISortableExtension>& InA, const TViewModelPtr<ISortableExtension>& InB) -> bool
@@ -792,7 +842,7 @@ void SequencerHelpers::SortOutlinerItems(FSequencer& Sequencer
 		for (int32 Index = SortableChildren.Num() - 1; Index >= 0; --Index)
 		{
 			const TViewModelPtr<ISortableExtension>& SortableChild = SortableChildren[Index];
-			const bool bHasItemToSort = SortItems.ContainsByPredicate([&SortableChild](const FSortItemCache& InItemCache)
+			const bool bHasItemToSort = Elements.ContainsByPredicate([&SortableChild](const FSortBarElement& InItemCache)
 				{
 					return InItemCache.SortableItem == SortableChild;
 				});
@@ -810,9 +860,9 @@ void SequencerHelpers::SortOutlinerItems(FSequencer& Sequencer
 		int32 CurrentIndex = 0;
 		for (const int32 ItemIndex : ItemIndicesToReinsert)
 		{
-			if (SortItems.IsValidIndex(CurrentIndex))
+			if (Elements.IsValidIndex(CurrentIndex))
 			{
-				SortableChildren.Insert(SortItems[CurrentIndex].SortableItem, ItemIndex);
+				SortableChildren.Insert(Elements[CurrentIndex].SortableItem, ItemIndex);
 				++CurrentIndex;
 			}
 		}
