Anim In Engine: Motion trails 2.0.

diff --git a/Engine/Plugins/Editor/SequencerAnimTools/Source/SequencerAnimTools/Public/TrailHierarchy.h b/Engine/Plugins/Editor/SequencerAnimTools/Source/SequencerAnimTools/Public/TrailHierarchy.h
--- a/Engine/Plugins/Editor/SequencerAnimTools/Source/SequencerAnimTools/Public/TrailHierarchy.h
+++ b/Engine/Plugins/Editor/SequencerAnimTools/Source/SequencerAnimTools/Public/TrailHierarchy.h
@@ -13,7 +13,6 @@ namespace UE
 namespace SequencerAnimTools
 {
 
-
 struct FTrailVisibilityManager
 {
 	bool IsTrailVisible(const FGuid& Guid, const FTrail* Trail) const
@@ -38,6 +37,14 @@ struct FTrailVisibilityManager
 		}
 	}
 
+	void Reset()
+	{
+		InactiveMask.Empty();
+		VisibilityMask.Empty();
+		AlwaysVisible.Empty();
+		Selected.Empty();
+		ControlSelected.Empty();
+	}
 
 	TSet<FGuid> InactiveMask; // Any trails whose cache state or parent's cache state has been marked as NotUpdated
 	TSet<FGuid> VisibilityMask; // Any trails masked out by the user interface, ex bone trails
@@ -73,11 +80,11 @@ class FTrailHierarchy
 public:
 
 	FTrailHierarchy()
-		: ViewRange(TRange<double>::All())
-		, EvalRange(TRange<double>::All())
-		, SecondsPerSegment(0.1)
-		, LastEvalRange(TRange<double>::All())
-		, LastSecondsPerSegment(0.1)
+		:TickViewRange(TRange<FFrameNumber>(0,0))
+		,TickEvalRange(TRange<FFrameNumber>(0,0))
+		,TicksPerSegment(1)
+		,LastTickEvalRange(TRange<FFrameNumber>(0,0))
+		,LastTicksPerSegment(1)
 		, AllTrails()
 		, TimingStats()
 		, VisibilityManager()
@@ -88,12 +95,16 @@ public:
 	virtual void Initialize() = 0;
 	virtual void Destroy() = 0; // TODO: make dtor?
 	virtual ITrailHierarchyRenderer* GetRenderer() const = 0;
-	virtual double GetSecondsPerFrame() const = 0;
-	virtual double GetSecondsPerSegment() const = 0;
 	virtual FFrameNumber GetFramesPerFrame() const = 0;
 	virtual FFrameNumber GetFramesPerSegment() const = 0;
 
-	// Optionally implemented methods
+	//new
+	virtual const FCurrentFramesInfo* GetCurrentFramesInfo() const = 0;
+	virtual bool IsVisible(const FGuid& InTrailGuid) const;
+	virtual bool CheckForChanges() = 0;
+	virtual bool IsTrailEvaluating(const FGuid& InTrailGuid, bool bIndirectlyOnly) const = 0;
+
+	//optional methods
 	virtual void CalculateEvalRangeArray();
 	virtual void Update();
 	virtual bool HandleClick(FEditorViewportClient* InViewportClient, HHitProxy* HitProxy, FInputClick Click);
@@ -103,8 +114,8 @@ public:
 
 	//Get calculated center of all selections into one position, return false if not selected
 	virtual bool IsAnythingSelected(FVector& OutVectorPosition)const;
-	//Get each selected item's position, return false if not selected.
-	virtual bool IsAnythingSelected(TArray<FVector>& OutVectorPositions) const;
+	//Get each selected item's position, return false if not selected. if bAllPositions is true it will get all selected positions from each trail, otherwise will be the average
+	virtual bool IsAnythingSelected(TArray<FVector>& OutVectorPositions, bool bAllPositions = false) const;
 	virtual bool IsAnythingSelected() const;
 	virtual void SelectNone();
 	virtual bool IsSelected(const FGuid& Key) const;
@@ -114,13 +125,14 @@ public:
 	virtual void RemoveTrail(const FGuid& Key);
 
 	virtual bool StartTracking();
-	virtual bool ApplyDelta(const FVector& Pos, const FRotator& Rot, const FVector& WidgetLocation);
+	virtual bool ApplyDelta(const FVector& Pos, const FRotator& Rot, const FVector& WidgetLocation, bool bApplyToOffset);
 	virtual bool EndTracking();
 
 	virtual void TranslateSelectedKeys(bool bRight);
 	virtual void DeleteSelectedKeys();
 
-	const TRange<double>& GetViewRange() const { return ViewRange; }
+	const TRange<FFrameNumber>& GetViewFrameRange() const { return TickViewRange; }
+	FFrameNumber GetTicksPerSegment() const { return TicksPerSegment; }
 	const TMap<FGuid, TUniquePtr<FTrail>>& GetAllTrails() const { return AllTrails; }
 
 	const TMap<FString, FTimespan>& GetTimingStats() const { return TimingStats; };
@@ -133,15 +145,11 @@ protected:
 
 	void OpenContextMenu(const FGuid& TrailGuid);
 protected:
-
-
-	TRange<double> ViewRange;
-	TRange<double> EvalRange;
-	TArray<double> EvalTimesArr;
-	double SecondsPerSegment;
-
-	TRange<double> LastEvalRange;
-	double LastSecondsPerSegment;
+	TRange<FFrameNumber> TickViewRange;
+	TRange<FFrameNumber> TickEvalRange;
+	FFrameNumber TicksPerSegment;
+	TRange<FFrameNumber> LastTickEvalRange = TRange<FFrameNumber>(0, 0);
+	FFrameNumber LastTicksPerSegment = FFrameNumber(1);
 
 	TMap<FGuid, TUniquePtr<FTrail>> AllTrails;
 
diff --git a/Engine/Plugins/Editor/SequencerAnimTools/Source/SequencerAnimTools/Public/TrajectoryCache.h b/Engine/Plugins/Editor/SequencerAnimTools/Source/SequencerAnimTools/Public/TrajectoryCache.h
--- a/Engine/Plugins/Editor/SequencerAnimTools/Source/SequencerAnimTools/Public/TrajectoryCache.h
+++ b/Engine/Plugins/Editor/SequencerAnimTools/Source/SequencerAnimTools/Public/TrajectoryCache.h
@@ -5,135 +5,3 @@
 #include "Math/Range.h"
 #include "Math/Transform.h"
 
-namespace UE
-{
-namespace SequencerAnimTools
-{
-
-bool RangesContain(const TSet<TRange<double>>& Ranges, const double Time);
-
-
-// Intermediate range reprsentation structure
-class FTrailEvaluateTimes
-{
-public:
-	FTrailEvaluateTimes()
-		: EvalTimes()
-		, Spacing(0.1)
-		, Range(TRange<double>::Empty())
-	{}
-
-	FTrailEvaluateTimes(TArrayView<double> InEvalTimes, double InSpacing) 
-		: EvalTimes(InEvalTimes)
-		, Spacing(InSpacing)
-		, Range(TRange<double>(InEvalTimes[0], InEvalTimes.Last()))
-	{}
-
-	TArrayView<double> EvalTimes;
-	double Spacing;
-	TRange<double> Range;
-};
-
-class FTrajectoryCache
-{
-public:
-	FTrajectoryCache()
-		: Default(FTransform::Identity)
-		, TrackRange()
-		, Spacing()
-	{}
-
-	FTrajectoryCache(const double InSpacing, const TRange<double>& InTrackRange, const FTransform& InDefault)
-		: Default(InDefault)
-		, TrackRange(TRange<double>(FMath::FloorToDouble(InTrackRange.GetLowerBoundValue() / InSpacing)* InSpacing, FMath::FloorToDouble(InTrackRange.GetUpperBoundValue() / InSpacing)* InSpacing)) // snap to spacing
-		, Spacing(InSpacing)
-
-	{}
-
-	virtual ~FTrajectoryCache() {}
-
-	// alternatively, GetEvaluateTimes(const FTrailEvaluateTimes&) and GetInterpEvaluateTimes(const FTrailEvaluateTimes&)
-	virtual const FTransform& Get(const double InTime) const = 0;
-	virtual FTransform GetInterp(const double InTime) const = 0;
-	virtual void GetInterp(const double InTime, FTransform& OutTransform, FTransform& OutParentTransform) const{};
-	virtual const FTransform& GetParent(const double InTime) const { return FTransform::Identity; }
-
-	virtual TArray<double> GetAllTimesInRange(const TRange<double>& InRange) const = 0;
-	
-	virtual void Set(const double InTime, const FTransform& InValue) = 0;
-	virtual void SetTransforms(TArray<FTransform>& InTrajectoryCache, TArray<FTransform>& InParentCache) = 0;
-	// Optionally implemented
-	virtual const FTransform& GetDefault() const { return Default; };
-	virtual void UpdateCacheTimes(const FTrailEvaluateTimes& InEvalTimes);
-	
-	const TRange<double>& GetTrackRange() const { return TrackRange; }
-
-protected:
-	FTransform Default;
-	TRange<double> TrackRange;
-	double Spacing;
-};
-
-class FArrayTrajectoryCache : public FTrajectoryCache
-{
-public:
-
-	// for FRootTrail
-	FArrayTrajectoryCache()
-		: FTrajectoryCache()
-		, TrajectoryCache()
-
-	{}
-
-	FArrayTrajectoryCache(const double InSpacing, const TRange<double>& InTrackRange, FTransform InDefault = FTransform::Identity)
-		: FTrajectoryCache(InSpacing,InTrackRange,InDefault)
-		, TrajectoryCache()
-
-	{
-		TrajectoryCache.SetNumUninitialized(int32(InTrackRange.Size<double>() / InSpacing) + 1);
-
-		for (FTransform& Transform : TrajectoryCache)
-		{
-			Transform = InDefault;
-		}
-	}
-
-	virtual const FTransform& Get(const double InTime) const override 
-	{ 
-		return TrajectoryCache.Num() > 0 ? TrajectoryCache[FMath::Clamp(int32((InTime - TrackRange.GetLowerBoundValue()) / Spacing), 0, TrajectoryCache.Num() - 1)] : Default; 
-	}
-
-	virtual FTransform GetInterp(const double InTime) const override;
-	virtual void GetInterp(const double InTime, FTransform& OutTransform, FTransform& OutParentTransform) const override;
-
-	virtual const FTransform& GetParent(const double InTime) const override
-	{
-		return ParentTrajectoryCache.Num() > 0 ? ParentTrajectoryCache[FMath::Clamp(int32((InTime - TrackRange.GetLowerBoundValue()) / Spacing), 0, ParentTrajectoryCache.Num() - 1)] : FTransform::Identity;
-	}
-
-
-	virtual TArray<double> GetAllTimesInRange(const TRange<double>& InRange) const override;
-
-	virtual void Set(const double InTime, const FTransform& InValue) override 
-	{ 
-		if (TrajectoryCache.Num() > 0)
-		{
-			TrajectoryCache[FMath::Clamp(int32((InTime - TrackRange.GetLowerBoundValue()) / Spacing), 0, TrajectoryCache.Num() - 1)] = InValue;
-		}
-	}
-
-	virtual void SetTransforms(TArray<FTransform>& InTrajectoryCache, TArray<FTransform>& InParentCache)
-	{
-		TrajectoryCache = InTrajectoryCache;
-		ParentTrajectoryCache = InParentCache;
-	}
-
-	
-private:
-	TArray<FTransform> TrajectoryCache;
-	TArray<FTransform> ParentTrajectoryCache;
-
-};
-
-} // namespace MovieScene
-} // namespace UE
diff --git a/Engine/Plugins/Editor/SequencerAnimTools/Source/SequencerAnimTools/Public/TrajectoryDrawInfo.h b/Engine/Plugins/Editor/SequencerAnimTools/Source/SequencerAnimTools/Public/TrajectoryDrawInfo.h
--- a/Engine/Plugins/Editor/SequencerAnimTools/Source/SequencerAnimTools/Public/TrajectoryDrawInfo.h
+++ b/Engine/Plugins/Editor/SequencerAnimTools/Source/SequencerAnimTools/Public/TrajectoryDrawInfo.h
@@ -6,7 +6,8 @@
 #include "Misc/Guid.h"
 #include "Math/Color.h"
 #include "Math/Range.h"
-
+#include "Tools/EvaluateSequencerTools.h"
+#include "Tools/MotionTrailOptions.h"
 
 class FSceneView;
 template <typename OptionalType> struct TOptional;
@@ -16,9 +17,6 @@ namespace UE
 namespace SequencerAnimTools
 {
 
-class FTrail;  
-class FMapTrajectoryCache;
-class FTrajectoryCache;
 
 class FTrailScreenSpaceTransform
 {
@@ -35,43 +33,100 @@ private:
 	const float DPIScale;
 };
 
-struct FDisplayContext
+
+//This struct will give set of indices to calculate based upon number of frames
+//and current bucket we are calculating. t
+struct FFrameCalculator
 {
-	FGuid YourNode;
-	const FTrailScreenSpaceTransform& ScreenSpaceTransform;
-	double SecondsPerTick;
-	const TRange<double>& TimeRange;
-	double SecondsPerSegment;
-};
+	FFrameCalculator()
+	{
+		IndicesToCalculate.Reserve(IndicesToCalculatePerBucket);
+	}
+	//Set it up
+	void SetUpFrameCalculator(const UE::AIE::FFrameTimeByIndex& InCurrentFrameTimes, TRange<FFrameNumber> ViewRange);
+	//Add must have indices that you want to calculate first, usually start/end or key times
+	void AddMustHaveIndices(const TArray<int32>& InMustHaveIndices);
+	//reset calculated indices and bucket to zero
+	void Reset();
+
+	//calculate next set of indices in the IndicesToCalculate array
+	//will return true if we need to calculate, will return false if this is the last time
+	bool CalculateIndices();
+
+	static int32 IndicesToCalculatePerBucket; //this could be auto calculated based on frame rate, now it's also a CVAR
+	UE::AIE::FFrameTimeByIndex CurrentFrameTimes;
+	//set of ranges of time we calculate, if view range is same as eval range this will just be one range that's equal to the CurrentFrameTimesRange
+	TArray<TRange<FFrameNumber>> Ranges;
+	int32 CurrentRange;
+	int32 CurrentRangeStartIndex;
+	int32 CurrentRangeEndIndex;
+
+	int32 NumOfBuckets;
+	int32 CurrentBucket;
+	TArray<int32> IndicesToCalculate;
+	TArray<int32> MustHaveIndices;
 
+private:
+	void CalculateStartEndIndices();
+
+};
 
-class FTrajectoryDrawInfo
+// Holds set of currently calculating frames
+struct FCurrentFramesInfo
 {
-public:
+	void SetViewRange(const TRange<FFrameNumber>& InViewRange);
+	void SetUpFrameTimes(const TRange<FFrameNumber>& InEvalFrameRange, const FFrameNumber& InFrameStep);
+	void AddMustHaveTimes(const TSet<FFrameNumber>& InMustHaveTimes);
+	const TArray<int32>& IndicesToCalculate() const;
+	bool KeepCalculating();
+	void Reset();
+
+	//range and frame rate, this is fixed and will only change if evaluation range(sequencer range/display rate) changes.
+	UE::AIE::FFrameTimeByIndex CurrentFrameTimes; 
+	TRange<FFrameNumber> ViewRange;
+	bool bViewRangeIsEvalRange = false;
+	//this also contains the indices to calculate
+	FFrameCalculator  FrameCalculator;
+
+	//set of indices that slowly grow
+	TSortedMap<int32, FFrameNumber> SortedTransformIndices;
+	//cached array of the transform indices
+	TArray<int32> TransformIndices;  
+	// all trails share this now
+	TArray<FFrameNumber> CurrentFrames; 
+};
+
 
-	FTrajectoryDrawInfo(const FLinearColor& InColor, FTrajectoryCache* InTrajectoryCache)
+struct FTrajectoryDrawInfo
+{
+	FTrajectoryDrawInfo(EMotionTrailTrailStyle InStyle, const FLinearColor& InColor, TSharedPtr<UE::AIE::FArrayOfTransforms>& InArrayOfTransforms,
+		TSharedPtr<UE::AIE::FArrayOfTransforms>& InParentSpaceArrayOfTransforms)
 		: Color(InColor)
-		, CachedViewRange(TRange<double>::Empty())
-		, TrajectoryCache(InTrajectoryCache)
+		, ArrayOfTransforms(InArrayOfTransforms)
+		, ParentSpaceTransforms(InParentSpaceArrayOfTransforms)
 	{}
 
-	virtual ~FTrajectoryDrawInfo() {}
+	~FTrajectoryDrawInfo() {}
+	//new
+	void GetTrajectoryPointsForDisplay(const FTransform& OffsetTransform, const FTransform& ParentSpaceTransform, const FCurrentFramesInfo& InCurrentFramesInfo, bool bIsEvaluating,  TArray<FVector>& OutPoints, TArray<FFrameNumber>& OutFrames);
+	void GetTickPointsForDisplay(const FTransform& OffsetTransform, const FTransform& ParentSpaceTransform, const FTrailScreenSpaceTransform& ScreenSpaceTransform, const FCurrentFramesInfo & InCurrentFramesInfo, bool bIsEvaluating, TArray<FVector2D>& OutTicks, TArray<FVector2D>& OutTickTangents);
+	FVector GetPoint(const FTransform& OffsetTransform, const FTransform& ParentSpaceTransform, const FCurrentFramesInfo& InDisplayContext, const FFrameNumber& InTime);
 
-	void GetTrajectoryPointsForDisplay(const FDisplayContext& InDisplayContext, TArray<FVector>& OutPoints, TArray<double>& OutSeconds);
-	void GetTickPointsForDisplay(const FDisplayContext& InDisplayContext, TArray<FVector2D>& OutTicks, TArray<FVector2D>& OutTickTangents);
-	FVector GetPoint(double InTime);
 
-	// Optionally implemented methods
 	void SetColor(const FLinearColor& InColor) { Color = InColor; }
 	FLinearColor GetColor() const { return Color; }
+	void SetStyle(EMotionTrailTrailStyle InStyle) { Style = InStyle; }
+	EMotionTrailTrailStyle GetStyle() const { return Style; }
+
 
-	virtual const TRange<double>& GetCachedViewRange() const { return CachedViewRange; }
 protected:
+	EMotionTrailTrailStyle Style;
 	FLinearColor Color;
-	TRange<double> CachedViewRange;
-	FTrajectoryCache* TrajectoryCache;
+	TSharedPtr<UE::AIE::FArrayOfTransforms> ArrayOfTransforms;
+	TSharedPtr<UE::AIE::FArrayOfTransforms> ParentSpaceTransforms;
 };
 
 
+
 } // namespace MovieScene
 } // namespace UE
diff --git a/Engine/Source/Editor/MovieSceneTools/Private/MovieSceneToolHelpers.cpp b/Engine/Source/Editor/MovieSceneTools/Private/MovieSceneToolHelpers.cpp
--- a/Engine/Source/Editor/MovieSceneTools/Private/MovieSceneToolHelpers.cpp
+++ b/Engine/Source/Editor/MovieSceneTools/Private/MovieSceneToolHelpers.cpp
@@ -5548,11 +5548,26 @@ UMovieScene3DTransformSection* MovieSceneToolHelpers::GetTransformSection(
 	{
 		return nullptr;
 	}
+
+	UMovieScene3DTransformSection* TransformSection =
+		Cast<UMovieScene3DTransformSection>(TransformTrack->GetSectionToKey());
+	if (TransformSection)
+	{
+		return TransformSection;
+	}
+
+	TransformSection = TransformTrack->GetAllSections().Num() > 0 ? Cast<UMovieScene3DTransformSection>(TransformTrack->GetAllSections()[0]) : nullptr;
+	if (TransformSection)
+	{
+		return TransformSection;
+	}
+
+	//there really should be a section to key or at least one section, if not then we go through creating it.
 	TransformTrack->Modify();
 
 	static constexpr FFrameNumber Frame0;
 	bool bSectionAdded = false;
-	UMovieScene3DTransformSection* TransformSection =
+	TransformSection =
 		Cast<UMovieScene3DTransformSection>(TransformTrack->FindOrAddSection(Frame0, bSectionAdded));
 	if (!TransformSection)
 	{
diff --git a/Engine/Source/Editor/Sequencer/Private/SequencerEdMode.cpp b/Engine/Source/Editor/Sequencer/Private/SequencerEdMode.cpp
--- a/Engine/Source/Editor/Sequencer/Private/SequencerEdMode.cpp
+++ b/Engine/Source/Editor/Sequencer/Private/SequencerEdMode.cpp
@@ -59,11 +59,11 @@ static TAutoConsoleVariable<bool> CVarDrawMeshTrails(
 	true,
 	TEXT("Toggle to show or hide Level Sequence VR Editor trails"));
 
-//if true still use old motion trails for sequencer objects.
+//DEPRECRATED
 TAutoConsoleVariable<bool> CVarUseOldSequencerMotionTrails(
 	TEXT("Sequencer.UseOldSequencerTrails"),
 	true,
-	TEXT("If true show old motion trails, if false use new editable motion trails."));
+	TEXT("DEPRECRATED:  Will show old trails outside animation mode, otherwise will show newer version for all objects."));
 
 
 namespace UE
@@ -889,7 +889,10 @@ void FSequencerEdMode::DrawTracks3D(FPrimitiveDrawInterface* PDI)
 {
 	using namespace UE::Sequencer;
 
-	if (CVarUseOldSequencerMotionTrails->GetBool() == false)
+	//if in control rig mode(aka animation mode) we show new trails
+	const FName ControlRigEditModeModeName("EditMode.ControlRig");
+	const bool bIsInControlRigEditMode = GLevelEditorModeTools().GetActiveMode(ControlRigEditModeModeName) != nullptr;
+	if (bIsInControlRigEditMode == true)
 	{
 		return;
 	}
diff --git a/Engine/Source/Editor/Sequencer/Private/Tools/SMotionTrailOptions.cpp b/Engine/Source/Editor/Sequencer/Private/Tools/SMotionTrailOptions.cpp
--- a/Engine/Source/Editor/Sequencer/Private/Tools/SMotionTrailOptions.cpp
+++ b/Engine/Source/Editor/Sequencer/Private/Tools/SMotionTrailOptions.cpp
@@ -16,6 +16,7 @@
 #include "UObject/UObjectGlobals.h"
 #include "Widgets/Layout/SBorder.h"
 #include "Widgets/SBoxPanel.h"
+#include "GameFramework/Actor.h"
 
 #define LOCTEXT_NAMESPACE "MotionTrail"
 
@@ -42,7 +43,6 @@ void SMotionTrailOptions::Construct(const FArguments& InArgs)
 			SNew(SBorder)
 			.HAlign(HAlign_Fill)
 			.VAlign(VAlign_Fill)
-			.Padding(FMargin(10.0,5.0,10.0,5.0))
 			[
 			SNew(SVerticalBox)
 			
@@ -57,5 +57,137 @@ void SMotionTrailOptions::Construct(const FArguments& InArgs)
 		];
 }
 
+int32 UMotionTrailToolOptions::GetNumPinned() const
+{
+	return PinnedTrails.Num();
+}
+
+UMotionTrailToolOptions::FPinnedTrail* UMotionTrailToolOptions::GetPinnedTrail(int32 Index)  
+{
+	if (Index >= 0 && Index < PinnedTrails.Num())
+	{
+		return &(PinnedTrails[Index]);
+	}
+	return nullptr;
+}
+
+void UMotionTrailToolOptions::ResetPinnedItems()
+{
+	PinnedTrails.Reset();
+}
+
+void UMotionTrailToolOptions::AddPinned(const FPinnedTrail& InPinnedTrail)
+{
+	if (UMotionTrailToolOptions::FPinnedTrail* Trail = PinnedTrails.FindByKey(InPinnedTrail.TrailGuid))
+	{
+		return;
+	}
+	PinnedTrails.Add(InPinnedTrail);
+	OnAddPinned.Broadcast(InPinnedTrail.TrailGuid);
+}
+
+int32 UMotionTrailToolOptions::GetIndexFromGuid(FGuid InGuid) const
+{
+	return PinnedTrails.IndexOfByKey(InGuid);
+}
+
+void UMotionTrailToolOptions::PinSelection() const
+{
+	OnPinSelection.Broadcast(); //motion trails will handle this
+}
+
+void UMotionTrailToolOptions::PinComponent(USceneComponent* InSceneComponent, const FName& InSocketName) const
+{
+	OnPinComponent.Broadcast(InSceneComponent, InSocketName);
+}
+
+void UMotionTrailToolOptions::DeletePinned(int32 Index)
+{
+	if (Index >= 0 && Index < PinnedTrails.Num())
+	{
+		FGuid Guid = PinnedTrails[Index].TrailGuid;
+		PinnedTrails.RemoveAt(Index);
+		OnDeletePinned.Broadcast(Guid);
+	}
+}
+
+void UMotionTrailToolOptions::PutPinnnedInSpace(int32 Index, AActor* InActor, const FName& InComponentName)
+{
+	if (Index >= 0 && Index < PinnedTrails.Num())
+	{
+		FGuid Guid = PinnedTrails[Index].TrailGuid;
+		if (InActor)
+		{
+			PinnedTrails[Index].SpaceName = FText::FromString(InActor->GetActorLabel());
+		}
+		else
+		{
+			PinnedTrails[Index].SpaceName.Reset();
+		}
+		OnPutPinnedInSpace.Broadcast(Guid,InActor,InComponentName);
+	}
+}
+
+void UMotionTrailToolOptions::SetLinearColor(int32 Index, const FLinearColor& Color)
+{
+	if (Index >= 0 && Index < PinnedTrails.Num())
+	{
+		PinnedTrails[Index].TrailColor = Color;
+		OnSetLinearColor.Broadcast(PinnedTrails[Index].TrailGuid, Color);
+	}
+}
+
+void UMotionTrailToolOptions::SetHasOffset(int32 Index, bool bHasOffset)
+{
+	if (Index >= 0 && Index < PinnedTrails.Num())
+	{
+		PinnedTrails[Index].bHasOffset = bHasOffset;
+		OnSetHasOffset.Broadcast(PinnedTrails[Index].TrailGuid, bHasOffset);
+	}
+}
+
+
+TArray<TPair<FText, FText>>& UMotionTrailToolOptions::GetTrailStyles()  
+{
+	if (TrailStylesText.Num() == 0)
+	{
+		TPair<FText, FText> TextToAdd;
+		TextToAdd.Key = (LOCTEXT("Default", "Default"));
+		TextToAdd.Value = (LOCTEXT("DefaultTooltip", "Use specified single trail color"));
+		TrailStylesText.Add(TextToAdd);
+
+		TextToAdd.Key = (LOCTEXT("Dashed", "Dashed"));
+		TextToAdd.Value = (LOCTEXT("DashedTooltip", "Alternate color every other frame"));
+		TrailStylesText.Add(TextToAdd);
+
+		TextToAdd.Key = (LOCTEXT("Time", "Time"));
+		TextToAdd.Value = (LOCTEXT("TimeTooltip", "Alternate color before and after current Sequencer time"));
+		TrailStylesText.Add(TextToAdd);
+
+		TextToAdd.Key = (LOCTEXT("HeatMap", "Heat Map"));
+		TextToAdd.Value = (LOCTEXT("HeatMapTooltip", "Color shows speed from fast(Red) to slow(Blue)"));
+		TrailStylesText.Add(TextToAdd);
+
+	}
+	return TrailStylesText;
+}
+
+void UMotionTrailToolOptions::SetTrailStyle(int32 Index)
+{
+	if (Index >= 0 && Index < 4)
+	{
+		TrailStyle = (EMotionTrailTrailStyle)(uint8)(Index);
+		FPropertyChangedEvent Event(UMotionTrailToolOptions::StaticClass()->FindPropertyByName(GET_MEMBER_NAME_CHECKED(UMotionTrailToolOptions, TrailStyle)));
+		PostEditChangeProperty(Event);
+	}
+}
+
+int32 UMotionTrailToolOptions::GetTrailStyleIndex() const
+{
+	return (int32)TrailStyle;
+}
+
+
+
 
 #undef LOCTEXT_NAMESPACE
diff --git a/Engine/Source/Editor/Sequencer/Public/Tools/MotionTrailOptions.h b/Engine/Source/Editor/Sequencer/Public/Tools/MotionTrailOptions.h
--- a/Engine/Source/Editor/Sequencer/Public/Tools/MotionTrailOptions.h
+++ b/Engine/Source/Editor/Sequencer/Public/Tools/MotionTrailOptions.h
@@ -3,12 +3,26 @@
 #pragma once
 #include "CoreMinimal.h"
 #include "HAL/IConsoleManager.h"
+#include "Delegates/Delegate.h"
 #include "MotionTrailOptions.generated.h"
 
 //if true still use old motion trails for sequencer objects.
 extern SEQUENCER_API TAutoConsoleVariable<bool> CVarUseOldSequencerMotionTrails;
 
 
+class USceneComponent;
+class AActor;
+
+UENUM(BlueprintType)
+enum class EMotionTrailTrailStyle : uint8
+{
+	Default = 0,
+	Dashed = 1,
+	Time = 2,
+	HeatMap = 3
+};
+
+
 // TODO: option to make tick size proportional to distance from camera to get a sense of perspective and scale
 UCLASS(config = EditorPerProjectUserSettings)
 class SEQUENCER_API UMotionTrailToolOptions : public UObject
@@ -18,97 +32,184 @@ class SEQUENCER_API UMotionTrailToolOptions : public UObject
 public:
 	UMotionTrailToolOptions()
 		: bShowTrails(false)
-		, TrailColor(.22,0.15,1.0)
+		, bShowSelectedTrails(true)
+		, TrailStyle(EMotionTrailTrailStyle::Default)
+		, DefaultColor(.22,0.15,1.0)
+		, TimePreColor(.22, 0.35, 0.8)
+		, TimePostColor(.85, 0.25, 0.1)
+		, DashPreColor(0.2, 0.9, 0.3)
+		, DashPostColor(0.7, 0.2, 0.7)
 		, bShowFullTrail(true)
 		, TrailThickness(0.0f)
 		, FramesBefore(10)
 		, FramesAfter(10)
 		, EvalsPerFrame(1)
 		, bShowKeys(true)
-		, bShowFrameNumber(true)
+		, bShowFrameNumber_DEPRECATED(false)
 		, KeyColor(1.0, 1.0, 1.0)
-		, KeySize(10.0f)
+		, KeySize(3.0f)
 		, bShowMarks(false)
 		, MarkColor(0.25,1.0,0.15)
 		, MarkSize(5.0)
-		, bLockMarksToFrames(true)
-		, SecondsPerMark(0.1)
+		, bLockMarksToFrames_DEPRECATED(true)
+		, SecondsPerMark_DEPRECATED(0.1)
 	{}
 
 	/** Whether or not to show motion trails */
 	UPROPERTY(EditAnywhere, Category = Trail);
 	bool bShowTrails;
 
+	/** Whether or not to show selected motion trails */
+	UPROPERTY(EditAnywhere, Config, Category = Trail);
+	bool bShowSelectedTrails;
+
+	/** How To Show Color*/
+	UPROPERTY(EditAnywhere, Config, Category = Trail);
+	EMotionTrailTrailStyle TrailStyle;
+
 	/** The color of the motion trail */
-	UPROPERTY(EditAnywhere,  Category = Trail)
-	FLinearColor TrailColor;
+	UPROPERTY(EditAnywhere, Config, Category = AdvancedDisplay)
+	FLinearColor DefaultColor;
+
+	/** The color of the motion trail before current time if show alternating time colors */
+	UPROPERTY(EditAnywhere, Config, Category = AdvancedDisplay)
+	FLinearColor TimePreColor;
+
+	/** The color of the motion trail after current time if show alternating time colors */
+	UPROPERTY(EditAnywhere, Config, Category = AdvancedDisplay)
+	FLinearColor TimePostColor;
+
+	/** The color of the first motion trail color when alternating between frames */
+	UPROPERTY(EditAnywhere, Config, Category = AdvancedDisplay)
+	FLinearColor DashPreColor;
+
+	/** The color of the next motion trail color when alternating between frames */
+	UPROPERTY(EditAnywhere, Config, Category = AdvancedDisplay)
+	FLinearColor DashPostColor;
 
 	/** Whether or not to show the full motion trail */
-	UPROPERTY(EditAnywhere, Category = Trail)
+	UPROPERTY(EditAnywhere, Category = Trail);
 	bool bShowFullTrail;
 
 	/* The thickness of the motion trail */
-	UPROPERTY(EditAnywhere, Category = Trail, Meta = (ClampMin = "0.0"))
-	float TrailThickness;
+	UPROPERTY(EditAnywhere, Config, Category = Trail, Meta = (ClampMin = "0.0"))
+	double TrailThickness;
 
 	/** The number of frames to draw before the start of the trail. Requires not showing the full trail */
-	UPROPERTY(EditAnywhere, Category = Trail, Meta = (EditCondition = "!bShowFullTrail", ClampMin = "0"))
+	UPROPERTY(EditAnywhere, Config, Category = Trail, Meta = (EditCondition = "!bShowFullTrail", ClampMin = "0"))
 	int32 FramesBefore;
 
 	/** The number of frames to draw after the end of the trail. Requires not showing the full trail */
-	UPROPERTY(EditAnywhere, Category = Trail, Meta = (EditCondition = "!bShowFullTrail", ClampMin = "0"))
+	UPROPERTY(EditAnywhere, Config, Category = Trail, Meta = (EditCondition = "!bShowFullTrail", ClampMin = "0"))
 	int32 FramesAfter;
 
 	/** No longer exposed and clamped to 1 The number of evaluations per frame */
 	int32 EvalsPerFrame;
 	
 	/** Whether or not to show keys on the motion trail */
-	UPROPERTY(EditAnywhere, Category = Keys)
+	UPROPERTY(EditAnywhere, Config, Category = AdvancedDisplay)
 	bool bShowKeys;
 
-	/** Whether or not to show frame numbers for the keys on the motion trail */
-	UPROPERTY(EditAnywhere, Category = Keys)
-	bool bShowFrameNumber;
+	/** Deprecated in 5.6 */
+	UPROPERTY()
+	bool bShowFrameNumber_DEPRECATED;
 
 	/** The color of the keys */
 	UPROPERTY(EditAnywhere, Category = Keys)
 	FLinearColor KeyColor;
 
 	/** The size of the keys */
-	UPROPERTY(EditAnywhere, Category = Keys, Meta = (ClampMin = "0.0"))
+	UPROPERTY(EditAnywhere, Config, Category = Keys, Meta = (ClampMin = "0.0"))
 	double KeySize;
 
 	/** Whether or not to show marks along the motion trail */
-	UPROPERTY(EditAnywhere, Category = Marks)
+	UPROPERTY(EditAnywhere, Config, Category = Marks)
 	bool bShowMarks;
 
 	/** The color of the marks */
-	UPROPERTY(EditAnywhere, Category = Marks)
+	UPROPERTY(EditAnywhere, Config, Category = Marks)
 	FLinearColor MarkColor;
 
 	/** The size of the marks */
-	UPROPERTY(EditAnywhere, Category = Marks, Meta = (ClampMin = "0.0"))
+	UPROPERTY(EditAnywhere, Config, Category = Marks, Meta = (ClampMin = "0.0"))
 	double MarkSize;
 
-	/** Whether or not to lock the marks to the frames */
-	UPROPERTY(EditAnywhere, Category = Marks)
-	bool bLockMarksToFrames;
 
-	/** The seconds per mark */
-	UPROPERTY(EditAnywhere, Category = Marks, Meta = (EditCondition = "!bLockMarksToFrames", ClampMin = "0.01"))
-	double SecondsPerMark;
+	/** Deprecated in 5.6 */
+	UPROPERTY()
+	bool bLockMarksToFrames_DEPRECATED;
+
+	/** Deprecated in 5.6 */
+	UPROPERTY()
+	double SecondsPerMark_DEPRECATED;
 
 	virtual void PostEditChangeProperty(struct FPropertyChangedEvent& PropertyChangedEvent) override
 	{
 		FName PropertyName = (PropertyChangedEvent.Property != NULL) ? PropertyChangedEvent.Property->GetFName() : NAME_None;
 		OnDisplayPropertyChanged.Broadcast(PropertyName);
+		SaveConfig();
 	}
 
 	DECLARE_MULTICAST_DELEGATE_OneParam(FOnDisplayPropertyChanged, FName);
 	FOnDisplayPropertyChanged OnDisplayPropertyChanged;
 
-
 	static UMotionTrailToolOptions* GetTrailOptions()  { return GetMutableDefault<UMotionTrailToolOptions>(); }
 
+	//We put Pinned motion trail information as part of this object so it can be shared by SequencerAnimTools and Control Rig module
+	struct FPinnedTrail
+	{
+		FText TrailName;
+		FLinearColor TrailColor;
+		bool bHasOffset = false;
+		TOptional<FText> SpaceName;
+		FGuid TrailGuid;
+
+		bool operator==(const FGuid& InId) const { return TrailGuid == InId; }
+
+	};
+	void ResetPinnedItems();
+	int32 GetNumPinned() const;
+	FPinnedTrail* GetPinnedTrail(int32 Index);
+
+	int32 GetIndexFromGuid(FGuid InGuid) const;
+	void PinSelection() const;
+	void AddPinned(const FPinnedTrail& InPinnedTrail);
+	void PinComponent(USceneComponent* InSceneComponent, const FName& InSocketName) const;
+	void DeletePinned(int32 Index);
+	void PutPinnnedInSpace(int32 Index, AActor* InActor, const FName& InComponentName);
+	void SetLinearColor(int32 Index, const FLinearColor& Color);
+	void SetHasOffset(int32 Index, bool bHasOffset);
+	//get name and tooltip name for the color types
+	TArray<TPair<FText, FText>>& GetTrailStyles();
+	void SetTrailStyle(int32 Index);
+	int32 GetTrailStyleIndex() const;
+
+	DECLARE_MULTICAST_DELEGATE(FOnPinSelection);
+	FOnPinSelection OnPinSelection;
+
+	DECLARE_MULTICAST_DELEGATE_OneParam(FOnAddPinned, FGuid);
+	FOnAddPinned OnAddPinned;
+
+	DECLARE_MULTICAST_DELEGATE_TwoParams(FPinComponent, USceneComponent*, FName);
+	FPinComponent OnPinComponent;
+
+	DECLARE_MULTICAST_DELEGATE_OneParam(FOnDeletePinned, FGuid);
+	FOnDeletePinned OnDeletePinned;
+
+	DECLARE_MULTICAST_DELEGATE_ThreeParams(FOnPutPinnedInSpace, FGuid, AActor*, FName);
+	FOnPutPinnedInSpace OnPutPinnedInSpace;
+
+	DECLARE_MULTICAST_DELEGATE_TwoParams(FOnSetLinearColor, FGuid, FLinearColor);
+	FOnSetLinearColor OnSetLinearColor;
+
+	DECLARE_MULTICAST_DELEGATE_TwoParams(FOnSetHasOffset, FGuid, bool);
+	FOnSetHasOffset OnSetHasOffset;
+
+protected:
+	TArray<FPinnedTrail> PinnedTrails;
+	TArray<TPair<FText, FText>> TrailStylesText;
+
+};
+
+
 
-};
\ No newline at end of file
