Sequencer Scripting: Add GetSections to the ScriptingObject so you can retrieve the sections given the nodes. This allows you to get the sections on a track row as well.

diff --git a/Engine/Source/Editor/Sequencer/Private/Scripting/SequencerModuleOutlinerScriptingObject.cpp b/Engine/Source/Editor/Sequencer/Private/Scripting/SequencerModuleOutlinerScriptingObject.cpp
--- a/Engine/Source/Editor/Sequencer/Private/Scripting/SequencerModuleOutlinerScriptingObject.cpp
+++ b/Engine/Source/Editor/Sequencer/Private/Scripting/SequencerModuleOutlinerScriptingObject.cpp
@@ -13,8 +13,42 @@
 #include "Containers/Array.h"
 #include "Misc/FrameNumber.h"
 
+#include "SequencerCommonHelpers.h"
+
 #define LOCTEXT_NAMESPACE "SequencerModuleOutlinerScriptingObject"
 
+TArray<UMovieSceneSection*> USequencerModuleOutlinerScriptingObject::GetSections(const TArray<FSequencerViewModelScriptingStruct>& InNodes) const
+{
+	using namespace UE::Sequencer;
+
+	TViewModelPtr<FSequencerOutlinerViewModel> Outliner = CastViewModel<FSequencerOutlinerViewModel>(WeakOutliner.Pin());
+	if (!Outliner)
+	{
+		FFrame::KismetExecutionMessage(*LOCTEXT("OutlinerInvalid", "Outliner is no longer valid.").ToString(), ELogVerbosity::Error);
+		return TArray<UMovieSceneSection*>();
+	}
+
+	TSet<TWeakObjectPtr<UMovieSceneSection> > AllSections;
+	TArray<TSharedRef<FViewModel>> Nodes;
+	for (const FSequencerViewModelScriptingStruct& InNode : InNodes)
+	{
+		FViewModelPtr Item = InNode.WeakViewModel.ImplicitPin();
+		if (Item)
+		{
+			SequencerHelpers::GetAllSections(Item, AllSections);
+		}
+	}
+
+	TArray<UMovieSceneSection*> Sections;
+	for (TWeakObjectPtr<UMovieSceneSection> WeakSection : AllSections)
+	{
+		if (WeakSection.IsValid())
+		{
+			Sections.Add(WeakSection.Get());
+		}
+	}
+	return Sections;
+}
 
 FFrameNumber USequencerModuleOutlinerScriptingObject::GetNextKey(const TArray<FSequencerViewModelScriptingStruct>& InNodes, FFrameNumber FrameNumber, EMovieSceneTimeUnit TimeUnit) const
 {
diff --git a/Engine/Source/Editor/Sequencer/Private/Scripting/SequencerModuleOutlinerScriptingObject.h b/Engine/Source/Editor/Sequencer/Private/Scripting/SequencerModuleOutlinerScriptingObject.h
--- a/Engine/Source/Editor/Sequencer/Private/Scripting/SequencerModuleOutlinerScriptingObject.h
+++ b/Engine/Source/Editor/Sequencer/Private/Scripting/SequencerModuleOutlinerScriptingObject.h
@@ -12,6 +12,8 @@
 struct FFrameNumber;
 struct FSequencerViewModelScriptingStruct;
 
+class UMovieSceneSection;
+
 UCLASS()
 class USequencerModuleOutlinerScriptingObject : public USequencerOutlinerScriptingObject
 {
@@ -19,6 +21,9 @@ public:
 
 	GENERATED_BODY()
 
+	UFUNCTION(BlueprintCallable, Category = "Sequencer Editor")
+	TArray<UMovieSceneSection*> GetSections(const TArray<FSequencerViewModelScriptingStruct>& InNodes) const;
+
 	UFUNCTION(BlueprintCallable, Category = "Sequencer Editor")
 	FFrameNumber GetNextKey(const TArray<FSequencerViewModelScriptingStruct>& InNodes, FFrameNumber FrameNumber, EMovieSceneTimeUnit TimeUnit = EMovieSceneTimeUnit::DisplayRate) const;
 
