Sequencer: Fix binding label not being updated when using Assign Actor->Add Selected or Replace with Selected

The added object wasn't being correctly tracked before continuing

diff --git a/Engine/Source/Editor/Sequencer/Private/SequencerUtilities.cpp b/Engine/Source/Editor/Sequencer/Private/SequencerUtilities.cpp
--- a/Engine/Source/Editor/Sequencer/Private/SequencerUtilities.cpp
+++ b/Engine/Source/Editor/Sequencer/Private/SequencerUtilities.cpp
@@ -4180,6 +4180,7 @@ void FSequencerUtilities::AddObjectsToBinding(TSharedRef<ISequencer> InSequencer
 		FGuid AddedGuid = InSequencer->GetHandleToObject(ObjectToAdd, false);
 		if (AddedGuid.IsValid())
 		{
+			AddedObjects.Add(ObjectToAdd);
 			continue;
 		}
 
