[Sequencer] Minor fix, avoid lambda capture uobject directly, using a weakobjectptr instead, was crashing when asset reloading a level sequence multiple times

diff --git a/Engine/Source/Editor/Sequencer/Private/SequencerCommonHelpers.cpp b/Engine/Source/Editor/Sequencer/Private/SequencerCommonHelpers.cpp
--- a/Engine/Source/Editor/Sequencer/Private/SequencerCommonHelpers.cpp
+++ b/Engine/Source/Editor/Sequencer/Private/SequencerCommonHelpers.cpp
@@ -396,9 +396,9 @@ void SequencerHelpers::BuildEditSectionMenu(const TWeakPtr<FSequencer>& InWeakSe
 				return MakeShared<FMovieSceneSectionDetailsCustomization>(TypeInterface, CurrentScene.Get());
 			}));
 
-		DetailsView->RegisterInstancedCustomPropertyTypeLayout("MovieSceneConditionContainer", FOnGetPropertyTypeCustomizationInstance::CreateLambda([Sequence, InWeakSequencer]() 
+		DetailsView->RegisterInstancedCustomPropertyTypeLayout("MovieSceneConditionContainer", FOnGetPropertyTypeCustomizationInstance::CreateLambda([WeakSequence = TWeakObjectPtr(Sequence), InWeakSequencer]() 
 		{
-			return FMovieSceneConditionCustomization::MakeInstance(Sequence, InWeakSequencer); 
+			return FMovieSceneConditionCustomization::MakeInstance(WeakSequence, InWeakSequencer); 
 		}));
 
 		DetailsView->RegisterInstancedCustomPropertyTypeLayout("MovieSceneDirectorBlueprintConditionData", FOnGetPropertyTypeCustomizationInstance::CreateLambda([=]() {
