Sequencer: Follow up to 40144626 which deletes duplicate keys when selection is emptied

On deleting keys, selection changes and the original selected index could be invalid. Need to check that the index is valid before getting key times for that key handle.

#rnx

diff --git a/Engine/Source/Editor/Sequencer/Private/MVVM/Selection.cpp b/Engine/Source/Editor/Sequencer/Private/MVVM/Selection.cpp
--- a/Engine/Source/Editor/Sequencer/Private/MVVM/Selection.cpp
+++ b/Engine/Source/Editor/Sequencer/Private/MVVM/Selection.cpp
@@ -27,6 +27,11 @@ void FKeySelection::RemoveDuplicateKeys(TArrayView<FKeyHandle> KeyHandles)
 		{
 			if (FMovieSceneChannel* Channel = ChannelModel->GetChannel())
 			{
+				if (Channel->GetIndex(KeyHandle) == INDEX_NONE)
+				{
+					continue;
+				}
+
 				FFrameNumber KeyTime;
 				Channel->GetKeyTime(KeyHandle, KeyTime);
 
