Sequencer: Fix keys disappearing when expanding a track that has no children

Only return IsExpanded() true if there are any children and no need to do anything in SetExpansion if there are no children

diff --git a/Engine/Source/Editor/Sequencer/Private/MVVM/ViewModels/OutlinerItemModel.cpp b/Engine/Source/Editor/Sequencer/Private/MVVM/ViewModels/OutlinerItemModel.cpp
--- a/Engine/Source/Editor/Sequencer/Private/MVVM/ViewModels/OutlinerItemModel.cpp
+++ b/Engine/Source/Editor/Sequencer/Private/MVVM/ViewModels/OutlinerItemModel.cpp
@@ -199,35 +199,42 @@ void FOutlinerItemModelMixin::SetIdentifier(FName InNewIdentifier)
 
 bool FOutlinerItemModelMixin::IsExpanded() const
 {
-	if (bInitializedExpansion)
-	{
-		return bIsExpanded;
-	}
+	const FViewModel* ViewModel = AsViewModel();
 
-	bInitializedExpansion = true;
+	if (!bInitializedExpansion)
+	{
+		bInitializedExpansion = true;
 
-	TStringBuilder<256> StringBuilder;
-	IOutlinerExtension::GetPathName(*AsViewModel(), StringBuilder);
+		TStringBuilder<256> StringBuilder;
+		IOutlinerExtension::GetPathName(*ViewModel, StringBuilder);
 
-	TSharedPtr<FSequenceModel> SequenceModel = AsViewModel()->FindAncestorOfType<FSequenceModel>();
-	UMovieSceneSequence*       Sequence      = SequenceModel ? SequenceModel->GetSequence() : nullptr;
+		TSharedPtr<FSequenceModel> SequenceModel = ViewModel->FindAncestorOfType<FSequenceModel>();
+		UMovieSceneSequence*       Sequence      = SequenceModel ? SequenceModel->GetSequence() : nullptr;
 
-	if (Sequence)
-	{
-		FStringView StringView = StringBuilder.ToView();
-		UMovieScene* MovieScene = Sequence->GetMovieScene();
-		FMovieSceneEditorData& EditorData = MovieScene->GetEditorData();
-		if (const FMovieSceneExpansionState* Expansion = EditorData.ExpansionStates.FindByHash(GetTypeHash(StringView), StringView))
+		if (Sequence)
 		{
-			const_cast<FOutlinerItemModelMixin*>(this)->bIsExpanded = Expansion->bExpanded;
-		}
-		else
-		{
-			const_cast<FOutlinerItemModelMixin*>(this)->bIsExpanded = GetDefaultExpansionState();
+			FStringView StringView = StringBuilder.ToView();
+			UMovieScene* MovieScene = Sequence->GetMovieScene();
+			FMovieSceneEditorData& EditorData = MovieScene->GetEditorData();
+			if (const FMovieSceneExpansionState* Expansion = EditorData.ExpansionStates.FindByHash(GetTypeHash(StringView), StringView))
+			{
+				const_cast<FOutlinerItemModelMixin*>(this)->bIsExpanded = Expansion->bExpanded;
+			}
+			else
+			{
+				const_cast<FOutlinerItemModelMixin*>(this)->bIsExpanded = GetDefaultExpansionState();
+			}
 		}
 	}
 
-	return bIsExpanded;
+	if (bIsExpanded)
+	{
+		// If there are no children, no need to allow this to be expanded
+		const bool bHasChildren = ViewModel->GetChildren(EViewModelListType::Outliner).ToArray().Num() > 0;
+		return bHasChildren;
+	}
+
+	return false;
 }
 
 bool FOutlinerItemModelMixin::GetDefaultExpansionState() const
@@ -237,9 +244,17 @@ bool FOutlinerItemModelMixin::GetDefaultExpansionState() const
 
 void FOutlinerItemModelMixin::SetExpansion(bool bInIsExpanded)
 {
+	FViewModel* ViewModel = AsViewModel();
+
+	// If no children, there's no need to set this expanded
+	const bool bHasChildren = ViewModel->GetChildren(EViewModelListType::Outliner).ToArray().Num() > 0;
+	if (!bHasChildren)
+	{
+		return;
+	}
+
 	SetExpansionWithoutSaving(bInIsExpanded);
 
-	FViewModel* ViewModel = AsViewModel();
 	if (ViewModel->GetParent())
 	{
 		// Expansion state has changed, save it to the movie scene now
