Sequencer: Follow up to 40432460. No need to convert to array

diff --git a/Engine/Source/Editor/Sequencer/Private/MVVM/ViewModels/OutlinerItemModel.cpp b/Engine/Source/Editor/Sequencer/Private/MVVM/ViewModels/OutlinerItemModel.cpp
--- a/Engine/Source/Editor/Sequencer/Private/MVVM/ViewModels/OutlinerItemModel.cpp
+++ b/Engine/Source/Editor/Sequencer/Private/MVVM/ViewModels/OutlinerItemModel.cpp
@@ -230,8 +230,11 @@ bool FOutlinerItemModelMixin::IsExpanded() const
 	if (bIsExpanded)
 	{
 		// If there are no children, no need to allow this to be expanded
-		const bool bHasChildren = ViewModel->GetChildren(EViewModelListType::Outliner).ToArray().Num() > 0;
-		return bHasChildren;
+		FViewModelListIterator OutlinerChildren = ViewModel->GetChildren(EViewModelListType::Outliner);
+		if (OutlinerChildren)
+		{
+			return true;
+		}
 	}
 
 	return false;
@@ -247,8 +250,8 @@ void FOutlinerItemModelMixin::SetExpansion(bool bInIsExpanded)
 	FViewModel* ViewModel = AsViewModel();
 
 	// If no children, there's no need to set this expanded
-	const bool bHasChildren = ViewModel->GetChildren(EViewModelListType::Outliner).ToArray().Num() > 0;
-	if (!bHasChildren)
+	FViewModelListIterator OutlinerChildren = ViewModel->GetChildren(EViewModelListType::Outliner);
+	if (!OutlinerChildren)
 	{
 		return;
 	}
