Sequencer: By default, the curve editor should not be visible until the user opens it

If the user opens the curve editor, closes/reopens the editor, the curve editor should be visible.

diff --git a/Engine/Source/Editor/Sequencer/Private/SequencerSettings.cpp b/Engine/Source/Editor/Sequencer/Private/SequencerSettings.cpp
--- a/Engine/Source/Editor/Sequencer/Private/SequencerSettings.cpp
+++ b/Engine/Source/Editor/Sequencer/Private/SequencerSettings.cpp
@@ -36,7 +36,7 @@ USequencerSettings::USequencerSettings( const FObjectInitializer& ObjectInitiali
 	bLinkCurveEditorTimeRange = false;
 	bSynchronizeCurveEditorSelection = true;
 	bIsolateCurveEditorToSelection = true;
-	bCurveEditorVisible = true;
+	bCurveEditorVisible = false;
 	LoopMode = ESequencerLoopMode::SLM_NoLoop;
 	bSnapKeysAndSectionsToPlayRange = false;
 	bResetPlayheadWhenNavigating = false;
