Sequencer: Remove DeleteKeys command because it confused people who rebound it

diff --git a/Engine/Source/Editor/Sequencer/Private/Sequencer.cpp b/Engine/Source/Editor/Sequencer/Private/Sequencer.cpp
--- a/Engine/Source/Editor/Sequencer/Private/Sequencer.cpp
+++ b/Engine/Source/Editor/Sequencer/Private/Sequencer.cpp
@@ -11719,9 +11719,6 @@ void FSequencer::BindCommands()
 		Commands.SnapToFrame,
 		FExecuteAction::CreateSP(this, &FSequencer::SnapToFrame),
 		FCanExecuteAction::CreateSP(this, &FSequencer::CanSnapToFrame));
-	SequencerCommandBindings->MapAction(
-		Commands.DeleteKeys,
-		FExecuteAction::CreateSP(this, &FSequencer::DeleteSelectedKeys));
 
 	SequencerCommandBindings->MapAction(
 		Commands.TogglePilotCamera,
diff --git a/Engine/Source/Editor/Sequencer/Private/SequencerCommands.cpp b/Engine/Source/Editor/Sequencer/Private/SequencerCommands.cpp
--- a/Engine/Source/Editor/Sequencer/Private/SequencerCommands.cpp
+++ b/Engine/Source/Editor/Sequencer/Private/SequencerCommands.cpp
@@ -187,7 +187,6 @@ void FSequencerCommands::RegisterCommands()
 	UI_COMMAND( SetKeyTime, "Set Key Time", "Set the key to a specified time", EUserInterfaceActionType::Button, FInputChord() );
 	UI_COMMAND( Rekey, "Rekey", "Set the selected key's time to the current time", EUserInterfaceActionType::Button, FInputChord() );
 	UI_COMMAND( SnapToFrame, "Snap To Frame", "Snap selected keys to frame", EUserInterfaceActionType::Button, FInputChord() );
-	UI_COMMAND( DeleteKeys, "Delete Keys", "Deletes the selected keys", EUserInterfaceActionType::Button, FInputChord() );
 
 	UI_COMMAND(TogglePilotCamera, "Pilot Camera", "Toggle piloting the last camera or the camera cut camera.", EUserInterfaceActionType::ToggleButton, FInputChord(EModifierKey::Shift, EKeys::P));
 
diff --git a/Engine/Source/Editor/Sequencer/Private/SequencerContextMenus.cpp b/Engine/Source/Editor/Sequencer/Private/SequencerContextMenus.cpp
--- a/Engine/Source/Editor/Sequencer/Private/SequencerContextMenus.cpp
+++ b/Engine/Source/Editor/Sequencer/Private/SequencerContextMenus.cpp
@@ -201,7 +201,15 @@ void FKeyContextMenu::PopulateMenu(FMenuBuilder& MenuBuilder, TSharedPtr<FExtend
 
 		FToolMenuSection& KeysRemovalSection = KeyContextMenu->AddSection("KeyRemoval");
 		KeysRemovalSection.AddSeparator(NAME_None);
-		KeysRemovalSection.AddMenuEntry(FSequencerCommands::Get().DeleteKeys);
+
+		// This used to be a command in FSequencerCommands; removed because people were rebinding it and confused that delete, bound to FGenericCommands::Delete, still deleted keys.
+		KeysRemovalSection.AddMenuEntry(
+			NAME_None,
+			LOCTEXT("DeleteKeys.Label", "Delete Keys"),
+			LOCTEXT("DeleteKeys.Description", "Deletes the selected keys"),
+			FSlateIcon(),
+			FUIAction(FExecuteAction::CreateSP(WeakSequencer.Pin().Get(), &FSequencer::DeleteSelectedKeys))
+			);
 	}
 
 	USequencerToolMenuContext* ContextObject = NewObject<USequencerToolMenuContext>();
diff --git a/Engine/Source/Editor/Sequencer/Public/SequencerCommands.h b/Engine/Source/Editor/Sequencer/Public/SequencerCommands.h
--- a/Engine/Source/Editor/Sequencer/Public/SequencerCommands.h
+++ b/Engine/Source/Editor/Sequencer/Public/SequencerCommands.h
@@ -457,9 +457,6 @@ public:
 	/** Snap selected keys to frame */
 	TSharedPtr< FUICommandInfo > SnapToFrame;
 
-	/** Deletes the selected keys */
-	TSharedPtr< FUICommandInfo > DeleteKeys;
-
 	/** Toggle piloting the last camera or the camera cut camera */
 	TSharedPtr< FUICommandInfo > TogglePilotCamera;
 
