Sequencer: Fix dilation factor when resizing sections

diff --git a/Engine/Source/Editor/Sequencer/Private/Tools/EditToolDragOperations.cpp b/Engine/Source/Editor/Sequencer/Private/Tools/EditToolDragOperations.cpp
--- a/Engine/Source/Editor/Sequencer/Private/Tools/EditToolDragOperations.cpp
+++ b/Engine/Source/Editor/Sequencer/Private/Tools/EditToolDragOperations.cpp
@@ -390,7 +390,7 @@ void FResizeSection::OnDrag(const FPointerEvent& MouseEvent, FVector2D LocalMous
 			FFrameNumber NewPosition    = bDraggingByEnd ? FMath::Clamp(StartPosition + DeltaTime, DilationOrigin, RightMovementMaximum.Get(TNumericLimits<int32>::Max()))
 														 : FMath::Clamp(StartPosition + DeltaTime, LeftMovementMaximum.Get(TNumericLimits<int32>::Lowest()), DilationOrigin);
 
-			float DilationFactor = FMath::Abs(NewPosition.Value - DilationOrigin.Value) / float(UE::MovieScene::DiscreteSize(DataRange));
+			float DilationFactor = FMath::Abs(NewPosition.Value - DilationOrigin.Value) / static_cast<float>(DataRange.Size<FFrameNumber>().Value);
 
 			if (bDraggingByEnd)
 			{
