Sequencer: Fixed inconsistency between virtual and physical alignment with layer bars

The error here was using Slate widget ComputeDesiredSize to constrain the size of the track area widgets for layer bars, which was changed when Binding Lifetime tracks warranted changing the track lane to a nested lane that encompassed all children. The correct approach is to bake the size into the arranged size returned by FLayerBarModel::ArrangeVirtualTrackLaneView to ensure that it matches virtual sizing computations (such as marquee selection) as well.

diff --git a/Engine/Source/Editor/Sequencer/Private/MVVM/ViewModels/LayerBarModel.cpp b/Engine/Source/Editor/Sequencer/Private/MVVM/ViewModels/LayerBarModel.cpp
--- a/Engine/Source/Editor/Sequencer/Private/MVVM/ViewModels/LayerBarModel.cpp
+++ b/Engine/Source/Editor/Sequencer/Private/MVVM/ViewModels/LayerBarModel.cpp
@@ -84,7 +84,21 @@ TSharedPtr<ITrackLaneWidget> FLayerBarModel::CreateTrackLaneView(const FCreateTr
 FTrackLaneVirtualAlignment FLayerBarModel::ArrangeVirtualTrackLaneView() const
 {
 	TRange<FFrameNumber> Range = ComputeRange();
-	return FTrackLaneVirtualAlignment::Proportional(Range, .8f, EVerticalAlignment::VAlign_Top);
+
+	FTrackLaneVirtualAlignment Result = FTrackLaneVirtualAlignment::Fixed(Range, 10.f, VAlign_Top);
+
+	TViewModelPtr<IOutlinerExtension>  OutlinerItem = GetLinkedOutlinerItem();
+	TViewModelPtr<ITrackAreaExtension> TrackArea = OutlinerItem.ImplicitCast();
+
+	if (TrackArea && OutlinerItem)
+	{
+		FTrackAreaParameters Parameters = TrackArea->GetTrackAreaParameters();
+
+		Result.VerticalAlignment.VSizeParam = OutlinerItem->GetOutlinerSizing().GetTotalHeight();
+		Result.VerticalAlignment.VPadding   = Parameters.TrackLanePadding.Top + Parameters.TrackLanePadding.Bottom;
+	}
+
+	return Result;
 }
 
 ESelectionIntent FLayerBarModel::IsSelectable() const
diff --git a/Engine/Source/Editor/Sequencer/Private/Widgets/SSequencerLayerBar.cpp b/Engine/Source/Editor/Sequencer/Private/Widgets/SSequencerLayerBar.cpp
--- a/Engine/Source/Editor/Sequencer/Private/Widgets/SSequencerLayerBar.cpp
+++ b/Engine/Source/Editor/Sequencer/Private/Widgets/SSequencerLayerBar.cpp
@@ -254,26 +254,6 @@ void SSequencerLayerBar::Construct(const FArguments& InArgs, TWeakPtr<STrackArea
 	];
 }
 
-FVector2D SSequencerLayerBar::ComputeDesiredSize(float LayoutScale) const
-{
-	TViewModelPtr<FLinkedOutlinerExtension> LayerBar = WeakLayerBar.Pin();
-	TViewModelPtr<IOutlinerExtension>  OutlinerItem = LayerBar ? LayerBar->GetLinkedOutlinerItem() : nullptr;
-	TViewModelPtr<ITrackAreaExtension> TrackArea = OutlinerItem.ImplicitCast();
-
-	float Height = 0.f;
-
-	if (TrackArea && OutlinerItem)
-	{
-		FTrackAreaParameters Parameters = TrackArea->GetTrackAreaParameters();
-
-		Height = OutlinerItem->GetOutlinerSizing().GetTotalHeight();
-
-		Height -= Parameters.TrackLanePadding.Top + Parameters.TrackLanePadding.Bottom;
-	}
-
-	return FVector2D(100.f, Height);
-}
-
 TSharedPtr<FLayerBarModel> SSequencerLayerBar::GetLayerBarModel() const
 {
 	return WeakLayerBar.Pin();
diff --git a/Engine/Source/Editor/Sequencer/Private/Widgets/SSequencerLayerBar.h b/Engine/Source/Editor/Sequencer/Private/Widgets/SSequencerLayerBar.h
--- a/Engine/Source/Editor/Sequencer/Private/Widgets/SSequencerLayerBar.h
+++ b/Engine/Source/Editor/Sequencer/Private/Widgets/SSequencerLayerBar.h
@@ -49,7 +49,6 @@ private:
 
 	/*~ SCompoundWidget interface */
 	int32 OnPaint( const FPaintArgs& Args, const FGeometry& AllottedGeometry, const FSlateRect& MyCullingRect, FSlateWindowElementList& OutDrawElements, int32 LayerId, const FWidgetStyle& InWidgetStyle, bool bParentEnabled ) const override;
-	FVector2D ComputeDesiredSize(float LayoutScale) const override;
 
 	FReply OnMouseButtonDown(const FGeometry& MyGeometry, const FPointerEvent& MouseEvent) override;
 	FReply OnMouseButtonUp(const FGeometry& MyGeometry, const FPointerEvent& MouseEvent) override;
