Sequencer: Allow layer bar selection only with left mouse button

diff --git a/Engine/Source/Editor/Sequencer/Private/Widgets/SSequencerLayerBar.cpp b/Engine/Source/Editor/Sequencer/Private/Widgets/SSequencerLayerBar.cpp
--- a/Engine/Source/Editor/Sequencer/Private/Widgets/SSequencerLayerBar.cpp
+++ b/Engine/Source/Editor/Sequencer/Private/Widgets/SSequencerLayerBar.cpp
@@ -82,7 +82,10 @@ struct FLayerBarHotspot
 	}
 	void HandleMouseSelection(FHotspotSelectionManager& SelectionManager) override
 	{
-		SelectionManager.DefaultModelSelection(LayerBar);
+		if (SelectionManager.MouseEvent->GetEffectingButton() == EKeys::LeftMouseButton)
+		{
+			SelectionManager.DefaultModelSelection(LayerBar);
+		}
 	}
 
 	TSharedPtr<FLayerBarModel> LayerBar;
@@ -136,7 +139,10 @@ struct FLayerBarStretchHotspot
 	}
 	void HandleMouseSelection(FHotspotSelectionManager& SelectionManager) override
 	{
-		SelectionManager.SelectModelExclusive(LayerBar);
+		if (SelectionManager.MouseEvent->GetEffectingButton() == EKeys::LeftMouseButton)
+		{
+			SelectionManager.SelectModelExclusive(LayerBar);
+		}
 	}
 	FCursorReply GetCursor() const override
 	{
