Sequencer: Fix OnContextMenuOpeining and SetSectionToKey iterators getting thrown away immediately

#rnx

diff --git a/Engine/Source/Editor/Sequencer/Private/SequencerContextMenus.cpp b/Engine/Source/Editor/Sequencer/Private/SequencerContextMenus.cpp
--- a/Engine/Source/Editor/Sequencer/Private/SequencerContextMenus.cpp
+++ b/Engine/Source/Editor/Sequencer/Private/SequencerContextMenus.cpp
@@ -901,14 +901,14 @@ void FSectionContextMenu::SetSectionToKey()
 		return;
 	}
 
-	if (Sequencer->GetViewModel()->GetSelection()->GetSelectedSections().Num() != 1)
+	TSet<UMovieSceneSection*> SelectedSections = Sequencer->GetViewModel()->GetSelection()->GetSelectedSections();
+	if (SelectedSections.Num() != 1)
 	{
 		return;
 	}
 
 	const bool bToggle = IsSectionToKey();
-	auto It = Sequencer->GetViewModel()->GetSelection()->GetSelectedSections().CreateConstIterator();
-	UMovieSceneSection* Section = *It;
+	UMovieSceneSection* Section = *SelectedSections.CreateConstIterator();
 	UMovieSceneTrack* Track = Section->GetTypedOuter<UMovieSceneTrack>();
 	if (Track)
 	{
@@ -946,13 +946,13 @@ bool FSectionContextMenu::CanSetSectionToKey() const
 		return false;
 	}
 
-	if (Sequencer->GetViewModel()->GetSelection()->GetSelectedSections().Num() != 1)
+	TSet<UMovieSceneSection*> SelectedSections = Sequencer->GetViewModel()->GetSelection()->GetSelectedSections();
+	if (SelectedSections.Num() != 1)
 	{
 		return false;
 	}
 
-	auto It = Sequencer->GetViewModel()->GetSelection()->GetSelectedSections().CreateConstIterator();
-	UMovieSceneSection* Section = *It;
+	UMovieSceneSection* Section = *SelectedSections.CreateConstIterator();
 	UMovieSceneTrack* Track = Section->GetTypedOuter<UMovieSceneTrack>();
 	if (Track && Section->GetBlendType().IsValid() && (Section->GetBlendType().Get() == EMovieSceneBlendType::Absolute 
 		|| Section->GetBlendType().Get() == EMovieSceneBlendType::Additive 
diff --git a/Engine/Source/Editor/SequencerCore/Private/MVVM/Views/SOutlinerView.cpp b/Engine/Source/Editor/SequencerCore/Private/MVVM/Views/SOutlinerView.cpp
--- a/Engine/Source/Editor/SequencerCore/Private/MVVM/Views/SOutlinerView.cpp
+++ b/Engine/Source/Editor/SequencerCore/Private/MVVM/Views/SOutlinerView.cpp
@@ -792,9 +792,9 @@ TSharedPtr<SWidget> SOutlinerView::OnContextMenuOpening()
 	TSharedPtr<SWidget> ContextMenu;
 
 	// Open a context menu for the first selected item if it is selectable
-	if (auto It = GetSelectedItems().CreateConstIterator(); It)
+	if (auto It = GetSelectedItems(); It.Num() > 0) 
 	{
-		TWeakViewModelPtr<IOutlinerExtension> WeakSelectedItem = *It;
+		TWeakViewModelPtr<IOutlinerExtension> WeakSelectedItem = It[0];
 		TViewModelPtr<IOutlinerExtension> SelectedItem = WeakSelectedItem.Pin();
 
 		// Items are selectable by default, but can opt-in to conditionally selectable
