[Sequencer] Enable "Key This Section" option for all valid blend types

diff --git a/Engine/Source/Editor/Sequencer/Private/SequencerContextMenus.cpp b/Engine/Source/Editor/Sequencer/Private/SequencerContextMenus.cpp
--- a/Engine/Source/Editor/Sequencer/Private/SequencerContextMenus.cpp
+++ b/Engine/Source/Editor/Sequencer/Private/SequencerContextMenus.cpp
@@ -954,9 +954,7 @@ bool FSectionContextMenu::CanSetSectionToKey() const
 
 	UMovieSceneSection* Section = *SelectedSections.CreateConstIterator();
 	UMovieSceneTrack* Track = Section->GetTypedOuter<UMovieSceneTrack>();
-	if (Track && Section->GetBlendType().IsValid() && (Section->GetBlendType().Get() == EMovieSceneBlendType::Absolute 
-		|| Section->GetBlendType().Get() == EMovieSceneBlendType::Additive 
-		|| Section->GetBlendType().Get() == EMovieSceneBlendType::Override))
+	if (Track && Section->GetBlendType().IsValid() && Section->GetBlendType().Get() != EMovieSceneBlendType::Invalid)
 	{
 		return true;
 	}
