Sequencer: Speculative fix for calling paste sections with a null track

Crash report indicates that paste sections is being called from python with a potential null track.

#rnx

diff --git a/Engine/Source/Editor/Sequencer/Private/SequencerUtilities.cpp b/Engine/Source/Editor/Sequencer/Private/SequencerUtilities.cpp
--- a/Engine/Source/Editor/Sequencer/Private/SequencerUtilities.cpp
+++ b/Engine/Source/Editor/Sequencer/Private/SequencerUtilities.cpp
@@ -2217,6 +2217,10 @@ bool FSequencerUtilities::PasteSections(const FString& TextToImport, FMovieScene
 	for (int32 Index = 0; Index < PasteSectionsParams.Tracks.Num(); ++Index)
 	{
 		UMovieSceneTrack* Track = PasteSectionsParams.Tracks[Index];
+		if (!Track)
+		{
+			continue;
+		}
 
 		const bool bAllowOverlap = Track->SupportsMultipleRows() || Track->GetSupportedBlendTypes().Num() > 0;
 
