[Sequencer] Filter improvements to parody Navigation Tool

diff --git a/Engine/Source/Editor/Sequencer/Private/Filters/Filters/SequencerTrackFilter_Modified.cpp b/Engine/Source/Editor/Sequencer/Private/Filters/Filters/SequencerTrackFilter_Modified.cpp
--- a/Engine/Source/Editor/Sequencer/Private/Filters/Filters/SequencerTrackFilter_Modified.cpp
+++ b/Engine/Source/Editor/Sequencer/Private/Filters/Filters/SequencerTrackFilter_Modified.cpp
@@ -60,9 +60,4 @@ bool FSequencerTrackFilter_Modified::PassesFilter(FSequencerTrackFilterType InIt
 	return false;
 }
 
-void FSequencerTrackFilter_Modified::ToggleShowOnlyModifiedTracks()
-{
-	SetActive(!IsActive());
-}
-
 #undef LOCTEXT_NAMESPACE
diff --git a/Engine/Source/Editor/Sequencer/Private/Filters/Filters/SequencerTrackFilter_Text.cpp b/Engine/Source/Editor/Sequencer/Private/Filters/Filters/SequencerTrackFilter_Text.cpp
--- a/Engine/Source/Editor/Sequencer/Private/Filters/Filters/SequencerTrackFilter_Text.cpp
+++ b/Engine/Source/Editor/Sequencer/Private/Filters/Filters/SequencerTrackFilter_Text.cpp
@@ -101,13 +101,15 @@ FString FSequencerTrackFilter_Text::GetName() const
 
 bool FSequencerTrackFilter_Text::PassesFilter(FSequencerTrackFilterType InItem) const
 {
-	FSequencerFilterData& FilterData = GetFilterInterface().GetFilterData();
-
 	for (const TSharedRef<FSequencerTextFilterExpressionContext>& TextFilterExpressionContext : TextFilterExpressionContexts)
 	{
-		UMovieSceneTrack* const TrackObject = FilterData.ResolveMovieSceneTrackObject(InItem);
-		TextFilterExpressionContext->SetFilterItem(InItem, TrackObject);
-		if (!TextFilterExpressionEvaluator.TestTextFilter(*TextFilterExpressionContext))
+		TextFilterExpressionContext->SetFilterItem(InItem);
+
+		const bool bPassedFilter = TextFilterExpressionEvaluator.TestTextFilter(*TextFilterExpressionContext);
+
+		TextFilterExpressionContext->SetFilterItem(nullptr);
+
+		if (!bPassedFilter)
 		{
 			return false;
 		}
diff --git a/Engine/Source/Editor/Sequencer/Private/Filters/Filters/SequencerTrackFilter_Unbound.cpp b/Engine/Source/Editor/Sequencer/Private/Filters/Filters/SequencerTrackFilter_Unbound.cpp
--- a/Engine/Source/Editor/Sequencer/Private/Filters/Filters/SequencerTrackFilter_Unbound.cpp
+++ b/Engine/Source/Editor/Sequencer/Private/Filters/Filters/SequencerTrackFilter_Unbound.cpp
@@ -64,9 +64,4 @@ bool FSequencerTrackFilter_Unbound::PassesFilter(FSequencerTrackFilterType InIte
 	return bPassed;
 }
 
-void FSequencerTrackFilter_Unbound::ToggleShowOnlyModifiedTracks()
-{
-	SetActive(!IsActive());
-}
-
 #undef LOCTEXT_NAMESPACE
diff --git a/Engine/Source/Editor/Sequencer/Private/Filters/Menus/SequencerFilterBarContextMenu.cpp b/Engine/Source/Editor/Sequencer/Private/Filters/Menus/SequencerFilterBarContextMenu.cpp
--- a/Engine/Source/Editor/Sequencer/Private/Filters/Menus/SequencerFilterBarContextMenu.cpp
+++ b/Engine/Source/Editor/Sequencer/Private/Filters/Menus/SequencerFilterBarContextMenu.cpp
@@ -134,16 +134,4 @@ void FSequencerFilterBarContextMenu::OnActivateAllFilters(const bool bInActivate
 	FilterBar->ActivateAllEnabledFilters(bInActivate, {});
 }
 
-void FSequencerFilterBarContextMenu::OnResetFilters()
-{
-	const TSharedPtr<FSequencerFilterBar> FilterBar = WeakFilterBar.Pin();
-	if (!FilterBar.IsValid())
-	{
-		return;
-	}
-
-	FilterBar->EnableAllFilters(false, {});
-	FilterBar->EnableCustomTextFilters(false, {});
-}
-
 #undef LOCTEXT_NAMESPACE
diff --git a/Engine/Source/Editor/Sequencer/Private/Filters/Menus/SequencerFilterBarContextMenu.h b/Engine/Source/Editor/Sequencer/Private/Filters/Menus/SequencerFilterBarContextMenu.h
--- a/Engine/Source/Editor/Sequencer/Private/Filters/Menus/SequencerFilterBarContextMenu.h
+++ b/Engine/Source/Editor/Sequencer/Private/Filters/Menus/SequencerFilterBarContextMenu.h
@@ -20,7 +20,6 @@ protected:
 	void PopulateFilterBulkOptionsSection(UToolMenu& InMenu);
 
 	void OnActivateAllFilters(const bool bInActivate);
-	void OnResetFilters();
 
 	TWeakPtr<FSequencerFilterBar> WeakFilterBar;
 };
diff --git a/Engine/Source/Editor/Sequencer/Private/Filters/Menus/SequencerTrackFilterContextMenu.cpp b/Engine/Source/Editor/Sequencer/Private/Filters/Menus/SequencerTrackFilterContextMenu.cpp
--- a/Engine/Source/Editor/Sequencer/Private/Filters/Menus/SequencerTrackFilterContextMenu.cpp
+++ b/Engine/Source/Editor/Sequencer/Private/Filters/Menus/SequencerTrackFilterContextMenu.cpp
@@ -123,17 +123,15 @@ void FSequencerTrackFilterContextMenu::PopulateBulkOptionsSection(UToolMenu& InM
 	FToolMenuSection& Section = InMenu.FindOrAddSection(TEXT("FilterBulkOptions")
 		, LOCTEXT("BulkOptionsContextHeading", "Filter Bulk Options"));
 
-	Section.AddMenuEntry(TEXT("ActivateAllFilters"),
-		LOCTEXT("ActivateAllFilters", "Activate All Filters"),
-		LOCTEXT("ActivateAllFiltersTooltip", "Activates all enabled filters."),
-		FSlateIcon(FAppStyle::GetAppStyleSetName(), TEXT("Icons.Plus")),
-		FUIAction(FExecuteAction::CreateRaw(this, &FSequencerTrackFilterContextMenu::OnActivateAllFilters, true)));
+	Section.AddMenuEntry(TrackFilterCommands.ActivateAllFilters
+		, TrackFilterCommands.ActivateAllFilters->GetLabel()
+		, TrackFilterCommands.ActivateAllFilters->GetDescription()
+		, FSlateIcon(FAppStyle::GetAppStyleSetName(), TEXT("Icons.Plus")));
 
-	Section.AddMenuEntry(TEXT("DeactivateAllFilters"),
-		LOCTEXT("DeactivateAllFilters", "Deactivate All Filters"),
-		LOCTEXT("DeactivateAllFiltersTooltip", "Deactivates all enabled filters."),
-		FSlateIcon(FAppStyle::GetAppStyleSetName(), TEXT("Icons.Minus")),
-		FUIAction(FExecuteAction::CreateRaw(this, &FSequencerTrackFilterContextMenu::OnActivateAllFilters, false)));
+	Section.AddMenuEntry(TrackFilterCommands.DeactivateAllFilters
+		, TrackFilterCommands.DeactivateAllFilters->GetLabel()
+		, TrackFilterCommands.DeactivateAllFilters->GetDescription()
+		, FSlateIcon(FAppStyle::GetAppStyleSetName(), TEXT("Icons.Minus")));
 
 	Section.AddSeparator(NAME_None);
 
@@ -178,17 +176,6 @@ void FSequencerTrackFilterContextMenu::OnDisableFilter()
 	Filter->GetFilterInterface().SetFilterEnabledByDisplayName(FilterName, false);
 }
 
-void FSequencerTrackFilterContextMenu::OnResetFilters()
-{
-	const TSharedPtr<FSequencerTrackFilter> Filter = GetFilter();
-	if (!Filter.IsValid())
-	{
-		return;
-	}
-
-	Filter->GetFilterInterface().EnableAllFilters(false, {});
-}
-
 void FSequencerTrackFilterContextMenu::OnActivateWithFilterException()
 {
 	const TSharedPtr<FSequencerTrackFilter> Filter = GetFilter();
@@ -229,9 +216,10 @@ void FSequencerTrackFilterContextMenu::OnEditFilter()
 		return;
 	}
 
-	FSequencerFilterBar& FilterBar = (FSequencerFilterBar&)Filter->GetFilterInterface();
+	ISequencerTrackFilters& FilterInterface = Filter->GetFilterInterface();
+	const TSharedRef<ISequencerTrackFilters> FilterBar = StaticCastSharedRef<ISequencerTrackFilters>(FilterInterface.AsShared());
 
-	SSequencerCustomTextFilterDialog::CreateWindow_EditCustomTextFilter(FilterBar.AsShared(), CustomTextFilter);
+	SSequencerCustomTextFilterDialog::CreateWindow_EditCustomTextFilter(FilterBar, CustomTextFilter);
 }
 
 void FSequencerTrackFilterContextMenu::OnDeleteFilter()
diff --git a/Engine/Source/Editor/Sequencer/Private/Filters/Menus/SequencerTrackFilterContextMenu.h b/Engine/Source/Editor/Sequencer/Private/Filters/Menus/SequencerTrackFilterContextMenu.h
--- a/Engine/Source/Editor/Sequencer/Private/Filters/Menus/SequencerTrackFilterContextMenu.h
+++ b/Engine/Source/Editor/Sequencer/Private/Filters/Menus/SequencerTrackFilterContextMenu.h
@@ -23,7 +23,6 @@ protected:
 	FText GetFilterDisplayName() const;
 
 	void OnDisableFilter();
-	void OnResetFilters();
 
 	void OnActivateWithFilterException();
 
diff --git a/Engine/Source/Editor/Sequencer/Private/Filters/SequencerFilterBar.cpp b/Engine/Source/Editor/Sequencer/Private/Filters/SequencerFilterBar.cpp
--- a/Engine/Source/Editor/Sequencer/Private/Filters/SequencerFilterBar.cpp
+++ b/Engine/Source/Editor/Sequencer/Private/Filters/SequencerFilterBar.cpp
@@ -15,6 +15,7 @@
 #include "Filters/Filters/SequencerTrackFilter_Unbound.h"
 #include "Filters/Filters/SequencerTrackFilters.h"
 #include "Filters/SequencerFilterBarConfig.h"
+#include "Filters/SequencerTextFilterExpressionContext.h"
 #include "Filters/SequencerTrackFilterCollection.h"
 #include "Filters/SequencerTrackFilterCommands.h"
 #include "Filters/SequencerTrackFilterExtension.h"
@@ -32,6 +33,7 @@
 #include "Stats/StatsMisc.h"
 #include "UObject/UObjectIterator.h"
 #include "Widgets/SFilterBarIsolateHideShow.h"
+#include "Widgets/SFilterExpressionHelpDialog.h"
 #include "Widgets/SSequencerSearchBox.h"
 
 using namespace UE::Sequencer;
@@ -72,6 +74,8 @@ FSequencerFilterBar::FSequencerFilterBar(FSequencer& InSequencer)
 	SelectedFilter->OnChanged().AddRaw(this, &FSequencerFilterBar::RequestFilterUpdate);
 
 	CreateDefaultFilters();
+
+	CreateCustomTextFiltersFromConfig();
 }
 
 FSequencerFilterBar::~FSequencerFilterBar()
@@ -179,39 +183,37 @@ void FSequencerFilterBar::BindCommands()
 	const TSharedRef<SSequencer> SequencerWidget = StaticCastSharedRef<SSequencer>(GetSequencer().GetSequencerWidget());
 
 	CommandList->MapAction(TrackFilterCommands.ToggleFilterBarVisibility,
-		FUIAction(
-			FExecuteAction::CreateSP(SequencerWidget, &SSequencer::ToggleFilterBarVisibility),
-			FCanExecuteAction(),
-			FIsActionChecked::CreateSP(SequencerWidget, &SSequencer::IsFilterBarVisible)
-		));
+		FExecuteAction::CreateSP(SequencerWidget, &SSequencer::ToggleFilterBarVisibility),
+		FCanExecuteAction(),
+		FIsActionChecked::CreateSP(SequencerWidget, &SSequencer::IsFilterBarVisible));
 
 	CommandList->MapAction(TrackFilterCommands.ResetFilters,
-		FUIAction(
-			FExecuteAction::CreateSP(this, &FSequencerFilterBar::ResetFilters),
-			FCanExecuteAction::CreateSP(this, &FSequencerFilterBar::CanResetFilters)
-		));
+		FExecuteAction::CreateSP(this, &FSequencerFilterBar::ResetFilters),
+		FCanExecuteAction::CreateSP(this, &FSequencerFilterBar::CanResetFilters));
 
 	CommandList->MapAction(TrackFilterCommands.ToggleMuteFilters,
-		FUIAction(
-			FExecuteAction::CreateSP(this, &FSequencerFilterBar::ToggleMuteFilters),
-			FCanExecuteAction(),
-			FIsActionChecked::CreateSP(this, &FSequencerFilterBar::AreFiltersMuted)
-		));
+		FExecuteAction::CreateSP(this, &FSequencerFilterBar::ToggleMuteFilters),
+		FCanExecuteAction(),
+		FIsActionChecked::CreateSP(this, &FSequencerFilterBar::AreFiltersMuted));
 
 	CommandList->MapAction(TrackFilterCommands.DisableAllFilters,
-		FUIAction(
-			FExecuteAction::CreateLambda([this]()
-				{
-					EnableAllFilters(false, {});
-				}),
-			FCanExecuteAction::CreateSP(this, &FSequencerFilterBar::HasAnyFilterEnabled)
-		));
+		FExecuteAction::CreateLambda([this]()
+			{
+				EnableAllFilters(false, {});
+			}),
+		FCanExecuteAction::CreateSP(this, &FSequencerFilterBar::HasAnyFilterEnabled));
 
 	CommandList->MapAction(TrackFilterCommands.ToggleActivateEnabledFilters,
-		FUIAction(
-			FExecuteAction::CreateSP(this, &FSequencerFilterBar::ToggleActivateAllEnabledFilters),
-			FCanExecuteAction::CreateSP(this, &FSequencerFilterBar::HasAnyFilterEnabled)
-		));
+		FExecuteAction::CreateSP(this, &FSequencerFilterBar::ToggleActivateAllEnabledFilters),
+		FCanExecuteAction::CreateSP(this, &FSequencerFilterBar::HasAnyFilterEnabled));
+
+	CommandList->MapAction(TrackFilterCommands.ActivateAllFilters,
+		FExecuteAction::CreateSP(this, &FSequencerFilterBar::ActivateAllEnabledFilters, true, TArray<FString>()),
+		FCanExecuteAction::CreateSP(this, &FSequencerFilterBar::AreAllEnabledFiltersActive, true, TArray<FString>()));
+
+	CommandList->MapAction(TrackFilterCommands.DeactivateAllFilters,
+		FExecuteAction::CreateSP(this, &FSequencerFilterBar::ActivateAllEnabledFilters, false, TArray<FString>()),
+		FCanExecuteAction::CreateSP(this, &FSequencerFilterBar::AreAllEnabledFiltersActive, false, TArray<FString>()));
 
 	// Bind all filter actions
 	UMovieSceneSequence* const FocusedSequence = Sequencer.GetFocusedMovieSceneSequence();
@@ -260,7 +262,6 @@ void FSequencerFilterBar::BindCommands()
 				}
 
 				CurveEditorCommands->Append(CurveEditorSharedBindings.ToSharedRef());
-
 			}
 		}
 	}
@@ -333,12 +334,6 @@ void FSequencerFilterBar::ToggleMuteFilters()
 	MuteFilters(!AreFiltersMuted());
 }
 
-void FSequencerFilterBar::ResetFilters()
-{
-	EnableAllFilters(false, {});
-	EnableCustomTextFilters(false);
-}
-
 bool FSequencerFilterBar::CanResetFilters() const
 {
 	const TArray<TSharedRef<FFilterCategory>> Categories = { GetClassTypeCategory(), GetComponentTypeCategory(), GetMiscCategory() };
@@ -346,6 +341,12 @@ bool FSequencerFilterBar::CanResetFilters() const
 	return HasEnabledFilter(ClassAndCompFilters);
 }
 
+void FSequencerFilterBar::ResetFilters()
+{
+	EnableAllFilters(false, {});
+	EnableCustomTextFilters(false);
+}
+
 FString FSequencerFilterBar::GetTextFilterString() const
 {
 	return TextFilter->GetRawFilterText().ToString();
@@ -450,6 +451,8 @@ void FSequencerFilterBar::EmptyIsolatedTracks()
 void FSequencerFilterBar::RequestFilterUpdate()
 {
 	Sequencer.GetNodeTree()->RequestFilterUpdate();
+
+	RequestUpdateEvent.Broadcast();
 }
 
 TSharedPtr<FSequencerTrackFilter> FSequencerFilterBar::FindFilterByDisplayName(const FString& InFilterName) const
@@ -713,7 +716,7 @@ void FSequencerFilterBar::ActivateCommonFilters(const bool bInActivate
 	RequestFilterUpdate();
 }
 
-bool FSequencerFilterBar::AreAllEnabledFiltersActive(const bool bInActive, const TArray<FString>& InExceptionFilterNames) const
+bool FSequencerFilterBar::AreAllEnabledFiltersActive(const bool bInActive, const TArray<FString> InExceptionFilterNames) const
 {
 	const TArray<TSharedRef<FSequencerTrackFilter>> EnabledFilters = GetEnabledFilters();
 	for (const TSharedRef<FSequencerTrackFilter>& Filter : EnabledFilters)
@@ -748,7 +751,7 @@ bool FSequencerFilterBar::AreAllEnabledFiltersActive(const bool bInActive, const
 	return true;
 }
 
-void FSequencerFilterBar::ActivateAllEnabledFilters(const bool bInActivate, const TArray<FString>& InExceptionFilterNames)
+void FSequencerFilterBar::ActivateAllEnabledFilters(const bool bInActivate, const TArray<FString> InExceptionFilterNames)
 {
 	const TArray<TSharedRef<FSequencerTrackFilter>> EnabledFilters = GetEnabledFilters();
 	for (const TSharedRef<FSequencerTrackFilter>& Filter : EnabledFilters)
@@ -1483,6 +1486,11 @@ bool FSequencerFilterBar::FilterNodesRecursive(const bool bInHasActiveFilter, co
 			SetTrackParentsExpanded(InStartNode.ImplicitCast(), true);
 		}
 
+		if (bPassedTextFilter)
+		{
+			FilterData.FilterInParentChildNodes(InStartNode, true, true, true);
+		}
+
 		FilterData.FilterInNodeWithAncestors(InStartNode);
 		return true;
 	}
@@ -1558,18 +1566,27 @@ void FSequencerFilterBar::SetTrackParentsExpanded(const TViewModelPtr<IOutlinerE
 {
 	for (const TViewModelPtr<IOutlinerExtension>& ParentNode : InNode.AsModel()->GetAncestorsOfType<IOutlinerExtension>())
 	{
-		if (!ParentNode->IsExpanded())
+		if (ParentNode->IsExpanded() != bInExpanded)
 		{
-			ParentNode->SetExpansion(true);
+			ParentNode->SetExpansion(bInExpanded);
 		}
 	}
 }
 
 FString FSequencerFilterBar::GenerateTextFilterStringFromEnabledFilters() const
 {
+	TArray<TSharedRef<FSequencerTrackFilter>> FiltersToSave;
+
+	FiltersToSave.Append(GetCommonFilters());
+
+	for (const TSharedRef<FSequencerTrackFilter_CustomText>& Filter : CustomTextFilters)
+	{
+		FiltersToSave.Add(Filter);
+	}
+
 	FString GeneratedFilterString = TextFilter->GetRawFilterText().ToString();
 
-	for (const TSharedRef<FSequencerTrackFilter>& Filter : GetCommonFilters())
+	for (const TSharedRef<FSequencerTrackFilter>& Filter : FiltersToSave)
 	{
 		if (IsFilterActive(Filter) && IsFilterEnabled(Filter))
 		{
@@ -1691,6 +1708,89 @@ TSharedRef<SComboButton> FSequencerFilterBar::MakeAddFilterButton()
 	return ComboButton;
 }
 
+bool FSequencerFilterBar::ShouldShowFilterBarWidget() const
+{
+	const USequencerSettings* const Settings = Sequencer.GetSequencerSettings();
+	if (ensure(IsValid(Settings)))
+	{
+		return Settings->IsFilterBarVisible() && HasAnyFiltersEnabled();
+	}
+	return false;
+}
+
+bool FSequencerFilterBar::IsFilterBarVisible() const
+{
+	const USequencerSettings* const Settings = Sequencer.GetSequencerSettings();
+	if (ensure(IsValid(Settings)))
+	{
+		return Settings->IsFilterBarVisible();
+	}
+	return false;
+}
+
+void FSequencerFilterBar::ToggleFilterBarVisibility()
+{
+	USequencerSettings* const Settings = Sequencer.GetSequencerSettings();
+	if (!ensure(IsValid(Settings)))
+	{
+		return;
+	}
+
+	const bool bNewVisibility = !Settings->IsFilterBarVisible();
+
+	Settings->SetFilterBarVisible(bNewVisibility);
+
+	StateChangedEvent.Broadcast(bNewVisibility, Settings->GetFilterBarLayout());
+}
+
+bool FSequencerFilterBar::IsFilterBarLayout(const EFilterBarLayout InLayout) const
+{
+	const USequencerSettings* const Settings = Sequencer.GetSequencerSettings();
+	if (ensure(IsValid(Settings)))
+	{
+		return (Settings->GetFilterBarLayout() == InLayout);
+	}
+	return false;
+}
+
+void FSequencerFilterBar::SetToVerticalLayout()
+{
+	USequencerSettings* const Settings = Sequencer.GetSequencerSettings();
+	if (!ensure(IsValid(Settings)))
+	{
+		return;
+	}
+
+	Settings->SetFilterBarLayout(EFilterBarLayout::Vertical);
+
+	StateChangedEvent.Broadcast(IsFilterBarVisible(), Settings->GetFilterBarLayout());
+}
+
+void FSequencerFilterBar::SetToHorizontalLayout()
+{
+	USequencerSettings* const Settings = Sequencer.GetSequencerSettings();
+	if (!ensure(IsValid(Settings)))
+	{
+		return;
+	}
+
+	Settings->SetFilterBarLayout(EFilterBarLayout::Horizontal);
+
+	StateChangedEvent.Broadcast(IsFilterBarVisible(), Settings->GetFilterBarLayout());
+}
+
+void FSequencerFilterBar::ToggleFilterBarLayout()
+{
+	if (IsFilterBarLayout(EFilterBarLayout::Horizontal))
+	{
+		SetToVerticalLayout();
+	}
+	else if (IsFilterBarLayout(EFilterBarLayout::Vertical))
+	{
+		SetToHorizontalLayout();
+	}
+}
+
 TSharedPtr<SSequencerFilterBar> FSequencerFilterBar::GetWidget() const
 {
 	const TSharedPtr<SSequencer> SequencerWidget = StaticCastSharedRef<SSequencer>(GetSequencer().GetSequencerWidget());
@@ -1726,4 +1826,27 @@ bool FSequencerFilterBar::IsFilterSupported(const FString& InFilterName) const
 	return FoundFilter ? IsFilterSupported(*FoundFilter) : false;
 }
 
+void FSequencerFilterBar::OpenTextExpressionHelp()
+{
+	FFilterExpressionHelpDialogConfig Config;
+	Config.IdentifierName = TEXT("NavigationToolCustomTextFilterHelp");
+	Config.DialogTitle = LOCTEXT("NavigationToolCustomTextFilterHelp", "Navigation Tool Custom Text Filter Help");
+	Config.TextFilterExpressionContexts = GetTextFilterExpressionContexts();
+	SFilterExpressionHelpDialog::Open(MoveTemp(Config));
+}
+
+void FSequencerFilterBar::SaveCurrentFilterSetAsCustomTextFilter()
+{
+	const FText NewFilterString = FText::FromString(GenerateTextFilterStringFromEnabledFilters());
+	// @TODO: update to the new way of handling this (similar to Navigation Tool)
+	//CreateWindow_AddCustomTextFilter(DefaultNewCustomTextFilterData(NewFilterString));
+}
+
+void FSequencerFilterBar::CreateNewTextFilter()
+{
+	const FText NewFilterString = FText::FromString(GetTextFilterString());
+	// @TODO: update to the new way of handling this (similar to Navigation Tool)
+	//CreateWindow_AddCustomTextFilter(DefaultNewCustomTextFilterData(NewFilterString));
+}
+
 #undef LOCTEXT_NAMESPACE
diff --git a/Engine/Source/Editor/Sequencer/Private/Filters/SequencerFilterBar.h b/Engine/Source/Editor/Sequencer/Private/Filters/SequencerFilterBar.h
--- a/Engine/Source/Editor/Sequencer/Private/Filters/SequencerFilterBar.h
+++ b/Engine/Source/Editor/Sequencer/Private/Filters/SequencerFilterBar.h
@@ -39,7 +39,7 @@ enum class ESequencerFilterChange : uint8
 };
 
 /** Holds the Sequencer track filter collection, the current text filter, and hidden/isolated lists. */
-class FSequencerFilterBar : public TSharedFromThis<FSequencerFilterBar>, public ISequencerTrackFilters
+class FSequencerFilterBar : public ISequencerTrackFilters
 {
 public:
 	/** An identifier shared by all filter bars, used to save and load settings common to every instance */
@@ -53,18 +53,16 @@ public:
 
 	void BindCommands();
 
-	void CreateCustomTextFiltersFromConfig();
-
 	TSharedPtr<ICustomTextFilter<FSequencerTrackFilterType>> CreateTextFilter();
 
 	TSharedRef<SSequencerFilterBar> GenerateWidget(const TSharedPtr<SSequencerSearchBox>& InSearchBox, const EFilterBarLayout InLayout);
 
-	bool AreFiltersMuted() const;
-	void MuteFilters(const bool bInMute);
+	virtual bool AreFiltersMuted() const override;
+	virtual void MuteFilters(const bool bInMute) override;
 	void ToggleMuteFilters();
 
-	void ResetFilters();
-	bool CanResetFilters() const;
+	virtual bool CanResetFilters() const override;
+	virtual void ResetFilters() override;
 
 	FSequencerFiltersChanged& GetOnFiltersChanged() { return FiltersChangedEvent; }
 
@@ -114,8 +112,8 @@ public:
 
 	virtual void ActivateCommonFilters(const bool bInActivate, const TArray<FString>& InExceptionFilterNames) override;
 
-	virtual bool AreAllEnabledFiltersActive(const bool bInActive, const TArray<FString>& InExceptionFilterNames) const;
-	virtual void ActivateAllEnabledFilters(const bool bInActivate, const TArray<FString>& InExceptionFilterNames) override;
+	virtual bool AreAllEnabledFiltersActive(const bool bInActive, const TArray<FString> InExceptionFilterNames) const override;
+	virtual void ActivateAllEnabledFilters(const bool bInActivate, const TArray<FString> InExceptionFilterNames) override;
 	virtual void ToggleActivateAllEnabledFilters();
 
 	virtual bool IsFilterActiveByDisplayName(const FString& InFilterName) const override;
@@ -158,8 +156,8 @@ public:
 		, const bool bInCheckCommonFilters = true
 		, const bool bInCheckInternalFilters = true
 		, const bool bInCheckCustomTextFilters = true) const override;
-	bool IsFilterActive(const TSharedRef<FSequencerTrackFilter> InFilter) const;
-	bool SetFilterActive(const TSharedRef<FSequencerTrackFilter>& InFilter, const bool bInActive, const bool bInRequestFilterUpdate = true);
+	virtual bool IsFilterActive(const TSharedRef<FSequencerTrackFilter> InFilter) const override;
+	virtual bool SetFilterActive(const TSharedRef<FSequencerTrackFilter>& InFilter, const bool bInActive, const bool bInRequestFilterUpdate = true) override;
 	void ActivateCommonFilters(const bool bInActivate
 		, const TArray<TSharedRef<FFilterCategory>> InMatchCategories
 		, const TArray<TSharedRef<FSequencerTrackFilter>>& InExceptions);
@@ -170,8 +168,8 @@ public:
 	bool HasEnabledCommonFilters() const;
 	bool HasEnabledFilter(const TArray<TSharedRef<FSequencerTrackFilter>>& InFilters = {}) const;
 	virtual bool HasAnyFilterEnabled() const override;
-	bool IsFilterEnabled(TSharedRef<FSequencerTrackFilter> InFilter) const;
-	bool SetFilterEnabled(const TSharedRef<FSequencerTrackFilter> InFilter, const bool bInEnabled, const bool bInRequestFilterUpdate = true);
+	virtual bool IsFilterEnabled(TSharedRef<FSequencerTrackFilter> InFilter) const override;
+	virtual bool SetFilterEnabled(const TSharedRef<FSequencerTrackFilter> InFilter, const bool bInEnabled, const bool bInRequestFilterUpdate = true) override;
 	void EnableFilters(const bool bInEnable
 		, const TArray<TSharedRef<FFilterCategory>> InMatchCategories = {}
 		, const TArray<TSharedRef<FSequencerTrackFilter>> InExceptions = {});
@@ -237,14 +235,32 @@ public:
 	TSharedRef<SFilterBarIsolateHideShow> MakeIsolateHideShowPanel();
 	TSharedRef<SComboButton> MakeAddFilterButton();
 
+	virtual bool ShouldShowFilterBarWidget() const override;
+
+	virtual bool IsFilterBarVisible() const override;
+	virtual void ToggleFilterBarVisibility() override;
+
+	virtual bool IsFilterBarLayout(const EFilterBarLayout InLayout) const override;
+	virtual void SetToVerticalLayout() override;
+	virtual void SetToHorizontalLayout() override;
+	virtual void ToggleFilterBarLayout() override;
+
 	/** Attempts to get the filter bar widget from the Sequencer widget */
 	TSharedPtr<SSequencerFilterBar> GetWidget() const;
 
 	bool IsFilterSupported(const TSharedRef<FSequencerTrackFilter>& InFilter) const;
 	bool IsFilterSupported(const FString& InFilterName) const;
 
+	virtual FOnFilterBarStateChanged& OnStateChanged() override { return StateChangedEvent; }
+	virtual FSimpleMulticastDelegate& OnRequestUpdate() override { return RequestUpdateEvent; }
+
+	virtual void OpenTextExpressionHelp() override;
+	virtual void SaveCurrentFilterSetAsCustomTextFilter() override;
+	virtual void CreateNewTextFilter() override;
+
 protected:
 	void CreateDefaultFilters();
+	void CreateCustomTextFiltersFromConfig();
 
 	bool PassesAnyCommonFilter(const UE::Sequencer::TViewModelPtr<UE::Sequencer::IOutlinerExtension>& InNode);
 	bool PassesAllInternalFilters(const UE::Sequencer::TViewModelPtr<UE::Sequencer::IOutlinerExtension>& InNode);
@@ -279,10 +295,12 @@ protected:
 
 	TSharedPtr<FSequencerTrackFilterMenu> FilterMenu;
 
-	FSequencerFiltersChanged FiltersChangedEvent;
-
 	FSequencerFilterData FilterData;
 
+	FSequencerFiltersChanged FiltersChangedEvent;
+	FOnFilterBarStateChanged StateChangedEvent;
+	FSimpleMulticastDelegate RequestUpdateEvent;
+
 private:
 	static int32 InstanceCount;
 
diff --git a/Engine/Source/Editor/Sequencer/Private/Filters/SequencerTrackFilterCommands.cpp b/Engine/Source/Editor/Sequencer/Private/Filters/SequencerTrackFilterCommands.cpp
--- a/Engine/Source/Editor/Sequencer/Private/Filters/SequencerTrackFilterCommands.cpp
+++ b/Engine/Source/Editor/Sequencer/Private/Filters/SequencerTrackFilterCommands.cpp
@@ -24,7 +24,10 @@ void FSequencerTrackFilterCommands::RegisterCommands()
 	UI_COMMAND(DisableAllFilters, "Remove All Filters", "Disables all filters and removes them from the filter bar", EUserInterfaceActionType::Button, FInputChord());
 
 	UI_COMMAND(ToggleActivateEnabledFilters, "Toggle Active State", "Activates or deactivates all enabled filters", EUserInterfaceActionType::Button, FInputChord());
-	
+
+	UI_COMMAND(ActivateAllFilters, "Activate All Filters", "Activates all enabled filters", EUserInterfaceActionType::Button, FInputChord());
+	UI_COMMAND(DeactivateAllFilters, "Deactivate All Filters", "Deactivates all enabled filters", EUserInterfaceActionType::Button, FInputChord());
+
 	UI_COMMAND(HideSelectedTracks, "Hide Selected Tracks", "Hide all selected tracks", EUserInterfaceActionType::Button, FInputChord(EModifierKey::Alt, EKeys::W));
 	UI_COMMAND(IsolateSelectedTracks, "Isolate Selected Tracks", "Isolate all selected tracks", EUserInterfaceActionType::Button, FInputChord(EModifierKey::Alt, EKeys::Q));
 
diff --git a/Engine/Source/Editor/Sequencer/Private/Filters/SequencerTrackFilterCommands.h b/Engine/Source/Editor/Sequencer/Private/Filters/SequencerTrackFilterCommands.h
--- a/Engine/Source/Editor/Sequencer/Private/Filters/SequencerTrackFilterCommands.h
+++ b/Engine/Source/Editor/Sequencer/Private/Filters/SequencerTrackFilterCommands.h
@@ -29,6 +29,9 @@ public:
 
 	TSharedPtr<FUICommandInfo> ToggleActivateEnabledFilters;
 
+	TSharedPtr<FUICommandInfo> ActivateAllFilters;
+	TSharedPtr<FUICommandInfo> DeactivateAllFilters;
+
 	/** Hide/Isolate Filter Commands */
 
 	TSharedPtr<FUICommandInfo> HideSelectedTracks;
diff --git a/Engine/Source/Editor/Sequencer/Private/Filters/TextExpressions/SequencerTextFilterExpressionContext.cpp b/Engine/Source/Editor/Sequencer/Private/Filters/TextExpressions/SequencerTextFilterExpressionContext.cpp
--- a/Engine/Source/Editor/Sequencer/Private/Filters/TextExpressions/SequencerTextFilterExpressionContext.cpp
+++ b/Engine/Source/Editor/Sequencer/Private/Filters/TextExpressions/SequencerTextFilterExpressionContext.cpp
@@ -14,10 +14,11 @@ FSequencerTextFilterExpressionContext::FSequencerTextFilterExpressionContext(ISe
 {
 }
 
-void FSequencerTextFilterExpressionContext::SetFilterItem(FSequencerTrackFilterType InFilterItem, UMovieSceneTrack* const InTrackObject)
+void FSequencerTextFilterExpressionContext::SetFilterItem(FSequencerTrackFilterType InFilterItem)
 {
 	FilterItem = InFilterItem;
-	WeakTrackObject = InTrackObject;
+
+	FSequencerFilterData& FilterData = FilterInterface.GetFilterData();
 }
 
 bool FSequencerTextFilterExpressionContext::TestBasicStringExpression(const FTextFilterString& InValue
@@ -95,30 +96,12 @@ UMovieScene* FSequencerTextFilterExpressionContext::GetFocusedGetMovieScene() co
 	return IsValid(FocusedMovieSceneSequence) ? FocusedMovieSceneSequence->GetMovieScene() : nullptr;
 }
 
-bool FSequencerTextFilterExpressionContext::CompareFStringForExactBool(const FTextFilterString& InValue, const bool bInPassedFilter) const
+UMovieSceneTrack* FSequencerTextFilterExpressionContext::GetMovieSceneTrack() const
 {
-	if (InValue.CompareFString(TEXT("TRUE"), ETextFilterTextComparisonMode::Exact))
-	{
-		return bInPassedFilter;
-	}
-	if (InValue.CompareFString(TEXT("FALSE"), ETextFilterTextComparisonMode::Exact))
-	{
-		return !bInPassedFilter;
-	}
-	return true;
+	return FilterInterface.GetFilterData().ResolveMovieSceneTrackObject(FilterItem);
 }
 
-bool FSequencerTextFilterExpressionContext::CompareFStringForExactBool(const FTextFilterString& InValue
-	, const ETextFilterComparisonOperation InComparisonOperation
-	, const bool bInPassedFilter) const
+UObject* FSequencerTextFilterExpressionContext::GetBoundObject() const
 {
-	switch (InComparisonOperation)
-	{
-	case ETextFilterComparisonOperation::Equal:
-		return CompareFStringForExactBool(InValue, bInPassedFilter);
-	case ETextFilterComparisonOperation::NotEqual:
-		return CompareFStringForExactBool(InValue, !bInPassedFilter);
-	default:
-		return true;
-	}
+	return FilterInterface.GetFilterData().ResolveTrackBoundObject(FilterInterface.GetSequencer(), FilterItem);
 }
diff --git a/Engine/Source/Editor/Sequencer/Private/Filters/TextExpressions/SequencerTextFilterExpression_TrackClass.cpp b/Engine/Source/Editor/Sequencer/Private/Filters/TextExpressions/SequencerTextFilterExpression_TrackClass.cpp
--- a/Engine/Source/Editor/Sequencer/Private/Filters/TextExpressions/SequencerTextFilterExpression_TrackClass.cpp
+++ b/Engine/Source/Editor/Sequencer/Private/Filters/TextExpressions/SequencerTextFilterExpression_TrackClass.cpp
@@ -38,9 +38,10 @@ bool FSequencerTextFilterExpression_TrackClass::TestComplexExpression(const FNam
 		return true;
 	}
 
-	if (WeakTrackObject.IsValid())
+	UMovieSceneTrack* const TrackObject = GetMovieSceneTrack();
+	if (IsValid(TrackObject))
 	{
-		const FString TrackClassName = WeakTrackObject->GetClass()->GetName();
+		const FString TrackClassName = TrackObject->GetClass()->GetName();
 		if (TextFilterUtils::TestComplexExpression(TrackClassName, InValue, InComparisonOperation, InTextComparisonMode))
 		{
 			return true;
diff --git a/Engine/Source/Editor/Sequencer/Private/Filters/Widgets/SSequencerCustomTextFilterDialog.cpp b/Engine/Source/Editor/Sequencer/Private/Filters/Widgets/SSequencerCustomTextFilterDialog.cpp
--- a/Engine/Source/Editor/Sequencer/Private/Filters/Widgets/SSequencerCustomTextFilterDialog.cpp
+++ b/Engine/Source/Editor/Sequencer/Private/Filters/Widgets/SSequencerCustomTextFilterDialog.cpp
@@ -4,7 +4,7 @@
 #include "Sequencer.h"
 #include "SPrimaryButton.h"
 #include "Filters/Filters/SequencerTrackFilter_CustomText.h"
-#include "Filters/SequencerFilterBar.h"
+#include "Filters/ISequencerTrackFilters.h"
 #include "Interfaces/IMainFrameModule.h"
 #include "Widgets/Colors/SColorBlock.h"
 #include "Widgets/Colors/SColorPicker.h"
@@ -14,7 +14,7 @@
 
 TSharedPtr<SSequencerCustomTextFilterDialog> SSequencerCustomTextFilterDialog::DialogInstance;
 
-void SSequencerCustomTextFilterDialog::Construct(const FArguments& InArgs, const TSharedRef<FSequencerFilterBar>& InFilterBar)
+void SSequencerCustomTextFilterDialog::Construct(const FArguments& InArgs, const TSharedRef<ISequencerTrackFilters>& InFilterBar)
 {
 	WeakFilterBar = InFilterBar;
 
@@ -78,7 +78,7 @@ void SSequencerCustomTextFilterDialog::Construct(const FArguments& InArgs, const
 	OnWindowClosed.BindSP(this, &SSequencerCustomTextFilterDialog::HandleWindowClosed);
 }
 
-void SSequencerCustomTextFilterDialog::CreateWindow_AddCustomTextFilter(const TSharedRef<FSequencerFilterBar>& InFilterBar
+void SSequencerCustomTextFilterDialog::CreateWindow_AddCustomTextFilter(const TSharedRef<ISequencerTrackFilters>& InFilterBar
 	, const FCustomTextFilterData& InCustomTextFilterData
 	, const TSharedPtr<SWindow> InParentWindow)
 {
@@ -95,7 +95,7 @@ void SSequencerCustomTextFilterDialog::CreateWindow_AddCustomTextFilter(const TS
 	ShowWindow(DialogInstance.ToSharedRef(), true, InParentWindow);
 }
 
-void SSequencerCustomTextFilterDialog::CreateWindow_EditCustomTextFilter(const TSharedRef<FSequencerFilterBar>& InFilterBar
+void SSequencerCustomTextFilterDialog::CreateWindow_EditCustomTextFilter(const TSharedRef<ISequencerTrackFilters>& InFilterBar
 	, TSharedPtr<FSequencerTrackFilter_CustomText> InCustomTextFilter
 	, const TSharedPtr<SWindow> InParentWindow)
 {
@@ -141,8 +141,6 @@ void SSequencerCustomTextFilterDialog::ShowWindow(const TSharedRef<SWindow>& InW
 	}
 	else
 	{
-		FSlateApplication::Get().AddWindow(InWindowToShow);
-
 		if (InParentWindow.IsValid())
 		{
 			FSlateApplication::Get().AddWindowAsNativeChild(InWindowToShow, InParentWindow.ToSharedRef());
@@ -322,7 +320,7 @@ TSharedRef<SWidget> SSequencerCustomTextFilterDialog::ConstructButtonRow()
 
 bool SSequencerCustomTextFilterDialog::CheckFilterNameValidity() const
 {
-	const TSharedPtr<FSequencerFilterBar> FilterBar = WeakFilterBar.Pin();
+	const TSharedPtr<ISequencerTrackFilters> FilterBar = WeakFilterBar.Pin();
 	if (!FilterBar.IsValid())
 	{
 		return false;
@@ -383,7 +381,7 @@ FReply SSequencerCustomTextFilterDialog::OnColorBlockMouseButtonDown(const FGeom
 
 void SSequencerCustomTextFilterDialog::OnCreateCustomTextFilter(const bool bInApplyFilter)
 {
-	const TSharedPtr<FSequencerFilterBar> FilterBar = WeakFilterBar.Pin();
+	const TSharedPtr<ISequencerTrackFilters> FilterBar = WeakFilterBar.Pin();
 	if (!FilterBar.IsValid())
 	{
 		return;
@@ -425,7 +423,7 @@ void SSequencerCustomTextFilterDialog::OnCreateCustomTextFilter(const bool bInAp
 
 void SSequencerCustomTextFilterDialog::OnModifyCustomTextFilter()
 {
-	const TSharedPtr<FSequencerFilterBar> FilterBar = WeakFilterBar.Pin();
+	const TSharedPtr<ISequencerTrackFilters> FilterBar = WeakFilterBar.Pin();
 	if (!FilterBar.IsValid())
 	{
 		return;
@@ -497,7 +495,7 @@ FReply SSequencerCustomTextFilterDialog::OnDeleteButtonClick()
 		return FReply::Handled();
 	}
 
-	const TSharedPtr<FSequencerFilterBar> FilterBar = WeakFilterBar.Pin();
+	const TSharedPtr<ISequencerTrackFilters> FilterBar = WeakFilterBar.Pin();
 	if (!FilterBar.IsValid())
 	{
 		return FReply::Handled();
diff --git a/Engine/Source/Editor/Sequencer/Private/Filters/Widgets/SSequencerCustomTextFilterDialog.h b/Engine/Source/Editor/Sequencer/Private/Filters/Widgets/SSequencerCustomTextFilterDialog.h
--- a/Engine/Source/Editor/Sequencer/Private/Filters/Widgets/SSequencerCustomTextFilterDialog.h
+++ b/Engine/Source/Editor/Sequencer/Private/Filters/Widgets/SSequencerCustomTextFilterDialog.h
@@ -8,16 +8,15 @@
 
 class ISequencerTrackFilters;
 class FSequencerTrackFilter_CustomText;
-class FSequencerFilterBar;
 
 class SSequencerCustomTextFilterDialog : public SWindow
 {
 public:
-	static void CreateWindow_AddCustomTextFilter(const TSharedRef<FSequencerFilterBar>& InFilterBar
+	static void CreateWindow_AddCustomTextFilter(const TSharedRef<ISequencerTrackFilters>& InFilterBar
 		, const FCustomTextFilterData& InCustomTextFilterData = FCustomTextFilterData()
 		, const TSharedPtr<SWindow> InParentWindow = nullptr);
 
-	static void CreateWindow_EditCustomTextFilter(const TSharedRef<FSequencerFilterBar>& InFilterBar
+	static void CreateWindow_EditCustomTextFilter(const TSharedRef<ISequencerTrackFilters>& InFilterBar
 		, TSharedPtr<FSequencerTrackFilter_CustomText> InCustomTextFilter
 		, const TSharedPtr<SWindow> InParentWindow = nullptr);
 
@@ -32,7 +31,7 @@ public:
 		SLATE_ARGUMENT(FCustomTextFilterData, CustomTextFilterData)
 	SLATE_END_ARGS()
 
-	void Construct(const FArguments& InArgs, const TSharedRef<FSequencerFilterBar>& InFilterBar);
+	void Construct(const FArguments& InArgs, const TSharedRef<ISequencerTrackFilters>& InFilterBar);
 
 private:
 	static void ShowWindow(const TSharedRef<SWindow>& InWindowToShow, const bool bInModal, const TSharedPtr<SWindow>& InParentWindow = nullptr);
@@ -59,7 +58,7 @@ private:
 
 	void HandleWindowClosed(const TSharedRef<SWindow>& InWindow);
 
-	TWeakPtr<FSequencerFilterBar> WeakFilterBar;
+	TWeakPtr<ISequencerTrackFilters> WeakFilterBar;
 
 	TSharedPtr<FSequencerTrackFilter_CustomText> CustomTextFilter;
 
diff --git a/Engine/Source/Editor/Sequencer/Private/Filters/Widgets/SSequencerFilterBar.h b/Engine/Source/Editor/Sequencer/Private/Filters/Widgets/SSequencerFilterBar.h
--- a/Engine/Source/Editor/Sequencer/Private/Filters/Widgets/SSequencerFilterBar.h
+++ b/Engine/Source/Editor/Sequencer/Private/Filters/Widgets/SSequencerFilterBar.h
@@ -19,18 +19,6 @@ enum class ESequencerFilterChange : uint8;
 class SSequencerFilterBar : public SCompoundWidget
 {
 public:
-	/** Delegate for when filters have changed */
-	DECLARE_DELEGATE(FOnFilterChanged);
-
-	/** Delegate to create a TTextFilter used to compare FilterType with text queries */
-	DECLARE_DELEGATE_RetVal(TSharedPtr<FSequencerTrackFilter_CustomText>, FCreateTextFilter);
-
-	DECLARE_DELEGATE_OneParam(FOnFilterBarLayoutChanging, EFilterBarLayout /*InNewLayout*/);
-
-	DECLARE_EVENT(FSequencerFilterBar, FSequencerFiltersChanged);
-
-	DECLARE_EVENT_OneParam(FSequencerFilterBar, FExternalCustomTextFilterEvent, const TSharedRef<FSequencerFilterBar>& /*InBroadcastingFilterBar*/);
-
 	SLATE_BEGIN_ARGS(SSequencerFilterBar)
 		: _FilterBarLayout(EFilterBarLayout::Vertical)
 		, _CanChangeOrientation(true)
diff --git a/Engine/Source/Editor/Sequencer/Private/Filters/Widgets/SSequencerFilterSuggestionListRow.cpp b/Engine/Source/Editor/Sequencer/Private/Filters/Widgets/SSequencerFilterSuggestionListRow.cpp
--- a/Engine/Source/Editor/Sequencer/Private/Filters/Widgets/SSequencerFilterSuggestionListRow.cpp
+++ b/Engine/Source/Editor/Sequencer/Private/Filters/Widgets/SSequencerFilterSuggestionListRow.cpp
@@ -77,8 +77,8 @@ TSharedRef<SWidget> SSequencerFilterSuggestionListRow::CreateSuggestionItem(cons
 			.Padding(12.f, 1.f, 9.f, 1.f)
 			[
 				SNew(STextBlock)
-				.Text(InSuggestionListItem->DisplayName)
-				.ToolTipText(InSuggestionListItem->Description)
+				.Text(InSuggestionListItem->Suggestion.DisplayName)
+				.ToolTipText(InSuggestionListItem->Suggestion.Description)
 				.HighlightText(HighlightText)
 			];
 	}
@@ -90,7 +90,7 @@ TSharedRef<SWidget> SSequencerFilterSuggestionListRow::CreateSuggestionItem(cons
 			.Padding(3.f, 1.f)
 			[
 				SNew(STextBlock)
-				.Text(InSuggestionListItem->Description)
+				.Text(InSuggestionListItem->Suggestion.Description)
 				.TextStyle(FAppStyle::Get(), TEXT("HintText"))
 			];
 	}
diff --git a/Engine/Source/Editor/Sequencer/Private/Filters/Widgets/SequencerFilterSuggestionListEntry.h b/Engine/Source/Editor/Sequencer/Private/Filters/Widgets/SequencerFilterSuggestionListEntry.h
--- a/Engine/Source/Editor/Sequencer/Private/Filters/Widgets/SequencerFilterSuggestionListEntry.h
+++ b/Engine/Source/Editor/Sequencer/Private/Filters/Widgets/SequencerFilterSuggestionListEntry.h
@@ -2,6 +2,7 @@
 
 #pragma once
 
+#include "Filters/SequencerFilterSuggestion.h"
 #include "Internationalization/Text.h"
 #include "Templates/SharedPointer.h"
 
@@ -62,11 +63,5 @@ public:
 		: FSequencerFilterSuggestionListEntryBase(ESequencerFilterSuggestionListEntryType::Suggestion)
 	{}
 
-	FString Suggestion;
-
-	FText DisplayName;
-
-	FText CategoryName;
-
-	FText Description;
+	FSequencerFilterSuggestion Suggestion;
 };
diff --git a/Engine/Source/Editor/Sequencer/Private/SSequencer.cpp b/Engine/Source/Editor/Sequencer/Private/SSequencer.cpp
--- a/Engine/Source/Editor/Sequencer/Private/SSequencer.cpp
+++ b/Engine/Source/Editor/Sequencer/Private/SSequencer.cpp
@@ -557,7 +557,6 @@ TSharedRef<SWidget> SSequencer::ConstructFilterBarContent()
 			{
 				GetSequencerSettings()->SetLastFilterBarSizeCoefficient(InNewCoefficient);
 			})
-		.Expose(FilterSplitterSlot)
 		[
 			SFilterBarClippingHorizontalBox::WrapVerticalListWithHeading(FilterBarWidgetRef
 				, FPointerEventHandler::CreateSP(FilterBarWidgetRef, &SSequencerFilterBar::OnMouseButtonUp))
@@ -1361,14 +1360,14 @@ TSharedRef<SWidget> SSequencer::ConstructSearchAndFilterRow()
 			.OnSaveSearchClicked(this, &SSequencer::OnOutlinerSearchSaved);
 	}
 
-	FilterBarWidget = Sequencer->GetFilterBar()->GenerateWidget(SearchBox, GetFilterBarLayout());
+	FilterBarWidget = FilterBar->GenerateWidget(SearchBox, GetFilterBarLayout());
 
 	if (!FilterComboButtonWidget.IsValid())
 	{
 		FilterComboButtonWidget = FilterBar->MakeAddFilterButton();
 	}
 
-	FilterBarWidget->SetMuted(Sequencer->GetFilterBar()->AreFiltersMuted());
+	FilterBarWidget->SetMuted(FilterBar->AreFiltersMuted());
 
 	return SNew(SHorizontalBox)
 
@@ -4067,12 +4066,10 @@ bool SSequencer::IsFilterBarVisible() const
 
 void SSequencer::ToggleFilterBarVisibility()
 {
-	bool bNewFilterBarVisible = true;
-
 	USequencerSettings* const SequencerSettings = GetSequencerSettings();
 	if (ensure(SequencerSettings))
 	{
-		bNewFilterBarVisible = !SequencerSettings->IsFilterBarVisible();
+		const bool bNewFilterBarVisible = !SequencerSettings->IsFilterBarVisible();
 		SequencerSettings->SetFilterBarVisible(bNewFilterBarVisible);
 	}
 
@@ -4122,17 +4119,14 @@ void SSequencer::RebuildSearchAndFilterRow()
 
 	if (FilterBarWidget.IsValid())
 	{
-		if (const TSharedPtr<FSequencerFilterBar> FilterBar = FilterBarWidget->GetFilterBar())
+		if (IsFilterBarVisible() && GetFilterBarLayout() == EFilterBarLayout::Horizontal)
 		{
-			if (IsFilterBarVisible() && GetFilterBarLayout() == EFilterBarLayout::Horizontal)
-			{
-				SearchAndFilterRow->AddSlot()
-					.AutoHeight()
-					.Padding(0)
-					[
-						FilterBarWidget.ToSharedRef()
-					];
-			}
+			SearchAndFilterRow->AddSlot()
+				.AutoHeight()
+				.Padding(0)
+				[
+					FilterBarWidget.ToSharedRef()
+				];
 		}
 	}
 }
diff --git a/Engine/Source/Editor/Sequencer/Private/SSequencer.h b/Engine/Source/Editor/Sequencer/Private/SSequencer.h
--- a/Engine/Source/Editor/Sequencer/Private/SSequencer.h
+++ b/Engine/Source/Editor/Sequencer/Private/SSequencer.h
@@ -888,7 +888,5 @@ private:
 	TSharedPtr<SComboButton> ViewOptionsComboButton;
 	TSharedPtr<FSequencerViewOptionsMenu> ViewOptionsMenu;
 
-	SSplitter::FSlot* FilterSplitterSlot = nullptr;
-
 	TSharedPtr<SVerticalBox> SearchAndFilterRow;
 };
diff --git a/Engine/Source/Editor/Sequencer/Private/Sequencer.cpp b/Engine/Source/Editor/Sequencer/Private/Sequencer.cpp
--- a/Engine/Source/Editor/Sequencer/Private/Sequencer.cpp
+++ b/Engine/Source/Editor/Sequencer/Private/Sequencer.cpp
@@ -169,6 +169,7 @@
 #include "Systems/MovieSceneMotionVectorSimulationSystem.h"
 #include "IKeyArea.h"
 #include "Editor/TransBuffer.h"
+#include "Filters/SequencerFilterBar.h"
 #include "Sidebar/SidebarDrawerConfig.h"
 #include "Widgets/Sidebar/SequencerSelectionDrawer.h"
 #include "Misc/Thumbnail/ThumbnailCaptureUtils.h"
@@ -183,7 +184,6 @@
 #include "Bindings/MovieSceneSpawnableBinding.h"
 #include "Bindings/MovieSceneSpawnableActorBinding.h"
 #include "Bindings/MovieSceneReplaceableActorBinding.h"
-#include "Filters/SequencerFilterBar.h"
 #include "TrackEditors/TimeWarpTrackEditor.h"
 #include "Variants/MovieSceneTimeWarpGetter.h"
 
@@ -336,6 +336,8 @@ void FSequencer::InitSequencer(const FSequencerInitParams& InitParams
 	using namespace UE::MovieScene;
 	using namespace UE::Sequencer;
 
+	const TSharedRef<FSequencer> SharedThisRef = SharedThis(this);
+
 	bIsEditingWithinLevelEditor = InitParams.bEditWithinLevelEditor;
 	ScrubStyle = InitParams.ViewParams.ScrubberStyle;
 	HostCapabilities = InitParams.HostCapabilities;
@@ -540,7 +542,7 @@ void FSequencer::InitSequencer(const FSequencerInitParams& InitParams
 	PlaybackSpeedBeforePlay = PlaybackSpeed;
 	ShuttleMultiplier = 0;
 	ObjectChangeListener = InObjectChangeListener;
-	PropertyKeyedStatusHandler = MakeShared<FSequencerPropertyKeyedStatusHandler>(SharedThis(this));
+	PropertyKeyedStatusHandler = MakeShared<FSequencerPropertyKeyedStatusHandler>(SharedThisRef);
 
 	check( ObjectChangeListener.IsValid() );
 	
@@ -617,7 +619,7 @@ void FSequencer::InitSequencer(const FSequencerInitParams& InitParams
 		UMovieSceneSignedObject::SetDeferredHandler(DeferredSignedObjectChangeHandler);
 	}
 
-	ViewModel = MakeShared<FSequencerEditorViewModel>(SharedThis(this), GetHostCapabilities());
+	ViewModel = MakeShared<FSequencerEditorViewModel>(SharedThisRef, GetHostCapabilities());
 	ViewModel->InitializeEditor();
 	ViewModel->SetSequence(InitParams.RootSequence);
 	ViewModel->GetRootSequenceModel()->InitializeExtensions();
@@ -639,7 +641,6 @@ void FSequencer::InitSequencer(const FSequencerInitParams& InitParams
 	PlayPosition.Reset(ConvertFrameTime(GetPlaybackRange().GetLowerBoundValue(), GetRootTickResolution(), PlayPosition.GetInputRate()));
 
 	FilterBar = MakeShared<FSequencerFilterBar>(*this);
-	FilterBar->CreateCustomTextFiltersFromConfig();
 
 	// Make internal widgets
 	SequencerWidget = SNew( SSequencer, SharedThis( this ) )
@@ -705,7 +706,7 @@ void FSequencer::InitSequencer(const FSequencerInitParams& InitParams
 	{
 		check(EditorObjectBindingDelegates[DelegateIndex].IsBound());
 		// Object bindings may exist in other modules, call a delegate that will create one for us 
-		TSharedRef<ISequencerEditorObjectBinding> ObjectBinding = EditorObjectBindingDelegates[DelegateIndex].Execute(SharedThis(this));
+		TSharedRef<ISequencerEditorObjectBinding> ObjectBinding = EditorObjectBindingDelegates[DelegateIndex].Execute(SharedThisRef);
 		ObjectBindings.Add(ObjectBinding);
 	}
 
@@ -724,7 +725,7 @@ void FSequencer::InitSequencer(const FSequencerInitParams& InitParams
 	if (CVarTimeUndo->GetBool())
 	{
 		TimeUndoRedoHandler.UndoRedoProxy = NewObject<USequencerTimeChangeUndoRedoProxy>(GetTransientPackage(), NAME_None);
-		TimeUndoRedoHandler.SetSequencer(SharedThis(this));
+		TimeUndoRedoHandler.SetSequencer(SharedThisRef);
 		TimeUndoRedoHandler.UndoRedoProxy->SetFlags(RF_Transactional | RF_Transient);
 	}
 
@@ -770,7 +771,7 @@ void FSequencer::InitSequencer(const FSequencerInitParams& InitParams
 		DetailsDrawer.InitialState = Settings->GetSidebarState().FindOrAddDrawerState(SelectionDrawerId);
 		RegisterDrawer(MoveTemp(DetailsDrawer));
 
-		RegisterDrawerSection(SelectionDrawerId, MakeShared<FSequencerSelectionDrawer>(SharedThis(this)));
+		RegisterDrawerSection(SelectionDrawerId, MakeShared<FSequencerSelectionDrawer>(SharedThisRef));
 	}
 }
 
@@ -881,7 +882,6 @@ FSequencer::~FSequencer()
 	RootTemplateInstance.TearDown();
 }
 
-
 void FSequencer::Close()
 {
 	for (FLevelEditorViewportClient* LevelVC : GEditor->GetLevelViewportClients())
diff --git a/Engine/Source/Editor/Sequencer/Private/Widgets/SSequencerSearchBox.cpp b/Engine/Source/Editor/Sequencer/Private/Widgets/SSequencerSearchBox.cpp
--- a/Engine/Source/Editor/Sequencer/Private/Widgets/SSequencerSearchBox.cpp
+++ b/Engine/Source/Editor/Sequencer/Private/Widgets/SSequencerSearchBox.cpp
@@ -258,7 +258,7 @@ void SSequencerSearchBox::HandleTextCommitted(const FText& InNewText, const ETex
 	FText CommittedText;
 	if (SelectedSuggestion.IsValid() && !SelectedSuggestion->IsHeader() && InCommitType != ETextCommit::OnCleared)
 	{
-		CommittedText = OnSuggestionChosen.Execute(InNewText, SelectedSuggestion->AsSuggestionEntry()->Suggestion);
+		CommittedText = OnSuggestionChosen.Execute(InNewText, SelectedSuggestion->AsSuggestionEntry()->Suggestion.Suggestion);
 	}
 	else
 	{
@@ -309,7 +309,7 @@ void SSequencerSearchBox::OnSelectionChanged(TSharedPtr<FSequencerFilterSuggesti
 	// If the user clicked directly on an item to select it, then accept the choice and close the window
 	if (InSelectInfo == ESelectInfo::OnMouseClick && !InNewValue->IsHeader())
 	{
-		const FText NewText = OnSuggestionChosen.Execute(SearchBox->GetText(), InNewValue->AsSuggestionEntry()->Suggestion);
+		const FText NewText = OnSuggestionChosen.Execute(SearchBox->GetText(), InNewValue->AsSuggestionEntry()->Suggestion.Suggestion);
 		SetText(NewText);
 
 		SuggestionMenuAnchor->SetIsOpen(false, false);
@@ -379,9 +379,7 @@ void SSequencerSearchBox::UpdateSuggestionList()
 			TArray<TSharedPtr<FSequencerFilterSuggestionListEntryBase>>& CategorySuggestions = CategorizedSuggestions.FindOrAdd(Suggestion.CategoryName);
 
 			const TSharedRef<FSequencerFilterSuggestionListEntry> NewSuggestionEntry = MakeShared<FSequencerFilterSuggestionListEntry>();
-			NewSuggestionEntry->Suggestion = Suggestion.Suggestion;
-			NewSuggestionEntry->DisplayName = Suggestion.DisplayName;
-			NewSuggestionEntry->Description = Suggestion.Description;
+			NewSuggestionEntry->Suggestion = Suggestion;
 
 			CategorySuggestions.Add(NewSuggestionEntry);
 		}
@@ -448,7 +446,7 @@ void SSequencerSearchBox::SetOnSaveSearchHandler(SFilterSearchBox::FOnSaveSearch
 	SearchBox->SetOnSaveSearchHandler(InOnSaveSearchHandler);
 }
 
-void SSequencerSearchBox::ExtractSearchFilterTerms(const FText& InSearchText, FParsedExpression& OutParsedPair)
+void SSequencerSearchBox::ExtractSearchFilterTerms(const FText& InSearchText, FParsedExpression& OutParsedPair) const
 {
 	const TSharedPtr<ISequencerFilterBar> FilterBar = WeakFilterBar.Pin();
 	if (!FilterBar.IsValid())
@@ -675,13 +673,11 @@ int32 SSequencerSearchBox::FindTokenIndex(const TArray<FExpressionToken>& Expres
 
 const FTextFilterExpressionEvaluator* SSequencerSearchBox::GetTextFilterExpressionEvaluator() const
 {
-	const TSharedPtr<ISequencerFilterBar> FilterBar = WeakFilterBar.Pin();
-	if (!FilterBar.IsValid())
+	if (const TSharedPtr<ISequencerFilterBar> FilterBar = WeakFilterBar.Pin())
 	{
-		return nullptr;
+		return &FilterBar->GetTextFilterExpressionEvaluator();
 	}
-
-	return &FilterBar->GetTextFilterExpressionEvaluator();
+	return nullptr;
 }
 
 #undef LOCTEXT_NAMESPACE
diff --git a/Engine/Source/Editor/Sequencer/Private/Widgets/SSequencerSearchBox.h b/Engine/Source/Editor/Sequencer/Private/Widgets/SSequencerSearchBox.h
--- a/Engine/Source/Editor/Sequencer/Private/Widgets/SSequencerSearchBox.h
+++ b/Engine/Source/Editor/Sequencer/Private/Widgets/SSequencerSearchBox.h
@@ -2,6 +2,7 @@
 
 #pragma once
 
+#include "Filters/SequencerFilterSuggestion.h"
 #include "Filters/SFilterSearchBox.h"
 #include "Framework/Text/TextLayout.h"
 #include "Misc/ExpressionParserTypesFwd.h"
@@ -18,17 +19,6 @@ struct FFocusEvent;
 struct FGeometry;
 struct FKeyEvent;
 
-struct FSequencerFilterSuggestion
-{
-	FString Suggestion;
-
-	FText DisplayName;
-
-	FText CategoryName;
-
-	FText Description;
-};
-
 /** A delegate for a callback to filter the given suggestion list, to allow custom filtering behavior */
 DECLARE_DELEGATE_ThreeParams(FOnSequencerSearchBoxSuggestionFilter, const FText& /*InSearchText*/, TArray<FSequencerFilterSuggestion>& /*InPossibleSuggestions*/, FText& /*InSuggestionHighlightText*/);
 
@@ -121,7 +111,7 @@ private:
 		int32 SuggestionIndex = 0;
 	};
 
-	void ExtractSearchFilterTerms(const FText& InSearchText, FParsedExpression& OutParsedPair);
+	void ExtractSearchFilterTerms(const FText& InSearchText, FParsedExpression& OutParsedPair) const;
 
 	TSharedRef<SWidget> GetSuggestionListMenuContent();
 
@@ -164,10 +154,10 @@ private:
 	/** Provides default value suggestions, if no suggestions are provided */
 	void DefaultValueSuggestions(const FString& InKeyValue, TArray<FSequencerFilterSuggestion>& OutPossibleSuggestions);
 
-	bool IsOperatorToken(const TExpressionToken<TCHAR>& InToken);
-	bool IsLogicalOperatorToken(const TExpressionToken<TCHAR>& InToken);
+	static bool IsOperatorToken(const TExpressionToken<TCHAR>& InToken);
+	static bool IsLogicalOperatorToken(const TExpressionToken<TCHAR>& InToken);
 
-	int32 FindTokenIndex(const TArray<TExpressionToken<TCHAR>>& ExpressionTokens, const int32 InIndexToFind);
+	static int32 FindTokenIndex(const TArray<TExpressionToken<TCHAR>>& ExpressionTokens, const int32 InIndexToFind);
 
 	const FTextFilterExpressionEvaluator* GetTextFilterExpressionEvaluator() const;
 
diff --git a/Engine/Source/Editor/Sequencer/Public/Filters/Filters/SequencerTrackFilter_Modified.h b/Engine/Source/Editor/Sequencer/Public/Filters/Filters/SequencerTrackFilter_Modified.h
--- a/Engine/Source/Editor/Sequencer/Public/Filters/Filters/SequencerTrackFilter_Modified.h
+++ b/Engine/Source/Editor/Sequencer/Public/Filters/Filters/SequencerTrackFilter_Modified.h
@@ -26,6 +26,4 @@ public:
 	virtual FString GetName() const override;
 	virtual bool PassesFilter(FSequencerTrackFilterType InItem) const override;
 	//~ End IFilter
-
-	void ToggleShowOnlyModifiedTracks();
 };
diff --git a/Engine/Source/Editor/Sequencer/Public/Filters/Filters/SequencerTrackFilter_Unbound.h b/Engine/Source/Editor/Sequencer/Public/Filters/Filters/SequencerTrackFilter_Unbound.h
--- a/Engine/Source/Editor/Sequencer/Public/Filters/Filters/SequencerTrackFilter_Unbound.h
+++ b/Engine/Source/Editor/Sequencer/Public/Filters/Filters/SequencerTrackFilter_Unbound.h
@@ -26,6 +26,4 @@ public:
 	virtual FString GetName() const override;
 	virtual bool PassesFilter(FSequencerTrackFilterType InItem) const override;
 	//~ End IFilter
-
-	void ToggleShowOnlyModifiedTracks();
 };
diff --git a/Engine/Source/Editor/Sequencer/Public/Filters/ISequencerFilterBar.h b/Engine/Source/Editor/Sequencer/Public/Filters/ISequencerFilterBar.h
--- a/Engine/Source/Editor/Sequencer/Public/Filters/ISequencerFilterBar.h
+++ b/Engine/Source/Editor/Sequencer/Public/Filters/ISequencerFilterBar.h
@@ -2,8 +2,10 @@
 
 #pragma once
 
+#include "Filters/SBasicFilterBar.h"
 #include "Framework/Commands/UICommandList.h"
 
+class FSequencerTrackFilter;
 class FSequencerTrackFilter_CustomText;
 class FString;
 class FText;
@@ -12,9 +14,11 @@ class ISequencer;
 class ISequencerTextFilterExpressionContext;
 
 /** Generic interface for all Sequencer filter bar implementations. */
-class ISequencerFilterBar
+class ISequencerFilterBar : public TSharedFromThis<ISequencerFilterBar>
 {
 public:
+	DECLARE_EVENT_TwoParams(ISequencerFilterBar, FOnFilterBarStateChanged, const bool /*InIsVisible*/, const EFilterBarLayout /*InNewLayout*/);
+
 	virtual ~ISequencerFilterBar() = default;
 
 	SEQUENCER_API virtual FName GetIdentifier() const = 0;
@@ -30,6 +34,12 @@ public:
 	 * The text expression must have key, operator, and value tokens. */
 	SEQUENCER_API virtual bool DoesTextFilterStringContainExpressionPair(const ISequencerTextFilterExpressionContext& InExpression) const = 0;
 
+	SEQUENCER_API virtual bool AreFiltersMuted() const = 0;
+	SEQUENCER_API virtual void MuteFilters(const bool bInMute) = 0;
+
+	SEQUENCER_API virtual bool CanResetFilters() const = 0;
+	SEQUENCER_API virtual void ResetFilters() = 0;
+
 	SEQUENCER_API virtual bool HasAnyFilterActive(const bool bCheckTextFilter = true
 		, const bool bInCheckHideIsolateFilter = true
 		, const bool bInCheckCommonFilters = true
@@ -44,7 +54,8 @@ public:
 
 	SEQUENCER_API virtual void ActivateCommonFilters(const bool bInActivate, const TArray<FString>& InExceptionFilterNames) = 0;
 
-	SEQUENCER_API virtual void ActivateAllEnabledFilters(const bool bInActivate, const TArray<FString>& InExceptionFilterNames) = 0;
+	SEQUENCER_API virtual bool AreAllEnabledFiltersActive(const bool bInActive, const TArray<FString> InExceptionFilterNames) const = 0;
+	SEQUENCER_API virtual void ActivateAllEnabledFilters(const bool bInActivate, const TArray<FString> InExceptionFilterNames) = 0;
 
 	SEQUENCER_API virtual bool IsFilterActiveByDisplayName(const FString& InFilterName) const = 0;
 	SEQUENCER_API virtual bool IsFilterEnabledByDisplayName(const FString& InFilterName) const = 0;
@@ -59,4 +70,24 @@ public:
 
 	SEQUENCER_API virtual const FTextFilterExpressionEvaluator& GetTextFilterExpressionEvaluator() const = 0;
 	SEQUENCER_API virtual TArray<TSharedRef<ISequencerTextFilterExpressionContext>> GetTextFilterExpressionContexts() const = 0;
+
+	SEQUENCER_API virtual void OpenTextExpressionHelp() = 0;
+	SEQUENCER_API virtual void SaveCurrentFilterSetAsCustomTextFilter() = 0;
+	SEQUENCER_API virtual void CreateNewTextFilter() = 0;
+
+	/** Event called when the visibility or layout state has changed */
+	SEQUENCER_API virtual FOnFilterBarStateChanged& OnStateChanged() = 0;
+
+	/** Event called to request that the subscribe update its view since the filters state has changed */
+	SEQUENCER_API virtual FSimpleMulticastDelegate& OnRequestUpdate() = 0;
+
+	SEQUENCER_API virtual bool ShouldShowFilterBarWidget() const = 0;
+
+	SEQUENCER_API virtual bool IsFilterBarVisible() const = 0;
+	SEQUENCER_API virtual void ToggleFilterBarVisibility() = 0;
+
+	SEQUENCER_API virtual bool IsFilterBarLayout(const EFilterBarLayout InLayout) const = 0;
+	SEQUENCER_API virtual void SetToVerticalLayout() = 0;
+	SEQUENCER_API virtual void SetToHorizontalLayout() = 0;
+	SEQUENCER_API virtual void ToggleFilterBarLayout() = 0;
 };
diff --git a/Engine/Source/Editor/Sequencer/Public/Filters/ISequencerTextFilterExpressionContext.h b/Engine/Source/Editor/Sequencer/Public/Filters/ISequencerTextFilterExpressionContext.h
--- a/Engine/Source/Editor/Sequencer/Public/Filters/ISequencerTextFilterExpressionContext.h
+++ b/Engine/Source/Editor/Sequencer/Public/Filters/ISequencerTextFilterExpressionContext.h
@@ -3,6 +3,7 @@
 #pragma once
 
 #include "Misc/TextFilterExpressionEvaluator.h"
+#include "Misc/TextFilterUtils.h"
 
 enum class SEQUENCER_API ESequencerTextFilterValueType : uint8
 {
@@ -19,10 +20,8 @@ struct SEQUENCER_API FSequencerTextFilterKeyword
 
 /** Extends the ITextFilterExpressionContext interface to add support for suggestions. */
 class ISequencerTextFilterExpressionContext : public ITextFilterExpressionContext
-{ 
+{
 public:
-	//~ Begin ISequencerTextFilterExpressionContext
-
 	SEQUENCER_API virtual TSet<FName> GetKeys() const = 0;
 
 	SEQUENCER_API virtual ESequencerTextFilterValueType GetValueType() const = 0;
@@ -31,5 +30,31 @@ public:
 	SEQUENCER_API virtual FText GetDescription() const = 0;
 	SEQUENCER_API virtual FText GetCategory() const { return FText::GetEmpty(); }
 
-	//~ End ISequencerTextFilterExpressionContext
+	static bool CompareFStringForExactBool(const FTextFilterString& InValue, const bool bInPassedFilter)
+	{
+		if (InValue.CompareFString(TEXT("TRUE"), ETextFilterTextComparisonMode::Exact))
+		{
+			return bInPassedFilter;
+		}
+		if (InValue.CompareFString(TEXT("FALSE"), ETextFilterTextComparisonMode::Exact))
+		{
+			return !bInPassedFilter;
+		}
+		return true;
+	}
+
+	static bool CompareFStringForExactBool(const FTextFilterString& InValue
+		, const ETextFilterComparisonOperation InComparisonOperation
+		, const bool bInPassedFilter)
+	{
+		switch (InComparisonOperation)
+		{
+		case ETextFilterComparisonOperation::Equal:
+			return CompareFStringForExactBool(InValue, bInPassedFilter);
+		case ETextFilterComparisonOperation::NotEqual:
+			return CompareFStringForExactBool(InValue, !bInPassedFilter);
+		default:
+			return true;
+		}
+	}
 };
diff --git a/Engine/Source/Editor/Sequencer/Public/Filters/ISequencerTrackFilters.h b/Engine/Source/Editor/Sequencer/Public/Filters/ISequencerTrackFilters.h
--- a/Engine/Source/Editor/Sequencer/Public/Filters/ISequencerTrackFilters.h
+++ b/Engine/Source/Editor/Sequencer/Public/Filters/ISequencerTrackFilters.h
@@ -4,6 +4,7 @@
 
 #include "ISequencerFilterBar.h"
 
+class FSequencerTrackFilter;
 class FSequencerTrackFilter_CustomText;
 struct FSequencerFilterData;
 
@@ -12,6 +13,12 @@ class ISequencerTrackFilters : public ISequencerFilterBar
 public:
 	virtual ~ISequencerTrackFilters() override = default;
 
+	SEQUENCER_API virtual bool IsFilterEnabled(TSharedRef<FSequencerTrackFilter> InFilter) const = 0;
+	SEQUENCER_API virtual bool SetFilterEnabled(const TSharedRef<FSequencerTrackFilter> InFilter, const bool bInEnabled, const bool bInRequestFilterUpdate = true) = 0;
+
+	SEQUENCER_API virtual bool IsFilterActive(const TSharedRef<FSequencerTrackFilter> InFilter) const = 0;
+	SEQUENCER_API virtual bool SetFilterActive(const TSharedRef<FSequencerTrackFilter>& InFilter, const bool bInActive, const bool bInRequestFilterUpdate = true) = 0;
+
 	SEQUENCER_API virtual bool AddCustomTextFilter(const TSharedRef<FSequencerTrackFilter_CustomText>& InFilter, const bool bInAddToConfig) = 0;
 	SEQUENCER_API virtual bool RemoveCustomTextFilter(const TSharedRef<FSequencerTrackFilter_CustomText>& InFilter, const bool bInAddToConfig) = 0;
 
diff --git a/Engine/Source/Editor/Sequencer/Public/Filters/SequencerFilterSuggestion.h b/Engine/Source/Editor/Sequencer/Public/Filters/SequencerFilterSuggestion.h
new file mode 100644
--- /dev/null
+++ b/Engine/Source/Editor/Sequencer/Public/Filters/SequencerFilterSuggestion.h
@@ -0,0 +1,29 @@
+﻿// Copyright Epic Games, Inc. All Rights Reserved.
+
+#pragma once
+
+#include "Containers/UnrealString.h"
+#include "Internationalization/Text.h"
+
+struct FSequencerFilterSuggestion
+{
+	/** The raw suggestion string that should be used with the search box */
+	FString Suggestion;
+
+	/** The user-facing display name of this suggestion */
+	FText DisplayName;
+
+	/** The user-facing category name of this suggestion (if any) */
+	FText CategoryName;
+
+	/** Describes what this suggestion will do */
+	FText Description;
+
+	static FSequencerFilterSuggestion MakeSimpleSuggestion(FString InSuggestionString)
+	{
+		FSequencerFilterSuggestion Suggestion;
+		Suggestion.Suggestion = MoveTemp(InSuggestionString);
+		Suggestion.DisplayName = FText::FromString(Suggestion.Suggestion);
+		return MoveTemp(Suggestion);
+	}
+};
diff --git a/Engine/Source/Editor/Sequencer/Public/Filters/SequencerTextFilterExpressionContext.h b/Engine/Source/Editor/Sequencer/Public/Filters/SequencerTextFilterExpressionContext.h
--- a/Engine/Source/Editor/Sequencer/Public/Filters/SequencerTextFilterExpressionContext.h
+++ b/Engine/Source/Editor/Sequencer/Public/Filters/SequencerTextFilterExpressionContext.h
@@ -16,7 +16,7 @@ class FSequencerTextFilterExpressionContext : public ISequencerTextFilterExpress
 public:
 	SEQUENCER_API FSequencerTextFilterExpressionContext(ISequencerTrackFilters& InFilterInterface);
 
-	void SetFilterItem(FSequencerTrackFilterType InFilterItem, UMovieSceneTrack* const InTrackObject);
+	void SetFilterItem(FSequencerTrackFilterType InFilterItem);
 
 	//~ Begin ITextFilterExpressionContext
 
@@ -34,15 +34,10 @@ protected:
 	SEQUENCER_API UMovieSceneSequence* GetFocusedMovieSceneSequence() const;
 	SEQUENCER_API UMovieScene* GetFocusedGetMovieScene() const;
 
-	SEQUENCER_API bool CompareFStringForExactBool(const FTextFilterString& InValue, const bool bInPassedFilter) const;
-
-	SEQUENCER_API bool CompareFStringForExactBool(const FTextFilterString& InValue
-		, const ETextFilterComparisonOperation InComparisonOperation
-		, const bool bInPassedFilter) const;
+	SEQUENCER_API UMovieSceneTrack* GetMovieSceneTrack() const;
+	SEQUENCER_API UObject* GetBoundObject() const;
 
 	ISequencerTrackFilters& FilterInterface;
 
 	FSequencerTrackFilterType FilterItem;
-
-	TWeakObjectPtr<UMovieSceneTrack> WeakTrackObject;
 };
