Sequencer: Fix crash in reduce tolerance OnValueCommitted_Lambda.

Don't capture by reference

diff --git a/Engine/Source/Editor/Sequencer/Private/SequencerContextMenus.cpp b/Engine/Source/Editor/Sequencer/Private/SequencerContextMenus.cpp
--- a/Engine/Source/Editor/Sequencer/Private/SequencerContextMenus.cpp
+++ b/Engine/Source/Editor/Sequencer/Private/SequencerContextMenus.cpp
@@ -767,7 +767,7 @@ void FSectionContextMenu::AddKeyEditMenu(FMenuBuilder& MenuBuilder)
 		[
 			SNew(SSpinBox<float>)
 			.Style(&FAppStyle::GetWidgetStyle<FSpinBoxStyle>(TEXT("Sequencer.HyperlinkSpinBox")))
-			.OnValueCommitted_Lambda([&OnReduceKeysToleranceChanged](float Value, ETextCommit::Type)
+			.OnValueCommitted_Lambda([OnReduceKeysToleranceChanged](float Value, ETextCommit::Type)
 			{
 				OnReduceKeysToleranceChanged(Value);
 			})
