Sequencer: Backing out 26129795 because it no longer repros and it was causing Sequencer to close when it wasn't necessary

#rnx

diff --git a/Engine/Source/Editor/Sequencer/Private/Sequencer.cpp b/Engine/Source/Editor/Sequencer/Private/Sequencer.cpp
--- a/Engine/Source/Editor/Sequencer/Private/Sequencer.cpp
+++ b/Engine/Source/Editor/Sequencer/Private/Sequencer.cpp
@@ -499,13 +499,6 @@ void FSequencer::InitSequencer(const FSequencerInitParams& InitParams
 				}
 			}
 
-			// If our playback client is getting killed (e.g. it's a Blueprint that has been recompiled), we have
-			// to bail out.
-			if (CachedPlaybackContext.IsStale())
-			{
-				GEditor->GetEditorSubsystem<UAssetEditorSubsystem>()->CloseAllEditorsForAsset(GetRootMovieSceneSequence());
-				return;
-			}
 			if (UObject* PlaybackContext = CachedPlaybackContext.Get())
 			{
 				if (ReplacementMap.Contains(PlaybackContext->GetClass()))
