[Sequencer] Fix "Isolate" filtering regression not properly showing control rig tracks

diff --git a/Engine/Source/Editor/Sequencer/Private/Filters/SequencerFilterBar.cpp b/Engine/Source/Editor/Sequencer/Private/Filters/SequencerFilterBar.cpp
--- a/Engine/Source/Editor/Sequencer/Private/Filters/SequencerFilterBar.cpp
+++ b/Engine/Source/Editor/Sequencer/Private/Filters/SequencerFilterBar.cpp
@@ -1486,11 +1486,6 @@ bool FSequencerFilterBar::FilterNodesRecursive(const bool bInHasActiveFilter, co
 			SetTrackParentsExpanded(InStartNode.ImplicitCast(), true);
 		}
 
-		if (bPassedTextFilter)
-		{
-			FilterData.FilterInParentChildNodes(InStartNode, true, true, true);
-		}
-
 		FilterData.FilterInNodeWithAncestors(InStartNode);
 		return true;
 	}
