Sequencer: Add a toggle to Link Filters with Curve Editor.

The curves shown in the curve editor area populdated by the Sequencer tree. You can filter out nodes in the Sequencer tree and those will propagate to the curve editor by default. Users want the ability to bypass those filters and still show curves in the curve editor that are filterd out by the Sequencer tree. This toggle allows you to do that.

diff --git a/Engine/Source/Editor/Sequencer/Private/Filters/Menus/SequencerViewOptionsMenu.cpp b/Engine/Source/Editor/Sequencer/Private/Filters/Menus/SequencerViewOptionsMenu.cpp
--- a/Engine/Source/Editor/Sequencer/Private/Filters/Menus/SequencerViewOptionsMenu.cpp
+++ b/Engine/Source/Editor/Sequencer/Private/Filters/Menus/SequencerViewOptionsMenu.cpp
@@ -1,4 +1,4 @@
-﻿// Copyright Epic Games, Inc. All Rights Reserved.
+// Copyright Epic Games, Inc. All Rights Reserved.
 
 #include "SequencerViewOptionsMenu.h"
 #include "Filters/Menus/SequencerMenuContext.h"
@@ -132,6 +132,17 @@ void FSequencerViewOptionsMenu::PopulateFilterOptionsSection(UToolMenu& InMenu)
 			FIsActionChecked::CreateRaw(this, &FSequencerViewOptionsMenu::IsAutoExpandPassedFilterNodes)
 		),
 		EUserInterfaceActionType::ToggleButton);
+
+	OptionsSection.AddMenuEntry(TEXT("LinkFiltersWithCurveEditor"),
+		LOCTEXT("LinkFiltersWithCurveEditor", "Link Filters with Curve Editor"),
+		LOCTEXT("LinkFiltersWithCurveEditorToolTip", "Toggle linking filters with the curve editor tree"),
+		FSlateIcon(),
+		FUIAction(
+			FExecuteAction::CreateRaw(this, &FSequencerViewOptionsMenu::ToggleLinkFiltersWithCurveEditor),
+			FCanExecuteAction(),
+			FIsActionChecked::CreateRaw(this, &FSequencerViewOptionsMenu::IsLinkFiltersWithCurveEditor)
+		),
+		EUserInterfaceActionType::ToggleButton);
 }
 
 void FSequencerViewOptionsMenu::PopulateLayoutSection(UToolMenu& InMenu)
@@ -265,6 +276,42 @@ void FSequencerViewOptionsMenu::ToggleAutoExpandPassedFilterNodes()
 	Sequencer->GetFilterInterface()->RequestFilterUpdate();
 }
 
+bool FSequencerViewOptionsMenu::IsLinkFiltersWithCurveEditor() const
+{
+	const TSharedPtr<FSequencer> Sequencer = WeakSequencer.Pin();
+	if (!Sequencer.IsValid())
+	{
+		return false;
+	}
+
+	const USequencerSettings* const SequencerSettings = Sequencer->GetSequencerSettings();
+	if (!IsValid(SequencerSettings))
+	{
+		return false;
+	}
+
+	return SequencerSettings->GetLinkFiltersWithCurveEditor();
+}
+
+void FSequencerViewOptionsMenu::ToggleLinkFiltersWithCurveEditor()
+{
+	const TSharedPtr<FSequencer> Sequencer = WeakSequencer.Pin();
+	if (!Sequencer.IsValid())
+	{
+		return;
+	}
+
+	USequencerSettings* const SequencerSettings = Sequencer->GetSequencerSettings();
+	if (!IsValid(SequencerSettings))
+	{
+		return;
+	}
+
+	SequencerSettings->SetLinkFiltersWithCurveEditor(!SequencerSettings->GetLinkFiltersWithCurveEditor());
+
+	Sequencer->GetFilterInterface()->RequestFilterUpdate();
+}
+
 TSharedPtr<SSequencer> FSequencerViewOptionsMenu::GetSequencerWidget() const
 {
 	const TSharedPtr<FSequencer> Sequencer = WeakSequencer.Pin();
diff --git a/Engine/Source/Editor/Sequencer/Private/Filters/Menus/SequencerViewOptionsMenu.h b/Engine/Source/Editor/Sequencer/Private/Filters/Menus/SequencerViewOptionsMenu.h
--- a/Engine/Source/Editor/Sequencer/Private/Filters/Menus/SequencerViewOptionsMenu.h
+++ b/Engine/Source/Editor/Sequencer/Private/Filters/Menus/SequencerViewOptionsMenu.h
@@ -1,4 +1,4 @@
-﻿// Copyright Epic Games, Inc. All Rights Reserved.
+// Copyright Epic Games, Inc. All Rights Reserved.
 
 #pragma once
 
@@ -31,6 +31,9 @@ protected:
 	bool IsAutoExpandPassedFilterNodes() const;
 	void ToggleAutoExpandPassedFilterNodes();
 
+	bool IsLinkFiltersWithCurveEditor() const;
+	void ToggleLinkFiltersWithCurveEditor();
+
 	TSharedPtr<SSequencer> GetSequencerWidget() const;
 
 	TWeakPtr<FSequencer> WeakSequencer;
diff --git a/Engine/Source/Editor/Sequencer/Private/MVVM/CurveEditorIntegrationExtension.cpp b/Engine/Source/Editor/Sequencer/Private/MVVM/CurveEditorIntegrationExtension.cpp
--- a/Engine/Source/Editor/Sequencer/Private/MVVM/CurveEditorIntegrationExtension.cpp
+++ b/Engine/Source/Editor/Sequencer/Private/MVVM/CurveEditorIntegrationExtension.cpp
@@ -12,6 +12,8 @@
 #include "MVVM/Extensions/ICurveEditorTreeItemExtension.h"
 #include "CurveEditor.h"
 
+#include "SequencerSettings.h"
+
 namespace UE
 {
 namespace Sequencer
@@ -67,6 +69,8 @@ void FCurveEditorIntegrationExtension::UpdateCurveEditor()
 		return;
 	}
 
+	const bool bUseFilteredOut = OwnerModel->GetSequencer()->GetSequencerSettings()->GetLinkFiltersWithCurveEditor();
+
 	FCurveEditorTree* CurveEditorTree = CurveEditor->GetTree();
 
 	// Guard against multiple broadcasts here and defer them until the end of this function
@@ -81,7 +85,7 @@ void FCurveEditorIntegrationExtension::UpdateCurveEditor()
 		if (ViewModel)
 		{
 			TViewModelPtr<IOutlinerExtension> OutlinerModel = ViewModel.ImplicitCast();
-			const bool bIsVisible = OutlinerModel && !OutlinerModel->IsFilteredOut();
+			const bool bIsVisible = bUseFilteredOut ? OutlinerModel && !OutlinerModel->IsFilteredOut() : true;
 			
 			const FCurveEditorTreeItem* TreeItem = CurveEditorTree->FindItem(ItemID);
 
@@ -119,7 +123,7 @@ void FCurveEditorIntegrationExtension::UpdateCurveEditor()
 	bool bItemsAdded = false;
 	for (TParentFirstChildIterator<IOutlinerExtension> It(OwnerModel, bIncludeRootNode); It; ++It)
 	{
-		if (It->IsFilteredOut())
+		if (bUseFilteredOut && It->IsFilteredOut())
 		{
 			It.IgnoreCurrentChildren();
 			continue;
diff --git a/Engine/Source/Editor/Sequencer/Private/SequencerCommands.cpp b/Engine/Source/Editor/Sequencer/Private/SequencerCommands.cpp
--- a/Engine/Source/Editor/Sequencer/Private/SequencerCommands.cpp
+++ b/Engine/Source/Editor/Sequencer/Private/SequencerCommands.cpp
@@ -155,6 +155,7 @@ void FSequencerCommands::RegisterCommands()
 
 	UI_COMMAND( ToggleShowCurveEditor, "Curve Editor", "Show the animation keys in a curve editor", EUserInterfaceActionType::ToggleButton, FInputChord() );
 	UI_COMMAND( ToggleLinkCurveEditorTimeRange, "Link Curve Editor Time Range", "Link the curve editor time range to the sequence", EUserInterfaceActionType::ToggleButton, FInputChord() );
+	UI_COMMAND( ToggleLinkFiltersWithCurveEditor, "Link Filters with Curve Editor", "Link the curve editor tree view with the sequencer filters", EUserInterfaceActionType::ToggleButton, FInputChord());
 
 	UI_COMMAND( ToggleShowPreAndPostRoll, "Pre/Post Roll", "Toggles visualization of pre and post roll", EUserInterfaceActionType::ToggleButton, FInputChord() );
 
diff --git a/Engine/Source/Editor/Sequencer/Private/SequencerSettings.cpp b/Engine/Source/Editor/Sequencer/Private/SequencerSettings.cpp
--- a/Engine/Source/Editor/Sequencer/Private/SequencerSettings.cpp
+++ b/Engine/Source/Editor/Sequencer/Private/SequencerSettings.cpp
@@ -34,6 +34,7 @@ USequencerSettings::USequencerSettings( const FObjectInitializer& ObjectInitiali
 	ZoomPosition = ESequencerZoomPosition::SZP_CurrentTime;
 	bAutoScrollEnabled = false;
 	bLinkCurveEditorTimeRange = false;
+	bLinkFiltersWithCurveEditor = true;
 	bSynchronizeCurveEditorSelection = true;
 	bIsolateCurveEditorToSelection = true;
 	bCurveEditorVisible = false;
@@ -500,6 +501,20 @@ void USequencerSettings::SetLinkCurveEditorTimeRange(bool InbLinkCurveEditorTime
 	}
 }
 
+bool USequencerSettings::GetLinkFiltersWithCurveEditor() const
+{
+	return bLinkFiltersWithCurveEditor;
+}
+
+void USequencerSettings::SetLinkFiltersWithCurveEditor(bool bInLinkFiltersWithCurveEditor)
+{
+	if (bLinkFiltersWithCurveEditor != bInLinkFiltersWithCurveEditor)
+	{
+		bLinkFiltersWithCurveEditor = bInLinkFiltersWithCurveEditor;
+		SaveConfig();
+	}
+}
+
 void USequencerSettings::SyncCurveEditorSelection(bool bInSynchronizeCurveEditorSelection)
 {
 	if (bSynchronizeCurveEditorSelection != bInSynchronizeCurveEditorSelection)
diff --git a/Engine/Source/Editor/Sequencer/Public/SequencerCommands.h b/Engine/Source/Editor/Sequencer/Public/SequencerCommands.h
--- a/Engine/Source/Editor/Sequencer/Public/SequencerCommands.h
+++ b/Engine/Source/Editor/Sequencer/Public/SequencerCommands.h
@@ -397,6 +397,9 @@ public:
 	/** Toggles whether or not the curve editor time range should be linked to the sequencer. */
 	TSharedPtr< FUICommandInfo > ToggleLinkCurveEditorTimeRange;
 
+	/** Toggles whether or not the curve editor time tree view should be linked with the sequencer filters. */
+	TSharedPtr< FUICommandInfo > ToggleLinkFiltersWithCurveEditor;
+
 	/** Toggles visualization of pre and post roll */
 	TSharedPtr< FUICommandInfo > ToggleShowPreAndPostRoll;
 
diff --git a/Engine/Source/Editor/Sequencer/Public/SequencerSettings.h b/Engine/Source/Editor/Sequencer/Public/SequencerSettings.h
--- a/Engine/Source/Editor/Sequencer/Public/SequencerSettings.h
+++ b/Engine/Source/Editor/Sequencer/Public/SequencerSettings.h
@@ -269,6 +269,11 @@ public:
 	/** Sets whether or not to link the curve editor time range. */
 	void SetLinkCurveEditorTimeRange(bool InbLinkCurveEditorTimeRange);
 
+	/** Gets whether or not to link sequence filters with the curve editor tree */
+	bool GetLinkFiltersWithCurveEditor() const;
+	/** Sets whether or not to link sequence filters with the curve editor tree */
+	void SetLinkFiltersWithCurveEditor(bool bInLinkFiltersWithCurveEditor);
+
 	/** Return true if we are to synchronize the curve editor and sequencer trees */
 	bool ShouldSyncCurveEditorSelection() const { return bSynchronizeCurveEditorSelection; }
 	/** Assign whether we are to synchronize the curve editor and sequencer trees */
@@ -651,6 +656,10 @@ protected:
 	UPROPERTY( config, EditAnywhere, Category=CurveEditor )
 	bool bLinkCurveEditorTimeRange;
 
+	/** Enable or disable linking sequence filters with the curve editor tree */
+	UPROPERTY(config, EditAnywhere, Category = CurveEditor)
+	bool bLinkFiltersWithCurveEditor;
+
 	/** When enabled, changing the sequencer tree selection will also select the relevant nodes in the curve editor tree if possible. */
 	UPROPERTY( config, EditAnywhere, Category=CurveEditor )
 	bool bSynchronizeCurveEditorSelection;
