Sequencer: Change synchronization behavior so that the corresponding bound actor is only selected when child Outliner Items are selected. When you select keys/sections/layer bars, the corresponding bound actor is no longer automatically selected.

diff --git a/Engine/Source/Editor/Sequencer/Private/Sequencer.cpp b/Engine/Source/Editor/Sequencer/Private/Sequencer.cpp
--- a/Engine/Source/Editor/Sequencer/Private/Sequencer.cpp
+++ b/Engine/Source/Editor/Sequencer/Private/Sequencer.cpp
@@ -6676,16 +6676,13 @@ void FSequencer::SynchronizeExternalSelectionWithSequencerSelection()
 	TSet<AActor*> SelectedSequencerActors;
 	TSet<UActorComponent*> SelectedSequencerComponents;
 
-	TSet<TWeakViewModelPtr<IOutlinerExtension>> Items = ViewModel->GetSelection()->GetNodesWithSelectedKeysOrSections();
-	Items.Append(ViewModel->GetSelection()->Outliner.GetSelected());
-
 	const FName ControlRigEditModeModeName("EditMode.ControlRig");
 	const bool bHACK_ControlRigEditMode = GLevelEditorModeTools().GetActiveMode(ControlRigEditModeModeName) != nullptr;
 
-	for ( TWeakViewModelPtr<IOutlinerExtension> WeakItem : Items)
+	for (TWeakViewModelPtr<IOutlinerExtension> WeakOutlinerItem : ViewModel->GetSelection()->Outliner.GetSelected())
 	{
-		FViewModelPtr SelectedItem = WeakItem.Pin();
-		if (!SelectedItem)
+		FViewModelPtr OutlinerItem = WeakOutlinerItem.Pin();
+		if (!OutlinerItem)
 		{
 			continue;
 		}
@@ -6693,8 +6690,8 @@ void FSequencer::SynchronizeExternalSelectionWithSequencerSelection()
 		//HACK for DHI, if we have an active control rig then one is selected so don't find a parent actor or compomonent to select	
 		//but if we do select the actor/compoent directly we still select it.
 		TViewModelPtr<IObjectBindingExtension> ObjectBinding = bHACK_ControlRigEditMode
-			? CastViewModel<IObjectBindingExtension>(SelectedItem)
-			: SelectedItem->FindAncestorOfType<IObjectBindingExtension>(true);
+			? CastViewModel<IObjectBindingExtension>(OutlinerItem)
+			: OutlinerItem->FindAncestorOfType<IObjectBindingExtension>(true);
 
 		// If the closest node is an object node, try to get the actor/component nodes from it.
 	if (ObjectBinding)
