[Sequencer] Navigation Tool initial submit

diff --git a/Engine/Plugins/MovieScene/LevelSequenceEditor/Source/LevelSequenceEditor/Private/LevelSequenceEditorToolkit.cpp b/Engine/Plugins/MovieScene/LevelSequenceEditor/Source/LevelSequenceEditor/Private/LevelSequenceEditorToolkit.cpp
--- a/Engine/Plugins/MovieScene/LevelSequenceEditor/Source/LevelSequenceEditor/Private/LevelSequenceEditorToolkit.cpp
+++ b/Engine/Plugins/MovieScene/LevelSequenceEditor/Source/LevelSequenceEditor/Private/LevelSequenceEditorToolkit.cpp
@@ -44,6 +44,8 @@
 #include "ClassViewerModule.h"
 #include "SequencerUtilities.h"
 #include "ViewportToolbar/UnrealEdViewportToolbarContext.h"
+#include "LevelSequenceNavigationToolProvider.h"
+#include "Tools/NavigationTool/Customization/NavigationToolExtender.h"
 
 // @todo sequencer: hack: setting defaults for transform tracks
 
@@ -248,6 +250,7 @@ void FLevelSequenceEditorToolkit::Initialize(const EToolkitMode::Type Mode, cons
 		SequencerInitParams.HostCapabilities.bSupportsAddFromContentBrowser = true;
 		SequencerInitParams.HostCapabilities.bSupportsSidebar = true;
 		SequencerInitParams.HostCapabilities.bSupportsViewportSelectability = true;
+		SequencerInitParams.HostCapabilities.bSupportsNavigationTool = true;
 	}
 
 	ExtendSequencerToolbar("Sequencer.MainToolBar");
@@ -315,6 +318,13 @@ void FLevelSequenceEditorToolkit::Initialize(const EToolkitMode::Type Mode, cons
 			SelectionSection.AddEntry(Entry);
 		}
 	}
+
+	using namespace UE::Sequencer::NavigationTool;
+	if (const TSharedPtr<INavigationTool> Tool = Sequencer->GetNavigationTool())
+	{
+		NavigationToolProvider = MakeShared<FLevelSequenceNavigationToolProvider>();
+		FNavigationToolExtender::RegisterOnlyToolProvider(Tool.ToSharedRef(), NavigationToolProvider.ToSharedRef());
+	}
 }
 
 
@@ -980,6 +990,16 @@ void FLevelSequenceEditorToolkit::SaveAsset_Execute()
 
 void FLevelSequenceEditorToolkit::OnClose()
 {
+	if (NavigationToolProvider.IsValid())
+	{
+		using namespace UE::Sequencer::NavigationTool;
+		if (const TSharedPtr<INavigationTool> Tool = Sequencer->GetNavigationTool())
+		{
+			FNavigationToolExtender::UnregisterToolProvider(Tool.ToSharedRef(), NavigationToolProvider->GetIdentifier());
+		}
+		NavigationToolProvider.Reset();
+	}
+
 	// Remove the viewport toolbar extensions we added earlier.
 	UToolMenus::Get()->UnregisterOwnerByName(UE::MovieScene::Private::ViewportToolbarOwnerName);
 
diff --git a/Engine/Plugins/MovieScene/LevelSequenceEditor/Source/LevelSequenceEditor/Private/LevelSequenceEditorToolkit.h b/Engine/Plugins/MovieScene/LevelSequenceEditor/Source/LevelSequenceEditor/Private/LevelSequenceEditorToolkit.h
--- a/Engine/Plugins/MovieScene/LevelSequenceEditor/Source/LevelSequenceEditor/Private/LevelSequenceEditorToolkit.h
+++ b/Engine/Plugins/MovieScene/LevelSequenceEditor/Source/LevelSequenceEditor/Private/LevelSequenceEditorToolkit.h
@@ -21,6 +21,11 @@ class FLevelSequencePlaybackContext;
 class UPrimitiveComponent;
 enum class EMapChangeType : uint8;
 
+namespace UE::Sequencer::NavigationTool
+{
+	class FLevelSequenceNavigationToolProvider;
+}
+
 /**
  * Implements an Editor toolkit for level sequences.
  */
@@ -147,6 +152,8 @@ private:
 	/** The sequencer used by this editor. */
 	TSharedPtr<ISequencer> Sequencer;
 
+	TSharedPtr<UE::Sequencer::NavigationTool::FLevelSequenceNavigationToolProvider> NavigationToolProvider;
+
 	/** Pointer to the style set to use for toolkits. */
 	TSharedRef<ISlateStyle> Style;
 
diff --git a/Engine/Plugins/MovieScene/LevelSequenceEditor/Source/LevelSequenceEditor/Private/LevelSequenceNavigationToolProvider.cpp b/Engine/Plugins/MovieScene/LevelSequenceEditor/Source/LevelSequenceEditor/Private/LevelSequenceNavigationToolProvider.cpp
new file mode 100644
--- /dev/null
+++ b/Engine/Plugins/MovieScene/LevelSequenceEditor/Source/LevelSequenceEditor/Private/LevelSequenceNavigationToolProvider.cpp
@@ -0,0 +1,183 @@
+﻿// Copyright Epic Games, Inc. All Rights Reserved.
+
+#include "LevelSequenceNavigationToolProvider.h"
+#include "ISequencer.h"
+#include "ISourceControlModule.h"
+#include "LevelSequence.h"
+#include "MovieScene.h"
+#include "Tools/NavigationTool/Columns/NavigationToolColorColumn.h"
+#include "Tools/NavigationTool/Columns/NavigationToolColumnExtender.h"
+#include "Tools/NavigationTool/Columns/NavigationToolCommentColumn.h"
+#include "Tools/NavigationTool/Columns/NavigationToolHBiasColumn.h"
+#include "Tools/NavigationTool/Columns/NavigationToolInTimeColumn.h"
+#include "Tools/NavigationTool/Columns/NavigationToolItemsColumn.h"
+#include "Tools/NavigationTool/Columns/NavigationToolLabelColumn.h"
+#include "Tools/NavigationTool/Columns/NavigationToolLengthColumn.h"
+#include "Tools/NavigationTool/Columns/NavigationToolLockColumn.h"
+#include "Tools/NavigationTool/Columns/NavigationToolMarkerVisibilityColumn.h"
+#include "Tools/NavigationTool/Columns/NavigationToolOutTimeColumn.h"
+#include "Tools/NavigationTool/Columns/NavigationToolRevisionControlColumn.h"
+#include "Tools/NavigationTool/Columns/NavigationToolStartFrameOffsetColumn.h"
+#include "Tools/NavigationTool/Columns/NavigationToolTakeColumn.h"
+#include "Tools/NavigationTool/Columns/NavigationToolDeactiveStateColumn.h"
+#include "Tools/NavigationTool/Filters/Filters/NavigationToolBuiltInFilterParams.h"
+#include "Tools/NavigationTool/Filters/Filters/NavigationToolFilterBase.h"
+#include "Tools/NavigationTool/INavigationTool.h"
+#include "Tools/NavigationTool/Items/NavigationToolSequence.h"
+#include "Tools/NavigationTool/NavigationToolSettings.h"
+
+#define LOCTEXT_NAMESPACE "LevelSequenceNavigationToolProvider"
+
+namespace UE::Sequencer::NavigationTool
+{
+
+const FName FLevelSequenceNavigationToolProvider::Identifier = TEXT("LevelSequence");
+
+const FName FLevelSequenceNavigationToolProvider::AnimationColumnViewId = TEXT("Animation");
+
+FLevelSequenceNavigationToolProvider::FLevelSequenceNavigationToolProvider()
+	: ToolCommands(MakeShared<FUICommandList>())
+{
+}
+
+void FLevelSequenceNavigationToolProvider::BindCommands(const TSharedRef<FUICommandList>& InCommandList)
+{
+	InCommandList->Append(ToolCommands);
+}
+
+FName FLevelSequenceNavigationToolProvider::GetIdentifier() const
+{
+	return Identifier;
+}
+
+FName FLevelSequenceNavigationToolProvider::GetDefaultColumnView() const
+{
+	return AnimationColumnViewId;
+}
+
+void FLevelSequenceNavigationToolProvider::OnExtendColumns(FNavigationToolColumnExtender& OutExtender)
+{
+	FNavigationToolProvider::OnExtendColumns(OutExtender);
+
+	OutExtender.AddColumn<FNavigationToolDeactiveStateColumn>();
+	OutExtender.AddColumn<FNavigationToolMarkerVisibilityColumn>();
+	OutExtender.AddColumn<FNavigationToolLockColumn>();
+	OutExtender.AddColumn<FNavigationToolColorColumn>();
+	OutExtender.AddColumn<FNavigationToolLabelColumn>();
+	OutExtender.AddColumn<FNavigationToolItemsColumn>();
+	OutExtender.AddColumn<FNavigationToolInTimeColumn>();
+	OutExtender.AddColumn<FNavigationToolOutTimeColumn>();
+	OutExtender.AddColumn<FNavigationToolLengthColumn>();
+	OutExtender.AddColumn<FNavigationToolHBiasColumn>();
+	OutExtender.AddColumn<FNavigationToolStartFrameOffsetColumn>();
+	OutExtender.AddColumn<FNavigationToolTakeColumn>();
+	OutExtender.AddColumn<FNavigationToolCommentColumn>();
+
+	ISourceControlModule& SourceControlModule = ISourceControlModule::Get();
+	if (SourceControlModule.IsEnabled())
+	{
+		OutExtender.AddColumn<FNavigationToolRevisionControlColumn>();
+	}
+}
+
+void FLevelSequenceNavigationToolProvider::OnExtendColumnViews(TSet<FNavigationToolColumnView>& OutColumnViews)
+{
+	FNavigationToolColumnView& AnimationColumnView = OutColumnViews
+		.FindOrAdd(FNavigationToolColumnView(AnimationColumnViewId));
+
+	AnimationColumnView.VisibleColumns =
+		{
+			FNavigationToolColorColumn::GetStaticTypeId().ToName(),
+			FNavigationToolLabelColumn::GetStaticTypeId().ToName(),
+			FNavigationToolItemsColumn::GetStaticTypeId().ToName(),
+			FNavigationToolInTimeColumn::GetStaticTypeId().ToName(),
+			FNavigationToolOutTimeColumn::GetStaticTypeId().ToName(),
+			FNavigationToolHBiasColumn::GetStaticTypeId().ToName()
+		};
+}
+
+void FLevelSequenceNavigationToolProvider::OnExtendItemChildren(INavigationTool& InTool
+	, const FNavigationToolItemPtr& InParentItem
+	, TArray<FNavigationToolItemPtr>& OutChildren
+	, const bool bInRecursive)
+{
+	FNavigationToolProvider::OnExtendItemChildren(InTool, InParentItem, OutChildren, bInRecursive);
+
+	// Only extending root item
+	if (InParentItem->GetItemId() != FNavigationToolItemId::RootId)
+	{
+		return;
+	}
+
+	const TSharedPtr<ISequencer> Sequencer = InTool.GetSequencer();
+	if (!Sequencer.IsValid())
+	{
+		return;
+	}
+
+	UMovieSceneSequence* const RootSequence = Sequencer->GetRootMovieSceneSequence();
+	if (!IsValid(RootSequence))
+	{
+		return;
+	}
+
+	const TSharedRef<FLevelSequenceNavigationToolProvider> Provider = SharedThis(this);
+
+	const FNavigationToolItemPtr NewItem = InTool.FindOrAdd<FNavigationToolSequence>(Provider, InParentItem, RootSequence, nullptr, 0);
+	const FNavigationToolItemFlagGuard Guard(NewItem, ENavigationToolItemFlags::IgnorePendingKill);
+	OutChildren.Add(NewItem);
+	if (bInRecursive)
+	{
+		NewItem->FindChildren(OutChildren, bInRecursive);
+	}
+}
+
+void FLevelSequenceNavigationToolProvider::OnExtendBuiltInFilters(TArray<FNavigationToolBuiltInFilterParams>& OutFilterParams)
+{
+	OutFilterParams.Add(FNavigationToolBuiltInFilterParams::CreateSequenceFilter());
+	OutFilterParams.Add(FNavigationToolBuiltInFilterParams::CreateBindingFilter());
+}
+
+FMovieSceneEditorData* FLevelSequenceNavigationToolProvider::GetRootMovieSceneEditorData(const INavigationTool& InTool) const
+{
+	const TSharedPtr<ISequencer> Sequencer = InTool.GetSequencer();
+	if (!Sequencer.IsValid())
+	{
+		return nullptr;
+	}
+
+	UMovieSceneSequence* const RootMovieSceneSequence = Sequencer->GetRootMovieSceneSequence();
+	if (!IsValid(RootMovieSceneSequence))
+	{
+		return nullptr;
+	}
+
+	UMovieScene* const RootMovieScene = RootMovieSceneSequence->GetMovieScene();
+	if (!IsValid(RootMovieScene))
+	{
+		return nullptr;
+	}
+
+	return &RootMovieScene->GetEditorData();
+}
+
+FNavigationToolSaveState* FLevelSequenceNavigationToolProvider::GetSaveState(const INavigationTool& InTool) const
+{
+	if (FMovieSceneEditorData* const EditorData = GetRootMovieSceneEditorData(InTool))
+	{
+		return &EditorData->NavigationToolState;
+	}
+	return nullptr;
+}
+
+void FLevelSequenceNavigationToolProvider::SetSaveState(const INavigationTool& InTool, const FNavigationToolSaveState& InSaveState) const
+{
+	if (FMovieSceneEditorData* const EditorData = GetRootMovieSceneEditorData(InTool))
+	{
+		EditorData->NavigationToolState = InSaveState;
+	}
+}
+
+}
+
+#undef LOCTEXT_NAMESPACE
diff --git a/Engine/Plugins/MovieScene/LevelSequenceEditor/Source/LevelSequenceEditor/Private/LevelSequenceNavigationToolProvider.h b/Engine/Plugins/MovieScene/LevelSequenceEditor/Source/LevelSequenceEditor/Private/LevelSequenceNavigationToolProvider.h
new file mode 100644
--- /dev/null
+++ b/Engine/Plugins/MovieScene/LevelSequenceEditor/Source/LevelSequenceEditor/Private/LevelSequenceNavigationToolProvider.h
@@ -0,0 +1,57 @@
+﻿// Copyright Epic Games, Inc. All Rights Reserved.
+
+#pragma once
+
+#include "Templates/SharedPointer.h"
+#include "Tools/NavigationTool/NavigationToolDefines.h"
+#include "Tools/NavigationTool/Providers/NavigationToolProvider.h"
+#include "UObject/WeakObjectPtr.h"
+
+class FUICommandList;
+class ISequencer;
+struct FMovieSceneEditorData;
+
+namespace UE::Sequencer::NavigationTool
+{
+
+class FNavigationToolItem;
+class INavigationTool;
+
+class FLevelSequenceNavigationToolProvider : public FNavigationToolProvider
+{
+public:
+	static const FName Identifier;
+
+	static const FName AnimationColumnViewId;
+
+	FLevelSequenceNavigationToolProvider();
+	virtual ~FLevelSequenceNavigationToolProvider() override = default;
+
+	void BindCommands(const TSharedRef<FUICommandList>& InCommandList);
+
+	//~ Begin INavigationToolProvider
+
+	virtual FName GetIdentifier() const override;
+
+	virtual FName GetDefaultColumnView() const override;
+
+	virtual void OnExtendColumns(FNavigationToolColumnExtender& OutExtender) override;
+	virtual void OnExtendColumnViews(TSet<FNavigationToolColumnView>& OutColumnViews) override;
+	virtual void OnExtendItemChildren(INavigationTool& InTool
+		, const FNavigationToolItemPtr& InParentItem
+		, TArray<FNavigationToolItemPtr>& OutChildren
+		, const bool bInRecursive) override;
+	virtual void OnExtendBuiltInFilters(TArray<FNavigationToolBuiltInFilterParams>& OutFilterParams) override;
+
+	virtual FNavigationToolSaveState* GetSaveState(const INavigationTool& InTool) const override;
+	virtual void SetSaveState(const INavigationTool& InTool, const FNavigationToolSaveState& InSaveState) const override;
+
+	//~ End INavigationToolProvider
+
+private:
+	FMovieSceneEditorData* GetRootMovieSceneEditorData(const INavigationTool& InTool) const;
+
+	TSharedRef<FUICommandList> ToolCommands;
+};
+
+}
diff --git a/Engine/Source/Editor/EditorStyle/Private/StarshipStyle.cpp b/Engine/Source/Editor/EditorStyle/Private/StarshipStyle.cpp
--- a/Engine/Source/Editor/EditorStyle/Private/StarshipStyle.cpp
+++ b/Engine/Source/Editor/EditorStyle/Private/StarshipStyle.cpp
@@ -2508,6 +2508,28 @@ void FStarshipEditorStyle::FStyle::SetupSequencerStyles()
 		Set("Sequencer.Indicator.Condition", new IMAGE_BRUSH_SVG("Sequencer/Column_Widgets/SequencerCondition", Icon12x12));
 		Set("Sequencer.Indicator.TimeWarp", new IMAGE_BRUSH_SVG("Sequencer/Column_Widgets/SequencerTimeWarp", Icon12x12));
 	}
+
+	// Navigation Tool Styling
+	{
+		const FTableRowStyle& AlternatingTableRowStyle = FAppStyle::Get().GetWidgetStyle<FTableRowStyle>(TEXT("TableView.AlternatingRow"));
+		Set("Sequencer.NavigationTool.TableViewRow", FTableRowStyle(AlternatingTableRowStyle));
+
+		const FSpinBoxStyle FrameTimeSpinStyle = GetWidgetStyle<FSpinBoxStyle>(TEXT("SpinBox"));
+		Set("Sequencer.NavigationTool.SpinBox", FrameTimeSpinStyle);
+
+		FSpinBoxStyle StartTimeSpinStyle = FrameTimeSpinStyle;
+		StartTimeSpinStyle.SetForegroundColor(FStyleColors::AccentGreen);
+		Set("Sequencer.NavigationTool.InTime.SpinBox", StartTimeSpinStyle);
+
+		FSpinBoxStyle EndTimeSpinStyle = FrameTimeSpinStyle;
+		EndTimeSpinStyle.SetForegroundColor(FStyleColors::AccentRed);
+		Set("Sequencer.NavigationTool.OutTime.SpinBox", EndTimeSpinStyle);
+
+		FEditableTextBoxStyle NonEditableTextBoxStyle = GetWidgetStyle<FEditableTextBoxStyle>(TEXT("Graph.EditableTextBox"));
+		NonEditableTextBoxStyle.SetForegroundColor(FSlateColor::UseSubduedForeground());
+		NonEditableTextBoxStyle.BackgroundImageReadOnly = FSlateRoundedBoxBrush(FStyleColors::Background, 4.0f, FSlateColor::UseSubduedForeground(), 1.0f);
+		Set("Sequencer.NavigationTool.NonEditableTextBox", NonEditableTextBoxStyle);
+	}
 }
 
 void FStarshipEditorStyle::FStyle::SetupViewportStyles()
diff --git a/Engine/Source/Editor/Sequencer/Private/Filters/SequencerTrackFilterCommands.cpp b/Engine/Source/Editor/Sequencer/Private/Filters/SequencerTrackFilterCommands.cpp
--- a/Engine/Source/Editor/Sequencer/Private/Filters/SequencerTrackFilterCommands.cpp
+++ b/Engine/Source/Editor/Sequencer/Private/Filters/SequencerTrackFilterCommands.cpp
@@ -15,7 +15,7 @@ FSequencerTrackFilterCommands::FSequencerTrackFilterCommands()
 
 void FSequencerTrackFilterCommands::RegisterCommands()
 {
-	UI_COMMAND(ToggleFilterBarVisibility, "Filter Bar Visibility", "Toggle the visibility of the filter bar", EUserInterfaceActionType::ToggleButton, FInputChord(EModifierKey::Control | EModifierKey::Shift, EKeys::F));
+	UI_COMMAND(ToggleFilterBarVisibility, "Filter Bar Visibility", "Toggle the visibility of the filter bar. The filter bar will only be displayed when there are enabled filters", EUserInterfaceActionType::ToggleButton, FInputChord(EModifierKey::Control | EModifierKey::Shift, EKeys::F));
 
 	UI_COMMAND(ResetFilters, "Reset Filters", "Reset all enabled filters", EUserInterfaceActionType::Button, FInputChord());
 
diff --git a/Engine/Source/Editor/Sequencer/Private/SSequencer.cpp b/Engine/Source/Editor/Sequencer/Private/SSequencer.cpp
--- a/Engine/Source/Editor/Sequencer/Private/SSequencer.cpp
+++ b/Engine/Source/Editor/Sequencer/Private/SSequencer.cpp
@@ -1660,6 +1660,18 @@ void SSequencer::PopulateToolBar(UToolMenu* InMenu)
 			CurveEditorSection.AddEntry(ShowCurveEditorEntry);
 		}
 	}
+
+	if (Sequencer->GetHostCapabilities().bSupportsNavigationTool)
+	{
+		FToolMenuSection& NavigationToolSection = InMenu->FindOrAddSection(TEXT("NavigationTool"));
+
+		FToolMenuEntry ShowNavigationToolEntry = FToolMenuEntry::InitToolBarButton(
+			FSequencerCommands::Get().ToggleNavigationToolVisible,
+			TAttribute<FText>(),TAttribute<FText>(),
+			FSlateIcon(FAppStyle::GetAppStyleSetName(), TEXT("GraphEditor.Sequence_16x")));
+		ShowNavigationToolEntry.StyleNameOverride = SequencerToolbarStyleName;
+		NavigationToolSection.AddEntry(ShowNavigationToolEntry);
+	}
 }
 
 TSharedRef<SWidget> SSequencer::MakeAddMenu()
diff --git a/Engine/Source/Editor/Sequencer/Private/Sequencer.cpp b/Engine/Source/Editor/Sequencer/Private/Sequencer.cpp
--- a/Engine/Source/Editor/Sequencer/Private/Sequencer.cpp
+++ b/Engine/Source/Editor/Sequencer/Private/Sequencer.cpp
@@ -171,6 +171,8 @@
 #include "Editor/TransBuffer.h"
 #include "Filters/SequencerFilterBar.h"
 #include "Sidebar/SidebarDrawerConfig.h"
+#include "Tools/NavigationTool/INavigationTool.h"
+#include "Tools/NavigationTool/NavigationTool.h"
 #include "Widgets/Sidebar/SequencerSelectionDrawer.h"
 #include "Misc/Thumbnail/ThumbnailCaptureUtils.h"
 #include "ToolMenus.h"
@@ -725,6 +727,14 @@ void FSequencer::InitSequencer(const FSequencerInitParams& InitParams
 	// Update the view range to the new current time
 	UpdateTimeBoundsToFocusedMovieScene();
 
+	// NavigationTool must be created before BindCommands()
+	if (HostCapabilities.bSupportsNavigationTool)
+	{
+		NavigationTool = MakeShared<NavigationTool::FNavigationTool>(SharedThisRef);
+		NavigationTool->SetBaseCommandList(SequencerCommandBindings);
+		NavigationTool->Init();
+	}
+
 	// NOTE: Could fill in asset editor commands here!
 
 	BindCommands();
@@ -11960,6 +11970,17 @@ void FSequencer::BindCommands()
 		FExecuteAction::CreateSP(this, &FSequencer::AlignSelectionToPlayhead),
 		FCanExecuteAction::CreateSP(this, &FSequencer::CanAlignSelectionToPlayhead));
 
+	if (NavigationTool.IsValid())
+	{
+		const TSharedRef<UE::Sequencer::NavigationTool::FNavigationTool> NavigationToolRef = NavigationTool.ToSharedRef();
+
+		SequencerCommandBindings->MapAction(
+			Commands.ToggleNavigationToolVisible,
+			FExecuteAction::CreateSP(NavigationToolRef, &UE::Sequencer::NavigationTool::FNavigationTool::ToggleToolTabVisible),
+			FCanExecuteAction(),
+			FIsActionChecked::CreateSP(NavigationToolRef, &UE::Sequencer::NavigationTool::FNavigationTool::IsToolTabVisible));
+	}
+
 	// If this sequencer supports a curve editor, let's add bindings for it.
 	FCurveEditorExtension* CurveEditorExtension = ViewModel->CastDynamic<FCurveEditorExtension>();
 	if (CurveEditorExtension && ensure(CurveEditorExtension->GetCurveEditor()))
@@ -12710,4 +12731,9 @@ void FSequencer::AlignSelectionToPlayhead()
 	FSequencerSelectionAlignmentUtils::AlignSelectionToPlayhead(*this);
 }
 
+TSharedPtr<UE::Sequencer::NavigationTool::INavigationTool> FSequencer::GetNavigationTool() const
+{
+	return NavigationTool;
+}
+
 #undef LOCTEXT_NAMESPACE
diff --git a/Engine/Source/Editor/Sequencer/Private/Sequencer.h b/Engine/Source/Editor/Sequencer/Private/Sequencer.h
--- a/Engine/Source/Editor/Sequencer/Private/Sequencer.h
+++ b/Engine/Source/Editor/Sequencer/Private/Sequencer.h
@@ -43,6 +43,7 @@
 #include "SequencerTimeChangeUndoRedoProxy.h"
 #include "Capabilities/CameraCutViewTargetCacheCapability.h"
 #include "Filters/ISequencerTrackFilters.h"
+#include "Tools/NavigationTool/INavigationTool.h"
 
 class AActor;
 class ACameraActor;
@@ -105,6 +106,14 @@ namespace UE
 		class FViewModel;
 		class FSequencerSelection;
 
+		namespace NavigationTool
+		{
+
+			class FNavigationTool;
+			class INavigationTool;
+
+		}
+
 	} // namespace Sequencer
 } // namespace UE
 
@@ -1270,6 +1279,8 @@ public:
 	/** Checks to see if a clock extension wants the standard PlayRate combo menu hidden */
 	EVisibility GetPlayRateComboVisibility() const;
 
+	virtual TSharedPtr<UE::Sequencer::NavigationTool::INavigationTool> GetNavigationTool() const override;
+
 private:
 
 	/** Update the time bases for the current movie scene */
@@ -1625,4 +1636,7 @@ private:
 	FOnViewportSelectionLimitedChanged OnSelectionLimitedChangedDelegate;
 
 	TSharedPtr<FSequencerFilterBar> FilterBar;
+
+	/** Navigation Tool instance for this Sequencer */
+	TSharedPtr<UE::Sequencer::NavigationTool::FNavigationTool> NavigationTool;
 };
diff --git a/Engine/Source/Editor/Sequencer/Private/SequencerCommands.cpp b/Engine/Source/Editor/Sequencer/Private/SequencerCommands.cpp
--- a/Engine/Source/Editor/Sequencer/Private/SequencerCommands.cpp
+++ b/Engine/Source/Editor/Sequencer/Private/SequencerCommands.cpp
@@ -200,6 +200,8 @@ void FSequencerCommands::RegisterCommands()
 	UI_COMMAND(ToggleSidebarDrawerDock, "Sidebar Drawer Docked", "Toggles the docked state of the sidebar current drawer. Undocks the docked sidebar drawer if docked or docks the sidebar drawer if there is one open and no currently docked drawer.", EUserInterfaceActionType::Button, FInputChord(EModifierKey::Alt, EKeys::D));
 
 	UI_COMMAND(AlignSelectionToPlayhead, "Align Selection To Playhead", "Aligns layer bars, sections, or key selections to the current playhead location", EUserInterfaceActionType::Button, FInputChord());
+
+	UI_COMMAND(ToggleNavigationToolVisible, "Sequence Navigation", "Show the Sequence Navigation Tool", EUserInterfaceActionType::ToggleButton, FInputChord());
 }
 
 #undef LOCTEXT_NAMESPACE
diff --git a/Engine/Source/Editor/Sequencer/Private/SequencerModule.cpp b/Engine/Source/Editor/Sequencer/Private/SequencerModule.cpp
--- a/Engine/Source/Editor/Sequencer/Private/SequencerModule.cpp
+++ b/Engine/Source/Editor/Sequencer/Private/SequencerModule.cpp
@@ -22,6 +22,7 @@
 #include "MovieSceneSignedObject.h"
 #include "MovieSceneSequence.h"
 #include "Sections/MovieSceneSubSection.h"
+#include "Tools/NavigationTool/Customization/NavigationToolExtender.h"
 
 #include "MVVM/CurveEditorExtension.h"
 #include "MVVM/CurveEditorIntegrationExtension.h"
@@ -704,6 +705,11 @@ public:
 		return ObjectSchemas;
 	}
 
+	UE::Sequencer::NavigationTool::FNavigationToolExtender& GetNavigationToolExtender() override
+	{
+		return NavigationToolExtender;
+	}
+
 private:
 
 	TSet<FAnimatedPropertyKey> PropertyAnimators;
@@ -777,6 +783,8 @@ private:
 	// Outliner Indicator Item Delegate Handles
 	FDelegateHandle ConditionOutlinerIndicatorHandle;
 	FDelegateHandle TimeWarpOutlinerIndicatorHandle;
+
+	UE::Sequencer::NavigationTool::FNavigationToolExtender NavigationToolExtender;
 };
 
 IMPLEMENT_MODULE(FSequencerModule, Sequencer);
diff --git a/Engine/Source/Editor/Sequencer/Private/SequencerSettings.cpp b/Engine/Source/Editor/Sequencer/Private/SequencerSettings.cpp
--- a/Engine/Source/Editor/Sequencer/Private/SequencerSettings.cpp
+++ b/Engine/Source/Editor/Sequencer/Private/SequencerSettings.cpp
@@ -75,6 +75,7 @@ USequencerSettings::USequencerSettings( const FObjectInitializer& ObjectInitiali
 	ViewDensity = "Relaxed";
 	AssetBrowserWidth = 500.f;
 	AssetBrowserHeight = 300.f;
+	bNavigationToolVisible = false;
 
 	SectionColorTints.Add(FColor(88, 102, 142, 255)); // blue
 	SectionColorTints.Add(FColor(99, 137, 132, 255)); // blue-green
@@ -1202,3 +1203,12 @@ bool USequencerSettings::ShouldShowThumbnailCaptureSettings()
 {
 	return UE::Sequencer::CVarEnableRelevantThumbnails.GetValueOnGameThread();
 }
+
+void USequencerSettings::SetNavigationToolVisible(const bool bInVisible)
+{
+	if (bNavigationToolVisible != bInVisible)
+	{
+		bNavigationToolVisible = bInVisible;
+		SaveConfig();
+	}
+}
diff --git a/Engine/Source/Editor/Sequencer/Private/Tools/NavigationTool/Columns/NavigationToolColorColumn.cpp b/Engine/Source/Editor/Sequencer/Private/Tools/NavigationTool/Columns/NavigationToolColorColumn.cpp
new file mode 100644
--- /dev/null
+++ b/Engine/Source/Editor/Sequencer/Private/Tools/NavigationTool/Columns/NavigationToolColorColumn.cpp
@@ -0,0 +1,46 @@
+// Copyright Epic Games, Inc. All Rights Reserved.
+
+#include "Tools/NavigationTool/Columns/NavigationToolColorColumn.h"
+#include "Tools/NavigationTool/NavigationToolView.h"
+#include "Tools/NavigationTool/Widgets/Columns/SNavigationToolColor.h"
+#include "Widgets/Layout/SBox.h"
+
+#define LOCTEXT_NAMESPACE "NavigationToolColorColumn"
+
+namespace UE::Sequencer::NavigationTool
+{
+
+FText FNavigationToolColorColumn::GetColumnDisplayNameText() const
+{
+	return LOCTEXT("ColorColumn", "Color");
+}
+
+SHeaderRow::FColumn::FArguments FNavigationToolColorColumn::ConstructHeaderRowColumn(const TSharedRef<INavigationToolView>& InToolView
+	, const float InFillSize)
+{
+	const FName ColumnId = GetColumnId();
+
+	return SHeaderRow::Column(ColumnId)
+	    .FixedWidth(12.f)
+		.HAlignHeader(HAlign_Center)
+		.VAlignHeader(VAlign_Center)
+		.HAlignCell(HAlign_Center)
+		.VAlignCell(VAlign_Center)
+		.DefaultLabel(GetColumnDisplayNameText())
+	    .DefaultTooltip(GetColumnDisplayNameText())
+		[
+			// Display nothing in the column header
+			SNew(SBox)
+		];
+}
+
+TSharedRef<SWidget> FNavigationToolColorColumn::ConstructRowWidget(const FNavigationToolItemPtr& InItem
+	, const TSharedRef<INavigationToolView>& InToolView
+	, const TSharedRef<SNavigationToolTreeRow>& InRow)
+{
+	return SNew(SNavigationToolColor, InItem, InToolView, InRow);
+}
+
+}
+
+#undef LOCTEXT_NAMESPACE
diff --git a/Engine/Source/Editor/Sequencer/Private/Tools/NavigationTool/Columns/NavigationToolColumn.cpp b/Engine/Source/Editor/Sequencer/Private/Tools/NavigationTool/Columns/NavigationToolColumn.cpp
new file mode 100644
--- /dev/null
+++ b/Engine/Source/Editor/Sequencer/Private/Tools/NavigationTool/Columns/NavigationToolColumn.cpp
@@ -0,0 +1,19 @@
+﻿// Copyright Epic Games, Inc. All Rights Reserved.
+
+#include "Tools/NavigationTool/Columns/NavigationToolColumn.h"
+#include "Tools/NavigationTool/Providers/NavigationToolProvider.h"
+
+namespace UE::Sequencer::NavigationTool
+{
+
+TSharedPtr<FNavigationToolProvider> FNavigationToolColumn::GetProvider() const
+{
+	return WeakProvider.Pin();
+}
+
+void FNavigationToolColumn::SetProvider(const TWeakPtr<FNavigationToolProvider>& InWeakProvider)
+{
+	WeakProvider = InWeakProvider;
+}
+
+}
\ No newline at end of file
diff --git a/Engine/Source/Editor/Sequencer/Private/Tools/NavigationTool/Columns/NavigationToolCommentColumn.cpp b/Engine/Source/Editor/Sequencer/Private/Tools/NavigationTool/Columns/NavigationToolCommentColumn.cpp
new file mode 100644
--- /dev/null
+++ b/Engine/Source/Editor/Sequencer/Private/Tools/NavigationTool/Columns/NavigationToolCommentColumn.cpp
@@ -0,0 +1,45 @@
+// Copyright Epic Games, Inc. All Rights Reserved.
+
+#include "Tools/NavigationTool/Columns/NavigationToolCommentColumn.h"
+#include "Tools/NavigationTool/NavigationToolView.h"
+#include "Tools/NavigationTool/Widgets/Columns/SNavigationToolComment.h"
+
+#define LOCTEXT_NAMESPACE "NavigationToolCommentColumn"
+
+namespace UE::Sequencer::NavigationTool
+{
+
+FText FNavigationToolCommentColumn::GetColumnDisplayNameText() const
+{
+	return LOCTEXT("CommentColumn", "Comment");
+}
+
+const FSlateBrush* FNavigationToolCommentColumn::GetIconBrush() const
+{
+	return FAppStyle::GetBrush(TEXT("FoliageEditMode.BubbleBorder"));
+}
+
+SHeaderRow::FColumn::FArguments FNavigationToolCommentColumn::ConstructHeaderRowColumn(const TSharedRef<INavigationToolView>& InToolView, const float InFillSize)
+{
+	const FName ColumnId = GetColumnId();
+
+	return SHeaderRow::Column(ColumnId)
+		.FillWidth(InFillSize)
+		.HAlignHeader(HAlign_Center)
+		.VAlignHeader(VAlign_Center)
+		.HAlignCell(HAlign_Fill)
+		.VAlignCell(VAlign_Center)
+		.DefaultLabel(GetColumnDisplayNameText())
+		.OnGetMenuContent(InToolView, &INavigationToolView::GetColumnMenuContent, ColumnId);
+}
+
+TSharedRef<SWidget> FNavigationToolCommentColumn::ConstructRowWidget(const FNavigationToolItemPtr& InItem
+	, const TSharedRef<INavigationToolView>& InToolView
+	, const TSharedRef<SNavigationToolTreeRow>& InRow)
+{
+	return SNew(SNavigationToolComment, InItem, InToolView, InRow);
+}
+
+}
+
+#undef LOCTEXT_NAMESPACE
diff --git a/Engine/Source/Editor/Sequencer/Private/Tools/NavigationTool/Columns/NavigationToolDeactiveStateColumn.cpp b/Engine/Source/Editor/Sequencer/Private/Tools/NavigationTool/Columns/NavigationToolDeactiveStateColumn.cpp
new file mode 100644
--- /dev/null
+++ b/Engine/Source/Editor/Sequencer/Private/Tools/NavigationTool/Columns/NavigationToolDeactiveStateColumn.cpp
@@ -0,0 +1,69 @@
+// Copyright Epic Games, Inc. All Rights Reserved.
+
+#include "Tools/NavigationTool/Columns/NavigationToolDeactiveStateColumn.h"
+#include "Tools/NavigationTool/Extensions/ISequenceInactivableExtension.h"
+#include "Tools/NavigationTool/Items/INavigationToolItem.h"
+#include "Tools/NavigationTool/NavigationToolView.h"
+#include "Tools/NavigationTool/Widgets/Columns/SNavigationToolDeactiveState.h"
+#include "Widgets/Layout/SGridPanel.h"
+#include "Widgets/SBoxPanel.h"
+#include "Widgets/SOverlay.h"
+#include "Widgets/Text/STextBlock.h"
+
+#define LOCTEXT_NAMESPACE "NavigationToolMuteStateColumn"
+
+namespace UE::Sequencer::NavigationTool
+{
+
+FText FNavigationToolDeactiveStateColumn::GetColumnDisplayNameText() const
+{
+	return LOCTEXT("NavigationToolMuteStateColumn", "Sequence Evaluation Mute State");
+}
+
+const FSlateBrush* FNavigationToolDeactiveStateColumn::GetIconBrush() const
+{
+	return FAppStyle::GetBrush(TEXT("Sequencer.Column.Mute"));
+}
+
+SHeaderRow::FColumn::FArguments FNavigationToolDeactiveStateColumn::ConstructHeaderRowColumn(const TSharedRef<INavigationToolView>& InToolView
+	, const float InFillSize)
+{
+	const FName ColumnId = GetColumnId();
+
+	return SHeaderRow::Column(ColumnId)
+		.FixedWidth(24.f)
+		.HAlignHeader(HAlign_Left)
+		.VAlignHeader(VAlign_Center)
+		.HAlignCell(HAlign_Center)
+		.VAlignCell(VAlign_Center)
+		.DefaultLabel(GetColumnDisplayNameText())
+		.DefaultTooltip(GetColumnDisplayNameText())
+		.OnGetMenuContent(InToolView, &INavigationToolView::GetColumnMenuContent, ColumnId)
+		[
+			SNew(SImage)
+			.ColorAndOpacity(FSlateColor::UseForeground())
+			.Image(this, &FNavigationToolDeactiveStateColumn::GetIconBrush)
+		];
+}
+
+TSharedRef<SWidget> FNavigationToolDeactiveStateColumn::ConstructRowWidget(const FNavigationToolItemPtr& InItem
+	, const TSharedRef<INavigationToolView>& InToolView
+	, const TSharedRef<SNavigationToolTreeRow>& InRow)
+{
+	if (!InItem->IsA<ISequenceInactivableExtension>())
+	{
+		return SNullWidget::NullWidget;
+	}
+
+	return SNew(SHorizontalBox)
+		+ SHorizontalBox::Slot()
+		.AutoWidth()
+		.VAlign(VAlign_Center)
+		[
+			SNew(SNavigationToolDeactiveState, SharedThis(this), InToolView, InItem, InRow)
+		];
+}
+
+}
+
+#undef LOCTEXT_NAMESPACE
diff --git a/Engine/Source/Editor/Sequencer/Private/Tools/NavigationTool/Columns/NavigationToolHBiasColumn.cpp b/Engine/Source/Editor/Sequencer/Private/Tools/NavigationTool/Columns/NavigationToolHBiasColumn.cpp
new file mode 100644
--- /dev/null
+++ b/Engine/Source/Editor/Sequencer/Private/Tools/NavigationTool/Columns/NavigationToolHBiasColumn.cpp
@@ -0,0 +1,52 @@
+// Copyright Epic Games, Inc. All Rights Reserved.
+
+#include "Tools/NavigationTool/Columns/NavigationToolHBiasColumn.h"
+#include "Tools/NavigationTool/Items/NavigationToolSequence.h"
+#include "Tools/NavigationTool/NavigationToolView.h"
+#include "Tools/NavigationTool/Widgets/Columns/SNavigationToolHBias.h"
+
+#define LOCTEXT_NAMESPACE "NavigationToolHBiasColumn"
+
+namespace UE::Sequencer::NavigationTool
+{
+	
+FText FNavigationToolHBiasColumn::GetColumnDisplayNameText() const
+{
+	return LOCTEXT("HBiasColumn", "HBias");
+}
+
+const FSlateBrush* FNavigationToolHBiasColumn::GetIconBrush() const
+{
+	return FAppStyle::GetBrush(TEXT("ClassIcon.TimelineComponent"));
+}
+
+SHeaderRow::FColumn::FArguments FNavigationToolHBiasColumn::ConstructHeaderRowColumn(const TSharedRef<INavigationToolView>& InToolView
+	, const float InFillSize)
+{
+	const FName ColumnId = GetColumnId();
+
+	return SHeaderRow::Column(ColumnId)
+		.FillWidth(InFillSize)
+		.HAlignHeader(HAlign_Center)
+		.VAlignHeader(VAlign_Center)
+		.HAlignCell(HAlign_Fill)
+		.VAlignCell(VAlign_Center)
+		.DefaultLabel(GetColumnDisplayNameText())
+		.OnGetMenuContent(InToolView, &INavigationToolView::GetColumnMenuContent, ColumnId);
+}
+
+TSharedRef<SWidget> FNavigationToolHBiasColumn::ConstructRowWidget(const FNavigationToolItemPtr& InItem
+	, const TSharedRef<INavigationToolView>& InToolView
+	, const TSharedRef<SNavigationToolTreeRow>& InRow)
+{
+	if (!InItem->IsA<FNavigationToolSequence>())
+	{
+		return SNullWidget::NullWidget;
+	}
+
+	return SNew(SNavigationToolHBias, InItem, InToolView, InRow);
+}
+
+}
+
+#undef LOCTEXT_NAMESPACE
diff --git a/Engine/Source/Editor/Sequencer/Private/Tools/NavigationTool/Columns/NavigationToolInTimeColumn.cpp b/Engine/Source/Editor/Sequencer/Private/Tools/NavigationTool/Columns/NavigationToolInTimeColumn.cpp
new file mode 100644
--- /dev/null
+++ b/Engine/Source/Editor/Sequencer/Private/Tools/NavigationTool/Columns/NavigationToolInTimeColumn.cpp
@@ -0,0 +1,52 @@
+// Copyright Epic Games, Inc. All Rights Reserved.
+
+#include "Tools/NavigationTool/Columns/NavigationToolInTimeColumn.h"
+#include "Tools/NavigationTool/Items/NavigationToolSequence.h"
+#include "Tools/NavigationTool/NavigationToolView.h"
+#include "Tools/NavigationTool/Widgets/Columns/SNavigationToolInTime.h"
+
+#define LOCTEXT_NAMESPACE "NavigationToolInTimeColumn"
+
+namespace UE::Sequencer::NavigationTool
+{
+
+FText FNavigationToolInTimeColumn::GetColumnDisplayNameText() const
+{
+	return LOCTEXT("InTimeColumn", "In");
+}
+
+const FSlateBrush* FNavigationToolInTimeColumn::GetIconBrush() const
+{
+	return FAppStyle::GetBrush(TEXT("Icons.Alignment.Left"));
+}
+
+SHeaderRow::FColumn::FArguments FNavigationToolInTimeColumn::ConstructHeaderRowColumn(const TSharedRef<INavigationToolView>& InToolView
+	, const float InFillSize)
+{
+	const FName ColumnId = GetColumnId();
+
+	return SHeaderRow::Column(ColumnId)
+		.FillWidth(InFillSize)
+		.HAlignHeader(HAlign_Center)
+		.VAlignHeader(VAlign_Center)
+		.HAlignCell(HAlign_Fill)
+		.VAlignCell(VAlign_Center)
+		.DefaultLabel(GetColumnDisplayNameText())
+		.OnGetMenuContent(InToolView, &INavigationToolView::GetColumnMenuContent, ColumnId);
+}
+
+TSharedRef<SWidget> FNavigationToolInTimeColumn::ConstructRowWidget(const FNavigationToolItemPtr& InItem
+	, const TSharedRef<INavigationToolView>& InToolView
+	, const TSharedRef<SNavigationToolTreeRow>& InRow)
+{
+	if (!InItem->IsA<FNavigationToolSequence>())
+	{
+		return SNullWidget::NullWidget;
+	}
+
+	return SNew(SNavigationToolInTime, InItem, InToolView, InRow);
+}
+
+}
+
+#undef LOCTEXT_NAMESPACE
diff --git a/Engine/Source/Editor/Sequencer/Private/Tools/NavigationTool/Columns/NavigationToolItemsColumn.cpp b/Engine/Source/Editor/Sequencer/Private/Tools/NavigationTool/Columns/NavigationToolItemsColumn.cpp
new file mode 100644
--- /dev/null
+++ b/Engine/Source/Editor/Sequencer/Private/Tools/NavigationTool/Columns/NavigationToolItemsColumn.cpp
@@ -0,0 +1,46 @@
+// Copyright Epic Games, Inc. All Rights Reserved.
+
+#include "Tools/NavigationTool/Columns/NavigationToolItemsColumn.h"
+#include "Tools/NavigationTool/NavigationToolView.h"
+#include "Tools/NavigationTool/Widgets/Columns/SNavigationToolItemList.h"
+
+#define LOCTEXT_NAMESPACE "NavigationToolItemsColumn"
+
+namespace UE::Sequencer::NavigationTool
+{
+
+FText FNavigationToolItemsColumn::GetColumnDisplayNameText() const
+{
+	return LOCTEXT("DisplayName", "Items");
+}
+
+const FSlateBrush* FNavigationToolItemsColumn::GetIconBrush() const
+{
+	return FAppStyle::GetBrush(TEXT("ClassIcon.AbilitySystemComponent"));
+}
+
+SHeaderRow::FColumn::FArguments FNavigationToolItemsColumn::ConstructHeaderRowColumn(const TSharedRef<INavigationToolView>& InToolView
+	, const float InFillSize)
+{
+	const FName ColumnId = GetColumnId();
+
+	return SHeaderRow::Column(ColumnId)
+		.FillWidth(InFillSize)
+		.HAlignHeader(HAlign_Left)
+		.VAlignHeader(VAlign_Center)
+		.HAlignCell(HAlign_Fill)
+		.VAlignCell(VAlign_Center)
+		.DefaultLabel(GetColumnDisplayNameText())
+		.OnGetMenuContent(InToolView, &INavigationToolView::GetColumnMenuContent, ColumnId);
+}
+
+TSharedRef<SWidget> FNavigationToolItemsColumn::ConstructRowWidget(const FNavigationToolItemPtr& InItem
+	, const TSharedRef<INavigationToolView>& InToolView
+	, const TSharedRef<SNavigationToolTreeRow>& InRow)
+{
+	return SNew(SNavigationToolItemList, InItem, InToolView, InRow);
+}
+
+}
+
+#undef LOCTEXT_NAMESPACE
diff --git a/Engine/Source/Editor/Sequencer/Private/Tools/NavigationTool/Columns/NavigationToolLabelColumn.cpp b/Engine/Source/Editor/Sequencer/Private/Tools/NavigationTool/Columns/NavigationToolLabelColumn.cpp
new file mode 100644
--- /dev/null
+++ b/Engine/Source/Editor/Sequencer/Private/Tools/NavigationTool/Columns/NavigationToolLabelColumn.cpp
@@ -0,0 +1,43 @@
+// Copyright Epic Games, Inc. All Rights Reserved.
+
+#include "Tools/NavigationTool/Columns/NavigationToolLabelColumn.h"
+#include "Tools/NavigationTool/Items/INavigationToolItem.h"
+#include "Tools/NavigationTool/NavigationToolView.h"
+
+#define LOCTEXT_NAMESPACE "NavigationToolLabelColumn"
+
+namespace UE::Sequencer::NavigationTool
+{
+
+FText FNavigationToolLabelColumn::GetColumnDisplayNameText() const
+{
+	return LOCTEXT("LabelColumn", "Label");
+}
+
+const FSlateBrush* FNavigationToolLabelColumn::GetIconBrush() const
+{
+	return FAppStyle::GetBrush(TEXT("ClassIcon.FontFace"));
+}
+
+SHeaderRow::FColumn::FArguments FNavigationToolLabelColumn::ConstructHeaderRowColumn(const TSharedRef<INavigationToolView>& InToolView
+	, const float InFillSize)
+{
+	const FName ColumnId = GetColumnId();
+
+	return SHeaderRow::Column(ColumnId)
+		.FillWidth(InFillSize)
+		.DefaultLabel(GetColumnDisplayNameText())
+		.ShouldGenerateWidget(true)
+		.OnGetMenuContent(InToolView, &INavigationToolView::GetColumnMenuContent, ColumnId);
+}
+
+TSharedRef<SWidget> FNavigationToolLabelColumn::ConstructRowWidget(const FNavigationToolItemPtr& InItem
+	, const TSharedRef<INavigationToolView>& InToolView
+	, const TSharedRef<SNavigationToolTreeRow>& InRow)
+{
+	return InItem->GenerateLabelWidget(InRow);
+}
+
+}
+
+#undef LOCTEXT_NAMESPACE
diff --git a/Engine/Source/Editor/Sequencer/Private/Tools/NavigationTool/Columns/NavigationToolLengthColumn.cpp b/Engine/Source/Editor/Sequencer/Private/Tools/NavigationTool/Columns/NavigationToolLengthColumn.cpp
new file mode 100644
--- /dev/null
+++ b/Engine/Source/Editor/Sequencer/Private/Tools/NavigationTool/Columns/NavigationToolLengthColumn.cpp
@@ -0,0 +1,46 @@
+// Copyright Epic Games, Inc. All Rights Reserved.
+
+#include "Tools/NavigationTool/Columns/NavigationToolLengthColumn.h"
+#include "Tools/NavigationTool/NavigationToolView.h"
+#include "Tools/NavigationTool/Widgets/Columns/SNavigationToolLength.h"
+
+#define LOCTEXT_NAMESPACE "NavigationToolLengthColumn"
+
+namespace UE::Sequencer::NavigationTool
+{
+
+FText FNavigationToolLengthColumn::GetColumnDisplayNameText() const
+{
+	return LOCTEXT("LengthColumn", "Length");
+}
+
+const FSlateBrush* FNavigationToolLengthColumn::GetIconBrush() const
+{
+	return FAppStyle::GetBrush(TEXT("CurveEd.FitHorizontal"));
+}
+
+SHeaderRow::FColumn::FArguments FNavigationToolLengthColumn::ConstructHeaderRowColumn(const TSharedRef<INavigationToolView>& InToolView
+	, const float InFillSize)
+{
+	const FName ColumnId = GetColumnId();
+
+	return SHeaderRow::Column(ColumnId)
+		.FillWidth(InFillSize)
+		.HAlignHeader(HAlign_Center)
+		.VAlignHeader(VAlign_Center)
+		.HAlignCell(HAlign_Fill)
+		.VAlignCell(VAlign_Center)
+		.DefaultLabel(GetColumnDisplayNameText())
+		.OnGetMenuContent(InToolView, &INavigationToolView::GetColumnMenuContent, ColumnId);
+}
+
+TSharedRef<SWidget> FNavigationToolLengthColumn::ConstructRowWidget(const FNavigationToolItemPtr& InItem
+	, const TSharedRef<INavigationToolView>& InToolView
+	, const TSharedRef<SNavigationToolTreeRow>& InRow)
+{
+	return SNew(SNavigationToolLength, InItem, InToolView, InRow);
+}
+
+}
+
+#undef LOCTEXT_NAMESPACE
diff --git a/Engine/Source/Editor/Sequencer/Private/Tools/NavigationTool/Columns/NavigationToolLockColumn.cpp b/Engine/Source/Editor/Sequencer/Private/Tools/NavigationTool/Columns/NavigationToolLockColumn.cpp
new file mode 100644
--- /dev/null
+++ b/Engine/Source/Editor/Sequencer/Private/Tools/NavigationTool/Columns/NavigationToolLockColumn.cpp
@@ -0,0 +1,66 @@
+// Copyright Epic Games, Inc. All Rights Reserved.
+
+#include "Tools/NavigationTool/Columns/NavigationToolLockColumn.h"
+#include "Styling/StyleColors.h"
+#include "Tools/NavigationTool/Extensions/ISequenceLockableExtension.h"
+#include "Tools/NavigationTool/Items/INavigationToolItem.h"
+#include "Tools/NavigationTool/NavigationToolView.h"
+#include "Tools/NavigationTool/Widgets/Columns/SNavigationToolLock.h"
+#include "Widgets/SBoxPanel.h"
+
+#define LOCTEXT_NAMESPACE "NavigationToolLockColumn"
+
+namespace UE::Sequencer::NavigationTool
+{
+
+FText FNavigationToolLockColumn::GetColumnDisplayNameText() const
+{
+	return LOCTEXT("LockColumn", "Lock");
+}
+
+const FSlateBrush* FNavigationToolLockColumn::GetIconBrush() const
+{
+	return FAppStyle::GetBrush(TEXT("Icons.Lock"));
+}
+
+SHeaderRow::FColumn::FArguments FNavigationToolLockColumn::ConstructHeaderRowColumn(const TSharedRef<INavigationToolView>& InToolView
+	, const float InFillSize)
+{
+	const FName ColumnId = GetColumnId();
+
+	return SHeaderRow::Column(ColumnId)
+		.FixedWidth(24.f)
+		.HAlignHeader(HAlign_Center)
+		.VAlignHeader(VAlign_Center)
+		.HAlignCell(HAlign_Center)
+		.VAlignCell(VAlign_Center)
+		.DefaultLabel(GetColumnDisplayNameText())
+		.DefaultTooltip(GetColumnDisplayNameText())
+		[
+			SNew(SImage)
+			.ColorAndOpacity(FStyleColors::Foreground)
+			.Image(FAppStyle::GetBrush(TEXT("Icons.Lock")))
+		];
+}
+
+TSharedRef<SWidget> FNavigationToolLockColumn::ConstructRowWidget(const FNavigationToolItemPtr& InItem
+	, const TSharedRef<INavigationToolView>& InToolView
+	, const TSharedRef<SNavigationToolTreeRow>& InRow)
+{
+	if (!InItem->IsA<ISequenceLockableExtension>())
+	{
+		return SNullWidget::NullWidget;
+	}
+
+	return SNew(SHorizontalBox)
+		+ SHorizontalBox::Slot()
+		.AutoWidth()
+		.VAlign(VAlign_Center)
+		[
+			SNew(SNavigationToolLock, InItem, InToolView, InRow)
+		];
+}
+
+}
+
+#undef LOCTEXT_NAMESPACE
diff --git a/Engine/Source/Editor/Sequencer/Private/Tools/NavigationTool/Columns/NavigationToolMarkerVisibilityColumn.cpp b/Engine/Source/Editor/Sequencer/Private/Tools/NavigationTool/Columns/NavigationToolMarkerVisibilityColumn.cpp
new file mode 100644
--- /dev/null
+++ b/Engine/Source/Editor/Sequencer/Private/Tools/NavigationTool/Columns/NavigationToolMarkerVisibilityColumn.cpp
@@ -0,0 +1,61 @@
+// Copyright Epic Games, Inc. All Rights Reserved.
+
+#include "Tools/NavigationTool/Columns/NavigationToolMarkerVisibilityColumn.h"
+#include "Tools/NavigationTool/Items/INavigationToolItem.h"
+#include "Tools/NavigationTool/Items/NavigationToolSequence.h"
+#include "Tools/NavigationTool/NavigationToolView.h"
+#include "Tools/NavigationTool/Widgets/Columns/SNavigationToolMarkerVisibility.h"
+#include "Widgets/SBoxPanel.h"
+#include "Widgets/SOverlay.h"
+#include "Widgets/Text/STextBlock.h"
+
+#define LOCTEXT_NAMESPACE "NavigationToolMarkerVisibilityColumn"
+
+namespace UE::Sequencer::NavigationTool
+{
+
+FText FNavigationToolMarkerVisibilityColumn::GetColumnDisplayNameText() const
+{
+	return LOCTEXT("NavigationToolMarkerVisibilityColumn", "Marker Visibility");
+}
+
+const FSlateBrush* FNavigationToolMarkerVisibilityColumn::GetIconBrush() const
+{
+	return FAppStyle::GetBrush(TEXT("AnimTimeline.SectionMarker"));
+}
+
+SHeaderRow::FColumn::FArguments FNavigationToolMarkerVisibilityColumn::ConstructHeaderRowColumn(const TSharedRef<INavigationToolView>& InToolView
+	, const float InFillSize)
+{
+	const FName ColumnId = GetColumnId();
+
+	return SHeaderRow::Column(ColumnId)
+		.FixedWidth(24.f)
+		.HAlignHeader(HAlign_Center)
+		.VAlignHeader(VAlign_Center)
+		.HAlignCell(HAlign_Center)
+		.VAlignCell(VAlign_Center)
+		.DefaultLabel(GetColumnDisplayNameText())
+		.DefaultTooltip(GetColumnDisplayNameText())
+		[
+			SNew(SImage)
+			.ColorAndOpacity(FSlateColor::UseForeground())
+			.Image(FAppStyle::GetBrush(TEXT("AnimTimeline.SectionMarker")))
+		];
+}
+
+TSharedRef<SWidget> FNavigationToolMarkerVisibilityColumn::ConstructRowWidget(const FNavigationToolItemPtr& InItem
+	, const TSharedRef<INavigationToolView>& InToolView
+	, const TSharedRef<SNavigationToolTreeRow>& InRow)
+{
+	if (!InItem->IsA<IMarkerVisibilityExtension>())
+	{
+		return SNullWidget::NullWidget;
+	}
+
+	return SNew(SNavigationToolMarkerVisibility, SharedThis(this), InItem, InToolView, InRow);
+}
+
+}
+
+#undef LOCTEXT_NAMESPACE
diff --git a/Engine/Source/Editor/Sequencer/Private/Tools/NavigationTool/Columns/NavigationToolOutTimeColumn.cpp b/Engine/Source/Editor/Sequencer/Private/Tools/NavigationTool/Columns/NavigationToolOutTimeColumn.cpp
new file mode 100644
--- /dev/null
+++ b/Engine/Source/Editor/Sequencer/Private/Tools/NavigationTool/Columns/NavigationToolOutTimeColumn.cpp
@@ -0,0 +1,54 @@
+// Copyright Epic Games, Inc. All Rights Reserved.
+
+#include "Tools/NavigationTool/Columns/NavigationToolOutTimeColumn.h"
+#include "Tools/NavigationTool/Items/NavigationToolSequence.h"
+#include "Tools/NavigationTool/NavigationTool.h"
+#include "Tools/NavigationTool/NavigationToolView.h"
+#include "Tools/NavigationTool/Widgets/Columns/SNavigationToolOutTime.h"
+#include "Tools/NavigationTool/Widgets/SNavigationToolTreeRow.h"
+
+#define LOCTEXT_NAMESPACE "NavigationToolOutTimeColumn"
+
+namespace UE::Sequencer::NavigationTool
+{
+
+FText FNavigationToolOutTimeColumn::GetColumnDisplayNameText() const
+{
+	return LOCTEXT("OutTimeColumn", "Out");
+}
+
+const FSlateBrush* FNavigationToolOutTimeColumn::GetIconBrush() const
+{
+	return FAppStyle::GetBrush(TEXT("Icons.Alignment.Right"));
+}
+
+SHeaderRow::FColumn::FArguments FNavigationToolOutTimeColumn::ConstructHeaderRowColumn(const TSharedRef<INavigationToolView>& InToolView
+	, const float InFillSize)
+{
+	const FName ColumnId = GetColumnId();
+
+	return SHeaderRow::Column(ColumnId)
+		.FillWidth(InFillSize)
+		.HAlignHeader(HAlign_Center)
+		.VAlignHeader(VAlign_Center)
+		.HAlignCell(HAlign_Fill)
+		.VAlignCell(VAlign_Center)
+		.DefaultLabel(GetColumnDisplayNameText())
+		.OnGetMenuContent(InToolView, &INavigationToolView::GetColumnMenuContent, ColumnId);
+}
+
+TSharedRef<SWidget> FNavigationToolOutTimeColumn::ConstructRowWidget(const FNavigationToolItemPtr& InItem
+	, const TSharedRef<INavigationToolView>& InToolView
+	, const TSharedRef<SNavigationToolTreeRow>& InRow)
+{
+	if (!InItem->IsA<FNavigationToolSequence>())
+	{
+		return SNullWidget::NullWidget;
+	}
+
+	return SNew(SNavigationToolOutTime, InItem, InToolView, InRow);
+}
+
+}
+
+#undef LOCTEXT_NAMESPACE
diff --git a/Engine/Source/Editor/Sequencer/Private/Tools/NavigationTool/Columns/NavigationToolRevisionControlColumn.cpp b/Engine/Source/Editor/Sequencer/Private/Tools/NavigationTool/Columns/NavigationToolRevisionControlColumn.cpp
new file mode 100644
--- /dev/null
+++ b/Engine/Source/Editor/Sequencer/Private/Tools/NavigationTool/Columns/NavigationToolRevisionControlColumn.cpp
@@ -0,0 +1,63 @@
+// Copyright Epic Games, Inc. All Rights Reserved.
+
+#include "Tools/NavigationTool/Columns/NavigationToolRevisionControlColumn.h"
+#include "RevisionControlStyle/RevisionControlStyle.h"
+#include "Tools/NavigationTool/Items/INavigationToolItem.h"
+#include "Tools/NavigationTool/Items/NavigationToolSequence.h"
+#include "Tools/NavigationTool/NavigationToolView.h"
+#include "Tools/NavigationTool/Widgets/Columns/SNavigationToolRevisionControl.h"
+#include "Widgets/SBoxPanel.h"
+#include "Widgets/SOverlay.h"
+#include "Widgets/Text/STextBlock.h"
+
+#define LOCTEXT_NAMESPACE "NavigationToolRevisionControlColumn"
+
+namespace UE::Sequencer::NavigationTool
+{
+
+FText FNavigationToolRevisionControlColumn::GetColumnDisplayNameText() const
+{
+	return LOCTEXT("RevisionControl", "Revision Control");
+}
+
+const FSlateBrush* FNavigationToolRevisionControlColumn::GetIconBrush() const
+{
+	return FRevisionControlStyleManager::Get().GetBrush(TEXT("RevisionControl.Icon"));
+}
+
+SHeaderRow::FColumn::FArguments FNavigationToolRevisionControlColumn::ConstructHeaderRowColumn(const TSharedRef<INavigationToolView>& InToolView
+	, const float InFillSize)
+{
+	const FName ColumnId = GetColumnId();
+
+	return SHeaderRow::Column(ColumnId)
+		.FixedWidth(24.f)
+		.HAlignHeader(HAlign_Center)
+		.VAlignHeader(VAlign_Center)
+		.HAlignCell(HAlign_Center)
+		.VAlignCell(VAlign_Center)
+		.DefaultLabel(GetColumnDisplayNameText())
+		.DefaultTooltip(GetColumnDisplayNameText())
+		.OnGetMenuContent(InToolView, &INavigationToolView::GetColumnMenuContent, ColumnId)
+		[
+			SNew(SImage)
+			.ColorAndOpacity(FSlateColor::UseForeground())
+			.Image(this, &FNavigationToolRevisionControlColumn::GetIconBrush)
+		];
+}
+
+TSharedRef<SWidget> FNavigationToolRevisionControlColumn::ConstructRowWidget(const FNavigationToolItemPtr& InItem
+	, const TSharedRef<INavigationToolView>& InToolView
+	, const TSharedRef<SNavigationToolTreeRow>& InRow)
+{
+	if (!InItem->IsA<FNavigationToolSequence>())
+	{
+		return SNullWidget::NullWidget;
+	}
+
+	return SNew(SNavigationToolRevisionControl, SharedThis(this), InItem, InToolView, InRow);
+}
+
+}
+
+#undef LOCTEXT_NAMESPACE
diff --git a/Engine/Source/Editor/Sequencer/Private/Tools/NavigationTool/Columns/NavigationToolStartFrameOffsetColumn.cpp b/Engine/Source/Editor/Sequencer/Private/Tools/NavigationTool/Columns/NavigationToolStartFrameOffsetColumn.cpp
new file mode 100644
--- /dev/null
+++ b/Engine/Source/Editor/Sequencer/Private/Tools/NavigationTool/Columns/NavigationToolStartFrameOffsetColumn.cpp
@@ -0,0 +1,46 @@
+// Copyright Epic Games, Inc. All Rights Reserved.
+
+#include "Tools/NavigationTool/Columns/NavigationToolStartFrameOffsetColumn.h"
+#include "Tools/NavigationTool/NavigationToolView.h"
+#include "Tools/NavigationTool/Widgets/Columns/SNavigationToolStartFrameOffset.h"
+
+#define LOCTEXT_NAMESPACE "NavigationToolStartFrameOffsetColumn"
+
+namespace UE::Sequencer::NavigationTool
+{
+
+FText FNavigationToolStartFrameOffsetColumn::GetColumnDisplayNameText() const
+{
+	return LOCTEXT("StartFrameOffset", "Start Frame Offset");
+}
+
+const FSlateBrush* FNavigationToolStartFrameOffsetColumn::GetIconBrush() const
+{
+	return FAppStyle::GetBrush(TEXT("ClassIcon.TimelineComponent"));
+}
+
+SHeaderRow::FColumn::FArguments FNavigationToolStartFrameOffsetColumn::ConstructHeaderRowColumn(const TSharedRef<INavigationToolView>& InToolView
+	, const float InFillSize)
+{
+	const FName ColumnId = GetColumnId();
+
+	return SHeaderRow::Column(ColumnId)
+		.FillWidth(InFillSize)
+		.HAlignHeader(HAlign_Center)
+		.VAlignHeader(VAlign_Center)
+		.HAlignCell(HAlign_Fill)
+		.VAlignCell(VAlign_Center)
+		.DefaultLabel(GetColumnDisplayNameText())
+		.OnGetMenuContent(InToolView, &INavigationToolView::GetColumnMenuContent, ColumnId);
+}
+
+TSharedRef<SWidget> FNavigationToolStartFrameOffsetColumn::ConstructRowWidget(const FNavigationToolItemPtr& InItem
+	, const TSharedRef<INavigationToolView>& InToolView
+	, const TSharedRef<SNavigationToolTreeRow>& InRow)
+{
+	return SNew(SNavigationToolStartFrameOffset, InItem, InToolView, InRow);
+}
+
+}
+
+#undef LOCTEXT_NAMESPACE
diff --git a/Engine/Source/Editor/Sequencer/Private/Tools/NavigationTool/Columns/NavigationToolTakeColumn.cpp b/Engine/Source/Editor/Sequencer/Private/Tools/NavigationTool/Columns/NavigationToolTakeColumn.cpp
new file mode 100644
--- /dev/null
+++ b/Engine/Source/Editor/Sequencer/Private/Tools/NavigationTool/Columns/NavigationToolTakeColumn.cpp
@@ -0,0 +1,53 @@
+// Copyright Epic Games, Inc. All Rights Reserved.
+
+#include "Tools/NavigationTool/Columns/NavigationToolTakeColumn.h"
+#include "Tools/NavigationTool/Items/NavigationToolSequence.h"
+#include "Tools/NavigationTool/NavigationTool.h"
+#include "Tools/NavigationTool/NavigationToolView.h"
+#include "Tools/NavigationTool/Widgets/Columns/SNavigationToolTake.h"
+
+#define LOCTEXT_NAMESPACE "NavigationToolTakesColumn"
+
+namespace UE::Sequencer::NavigationTool
+{
+
+FText FNavigationToolTakeColumn::GetColumnDisplayNameText() const
+{
+	return LOCTEXT("TakesColumn", "Takes");
+}
+
+const FSlateBrush* FNavigationToolTakeColumn::GetIconBrush() const
+{
+	return FAppStyle::GetBrush(TEXT("ClassIcon.LevelSequence"));
+}
+
+SHeaderRow::FColumn::FArguments FNavigationToolTakeColumn::ConstructHeaderRowColumn(const TSharedRef<INavigationToolView>& InToolView
+	, const float InFillSize)
+{
+	const FName ColumnId = GetColumnId();
+
+	return SHeaderRow::Column(ColumnId)
+		.FillWidth(InFillSize)
+		.HAlignHeader(HAlign_Center)
+		.VAlignHeader(VAlign_Center)
+		.HAlignCell(HAlign_Fill)
+		.VAlignCell(VAlign_Center)
+		.DefaultLabel(GetColumnDisplayNameText())
+		.OnGetMenuContent(InToolView, &INavigationToolView::GetColumnMenuContent, ColumnId);
+}
+
+TSharedRef<SWidget> FNavigationToolTakeColumn::ConstructRowWidget(const FNavigationToolItemPtr& InItem
+	, const TSharedRef<INavigationToolView>& InToolView
+	, const TSharedRef<SNavigationToolTreeRow>& InRow)
+{
+	if (!InItem->IsA<FNavigationToolSequence>())
+	{
+		return SNullWidget::NullWidget;
+	}
+
+	return SNew(SNavigationToolTake, InItem, InToolView, InRow);
+}
+
+}
+
+#undef LOCTEXT_NAMESPACE
diff --git a/Engine/Source/Editor/Sequencer/Private/Tools/NavigationTool/Customization/NavigationToolExtender.cpp b/Engine/Source/Editor/Sequencer/Private/Tools/NavigationTool/Customization/NavigationToolExtender.cpp
new file mode 100644
--- /dev/null
+++ b/Engine/Source/Editor/Sequencer/Private/Tools/NavigationTool/Customization/NavigationToolExtender.cpp
@@ -0,0 +1,324 @@
+// Copyright Epic Games, Inc. All Rights Reserved.
+
+#include "Tools/NavigationTool/Customization/NavigationToolExtender.h"
+#include "ISequencer.h"
+#include "ISequencerModule.h"
+#include "Modules/ModuleManager.h"
+#include "SequencerLog.h"
+#include "SequencerSettings.h"
+#include "Tools/NavigationTool/Customization/NavigationToolIconCustomization.h"
+#include "Tools/NavigationTool/INavigationTool.h"
+#include "Tools/NavigationTool/Items/NavigationToolComponent.h"
+#include "Tools/NavigationTool/NavigationToolCommands.h"
+#include "Tools/NavigationTool/Providers/NavigationToolProvider.h"
+
+namespace UE::Sequencer::NavigationTool
+{
+
+FNavigationToolExtender& FNavigationToolExtender::Get()
+{
+	ISequencerModule& SequencerModule = FModuleManager::LoadModuleChecked<ISequencerModule>(TEXT("Sequencer"));
+	return SequencerModule.GetNavigationToolExtender();
+}
+
+FNavigationToolExtender::FNavigationToolExtender()
+{
+	FNavigationToolCommands::Register();
+
+	// @TODO: register item proxies?
+	//ItemProxyRegistry.RegisterItemProxyWithDefaultFactory<FNavigationToolComponentProxy, 10>();
+}
+
+FNavigationToolExtender::~FNavigationToolExtender()
+{
+	FNavigationToolCommands::Unregister();
+}
+
+FName FNavigationToolExtender::GetToolInstanceName(const INavigationTool& InTool)
+{
+	if (const TSharedPtr<ISequencer> Sequencer = InTool.GetSequencer())
+	{
+		const USequencerSettings* const SequencerSettings = Sequencer->GetSequencerSettings();
+		if (IsValid(SequencerSettings))
+		{
+			return *SequencerSettings->GetName();
+		}
+	}
+	return NAME_None;
+}
+
+bool FNavigationToolExtender::RegisterToolProvider(const TSharedRef<INavigationTool>& InTool
+	, const TSharedRef<FNavigationToolProvider>& InProvider)
+{
+	FNavigationToolExtender& ThisInstance = Get();
+	INavigationTool& RawToolRef = *InTool;
+
+	// Check if provider already exists for this tool
+	const FName Identifier = InProvider->GetIdentifier();
+	if (const TSharedPtr<FNavigationToolProvider> Provider = FindToolProvider(RawToolRef, Identifier))
+	{
+		UnregisterToolProvider(InTool, Identifier);
+		//return false;
+	}
+
+	// Add the provider if an entry for this tool already exists
+	if (FNavigationToolProviderEntry* const Entry = FindToolProvidersEntry(RawToolRef))
+	{
+		Entry->WeakProviders.Add(InProvider);
+		UE_LOG(LogSequencer, Display, TEXT("Navigation Tool Provider Added: %s"), *InProvider->GetIdentifier().ToString());
+
+		ThisInstance.ProvidersChangedDelegate.Broadcast(InTool, InProvider, ENavigationToolProvidersChangeType::Add);
+
+		return true;
+	}
+
+	// Find and verify the tool instance name
+	const FName ToolInstanceName = GetToolInstanceName(RawToolRef);
+
+	if (ToolInstanceName.IsNone())
+	{
+		UE_LOG(LogSequencer, Error, TEXT("Invalid Navigation Tool Instance Name"));
+		return false;
+	}
+
+	// Add new tool entry
+	FNavigationToolProviderEntry NewEntry(InTool, { InProvider });
+
+	ThisInstance.ToolProviders.Add(MoveTemp(NewEntry));
+	UE_LOG(LogSequencer, Display, TEXT("Navigation Tool Provider Entry Added: %s"), *ToolInstanceName.ToString());
+
+	ThisInstance.ProvidersChangedDelegate.Broadcast(InTool, InProvider, ENavigationToolProvidersChangeType::Add);
+
+	return true;
+}
+
+bool FNavigationToolExtender::RegisterOnlyToolProvider(const TSharedRef<INavigationTool>& InTool,
+	const TSharedRef<FNavigationToolProvider>& InProvider)
+{
+	UnregisterAllToolProviders(InTool);
+
+	return RegisterToolProvider(InTool, InProvider);
+}
+
+bool FNavigationToolExtender::UnregisterToolProvider(const TSharedRef<INavigationTool>& InTool, const FName InProviderId)
+{
+	FNavigationToolExtender& ThisInstance = Get();
+
+	for (TSet<FNavigationToolProviderEntry>::TIterator EntryIter(ThisInstance.ToolProviders); EntryIter; ++EntryIter)
+	{
+		FNavigationToolProviderEntry& Entry = *EntryIter;
+
+		if (const TSharedPtr<INavigationTool> Tool = Entry.WeakTool.Pin())
+		{
+			if (Tool.ToSharedRef() == InTool)
+			{
+				for (TSet<TWeakPtr<FNavigationToolProvider>>::TIterator ProviderIter(Entry.WeakProviders); ProviderIter; ++ProviderIter)
+				{
+					if (const TSharedPtr<FNavigationToolProvider> Provider = ProviderIter->Pin())
+					{
+						if (Provider->GetIdentifier() == InProviderId)
+						{
+							ProviderIter.RemoveCurrent();
+
+							ThisInstance.ProvidersChangedDelegate.Broadcast(InTool, Provider.ToSharedRef(), ENavigationToolProvidersChangeType::Remove);
+
+							return true;
+						}
+					}
+					else
+					{
+						ProviderIter.RemoveCurrent();
+					}
+				}
+			}
+		}
+	}
+
+	return false;
+}
+
+void FNavigationToolExtender::UnregisterAllToolProviders(const TSharedRef<INavigationTool>& InTool)
+{
+	FNavigationToolExtender& ThisInstance = Get();
+
+	for (TSet<FNavigationToolProviderEntry>::TIterator EntryIter(ThisInstance.ToolProviders); EntryIter; ++EntryIter)
+	{
+		FNavigationToolProviderEntry& Entry = *EntryIter;
+
+		if (const TSharedPtr<INavigationTool> Tool = Entry.WeakTool.Pin())
+		{
+			if (Tool.ToSharedRef() == InTool)
+			{
+				for (TSet<TWeakPtr<FNavigationToolProvider>>::TIterator ProviderIter(Entry.WeakProviders); ProviderIter; ++ProviderIter)
+				{
+					if (const TSharedPtr<FNavigationToolProvider> Provider = ProviderIter->Pin())
+					{
+						ProviderIter.RemoveCurrent();
+
+						ThisInstance.ProvidersChangedDelegate.Broadcast(InTool, Provider.ToSharedRef(), ENavigationToolProvidersChangeType::Remove);
+					}
+					else
+					{
+						ProviderIter.RemoveCurrent();
+					}
+				}
+			}
+		}
+	}
+}
+
+FNavigationToolExtender::FNavigationToolProviderEntry* FNavigationToolExtender::FindToolProvidersEntry(const INavigationTool& InTool)
+{
+	FNavigationToolExtender& ThisInstance = Get();
+
+	for (TSet<FNavigationToolProviderEntry>::TIterator EntryIter(ThisInstance.ToolProviders); EntryIter; ++EntryIter)
+	{
+		FNavigationToolProviderEntry& Entry = *EntryIter;
+
+		if (const TSharedPtr<INavigationTool> Tool = Entry.WeakTool.Pin())
+		{
+			if (Tool.Get() == &InTool)
+			{
+				return &Entry;
+			}
+		}
+	}
+
+	return nullptr;
+}
+
+TSharedPtr<FNavigationToolProvider> FNavigationToolExtender::FindToolProvider(const INavigationTool& InTool, const FName InProviderId)
+{
+	if (FNavigationToolProviderEntry* const Entry = FindToolProvidersEntry(InTool))
+	{
+		for (TSet<TWeakPtr<FNavigationToolProvider>>::TIterator ProviderIter(Entry->WeakProviders); ProviderIter; ++ProviderIter)
+		{
+			if (const TSharedPtr<FNavigationToolProvider> Provider = ProviderIter->Pin())
+			{
+				if (Provider->GetIdentifier() == InProviderId)
+				{
+					return Provider;
+				}
+			}
+		}
+	}
+
+	return nullptr;
+}
+
+bool FNavigationToolExtender::FindToolProviders(const INavigationTool& InTool
+	, TSet<TWeakPtr<FNavigationToolProvider>>& OutProviders)
+{
+	if (const FNavigationToolProviderEntry* const Entry = FindToolProvidersEntry(InTool))
+	{
+		OutProviders.Append(Entry->WeakProviders);
+		return true;
+	}
+	return false;
+}
+
+void FNavigationToolExtender::ForEachProvider(const TFunction<bool(INavigationTool& InTool
+	, const TSharedRef<FNavigationToolProvider>& InProvider)>& InPredicate)
+{
+	for (TSet<FNavigationToolProviderEntry>::TIterator EntryIter(Get().ToolProviders); EntryIter; ++EntryIter)
+	{
+		FNavigationToolProviderEntry& Entry = *EntryIter;
+
+		const TSharedPtr<INavigationTool> Tool = Entry.WeakTool.Pin();
+		if (!Tool.IsValid())
+		{
+			EntryIter.RemoveCurrent();
+			continue;
+		}
+
+		const TSharedRef<INavigationTool> ToolRef = Tool.ToSharedRef();
+
+		for (TSet<TWeakPtr<FNavigationToolProvider>>::TIterator ProviderIter(Entry.WeakProviders); ProviderIter; ++ProviderIter)
+		{
+			if (const TSharedPtr<FNavigationToolProvider> Provider = ProviderIter->Pin())
+			{
+				// Stop processing if the function returns false
+				if (!InPredicate(*ToolRef, Provider.ToSharedRef()))
+				{
+					return;
+				}
+			}
+		}
+	}
+}
+
+void FNavigationToolExtender::ForEachToolProvider(INavigationTool& InTool,
+	const TFunction<bool(const TSharedRef<FNavigationToolProvider>& InProvider)>& InPredicate)
+{
+	FNavigationToolProviderEntry* const ToolEntry = Get().FindToolProvidersEntry(InTool);
+	if (!ToolEntry)
+	{
+		return;
+	}
+
+	for (TSet<TWeakPtr<FNavigationToolProvider>>::TIterator ProviderIter(ToolEntry->WeakProviders); ProviderIter; ++ProviderIter)
+	{
+		if (const TSharedPtr<FNavigationToolProvider> Provider = ProviderIter->Pin())
+		{
+			// Stop processing if the function returns false
+			if (!InPredicate(Provider.ToSharedRef()))
+			{
+				return;
+			}
+		}
+	}
+}
+
+FSlateIcon FNavigationToolExtender::FindOverrideIcon(TSharedPtr<const FNavigationToolItem> InItem)
+{
+	FNavigationToolExtender& ThisInstance = Get();
+
+	if (InItem.IsValid())
+	{
+		const TSharedPtr<INavigationToolIconCustomization> CustomizationToUse = ThisInstance.GetCustomizationForItem(InItem);
+
+		if (CustomizationToUse.IsValid())
+		{
+			return CustomizationToUse->GetOverrideIcon(InItem);
+		}
+	}
+
+	return FSlateIcon();
+}
+
+void FNavigationToolExtender::UnregisterOverriddenIcon_Internal(const FNavigationToolItemTypeId& InItemTypeId, const FName& InSpecializationIdentifier)
+{
+	FIconCustomizationKey Key;
+	Key.ItemTypeId = InItemTypeId;
+	Key.CustomizationSpecializationIdentifier = InSpecializationIdentifier;
+
+	IconRegistry.Remove(Key);
+}
+
+void FNavigationToolExtender::RegisterOverridenIcon_Internal(const FNavigationToolItemTypeId& InItemTypeId, const TSharedRef<INavigationToolIconCustomization>& InNavigationToolIconCustomization)
+{
+	FIconCustomizationKey Key;
+	Key.ItemTypeId = InItemTypeId;
+	Key.CustomizationSpecializationIdentifier = InNavigationToolIconCustomization->GetToolItemIdentifier();
+
+	if (!IconRegistry.Contains(Key))
+	{
+		IconRegistry.Add(Key, InNavigationToolIconCustomization);
+	}
+}
+
+TSharedPtr<INavigationToolIconCustomization> FNavigationToolExtender::GetCustomizationForItem(const TSharedPtr<const FNavigationToolItem>& InItem) const
+{
+	// There is gonna always be just one icon customization that support the item since it checks specific thing for the SupportCustomization
+	for (const TPair<FIconCustomizationKey, TSharedPtr<INavigationToolIconCustomization>>& IconCustomization : IconRegistry)
+	{
+		if (IconCustomization.Key.ItemTypeId == InItem->GetTypeId() && IconCustomization.Value->HasOverrideIcon(InItem))
+		{
+			return IconCustomization.Value;
+		}
+	}
+
+	return nullptr;
+}
+
+}
diff --git a/Engine/Source/Editor/Sequencer/Private/Tools/NavigationTool/Customization/NavigationToolObjectIconCustomization.cpp b/Engine/Source/Editor/Sequencer/Private/Tools/NavigationTool/Customization/NavigationToolObjectIconCustomization.cpp
new file mode 100644
--- /dev/null
+++ b/Engine/Source/Editor/Sequencer/Private/Tools/NavigationTool/Customization/NavigationToolObjectIconCustomization.cpp
@@ -0,0 +1,57 @@
+// Copyright Epic Games, Inc. All Rights Reserved.
+
+#include "Tools/NavigationTool/Customization/NavigationToolObjectIconCustomization.h"
+#include "Textures/SlateIcon.h"
+#include "Tools/NavigationTool/Items/NavigationToolObject.h"
+#include "UObject/Class.h"
+
+namespace UE::Sequencer::NavigationTool
+{
+
+FNavigationToolObjectIconCustomization::FNavigationToolObjectIconCustomization(const UClass* InSupportedClass)
+{
+	SupportedClassName = FName(TEXT(""));
+
+	if (IsValid(InSupportedClass))
+	{
+		SupportedClassName = InSupportedClass->GetFName();
+	}
+}
+
+void FNavigationToolObjectIconCustomization::SetOverriddenIcon(const FOnGetOverriddenObjectIcon& InOverriddenIcon)
+{
+	OnGetOverriddenIcon = InOverriddenIcon;
+}
+
+FName FNavigationToolObjectIconCustomization::GetToolItemIdentifier() const
+{
+	return SupportedClassName;
+}
+
+bool FNavigationToolObjectIconCustomization::HasOverrideIcon(TSharedPtr<const FNavigationToolItem> InToolItem) const
+{
+	if (const FNavigationToolObject* ObjectItem = InToolItem->CastTo<FNavigationToolObject>())
+	{
+		if (const UObject* Object = ObjectItem->GetObject())
+		{
+			if (const UObject* ObjectClass = Object->GetClass())
+			{
+				return ObjectClass->GetFName() == SupportedClassName;
+			}
+		}
+	}
+
+	return false;
+}
+
+FSlateIcon FNavigationToolObjectIconCustomization::GetOverrideIcon(TSharedPtr<const FNavigationToolItem> InToolItem) const
+{
+	if (HasOverrideIcon(InToolItem) && OnGetOverriddenIcon.IsBound())
+	{
+		return OnGetOverriddenIcon.Execute(InToolItem);
+	}
+
+	return FSlateIcon();
+}
+
+}
diff --git a/Engine/Source/Editor/Sequencer/Private/Tools/NavigationTool/DragDropOps/Handlers/NavigationToolActorDropHandler.cpp b/Engine/Source/Editor/Sequencer/Private/Tools/NavigationTool/DragDropOps/Handlers/NavigationToolActorDropHandler.cpp
new file mode 100644
--- /dev/null
+++ b/Engine/Source/Editor/Sequencer/Private/Tools/NavigationTool/DragDropOps/Handlers/NavigationToolActorDropHandler.cpp
@@ -0,0 +1,91 @@
+// Copyright Epic Games, Inc. All Rights Reserved.
+
+#include "NavigationToolActorDropHandler.h"
+#include "Framework/Application/SlateApplication.h"
+#include "Tools/NavigationTool/ItemActions/NavigationToolAddItem.h"
+#include "Tools/NavigationTool/Items/NavigationToolActor.h"
+#include "Tools/NavigationTool/NavigationTool.h"
+#include "Widgets/Views/STableRow.h"
+
+namespace UE::Sequencer::NavigationTool
+{
+
+bool FNavigationToolActorDropHandler::IsDraggedItemSupported(const TSharedPtr<INavigationToolItem>& InDraggedItem) const
+{
+	return InDraggedItem->IsA<FNavigationToolActor>();
+}
+
+TOptional<EItemDropZone> FNavigationToolActorDropHandler::CanDrop(EItemDropZone InDropZone, const FNavigationToolItemPtr& InTargetItem) const
+{
+	/*switch (ActionType)
+	{
+		case ENavigationToolDragDropActionType::Move:
+			// When moving, make sure the Destination is not one of the Items we're moving
+			if (!Items.Contains(InTargetItem))
+			{
+				return InDropZone;
+			}
+			break;
+
+		case ENavigationToolDragDropActionType::Copy:
+			return InDropZone;
+	}*/
+
+	return TOptional<EItemDropZone>();
+}
+
+bool FNavigationToolActorDropHandler::Drop(EItemDropZone InDropZone, const FNavigationToolItemPtr& InTargetItem)
+{
+	/*switch (ActionType)
+	{
+	case ENavigationToolDragDropActionType::Move:
+		MoveItems(InDropZone, InTargetItem);
+		break;
+
+	case ENavigationToolDragDropActionType::Copy:
+		StaticCastSharedRef<FNavigationTool>(InTargetItem->GetOwnerTool())->DuplicateItems(Items, InTargetItem, InDropZone);
+		break;
+
+	default:
+		return false;
+	}*/
+
+	return false;
+}
+
+void FNavigationToolActorDropHandler::MoveItems(EItemDropZone InDropZone, const FNavigationToolItemPtr& InTargetItem)
+{
+	// Can't access this mouse event at this point, so use the Slate Application to check modifiers.
+	/*const FAttachmentTransformRules& TransformRules = InTargetItem->GetOwnerTool()->GetProvider().GetTransformRule(!FSlateApplication::Get().GetModifierKeys().IsShiftDown());
+	
+	FNavigationToolAddItemParams Params;
+	Params.RelativeItem             = InTargetItem;
+	Params.RelativeDropZone         = InDropZone;
+	Params.Flags                    = ENavigationToolAddItemFlags::Select | ENavigationToolAddItemFlags::Transact;
+	Params.AttachmentTransformRules = TransformRules;
+	Params.SelectionFlags           = ENavigationToolItemSelectionFlags::AppendToCurrentSelection;
+
+	const TSet<FNavigationToolItemPtr> DraggedItemSet(Items);
+
+	// Remove all Items whose Parent is in the Item Set 
+	Items.RemoveAll([&DraggedItemSet](const FNavigationToolItemPtr& InItem)
+		{
+			return !InItem.IsValid() || !InItem->GetParent().IsValid() || DraggedItemSet.Contains(InItem->GetParent());
+		});
+
+	// Reverse order for Onto since Item->AddChild(...) adds it to Index 0, so last item would be at the top, which is reversed
+	if (InDropZone == EItemDropZone::OntoItem)
+	{
+		Algo::Reverse(Items);
+	}
+
+	const TSharedRef<FNavigationTool> Tool = StaticCastSharedRef<FNavigationTool>(InTargetItem->GetOwnerTool());
+
+	for (const FNavigationToolItemPtr& Item : Items)
+	{
+		Params.Item = Item;
+		Tool->EnqueueItemAction<FNavigationToolAddItem>(Params);
+	}*/
+}
+
+}
diff --git a/Engine/Source/Editor/Sequencer/Private/Tools/NavigationTool/DragDropOps/Handlers/NavigationToolActorDropHandler.h b/Engine/Source/Editor/Sequencer/Private/Tools/NavigationTool/DragDropOps/Handlers/NavigationToolActorDropHandler.h
new file mode 100644
--- /dev/null
+++ b/Engine/Source/Editor/Sequencer/Private/Tools/NavigationTool/DragDropOps/Handlers/NavigationToolActorDropHandler.h
@@ -0,0 +1,26 @@
+﻿// Copyright Epic Games, Inc. All Rights Reserved.
+
+#pragma once
+
+#include "Tools/NavigationTool/DragDropOps/Handlers/NavigationToolItemDropHandler.h"
+
+namespace UE::Sequencer::NavigationTool
+{
+
+/** Class that handles Dropping Actor Items into a Target Item */
+class FNavigationToolActorDropHandler : public FNavigationToolItemDropHandler
+{
+public:
+	UE_NAVIGATIONTOOL_INHERITS(FNavigationToolActorDropHandler, FNavigationToolItemDropHandler);
+
+private:
+	//~ Begin FNavigationToolItemDropHandler
+	virtual bool IsDraggedItemSupported(const TSharedPtr<INavigationToolItem>& InDraggedItem) const override;
+	virtual TOptional<EItemDropZone> CanDrop(EItemDropZone InDropZone, const FNavigationToolItemPtr& InTargetItem) const override;
+	virtual bool Drop(EItemDropZone InDropZone, const FNavigationToolItemPtr& InTargetItem) override;
+	//~ End FNavigationToolItemDropHandler
+
+	void MoveItems(EItemDropZone InDropZone, const FNavigationToolItemPtr& InTargetItem);
+};
+
+}
diff --git a/Engine/Source/Editor/Sequencer/Private/Tools/NavigationTool/DragDropOps/Handlers/NavigationToolItemDropHandler.cpp b/Engine/Source/Editor/Sequencer/Private/Tools/NavigationTool/DragDropOps/Handlers/NavigationToolItemDropHandler.cpp
new file mode 100644
--- /dev/null
+++ b/Engine/Source/Editor/Sequencer/Private/Tools/NavigationTool/DragDropOps/Handlers/NavigationToolItemDropHandler.cpp
@@ -0,0 +1,21 @@
+﻿// Copyright Epic Games, Inc. All Rights Reserved.
+
+#include "Tools/NavigationTool/DragDropOps/Handlers/NavigationToolItemDropHandler.h"
+#include "Tools/NavigationTool/DragDropOps/NavigationToolItemDragDropOp.h"
+
+namespace UE::Sequencer::NavigationTool
+{
+
+void FNavigationToolItemDropHandler::Initialize(const FNavigationToolItemDragDropOp& InDragDropOp)
+{
+	Items = InDragDropOp.GetItems();
+	ActionType = InDragDropOp.GetActionType();
+
+	// Remove all Items that are Invalid or out of the Scope of this Handler
+	Items.RemoveAll([this](const FNavigationToolItemPtr& InItem)
+		{
+			return !InItem.IsValid() || !IsDraggedItemSupported(InItem);
+		});
+}
+
+}
diff --git a/Engine/Source/Editor/Sequencer/Private/Tools/NavigationTool/DragDropOps/NavigationToolItemDragDropOp.cpp b/Engine/Source/Editor/Sequencer/Private/Tools/NavigationTool/DragDropOps/NavigationToolItemDragDropOp.cpp
new file mode 100644
--- /dev/null
+++ b/Engine/Source/Editor/Sequencer/Private/Tools/NavigationTool/DragDropOps/NavigationToolItemDragDropOp.cpp
@@ -0,0 +1,159 @@
+// Copyright Epic Games, Inc. All Rights Reserved.
+
+#include "Tools/NavigationTool/DragDropOps/NavigationToolItemDragDropOp.h"
+#include "GameFramework/Actor.h"
+#include "ScopedTransaction.h"
+#include "Styling/SlateIconFinder.h"
+#include "Tools/NavigationTool/DragDropOps/Handlers/NavigationToolActorDropHandler.h"
+#include "Tools/NavigationTool/DragDropOps/Handlers/NavigationToolItemDropHandler.h"
+#include "Tools/NavigationTool/Items/INavigationToolItem.h"
+#include "Tools/NavigationTool/Items/NavigationToolActor.h"
+#include "Tools/NavigationTool/NavigationTool.h"
+#include "Tools/NavigationTool/NavigationToolView.h"
+
+#define LOCTEXT_NAMESPACE "NavigationToolItemDragDropOp"
+
+namespace UE::Sequencer::NavigationTool
+{
+
+FText GetActionName(const ENavigationToolDragDropActionType InActionType)
+{
+	switch (InActionType)
+	{
+		case ENavigationToolDragDropActionType::Move:
+			return LOCTEXT("MoveAction", "Moving");
+
+		case ENavigationToolDragDropActionType::Copy:
+			return LOCTEXT("CopyAction", "Copying");
+	}
+	return FText::GetEmpty();
+}
+
+FText GetItemName(const TArray<FNavigationToolItemPtr>& InItems)
+{
+	if (InItems.IsEmpty())
+	{
+		return LOCTEXT("NoItems", "0 Items");
+	}
+
+	if (InItems.Num() == 1)
+	{
+		return InItems[0]->GetDisplayName();
+	}
+
+	return FText::Format(LOCTEXT("ManyItems", "{0} and {1} other item(s)"), InItems[0]->GetDisplayName(), InItems.Num() - 1);
+}
+
+TSharedRef<FNavigationToolItemDragDropOp> FNavigationToolItemDragDropOp::New(const TArray<FNavigationToolItemPtr>& InItems
+	, const TSharedPtr<FNavigationToolView>& InToolView
+	, ENavigationToolDragDropActionType InActionType)
+{
+	TSharedRef<FNavigationToolItemDragDropOp> DragDropOp = MakeShared<FNavigationToolItemDragDropOp>();
+	DragDropOp->Init(InItems, InToolView, InActionType);
+	return DragDropOp;
+}
+
+void FNavigationToolItemDragDropOp::GetDragDropOpActors(TArray<TWeakObjectPtr<AActor>>& OutActors) const
+{
+	for (const FNavigationToolItemPtr& Item : Items)
+	{
+		if (!Item.IsValid())
+		{
+			continue;
+		}
+		if (FNavigationToolActor* const ActorItem = Item->CastTo<FNavigationToolActor>())
+		{
+			if (AActor* const UnderlyingActor = ActorItem->GetActor())
+			{
+				OutActors.Add(UnderlyingActor);
+			}
+		}
+	}
+}
+
+FReply FNavigationToolItemDragDropOp::Drop(EItemDropZone InDropZone, const FNavigationToolItemPtr& InTargetItem)
+{
+	FReply Reply = FReply::Unhandled();
+
+	FScopedTransaction Transaction(LOCTEXT("DropItems", "Navigation Tool Drop Items"));
+
+	for (const TSharedRef<FNavigationToolItemDropHandler>& DropHandler : DropHandlers)
+	{
+		if (DropHandler->GetItems().IsEmpty())
+		{
+			continue;
+		}
+
+		if (DropHandler->Drop(InDropZone, InTargetItem))
+		{
+			Reply = FReply::Handled();
+		}
+	}
+
+	if (!Reply.IsEventHandled())
+	{
+		Transaction.Cancel();
+	}
+
+	return Reply;
+}
+
+TOptional<EItemDropZone> FNavigationToolItemDragDropOp::CanDrop(EItemDropZone InDropZone, const FNavigationToolItemPtr& InTargetItem) const
+{
+	const TSharedPtr<INavigationToolView> ToolView = GetToolView();
+
+	// Only support Drag/Drop from Same Navigation Tool
+	if (!ToolView.IsValid() || &InTargetItem->GetOwnerTool() != ToolView->GetOwnerTool().Get())
+	{
+		return TOptional<EItemDropZone>();
+	}
+
+	for (const TSharedRef<FNavigationToolItemDropHandler>& DropHandler : DropHandlers)
+	{
+		if (DropHandler->GetItems().IsEmpty())
+		{
+			continue;
+		}
+
+		// Return the Item Drop Zone of the first Drop Handler that supports the Drop Zone and Target Item
+		const TOptional<EItemDropZone> DropZone = DropHandler->CanDrop(InDropZone, InTargetItem);
+		if (DropZone.IsSet())
+		{
+			return *DropZone;
+		}
+	}
+
+	return TOptional<EItemDropZone>();
+}
+
+void FNavigationToolItemDragDropOp::Init(const TArray<FNavigationToolItemPtr>& InItems
+	, const TSharedPtr<FNavigationToolView>& InToolView
+	, ENavigationToolDragDropActionType InActionType)
+{
+	Items            = InItems;
+	WeakToolView = InToolView;
+	ActionType       = InActionType;
+	MouseCursor      = EMouseCursor::GrabHandClosed;
+
+	CurrentIconBrush = Items.Num() == 1
+		? Items[0]->GetIconBrush()
+		: FSlateIconFinder::FindIconForClass(AActor::StaticClass()).GetIcon();
+
+	CurrentHoverText = FText::Format(LOCTEXT("HoverText", "{0} {1}")
+		, GetActionName(InActionType)
+		, GetItemName(InItems));
+
+	CurrentIconColorAndOpacity = FSlateColor::UseForeground();
+
+	SetupDefaults();
+	Construct();
+
+	// Add Default Drop Handlers 
+	AddDropHandler<FNavigationToolActorDropHandler>();
+
+	FNavigationToolExtender::OnItemDragDropOpInitialized().Broadcast(*this);
+}
+
+}
+
+#undef LOCTEXT_NAMESPACE
diff --git a/Engine/Source/Editor/Sequencer/Private/Tools/NavigationTool/Filters/Filters/NavigationToolBuiltInFilter.cpp b/Engine/Source/Editor/Sequencer/Private/Tools/NavigationTool/Filters/Filters/NavigationToolBuiltInFilter.cpp
new file mode 100644
--- /dev/null
+++ b/Engine/Source/Editor/Sequencer/Private/Tools/NavigationTool/Filters/Filters/NavigationToolBuiltInFilter.cpp
@@ -0,0 +1,119 @@
+﻿// Copyright Epic Games, Inc. All Rights Reserved.
+
+#include "Tools/NavigationTool/Filters/Filters/NavigationToolBuiltInFilter.h"
+#include "Tools/NavigationTool/Items/NavigationToolSequence.h"
+
+namespace UE::Sequencer::NavigationTool
+{
+
+FNavigationToolBuiltInFilter::FNavigationToolBuiltInFilter(const FNavigationToolBuiltInFilterParams& InFilterParams)
+	: FFilterBase<FNavigationToolItemPtr>(nullptr)
+	, bIsActive(true)
+	, FilterParams(InFilterParams)
+{
+}
+
+FText FNavigationToolBuiltInFilter::GetDisplayName() const
+{ 
+	return FilterParams.GetDisplayName();
+}
+
+FText FNavigationToolBuiltInFilter::GetToolTipText() const
+{
+	return FilterParams.GetTooltipText();
+}
+
+FLinearColor FNavigationToolBuiltInFilter::GetColor() const
+{
+	return FLinearColor::White;
+}
+
+FName FNavigationToolBuiltInFilter::GetIconName() const
+{
+	if (const FSlateBrush* const IconBrush = FilterParams.GetIconBrush())
+	{
+		return IconBrush->GetResourceName();
+	}
+	return NAME_None;
+}
+
+FString FNavigationToolBuiltInFilter::GetName() const
+{
+	return FilterParams.GetFilterId().ToString();
+}
+
+FSlateIcon FNavigationToolBuiltInFilter::GetIcon() const
+{
+	return FSlateIcon(FAppStyle::GetAppStyleSetName(), GetIconName());
+}
+
+bool FNavigationToolBuiltInFilter::PassesFilter(const FNavigationToolFilterType InItem) const
+{
+	auto ItemObjectPassesFilter = [this](const FNavigationToolItemPtr& InItem)
+		{
+			if (const FNavigationToolObject* ObjectItem = InItem->CastTo<FNavigationToolObject>())
+			{
+				const UObject* const Object = ObjectItem->GetObject();
+				return IsValid(Object) && FilterParams.IsValidObjectClass(Object->GetClass());
+			}
+			return false;
+		};
+
+	auto ItemTypePassesFilter = [this](const FNavigationToolItemPtr& InItem)
+		{
+			return FilterParams.IsValidItemClass(InItem->GetTypeId());
+		};
+
+	const ENavigationToolFilterMode FilterMode = FilterParams.GetFilterMode();
+
+	if (EnumHasAnyFlags(FilterMode, ENavigationToolFilterMode::MatchesType)
+		&& ItemObjectPassesFilter(InItem) || ItemTypePassesFilter(InItem))
+	{
+		return true;
+	}
+
+	if (EnumHasAnyFlags(FilterMode, ENavigationToolFilterMode::ContainerOfType))
+	{
+		TArray<FNavigationToolItemPtr> RemainingItems = InItem->GetChildren();
+
+		while (!RemainingItems.IsEmpty())
+		{
+			FNavigationToolItemPtr Item = RemainingItems.Pop();
+
+			// Stop here if Item is invalid or if it's a Sequence Item as any children of this item will be considered contained by the item itself, not the querying item
+			// this could be improved if the Item Type filter as a whole ever moves past only checking UObject types, and moving past that Sequences are the top level items
+			if (!Item.IsValid() || Item->IsA<FNavigationToolSequence>())
+			{
+				continue;
+			}
+
+			if (ItemObjectPassesFilter(Item) || ItemTypePassesFilter(InItem))
+			{
+				return true;
+			}
+
+			RemainingItems.Append(Item->GetChildren());
+		}
+	}
+
+	return false;
+}
+
+bool FNavigationToolBuiltInFilter::IsActive() const
+{
+	return bIsActive;
+}
+
+void FNavigationToolBuiltInFilter::SetActive(const bool bInActive)
+{
+	FFilterBase<FNavigationToolItemPtr>::SetActive(bInActive);
+	
+	bIsActive = bInActive;
+}
+
+const FNavigationToolBuiltInFilterParams& FNavigationToolBuiltInFilter::GetFilterParams() const
+{
+	return FilterParams;
+}
+
+}
diff --git a/Engine/Source/Editor/Sequencer/Private/Tools/NavigationTool/Filters/Filters/NavigationToolFilterBase.cpp b/Engine/Source/Editor/Sequencer/Private/Tools/NavigationTool/Filters/Filters/NavigationToolFilterBase.cpp
new file mode 100644
--- /dev/null
+++ b/Engine/Source/Editor/Sequencer/Private/Tools/NavigationTool/Filters/Filters/NavigationToolFilterBase.cpp
@@ -0,0 +1,22 @@
+﻿// Copyright Epic Games, Inc. All Rights Reserved.
+
+#include "Tools/NavigationTool/Filters/Filters/NavigationToolFilterBase.h"
+#include "Filters/ISequencerFilterBar.h"
+#include "Tools/NavigationTool/Filters/INavigationToolFilterBar.h"
+
+namespace UE::Sequencer::NavigationTool
+{
+
+using namespace UE::Sequencer;
+
+FNavigationToolFilter::FNavigationToolFilter(INavigationToolFilterBar& InOutFilterInterface, TSharedPtr<FFilterCategory>&& InCategory)
+	: FSequencerFilterBase<FNavigationToolItemPtr>(InOutFilterInterface, MoveTemp(InCategory))
+{
+}
+
+INavigationToolFilterBar& FNavigationToolFilter::GetFilterInterface() const
+{
+	return static_cast<INavigationToolFilterBar&>(FilterInterface);
+}
+
+}
diff --git a/Engine/Source/Editor/Sequencer/Private/Tools/NavigationTool/Filters/Filters/NavigationToolFilter_Marks.cpp b/Engine/Source/Editor/Sequencer/Private/Tools/NavigationTool/Filters/Filters/NavigationToolFilter_Marks.cpp
new file mode 100644
--- /dev/null
+++ b/Engine/Source/Editor/Sequencer/Private/Tools/NavigationTool/Filters/Filters/NavigationToolFilter_Marks.cpp
@@ -0,0 +1,70 @@
+﻿// Copyright Epic Games, Inc. All Rights Reserved.
+
+#include "Tools/NavigationTool/Filters/Filters/NavigationToolFilter_Marks.h"
+#include "Filters/SequencerFilterBar.h"
+#include "Tools/NavigationTool/Filters/NavigationToolFilterCommands.h"
+#include "Tools/NavigationTool/Items/NavigationToolSequence.h"
+
+#define LOCTEXT_NAMESPACE "NavigationToolFilter_Marks"
+
+namespace UE::Sequencer::NavigationTool
+{
+
+using namespace UE::Sequencer;
+
+FNavigationToolFilter_Marks::FNavigationToolFilter_Marks(INavigationToolFilterBar& InFilterInterface, TSharedPtr<FFilterCategory> InCategory)
+	: FNavigationToolFilter(InFilterInterface, MoveTemp(InCategory))
+{
+}
+
+FText FNavigationToolFilter_Marks::GetDefaultToolTipText() const
+{
+	return LOCTEXT("NavigationToolFilter_MarksToolTip", "Show only sequences that contain marked frames");
+}
+
+TSharedPtr<FUICommandInfo> FNavigationToolFilter_Marks::GetToggleCommand() const
+{
+	return FNavigationToolFilterCommands::Get().ToggleFilter_Marks;
+}
+
+FText FNavigationToolFilter_Marks::GetDisplayName() const
+{
+	return LOCTEXT("NavigationToolFilter_Marks", "Marks");
+}
+
+FSlateIcon FNavigationToolFilter_Marks::GetIcon() const
+{
+	return FSlateIcon(FAppStyle::GetAppStyleSetName(), TEXT("AnimTimeline.SectionMarker"));
+}
+
+FString FNavigationToolFilter_Marks::GetName() const
+{
+	return StaticName();
+}
+
+bool FNavigationToolFilter_Marks::PassesFilter(FNavigationToolFilterType InItem) const
+{
+	const FNavigationToolSequence* const SequenceItem = InItem->CastTo<FNavigationToolSequence>();
+	if (!SequenceItem)
+	{
+		return false;
+	}
+
+	UMovieSceneSequence* const Sequence = SequenceItem->GetSequence();
+	if (IsValid(Sequence))
+	{
+		return false;
+	}
+
+	UMovieScene* const MovieScene = Sequence->GetMovieScene();
+	if (IsValid(MovieScene))
+	{
+		return false;
+	}
+
+	return MovieScene->GetMarkedFrames().Num() > 0;
+}
+
+}
+
+#undef LOCTEXT_NAMESPACE
diff --git a/Engine/Source/Editor/Sequencer/Private/Tools/NavigationTool/Filters/Filters/NavigationToolFilter_Unbound.cpp b/Engine/Source/Editor/Sequencer/Private/Tools/NavigationTool/Filters/Filters/NavigationToolFilter_Unbound.cpp
new file mode 100644
--- /dev/null
+++ b/Engine/Source/Editor/Sequencer/Private/Tools/NavigationTool/Filters/Filters/NavigationToolFilter_Unbound.cpp
@@ -0,0 +1,60 @@
+﻿// Copyright Epic Games, Inc. All Rights Reserved.
+
+#include "Tools/NavigationTool/Filters/Filters/NavigationToolFilter_Unbound.h"
+#include "Sequencer.h"
+#include "Tools/NavigationTool/Filters/NavigationToolFilterCommands.h"
+#include "Tools/NavigationTool/Items/NavigationToolBinding.h"
+
+#define LOCTEXT_NAMESPACE "NavigationToolFilter_Unbound"
+
+namespace UE::Sequencer::NavigationTool
+{
+
+using namespace UE::Sequencer;
+
+FNavigationToolFilter_Unbound::FNavigationToolFilter_Unbound(INavigationToolFilterBar& InFilterInterface, TSharedPtr<FFilterCategory> InCategory)
+	: FNavigationToolFilter(InFilterInterface, MoveTemp(InCategory))
+{
+}
+
+FText FNavigationToolFilter_Unbound::GetDefaultToolTipText() const
+{
+	return LOCTEXT("NavigationToolFilter_UnboundToolTip", "Show only sequences with Unbound tracks");
+}
+
+TSharedPtr<FUICommandInfo> FNavigationToolFilter_Unbound::GetToggleCommand() const
+{
+	return FNavigationToolFilterCommands::Get().ToggleFilter_Unbound;
+}
+
+FText FNavigationToolFilter_Unbound::GetDisplayName() const
+{
+	return LOCTEXT("NavigationToolFilter_Unbound", "Unbound");
+}
+
+FSlateIcon FNavigationToolFilter_Unbound::GetIcon() const
+{
+	return FSlateIcon(FAppStyle::GetAppStyleSetName(), TEXT("BTEditor.Graph.BTNode.Decorator.DoesPathExist.Icon"));
+}
+
+FString FNavigationToolFilter_Unbound::GetName() const
+{
+	return StaticName();
+}
+
+bool FNavigationToolFilter_Unbound::PassesFilter(FNavigationToolFilterType InItem) const
+{
+	const FNavigationToolBinding* const BindingItem = InItem->CastTo<FNavigationToolBinding>();
+	if (!BindingItem)
+	{
+		return false;
+	}
+
+	UObject* const BoundObject = BindingItem->FindBoundObject();
+
+	return IsValid(BoundObject);
+}
+
+}
+
+#undef LOCTEXT_NAMESPACE
diff --git a/Engine/Source/Editor/Sequencer/Private/Tools/NavigationTool/Filters/NavigationToolFilterData.cpp b/Engine/Source/Editor/Sequencer/Private/Tools/NavigationTool/Filters/NavigationToolFilterData.cpp
new file mode 100644
--- /dev/null
+++ b/Engine/Source/Editor/Sequencer/Private/Tools/NavigationTool/Filters/NavigationToolFilterData.cpp
@@ -0,0 +1,141 @@
+﻿// Copyright Epic Games, Inc. All Rights Reserved.
+
+#include "Tools/NavigationTool/Filters/NavigationToolFilterData.h"
+#include "Tools/NavigationTool/Items/NavigationToolSequence.h"
+
+namespace UE::Sequencer::NavigationTool
+{
+using namespace UE::Sequencer;
+
+FNavigationToolFilterData::FNavigationToolFilterData(const FString& InRawFilterText)
+	: RawFilterText(InRawFilterText)
+{
+}
+
+bool FNavigationToolFilterData::operator==(const FNavigationToolFilterData& InRhs) const
+{
+	return ContainsFilterInNodes(InRhs) && TotalNodeCount == InRhs.GetTotalNodeCount();
+}
+
+bool FNavigationToolFilterData::operator!=(const FNavigationToolFilterData& InRhs) const
+{
+	return !(*this == InRhs);
+}
+
+void FNavigationToolFilterData::Reset()
+{
+	RawFilterText.Reset();
+	TotalNodeCount = 0;
+	FilterInNodes.Reset();
+}
+
+FString FNavigationToolFilterData::GetRawFilterText() const
+{
+	return RawFilterText;
+}
+
+uint32 FNavigationToolFilterData::GetDisplayNodeCount() const
+{
+	return FilterInNodes.Num();
+}
+
+uint32 FNavigationToolFilterData::GetTotalNodeCount() const
+{
+	return TotalNodeCount;
+}
+
+uint32 FNavigationToolFilterData::GetFilterInCount() const
+{
+	return FilterInNodes.Num();
+}
+
+uint32 FNavigationToolFilterData::GetFilterOutCount() const
+{
+	return GetTotalNodeCount() - GetFilterInCount();
+}
+
+void FNavigationToolFilterData::IncrementTotalNodeCount()
+{
+	++TotalNodeCount;
+}
+
+void FNavigationToolFilterData::FilterInNode(const FNavigationToolItemPtr& InNode)
+{
+	FilterInNodes.Add(InNode);
+}
+
+void FNavigationToolFilterData::FilterOutNode(const FNavigationToolItemPtr& InNode)
+{
+	const FSetElementId ElementId = FilterInNodes.FindId(InNode);
+	if (ElementId.IsValidId())
+	{
+		FilterInNodes.Remove(ElementId);
+	}
+}
+
+void FNavigationToolFilterData::FilterInParentChildNodes(const FNavigationToolItemPtr& InNode
+	, const bool bInIncludeSelf, const bool bInIncludeParents, const bool bInIncludeChildren)
+{
+	if (!InNode.IsValid())
+	{
+		return;
+	}
+
+	if (bInIncludeParents)
+	{
+		FNavigationToolItemPtr CurrentParent = InNode->GetParent();
+		while (CurrentParent.IsValid())
+		{
+			FilterInNode(CurrentParent);
+
+			CurrentParent = CurrentParent->GetParent();
+			if (CurrentParent.IsValid() && CurrentParent->GetItemId() == FNavigationToolItemId::RootId)
+			{
+				CurrentParent = nullptr;
+			}
+		}
+	}
+
+	if (bInIncludeSelf)
+	{
+		FilterInNode(InNode);
+	}
+
+	if (bInIncludeChildren)
+	{
+		auto AddAllChildren = [this](const FNavigationToolItemPtr& InNode)
+		{
+			for (const FNavigationToolItemPtr& Child : InNode->GetChildren())
+			{
+				FilterInNode(Child);
+			}
+		};
+
+		for (const FNavigationToolItemPtr& Child : InNode->GetChildren())
+		{
+			AddAllChildren(Child);
+		}
+	}
+}
+
+void FNavigationToolFilterData::FilterInNodeWithAncestors(const FNavigationToolItemPtr& InNode)
+{
+	FilterInParentChildNodes(InNode, true, true, false);
+}
+
+bool FNavigationToolFilterData::ContainsFilterInNodes(const FNavigationToolFilterData& InOtherData) const
+{
+	return FilterInNodes.Includes(InOtherData.FilterInNodes);
+}
+
+bool FNavigationToolFilterData::IsFilteredIn(const FNavigationToolItemPtr& InNode) const
+{
+	return FilterInNodes.Contains(InNode);
+}
+
+bool FNavigationToolFilterData::IsFilteredOut(const FNavigationToolItemPtr& InNode) const
+{
+	return !FilterInNodes.Contains(InNode);
+}
+
+}
diff --git a/Engine/Source/Editor/Sequencer/Private/Tools/NavigationTool/Filters/TextExpressions/NavigationToolFilterTextExpressionContext.cpp b/Engine/Source/Editor/Sequencer/Private/Tools/NavigationTool/Filters/TextExpressions/NavigationToolFilterTextExpressionContext.cpp
new file mode 100644
--- /dev/null
+++ b/Engine/Source/Editor/Sequencer/Private/Tools/NavigationTool/Filters/TextExpressions/NavigationToolFilterTextExpressionContext.cpp
@@ -0,0 +1,68 @@
+// Copyright Epic Games, Inc. All Rights Reserved.
+
+#include "Tools/NavigationTool/Filters/TextExpressions/NavigationToolFilterTextExpressionContext.h"
+#include "Sequencer.h"
+#include "Tools/NavigationTool/Filters/INavigationToolFilterBar.h"
+#include "Tools/NavigationTool/Items/NavigationToolItem.h"
+
+namespace UE::Sequencer::NavigationTool
+{
+
+using namespace UE::Sequencer;
+
+FNavigationToolFilterTextExpressionContext::FNavigationToolFilterTextExpressionContext(INavigationToolFilterBar& InFilterInterface)
+	: FilterInterface(InFilterInterface)
+{
+}
+
+void FNavigationToolFilterTextExpressionContext::SetFilterItem(const FNavigationToolItemPtr& InFilterItem)
+{
+	WeakFilterItem = InFilterItem;
+}
+
+bool FNavigationToolFilterTextExpressionContext::TestBasicStringExpression(const FTextFilterString& InValue
+	, const ETextFilterTextComparisonMode InTextComparisonMode) const
+{
+	if (const TSharedPtr<INavigationToolItem> FilterItem = WeakFilterItem.Pin())
+	{
+		const FString Label = FilterItem->GetDisplayName().ToString();
+		if (TextFilterUtils::TestBasicStringExpression(Label, InValue, ETextFilterTextComparisonMode::Partial))
+		{
+			return true;
+		}
+	}
+
+	return false;
+}
+
+bool FNavigationToolFilterTextExpressionContext::TestComplexExpression(const FName& InKey
+	, const FTextFilterString& InValue
+	, const ETextFilterComparisonOperation InComparisonOperation
+	, const ETextFilterTextComparisonMode InTextComparisonMode) const
+{
+	if (!WeakFilterItem.IsValid())
+	{
+		return false;
+	}
+
+	const TSet<FName> Keys = GetKeys();
+	if (!Keys.IsEmpty() && !Keys.Contains(InKey))
+	{
+		return false;
+	}
+
+	return !InValue.IsEmpty();
+}
+
+UMovieSceneSequence* FNavigationToolFilterTextExpressionContext::GetFocusedMovieSceneSequence() const
+{
+	return FilterInterface.GetSequencer().GetFocusedMovieSceneSequence();
+}
+
+UMovieScene* FNavigationToolFilterTextExpressionContext::GetFocusedGetMovieScene() const
+{
+	const UMovieSceneSequence* const FocusedMovieSceneSequence = GetFocusedMovieSceneSequence();
+	return IsValid(FocusedMovieSceneSequence) ? FocusedMovieSceneSequence->GetMovieScene() : nullptr;
+}
+
+}
diff --git a/Engine/Source/Editor/Sequencer/Private/Tools/NavigationTool/ItemActions/NavigationToolAddItem.cpp b/Engine/Source/Editor/Sequencer/Private/Tools/NavigationTool/ItemActions/NavigationToolAddItem.cpp
new file mode 100644
--- /dev/null
+++ b/Engine/Source/Editor/Sequencer/Private/Tools/NavigationTool/ItemActions/NavigationToolAddItem.cpp
@@ -0,0 +1,123 @@
+// Copyright Epic Games, Inc. All Rights Reserved.
+
+#include "Tools/NavigationTool/ItemActions/NavigationToolAddItem.h"
+#include "Tools/NavigationTool/NavigationTool.h"
+#include "Tools/NavigationTool/NavigationToolScopedSelection.h"
+#include "Tools/NavigationTool/Items/INavigationToolItem.h"
+#include "Widgets/Views/STableRow.h"
+
+namespace UE::Sequencer::NavigationTool
+{
+
+FNavigationToolAddItem::FNavigationToolAddItem(const FNavigationToolAddItemParams& InAddItemParams)
+	: AddParams(InAddItemParams)
+{
+}
+
+bool FNavigationToolAddItem::ShouldTransact() const
+{
+	return EnumHasAnyFlags(AddParams.Flags, ENavigationToolAddItemFlags::Transact);
+}
+
+void FNavigationToolAddItem::Execute(INavigationTool& InTool)
+{
+	if (!AddParams.Item.IsValid())
+	{
+		return;
+	}
+
+	// Try to Create Children on Find
+	if (EnumHasAnyFlags(AddParams.Flags, ENavigationToolAddItemFlags::AddChildren))
+	{
+		constexpr bool bRecursiveFind = true;
+		TArray<FNavigationToolItemPtr> Children;
+		AddParams.Item->FindValidChildren(Children, bRecursiveFind);
+	}
+
+	const FNavigationToolItemPtr ParentItem = AddParams.Item->GetParent();
+
+	// If this array has elements in it, then we need to stop a circular dependency from forming
+	const TArray<FNavigationToolItemPtr> PathToRelativeItem = AddParams.Item->FindPath({ AddParams.RelativeItem });
+	if (PathToRelativeItem.Num() > 0 && ParentItem.IsValid())
+	{
+		FNavigationToolAddItemParams CircularSolverParams;
+		CircularSolverParams.Item             = PathToRelativeItem[0];
+		CircularSolverParams.RelativeItem     = AddParams.Item;
+		CircularSolverParams.RelativeDropZone = EItemDropZone::AboveItem;
+		CircularSolverParams.Flags            = AddParams.Flags;
+
+		ParentItem->AddChild(MoveTemp(CircularSolverParams));
+	}
+
+	if (AddParams.RelativeItem.IsValid())
+	{
+		const FNavigationToolItemPtr RelativeItemParent = AddParams.RelativeItem->GetParent();
+
+		// If it's onto item, the Relative Item is going to be the Parent
+		if (!AddParams.RelativeDropZone.IsSet() || AddParams.RelativeDropZone == EItemDropZone::OntoItem)
+		{
+			//If the Relative Item is Onto and it's the same as the Current Parent, shift Item up in the Hierarchy
+			//(as long as the parent is a valid one)
+			if (AddParams.RelativeItem == ParentItem && RelativeItemParent)
+			{
+				AddParams.RelativeDropZone = EItemDropZone::BelowItem;
+				RelativeItemParent->AddChild(AddParams);
+			}
+			else
+			{
+				AddParams.RelativeItem->AddChild(AddParams);
+			}
+		}
+		//else we place it as a Sibling to the Relative Item
+		else if (RelativeItemParent)
+		{
+			RelativeItemParent->AddChild(AddParams);
+		}
+		//if no parent, then add it to the tree root
+		else
+		{
+			InTool.GetTreeRoot()->AddChild(AddParams);
+		}
+	}
+	else
+	{
+		//If no Relative Item, add to tree root
+		InTool.GetTreeRoot()->AddChild(AddParams);
+	}
+
+	const FNavigationToolScopedSelection ScopedSelection(*InTool.GetSequencer(), ENavigationToolScopedSelectionPurpose::Read);
+
+	// Automatically select item if it's selected
+	if (AddParams.Item->IsSelected(ScopedSelection))
+	{
+		// Select in Navigation Tool but don't signal selection as we already have it selected in Mode Tools
+		AddParams.Flags = ENavigationToolAddItemFlags::Select;
+		AddParams.SelectionFlags &= ~ENavigationToolItemSelectionFlags::SignalSelectionChange;
+	}
+	// Signal selection change when we attempt to select this item in the Navigation Tool but it isn't selected in Sequencer
+	else if (EnumHasAnyFlags(AddParams.Flags, ENavigationToolAddItemFlags::Select))
+	{
+		AddParams.SelectionFlags |= ENavigationToolItemSelectionFlags::SignalSelectionChange;
+	}
+
+	if (EnumHasAnyFlags(AddParams.Flags, ENavigationToolAddItemFlags::Select))
+	{
+		InTool.SelectItems({ AddParams.Item }, AddParams.SelectionFlags);
+	}
+
+	//InTool.SetToolModified();
+}
+
+void FNavigationToolAddItem::OnObjectsReplaced(const TMap<UObject*, UObject*>& InReplacementMap, bool bRecursive)
+{
+	if (AddParams.Item.IsValid())
+	{
+		AddParams.Item->OnObjectsReplaced(InReplacementMap, bRecursive);
+	}
+	if (AddParams.RelativeItem.IsValid())
+	{
+		AddParams.RelativeItem->OnObjectsReplaced(InReplacementMap, bRecursive);
+	}
+}
+
+}
diff --git a/Engine/Source/Editor/Sequencer/Private/Tools/NavigationTool/ItemActions/NavigationToolItemAction.cpp b/Engine/Source/Editor/Sequencer/Private/Tools/NavigationTool/ItemActions/NavigationToolItemAction.cpp
new file mode 100644
--- /dev/null
+++ b/Engine/Source/Editor/Sequencer/Private/Tools/NavigationTool/ItemActions/NavigationToolItemAction.cpp
@@ -0,0 +1,3 @@
+// Copyright Epic Games, Inc. All Rights Reserved.
+
+#include "Tools/NavigationTool/ItemActions/NavigationToolItemAction.h"
