[Sequencer] Navigation Tool initial submit

diff --git a/Engine/Source/Editor/Sequencer/Private/Tools/NavigationTool/ItemActions/NavigationToolRemoveItem.cpp b/Engine/Source/Editor/Sequencer/Private/Tools/NavigationTool/ItemActions/NavigationToolRemoveItem.cpp
new file mode 100644
--- /dev/null
+++ b/Engine/Source/Editor/Sequencer/Private/Tools/NavigationTool/ItemActions/NavigationToolRemoveItem.cpp
@@ -0,0 +1,95 @@
+// Copyright Epic Games, Inc. All Rights Reserved.
+
+#include "Tools/NavigationTool/ItemActions/NavigationToolRemoveItem.h"
+#include "Tools/NavigationTool/ItemActions/NavigationToolAddItem.h"
+#include "Tools/NavigationTool/Items/INavigationToolItem.h"
+#include "Tools/NavigationTool/NavigationTool.h"
+#include "Widgets/Views/STableRow.h"
+
+namespace UE::Sequencer::NavigationTool
+{
+
+FNavigationToolRemoveItem::FNavigationToolRemoveItem(const FNavigationToolRemoveItemParams& InRemoveItemParams)
+	: RemoveParams(InRemoveItemParams)
+{
+	if (RemoveParams.Item.IsValid())
+	{
+		RemoveParams.Item->AddFlags(ENavigationToolItemFlags::PendingRemoval);
+	}
+}
+
+void FNavigationToolRemoveItem::Execute(INavigationTool& InTool)
+{
+	if (RemoveParams.Item.IsValid())
+	{
+		FNavigationToolItemFlagGuard Guard(RemoveParams.Item, ENavigationToolItemFlags::IgnorePendingKill);
+
+		//Copy the Array since we may modify it below on Add/Remove Child
+		TArray<FNavigationToolItemPtr> Children = RemoveParams.Item->GetChildren();
+
+		FNavigationToolItemPtr Parent       = RemoveParams.Item->GetParent();
+		FNavigationToolItemPtr RelativeItem = RemoveParams.Item;
+
+		//Search the lowest parent that is not pending removal
+		while (Parent.IsValid() && Parent->HasAnyFlags(ENavigationToolItemFlags::PendingRemoval))
+		{
+			RelativeItem = Parent;
+			Parent       = Parent->GetParent();
+		}
+
+		//Reparent the Item's Children to the Valid Parent found above
+		if (Parent.IsValid())
+		{
+			FNavigationToolAddItemParams ReparentParams;
+			ReparentParams.RelativeItem             = RelativeItem;
+			ReparentParams.RelativeDropZone         = EItemDropZone::BelowItem;
+			ReparentParams.Flags                    = ENavigationToolAddItemFlags::Select;
+
+			TArray<FNavigationToolItemPtr> ItemsToReparent;
+
+			ItemsToReparent.Append(Children);
+			while (ItemsToReparent.Num() > 0)
+			{
+				ReparentParams.Item = ItemsToReparent.Pop();
+
+				//Then Parent it to the Item's Parent.
+				//If we couldn't add the Child, then it means either the Child Item itself is invalid
+				//or the underlying info (e.g. object) is invalid (e.g. actor pending kill)
+				if (Parent->AddChild(ReparentParams) == false)
+				{
+					//so try reparenting the children of this invalid child (since this invalid child will be removed)
+					if (ReparentParams.Item.IsValid())
+					{
+						ItemsToReparent.Append(ReparentParams.Item->GetChildren());
+					}
+				}
+			}
+
+			//In case the Parent is still the same, Remove Child Item, else the Parent is going to be removed anyways
+			if (Parent == RemoveParams.Item->GetParent())
+			{
+				Parent->RemoveChild(RemoveParams.Item);
+			}
+		}
+		else
+		{
+			for (const FNavigationToolItemPtr& Child : Children)
+			{
+				RemoveParams.Item->RemoveChild(Child);
+			}
+		}
+
+		InTool.UnregisterItem(RemoveParams.Item->GetItemId());
+		//InTool.SetToolModified();
+	}
+}
+
+void FNavigationToolRemoveItem::OnObjectsReplaced(const TMap<UObject*, UObject*>& InReplacementMap, bool bRecursive)
+{
+	if (RemoveParams.Item.IsValid())
+	{
+		RemoveParams.Item->OnObjectsReplaced(InReplacementMap, bRecursive);
+	}
+}
+
+}
diff --git a/Engine/Source/Editor/Sequencer/Private/Tools/NavigationTool/ItemProxies/NavigationToolComponentProxy.cpp b/Engine/Source/Editor/Sequencer/Private/Tools/NavigationTool/ItemProxies/NavigationToolComponentProxy.cpp
new file mode 100644
--- /dev/null
+++ b/Engine/Source/Editor/Sequencer/Private/Tools/NavigationTool/ItemProxies/NavigationToolComponentProxy.cpp
@@ -0,0 +1,69 @@
+// Copyright Epic Games, Inc. All Rights Reserved.
+
+#include "Tools/NavigationTool/ItemProxies/NavigationToolComponentProxy.h"
+#include "Components/ActorComponent.h"
+#include "GameFramework/Actor.h"
+#include "Styling/SlateIconFinder.h"
+#include "Tools/NavigationTool/Items/NavigationToolActor.h"
+#include "Tools/NavigationTool/Items/NavigationToolObject.h"
+#include "Tools/NavigationTool/NavigationTool.h"
+
+#define LOCTEXT_NAMESPACE "NavigationToolComponentProxy"
+
+namespace UE::Sequencer::NavigationTool
+{
+
+FNavigationToolComponentProxy::FNavigationToolComponentProxy(INavigationTool& InTool, const FNavigationToolItemPtr& InParentItem)
+	: Super(InTool, InParentItem)
+{
+}
+
+FText FNavigationToolComponentProxy::GetDisplayName() const
+{
+	return LOCTEXT("DisplayName", "Components");
+}
+
+FSlateIcon FNavigationToolComponentProxy::GetIcon() const
+{
+	return FSlateIconFinder::FindIconForClass(USceneComponent::StaticClass());
+}
+
+FText FNavigationToolComponentProxy::GetIconTooltipText() const
+{
+	return LOCTEXT("Tooltip", "Shows the Components in an Actor. Visualization, non-editable and UCS Components are excluded");
+}
+
+ENavigationToolItemViewMode FNavigationToolComponentProxy::GetSupportedViewModes(const INavigationToolView& InToolView) const
+{
+	// Components should only be visualized in Navigation Tool View and not appear in the Item Column List
+	// Support any other type of View Mode
+	return ~ENavigationToolItemViewMode::ItemTree | ENavigationToolItemViewMode::HorizontalItemList;
+}
+
+void FNavigationToolComponentProxy::GetProxiedItems(const TSharedRef<INavigationToolItem>& InParent, TArray<FNavigationToolItemPtr>& OutChildren, bool bInRecursive)
+{
+	if (const FNavigationToolActor* const ActorItem = InParent->CastTo<FNavigationToolActor>())
+	{
+		const AActor* const Actor = ActorItem->GetActor();
+		if (!IsValid(Actor))
+		{
+			return;
+		}
+
+		const TSharedRef<FNavigationToolComponentProxy> SharedThisRef = SharedThis(this);
+
+		for (UActorComponent* const ActorComponent : Actor->GetComponents())
+		{
+			const FNavigationToolItemPtr NewItem = Tool.FindOrAdd<FNavigationToolObject>(InParent->GetProvider(), SharedThisRef, ActorComponent, SharedThisRef);
+			OutChildren.Add(NewItem);
+			if (bInRecursive)
+			{
+				NewItem->FindChildren(OutChildren, bInRecursive);
+			}
+		}
+	}
+}
+
+}
+
+#undef LOCTEXT_NAMESPACE
diff --git a/Engine/Source/Editor/Sequencer/Private/Tools/NavigationTool/Items/NavigationToolActor.cpp b/Engine/Source/Editor/Sequencer/Private/Tools/NavigationTool/Items/NavigationToolActor.cpp
new file mode 100644
--- /dev/null
+++ b/Engine/Source/Editor/Sequencer/Private/Tools/NavigationTool/Items/NavigationToolActor.cpp
@@ -0,0 +1,185 @@
+// Copyright Epic Games, Inc. All Rights Reserved.
+
+#include "Tools/NavigationTool/Items/NavigationToolActor.h"
+#include "Editor/EditorEngine.h"
+#include "GameFramework/Actor.h"
+#include "ISequencer.h"
+#include "MovieScene.h"
+#include "MovieSceneSequence.h"
+#include "ScopedTransaction.h"
+#include "Tools/NavigationTool/Items/NavigationToolComponent.h"
+#include "Tools/NavigationTool/Items/NavigationToolSequence.h"
+#include "Tools/NavigationTool/ItemProxies/NavigationToolComponentProxy.h"
+#include "Tools/NavigationTool/NavigationTool.h"
+#include "Tools/NavigationTool/NavigationToolSettings.h"
+
+#define LOCTEXT_NAMESPACE "NavigationToolActor"
+
+namespace UE::Sequencer::NavigationTool
+{
+
+FNavigationToolActor::FNavigationToolActor(INavigationTool& InTool
+	, const FNavigationToolItemPtr& InParentItem
+	, const FMovieSceneBinding& InBinding
+	, UObject* const InObject
+	, const FNavigationToolItemPtr& InReferencingItem)
+	: Super(InTool
+	, InParentItem
+	, InBinding
+	, InObject
+	, InReferencingItem)
+{
+}
+
+void FNavigationToolActor::FindChildren(TArray<FNavigationToolItemPtr>& OutChildren, const bool bInRecursive)
+{
+	Super::FindChildren(OutChildren, bInRecursive);
+
+	AActor* const ThisBoundActor = GetActor();
+	if (!IsValid(ThisBoundActor))
+	{
+		return;
+	}
+
+	const TSharedPtr<FNavigationToolSequence> SequenceParentItem = StaticCastSharedPtr<FNavigationToolSequence>(WeakReferencingItem.Pin());
+	if (!SequenceParentItem.IsValid())
+	{
+		return;
+	}
+
+	UMovieSceneSequence* const Sequence = SequenceParentItem->GetSequence();
+	if (!IsValid(Sequence))
+	{
+		return;
+	}
+
+	UMovieScene* const MovieScene = Sequence->GetMovieScene();
+	if (!IsValid(MovieScene))
+	{
+		return;
+	}
+
+	const TSharedPtr<ISequencer> Sequencer = Tool.GetSequencer();
+	if (!Sequencer.IsValid())
+	{
+		return;
+	}
+
+	const TSharedRef<FNavigationToolBinding> SharedThisRef = SharedThis(this);
+
+	TArray<FMovieSceneBinding>& Bindings = MovieScene->GetBindings();
+	Bindings.Sort([](const FMovieSceneBinding& InA, const FMovieSceneBinding& InB)
+		{
+			return InA.GetSortingOrder() < InB.GetSortingOrder();
+		});
+
+	for (const FMovieSceneBinding& CurrentBinding : Bindings)
+	{
+		UObject* const BoundObject = Sequencer->FindSpawnedObjectOrTemplate(CurrentBinding.GetObjectGuid());
+		if (IsValid(BoundObject))
+		{
+			UActorComponent* const Component = Cast<UActorComponent>(BoundObject);
+			if (IsValid(Component) && Component->GetOwner() == ThisBoundActor)
+			{
+				const FNavigationToolItemPtr NewItem = Tool.FindOrAdd<FNavigationToolComponent>(GetProvider(), SharedThisRef, CurrentBinding, BoundObject, SharedThisRef);
+				OutChildren.Add(NewItem);
+				if (bInRecursive)
+				{
+					NewItem->FindChildren(OutChildren, bInRecursive);
+				}
+			}
+		}
+	}
+}
+
+void FNavigationToolActor::GetItemProxies(TArray<TSharedPtr<FNavigationToolItemProxy>>& OutItemProxies)
+{
+	Super::GetItemProxies(OutItemProxies);
+
+	if (const TSharedPtr<FNavigationToolItemProxy> ComponentItemProxy = Tool.GetOrCreateItemProxy<FNavigationToolComponentProxy>(SharedThis(this)))
+	{
+		OutItemProxies.Add(ComponentItemProxy);
+	}
+}
+
+bool FNavigationToolActor::IsAllowedInTool() const
+{
+	FNavigationTool& ToolPrivate = static_cast<FNavigationTool&>(Tool);
+	const bool bIsAllowed = ToolPrivate.IsObjectAllowedInTool(GetActor());
+	return bIsAllowed;
+}
+
+ENavigationToolItemViewMode FNavigationToolActor::GetSupportedViewModes(const INavigationToolView& InToolView) const
+{
+	return ENavigationToolItemViewMode::ItemTree | ENavigationToolItemViewMode::HorizontalItemList;
+}
+
+FText FNavigationToolActor::GetDisplayName() const
+{
+	if (const AActor* const UnderlyingActor = GetActor())
+	{
+		return FText::FromString(UnderlyingActor->GetActorLabel());
+	}
+	return FText::GetEmpty();
+}
+
+TArray<FName> FNavigationToolActor::GetTags() const
+{
+	if (AActor* const UnderlyingActor = GetActor())
+	{
+		return UnderlyingActor->Tags;
+	}
+	return Super::GetTags();
+}
+
+bool FNavigationToolActor::GetVisibility() const
+{
+	if (const AActor* const UnderlyingActor = GetActor())
+	{
+		return !UnderlyingActor->IsTemporarilyHiddenInEditor(true);
+	}
+	return false;
+}
+
+void FNavigationToolActor::OnVisibilityChanged(const bool bInNewVisibility)
+{
+	if (AActor* const UnderlyingActor = GetActor())
+	{
+		UnderlyingActor->SetIsTemporarilyHiddenInEditor(!bInNewVisibility);
+	}
+}
+
+bool FNavigationToolActor::CanRename() const
+{
+	return Super::CanRename() && IsValid(GetActor());
+}
+
+bool FNavigationToolActor::Rename(const FString& InName)
+{
+	AActor* const UnderlyingActor = GetActor();
+	const bool bIsActorLabelEditable = IsValid(UnderlyingActor) && UnderlyingActor->IsActorLabelEditable();
+
+	bool bRenamed = false;
+	
+	if (bIsActorLabelEditable && !InName.Equals(UnderlyingActor->GetActorLabel(), ESearchCase::CaseSensitive))
+	{
+		const FScopedTransaction Transaction(LOCTEXT("NavigationToolRenameActor", "Rename Actor"));
+
+		FActorLabelUtilities::RenameExistingActor(UnderlyingActor, InName);
+
+		Super::Rename(InName);
+
+		bRenamed = true;
+	}
+
+	return bRenamed;
+}
+
+AActor* FNavigationToolActor::GetActor() const
+{
+	return WeakObject.IsValid() ? Cast<AActor>(WeakObject.Get(IsIgnoringPendingKill())) : nullptr;
+}
+
+}
+
+#undef LOCTEXT_NAMESPACE
diff --git a/Engine/Source/Editor/Sequencer/Private/Tools/NavigationTool/Items/NavigationToolBinding.cpp b/Engine/Source/Editor/Sequencer/Private/Tools/NavigationTool/Items/NavigationToolBinding.cpp
new file mode 100644
--- /dev/null
+++ b/Engine/Source/Editor/Sequencer/Private/Tools/NavigationTool/Items/NavigationToolBinding.cpp
@@ -0,0 +1,152 @@
+﻿// Copyright Epic Games, Inc. All Rights Reserved.
+
+#include "Tools/NavigationTool/Items/NavigationToolBinding.h"
+#include "ISequencer.h"
+#include "ScopedTransaction.h"
+#include "Styling/SlateIconFinder.h"
+#include "Tools/NavigationTool/INavigationTool.h"
+#include "Tools/NavigationTool/Items/NavigationToolItemUtils.h"
+#include "Tools/NavigationTool/Items/NavigationToolSequence.h"
+#include "Tools/NavigationTool/NavigationToolScopedSelection.h"
+#include "Widgets/Views/STableRow.h"
+
+#define LOCTEXT_NAMESPACE "NavigationToolBinding"
+
+namespace UE::Sequencer::NavigationTool
+{
+
+using namespace UE::Sequencer;
+
+FNavigationToolBinding::FNavigationToolBinding(INavigationTool& InTool
+	, const FNavigationToolItemPtr& InParentItem
+	, const FMovieSceneBinding& InBinding
+	, UObject* const InObject
+	, const FNavigationToolItemPtr& InReferencingItem)
+	: Super(InTool, InParentItem, InObject, InReferencingItem, Binding.GetObjectGuid().ToString())
+	, Binding(InBinding)
+{
+}
+
+void FNavigationToolBinding::FindChildren(TArray<FNavigationToolItemPtr>& OutChildren, const bool bInRecursive)
+{
+	Super::FindChildren(OutChildren, bInRecursive);
+}
+
+FText FNavigationToolBinding::GetDisplayName() const
+{
+	return FText::FromString(Binding.GetName());
+}
+
+const FSlateBrush* FNavigationToolBinding::GetDefaultIconBrush() const
+{
+	if (WeakObject.IsValid())
+	{
+		return FSlateIconFinder::FindIconForClass(WeakObject->GetClass()).GetIcon();
+	}
+	return FAppStyle::GetBrush(TEXT("Sequencer.InvalidSpawnableIcon"));
+}
+
+FSlateIcon FNavigationToolBinding::GetIcon() const
+{
+	return Icon;
+}
+
+FText FNavigationToolBinding::GetIconTooltipText() const
+{
+	return GetDisplayName();
+}
+
+bool FNavigationToolBinding::IsSelected(const FNavigationToolScopedSelection& InSelection) const
+{
+	return InSelection.IsSelected(Binding.GetObjectGuid());
+}
+
+void FNavigationToolBinding::Select(FNavigationToolScopedSelection& InSelection) const
+{
+	InSelection.Select(Binding.GetObjectGuid());
+}
+
+bool FNavigationToolBinding::CanDelete() const
+{
+	return WeakObject.IsValid();
+}
+
+bool FNavigationToolBinding::Delete()
+{
+	return !WeakObject.IsValid();
+}
+
+bool FNavigationToolBinding::CanRename() const
+{
+	return WeakObject.IsValid();
+}
+
+bool FNavigationToolBinding::Rename(const FString& InName)
+{
+	Binding.SetName(InName);
+
+	return false;
+}
+/*
+EItemSequenceLockState FNavigationToolBinding::GetLockState() const
+{
+	TArray<UMovieSceneSection*> Sections;
+
+	for (UMovieSceneTrack* const Track : Binding.GetTracks())
+	{
+		for (UMovieSceneSection* const Section : Track->GetAllSections())
+		{
+			Sections.Add(Section);
+		}
+	}
+
+	const ENavigationToolCompareState State = CompareArrayState<UMovieSceneSection>(Sections,
+		[](const UMovieSceneSection* const InSection)
+			{
+				return InSection->IsLocked();
+			},
+		[](const UMovieSceneSection* const InSection)
+			{
+				return !InSection->IsLocked();
+			}); 
+
+	return static_cast<EItemSequenceLockState>(State);
+}
+
+void FNavigationToolBinding::SetIsLocked(const bool bInIsLocked)
+{
+	for (UMovieSceneTrack* const Track : Binding.GetTracks())
+	{
+		for (UMovieSceneSection* const Section : Track->GetAllSections())
+		{
+			if (Section->IsLocked() != bInIsLocked)
+			{
+				Section->Modify();
+				Section->SetIsLocked(bInIsLocked);
+			}
+		}
+	}
+
+	for (ISequenceLockableExtension* const LockableItem : GetChildrenOfType<ISequenceLockableExtension>())
+	{
+		LockableItem->SetIsLocked(bInIsLocked);
+	}
+}*/
+
+UObject* FNavigationToolBinding::FindBoundObject() const
+{
+	if (const TSharedPtr<ISequencer> Sequencer = Tool.GetSequencer())
+	{
+		return Sequencer->FindSpawnedObjectOrTemplate(Binding.GetObjectGuid());
+	}
+	return nullptr;
+}
+
+const FMovieSceneBinding& FNavigationToolBinding::GetBinding() const
+{
+	return Binding;
+}
+
+}
+
+#undef LOCTEXT_NAMESPACE
diff --git a/Engine/Source/Editor/Sequencer/Private/Tools/NavigationTool/Items/NavigationToolComponent.cpp b/Engine/Source/Editor/Sequencer/Private/Tools/NavigationTool/Items/NavigationToolComponent.cpp
new file mode 100644
--- /dev/null
+++ b/Engine/Source/Editor/Sequencer/Private/Tools/NavigationTool/Items/NavigationToolComponent.cpp
@@ -0,0 +1,159 @@
+// Copyright Epic Games, Inc. All Rights Reserved.
+
+#include "Tools/NavigationTool/Items/NavigationToolComponent.h"
+#include "Components/PrimitiveComponent.h"
+#include "Components/SceneComponent.h"
+#include "ScopedTransaction.h"
+#include "Styling/StyleColors.h"
+#include "Tools/NavigationTool/Items/NavigationToolItemProxy.h"
+#include "Tools/NavigationTool/NavigationTool.h"
+#include "Tools/NavigationTool/Widgets/Columns/SNavigationToolLabelComponent.h"
+
+#define LOCTEXT_NAMESPACE "NavigationToolComponent"
+
+namespace UE::Sequencer::NavigationTool
+{
+
+FNavigationToolComponent::FNavigationToolComponent(INavigationTool& InTool
+	, const FNavigationToolItemPtr& InParentItem
+	, const FMovieSceneBinding& InBinding
+	, UObject* const InObject
+	, const FNavigationToolItemPtr& InReferencingItem)
+	: Super(InTool
+	, InParentItem
+	, InBinding
+	, InObject
+	, InReferencingItem)
+{
+}
+
+void FNavigationToolComponent::FindChildren(TArray<FNavigationToolItemPtr>& OutChildren, const bool bInRecursive)
+{
+	Super::FindChildren(OutChildren, bInRecursive);
+
+	/*const TSharedRef<FNavigationToolComponent> SharedThisRef = SharedThis(this);
+
+	for (UMovieSceneTrack* const Track : Binding.GetTracks())
+	{
+		const TArray<UMovieSceneSection*>& AllSections = Track->GetAllSections();
+		for (int32 SectionIndex = 0; SectionIndex < AllSections.Num(); ++SectionIndex)
+		{
+			const FNavigationToolItemPtr NewItem = Tool.FindOrAdd<FNavigationToolTrack>(GetProvider(), SharedThisRef, Track, Sequence, AllSections[SectionIndex], SectionIndex);
+			OutChildren.Add(NewItem);
+			if (bInRecursive)
+			{
+				NewItem->FindChildren(OutChildren, bInRecursive);
+			}
+		}
+	}*/
+}
+
+void FNavigationToolComponent::GetItemProxies(TArray<TSharedPtr<FNavigationToolItemProxy>>& OutItemProxies)
+{
+	Super::GetItemProxies(OutItemProxies);
+
+	if (UPrimitiveComponent* const PrimitiveComponent = Cast<UPrimitiveComponent>(GetComponent()))
+	{
+		if (const TSharedPtr<FNavigationToolItemProxy> MaterialItemProxy = Tool.GetOrCreateItemProxy<FNavigationToolItemProxy>(SharedThis(this)))
+		{
+			OutItemProxies.Add(MaterialItemProxy);
+		}
+	}
+}
+
+bool FNavigationToolComponent::IsAllowedInTool() const
+{
+	FNavigationTool& ToolPrivate = static_cast<FNavigationTool&>(Tool);
+
+	// Make sure Owner is Allowed (this also returns false if Owner is invalid)
+	const UActorComponent* const UnderlyingComponent = GetComponent();
+	if (!IsValid(UnderlyingComponent) || !ToolPrivate.IsObjectAllowedInTool(UnderlyingComponent->GetOwner()))
+	{
+		return false;
+	}
+
+	return ToolPrivate.IsObjectAllowedInTool(UnderlyingComponent);
+}
+
+ENavigationToolItemViewMode FNavigationToolComponent::GetSupportedViewModes(const INavigationToolView& InToolView) const
+{
+	return ENavigationToolItemViewMode::ItemTree | ENavigationToolItemViewMode::HorizontalItemList;
+}
+
+TSharedRef<SWidget> FNavigationToolComponent::GenerateLabelWidget(const TSharedRef<SNavigationToolTreeRow>& InRow)
+{
+	return SNew(SNavigationToolLabelComponent, SharedThis(this), InRow);
+}
+
+FLinearColor FNavigationToolComponent::GetItemBrushColor() const
+{
+	return FStyleColors::White25.GetSpecifiedColor();
+}
+
+TArray<FName> FNavigationToolComponent::GetTags() const
+{
+	UActorComponent* const UnderlyingComponent = GetComponent();
+	if (IsValid(UnderlyingComponent))
+	{
+		return UnderlyingComponent->ComponentTags;
+	}
+	return Super::GetTags();
+}
+
+bool FNavigationToolComponent::GetVisibility() const
+{
+	const USceneComponent* const UnderlyingComponent = Cast<USceneComponent>(GetComponent());
+	if (IsValid(UnderlyingComponent))
+	{
+		return UnderlyingComponent->IsVisibleInEditor();
+	}
+	return false;
+}
+
+void FNavigationToolComponent::OnVisibilityChanged(const bool bInNewVisibility)
+{
+	USceneComponent* const UnderlyingComponent = Cast<USceneComponent>(GetComponent());
+	if (IsValid(UnderlyingComponent))
+	{
+		UnderlyingComponent->SetVisibility(bInNewVisibility);
+	}
+}
+
+bool FNavigationToolComponent::CanRename() const
+{
+	return Super::CanRename() && IsValid(GetComponent());
+}
+
+bool FNavigationToolComponent::Rename(const FString& InName)
+{
+	UActorComponent* const UnderlyingComponent = GetComponent();
+	if (!IsValid(UnderlyingComponent))
+	{
+		return false;
+	}
+
+	bool bRenamed = false;
+
+	if (!InName.Equals(UnderlyingComponent->GetName(), ESearchCase::CaseSensitive))
+	{
+		const FScopedTransaction Transaction(LOCTEXT("NavigationToolRenameComponent", "Rename Component"));
+
+		UnderlyingComponent->Modify();
+		UnderlyingComponent->Rename(*InName);
+
+		Super::Rename(InName);
+
+		bRenamed = true;
+	}
+
+	return bRenamed;
+}
+
+UActorComponent* FNavigationToolComponent::GetComponent() const
+{
+	return WeakObject.IsValid() ? Cast<UActorComponent>(WeakObject.Get(IsIgnoringPendingKill())) : nullptr;
+}
+
+}
+
+#undef LOCTEXT_NAMESPACE
diff --git a/Engine/Source/Editor/Sequencer/Private/Tools/NavigationTool/Items/NavigationToolItem.cpp b/Engine/Source/Editor/Sequencer/Private/Tools/NavigationTool/Items/NavigationToolItem.cpp
new file mode 100644
--- /dev/null
+++ b/Engine/Source/Editor/Sequencer/Private/Tools/NavigationTool/Items/NavigationToolItem.cpp
@@ -0,0 +1,501 @@
+// Copyright Epic Games, Inc. All Rights Reserved.
+
+#include "Tools/NavigationTool/Items/NavigationToolItem.h"
+#include "Editor.h"
+#include "Framework/Application/SlateApplication.h"
+#include "Misc/Optional.h"
+#include "Styling/StyleColors.h"
+#include "Textures/SlateIcon.h"
+#include "Tools/NavigationTool/DragDropOps/NavigationToolItemDragDropOp.h"
+#include "Tools/NavigationTool/Items/NavigationToolActor.h"
+#include "Tools/NavigationTool/Items/NavigationToolItemUtils.h"
+#include "Tools/NavigationTool/ItemActions/NavigationToolAddItem.h"
+#include "Tools/NavigationTool/ItemActions/NavigationToolRemoveItem.h"
+#include "Tools/NavigationTool/NavigationTool.h"
+#include "Tools/NavigationTool/NavigationToolView.h"
+#include "Tools/NavigationTool/NavigationToolScopedSelection.h"
+#include "Tools/NavigationTool/Providers/NavigationToolProvider.h"
+#include "Tools/NavigationTool/Widgets/Columns/SNavigationToolLabelItem.h"
+#include "Widgets/Views/STableRow.h"
+
+#define LOCTEXT_NAMESPACE "NavigationToolItem"
+
+namespace UE::Sequencer::NavigationTool
+{
+
+EDetachmentRule GetDetachmentRuleFromAttachmentRule(EAttachmentRule AttachmentRule)
+{
+	switch (AttachmentRule)
+	{
+		case EAttachmentRule::KeepRelative:
+			return EDetachmentRule::KeepRelative;
+
+		default:
+			return EDetachmentRule::KeepWorld;
+	}
+};
+
+FDetachmentTransformRules GetDetachmentRulesFromAttachmentRules(const TOptional<FAttachmentTransformRules>& AttachmentTransformRules)
+{
+	if (AttachmentTransformRules.IsSet())
+	{
+		return FDetachmentTransformRules(
+			GetDetachmentRuleFromAttachmentRule(AttachmentTransformRules->LocationRule),
+			GetDetachmentRuleFromAttachmentRule(AttachmentTransformRules->RotationRule),
+			GetDetachmentRuleFromAttachmentRule(AttachmentTransformRules->ScaleRule),
+			true
+		);
+	}
+	return FDetachmentTransformRules::KeepWorldTransform;
+};
+
+FNavigationToolItem::FNavigationToolItem(INavigationTool& InTool, const FNavigationToolItemPtr& InParentItem)
+	: Tool(InTool)
+	, WeakParent(InParentItem)
+{
+}
+
+INavigationTool& FNavigationToolItem::GetOwnerTool() const
+{
+	return Tool;
+}
+
+TSharedPtr<FNavigationToolProvider> FNavigationToolItem::GetProvider() const
+{
+	return WeakProvider.Pin();
+}
+
+bool FNavigationToolItem::IsItemValid() const
+{
+	return true;
+}
+
+void FNavigationToolItem::RefreshChildren()
+{
+	TArray<FNavigationToolItemPtr> FoundChildren;
+	FindValidChildren(FoundChildren, /*bRecursiveFind=*/false);
+
+	TArray<FNavigationToolItemPtr> Sortable;
+	TArray<FNavigationToolItemPtr> Unsortable;
+	SplitSortableAndUnsortableItems(FoundChildren, Sortable, Unsortable);
+
+	// Start with all Sortable/Unsortable Items, and remove every item seen by iterating Children
+	TSet<FNavigationToolItemPtr> NewSortableChildren(Sortable);
+	TSet<FNavigationToolItemPtr> NewUnsortableChildren(Unsortable);
+	
+	// Remove items from "Children" that were not present in the Sortable Found Children (we'll add non-sortable later)
+	// Result is have Children only contain Items that existed previously
+	for (TArray<FNavigationToolItemPtr>::TIterator ItemIter(Children); ItemIter; ++ItemIter)
+	{
+		const FNavigationToolItemPtr Item(*ItemIter);
+
+		if (!Item.IsValid() || NewUnsortableChildren.Contains(Item))
+		{
+			ItemIter.RemoveCurrent();
+		}
+		else if (!NewSortableChildren.Contains(Item) || !Item->IsItemValid())
+		{
+			Item->SetParent(nullptr);
+			ItemIter.RemoveCurrent();
+		}
+
+		NewSortableChildren.Remove(Item);
+		NewUnsortableChildren.Remove(Item);
+	}
+
+	// Find Children for New Children in case these new children have grand children
+	// Note: This does not affect any of the current containers. It's just called for discovery
+	auto FindGrandChildren = [](const TSet<FNavigationToolItemPtr>& InChildren)
+		{
+			for (const FNavigationToolItemPtr& Child : InChildren)
+			{
+				TArray<FNavigationToolItemPtr> GrandChildren;
+				Child->FindValidChildren(GrandChildren, /*bRecursiveFind=*/true);
+			}
+		};
+
+	FindGrandChildren(NewUnsortableChildren);
+	FindGrandChildren(NewSortableChildren);
+
+	// After removing Children not present in Sortable
+	// Children should either be equal in size with Sortable (which means no new sortable children were added)
+	// or Sortable has more entries which means there are new items to add
+	if (Sortable.Num() > Children.Num())
+	{
+		check(!NewSortableChildren.IsEmpty());
+		HandleNewSortableChildren(NewSortableChildren.Array());
+	}
+
+	// Rearrange so that Children are arranged like so:
+	// [Unsortable Children][Sortable Children]
+	Unsortable.Append(MoveTemp(Children));
+	Children = MoveTemp(Unsortable);
+
+	// Update the Parents of every Child in the List
+	const FNavigationToolItemPtr This = SharedThis(this);
+	for (const FNavigationToolItemPtr& Child : Children)
+	{
+		Child->SetParent(This);
+	}
+}
+
+void FNavigationToolItem::ResetChildren()
+{
+	for (const FNavigationToolItemPtr& Item : GetChildren())
+	{
+		if (Item.IsValid())
+		{
+			Item->SetParent(nullptr);
+		}
+	}
+	GetChildrenMutable().Reset();	
+};
+
+void FNavigationToolItem::FindChildren(TArray<FNavigationToolItemPtr>& OutChildren, const bool bInRecursive)
+{
+	FNavigationTool& ToolPrivate = static_cast<FNavigationTool&>(Tool);
+
+	TArray<TSharedPtr<FNavigationToolItemProxy>> ItemProxies;
+	ToolPrivate.GetItemProxiesForItem(SharedThis(this), ItemProxies);
+	OutChildren.Reserve(OutChildren.Num() + ItemProxies.Num());
+
+	for (const TSharedPtr<FNavigationToolItemProxy>& ItemProxy : ItemProxies)
+	{
+		OutChildren.Add(ItemProxy);
+		if (bInRecursive)
+		{
+			ItemProxy->FindChildren(OutChildren, bInRecursive);
+		}
+	}
+}
+
+void FNavigationToolItem::FindValidChildren(TArray<FNavigationToolItemPtr>& OutChildren, bool bInRecursive)
+{
+	FindChildren(OutChildren, bInRecursive);
+
+	OutChildren.RemoveAll(
+		[](const FNavigationToolItemPtr& InItem)
+		{
+			return !InItem.IsValid() || !InItem->IsAllowedInTool();
+		});
+}
+
+TArray<FNavigationToolItemPtr> FNavigationToolItem::FindPath(const TArray<FNavigationToolItemPtr>& InItems) const
+{
+	const TSharedPtr<const INavigationToolItem> This = SharedThis(this);
+	TArray<FNavigationToolItemPtr> Path;
+	for (const FNavigationToolItemPtr& Item : InItems)
+	{
+		Path.Reset();
+		FNavigationToolItemPtr CurrentItem = Item;
+		while (CurrentItem.IsValid())
+		{
+			if (This == CurrentItem)
+			{
+				Algo::Reverse(Path);
+				return Path;
+			}
+			Path.Add(CurrentItem);
+			CurrentItem = CurrentItem->GetParent();
+		}
+	}
+	return TArray<FNavigationToolItemPtr>();
+}
+
+INavigationToolItem::IndexType FNavigationToolItem::GetChildIndex(const FNavigationToolItemPtr& ChildItem) const
+{
+	return GetChildren().Find(ChildItem);
+}
+
+bool FNavigationToolItem::CanAddChild(const FNavigationToolItemPtr& InChild) const
+{
+	return InChild.IsValid();
+}
+
+bool FNavigationToolItem::AddChild(const FNavigationToolAddItemParams& InAddItemParams)
+{
+	if (CanAddChild(InAddItemParams.Item))
+	{
+		AddChildChecked(InAddItemParams);
+		return true;
+	}
+	return false;
+}
+
+bool FNavigationToolItem::RemoveChild(const FNavigationToolRemoveItemParams& InRemoveItemParams)
+{
+	if (InRemoveItemParams.Item.IsValid())
+	{
+		return RemoveChildChecked(InRemoveItemParams);
+	}
+	return false;
+}
+
+void FNavigationToolItem::SetParent(FNavigationToolItemPtr InParent)
+{
+	//check that one of the parent's children is this
+	WeakParent = InParent;
+}
+
+ENavigationToolItemViewMode FNavigationToolItem::GetSupportedViewModes(const INavigationToolView& InToolView) const
+{
+	return InToolView.GetItemDefaultViewMode();
+}
+
+FNavigationToolItemId FNavigationToolItem::GetItemId() const
+{
+	if (ItemId.IsValid())
+	{
+		return ItemId;
+	}
+	const_cast<FNavigationToolItem*>(this)->RecalculateItemId();
+	return ItemId;
+}
+
+const FSlateBrush* FNavigationToolItem::GetIconBrush() const
+{
+	const FSlateIcon Icon = FNavigationToolExtender::FindOverrideIcon(SharedThis(this));
+	if (Icon.IsSet())
+	{
+		return Icon.GetIcon();
+	}
+
+	if (const FSlateBrush* const DefaultIconBrush = GetDefaultIconBrush())
+	{
+		return DefaultIconBrush;
+	}
+
+	return GetIcon().GetIcon();
+}
+
+TSharedRef<SWidget> FNavigationToolItem::GenerateLabelWidget(const TSharedRef<SNavigationToolTreeRow>& InRow)
+{
+	return SNew(SNavigationToolLabelItem, SharedThis(this), InRow);
+}
+
+bool FNavigationToolItem::Delete()
+{
+	GetOwnerTool().NotifyToolItemDeleted(SharedThis(this));
+	return true;
+}
+
+void FNavigationToolItem::AddFlags(ENavigationToolItemFlags Flags)
+{
+	EnumAddFlags(ItemFlags, Flags);
+}
+
+void FNavigationToolItem::RemoveFlags(ENavigationToolItemFlags Flags)
+{
+	EnumRemoveFlags(ItemFlags, Flags);
+}
+
+bool FNavigationToolItem::HasAnyFlags(ENavigationToolItemFlags Flags) const
+{
+	return EnumHasAnyFlags(ItemFlags, Flags);
+}
+
+bool FNavigationToolItem::HasAllFlags(ENavigationToolItemFlags Flags) const
+{
+	return EnumHasAllFlags(ItemFlags, Flags);
+}
+
+TOptional<EItemDropZone> FNavigationToolItem::CanAcceptDrop(const FDragDropEvent& InDragDropEvent, EItemDropZone InDropZone)
+{
+	if (const TSharedPtr<FNavigationToolItemDragDropOp> ItemDragDropOp = InDragDropEvent.GetOperationAs<FNavigationToolItemDragDropOp>())
+	{
+		const TOptional<EItemDropZone> DropZone = ItemDragDropOp->CanDrop(InDropZone, SharedThis(this));
+		if (DropZone.IsSet())
+		{
+			return *DropZone;
+		}
+	}
+
+	const TSharedRef<FNavigationToolItem> SharedThisRef = SharedThis(this);
+
+	TOptional<EItemDropZone> OutDropZone;
+
+	FNavigationToolExtender::ForEachToolProvider(GetOwnerTool(),
+		[this, &InDragDropEvent, InDropZone, &SharedThisRef, &OutDropZone]
+		(const TSharedRef<FNavigationToolProvider>& InProvider)
+		{
+			OutDropZone = InProvider->OnToolItemCanAcceptDrop(InDragDropEvent, InDropZone, SharedThisRef);
+			if (OutDropZone.IsSet())
+			{
+				return false;
+			}
+			return true;
+		});
+
+	return OutDropZone;
+}
+
+FReply FNavigationToolItem::AcceptDrop(const FDragDropEvent& InDragDropEvent, EItemDropZone InDropZone)
+{
+	if (const TSharedPtr<FNavigationToolItemDragDropOp> ItemDragDropOp = InDragDropEvent.GetOperationAs<FNavigationToolItemDragDropOp>())
+	{
+		const FReply Reply = ItemDragDropOp->Drop(InDropZone, SharedThis(this));
+		if (Reply.IsEventHandled())
+		{
+			return Reply;
+		}
+	}
+
+	const TSharedRef<FNavigationToolItem> SharedThisRef = SharedThis(this);
+
+	FReply OutReply = FReply::Unhandled();
+
+	FNavigationToolExtender::ForEachToolProvider(GetOwnerTool(),
+		[this, &InDragDropEvent, InDropZone, &SharedThisRef, &OutReply]
+		(const TSharedRef<FNavigationToolProvider>& InProvider)
+		{
+			OutReply = InProvider->OnToolItemAcceptDrop(InDragDropEvent, InDropZone, SharedThisRef);
+			if (OutReply.IsEventHandled())
+			{
+				return false;
+			}
+			return true;
+		});
+
+	return OutReply;
+}
+
+FLinearColor FNavigationToolItem::GetItemBrushColor() const
+{
+	return FStyleColors::White.GetSpecifiedColor();
+}
+
+bool FNavigationToolItem::IsExpanded() const
+{
+	if (const TSharedPtr<INavigationToolView> ToolView = Tool.GetMostRecentToolView())
+	{
+		return ToolView->IsItemExpanded(SharedThis(const_cast<FNavigationToolItem*>(this)));
+	}
+	return false;
+}
+
+void FNavigationToolItem::SetExpansion(const bool bInExpand)
+{
+	if (const TSharedPtr<INavigationToolView> ToolView = Tool.GetMostRecentToolView())
+	{
+		return ToolView->SetItemExpansion(SharedThis(this), bInExpand);
+	}
+}
+
+TOptional<FColor> FNavigationToolItem::GetColor() const
+{
+	return Tool.FindItemColor(SharedThis(const_cast<FNavigationToolItem*>(this)));
+}
+
+void FNavigationToolItem::SetColor(const TOptional<FColor>& InColor)
+{
+	Tool.SetItemColor(SharedThis(this), InColor.Get(FColor()));
+}
+
+void FNavigationToolItem::SetProvider(const TWeakPtr<FNavigationToolProvider>& InWeakProvider)
+{
+	WeakProvider = InWeakProvider;
+}
+
+void FNavigationToolItem::RecalculateItemId()
+{
+	const FNavigationToolItemId OldItemId = ItemId;
+	ItemId = CalculateItemId();
+
+	FNavigationTool& ToolPrivate = static_cast<FNavigationTool&>(Tool);
+	ToolPrivate.NotifyItemIdChanged(OldItemId, SharedThis(this));
+}
+
+void FNavigationToolItem::AddChildChecked(const FNavigationToolAddItemParams& InAddItemParams)
+{
+	if (const FNavigationToolItemPtr OldParent = InAddItemParams.Item->GetParent())
+	{
+		// if we are adding the child and the old parent is this,
+		// then it means we're just rearranging, only remove from array
+		if (OldParent.Get() == this)
+		{
+			Children.Remove(InAddItemParams.Item);
+		}
+		else
+		{
+			FNavigationToolRemoveItemParams RemoveParams(InAddItemParams.Item);
+			OldParent->RemoveChild(RemoveParams);
+		}
+	}
+
+	if (InAddItemParams.RelativeItem.IsValid() && InAddItemParams.RelativeDropZone.IsSet() && InAddItemParams.RelativeDropZone != EItemDropZone::OntoItem)
+	{
+		const int32 RelativeItemIndex = Children.Find(InAddItemParams.RelativeItem);
+		if (RelativeItemIndex != INDEX_NONE)
+		{
+			const int32 TargetIndex = InAddItemParams.RelativeDropZone == EItemDropZone::BelowItem
+				? RelativeItemIndex + 1
+				: RelativeItemIndex;
+
+			Children.EmplaceAt(TargetIndex, InAddItemParams.Item);
+		}
+		else
+		{
+			Children.EmplaceAt(0, InAddItemParams.Item);
+		}
+	}
+	else
+	{
+		Children.EmplaceAt(0, InAddItemParams.Item);
+	}
+
+	InAddItemParams.Item->SetParent(SharedThis(this));
+}
+
+bool FNavigationToolItem::RemoveChildChecked(const FNavigationToolRemoveItemParams& InRemoveItemParams)
+{
+	InRemoveItemParams.Item->SetParent(nullptr);
+	return Children.Remove(InRemoveItemParams.Item) > 0;
+}
+
+void FNavigationToolItem::HandleNewSortableChildren(TArray<FNavigationToolItemPtr> InSortableChildren)
+{
+	const TSharedPtr<FNavigationToolProvider> Provider = GetProvider();
+	if (!Provider.IsValid())
+	{
+		return;
+	}
+
+	const FNavigationToolSaveState* const SaveState = Provider->GetSaveState(Tool);
+	if (!SaveState)
+	{
+		return;
+	}
+
+	InSortableChildren.Sort([SaveState](const FNavigationToolItemPtr& A, const FNavigationToolItemPtr& B)
+	{
+		const FNavigationToolSerializedTreeNode* const NodeA = SaveState->SerializedTree.FindTreeNode(FNavigationTool::MakeSerializedItemFromToolItem(A));
+		const FNavigationToolSerializedTreeNode* const NodeB = SaveState->SerializedTree.FindTreeNode(FNavigationTool::MakeSerializedItemFromToolItem(B));
+		return FNavigationToolSerializedTree::CompareTreeItemOrder(NodeA, NodeB);
+	});
+
+	FNavigationToolAddItemParams AddItemParams;
+	for (const FNavigationToolItemPtr& NewChild : InSortableChildren)
+	{
+		AddItemParams.Item = NewChild;
+
+		const FNavigationToolSerializedTreeNode* const TreeNode = SaveState->SerializedTree.FindTreeNode(FNavigationTool::MakeSerializedItemFromToolItem(NewChild));
+		if (TreeNode && Children.IsValidIndex(TreeNode->GetLocalIndex()))
+		{
+			// Add Before the Child at Index, so this Item is at the specific Index
+			AddItemParams.RelativeItem = Children[TreeNode->GetLocalIndex()];
+			AddItemParams.RelativeDropZone = EItemDropZone::AboveItem;
+		}
+		else
+		{
+			// Add After Last, so this Item is the last item in the List
+			AddItemParams.RelativeItem = Children.IsEmpty() ? nullptr : Children.Last();
+			AddItemParams.RelativeDropZone = EItemDropZone::BelowItem;
+		}
+
+		AddChild(AddItemParams);
+	}
+}
+
+}
+
+#undef LOCTEXT_NAMESPACE
diff --git a/Engine/Source/Editor/Sequencer/Private/Tools/NavigationTool/Items/NavigationToolItemId.cpp b/Engine/Source/Editor/Sequencer/Private/Tools/NavigationTool/Items/NavigationToolItemId.cpp
new file mode 100644
--- /dev/null
+++ b/Engine/Source/Editor/Sequencer/Private/Tools/NavigationTool/Items/NavigationToolItemId.cpp
@@ -0,0 +1,134 @@
+﻿// Copyright Epic Games, Inc. All Rights Reserved.
+
+#include "Tools/NavigationTool/Items/NavigationToolItemId.h"
+#include "MovieSceneSequence.h"
+#include "Sections/MovieSceneSubSection.h"
+#include "Tools/NavigationTool/Items/INavigationToolItem.h"
+#include "Tools/NavigationTool/Items/NavigationToolItem.h"
+
+namespace UE::Sequencer::NavigationTool
+{
+
+FNavigationToolItemId FNavigationToolItemId::RootId = FNavigationToolItemId(TEXT("Root"));
+
+FNavigationToolItemId::FNavigationToolItemId(const FString& InUniqueId)
+{
+	Id = InUniqueId;
+	CalculateTypeHash();
+}
+
+FNavigationToolItemId::FNavigationToolItemId(const UObject* InObject)
+{
+	Id = GetObjectPath(InObject);
+	CalculateTypeHash();
+}
+
+FNavigationToolItemId::FNavigationToolItemId(const UObject* InObject
+	, const FNavigationToolItemPtr& InReferencingItem
+	, const FString& InReferencingId)
+	: FNavigationToolItemId(InReferencingItem->GetItemId())
+{
+	AddSeparatedSegment(Id, GetObjectPath(InObject));
+	AddSeparatedSegment(Id, InReferencingId);
+	CalculateTypeHash();
+}
+
+FNavigationToolItemId::FNavigationToolItemId(const FNavigationToolItemPtr& InParentItem
+	, const INavigationToolItemProxyFactory& InItemProxyFactory)
+	: FNavigationToolItemId(InParentItem->GetItemId())
+{
+	AddSeparatedSegment(Id, InItemProxyFactory.GetItemProxyTypeName().ToString());
+	CalculateTypeHash();
+}
+
+FNavigationToolItemId::FNavigationToolItemId(const FNavigationToolItemPtr& InParentItem
+	, const INavigationToolItem& InItemProxy)
+	: FNavigationToolItemId(InParentItem->GetItemId())
+{
+	AddSeparatedSegment(Id, InItemProxy.GetTypeId().ToString());
+	CalculateTypeHash();
+}
+
+FNavigationToolItemId::FNavigationToolItemId(const FNavigationToolItemPtr& InParentItem
+	, const UMovieSceneSequence* const InSequence
+	, const UMovieSceneSubSection* const InSubSection
+	, const int32 InSectionIndex)
+{
+	check(InSequence);
+
+	if (InParentItem.IsValid())
+	{
+		AddSeparatedSegment(Id, InParentItem->GetFullPath());
+	}
+
+	AddSeparatedSegment(Id, GetObjectPath(InSequence));
+
+	if (InSubSection)
+	{
+		AddSeparatedSegment(Id, GetObjectPath(InSubSection));
+		AddSeparatedSegment(Id, FString::FromInt(InSectionIndex));
+	}
+}
+
+FNavigationToolItemId::FNavigationToolItemId(const FNavigationToolItemId& Other)
+{
+	*this = Other;
+}
+
+FNavigationToolItemId::FNavigationToolItemId(FNavigationToolItemId&& Other) noexcept
+{
+	*this = MoveTemp(Other);
+}
+
+FNavigationToolItemId& FNavigationToolItemId::operator=(const FNavigationToolItemId& Other)
+{
+	Id = Other.Id;
+	CalculateTypeHash();
+	return *this;
+}
+
+FNavigationToolItemId& FNavigationToolItemId::operator=(FNavigationToolItemId&& Other) noexcept
+{
+	Swap(*this, Other);
+	CalculateTypeHash();
+	return *this;
+}
+
+bool FNavigationToolItemId::operator==(const FNavigationToolItemId& Other) const
+{
+	return Id == Other.Id && CachedHash == Other.CachedHash;
+}
+
+bool FNavigationToolItemId::IsValid() const
+{
+	return bHasCachedHash;
+}
+
+FString FNavigationToolItemId::GetStringId() const
+{
+	return Id;
+}
+
+FString FNavigationToolItemId::GetObjectPath(const UObject* const InObject)
+{
+	return FSoftObjectPath(InObject).ToString();
+}
+
+void FNavigationToolItemId::AddSeparatedSegment(FString& OutString, const FString& InSegment)
+{
+	if (OutString.IsEmpty())
+	{
+		OutString = InSegment;
+		return;
+	}
+
+	OutString += Separator + InSegment;
+}
+
+void FNavigationToolItemId::CalculateTypeHash()
+{
+	CachedHash = GetTypeHash(Id);
+	bHasCachedHash = true;
+}
+
+}
diff --git a/Engine/Source/Editor/Sequencer/Private/Tools/NavigationTool/Items/NavigationToolItemProxy.cpp b/Engine/Source/Editor/Sequencer/Private/Tools/NavigationTool/Items/NavigationToolItemProxy.cpp
new file mode 100644
--- /dev/null
+++ b/Engine/Source/Editor/Sequencer/Private/Tools/NavigationTool/Items/NavigationToolItemProxy.cpp
@@ -0,0 +1,69 @@
+// Copyright Epic Games, Inc. All Rights Reserved.
+
+#include "Tools/NavigationTool/Items/NavigationToolItemProxy.h"
+#include "Input/Reply.h"
+#include "Tools/NavigationTool/NavigationTool.h"
+#include "Tools/NavigationTool/NavigationToolView.h"
+
+namespace UE::Sequencer::NavigationTool
+{
+
+FNavigationToolItemProxy::FNavigationToolItemProxy(INavigationTool& InTool, const FNavigationToolItemPtr& InParentItem)
+	: FNavigationToolItem(InTool, InParentItem)
+{
+	WeakParent = InParentItem;
+}
+
+bool FNavigationToolItemProxy::IsItemValid() const
+{
+	return FNavigationToolItem::IsItemValid()
+		&& WeakParent.IsValid()
+		&& Tool.FindItem(WeakParent.Pin()->GetItemId()).IsValid();
+}
+
+void FNavigationToolItemProxy::FindChildren(TArray<FNavigationToolItemPtr>& OutChildren, bool bRecursive)
+{
+	const FNavigationToolItemPtr Parent = GetParent();
+	if (!Parent.IsValid() || !Parent->IsAllowedInTool())
+	{
+		return;
+	}
+	
+	Super::FindChildren(OutChildren, bRecursive);
+	
+	GetProxiedItems(Parent.ToSharedRef(), OutChildren, bRecursive);
+}
+
+void FNavigationToolItemProxy::SetParent(FNavigationToolItemPtr InParent)
+{
+	Super::SetParent(InParent);
+
+	// Recalculate our Item Id because we rely on what our parent is for our id
+	RecalculateItemId();
+}
+
+ENavigationToolItemViewMode FNavigationToolItemProxy::GetSupportedViewModes(const INavigationToolView& InToolView) const
+{
+	// Hide Proxies if it has no Children
+	if (Children.IsEmpty())
+	{
+		return ENavigationToolItemViewMode::None;
+	}
+	return InToolView.GetItemProxyViewMode();
+}
+
+bool FNavigationToolItemProxy::CanAutoExpand() const
+{
+	return false;
+}
+
+FNavigationToolItemId FNavigationToolItemProxy::CalculateItemId() const
+{
+	if (WeakParent.IsValid())
+	{
+		return FNavigationToolItemId(WeakParent.Pin(), *this);	
+	}
+	return FNavigationToolItemId();
+}
+
+}
diff --git a/Engine/Source/Editor/Sequencer/Private/Tools/NavigationTool/Items/NavigationToolItemUtils.cpp b/Engine/Source/Editor/Sequencer/Private/Tools/NavigationTool/Items/NavigationToolItemUtils.cpp
new file mode 100644
--- /dev/null
+++ b/Engine/Source/Editor/Sequencer/Private/Tools/NavigationTool/Items/NavigationToolItemUtils.cpp
@@ -0,0 +1,192 @@
+// Copyright Epic Games, Inc. All Rights Reserved.
+
+#include "Tools/NavigationTool/Items/NavigationToolItemUtils.h"
+#include "Containers/Array.h"
+#include "ISourceControlModule.h"
+#include "ISourceControlProvider.h"
+#include "MovieSceneSequence.h"
+#include "ScopedTransaction.h"
+#include "Sections/MovieSceneCinematicShotSection.h"
+#include "Sections/MovieSceneSubSection.h"
+#include "Sequencer.h"
+#include "SequencerUtilities.h"
+#include "Tools/NavigationTool/Items/INavigationToolItem.h"
+#include "Tools/NavigationTool/Items/NavigationToolSequence.h"
+#include "Tools/NavigationTool/Items/NavigationToolTreeRoot.h"
+#include "Tools/NavigationTool/NavigationTool.h"
+#include "Tools/NavigationTool/NavigationToolView.h"
+#include "Tools/NavigationTool/Items/NavigationToolSubTrack.h"
+#include "TrackEditors/SubTrackEditorBase.h"
+#include "Tracks/MovieSceneSubTrack.h"
+
+namespace UE::Sequencer::NavigationTool
+{
+
+bool CompareToolItemOrder(const FNavigationToolItemPtr& A, const FNavigationToolItemPtr& B)
+{
+	if (!A.IsValid())
+	{
+		return false;
+	}
+	if (!B.IsValid())
+	{
+		return false;
+	}
+	if (const FNavigationToolItemPtr LowestCommonAncestor = FNavigationTool::FindLowestCommonAncestor({A, B}))
+	{
+		const TArray<FNavigationToolItemPtr> PathToA = LowestCommonAncestor->FindPath({A});
+		const TArray<FNavigationToolItemPtr> PathToB = LowestCommonAncestor->FindPath({B});
+
+		int32 Index = 0;
+	
+		int32 PathAIndex = -1;
+		int32 PathBIndex = -1;
+
+		while (PathAIndex == PathBIndex)
+		{
+			if (!PathToA.IsValidIndex(Index))
+			{
+				return true;
+			}
+			if (!PathToB.IsValidIndex(Index))
+			{
+				return false;
+			}
+
+			PathAIndex = LowestCommonAncestor->GetChildIndex(PathToA[Index]);
+			PathBIndex = LowestCommonAncestor->GetChildIndex(PathToB[Index]);
+			++Index;
+		}
+		return PathAIndex < PathBIndex;
+	}
+	return false;
+}
+
+void SplitSortableAndUnsortableItems(const TArray<FNavigationToolItemPtr>& InItems
+	, TArray<FNavigationToolItemPtr>& OutSortable
+	, TArray<FNavigationToolItemPtr>& OutUnsortable)
+{
+	// Allocate both for worst case scenarios
+	OutSortable.Reserve(InItems.Num());
+	OutUnsortable.Reserve(InItems.Num());
+
+	for (const FNavigationToolItemPtr& Item : InItems)
+	{
+		if (!Item.IsValid())
+		{
+			continue;
+		}
+
+		if (Item->ShouldSort())
+		{
+			OutSortable.Add(Item);
+		}
+		else
+		{
+			OutUnsortable.Add(Item);
+		}
+	}
+}
+
+UMovieSceneSubSection* GetSequenceItemSubSection(const TSharedPtr<INavigationToolItem>& InItem)
+{
+	if (const FNavigationToolSequence* const SequenceItem = InItem->CastTo<FNavigationToolSequence>())
+	{
+		return SequenceItem->GetSubSection();
+	}
+	return nullptr;
+}
+
+UMovieSceneMetaData* GetSequenceItemMetaData(const TSharedPtr<INavigationToolItem>& InItem)
+{
+	if (const FNavigationToolSequence* const SequenceItem = InItem->CastTo<FNavigationToolSequence>())
+	{
+		return FSubTrackEditorUtil::FindOrAddMetaData(SequenceItem->GetSequence());
+	}
+	return nullptr;
+}
+
+FText SequenceDisplayNameWithModifyAndInactiveStatus(const TSharedPtr<FNavigationToolSequence>& InSequenceItem
+	, const bool bInRemoveParentName
+	, const bool bInIsActive)
+{
+	if (!InSequenceItem.IsValid())
+	{
+		return FText::GetEmpty();
+	}
+
+	UMovieSceneSequence* const Sequence = InSequenceItem->GetSequence();
+	if (!IsValid(Sequence))
+	{
+		return FText::GetEmpty();
+	}
+
+	UMovieSceneCinematicShotSection* const ShotSection = Cast<UMovieSceneCinematicShotSection>(InSequenceItem->GetSubSection());
+
+	const FText OriginalSequenceName = IsValid(ShotSection)
+		? FText::FromString(ShotSection->GetShotDisplayName()) : Sequence->GetDisplayName();
+
+	FText ParentSequenceName = FText::GetEmpty();
+	FText NewSequenceName = OriginalSequenceName;
+
+	// Remove parent name from child name if it contains it
+	if (bInRemoveParentName)
+	{
+		if (const FNavigationToolItemPtr Parent = InSequenceItem->GetParent())
+		{
+			if (!Parent->IsA<FNavigationToolTreeRoot>())
+			{
+				// Go up the hierarchy to find the parent sequence
+				if (const FNavigationToolSequence* const ParentSequenceItem = InSequenceItem->FindAncestorOfType<FNavigationToolSequence>())
+				{
+					UMovieSceneSequence* const ParentSequence = ParentSequenceItem->GetSequence();
+					if (IsValid(ParentSequence))
+					{
+						ParentSequenceName = ParentSequence->GetDisplayName();
+					}
+				}
+
+				// Remove the parent name from children
+				FString NewSequenceNameStr = ParentSequenceName.IsEmpty() ? NewSequenceName.ToString() : ParentSequenceName.ToString();
+				const FString ParentSequenceNameStr = ParentSequenceName.ToString();
+
+				if (NewSequenceNameStr.StartsWith(ParentSequenceNameStr) && NewSequenceNameStr.Len() > ParentSequenceNameStr.Len())
+				{
+					NewSequenceNameStr.RemoveFromStart(ParentSequenceNameStr);
+					ParentSequenceName = FText::FromString(NewSequenceNameStr);
+				}
+			}
+		}
+	}
+
+	FText OutText = NewSequenceName;
+
+	const bool bIsDirty = Sequence->GetPackage()->IsDirty();
+
+	if (bInIsActive)
+	{
+		if (bIsDirty)
+		{
+			OutText = FText::Format(NSLOCTEXT("Sequencer", "ActiveDirtySequenceLabel", "{0}*"), NewSequenceName);
+		}
+	}
+	else
+	{
+		const FText InactiveSequenceLabel = NSLOCTEXT("Sequencer", "InactiveSequenceLabel", "Inactive");
+
+		if (bIsDirty)
+		{
+			OutText = FText::Format(NSLOCTEXT("Sequencer", "InactiveDirtySequenceLabel", "{0}* [{1}]")
+				, NewSequenceName, InactiveSequenceLabel);
+		}
+		else
+		{
+			OutText = FText::Format(NSLOCTEXT("Sequencer", "InactiveDirtySequenceLabel", "{0} [{1}]")
+				, NewSequenceName, InactiveSequenceLabel);
+		}
+	}
+
+	return OutText;
+}
+
+}
diff --git a/Engine/Source/Editor/Sequencer/Private/Tools/NavigationTool/Items/NavigationToolObject.cpp b/Engine/Source/Editor/Sequencer/Private/Tools/NavigationTool/Items/NavigationToolObject.cpp
new file mode 100644
--- /dev/null
+++ b/Engine/Source/Editor/Sequencer/Private/Tools/NavigationTool/Items/NavigationToolObject.cpp
@@ -0,0 +1,151 @@
+// Copyright Epic Games, Inc. All Rights Reserved.
+
+#include "Tools/NavigationTool/Items/NavigationToolObject.h"
+#include "Sequencer.h"
+#include "Tools/NavigationTool/Items/NavigationToolItemId.h"
+#include "Tools/NavigationTool/NavigationTool.h"
+#include "Tools/NavigationTool/NavigationToolScopedSelection.h"
+#include "Styling/SlateIconFinder.h"
+
+#define LOCTEXT_NAMESPACE "NavigationToolObject"
+
+namespace UE::Sequencer::NavigationTool
+{
+
+FNavigationToolObject::FNavigationToolObject(INavigationTool& InTool
+	, const FNavigationToolItemPtr& InParentItem
+	, UObject* const InObject
+	, const FNavigationToolItemWeakPtr& InWeakReferencingItem
+	, const FString& InReferencingId)
+	: FNavigationToolItem(InTool, InParentItem)
+	, WeakObject(InObject)
+	, WeakReferencingItem(InWeakReferencingItem)
+	, ReferencingId(InReferencingId)
+{
+}
+
+bool FNavigationToolObject::IsItemValid() const
+{
+	return FNavigationToolItem::IsItemValid() && WeakObject.IsValid();
+}
+
+void FNavigationToolObject::Select(FNavigationToolScopedSelection& InSelection) const
+{
+	//InSelection.Select(GetObject());
+}
+
+bool FNavigationToolObject::IsSelected(const FNavigationToolScopedSelection& InSelection) const
+{
+	return InSelection.IsSelected(GetObject());
+}
+
+void FNavigationToolObject::OnSelect()
+{
+	Super::OnSelect();
+}
+
+void FNavigationToolObject::OnDoubleClick()
+{
+	Super::OnDoubleClick();
+
+	if (!WeakObject.IsValid())
+	{
+		return;
+	}
+
+	const TSharedPtr<ISequencer> Sequencer = Tool.GetSequencer();
+	if (!Sequencer.IsValid())
+	{
+		return;
+	}
+
+	const FGuid ObjectId = Sequencer->FindObjectId(*WeakObject.Get(), Sequencer->GetFocusedTemplateID());
+	Sequencer->SelectObject(ObjectId);
+}
+
+FText FNavigationToolObject::GetDisplayName() const
+{
+	const UObject* const UnderlyingObject = GetObject();
+	if (IsValid(UnderlyingObject))
+	{
+		return FText::FromString(UnderlyingObject->GetName());
+	}
+	return FText::GetEmpty();
+}
+
+FText FNavigationToolObject::GetClassName() const
+{
+	const UObject* const UnderlyingObject = GetObject();
+	if (IsValid(UnderlyingObject))
+	{
+		return UnderlyingObject->GetClass()->GetDisplayNameText();
+	}
+	return FText::GetEmpty();
+}
+
+FSlateIcon FNavigationToolObject::GetIcon() const
+{
+	const UObject* const UnderlyingObject = GetObject();
+	if (IsValid(UnderlyingObject))
+	{
+		return FSlateIconFinder::FindIconForClass(UnderlyingObject->GetClass());
+	}
+	return FSlateIcon();
+}
+
+FText FNavigationToolObject::GetIconTooltipText() const
+{
+	const UObject* const UnderlyingObject = GetObject();
+	if (IsValid(UnderlyingObject))
+	{
+		return UnderlyingObject->GetClass()->GetDisplayNameText();
+	}
+	return FText::GetEmpty();
+}
+
+bool FNavigationToolObject::IsAllowedInTool() const
+{
+	return IsValid(GetObject());
+}
+
+void FNavigationToolObject::OnObjectsReplaced(const TMap<UObject*, UObject*>& InReplacementMap, bool bRecursive)
+{
+	// Get the Object even if it's Pending Kill (most likely it is)
+	const UObject* ObjectPendingKill = WeakObject.Get(true);
+	if (ObjectPendingKill && InReplacementMap.Contains(ObjectPendingKill))
+	{
+		SetObject(InReplacementMap[ObjectPendingKill]);
+	}
+
+	// This handles calling OnObjectsReplaced for every child item
+	Super::OnObjectsReplaced(InReplacementMap, bRecursive);
+}
+
+UObject* FNavigationToolObject::GetObject() const
+{
+	return WeakObject.Get(IsIgnoringPendingKill());
+}
+
+void FNavigationToolObject::SetObject(UObject* InObject)
+{
+	SetObject_Impl(InObject);
+	RecalculateItemId();
+}
+
+FNavigationToolItemId FNavigationToolObject::CalculateItemId() const
+{
+	if (WeakReferencingItem.IsValid())
+	{
+		return FNavigationToolItemId(GetObject(), WeakReferencingItem.Pin(), ReferencingId);
+	}
+	return FNavigationToolItemId(GetObject());
+}
+
+void FNavigationToolObject::SetObject_Impl(UObject* InObject)
+{
+	WeakObject = InObject;
+}
+
+}
+
+#undef LOCTEXT_NAMESPACE
diff --git a/Engine/Source/Editor/Sequencer/Private/Tools/NavigationTool/Items/NavigationToolObjectReference.cpp b/Engine/Source/Editor/Sequencer/Private/Tools/NavigationTool/Items/NavigationToolObjectReference.cpp
new file mode 100644
--- /dev/null
+++ b/Engine/Source/Editor/Sequencer/Private/Tools/NavigationTool/Items/NavigationToolObjectReference.cpp
@@ -0,0 +1,112 @@
+// Copyright Epic Games, Inc. All Rights Reserved.
+
+#include "Tools/NavigationTool/Items/NavigationToolObjectReference.h"
+#include "Algo/Transform.h"
+#include "Tools/NavigationTool/NavigationTool.h"
+#include "Tools/NavigationTool/ItemActions/NavigationToolRemoveItem.h"
+#include "Tools/NavigationTool/Items/NavigationToolItemId.h"
+
+#define LOCTEXT_NAMESPACE "NavigationToolObjectReference"
+
+namespace UE::Sequencer::NavigationTool
+{
+
+FNavigationToolSharedObject::FNavigationToolSharedObject(INavigationTool& InTool
+	, const FNavigationToolItemPtr& InParentItem
+	, UObject* const InObject)
+	: Super(InTool
+	, InParentItem
+	, InObject)
+{
+}
+
+void FNavigationToolSharedObject::AddReference(const TSharedRef<FNavigationToolObjectReference>& InObjectReference)
+{
+	ObjectReferences.Add(InObjectReference);
+}
+
+void FNavigationToolSharedObject::RemoveReference(const TSharedRef<FNavigationToolObjectReference>& InObjectReference)
+{
+	ObjectReferences.Remove(InObjectReference);
+	ObjectReferences.RemoveAll([](const TWeakPtr<FNavigationToolObjectReference>& InItemWeak){ return !InItemWeak.IsValid(); });
+	
+	if (ObjectReferences.IsEmpty())
+	{
+		FNavigationToolRemoveItemParams Params;
+		Params.Item = SharedThis(this);
+		Tool.UnregisterItem(GetItemId());
+	}
+}
+
+TArray<FNavigationToolItemPtr> FNavigationToolSharedObject::GetObjectReferences() const
+{
+	TArray<FNavigationToolItemPtr> OutObjectReferences;
+	
+	Algo::TransformIf(ObjectReferences, OutObjectReferences
+		, [](const TWeakPtr<FNavigationToolObjectReference>& InItem) { return InItem.IsValid(); }
+		, [](const TWeakPtr<FNavigationToolObjectReference>& InItem) { return InItem.Pin(); });
+	
+	return OutObjectReferences;
+}
+
+FNavigationToolObjectReference::FNavigationToolObjectReference(INavigationTool& InTool
+	, const FNavigationToolItemPtr& InParentItem
+	, UObject* InObject
+	, const FNavigationToolItemPtr& InReferencingItem
+	, const FString& InReferenceId)
+	: Super(InTool, InParentItem, InObject)
+	, ReferencingItemWeak(InReferencingItem)
+	, ReferenceId(InReferenceId)
+{
+}
+
+void FNavigationToolObjectReference::OnItemRegistered()
+{
+	const TSharedRef<FNavigationToolObjectReference> SharedThisRef = SharedThis(this);
+
+	const FNavigationToolItemPtr SharedObject = Tool.FindOrAdd<FNavigationToolSharedObject>(GetProvider(), SharedThisRef, GetObject());
+
+	if (FNavigationToolSharedObject* const SharedObjectItem = SharedObject->CastTo<FNavigationToolSharedObject>())
+	{
+		SharedObjectItem->AddReference(SharedThisRef);
+	}
+}
+
+void FNavigationToolObjectReference::OnItemUnregistered()
+{
+	const TSharedRef<FNavigationToolObjectReference> SharedThisRef = SharedThis(this);
+
+	const FNavigationToolItemPtr SharedObject = Tool.FindOrAdd<FNavigationToolSharedObject>(GetProvider(), SharedThisRef, GetObject());
+
+	if (FNavigationToolSharedObject* const SharedObjectItem = SharedObject->CastTo<FNavigationToolSharedObject>())
+	{
+		SharedObjectItem->RemoveReference(SharedThisRef);
+	}
+}
+
+bool FNavigationToolObjectReference::IsItemValid() const
+{
+	return FNavigationToolObject::IsItemValid()
+		;//&& ReferencingItemWeak.IsValid()
+		//&& ReferencingItemWeak.Pin()->IsItemValid();
+}
+
+FText FNavigationToolObjectReference::GetDisplayName() const
+{
+	return FText::Format(LOCTEXT("DisplayName", "{0} {1}")
+		, FText::FromString(ReferenceId)
+		, FNavigationToolObject::GetDisplayName());
+}
+
+FNavigationToolItemId FNavigationToolObjectReference::CalculateItemId() const
+{
+	if (ReferencingItemWeak.IsValid())
+	{
+		return FNavigationToolItemId(GetObject(), ReferencingItemWeak.Pin(), ReferenceId);	
+	}
+	return FNavigationToolItemId();
+}
+
+}
+
+#undef LOCTEXT_NAMESPACE
diff --git a/Engine/Source/Editor/Sequencer/Private/Tools/NavigationTool/Items/NavigationToolSequence.cpp b/Engine/Source/Editor/Sequencer/Private/Tools/NavigationTool/Items/NavigationToolSequence.cpp
new file mode 100644
--- /dev/null
+++ b/Engine/Source/Editor/Sequencer/Private/Tools/NavigationTool/Items/NavigationToolSequence.cpp
@@ -0,0 +1,497 @@
+// Copyright Epic Games, Inc. All Rights Reserved.
+
+#include "Tools/NavigationTool/Items/NavigationToolSequence.h"
+#include "ISourceControlState.h"
+#include "LevelSequenceActor.h"
+#include "GameFramework/Actor.h"
+#include "MovieScene.h"
+#include "ScopedTransaction.h"
+#include "Sections/MovieSceneCinematicShotSection.h"
+#include "Sections/MovieSceneSubSection.h"
+#include "Sequencer.h"
+#include "SequencerChannelInterface.h"
+#include "Styling/SlateIconFinder.h"
+#include "Tools/NavigationTool/INavigationTool.h"
+#include "Tools/NavigationTool/Items/NavigationToolActor.h"
+#include "Tools/NavigationTool/Items/NavigationToolItemUtils.h"
+#include "Tools/NavigationTool/Utils/NavigationToolMovieSceneUtils.h"
+#include "Tools/NavigationTool/Items/NavigationToolSubTrack.h"
+#include "Tools/NavigationTool/Items/NavigationToolTrack.h"
+#include "Tools/NavigationTool/NavigationToolScopedSelection.h"
+#include "Tools/NavigationTool/Utils/NavigationToolMiscUtils.h"
+#include "Tracks/MovieSceneCinematicShotTrack.h"
+
+#define LOCTEXT_NAMESPACE "NavigationToolSequence"
+
+namespace UE::Sequencer::NavigationTool
+{
+
+FNavigationToolSequence::FNavigationToolSequence(INavigationTool& InTool
+	, const FNavigationToolItemPtr& InParentItem
+	, UMovieSceneSequence* const InSequence
+	, UMovieSceneSubSection* const InSubSection
+    , const int32 InSubSectionIndex)
+	: Super(InTool, InParentItem, InSequence)
+	, WeakSubSection(InSubSection)
+	, SubSectionIndex(InSubSectionIndex)
+{
+}
+
+bool FNavigationToolSequence::IsSelected(const FNavigationToolScopedSelection& InSelection) const
+{
+	return WeakSubSection.IsValid() && InSelection.IsSelected(WeakSubSection.Get());
+}
+
+void FNavigationToolSequence::Select(FNavigationToolScopedSelection& InSelection) const
+{
+	if (WeakSubSection.IsValid())
+	{
+		InSelection.Select(WeakSubSection.Get());
+	}
+}
+
+void FNavigationToolSequence::FindChildren(TArray<FNavigationToolItemPtr>& OutChildren, const bool bInRecursive)
+{
+	Super::FindChildren(OutChildren, bInRecursive);
+
+	UMovieSceneSequence* const Sequence = GetSequence();
+	if (!IsValid(Sequence))
+	{
+		return;
+	}
+
+	UMovieScene* const MovieScene = Sequence->GetMovieScene();
+	if (!IsValid(MovieScene))
+	{
+		return;
+	}
+
+	const TSharedPtr<ISequencer> Sequencer = Tool.GetSequencer();
+	if (!Sequencer.IsValid())
+	{
+		return;
+	}
+
+	const TSharedRef<FNavigationToolSequence> SharedThisRef = SharedThis(this);
+
+	for (UMovieSceneTrack* const Track : MovieScene->GetTracks())
+	{
+		FNavigationToolItemPtr NewItem;
+
+		UMovieSceneSubTrack* const SubTrack = Cast<UMovieSceneSubTrack>(Track);
+		if (IsValid(SubTrack))
+		{
+			NewItem = Tool.FindOrAdd<FNavigationToolSubTrack>(GetProvider(), SharedThisRef, SubTrack, Sequence, WeakSubSection, SubSectionIndex);
+		}
+		else if (IsValid(Track))
+		{
+			NewItem = Tool.FindOrAdd<FNavigationToolTrack>(GetProvider(), SharedThisRef, Track, Sequence, WeakSubSection, SubSectionIndex);
+		}
+
+		if (NewItem.IsValid())
+		{
+			OutChildren.Add(NewItem);
+			if (bInRecursive)
+			{
+				NewItem->FindChildren(OutChildren, bInRecursive);
+			}
+		}
+	}
+
+	TArray<FMovieSceneBinding>& Bindings = MovieScene->GetBindings();
+	Bindings.Sort([](const FMovieSceneBinding& InA, const FMovieSceneBinding& InB)
+		{
+			return InA.GetSortingOrder() < InB.GetSortingOrder();
+		});
+
+	// Only show actor bindings for a Sequence
+	for (const FMovieSceneBinding& Binding : Bindings)
+	{
+		UObject* const BoundObject = Sequencer->FindSpawnedObjectOrTemplate(Binding.GetObjectGuid());
+		if (IsValid(BoundObject))
+		{
+			if (AActor* const Actor = Cast<AActor>(BoundObject))
+			{
+				const FNavigationToolItemPtr NewItem = Tool.FindOrAdd<FNavigationToolActor>(GetProvider(), SharedThisRef, Binding, BoundObject, SharedThisRef);
+				OutChildren.Add(NewItem);
+				if (bInRecursive)
+				{
+					NewItem->FindChildren(OutChildren, bInRecursive);
+				}
+			}
+		}
+	}
+}
+
+void FNavigationToolSequence::GetItemProxies(TArray<TSharedPtr<FNavigationToolItemProxy>>& OutItemProxies)
+{
+	Super::GetItemProxies(OutItemProxies);
+}
+
+bool FNavigationToolSequence::AddChild(const FNavigationToolAddItemParams& InAddItemParams)
+{
+	// @TODO: handle sequence being moved to this item
+	return Super::AddChild(InAddItemParams);
+}
+
+bool FNavigationToolSequence::RemoveChild(const FNavigationToolRemoveItemParams& InRemoveItemParams)
+{
+	/// @TODO: handle sequence being moved to this item
+	return Super::RemoveChild(InRemoveItemParams);
+}
+
+ENavigationToolItemViewMode FNavigationToolSequence::GetSupportedViewModes(const INavigationToolView& InToolView) const
+{
+	// Sequences should only be visualized in Navigation Tool View and not appear in the Item Column List
+	// Support any other type of View Mode
+	return ENavigationToolItemViewMode::ItemTree | ~ENavigationToolItemViewMode::HorizontalItemList;
+}
+
+bool FNavigationToolSequence::IsAllowedInTool() const
+{
+	return GetSequence() != nullptr;
+}
+
+FText FNavigationToolSequence::GetDisplayName() const
+{
+	UMovieSceneCinematicShotSection* const ShotSection = Cast<UMovieSceneCinematicShotSection>(WeakSubSection.Get());
+	if (IsValid(ShotSection))
+	{
+		return FText::FromString(ShotSection->GetShotDisplayName());
+	}
+
+	UMovieSceneSequence* const Sequence = GetSequence();
+	if (IsValid(Sequence))
+	{
+		return Sequence->GetDisplayName();
+	}
+
+	return FText::GetEmpty();
+}
+
+FText FNavigationToolSequence::GetClassName() const
+{
+	UMovieSceneSequence* const Sequence = GetSequence();
+	if (IsValid(Sequence))
+	{
+		return FText::FromString(Sequence->GetClass()->GetName());
+	}
+	return FText::FromString(UMovieSceneSequence::StaticClass()->GetName());
+}
+
+FSlateIcon FNavigationToolSequence::GetIcon() const
+{
+	return FSlateIconFinder::FindIconForClass(ALevelSequenceActor::StaticClass());
+}
+
+FText FNavigationToolSequence::GetIconTooltipText() const
+{
+	return FText::GetEmpty();
+}
+
+UMovieSceneSequence* FNavigationToolSequence::GetSequence() const
+{
+	return Cast<UMovieSceneSequence>(GetObject());
+}
+
+UMovieSceneSubSection* FNavigationToolSequence::GetSubSection() const
+{
+	return WeakSubSection.IsValid() ? WeakSubSection.Get() : nullptr;
+}
+
+UMovieScene* FNavigationToolSequence::GetSequenceMovieScene() const
+{
+	UMovieSceneSequence* const Sequence = GetSequence();
+	if (IsValid(Sequence))
+	{
+		return Sequence->GetMovieScene();
+	}
+	return nullptr;
+}
+
+void FNavigationToolSequence::OnSelect()
+{
+	Super::OnSelect();
+}
+
+void FNavigationToolSequence::OnDoubleClick()
+{
+	// Allow selection of this item
+	Super::OnDoubleClick();
+
+	UMovieSceneSequence* const Sequence = GetSequence();
+	if (IsValid(Sequence))
+	{
+		FocusSequence(Tool, *Sequence);
+	}
+}
+
+FNavigationToolItemId FNavigationToolSequence::CalculateItemId() const
+{
+	return FNavigationToolItemId(GetParent(), GetSequence(), WeakSubSection.Get(), SubSectionIndex);
+}
+
+bool FNavigationToolSequence::CanRename() const
+{
+	const UMovieSceneSequence* const Sequence = GetSequence();
+	if (!IsValid(Sequence))
+	{
+		return false;
+	}
+
+	const UMovieScene* const MovieScene = Sequence->GetMovieScene();
+	if (!IsValid(MovieScene))
+	{
+		return false;
+	}
+
+	return !MovieScene->IsReadOnly();
+}
+
+bool FNavigationToolSequence::Rename(const FString& InName)
+{
+	UMovieSceneSequence* const Sequence = GetSequence();
+	if (!IsValid(Sequence))
+	{
+		return false;
+	}
+
+	bool bRenamed = false;
+
+	if (!InName.Equals(Sequence->GetName(), ESearchCase::CaseSensitive))
+	{
+		//const FScopedTransaction Transaction(LOCTEXT("SequenceRename", "Rename Sequence"));
+
+		if (UMovieSceneCinematicShotSection* const ShotSection = Cast<UMovieSceneCinematicShotSection>(WeakSubSection.Get()))
+		{
+			ShotSection->Modify();
+			ShotSection->SetShotDisplayName(InName);
+
+			Tool.NotifyToolItemRenamed(SharedThis(this));
+
+			bRenamed = true;
+		}
+		else
+		{
+			Sequence->Modify();
+
+			if (Sequence->Rename(*InName))
+			{
+				Tool.NotifyToolItemRenamed(SharedThis(this));
+
+				bRenamed = true;
+			}
+		}
+	}
+
+	return bRenamed;
+}
+
+EItemSequenceInactiveState FNavigationToolSequence::GetInactiveState() const
+{
+	if (WeakSubSection.IsValid())
+	{
+		return WeakSubSection->IsActive() ? EItemSequenceInactiveState::None : EItemSequenceInactiveState::Inactive;
+	}
+
+	const ENavigationToolCompareState State = CompareChildrenItemState<ISequenceInactivableExtension>(this,
+		[](const ISequenceInactivableExtension* const InItem)
+			{
+				return InItem->GetInactiveState() == EItemSequenceInactiveState::Inactive;
+			},
+		[](const ISequenceInactivableExtension* const InItem)
+			{
+				return InItem->GetInactiveState() == EItemSequenceInactiveState::None;
+			});
+
+	return static_cast<EItemSequenceInactiveState>(State);
+}
+
+void FNavigationToolSequence::SetIsInactive(const bool bInIsInactive)
+{
+	const bool bNewActiveState = !bInIsInactive;
+
+	if (WeakSubSection.IsValid())
+	{
+		if (!WeakSubSection->IsReadOnly() && WeakSubSection->IsActive() != bNewActiveState)
+		{
+			WeakSubSection->Modify();
+			WeakSubSection->SetIsActive(bNewActiveState);
+		}
+	}
+
+	for (ISequenceInactivableExtension* const InactivableItem : GetChildrenOfType<ISequenceInactivableExtension>())
+	{
+		InactivableItem->SetIsInactive(bInIsInactive);
+	}
+}
+
+EItemMarkerVisibility FNavigationToolSequence::GetMarkerVisibility() const
+{
+	return IsGloballyMarkedFramesForSequence(GetSequence())
+		? EItemMarkerVisibility::Visible : EItemMarkerVisibility::None;
+}
+
+void FNavigationToolSequence::SetMarkerVisibility(const bool bInVisible)
+{
+	const TSharedPtr<ISequencer> Sequencer = Tool.GetSequencer();
+	if (!Sequencer.IsValid())
+	{
+		return;
+	}
+
+	const bool bVisible = (GetMarkerVisibility() == EItemMarkerVisibility::Visible);
+	if (bVisible == bInVisible)
+	{
+		return;
+	}
+
+	UMovieSceneSequence* const Sequence = GetSequence();
+	if (Sequence == Sequencer->GetRootMovieSceneSequence())
+	{
+		Sequencer->GetSequencerSettings()->SetShowMarkedFrames(bInVisible);
+	}
+
+	ShowGloballyMarkedFramesForSequence(*Sequencer, Sequence, bInVisible);
+
+	for (IMarkerVisibilityExtension* const MarkerVisibilityItem : GetChildrenOfType<IMarkerVisibilityExtension>())
+	{
+		MarkerVisibilityItem->SetMarkerVisibility(bInVisible);
+	}
+}
+
+EItemSequenceLockState FNavigationToolSequence::GetLockState() const
+{
+	if (WeakSubSection.IsValid())
+	{
+		return WeakSubSection->IsLocked() ? EItemSequenceLockState::Locked : EItemSequenceLockState::None;
+	}
+
+	const ENavigationToolCompareState State = CompareChildrenItemState<ISequenceLockableExtension>(this,
+		[](const ISequenceLockableExtension* const InItem)
+			{
+				return InItem->GetLockState() == EItemSequenceLockState::Locked;
+			},
+		[](const ISequenceLockableExtension* const InItem)
+			{
+				return InItem->GetLockState() != EItemSequenceLockState::Locked;
+			});
+
+	return static_cast<EItemSequenceLockState>(State);
+}
+
+void FNavigationToolSequence::SetIsLocked(const bool bInIsLocked)
+{
+	if (WeakSubSection.IsValid())
+	{
+		if (WeakSubSection->IsLocked() != bInIsLocked)
+		{
+			WeakSubSection->Modify();
+			WeakSubSection->SetIsLocked(bInIsLocked);
+		}
+
+		return;
+	}
+
+	for (ISequenceLockableExtension* const LockableItem : GetChildrenOfType<ISequenceLockableExtension>())
+	{
+		LockableItem->SetIsLocked(bInIsLocked);
+	}
+}
+
+TOptional<FColor> FNavigationToolSequence::GetColor() const
+{
+	if (!WeakSubSection.IsValid())
+	{
+		return FColor::Transparent;
+	}
+
+	const TSharedPtr<ISequencer> Sequencer = Tool.GetSequencer();
+	if (!Sequencer.IsValid())
+	{
+		return FColor::Transparent;
+	}
+
+	const FColor TrackColor = WeakSubSection->GetColorTint();
+	return (TrackColor == FColor()) ? TOptional<FColor>() : WeakSubSection->GetColorTint();
+}
+
+void FNavigationToolSequence::SetColor(const TOptional<FColor>& InColor)
+{
+	if (!WeakSubSection.IsValid())
+	{
+		return;
+	}
+
+	const TSharedPtr<ISequencer> Sequencer = Tool.GetSequencer();
+	if (!Sequencer.IsValid())
+	{
+		return;
+	}
+
+	WeakSubSection->Modify();
+	WeakSubSection->SetColorTint(InColor.Get(FColor()));
+}
+
+EItemRevisionControlState FNavigationToolSequence::GetRevisionControlState() const
+{
+	if (!WeakSubSection.IsValid())
+	{
+		return EItemRevisionControlState::None;
+	}
+
+	const TSharedPtr<ISequencer> Sequencer = Tool.GetSequencer();
+	if (!Sequencer.IsValid())
+	{
+		return EItemRevisionControlState::None;
+	}
+
+	const FSourceControlStatePtr RevisionControlState = GetSourceControlState(WeakSubSection->GetPackage());
+	if (!RevisionControlState.IsValid())
+	{
+		return EItemRevisionControlState::None;
+	}
+
+	return RevisionControlState->IsCheckedOut() ? EItemRevisionControlState::CheckedOut : EItemRevisionControlState::None;
+}
+
+void FNavigationToolSequence::SetIsRevisionControlled(const bool bInIsRevisionControlled)
+{
+	// Revision Control is only used to view state.
+}
+
+const FSlateBrush* FNavigationToolSequence::GetRevisionControlStatusIcon() const
+{
+	if (!WeakSubSection.IsValid())
+	{
+		return nullptr;
+	}
+
+	const TSharedPtr<ISequencer> Sequencer = Tool.GetSequencer();
+	if (!Sequencer.IsValid())
+	{
+		return nullptr;
+	}
+
+	return GetSourceControlStatusBrush(WeakSubSection->GetPackage());
+}
+
+FText FNavigationToolSequence::GetRevisionControlStatusText() const
+{
+	if (!WeakSubSection.IsValid())
+	{
+		return FText::GetEmpty();
+	}
+
+	const TSharedPtr<ISequencer> Sequencer = Tool.GetSequencer();
+	if (!Sequencer.IsValid())
+	{
+		return FText::GetEmpty();
+	}
+
+	return GetSourceControlStatusText(WeakSubSection->GetPackage());
+}
+
+}
+
+#undef LOCTEXT_NAMESPACE
diff --git a/Engine/Source/Editor/Sequencer/Private/Tools/NavigationTool/Items/NavigationToolSubTrack.cpp b/Engine/Source/Editor/Sequencer/Private/Tools/NavigationTool/Items/NavigationToolSubTrack.cpp
new file mode 100644
--- /dev/null
+++ b/Engine/Source/Editor/Sequencer/Private/Tools/NavigationTool/Items/NavigationToolSubTrack.cpp
@@ -0,0 +1,166 @@
+﻿// Copyright Epic Games, Inc. All Rights Reserved.
+
+#include "Tools/NavigationTool/Items/NavigationToolSubTrack.h"
+#include "MovieSceneSequence.h"
+#include "MovieSceneTrack.h"
+#include "Sections/MovieSceneSubSection.h"
+#include "Tools/NavigationTool/INavigationTool.h"
+#include "Tools/NavigationTool/Items/NavigationToolItemUtils.h"
+#include "Tools/NavigationTool/Items/NavigationToolSequence.h"
+#include "Tools/NavigationTool/Utils/NavigationToolMovieSceneUtils.h"
+#include "Tracks/MovieSceneSubTrack.h"
+#include "Widgets/Views/STableRow.h"
+
+#define LOCTEXT_NAMESPACE "NavigationToolSubTrack"
+
+namespace UE::Sequencer::NavigationTool
+{
+
+using namespace UE::Sequencer;
+
+FNavigationToolSubTrack::FNavigationToolSubTrack(INavigationTool& InTool
+	, const FNavigationToolItemPtr& InParentItem
+	, UMovieSceneSubTrack* const InSubTrack
+	, const TWeakObjectPtr<UMovieSceneSequence>& InSequence
+	, const TWeakObjectPtr<UMovieSceneSection>& InSection
+	, const int32 InSubSectionIndex)
+	: Super(InTool
+	, InParentItem
+	, InSubTrack
+	, InSequence
+	, InSection
+	, InSubSectionIndex)
+{
+	OnTrackObjectChanged();
+}
+
+void FNavigationToolSubTrack::FindChildren(TArray<FNavigationToolItemPtr>& OutChildren, const bool bInRecursive)
+{
+	Super::FindChildren(OutChildren, bInRecursive);
+
+	UMovieSceneSubTrack* const SubTrack = GetSubTrack();
+	if (IsValid(SubTrack))
+	{
+		const TSharedRef<FNavigationToolSubTrack> SharedThisRef = SharedThis(this);
+
+		const TArray<UMovieSceneSection*>& AllSections = SubTrack->GetAllSections();
+		for (int32 SectionIndex = 0; SectionIndex < AllSections.Num(); ++SectionIndex)
+		{
+			UMovieSceneSubSection* const SubSection = Cast<UMovieSceneSubSection>(AllSections[SectionIndex]);
+			if (IsValid(SubSection))
+			{
+				UMovieSceneSequence* const Sequence = SubSection->GetSequence();
+				if (ensure(IsValid(Sequence)))
+				{
+					const FNavigationToolItemPtr NewItem = Tool.FindOrAdd<FNavigationToolSequence>(GetProvider(), SharedThisRef, Sequence, SubSection, SectionIndex);
+					OutChildren.Add(NewItem);
+					if (bInRecursive)
+					{
+						NewItem->FindChildren(OutChildren, bInRecursive);
+					}
+				}
+			}
+		}
+	}
+}
+
+TOptional<EItemDropZone> FNavigationToolSubTrack::CanAcceptDrop(const FDragDropEvent& InDragDropEvent, EItemDropZone InDropZone)
+{
+	return Super::CanAcceptDrop(InDragDropEvent, InDropZone);
+}
+
+FReply FNavigationToolSubTrack::AcceptDrop(const FDragDropEvent& InDragDropEvent, EItemDropZone InDropZone)
+{
+	return Super::AcceptDrop(InDragDropEvent, InDropZone);
+}
+
+UMovieSceneSubTrack* FNavigationToolSubTrack::GetSubTrack() const
+{
+	return Cast<UMovieSceneSubTrack>(GetTrack());
+}
+
+EItemSequenceInactiveState FNavigationToolSubTrack::GetInactiveState() const
+{
+	UMovieSceneSubSection* const ThisSubSection = Cast<UMovieSceneSubSection>(WeakSection.Get());
+	if (IsValid(ThisSubSection))
+	{
+		return ThisSubSection->IsActive() ? EItemSequenceInactiveState::None : EItemSequenceInactiveState::Inactive;
+	}
+
+	const ENavigationToolCompareState State = CompareChildrenItemState<ISequenceInactivableExtension>(this,
+		[](const ISequenceInactivableExtension* const InInactivableItem)
+			{
+				return InInactivableItem->GetInactiveState() == EItemSequenceInactiveState::Inactive;
+			},
+		[](const ISequenceInactivableExtension* const InInactivableItem)
+			{
+				return InInactivableItem->GetInactiveState() == EItemSequenceInactiveState::None;
+			});
+
+	return static_cast<EItemSequenceInactiveState>(State);
+}
+
+void FNavigationToolSubTrack::SetIsInactive(const bool bInIsInactive)
+{
+	const bool bNewActiveState = !bInIsInactive;
+
+	UMovieSceneSubSection* const ThisSubSection = Cast<UMovieSceneSubSection>(WeakSection.Get());
+	if (IsValid(ThisSubSection))
+	{
+		if (ThisSubSection->IsActive() != bNewActiveState)
+		{
+			ThisSubSection->Modify();
+			ThisSubSection->SetIsActive(bNewActiveState);
+		}
+	}
+
+	for (ISequenceInactivableExtension* const InactivableItem : GetChildrenOfType<ISequenceInactivableExtension>())
+	{
+		InactivableItem->SetIsInactive(bInIsInactive);
+	}
+}
+
+EItemMarkerVisibility FNavigationToolSubTrack::GetMarkerVisibility() const
+{
+	const ENavigationToolCompareState State = CompareChildrenItemState<IMarkerVisibilityExtension>(this,
+		[](const IMarkerVisibilityExtension* const InInactivableItem)
+			{
+				return InInactivableItem->GetMarkerVisibility() == EItemMarkerVisibility::Visible;
+			},
+		[](const IMarkerVisibilityExtension* const InInactivableItem)
+			{
+				return InInactivableItem->GetMarkerVisibility() == EItemMarkerVisibility::None;
+			});
+
+	return static_cast<EItemMarkerVisibility>(State);
+}
+
+void FNavigationToolSubTrack::SetMarkerVisibility(const bool bInVisible)
+{
+	const TSharedPtr<ISequencer> Sequencer = Tool.GetSequencer();
+	if (!Sequencer.IsValid())
+	{
+		return;
+	}
+
+	UMovieSceneSubSection* const ThisSubSection = Cast<UMovieSceneSubSection>(WeakSection.Get());
+	if (IsValid(ThisSubSection))
+	{
+		UMovieSceneSequence* const Sequence = ThisSubSection->GetSequence();
+		const bool bIsVisible = IsGloballyMarkedFramesForSequence(Sequence);
+		if (bIsVisible != bInVisible)
+		{
+			ModifySequenceAndMovieScene(Sequence);
+			ShowGloballyMarkedFramesForSequence(*Sequencer, Sequence, bInVisible);
+		}
+	}
+
+	for (IMarkerVisibilityExtension* const MarkerVisibilityItem : GetChildrenOfType<IMarkerVisibilityExtension>())
+	{
+		MarkerVisibilityItem->SetMarkerVisibility(bInVisible);
+	}
+}
+
+}
+
+#undef LOCTEXT_NAMESPACE
diff --git a/Engine/Source/Editor/Sequencer/Private/Tools/NavigationTool/Items/NavigationToolTrack.cpp b/Engine/Source/Editor/Sequencer/Private/Tools/NavigationTool/Items/NavigationToolTrack.cpp
new file mode 100644
--- /dev/null
+++ b/Engine/Source/Editor/Sequencer/Private/Tools/NavigationTool/Items/NavigationToolTrack.cpp
@@ -0,0 +1,204 @@
+﻿// Copyright Epic Games, Inc. All Rights Reserved.
+
+#include "Tools/NavigationTool/Items/NavigationToolTrack.h"
+#include "ISequencer.h"
+#include "ISequencerTrackEditor.h"
+#include "MovieSceneSequence.h"
+#include "MovieSceneTrack.h"
+#include "ScopedTransaction.h"
+#include "Styling/SlateIconFinder.h"
+#include "Tools/NavigationTool/INavigationTool.h"
+#include "Tools/NavigationTool/Items/NavigationToolItemUtils.h"
+#include "Tools/NavigationTool/Items/NavigationToolSequence.h"
+#include "Tools/NavigationTool/NavigationToolScopedSelection.h"
+#include "Widgets/Views/STableRow.h"
+
+#define LOCTEXT_NAMESPACE "NavigationToolTrack"
+
+namespace UE::Sequencer::NavigationTool
+{
+
+FNavigationToolTrack::FNavigationToolTrack(INavigationTool& InTool
+	, const FNavigationToolItemPtr& InParentItem
+	, UMovieSceneTrack* const InTrack
+	, const TWeakObjectPtr<UMovieSceneSequence>& InSequence
+	, const TWeakObjectPtr<UMovieSceneSection>& InSection
+	, const int32 InSubSectionIndex)
+	: Super(InTool, InParentItem, InTrack)
+	, WeakSequence(InSequence)
+	, WeakSection(InSection)
+	, SubSectionIndex(InSubSectionIndex)
+{
+	OnTrackObjectChanged();
+}
+
+FNavigationToolItemId FNavigationToolTrack::CalculateItemId() const
+{
+	UMovieSceneTrack* const Track = GetTrack();
+	check(IsValid(Track));
+
+	UMovieSceneSequence* const Sequence = WeakSequence.Get();
+	check(IsValid(Sequence));
+
+	const FNavigationToolItemPtr ParentItem = GetParent();
+
+	FString NewId = ParentItem.IsValid() ? ParentItem->GetFullPath() : FString();
+
+	if (WeakSection.IsValid())
+	{
+		NewId += FSoftObjectPath(WeakSection.Get()).ToString()
+			+ FNavigationToolItemId::Separator
+			+ FString::FromInt(SubSectionIndex)
+			+ FNavigationToolItemId::Separator;
+	}
+
+	NewId += /*FSoftObjectPath(Sequence).ToString()
+		+ FNavigationToolItemId::Separator
+		+ */FSoftObjectPath(Track).ToString();
+
+	return FNavigationToolItemId(NewId);
+}
+
+void FNavigationToolTrack::FindChildren(TArray<FNavigationToolItemPtr>& OutChildren, const bool bInRecursive)
+{
+	Super::FindChildren(OutChildren, bInRecursive);
+}
+
+TOptional<EItemDropZone> FNavigationToolTrack::CanAcceptDrop(const FDragDropEvent& InDragDropEvent, EItemDropZone InDropZone)
+{
+	return Super::CanAcceptDrop(InDragDropEvent, InDropZone);
+}
+
+FReply FNavigationToolTrack::AcceptDrop(const FDragDropEvent& InDragDropEvent, EItemDropZone InDropZone)
+{
+	return Super::AcceptDrop(InDragDropEvent, InDropZone);
+}
+
+FText FNavigationToolTrack::GetDisplayName() const
+{
+	if (const UMovieSceneTrack* const Track = GetTrack())
+	{
+		return Track->GetDisplayName();
+	}
+	return Super::GetDisplayName();
+}
+
+const FSlateBrush* FNavigationToolTrack::GetDefaultIconBrush() const
+{
+	UMovieSceneTrack* const Track = GetTrack();
+	if (!IsValid(Track))
+	{
+		return nullptr;
+	}
+
+	if (const TSharedPtr<ISequencer> Sequencer = Tool.GetSequencer())
+	{
+		if (const TSharedPtr<ISequencerTrackEditor> TrackEditor = Sequencer->GetTrackEditor(Track))
+		{
+			if (const FSlateBrush* const IconBrush = TrackEditor->GetIconBrush())
+			{
+				return IconBrush;
+			}
+		}
+	}
+
+	return FSlateIconFinder::FindIconForClass(UMovieSceneTrack::StaticClass()).GetIcon();
+}
+
+FSlateIcon FNavigationToolTrack::GetIcon() const
+{
+	return Icon;
+}
+
+bool FNavigationToolTrack::IsSelected(const FNavigationToolScopedSelection& InSelection) const
+{
+	UMovieSceneTrack* const Track = GetTrack();
+	return IsValid(Track) && InSelection.IsSelected(Track);
+}
+
+void FNavigationToolTrack::Select(FNavigationToolScopedSelection& InSelection) const
+{
+	UMovieSceneTrack* const Track = GetTrack();
+	if (IsValid(Track))
+	{
+		InSelection.Select(Track);
+	}
+}
+
+void FNavigationToolTrack::SetObject_Impl(UObject* InObject)
+{
+	Super::SetObject_Impl(InObject);
+
+	OnTrackObjectChanged();
+}
+
+EItemSequenceLockState FNavigationToolTrack::GetLockState() const
+{
+	if (WeakSection.IsValid())
+	{
+		return WeakSection->IsLocked() ? EItemSequenceLockState::Locked : EItemSequenceLockState::None;
+	}
+
+	const ENavigationToolCompareState State = CompareChildrenItemState<ISequenceLockableExtension>(this,
+		[](const ISequenceLockableExtension* const InLockableItem)
+			{
+				return InLockableItem->GetLockState() == EItemSequenceLockState::Locked;
+			},
+		[](const ISequenceLockableExtension* const InLockableItem)
+			{
+				return InLockableItem->GetLockState() != EItemSequenceLockState::Locked;
+			});
+
+	return static_cast<EItemSequenceLockState>(State);
+}
+
+void FNavigationToolTrack::SetIsLocked(const bool bInIsLocked)
+{
+	if (WeakSection.IsValid())
+	{
+		if (WeakSection->IsLocked() != bInIsLocked)
+		{
+			WeakSection->Modify();
+			WeakSection->SetIsLocked(bInIsLocked);
+		}
+	}
+
+	for (ISequenceLockableExtension* const LockableItem : GetChildrenOfType<ISequenceLockableExtension>())
+	{
+		LockableItem->SetIsLocked(bInIsLocked);
+	}
+}
+
+void FNavigationToolTrack::OnTrackObjectChanged()
+{
+	UMovieSceneTrack* const Track = GetTrack();
+	if (!IsValid(Track))
+	{
+		return;
+	}
+
+	const TSharedPtr<ISequencer> Sequencer = Tool.GetSequencer();
+	if (!Sequencer.IsValid())
+	{
+		return;
+	}
+
+	// Refresh icon from the track editor
+	const TSharedPtr<ISequencerTrackEditor> TrackEditor = Sequencer->GetTrackEditor(Track);
+	if (TrackEditor.IsValid())
+	{
+		if (const FSlateBrush* const IconBrush = TrackEditor->GetIconBrush())
+		{
+			Icon = FSlateIcon(FAppStyle::GetAppStyleSetName(), IconBrush->GetResourceName());
+		}
+	}
+}
+
+UMovieSceneTrack* FNavigationToolTrack::GetTrack() const
+{
+	return Cast<UMovieSceneTrack>(GetObject());
+}
+
+}
+
+#undef LOCTEXT_NAMESPACE
diff --git a/Engine/Source/Editor/Sequencer/Private/Tools/NavigationTool/Items/NavigationToolTreeRoot.cpp b/Engine/Source/Editor/Sequencer/Private/Tools/NavigationTool/Items/NavigationToolTreeRoot.cpp
new file mode 100644
--- /dev/null
+++ b/Engine/Source/Editor/Sequencer/Private/Tools/NavigationTool/Items/NavigationToolTreeRoot.cpp
@@ -0,0 +1,96 @@
+// Copyright Epic Games, Inc. All Rights Reserved.
+
+#include "NavigationToolTreeRoot.h"
+#include "Textures/SlateIcon.h"
+#include "Tools/NavigationTool/NavigationTool.h"
+#include "Tools/NavigationTool/Items/NavigationToolSequence.h"
+#include "Widgets/SNullWidget.h"
+#include "Widgets/Views/STableRow.h"
+
+namespace UE::Sequencer::NavigationTool
+{
+
+FNavigationToolTreeRoot::FNavigationToolTreeRoot(INavigationTool& InTool)
+	: FNavigationToolItem(InTool, nullptr)
+{
+}
+
+void FNavigationToolTreeRoot::FindChildren(TArray<FNavigationToolItemPtr>& OutChildren, const bool bInRecursive)
+{
+	Super::FindChildren(OutChildren, bInRecursive);
+
+	const TSharedRef<FNavigationToolTreeRoot> SharedThisRef = SharedThis(this);
+
+	// Get items to add from providers
+	FNavigationToolExtender::ForEachToolProvider(GetOwnerTool(), [this, &OutChildren, bInRecursive, &SharedThisRef]
+		(const TSharedRef<FNavigationToolProvider>& InToolProvider)
+		{
+			TArray<FNavigationToolItemPtr> ThisProviderItems;
+			InToolProvider->OnExtendItemChildren(Tool, SharedThisRef, ThisProviderItems, bInRecursive);
+
+			OutChildren.Append(ThisProviderItems);
+
+			return true;
+		});
+}
+
+bool FNavigationToolTreeRoot::CanAddChild(const FNavigationToolItemPtr& InChild) const
+{
+	return Super::CanAddChild(InChild) && InChild->CanBeTopLevel();
+}
+
+bool FNavigationToolTreeRoot::IsAllowedInTool() const
+{
+	checkNoEntry();
+	return false;
+}
+
+FText FNavigationToolTreeRoot::GetDisplayName() const
+{
+	checkNoEntry();
+	return FText::GetEmpty();
+}
+
+FText FNavigationToolTreeRoot::GetClassName() const
+{
+	checkNoEntry();
+	return FText::GetEmpty();
+}
+
+FText FNavigationToolTreeRoot::GetIconTooltipText() const
+{
+	checkNoEntry();
+	return FText::GetEmpty();
+}
+
+FSlateIcon FNavigationToolTreeRoot::GetIcon() const
+{
+	checkNoEntry();
+	return FSlateIcon();
+}
+
+TSharedRef<SWidget> FNavigationToolTreeRoot::GenerateLabelWidget(const TSharedRef<SNavigationToolTreeRow>& InRow)
+{
+	checkNoEntry();
+	return SNullWidget::NullWidget;
+}
+
+TOptional<EItemDropZone> FNavigationToolTreeRoot::CanAcceptDrop(const FDragDropEvent& InDragDropEvent, EItemDropZone InDropZone)
+{
+	return Super::CanAcceptDrop(InDragDropEvent, InDropZone);
+}
+
+FReply FNavigationToolTreeRoot::AcceptDrop(const FDragDropEvent& InDragDropEvent, EItemDropZone InDropZone)
+{
+	// Force Tree Root to always be Onto Item so we always create new Items as Children of Root
+	const EItemDropZone ForcedDropZone = EItemDropZone::OntoItem;
+
+	return Super::AcceptDrop(InDragDropEvent, ForcedDropZone);
+}
+
+FNavigationToolItemId FNavigationToolTreeRoot::CalculateItemId() const
+{
+	return FNavigationToolItemId::RootId;
+}
+
+}
diff --git a/Engine/Source/Editor/Sequencer/Private/Tools/NavigationTool/Items/NavigationToolTreeRoot.h b/Engine/Source/Editor/Sequencer/Private/Tools/NavigationTool/Items/NavigationToolTreeRoot.h
new file mode 100644
--- /dev/null
+++ b/Engine/Source/Editor/Sequencer/Private/Tools/NavigationTool/Items/NavigationToolTreeRoot.h
@@ -0,0 +1,43 @@
+// Copyright Epic Games, Inc. All Rights Reserved.
+
+#pragma once
+
+#include "Tools/NavigationTool/Items/NavigationToolItem.h"
+
+namespace UE::Sequencer::NavigationTool
+{
+
+class INavigationTool;
+
+/*
+ * The Item that ensures that every item (except self) has a parent in the hierarchy to make it easier
+ * to handle. This is not really a visual item, so it can't appear in the Navigation Tool View
+ */
+class FNavigationToolTreeRoot final : public FNavigationToolItem
+{
+public:
+	UE_NAVIGATIONTOOL_INHERITS_WITH_SUPER(FNavigationToolTreeRoot, FNavigationToolItem)
+
+	FNavigationToolTreeRoot(INavigationTool& InTool);
+	virtual ~FNavigationToolTreeRoot() override = default;
+
+	//~ Begin INavigationToolItem
+	virtual void FindChildren(TArray<FNavigationToolItemPtr>& OutChildren, const bool bInRecursive) override;
+	virtual bool CanAddChild(const FNavigationToolItemPtr& InChild) const override;
+	virtual bool IsAllowedInTool() const override;
+	virtual FText GetDisplayName() const override;
+	virtual FText GetClassName() const override;
+	virtual FText GetIconTooltipText() const override;
+	virtual FSlateIcon GetIcon() const override;
+	virtual TSharedRef<SWidget> GenerateLabelWidget(const TSharedRef<SNavigationToolTreeRow>& InRow) override;
+	virtual TOptional<EItemDropZone> CanAcceptDrop(const FDragDropEvent& InDragDropEvent, EItemDropZone InDropZone) override;	
+	virtual FReply AcceptDrop(const FDragDropEvent& InDragDropEvent, EItemDropZone InDropZone) override;
+	//~ End INavigationToolItem
+
+protected:
+	//~ Begin FNavigationToolItem
+	virtual FNavigationToolItemId CalculateItemId() const override;
+	//~ End FNavigationToolItem
+};
+
+}
diff --git a/Engine/Source/Editor/Sequencer/Private/Tools/NavigationTool/Menus/NavigationToolItemContextMenu.cpp b/Engine/Source/Editor/Sequencer/Private/Tools/NavigationTool/Menus/NavigationToolItemContextMenu.cpp
new file mode 100644
--- /dev/null
+++ b/Engine/Source/Editor/Sequencer/Private/Tools/NavigationTool/Menus/NavigationToolItemContextMenu.cpp
@@ -0,0 +1,112 @@
+﻿// Copyright Epic Games, Inc. All Rights Reserved.
+
+#include "NavigationToolItemContextMenu.h"
+#include "Framework/Commands/GenericCommands.h"
+#include "ToolMenu.h"
+#include "ToolMenus.h"
+#include "Tools/NavigationTool/Menus/NavigationToolItemMenuContext.h"
+#include "Tools/NavigationTool/NavigationTool.h"
+#include "Tools/NavigationTool/NavigationToolCommands.h"
+#include "Tools/NavigationTool/NavigationToolView.h"
+
+#define LOCTEXT_NAMESPACE "NavigationToolItemContextMenu"
+
+namespace UE::Sequencer::NavigationTool
+{
+
+FName FNavigationToolItemContextMenu::GetMenuName()
+{
+	static const FName MenuName(TEXT("Sequencer.NavigationTool.ItemContextMenu"));
+	return MenuName;
+}
+
+TSharedRef<SWidget> FNavigationToolItemContextMenu::CreateMenu(const TSharedRef<FNavigationToolView>& InToolView, const TArray<FNavigationToolItemPtr>& InItemList)
+{
+	UToolMenus* const ToolMenus = UToolMenus::Get();
+
+	const FName MenuName = GetMenuName();
+
+	if (!ToolMenus->IsMenuRegistered(MenuName))
+	{
+		UToolMenu* const ContextMenu = ToolMenus->RegisterMenu(MenuName, NAME_None, EMultiBoxType::Menu);
+		ContextMenu->AddDynamicSection(TEXT("Main"), FNewToolMenuDelegate::CreateStatic(&FNavigationToolItemContextMenu::PopulateMenu));
+	}
+
+	UNavigationToolItemMenuContext* const ContextObject = NewObject<UNavigationToolItemMenuContext>();
+	ContextObject->Init(InToolView->GetOwnerTool(), InItemList);
+
+	FToolMenuContext Context(InToolView->GetViewCommandList(), nullptr, ContextObject);
+	return ToolMenus->GenerateWidget(MenuName, Context);
+}
+
+void FNavigationToolItemContextMenu::PopulateMenu(UToolMenu* const InMenu)
+{
+	if (!InMenu)
+	{
+		return;
+	}
+
+	UToolMenu& MenuRef = *InMenu;
+
+	UNavigationToolItemMenuContext* const ItemsContext = MenuRef.FindContext<UNavigationToolItemMenuContext>();
+	if (!IsValid(ItemsContext))
+	{
+		return;
+	}
+
+	CreateToolSection(MenuRef);
+	CreateGenericSection(MenuRef);
+
+	/*FToolMenuSection* ContextActionsSection = MenuRef.FindSection(TEXT("ContextActions"));
+	if (!ContextActionsSection)
+	{
+		ContextActionsSection = &MenuRef.AddSection(TEXT("ContextActions")
+			, LOCTEXT("ContextActions", "Context Actions")
+			, FToolMenuInsert(NAME_None, EToolMenuInsertType::First));
+	}*/
+
+	// Note: Since the Navigation Tool Command List is linked to the Command List (see INavigationTool::SetBaseCommandList),
+	// we do NOT need to add the entry with a different Command List
+	/*ContextActionsSection->AddMenuEntry(
+		ToolCommands.OpenAdvancedRenamerTool_SelectedActors,
+		TAttribute<FText>(),
+		TAttribute<FText>(),
+		FSlateIcon(FAppStyle::GetAppStyleSetName(), TEXT("GenericCommands.Rename")));
+
+	ContextActionsSection->AddMenuEntry(
+		ToolCommands.OpenAdvancedRenamerTool_SharedClassActors,
+		TAttribute<FText>(),
+		TAttribute<FText>(),
+		FSlateIcon(FAppStyle::GetAppStyleSetName(), TEXT("GenericCommands.Rename")));*/
+}
+
+void FNavigationToolItemContextMenu::CreateGenericSection(UToolMenu& InMenu)
+{
+	const FGenericCommands& GenericCommands = FGenericCommands::Get();
+
+	FToolMenuSection& GenericSection = InMenu.FindOrAddSection(TEXT("GenericActions"), LOCTEXT("GenericActionsHeader", "Generic Actions"));
+
+	//GenericSection.AddMenuEntry(GenericCommands.Cut);
+	//GenericSection.AddMenuEntry(GenericCommands.Copy);
+	//GenericSection.AddMenuEntry(GenericCommands.Paste);
+	//GenericSection.AddMenuEntry(GenericCommands.Duplicate);
+	GenericSection.AddMenuEntry(GenericCommands.Delete);
+	GenericSection.AddMenuEntry(GenericCommands.Rename);
+}
+
+void FNavigationToolItemContextMenu::CreateToolSection(UToolMenu& InMenu)
+{
+	const FNavigationToolCommands& ToolCommands = FNavigationToolCommands::Get();
+
+	FToolMenuSection& ToolSection = InMenu.FindOrAddSection(TEXT("ToolActions"), LOCTEXT("ToolActionsHeader", "Navigation Tool Actions"));
+
+	ToolSection.AddMenuEntry(ToolCommands.ExpandAll);
+	ToolSection.AddMenuEntry(ToolCommands.CollapseAll);
+	ToolSection.AddMenuEntry(ToolCommands.SelectAllChildren);
+	ToolSection.AddMenuEntry(ToolCommands.SelectImmediateChildren);
+	ToolSection.AddMenuEntry(ToolCommands.FocusSingleSelection);
+}
+
+}
+
+#undef LOCTEXT_NAMESPACE
diff --git a/Engine/Source/Editor/Sequencer/Private/Tools/NavigationTool/Menus/NavigationToolItemContextMenu.h b/Engine/Source/Editor/Sequencer/Private/Tools/NavigationTool/Menus/NavigationToolItemContextMenu.h
new file mode 100644
--- /dev/null
+++ b/Engine/Source/Editor/Sequencer/Private/Tools/NavigationTool/Menus/NavigationToolItemContextMenu.h
@@ -0,0 +1,31 @@
+// Copyright Epic Games, Inc. All Rights Reserved.
+
+#pragma once
+
+#include "Templates/SharedPointer.h"
+#include "Tools/NavigationTool/NavigationToolDefines.h"
+
+class FName;
+class SWidget;
+class UToolMenu;
+
+namespace UE::Sequencer::NavigationTool
+{
+
+class FNavigationToolView;
+
+class FNavigationToolItemContextMenu : public TSharedFromThis<FNavigationToolItemContextMenu>
+{
+public:
+	static FName GetMenuName();
+
+	TSharedRef<SWidget> CreateMenu(const TSharedRef<FNavigationToolView>& InToolView, const TArray<FNavigationToolItemPtr>& InItemList);
+
+protected:
+	static void PopulateMenu(UToolMenu* const InMenu);
+
+	static void CreateGenericSection(UToolMenu& InMenu);
+	static void CreateToolSection(UToolMenu& InMenu);
+};
+
+}
diff --git a/Engine/Source/Editor/Sequencer/Private/Tools/NavigationTool/Menus/NavigationToolViewMenuContext.h b/Engine/Source/Editor/Sequencer/Private/Tools/NavigationTool/Menus/NavigationToolViewMenuContext.h
new file mode 100644
--- /dev/null
+++ b/Engine/Source/Editor/Sequencer/Private/Tools/NavigationTool/Menus/NavigationToolViewMenuContext.h
@@ -0,0 +1,45 @@
+// Copyright Epic Games, Inc. All Rights Reserved.
+
+#pragma once
+
+#include "Templates/SharedPointer.h"
+#include "UObject/Object.h"
+#include "ToolMenu.h"
+#include "Tools/NavigationTool/NavigationTool.h"
+#include "Tools/NavigationTool/NavigationToolView.h"
+#include "NavigationToolViewMenuContext.generated.h"
+
+DECLARE_DELEGATE_OneParam(FOnPopulateNavigationToolViewToolbarMenu, UToolMenu*);
+
+UCLASS(Transient)
+class UNavigationToolViewMenuContext : public UObject
+{
+	GENERATED_BODY()
+
+	friend UE::Sequencer::NavigationTool::FNavigationToolView;
+
+public:
+	void Init(const TWeakPtr<UE::Sequencer::NavigationTool::FNavigationToolView>& InToolView)
+	{
+		WeakToolView = InToolView;
+	}
+
+	TSharedPtr<UE::Sequencer::NavigationTool::FNavigationToolView> GetToolView() const
+	{
+		return WeakToolView.Pin();
+	}
+
+	TSharedPtr<UE::Sequencer::NavigationTool::FNavigationTool> GetTool() const
+	{
+		if (const TSharedPtr<UE::Sequencer::NavigationTool::FNavigationToolView> ToolView = GetToolView())
+		{
+			return StaticCastSharedPtr<UE::Sequencer::NavigationTool::FNavigationTool>(ToolView->GetOwnerTool());
+		}
+		return nullptr;
+	}
+
+	FOnPopulateNavigationToolViewToolbarMenu OnPopulateMenu;
+
+private:
+	TWeakPtr<UE::Sequencer::NavigationTool::FNavigationToolView> WeakToolView;
+};
