[Sequencer] Navigation Tool initial submit

diff --git a/Engine/Source/Editor/Sequencer/Private/Tools/NavigationTool/NavigationTool.cpp b/Engine/Source/Editor/Sequencer/Private/Tools/NavigationTool/NavigationTool.cpp
new file mode 100644
--- /dev/null
+++ b/Engine/Source/Editor/Sequencer/Private/Tools/NavigationTool/NavigationTool.cpp
@@ -0,0 +1,1329 @@
+// Copyright Epic Games, Inc. All Rights Reserved.
+
+#include "NavigationTool.h"
+#include "Algo/Reverse.h"
+#include "BlueprintEditorSettings.h"
+#include "Components/ActorComponent.h"
+#include "Editor/Transactor.h"
+#include "Engine/Level.h"
+#include "Engine/LevelStreaming.h"
+#include "Engine/World.h"
+#include "Framework/Application/SlateApplication.h"
+#include "Framework/MultiBox/MultiBoxBuilder.h"
+#include "GameFramework/Actor.h"
+#include "GameFramework/WorldSettings.h"
+#include "Kismet2/ComponentEditorUtils.h"
+#include "Misc/Optional.h"
+#include "ScopedTransaction.h"
+#include "Sections/MovieSceneSubSection.h"
+#include "Sequencer.h"
+#include "Textures/SlateIcon.h"
+#include "Tools/NavigationTool/Customization/NavigationToolExtender.h"
+#include "Tools/NavigationTool/Filters/Filters/NavigationToolBuiltInFilter.h"
+#include "Tools/NavigationTool/ItemActions/NavigationToolItemAction.h"
+#include "Tools/NavigationTool/ItemProxies/NavigationToolItemProxyRegistry.h"
+#include "Tools/NavigationTool/Items/NavigationToolItemUtils.h"
+#include "Tools/NavigationTool/Items/NavigationToolObjectReference.h"
+#include "Tools/NavigationTool/Items/NavigationToolSequence.h"
+#include "Tools/NavigationTool/Items/NavigationToolTreeRoot.h"
+#include "Tools/NavigationTool/NavigationToolEditorSelection.h"
+#include "Tools/NavigationTool/NavigationToolScopedSelection.h"
+#include "Tools/NavigationTool/NavigationToolTab.h"
+#include "Tools/NavigationTool/NavigationToolView.h"
+#include "Tools/NavigationTool/Providers/NavigationToolProvider.h"
+#include "Widgets/Views/STableRow.h"
+
+#define LOCTEXT_NAMESPACE "NavigationTool"
+
+namespace UE::Sequencer::NavigationTool
+{
+
+FNavigationTool::FNavigationTool(const TWeakPtr<ISequencer>& InWeakSequencer)
+	: WeakSequencer(InWeakSequencer)
+	, ToolTab(MakeShared<FNavigationToolTab>(*this))
+	, RootItem(MakeShared<FNavigationToolTreeRoot>(*this))
+{
+}
+
+void FNavigationTool::Init()
+{
+	// Register the default view. @TODO: This could/should probably be extendable and setup by the providers?
+	constexpr int32 DefaultViewId = 0;
+	const TSharedPtr<INavigationToolView> ToolView = RegisterToolView(DefaultViewId);
+	check(ToolView.IsValid());
+
+	LoadAllProviderStates();
+
+	//LoadAllProviderStates()
+	{
+		GlobalFilters.Reset();
+
+		ForEachProvider([this](const TSharedRef<FNavigationToolProvider>& InProvider)
+			{
+				LoadProviderState(InProvider);
+				return true;
+			});
+
+		Refresh();
+	}
+
+	if (const TSharedPtr<FSequencer> Sequencer = StaticCastSharedPtr<FSequencer>(WeakSequencer.Pin()))
+	{
+		Sequencer->OnChannelChanged().AddSPLambda(this, [this](const FMovieSceneChannelMetaData*, UMovieSceneSection*)
+			{
+				OnTreeViewChanged();
+			});
+		Sequencer->OnMovieSceneBindingsChanged().AddSPLambda(this, [this]()
+			{
+				OnTreeViewChanged();
+			});
+		Sequencer->OnMovieSceneBindingsPasted().AddSPLambda(this, [this](const TArray<FMovieSceneBinding>&)
+			{
+				OnTreeViewChanged();
+			});
+		Sequencer->OnMovieSceneDataChanged().AddSPLambda(this, [this](const EMovieSceneDataChangeType)
+			{
+				OnTreeViewChanged();
+			});
+		Sequencer->OnTreeViewChanged().AddSP(this, &FNavigationTool::OnTreeViewChanged);
+
+		Sequencer->OnCloseEvent().AddSPLambda(this, [this](const TSharedRef<ISequencer> InSequencer)
+			{
+				if (InSequencer == WeakSequencer.Pin())
+				{
+					Shutdown();
+				}
+			});
+	}
+
+	// Listen to object replacement changes
+	FCoreUObjectDelegates::OnObjectsReplaced.AddSP(this, &FNavigationTool::OnObjectsReplaced);
+	FEditorDelegates::OnEditorActorReplaced.AddSP(this, &FNavigationTool::OnActorReplaced);
+
+	FNavigationToolExtender::OnProvidersChanged().AddSP(this, &FNavigationTool::OnProvidersChanged);
+
+	// ToolTab uses the most recent tool view as content, so must call after view registration
+	ToolTab->Init();
+
+	if (GEditor)
+	{
+		GEditor->RegisterForUndo(this);
+	}
+}
+
+void FNavigationTool::Shutdown()
+{
+	SaveAllProviderStates();
+
+	if (const TSharedPtr<ISequencer> Sequencer = WeakSequencer.Pin())
+	{
+		Sequencer->OnChannelChanged().RemoveAll(this);
+		Sequencer->OnMovieSceneBindingsChanged().RemoveAll(this);
+		Sequencer->OnMovieSceneBindingsPasted().RemoveAll(this);
+		Sequencer->OnMovieSceneDataChanged().RemoveAll(this);
+		Sequencer->OnTreeViewChanged().RemoveAll(this);
+	}
+
+	FCoreUObjectDelegates::OnObjectsReplaced.RemoveAll(this);
+	FEditorDelegates::OnEditorActorReplaced.RemoveAll(this);
+
+	FNavigationToolExtender::OnProvidersChanged().RemoveAll(this);
+
+	ToolTab->Shutdown();
+
+	if (GEditor)
+	{
+		GEditor->UnregisterForUndo(this);
+	}
+}
+
+void FNavigationTool::ForEachProvider(const TFunction<bool(const TSharedRef<FNavigationToolProvider>& InToolProvider)>& InPredicate)
+{
+	FNavigationToolExtender::ForEachToolProvider(*this, InPredicate);
+}
+
+void FNavigationTool::LoadAllProviderStates()
+{
+	GlobalFilters.Reset();
+
+	ForEachProvider([this](const TSharedRef<FNavigationToolProvider>& InProvider)
+		{
+			LoadProviderState(InProvider);
+			return true;
+		});
+
+	Refresh();
+}
+
+void FNavigationTool::SaveAllProviderStates()
+{
+	ForEachProvider([this](const TSharedRef<FNavigationToolProvider>& InProvider)
+		{
+			SaveProviderState(InProvider);
+			return true;
+		});
+}
+
+void FNavigationTool::OnProvidersChanged(const TSharedRef<INavigationTool>& InTool
+	, const TSharedRef<FNavigationToolProvider>& InProvider
+	, const ENavigationToolProvidersChangeType InChangeType)
+{
+	if (&InTool.Get() != this)
+	{
+		return;
+	}
+
+	if (InChangeType == ENavigationToolProvidersChangeType::Add)
+	{
+		const UNavigationToolSettings* const ToolSettings = GetDefault<UNavigationToolSettings>();
+
+		LoadProviderState(InProvider);
+
+		ForEachToolView([&InProvider](const TSharedPtr<FNavigationToolView>& InToolView)
+			{
+				InToolView->LoadViewState(InProvider);
+
+				InToolView->CreateColumns(InProvider);
+				InToolView->CreateDefaultColumnViews(InProvider);
+			});
+
+		// We apply default views *after* all columns have been created for all providers
+		if (ToolSettings->ShouldApplyDefaultColumnView())
+		{
+			ForEachToolView([&InProvider](const TSharedPtr<FNavigationToolView>& InToolView)
+				{
+					InToolView->ApplyCustomColumnView(InProvider->GetDefaultColumnView());
+				});
+		}
+	}
+	else if (InChangeType == ENavigationToolProvidersChangeType::Remove)
+	{
+		// Remove extended column views
+		UNavigationToolSettings* const ToolSettings = GetMutableDefault<UNavigationToolSettings>();
+		TSet<FNavigationToolColumnView>& CustomColumnViews = ToolSettings->GetCustomColumnViews();
+
+		for (const FName ColumnViewName : InProvider->GetExtendedColumnNames())
+		{
+			CustomColumnViews.Remove(ColumnViewName);
+		}
+
+		ToolSettings->SaveConfig();
+	}
+
+	Refresh();
+}
+
+USequencerSettings* FNavigationTool::GetSequencerSettings() const
+{
+	if (const TSharedPtr<ISequencer> Sequencer = WeakSequencer.Pin())
+	{
+		return Sequencer->GetSequencerSettings();
+	}
+	return nullptr;
+}
+
+bool FNavigationTool::CanProcessSequenceSpawn(UMovieSceneSequence* const InSequence) const
+{
+	return true;
+}
+
+TSharedPtr<FUICommandList> FNavigationTool::GetBaseCommandList() const
+{
+	return WeakBaseCommandList.Pin();
+}
+
+TArray<FName> FNavigationTool::GetRegisteredItemProxyTypeNames() const
+{
+	TArray<FName> OutItemProxyTypeNames;
+	{
+		// Get Tool-registered Item Types first
+		TSet<FName> NameSet;
+		ItemProxyRegistry.GetRegisteredItemProxyTypeNames(NameSet);
+
+		// Get the Module-registered Item Types second
+		TSet<FName> ModuleNameSet;
+		FNavigationToolExtender::GetItemProxyRegistry().GetRegisteredItemProxyTypeNames(ModuleNameSet);
+		NameSet.Append(ModuleNameSet);
+
+		OutItemProxyTypeNames = NameSet.Array();
+	}
+
+	OutItemProxyTypeNames.Sort(FNameLexicalLess());
+
+	return OutItemProxyTypeNames;
+}
+
+void FNavigationTool::GetItemProxiesForItem(const FNavigationToolItemPtr& InItem, TArray<TSharedPtr<FNavigationToolItemProxy>>& OutItemProxies)
+{
+	// No Item Proxy support for Root
+	if (!InItem.IsValid() || InItem == RootItem)
+	{
+		return;
+	}
+
+	InItem->GetItemProxies(OutItemProxies);
+
+	ForEachProvider([this, &InItem, &OutItemProxies]
+		(const TSharedRef<FNavigationToolProvider>& InProvider)
+		{
+			InProvider->OnExtendItemProxiesForItem(*this, InItem, OutItemProxies);
+			return true;
+		});
+
+	// Clean up any invalid Item Proxy
+	OutItemProxies.RemoveAll([](const TSharedPtr<FNavigationToolItemProxy>& InItemProxy) { return !InItemProxy.IsValid(); });
+
+	// Sort Proxies by their Priority
+	OutItemProxies.Sort([](const TSharedPtr<FNavigationToolItemProxy>& InItemProxyA, const TSharedPtr<FNavigationToolItemProxy>& InItemProxyB)
+		{
+			return InItemProxyA->GetPriority() > InItemProxyB->GetPriority();
+		});
+}
+
+INavigationToolItemProxyFactory* FNavigationTool::GetItemProxyFactory(const FName InItemProxyTypeName) const
+{
+	// First look for the Registry in Navigation Tool
+	if (INavigationToolItemProxyFactory* const Factory = ItemProxyRegistry.GetItemProxyFactory(InItemProxyTypeName))
+	{
+		return Factory;
+	}
+
+	// Fallback to finding the Factory in the Module if the Navigation Tool did not find it
+	return FNavigationToolExtender::GetItemProxyRegistry().GetItemProxyFactory(InItemProxyTypeName);
+}
+
+bool FNavigationTool::IsToolLocked() const
+{
+	bool bOutShouldLock = false;
+
+	FNavigationTool& SharedThisRef = const_cast<FNavigationTool&>(*this);
+
+	FNavigationToolExtender::ForEachToolProvider(SharedThisRef, [&bOutShouldLock]
+		(const TSharedRef<FNavigationToolProvider>& InToolProvider)
+		{
+			if (InToolProvider->ShouldLockTool())
+			{
+				bOutShouldLock = true;
+				return false; // No need to continue processing tool providers
+			}
+			return true;
+		});
+
+	return bOutShouldLock;
+}
+
+void FNavigationTool::HandleUndoRedoTransaction(const FTransaction* const InTransaction, const bool bInIsUndo)
+{
+	RequestRefresh();
+}
+
+TSharedPtr<ISequencer> FNavigationTool::GetSequencer() const
+{
+	return WeakSequencer.Pin();
+}
+
+void FNavigationTool::SetBaseCommandList(const TSharedPtr<FUICommandList>& InBaseCommandList)
+{
+	WeakBaseCommandList = InBaseCommandList;
+}
+
+void FNavigationTool::SaveProviderState(const TSharedRef<FNavigationToolProvider>& InProvider)
+{
+	SaveSerializedTree(/*bInResetTree*/true);
+
+	/*ForEachToolView([&InProvider](const TSharedPtr<FNavigationToolView>& InToolView)
+		{
+			InToolView->SaveViewState(InProvider);
+		});*/
+
+	if (FNavigationToolSaveState* const SaveState = InProvider->GetSaveState(*this))
+	{
+		// Remove any saved item colors that can no longer be found
+		TArray<FString> ItemIds;
+		SaveState->ItemColorMap.GetKeys(ItemIds);
+		for (const FString& ItemId : ItemIds)
+		{
+			if (!FindItem(ItemId).IsValid())
+			{
+				SaveState->ItemColorMap.Remove(ItemId);
+			}
+		}
+	}
+
+	InProvider->UpdateItemIdContexts(*this);
+}
+
+void FNavigationTool::LoadProviderState(const TSharedRef<FNavigationToolProvider>& InProvider)
+{
+	InProvider->UpdateItemIdContexts(*this);
+	
+	if (FNavigationToolSaveState* const SaveState = InProvider->GetSaveState(*this))
+	{
+		LoadSerializedTree(RootItem, &SaveState->SerializedTree);
+	}
+
+	ForEachToolView([&InProvider](const TSharedPtr<FNavigationToolView>& InToolView)
+		{
+			InToolView->LoadViewState(InProvider);
+		});
+
+	// Extend built in filters, make sure to not add any duplicates
+	TArray<FNavigationToolBuiltInFilterParams> NewGlobalFilterParams;
+	InProvider->OnExtendBuiltInFilters(NewGlobalFilterParams);
+
+	auto DoesGlobalFilterExist = [this](const FName InFilterId)
+	{
+		for (const TSharedPtr<FNavigationToolBuiltInFilter>& GlobalFilter : GlobalFilters)
+		{
+			if (GlobalFilter->GetName().Equals(InFilterId.ToString()))
+			{
+				return true;
+			}
+		}
+		return false;
+	};
+
+	for (const FNavigationToolBuiltInFilterParams& NewFilterParams : NewGlobalFilterParams)
+	{
+		if (!DoesGlobalFilterExist(NewFilterParams.GetFilterId()))
+		{
+			const TSharedPtr<FNavigationToolBuiltInFilter> NewFilter = MakeShared<FNavigationToolBuiltInFilter>(NewFilterParams);
+
+			if (!NewFilterParams.IsEnabledByDefault())
+			{
+				NewFilter->SetActive(false);
+			}
+
+			GlobalFilters.Add(NewFilter);
+		}
+	}
+}
+
+void FNavigationTool::LoadSerializedTree(const FNavigationToolItemPtr& InParentItem
+	, FNavigationToolSerializedTree* const InSerializedTree)
+{
+	TArray<FNavigationToolItemPtr>& Children = InParentItem->GetChildrenMutable();
+	
+	TArray<FNavigationToolItemPtr> Sortable;
+	TArray<FNavigationToolItemPtr> Unsortable;
+	SplitSortableAndUnsortableItems(Children, Sortable, Unsortable);
+
+	// If Scene Tree is valid, Item Sorting should be empty as this function only takes a valid Scene Tree if
+	// loaded version supports Scene Trees (i.e. when Item Sorting stops being loaded in)
+	if (InSerializedTree)
+	{
+		Sortable.Sort([InSerializedTree](const FNavigationToolItemPtr& InA, const FNavigationToolItemPtr& InB)
+			{
+				const FNavigationToolSerializedTreeNode* const NodeA = InSerializedTree->FindTreeNode(MakeSerializedItemFromToolItem(InA));
+				const FNavigationToolSerializedTreeNode* const NodeB = InSerializedTree->FindTreeNode(MakeSerializedItemFromToolItem(InB));
+				return FNavigationToolSerializedTree::CompareTreeItemOrder(NodeA, NodeB);
+			});
+	}
+
+	Children = MoveTemp(Unsortable);
+	Children.Append(MoveTemp(Sortable));
+
+	for (FNavigationToolItemPtr& Child : Children)
+	{
+		LoadSerializedTree(Child, InSerializedTree);
+	}
+}
+
+void FNavigationTool::SaveSerializedTree(const bool bInResetTree)
+{
+	ForEachProvider([this, bInResetTree](const TSharedRef<FNavigationToolProvider>& InProvider)
+		{
+			FNavigationToolSaveState* const SaveState = InProvider->GetSaveState(*this);
+			if (!SaveState)
+			{
+				return true;
+			}
+
+			if (bInResetTree)
+			{
+				SaveState->SerializedTree.Reset();
+			}
+
+			SaveSerializedTreeRecursive(RootItem, SaveState->SerializedTree);
+			return true;
+		});
+}
+
+void FNavigationTool::SaveSerializedTreeRecursive(const FNavigationToolItemPtr& InParentItem
+	, FNavigationToolSerializedTree& InSerializedTree)
+{
+	const TArray<FNavigationToolItemPtr>& Children = InParentItem->GetChildren();
+
+	const FNavigationToolSerializedItem ParentSceneItem = MakeSerializedItemFromToolItem(InParentItem);
+
+	for (int32 Index = 0; Index < Children.Num(); ++Index)
+	{
+		if (const FNavigationToolItemPtr& ChildItem = Children[Index])
+		{
+			if (ChildItem->ShouldSort())
+			{
+				const FNavigationToolSerializedItem SceneItem = MakeSerializedItemFromToolItem(ChildItem);
+				if (SceneItem.IsValid())
+				{
+					InSerializedTree.GetOrAddTreeNode(SceneItem, ParentSceneItem);
+				}
+			}
+			SaveSerializedTreeRecursive(ChildItem, InSerializedTree);
+		}
+	}
+}
+
+TSharedPtr<INavigationToolView> FNavigationTool::RegisterToolView(const int32 InToolViewId)
+{
+	const TSharedPtr<FNavigationToolView> ToolView = FNavigationToolView::CreateInstance(InToolViewId
+		, SharedThis(this)
+		, GetBaseCommandList());
+	ToolViews.Add(InToolViewId, ToolView);
+
+	ForEachProvider([this, &ToolView](const TSharedRef<FNavigationToolProvider>& InProvider)
+		{
+			ToolView->LoadViewState(InProvider);
+
+			ToolView->CreateColumns(InProvider);
+			ToolView->CreateDefaultColumnViews(InProvider);
+
+			return true;
+		});
+
+	return ToolView;
+}
+
+TSharedPtr<INavigationToolView> FNavigationTool::GetToolView(const int32 InToolViewId) const
+{
+	if (const TSharedPtr<FNavigationToolView>* const FoundToolView = ToolViews.Find(InToolViewId))
+	{
+		return *FoundToolView;
+	}
+	return nullptr;
+}
+
+bool FNavigationTool::IsObjectAllowedInTool(const UObject* const InObject) const
+{
+	const AActor* const Actor = Cast<AActor>(InObject);
+	if (IsValid(Actor))
+	{
+		// Do not show Actors that aren't editable or not meant to be listed in Navigation Tool
+		if (!Actor->IsEditable())
+		{
+			return false;
+		}
+
+		return true;
+	}
+
+	const UActorComponent* const ActorComponent = Cast<UActorComponent>(InObject);
+	if (IsValid(ActorComponent))
+	{
+		const bool bHideConstructionScriptComponents = GetDefault<UBlueprintEditorSettings>()->bHideConstructionScriptComponentsInDetailsView;
+		return !ActorComponent->IsVisualizationComponent()
+			&& (ActorComponent->CreationMethod != EComponentCreationMethod::UserConstructionScript || !bHideConstructionScriptComponents)
+			&& (ActorComponent->CreationMethod != EComponentCreationMethod::Native ||
+				FComponentEditorUtils::GetPropertyForEditableNativeComponent(ActorComponent));
+	}
+
+	return false;
+}
+
+void FNavigationTool::RegisterItem(const FNavigationToolItemPtr& InItem)
+{
+	if (!InItem.IsValid())
+	{
+		return;
+	}
+
+	const FNavigationToolItemId ItemId = InItem->GetItemId();
+	const FNavigationToolItemPtr ExistingItem = FindItem(ItemId);
+
+	// if there's no existing item or the existing item does not match its new replacement,
+	// then call OnItemRegistered and Refresh Navigation Tool
+	if (!ExistingItem.IsValid() || ExistingItem != InItem)
+	{
+		InItem->OnItemRegistered();
+
+		AddItem(InItem);
+
+		RequestRefresh();
+	}
+}
+
+void FNavigationTool::UnregisterItem(const FNavigationToolItemId& InItemId)
+{
+	const FNavigationToolItemPtr FoundItem = FindItem(InItemId);
+	if (!FoundItem.IsValid())
+	{
+		return;
+	}
+
+	FoundItem->OnItemUnregistered();
+
+	RemoveItem(InItemId);
+
+	RequestRefresh();
+}
+
+void FNavigationTool::RequestRefresh()
+{
+	bRefreshRequested = true;
+}
+
+void FNavigationTool::Refresh()
+{
+	TGuardValue RefreshGuard(bRefreshing, true);
+	bRefreshRequested = false;
+
+	// Flush Pending Actions
+	{
+		// Make a Transaction if there's a Pending Action Requesting it
+		TUniquePtr<FScopedTransaction> Transaction;
+
+		if (!GIsTransacting)
+		{
+			const bool bShouldTransact = PendingActions.ContainsByPredicate(
+				[](const TSharedPtr<INavigationToolAction>& InAction)
+				{
+					return InAction.IsValid() && InAction->ShouldTransact();
+				});
+
+			if (bShouldTransact)
+			{
+				Transaction = MakeUnique<FScopedTransaction>(LOCTEXT("NavigationToolItemAction", "Navigation Tool Item Action"));
+			}
+		}
+
+		// Execute Pending Actions
+		for (const TSharedPtr<INavigationToolAction>& Action : PendingActions)
+		{
+			if (Action.IsValid())
+			{
+				Action->Execute(*this);
+			}
+		}
+
+		PendingActions.Empty();
+	}
+
+	// Save the Current Item ordering before Refreshing Children
+	// Do not reset tree just yet as there might be actors that still need to be discovered on the next pass
+	// This is done to save the items added from the Queued Actions and be considered when adding new items from Discovery
+	SaveSerializedTree(/*bInResetTree*/false);
+
+	// Refresh each item's children. This also Updates each Child's parent var
+	ForEachItem([](const FNavigationToolItemPtr& InItem)
+		{
+			InItem->RefreshChildren();
+		});
+
+	RootItem->RefreshChildren();
+
+	ForEachToolView([](const TSharedPtr<FNavigationToolView>& InToolView)
+		{
+			InToolView->Refresh();
+		});
+
+	// Save so that the Tree is updated to the latest Navigation Tool State
+	SaveSerializedTree(/*bInResetTree*/false);
+}
+
+TSharedRef<FNavigationToolItem> FNavigationTool::GetTreeRoot() const
+{
+	return RootItem;
+}
+
+FNavigationToolItemPtr FNavigationTool::FindItem(const FNavigationToolItemId& InItemId) const
+{
+	if (bIteratingItemMap)
+	{
+		if (const FNavigationToolItemPtr* const FoundItem = ItemsPendingAdd.Find(InItemId))
+		{
+			return *FoundItem;
+		}
+	}
+
+	if (const FNavigationToolItemPtr* const FoundItem = ItemMap.Find(InItemId))
+	{
+		return *FoundItem;
+	}
+
+	return nullptr;
+}
+
+void FNavigationTool::SetIgnoreNotify(const ENavigationToolIgnoreNotifyFlags InFlag, const bool bInIgnore)
+{
+	if (bInIgnore)
+	{
+		EnumAddFlags(IgnoreNotifyFlags, InFlag);
+	}
+	else
+	{
+		EnumRemoveFlags(IgnoreNotifyFlags, InFlag);
+	}
+}
+
+void FNavigationTool::PostUndo(const bool bInSuccess)
+{
+	if (bInSuccess)
+	{
+		const int32 QueueIndex = GEditor->Trans->GetQueueLength() - GEditor->Trans->GetUndoCount();
+		const FTransaction* Transaction = GEditor->Trans->GetTransaction(QueueIndex);
+		HandleUndoRedoTransaction(Transaction, true);
+	}
+}
+
+void FNavigationTool::PostRedo(const bool bInSuccess)
+{
+	if (bInSuccess)
+	{
+		const int32 QueueIndex = GEditor->Trans->GetQueueLength() - GEditor->Trans->GetUndoCount();
+		const FTransaction* Transaction = GEditor->Trans->GetTransaction(QueueIndex);
+		HandleUndoRedoTransaction(Transaction, false);
+	}
+}
+
+TStatId FNavigationTool::GetStatId() const
+{
+	RETURN_QUICK_DECLARE_CYCLE_STAT(FNavigationTool, STATGROUP_Tickables);
+}
+
+void FNavigationTool::Tick(const float InDeltaTime)
+{
+	if (NeedsRefresh())
+	{
+		Refresh();
+	}
+
+	if (bToolDirty)
+	{
+		bToolDirty = false;
+	}
+
+	// Select Objects Pending Selection
+	if (ObjectsLastSelected.IsValid())
+	{
+		TArray<FNavigationToolItemPtr> ItemsToSelect;
+
+		if (ObjectsLastSelected->Num() > 0)
+		{
+			ItemsToSelect.Reserve(ObjectsLastSelected->Num());
+			for (TWeakObjectPtr<UObject> Object : *ObjectsLastSelected)
+			{
+				if (Object.IsValid())
+				{
+					if (const FNavigationToolItemPtr Item = FindItem(FNavigationToolItemId(Object.Get())))
+					{
+						// If it's a Shared Object, select all its References as Shared Objects are not in the Navigation Tool directly
+						// and the outliner does not know which Object Reference should be selected
+						if (const FNavigationToolSharedObject* const SharedObjectItem = Item->CastTo<FNavigationToolSharedObject>())
+						{
+							ItemsToSelect.Append(SharedObjectItem->GetObjectReferences());
+						}
+						else
+						{
+							ItemsToSelect.Add(Item);
+						}
+					}
+				}
+			}
+		}
+
+		// Only Scroll Into View, don't signal selection since we just come from the selection notify itself
+		SelectItems(ItemsToSelect, ENavigationToolItemSelectionFlags::ScrollIntoView);
+		ObjectsLastSelected.Reset();
+	}
+
+	ForEachToolView([InDeltaTime](const TSharedPtr<FNavigationToolView>& InToolView)
+		{
+			InToolView->Tick(InDeltaTime);
+		});
+}
+
+void FNavigationTool::DeleteItems(TArray<FNavigationToolItemPtr> InItems)
+{
+	SortItems(InItems);
+
+	for (TArray<FNavigationToolItemPtr>::TIterator It(InItems); It; ++It)
+	{
+		FNavigationToolItemPtr Item = *It;
+
+		if (!Item.IsValid() || !Item->CanDelete())
+		{
+			It.RemoveCurrent();
+		}
+	}
+
+	if (InItems.IsEmpty())
+	{
+		return;
+	}
+
+	FScopedTransaction DeleteTransaction(LOCTEXT("NavigationToolItemDeleteAction", "Navigation Tool Delete Item(s)"), !GIsTransacting);
+
+	bool bRequestRefresh = false;
+	for (FNavigationToolItemPtr& Item : InItems)
+	{
+		if (Item && Item->Delete())
+		{
+			bRequestRefresh = true;
+		}
+	}
+
+	if (bRequestRefresh)
+	{
+		RequestRefresh();
+	}
+}
+
+void FNavigationTool::UnregisterToolView(const int32 InToolViewId)
+{
+	ToolViews.Remove(InToolViewId);
+}
+
+void FNavigationTool::UpdateRecentToolViews(const int32 InToolViewId)
+{
+	RecentToolViews.Remove(InToolViewId);
+	RecentToolViews.Add(InToolViewId);
+}
+
+TSharedPtr<INavigationToolView> FNavigationTool::GetMostRecentToolView() const
+{
+	for (int32 Index = RecentToolViews.Num() - 1; Index >= 0; --Index)
+	{
+		if (const TSharedPtr<INavigationToolView> ToolView = GetToolView(RecentToolViews[Index]))
+		{
+			return StaticCastSharedPtr<FNavigationToolView>(ToolView);
+		}
+	}
+	return nullptr;
+}
+
+void FNavigationTool::ForEachToolView(const TFunction<void(const TSharedPtr<FNavigationToolView>& InToolView)>& InPredicate) const
+{
+	FNavigationTool* const MutableThis = const_cast<FNavigationTool*>(this);
+	for (TMap<int32, TSharedPtr<FNavigationToolView>>::TIterator Iter(MutableThis->ToolViews); Iter; ++Iter)
+	{
+		const TSharedPtr<FNavigationToolView>& ToolView = Iter.Value();
+		if (ToolView.IsValid())
+		{
+			InPredicate(ToolView);
+		}
+		else
+		{
+			Iter.RemoveCurrent();
+		}
+	}
+}
+
+void FNavigationTool::EnqueueItemActions(TArray<TSharedPtr<INavigationToolAction>>&& InItemActions) noexcept
+{
+	PendingActions.Append(MoveTemp(InItemActions));
+}
+
+int32 FNavigationTool::GetPendingItemActionCount() const
+{
+	return PendingActions.Num();
+}
+
+bool FNavigationTool::NeedsRefresh() const
+{
+	if (bRefreshing)
+	{
+		return false;
+	}
+
+	if (bRefreshRequested || PendingActions.Num() > 0)
+	{
+		return true;
+	}
+
+	return false;
+}
+
+TOptional<FColor> FNavigationTool::FindItemColor(const FNavigationToolItemPtr& InItem, bool bRecurseParent) const
+{
+	if (!InItem.IsValid())
+	{
+		return TOptional<FColor>();
+	}
+
+	const TSharedPtr<INavigationToolProvider> Provider = InItem->GetProvider();
+	if (!Provider.IsValid())
+	{
+		return TOptional<FColor>();
+	}
+
+	FNavigationToolSaveState* const SaveState = Provider->GetSaveState(*this);
+	if (!SaveState)
+	{
+		return TOptional<FColor>();
+	}
+
+	if (const FColor* const FoundColor = SaveState->ItemColorMap.Find(InItem->GetItemId().GetStringId()))
+	{
+		return *FoundColor;
+	}
+
+	//If no Item Coloring Mapping was found for the Specific Item, then try find the Item Color of the Parent
+	if (bRecurseParent)
+	{
+		return FindItemColor(InItem->GetParent(), bRecurseParent);
+	}
+
+	return TOptional<FColor>();
+}
+
+void FNavigationTool::SetItemColor(const FNavigationToolItemPtr& InItem, const FColor& InColor)
+{
+	if (!InItem.IsValid())
+	{
+		return;
+	}
+
+	const TSharedPtr<INavigationToolProvider> Provider = InItem->GetProvider();
+	if (!Provider.IsValid())
+	{
+		return;
+	}
+
+	FNavigationToolSaveState* const SaveState = Provider->GetSaveState(*this);
+	if (!SaveState)
+	{
+		return;
+	}
+
+	bool bShouldChangeItemColor = true;
+
+	const TOptional<FColor> InheritedColor = FindItemColor(InItem, true);
+	const bool bHasInheritedColor = InheritedColor.IsSet();
+
+	if (bHasInheritedColor)
+	{
+		//Make sure the Inherited Color is different than the Color we're trying to Set.
+		bShouldChangeItemColor = InheritedColor != InColor;
+	}
+
+	if (bShouldChangeItemColor)
+	{
+		const TOptional<FColor> ParentInheritedColor = FindItemColor(InItem->GetParent(), true);
+
+		//First check if the Color to Set matches the one inherited from the Parent
+		if (ParentInheritedColor.IsSet() && ParentInheritedColor == InColor)
+		{
+			//If it matches, remove this Item from the map as we will inherit from parent
+			SaveState->ItemColorMap.Remove(InItem->GetItemId().GetStringId());
+		}
+		else
+		{
+			SaveState->ItemColorMap.Add(InItem->GetItemId().GetStringId(), InColor);
+		}
+
+		//Recurse Children to Remove their Mapping if same color with parent
+		TArray<FNavigationToolItemPtr> RemainingChildren = InItem->GetChildren();
+		while (RemainingChildren.Num() > 0)
+		{
+			const FNavigationToolItemPtr Child = RemainingChildren.Pop();
+			if (Child.IsValid())
+			{
+				const TOptional<FColor> ChildColor = FindItemColor(Child, false);
+				if (ChildColor.IsSet() && ChildColor == InColor)
+				{
+					RemoveItemColor(Child);
+
+					//Only check Grand Children if Child had same Color.
+					//Allow the situation where Parent is ColorA, Child is ColorB and GrandChild is ColorA.
+					RemainingChildren.Append(Child->GetChildren());
+				}
+			}
+		}
+		SetToolModified();
+	}
+}
+
+void FNavigationTool::RemoveItemColor(const FNavigationToolItemPtr& InItem)
+{
+	if (!InItem.IsValid())
+	{
+		return;
+	}
+
+	const TSharedPtr<INavigationToolProvider> Provider = InItem->GetProvider();
+	if (!Provider.IsValid())
+	{
+		return;
+	}
+
+	FNavigationToolSaveState* const SaveState = Provider->GetSaveState(*this);
+	if (!SaveState)
+	{
+		return;
+	}
+
+	const bool bRemoved = SaveState->ItemColorMap.Remove(InItem->GetItemId().GetStringId()) > 0;
+	if (bRemoved)
+	{
+		SetToolModified();
+	}
+}
+
+void FNavigationTool::NotifyItemIdChanged(const FNavigationToolItemId& OldId, const FNavigationToolItemPtr& InItem)
+{
+	const FNavigationToolItemId NewId = InItem->GetItemId();
+	if (OldId == NewId)
+	{
+		return;
+	}
+
+	const FNavigationToolItemPtr FoundItem = FindItem(OldId);
+	if (FoundItem.IsValid() && FoundItem == InItem)
+	{
+		AddItem(InItem);
+		RemoveItem(OldId);
+		SetToolModified();
+	}
+}
+
+void FNavigationTool::NotifyToolItemRenamed(const FNavigationToolItemPtr& InItem)
+{
+	ForEachProvider([&InItem](const TSharedRef<FNavigationToolProvider>& InProvider)
+		{
+			InProvider->NotifyToolItemRenamed(InItem);
+			return true;
+		});
+}
+
+void FNavigationTool::NotifyToolItemDeleted(const FNavigationToolItemPtr& InItem)
+{
+	ForEachProvider([&InItem](const TSharedRef<INavigationToolProvider>& InProvider)
+		{
+			InProvider->NotifyToolItemDeleted(InItem);
+			return true;
+		});
+}
+
+TArray<FNavigationToolItemPtr> FNavigationTool::GetSelectedItems(const bool bInNormalizeToTopLevelSelections) const
+{
+	if (const TSharedPtr<INavigationToolView>& ToolView = GetMostRecentToolView())
+	{
+		TArray<FNavigationToolItemPtr> SelectedItems = ToolView->GetSelectedItems();
+
+		if (bInNormalizeToTopLevelSelections)
+		{
+			NormalizeToTopLevelSelections(SelectedItems);
+			return SelectedItems;
+		}
+
+		return SelectedItems;
+	}
+
+	return TArray<FNavigationToolItemPtr>();
+}
+
+void FNavigationTool::SelectItems(const TArray<FNavigationToolItemPtr>& InItems, const ENavigationToolItemSelectionFlags InFlags) const
+{
+	ForEachToolView([InItems, InFlags](const TSharedPtr<INavigationToolView>& InToolView)
+		{
+			InToolView->SelectItems(InItems, InFlags);
+		});
+}
+
+void FNavigationTool::ClearItemSelection(const bool bInSignalSelectionChange) const
+{
+	ForEachToolView([bInSignalSelectionChange](const TSharedPtr<INavigationToolView>& InToolView)
+		{
+			InToolView->ClearItemSelection(bInSignalSelectionChange);
+		});
+}
+
+FNavigationToolItemPtr FNavigationTool::FindLowestCommonAncestor(const TArray<FNavigationToolItemPtr>& Items)
+{
+	TSet<FNavigationToolItemPtr> IntersectedAncestors;
+
+	for (const FNavigationToolItemPtr& Item : Items)
+	{
+		FNavigationToolItemPtr Parent = Item->GetParent();
+		TSet<FNavigationToolItemPtr> ItemAncestors;
+
+		//Add all Item's Ancestors
+		while (Parent.IsValid())
+		{
+			ItemAncestors.Add(Parent);
+			Parent = Parent->GetParent();
+		}
+
+		//cant check for intersection if empty so just init
+		if (IntersectedAncestors.Num() == 0)
+		{
+			IntersectedAncestors = ItemAncestors;
+		}
+		else
+		{
+			IntersectedAncestors = IntersectedAncestors.Intersect(ItemAncestors);
+
+			//We are sure the intersection is the Root if only one item is remaining. Stop iterating
+			if (IntersectedAncestors.Num() == 1)
+			{
+				break;
+			}
+		}
+	}
+
+	FNavigationToolItemPtr LowestCommonAncestor;
+	for (const FNavigationToolItemPtr& Item : IntersectedAncestors)
+	{
+		//Find Item with most tree height (i.e. lowest down the tree, closer to the selected nodes)
+		if (!LowestCommonAncestor.IsValid() || Item->GetItemTreeHeight() > LowestCommonAncestor->GetItemTreeHeight())
+		{
+			LowestCommonAncestor = Item;
+		}
+	}
+	return LowestCommonAncestor;
+}
+
+FNavigationToolSerializedItem FNavigationTool::MakeSerializedItemFromToolItem(const FNavigationToolItemPtr& InItem)
+{
+	if (!InItem.IsValid())
+	{
+		return FNavigationToolSerializedItem();
+	}
+
+	return FNavigationToolSerializedItem(InItem->GetItemId().GetStringId());
+}
+
+void FNavigationTool::SortItems(TArray<FNavigationToolItemPtr>& OutToolItems, const bool bInReverseOrder)
+{
+	OutToolItems.Sort([bInReverseOrder](const FNavigationToolItemPtr& ItemA, const FNavigationToolItemPtr& ItemB)
+	{
+		return UE::Sequencer::NavigationTool::CompareToolItemOrder(ItemA, ItemB) != bInReverseOrder;
+	});
+}
+
+void FNavigationTool::NormalizeToTopLevelSelections(TArray<FNavigationToolItemPtr>& InOutItems)
+{
+	if (InOutItems.IsEmpty())
+	{
+		return;
+	}
+
+	//Set for Quick Lookup
+	const TSet<FNavigationToolItemPtr> SelectedItemSet(InOutItems);
+
+	//Normalize Selection: Remove all Items that have Parents that are in the Selection. Swapping since we're sorting afterwards
+	InOutItems.RemoveAllSwap([&SelectedItemSet](const FNavigationToolItemPtr& Item)
+	{
+		FNavigationToolItemPtr Parent = Item->GetParent();
+		while (Parent.IsValid())
+		{
+			if (SelectedItemSet.Contains(Parent))
+			{
+				return true;
+			}
+			Parent = Parent->GetParent();
+		}
+		return false;
+	});
+}
+
+void FNavigationTool::SyncSequencerSelection(const TArray<FNavigationToolItemPtr>& InSelectedItems) const
+{
+	const TSharedPtr<ISequencer> Sequencer = GetSequencer();
+	if (!Sequencer.IsValid())
+	{
+		return;
+	}
+
+	FNavigationToolScopedSelection ScopedSelection(*Sequencer, ENavigationToolScopedSelectionPurpose::Sync);
+	for (const FNavigationToolItemPtr& Item : InSelectedItems)
+	{
+		Item->Select(ScopedSelection);
+	}
+}
+
+void FNavigationTool::OnObjectSelectionChanged(const FNavigationToolEditorSelection& InEditorSelection)
+{
+	bool bIsSyncingItemSelection = false;
+
+	// find if any outliner view is syncing item selection
+	// if we are, ignore object selection or it will cause another round of selections next tick
+	ForEachToolView([&bIsSyncingItemSelection](const TSharedPtr<FNavigationToolView>& InToolView)
+		{
+			bIsSyncingItemSelection |= InToolView->IsSyncingItemSelection();
+		});
+
+	if (bIsSyncingItemSelection)
+	{
+		return;
+	}
+	
+	if (!ObjectsLastSelected.IsValid())
+	{
+		ObjectsLastSelected = MakeShared<TArray<TWeakObjectPtr<UObject>>>();
+	}
+
+	TArray<UObject*> SelectedObjects = InEditorSelection.GetSelectedObjects<UObject, ENavigationToolSelectionSource::All>();
+	ObjectsLastSelected->Reserve(ObjectsLastSelected->Num() + SelectedObjects.Num());
+
+	for (UObject* const Object : SelectedObjects)
+	{
+		ObjectsLastSelected->AddUnique(Object);
+	}
+}
+
+const FNavigationToolItemProxyRegistry& FNavigationTool::GetItemProxyRegistry() const
+{
+	return ItemProxyRegistry;
+}
+
+FNavigationToolItemProxyRegistry& FNavigationTool::GetItemProxyRegistry()
+{
+	return ItemProxyRegistry;
+}
+
+void FNavigationTool::OnObjectsReplaced(const TMap<UObject*, UObject*>& InReplacementMap)
+{
+	ForEachItem([&InReplacementMap](const FNavigationToolItemPtr& InItem)
+		{
+			// Recursive not needed since we're calling it on all Items in Map Anyway
+			InItem->OnObjectsReplaced(InReplacementMap, /*bRecursive*/false);
+		});
+
+	if (ObjectsLastSelected.IsValid())
+	{
+		for (TWeakObjectPtr<UObject>& Object : *ObjectsLastSelected)
+		{
+			if (Object.IsValid() && InReplacementMap.Contains(Object.Get()))
+			{
+				Object = InReplacementMap[Object.Get()];
+			}
+		}
+	}
+
+	for (const TSharedPtr<INavigationToolAction>& Action : PendingActions)
+	{
+		//Recursive needed since we only have direct reference to the underlying item in the Action, not its children
+		if (Action.IsValid())
+		{
+			Action->OnObjectsReplaced(InReplacementMap, true);
+		}
+	}
+
+	ForEachToolView([](const TSharedPtr<FNavigationToolView>& InToolView)
+		{
+			InToolView->NotifyObjectsReplaced();
+		});
+
+	RequestRefresh();
+}
+
+void FNavigationTool::OnActorReplaced(AActor* const InOldActor, AActor* const InNewActor)
+{
+	TMap<UObject*, UObject*> ReplacementMap;
+	ReplacementMap.Add(InOldActor, InNewActor);
+	OnObjectsReplaced(ReplacementMap);
+}
+
+void FNavigationTool::SetToolModified()
+{
+	if (!bToolDirty)
+	{
+		bToolDirty = true;
+	}
+}
+
+bool FNavigationTool::IsToolTabVisible() const
+{
+	return ToolTab->IsToolTabVisible();
+}
+
+void FNavigationTool::ShowHideToolTab(const bool bInVisible)
+{
+	ToolTab->ShowHideToolTab(bInVisible);
+
+	if (!bInVisible)
+	{
+		SaveAllProviderStates();
+	}
+}
+
+void FNavigationTool::ToggleToolTabVisible()
+{
+	ToolTab->ToggleToolTabVisible();
+
+	if (!ToolTab->IsToolTabVisible())
+	{
+		SaveAllProviderStates();
+	}
+}
+
+void FNavigationTool::AddItem(const FNavigationToolItemPtr& InItem)
+{
+	const FNavigationToolItemId ItemId = InItem->GetItemId();
+
+	ItemsPendingRemove.Remove(ItemId);
+
+	if (bIteratingItemMap)
+	{
+		ItemsPendingAdd.Add(ItemId, InItem);
+	}
+	else
+	{
+		ItemMap.Add(ItemId, InItem);
+	}
+}
+
+void FNavigationTool::RemoveItem(const FNavigationToolItemId& InItemId)
+{
+	ItemsPendingAdd.Remove(InItemId);
+
+	if (bIteratingItemMap)
+	{
+		ItemsPendingRemove.Add(InItemId);
+	}
+	else
+	{
+		ItemMap.Remove(InItemId);
+	}
+}
+
+void FNavigationTool::ForEachItem(TFunctionRef<void(const FNavigationToolItemPtr&)> InFunc)
+{
+	// iteration scope, allowing for nested for-each
+	{
+		TGuardValue<bool> Guard(bIteratingItemMap, true);
+		for (const TPair<FNavigationToolItemId, FNavigationToolItemPtr>& Pair : ItemMap)
+		{
+			InFunc(Pair.Value);
+		}
+	}
+
+	if (!bIteratingItemMap && (!ItemsPendingAdd.IsEmpty() || !ItemsPendingRemove.IsEmpty()))
+	{
+		for (const TPair<FNavigationToolItemId, FNavigationToolItemPtr>& ItemToAdd : ItemsPendingAdd)
+		{
+			ItemMap.Add(ItemToAdd);
+		}
+		ItemsPendingAdd.Empty();
+
+		for (const FNavigationToolItemId& ItemIdToRemove : ItemsPendingRemove)
+		{
+			ItemMap.Remove(ItemIdToRemove);
+		}
+		ItemsPendingRemove.Empty();
+	}
+}
+
+void FNavigationTool::OnTreeViewChanged()
+{
+	if (const TSharedPtr<INavigationToolView> RecentToolView = GetMostRecentToolView())
+	{
+		RecentToolView->RequestRefresh();
+	}
+
+	RequestRefresh();
+}
+
+}
+
+#undef LOCTEXT_NAMESPACE
diff --git a/Engine/Source/Editor/Sequencer/Private/Tools/NavigationTool/NavigationTool.h b/Engine/Source/Editor/Sequencer/Private/Tools/NavigationTool/NavigationTool.h
new file mode 100644
--- /dev/null
+++ b/Engine/Source/Editor/Sequencer/Private/Tools/NavigationTool/NavigationTool.h
@@ -0,0 +1,265 @@
+// Copyright Epic Games, Inc. All Rights Reserved.
+
+#pragma once
+
+#include "EditorUndoClient.h"
+#include "EditorSequenceNavigationDefs.h"
+#include "Engine/EngineTypes.h"
+#include "Tools/NavigationTool/INavigationTool.h"
+#include "Tools/NavigationTool/Items/INavigationToolItem.h"
+#include "Tools/NavigationTool/Items/NavigationToolItemId.h"
+#include "Tools/NavigationTool/Items/NavigationToolItemProxy.h"
+#include "Tools/NavigationTool/ItemProxies/NavigationToolItemProxyRegistry.h"
+#include "Tools/NavigationTool/ItemProxies/INavigationToolItemProxyFactory.h"
+#include "Tools/NavigationTool/NavigationToolDefines.h"
+#include "TickableEditorObject.h"
+#include "UObject/WeakObjectPtrTemplates.h"
+
+class AActor;
+class FTransaction;
+class FUICommandList;
+class ISequencer;
+class UMovieSceneSequence;
+class USequencerSettings;
+class USceneComponent;
+enum class EItemDropZone;
+
+namespace UE::Sequencer::NavigationTool
+{
+
+class INavigationToolAction;
+class FNavigationToolBuiltInFilter;
+class FNavigationToolProvider;
+class FNavigationToolTab;
+class FNavigationToolTreeRoot;
+class FNavigationToolView;
+class INavigationToolView;
+
+class FNavigationTool
+	: public INavigationTool
+	, public FTickableEditorObject
+	, public FEditorUndoClient
+{
+public:
+	FNavigationTool(const TWeakPtr<ISequencer>& InWeakSequencer);
+
+	void Init();
+	void Shutdown();
+
+	void ForEachProvider(const TFunction<bool(const TSharedRef<FNavigationToolProvider>& InToolProvider)>& InPredicate);
+
+	void LoadAllProviderStates();
+	void SaveAllProviderStates();
+
+	bool CanProcessSequenceSpawn(UMovieSceneSequence* const InSequence) const;
+
+	TSharedPtr<FUICommandList> GetBaseCommandList() const;
+
+	FOnToolLoaded OnToolLoaded;
+
+	/** Gathers all previously existing and new Item Proxies for a given Item */
+	void GetItemProxiesForItem(const FNavigationToolItemPtr& InItem, TArray<TSharedPtr<FNavigationToolItemProxy>>& OutItemProxies);
+
+	/** Tries to find the Item Proxy Factory for the given Item Proxy Type Name */
+	INavigationToolItemProxyFactory* GetItemProxyFactory(const FName InItemProxyTypeName) const;
+
+	/** Returns whether the Navigation Tool is in Read-only mode */
+	bool IsToolLocked() const;
+
+	void HandleUndoRedoTransaction(const FTransaction* const InTransaction, const bool bInIsUndo);
+
+	//~ Begin INavigationTool
+
+	virtual bool IsToolTabVisible() const override;
+	virtual void ShowHideToolTab(const bool bInVisible) override;
+	virtual void ToggleToolTabVisible() override;
+
+	virtual FOnToolLoaded& GetOnToolLoaded() override { return OnToolLoaded; }
+
+	virtual TSharedPtr<ISequencer> GetSequencer() const override;
+
+	virtual void SetBaseCommandList(const TSharedPtr<FUICommandList>& InBaseCommandList) override;
+
+	virtual TSharedPtr<INavigationToolView> RegisterToolView(const int32 InToolViewId) override;
+	virtual TSharedPtr<INavigationToolView> GetToolView(const int32 InToolViewId) const override;
+	virtual TSharedPtr<INavigationToolView> GetMostRecentToolView() const override;
+
+	virtual bool IsObjectAllowedInTool(const UObject* const InObject) const override;
+
+	virtual void RegisterItem(const FNavigationToolItemPtr& InItem) override;
+	virtual void UnregisterItem(const FNavigationToolItemId& InItemId) override;
+
+	virtual void RequestRefresh() override;
+	virtual void Refresh() override;
+
+	virtual TSharedRef<FNavigationToolItem> GetTreeRoot() const override;
+
+	virtual FNavigationToolItemPtr FindItem(const FNavigationToolItemId& InItemId) const override;
+
+	virtual void SetIgnoreNotify(const ENavigationToolIgnoreNotifyFlags InFlag, const bool bInIgnore) override;
+
+	virtual void OnObjectSelectionChanged(const FNavigationToolEditorSelection& InEditorSelection) override;
+	virtual TArray<FNavigationToolItemPtr> GetSelectedItems(const bool bInNormalizeToTopLevelSelections = false) const override;
+	virtual void SelectItems(const TArray<FNavigationToolItemPtr>& InItems
+		, const ENavigationToolItemSelectionFlags InFlags = ENavigationToolItemSelectionFlags::SignalSelectionChange) const override;
+	virtual void ClearItemSelection(const bool bInSignalSelectionChange) const override;
+
+	virtual const FNavigationToolItemProxyRegistry& GetItemProxyRegistry() const override;
+	virtual TArray<FName> GetRegisteredItemProxyTypeNames() const override;
+
+	virtual void SetItemColor(const FNavigationToolItemPtr& InItem, const FColor& InColor) override;
+	virtual void RemoveItemColor(const FNavigationToolItemPtr& InItem) override;
+	virtual TOptional<FColor> FindItemColor(const FNavigationToolItemPtr& InItem, bool bRecurseParent = true) const override;
+
+	virtual void EnqueueItemActions(TArray<TSharedPtr<INavigationToolAction>>&& InItemActions) noexcept override;
+
+	virtual void NotifyToolItemRenamed(const FNavigationToolItemPtr& InItem) override;
+	virtual void NotifyToolItemDeleted(const FNavigationToolItemPtr& InItem) override;
+
+	//~ End INavigationTool
+
+	FNavigationToolItemProxyRegistry& GetItemProxyRegistry();
+
+	//~ Begin FEditorUndoClient
+	virtual void PostUndo(const bool bInSuccess) override;
+	virtual void PostRedo(const bool bInSuccess) override;
+	//~ End FEditorUndoClient
+
+	//~ Begin FTickableObjectBase
+	virtual TStatId GetStatId() const override;
+	virtual void Tick(const float InDeltaTime) override;
+	//~ End FTickableObjectBase
+
+	/** Delete a set of items in the outliner by calling their custom delete handler */
+	void DeleteItems(TArray<FNavigationToolItemPtr> InItems);
+
+	/** Unregisters the Navigation Tool View bound to the given id */
+	void UnregisterToolView(const int32 InToolViewId);
+
+	/** Sets the given Navigation Tool View Id as the most recent Navigation Tool View */
+	void UpdateRecentToolViews(const int32 InToolViewId);
+
+	/** Executes the given predicate for each Navigation Tool View registered */
+	void ForEachToolView(const TFunction<void(const TSharedPtr<FNavigationToolView>& InToolView)>& InPredicate) const;
+
+	/** @return Number of actions that been added to the queue so far before triggering a refresh */
+	int32 GetPendingItemActionCount() const;
+
+	/** @return True if the Navigation Tool is currently in need of a refresh */
+	bool NeedsRefresh() const;
+
+	/**
+	 * Replaces the Item's Id in the Item Map. This can be due to an object item changing it's object
+	 * (e.g. a bp component getting destroyed and recreated, the item should be the same but the underlying component will not be)
+	 */
+	void NotifyItemIdChanged(const FNavigationToolItemId& OldId, const FNavigationToolItemPtr& InItem);
+
+	/** Gets the closest item to all the given items while also being their common ancestor */
+	static FNavigationToolItemPtr FindLowestCommonAncestor(const TArray<FNavigationToolItemPtr>& Items);
+
+	/** Converts the given Navigation Tool Item to a Scene Item that can be serialized in the Scene Tree */
+	static FNavigationToolSerializedItem MakeSerializedItemFromToolItem(const FNavigationToolItemPtr& InItem);
+
+	/**
+	 * Sort the given array of items based on their ordering in the Navigation Tool
+	 * @see FNavigationTool::CompareToolItemOrder
+	 */
+	static void SortItems(TArray<FNavigationToolItemPtr>& OutToolItems, const bool bInReverseOrder = false);
+
+	/** Normalizes the given Items by removing selected items that have their parent item also in the selection */
+	static void NormalizeToTopLevelSelections(TArray<FNavigationToolItemPtr>& InOutItems);
+
+	/** Have the given Selected Items sync to the USelection Instances of Mode Tools */
+	void SyncSequencerSelection(const TArray<FNavigationToolItemPtr>& InSelectedItems) const;
+
+	/** Called when the engine replaces an object. A common example is when a BP Component is destroyed, and replaced */
+	void OnObjectsReplaced(const TMap<UObject*, UObject*>& InReplacementMap);
+
+	/** Called when the engine replaces an Actor */
+	void OnActorReplaced(AActor* const InOldActor, AActor* const InNewActor);
+
+	/** Marks the Navigation Tool dirty. This triggers INavigationToolProvider::OnToolModified on next tick */
+	void SetToolModified();
+
+	void OnTreeViewChanged();
+
+private:
+	friend class FNavigationToolView;
+	friend class FNavigationToolToolbarMenu;
+
+	void OnProvidersChanged(const TSharedRef<INavigationTool>& InTool
+		, const TSharedRef<FNavigationToolProvider>& InProvider
+		, const ENavigationToolProvidersChangeType InChangeType);
+
+	void AddItem(const FNavigationToolItemPtr& InItem);
+	void RemoveItem(const FNavigationToolItemId& InItemId);
+	void ForEachItem(TFunctionRef<void(const FNavigationToolItemPtr&)> InFunc);
+
+	void SaveProviderState(const TSharedRef<FNavigationToolProvider>& InProvider);
+	void LoadProviderState(const TSharedRef<FNavigationToolProvider>& InProvider);
+
+	/** Save all Navigation Tool items */
+	void SaveSerializedTree(const bool bInResetTree);
+
+	/** Save a Navigation Tool item */
+	void SaveSerializedTreeRecursive(const FNavigationToolItemPtr& InParentItem
+		, FNavigationToolSerializedTree& InSerializedTree);
+
+	/** Load a Navigation Tool item */
+	void LoadSerializedTree(const FNavigationToolItemPtr& InParentItem
+		, FNavigationToolSerializedTree* const InSerializedTree);
+
+	USequencerSettings* GetSequencerSettings() const;
+
+	TWeakPtr<ISequencer> WeakSequencer;
+
+	TSharedRef<FNavigationToolTab> ToolTab;
+
+	/** The root of all the items in the outliner */
+	TSharedRef<FNavigationToolTreeRoot> RootItem;
+
+	/** the base command list gotten from the outliner provider to append our command list to */
+	TWeakPtr<FUICommandList> WeakBaseCommandList;
+
+	/** the map of the registered items */
+	TMap<FNavigationToolItemId, FNavigationToolItemPtr> ItemMap;
+
+	TMap<FNavigationToolItemId, FNavigationToolItemPtr> ItemsPendingAdd;
+
+	TSet<FNavigationToolItemId> ItemsPendingRemove;
+
+	/** Navigation Tool's Item Proxy Factory Registry Instance. This takes precedence over the Module's Factory Registry */
+	FNavigationToolItemProxyRegistry ItemProxyRegistry;
+	
+	/** The current pending actions before refresh is called */
+	TArray<TSharedPtr<INavigationToolAction>> PendingActions;
+
+	/** The list of objects pending selection processing. filled in when getting notifies from the USelection instances */
+	TSharedPtr<TArray<TWeakObjectPtr<UObject>>> ObjectsLastSelected;
+
+	/** The map of registered outliner views */
+	TMap<int32, TSharedPtr<FNavigationToolView>> ToolViews;
+
+	/** List of Navigation Tool View Ids in order from least recent to most recent (i.e. Index 0 is least recent) */
+	TArray<int32> RecentToolViews;
+
+	/** The current events to ignore and not handle automatically */
+	ENavigationToolIgnoreNotifyFlags IgnoreNotifyFlags = ENavigationToolIgnoreNotifyFlags::None;
+
+	/** Flag indicating whether the Navigation Tool has been changed this tick and should call INavigationToolProvider::OnToolModified next tick */
+	bool bToolDirty = false;
+	
+	/** Flag indicating Refreshing is taking place */
+	bool bRefreshing = false;
+
+	/** Flag indicating that a refresh must take place next tick */
+	bool bRefreshRequested = false;
+
+	/** Flag indicating that the Item Map is iterating */
+	bool bIteratingItemMap = false;
+
+	/** Built in "global" item type filters for all views */
+	TArray<TSharedPtr<FNavigationToolBuiltInFilter>> GlobalFilters;
+};
+
+}
diff --git a/Engine/Source/Editor/Sequencer/Private/Tools/NavigationTool/NavigationToolDefines.cpp b/Engine/Source/Editor/Sequencer/Private/Tools/NavigationTool/NavigationToolDefines.cpp
new file mode 100644
--- /dev/null
+++ b/Engine/Source/Editor/Sequencer/Private/Tools/NavigationTool/NavigationToolDefines.cpp
@@ -0,0 +1,18 @@
+﻿// Copyright Epic Games, Inc. All Rights Reserved.
+
+#include "Tools/NavigationTool/NavigationToolDefines.h"
+#include "Tools/NavigationTool/Menus/NavigationToolItemContextMenu.h"
+#include "Tools/NavigationTool/Menus/NavigationToolToolbarMenu.h"
+
+namespace UE::Sequencer::NavigationTool
+{
+	FName GetToolBarMenuName()
+	{
+		return FNavigationToolToolbarMenu::GetMenuName();
+	}
+
+	FName GetItemContextMenuName()
+	{
+		return FNavigationToolItemContextMenu::GetMenuName();
+	}
+}
diff --git a/Engine/Source/Editor/Sequencer/Private/Tools/NavigationTool/NavigationToolTab.h b/Engine/Source/Editor/Sequencer/Private/Tools/NavigationTool/NavigationToolTab.h
new file mode 100644
--- /dev/null
+++ b/Engine/Source/Editor/Sequencer/Private/Tools/NavigationTool/NavigationToolTab.h
@@ -0,0 +1,59 @@
+﻿// Copyright Epic Games, Inc. All Rights Reserved.
+
+#pragma once
+
+#include "Templates/SharedPointer.h"
+#include "Widgets/Docking/SDockTab.h"
+
+class FName;
+class FTabManager;
+class ISequencer;
+class IToolkitHost;
+class USequencerSettings;
+
+namespace UE::Sequencer::NavigationTool
+{
+
+class INavigationTool;
+
+/** Manages a Navigation Tool tab, including its visible state settings in USequencerSettings */
+class FNavigationToolTab : public TSharedFromThis<FNavigationToolTab>
+{
+public:
+	/** @return Returns the auto generated tool tab id using the passed in ISequencer.
+	 * Tool tab ids are based on the sequencer instance settings name. */
+	static FName StaticToolTabId(const ISequencer& InSequencer);
+
+	FNavigationToolTab(INavigationTool& InOwnerTool);
+
+	void Init();
+	void Shutdown();
+
+	/** @return True if the USequencerSettings setting for the Navigation Tool is set to visible */
+	bool ShouldShowToolTab() const;
+
+	/** @return True if the Navigation Tool tab is visible */
+	bool IsToolTabVisible() const;
+
+	/** Shows or hides the Navigation Tool tab and toggles the state in the saved settings */
+	void ShowHideToolTab(const bool bInVisible);
+
+	void ToggleToolTabVisible();
+
+protected:
+	TSharedPtr<IToolkitHost> GetToolkitHost() const;
+	TSharedPtr<FTabManager> GetTabManager() const;
+	USequencerSettings* GetSequencerSettings() const;
+
+	void RegisterToolTab();
+	void UnregisterToolTab();
+
+	TSharedRef<SDockTab> SpawnToolTab(const FSpawnTabArgs& InArgs);
+	void CloseToolTab();
+
+	INavigationTool& OwnerTool;
+
+	FName NavigationToolTabId;
+};
+
+}
