[Sequencer] Navigation Tool initial submit

diff --git a/Engine/Source/Editor/Sequencer/Private/Tools/NavigationTool/NavigationToolView.cpp b/Engine/Source/Editor/Sequencer/Private/Tools/NavigationTool/NavigationToolView.cpp
new file mode 100644
--- /dev/null
+++ b/Engine/Source/Editor/Sequencer/Private/Tools/NavigationTool/NavigationToolView.cpp
@@ -0,0 +1,2228 @@
+// Copyright Epic Games, Inc. All Rights Reserved.
+
+#include "NavigationToolView.h"
+#include "Framework/Application/SlateApplication.h"
+#include "Framework/Commands/GenericCommands.h"
+#include "Framework/MultiBox/MultiBoxBuilder.h"
+#include "Misc/MessageDialog.h"
+#include "MVVM/Selection/Selection.h"
+#include "Sequencer.h"
+#include "Filters/Filters/NavigationToolBuiltInFilter.h"
+#include "Styling/SlateTypes.h"
+#include "Tools/NavigationTool/Columns/NavigationToolColumn.h"
+#include "Tools/NavigationTool/Columns/NavigationToolColumnExtender.h"
+#include "Tools/NavigationTool/Customization/NavigationToolExtender.h"
+#include "Tools/NavigationTool/DragDropOps/NavigationToolItemDragDropOp.h"
+#include "Tools/NavigationTool/Filters/Filters/NavigationToolFilterBase.h"
+#include "Tools/NavigationTool/Filters/NavigationToolFilterBar.h"
+#include "Tools/NavigationTool/ItemProxies/NavigationToolComponentProxy.h"
+#include "Tools/NavigationTool/Items/NavigationToolActor.h"
+#include "Tools/NavigationTool/Items/NavigationToolComponent.h"
+#include "Tools/NavigationTool/Items/NavigationToolItemUtils.h"
+#include "Tools/NavigationTool/Items/NavigationToolSequence.h"
+#include "Tools/NavigationTool/Items/NavigationToolTreeRoot.h"
+#include "Tools/NavigationTool/Menus/NavigationToolItemContextMenu.h"
+#include "Tools/NavigationTool/NavigationTool.h"
+#include "Tools/NavigationTool/NavigationToolCommands.h"
+#include "Tools/NavigationTool/NavigationToolSettings.h"
+#include "Tools/NavigationTool/Providers/NavigationToolProvider.h"
+#include "Tools/NavigationTool/Utils/NavigationToolMiscUtils.h"
+#include "Tools/NavigationTool/Widgets/SNavigationToolView.h"
+#include "Tools/NavigationTool/Widgets/SNavigationToolTreeView.h"
+#include "Widgets/ModalTextInputDialog.h"
+#include "Widgets/SNullWidget.h"
+
+#define LOCTEXT_NAMESPACE "NavigationToolView"
+
+namespace UE::Sequencer::NavigationTool
+{
+
+FNavigationToolView::FNavigationToolView(FPrivateToken)
+	: ToolSettings(GetMutableDefault<UNavigationToolSettings>())
+	, ItemContextMenu(MakeShared<FNavigationToolItemContextMenu>())
+{
+}
+
+FNavigationToolView::~FNavigationToolView()
+{
+	if (UObjectInitialized())
+	{
+		ToolSettings->OnSettingChanged().RemoveAll(this);
+	}
+
+	if (const TSharedPtr<ISequencer> Sequencer = GetSequencer())
+	{
+		if (const TSharedPtr<FSequencerEditorViewModel> ViewModel = Sequencer->GetViewModel())
+		{
+			if (const TSharedPtr<FSequencerSelection> SequencerSelection = ViewModel->GetSelection())
+			{
+				SequencerSelection->OnChanged.AddSP(this, &FNavigationToolView::OnSequencerOutlinerSelectionChanged);
+			}
+		}
+	}
+}
+
+void FNavigationToolView::Init(const TSharedRef<FNavigationTool>& InTool, const TSharedPtr<FUICommandList>& InBaseCommandList)
+{
+	const TSharedPtr<ISequencer> Sequencer = InTool->GetSequencer();
+	if (!Sequencer.IsValid())
+	{
+		return;
+	}
+
+	WeakTool = InTool;
+
+	BindCommands(InBaseCommandList);
+
+	FilterBar = MakeShared<FNavigationToolFilterBar>(*InTool);
+	FilterBar->Init();
+	FilterBar->BindCommands(GetBaseCommandList());
+	FilterBar->OnStateChanged().AddSPLambda(this,
+		[this](const bool bInIsVisible, const EFilterBarLayout InNewLayout)
+		{
+			if (const TSharedPtr<SNavigationToolView> ToolWidget = StaticCastSharedPtr<SNavigationToolView>(GetToolWidget()))
+			{
+				ToolWidget->RebuildWidget();
+			}
+		});
+
+	if (const TSharedPtr<FSequencerEditorViewModel> ViewModel = Sequencer->GetViewModel())
+	{
+		if (const TSharedPtr<FSequencerSelection> SequencerSelection = ViewModel->GetSelection())
+		{
+			SequencerSelection->OnChanged.AddSP(this, &FNavigationToolView::OnSequencerOutlinerSelectionChanged);
+		}
+	}
+
+	ToolSettings->OnSettingChanged().AddSP(this, &FNavigationToolView::OnToolSettingsChanged);
+
+	ToolViewWidget = SNew(SNavigationToolView, SharedThis(this));
+
+	UpdateRecentViews();
+}
+
+void FNavigationToolView::CreateColumns(const TSharedRef<FNavigationToolProvider>& InProvider)
+{
+	const TSharedPtr<FNavigationTool> Tool = WeakTool.Pin();
+	if (!Tool.IsValid())
+	{
+		return;
+	}
+
+	FNavigationToolColumnExtender ColumnExtender;
+	InProvider->OnExtendColumns(ColumnExtender);
+
+	// Automatically set provider for the extended columns
+	for (const TSharedPtr<FNavigationToolColumn>& Column : ColumnExtender.GetColumns())
+	{
+		Column->SetProvider(InProvider);
+	}
+
+	// Sort and re-cache columns
+	const TArray<TSharedPtr<FNavigationToolColumn>>& ColumnsToAdd = ColumnExtender.GetColumns();
+
+	for (const TSharedPtr<FNavigationToolColumn>& Column : ColumnsToAdd)
+	{
+		const FName ColumnId = Column->GetColumnId();
+		if (!Columns.Contains(ColumnId))
+		{
+			Columns.Add(ColumnId, Column);
+		}
+	}
+
+	if (ToolViewWidget.IsValid())
+	{
+		ToolViewWidget->ReconstructColumns();
+	}
+}
+
+void FNavigationToolView::CreateDefaultColumnViews(const TSharedRef<FNavigationToolProvider>& InProvider)
+{
+	const TSharedPtr<FNavigationTool> Tool = WeakTool.Pin();
+	if (!Tool.IsValid())
+	{
+		return;
+	}
+
+	InProvider->OnExtendColumnViews(ToolSettings->GetCustomColumnViews());
+	ToolSettings->SaveConfig();
+}
+
+TSharedRef<FNavigationToolView> FNavigationToolView::CreateInstance(const int32 InToolViewId
+	, const TSharedRef<FNavigationTool>& InTool
+	, const TSharedPtr<FUICommandList>& InBaseCommandList)
+{
+	const TSharedRef<FNavigationToolView> Instance = MakeShared<FNavigationToolView>(FPrivateToken{});
+	Instance->ToolViewId = InToolViewId;
+	Instance->Init(InTool, InBaseCommandList);
+	return Instance;
+}
+
+void FNavigationToolView::PostLoad()
+{
+	if (ToolViewWidget.IsValid())
+	{
+		ToolViewWidget->ReconstructColumns();
+	}
+}
+
+void FNavigationToolView::OnToolSettingsChanged(UObject* InObject, FPropertyChangedEvent& InPropertyChangedEvent)
+{
+	RefreshTool(false);
+}
+
+void FNavigationToolView::Tick(float InDeltaTime)
+{
+	if (bRefreshRequested)
+	{
+		bRefreshRequested = false;
+		Refresh();
+	}
+
+	for (const TPair<FName, TSharedPtr<INavigationToolColumn>>& Pair : Columns)
+	{
+		if (const TSharedPtr<INavigationToolColumn>& Column = Pair.Value)
+		{
+			Column->Tick(InDeltaTime);
+		}
+	}
+
+	// Check if we have pending items to rename and we are not currently renaming an item
+	if (bRenamingItems && ItemsRemainingRename.Num() > 0 && !CurrentItemRenaming.IsValid())
+	{
+		CurrentItemRenaming = ItemsRemainingRename[0];
+		ItemsRemainingRename.RemoveAt(0);
+
+		if (CurrentItemRenaming.IsValid())
+		{
+			CurrentItemRenaming->OnRenameAction().AddSP(this, &FNavigationToolView::OnItemRenameAction);
+			CurrentItemRenaming->OnRenameAction().Broadcast(ENavigationToolRenameAction::Requested, SharedThis(this));
+		}
+	}
+
+	if (bRequestedRename)
+	{
+		bRequestedRename = false;
+		RenameSelected();
+	}
+}
+
+void FNavigationToolView::BindCommands(const TSharedPtr<FUICommandList>& InBaseCommandList)
+{
+	const FGenericCommands& GenericCommands = FGenericCommands::Get();
+	const FNavigationToolCommands& ToolCommands = FNavigationToolCommands::Get();
+
+	ViewCommandList = MakeShared<FUICommandList>();
+
+	if (InBaseCommandList.IsValid())
+	{
+		InBaseCommandList->Append(ViewCommandList.ToSharedRef());
+	}
+
+	ViewCommandList->MapAction(ToolCommands.OpenToolSettings
+		, FExecuteAction::CreateStatic(&UNavigationToolSettings::OpenEditorSettings));
+
+	ViewCommandList->MapAction(ToolCommands.Refresh
+		, FExecuteAction::CreateSP(this, &FNavigationToolView::RefreshTool, true));
+
+	ViewCommandList->MapAction(GenericCommands.Rename
+		, FExecuteAction::CreateSP(this, &FNavigationToolView::RenameSelected)
+		, FCanExecuteAction::CreateSP(this, &FNavigationToolView::CanRenameSelected));
+
+	ViewCommandList->MapAction(GenericCommands.Delete
+		, FExecuteAction::CreateSP(this, &FNavigationToolView::DeleteSelected)
+		, FCanExecuteAction::CreateSP(this, &FNavigationToolView::CanDeleteSelected));
+
+	ViewCommandList->MapAction(GenericCommands.Duplicate
+		, FExecuteAction::CreateSP(this, &FNavigationToolView::DuplicateSelected)
+		, FCanExecuteAction::CreateSP(this, &FNavigationToolView::CanDuplicateSelected));
+
+	ViewCommandList->MapAction(ToolCommands.SelectAllChildren
+		, FExecuteAction::CreateSP(this, &FNavigationToolView::SelectChildren, true)
+		, FCanExecuteAction::CreateSP(this, &FNavigationToolView::CanSelectChildren));
+
+	ViewCommandList->MapAction(ToolCommands.SelectImmediateChildren
+		, FExecuteAction::CreateSP(this, &FNavigationToolView::SelectChildren, false)
+		, FCanExecuteAction::CreateSP(this, &FNavigationToolView::CanSelectChildren));
+
+	ViewCommandList->MapAction(ToolCommands.SelectParent
+		, FExecuteAction::CreateSP(this, &FNavigationToolView::SelectParent)
+		, FCanExecuteAction::CreateSP(this, &FNavigationToolView::CanSelectParent));
+
+	ViewCommandList->MapAction(ToolCommands.SelectFirstChild
+		, FExecuteAction::CreateSP(this, &FNavigationToolView::SelectFirstChild)
+		, FCanExecuteAction::CreateSP(this, &FNavigationToolView::CanSelectFirstChild));
+
+	ViewCommandList->MapAction(ToolCommands.SelectNextSibling
+		, FExecuteAction::CreateSP(this, &FNavigationToolView::SelectSibling, +1)
+		, FCanExecuteAction::CreateSP(this, &FNavigationToolView::CanSelectSibling));
+
+	ViewCommandList->MapAction(ToolCommands.SelectPreviousSibling
+		, FExecuteAction::CreateSP(this, &FNavigationToolView::SelectSibling, -1)
+		, FCanExecuteAction::CreateSP(this, &FNavigationToolView::CanSelectSibling));
+
+	ViewCommandList->MapAction(ToolCommands.ExpandAll
+		, FExecuteAction::CreateSP(this, &FNavigationToolView::ExpandAll)
+		, FCanExecuteAction::CreateSP(this, &FNavigationToolView::CanExpandAll));
+
+	ViewCommandList->MapAction(ToolCommands.CollapseAll
+		, FExecuteAction::CreateSP(this, &FNavigationToolView::CollapseAll)
+		, FCanExecuteAction::CreateSP(this, &FNavigationToolView::CanCollapseAll));
+
+	ViewCommandList->MapAction(ToolCommands.ScrollNextSelectionIntoView
+		, FExecuteAction::CreateSP(this, &FNavigationToolView::ScrollNextIntoView)
+		, FCanExecuteAction::CreateSP(this, &FNavigationToolView::CanScrollNextIntoView));
+
+	ViewCommandList->MapAction(ToolCommands.ToggleMutedHierarchy
+		, FExecuteAction::CreateSP(this, &FNavigationToolView::ToggleMutedHierarchy)
+		, FCanExecuteAction::CreateSP(this, &FNavigationToolView::CanToggleMutedHierarchy)
+		, FIsActionChecked::CreateSP(this, &FNavigationToolView::IsMutedHierarchyActive));
+
+	ViewCommandList->MapAction(ToolCommands.ToggleAutoExpandToSelection
+		, FExecuteAction::CreateSP(this, &FNavigationToolView::ToggleAutoExpandToSelection)
+		, FCanExecuteAction::CreateSP(this, &FNavigationToolView::CanToggleAutoExpandToSelection)
+		, FIsActionChecked::CreateSP(this, &FNavigationToolView::ShouldAutoExpandToSelection));
+
+	ViewCommandList->MapAction(ToolCommands.ToggleShortNames
+		, FExecuteAction::CreateSP(this, &FNavigationToolView::ToggleUseShortNames)
+		, FCanExecuteAction::CreateSP(this, &FNavigationToolView::CanToggleUseShortNames)
+		, FIsActionChecked::CreateSP(this, &FNavigationToolView::ShouldUseShortNames));
+
+	ViewCommandList->MapAction(ToolCommands.ResetVisibleColumnSizes
+		, FExecuteAction::CreateSP(this, &FNavigationToolView::ResetVisibleColumnSizes)
+		, FCanExecuteAction::CreateSP(this, &FNavigationToolView::CanResetAllColumnSizes));
+
+	ViewCommandList->MapAction(ToolCommands.SaveCurrentColumnView
+		, FExecuteAction::CreateSP(this, &FNavigationToolView::SaveNewCustomColumnView));
+
+	ViewCommandList->MapAction(ToolCommands.FocusSingleSelection
+		, FExecuteAction::CreateSP(this, &FNavigationToolView::FocusSingleSelection)
+		, FCanExecuteAction::CreateSP(this, &FNavigationToolView::CanFocusSingleSelection));
+}
+
+TSharedPtr<FUICommandList> FNavigationToolView::GetBaseCommandList() const
+{
+	if (const TSharedPtr<FNavigationTool> Tool = WeakTool.Pin())
+	{
+		return Tool->GetBaseCommandList();
+	}
+	return nullptr;
+}
+
+void FNavigationToolView::UpdateRecentViews()
+{
+	if (const TSharedPtr<FNavigationTool> Tool = WeakTool.Pin())
+	{
+		Tool->UpdateRecentToolViews(ToolViewId);
+	}
+}
+
+bool FNavigationToolView::IsMostRecentToolView() const
+{
+	return WeakTool.IsValid() && WeakTool.Pin()->GetMostRecentToolView().Get() == this;
+}
+
+TSharedPtr<ISequencer> FNavigationToolView::GetSequencer() const
+{
+	if (const TSharedPtr<INavigationTool> Tool = GetOwnerTool())
+	{
+		return Tool->GetSequencer();
+	}
+	return nullptr;
+}
+
+TSharedPtr<INavigationTool> FNavigationToolView::GetOwnerTool() const
+{
+	return WeakTool.Pin();
+}
+
+TSharedPtr<SWidget> FNavigationToolView::GetToolWidget() const
+{
+	return ToolViewWidget;
+}
+
+TSharedPtr<SWidget> FNavigationToolView::CreateItemContextMenu()
+{
+	return ItemContextMenu->CreateMenu(SharedThis(this), SelectedItems);
+}
+
+bool FNavigationToolView::ShouldShowColumnByDefault(const TSharedPtr<INavigationToolColumn>& InColumn) const
+{
+	if (!InColumn.IsValid())
+	{
+		return false;
+	}
+
+	const TSharedPtr<FNavigationToolProvider> Provider = InColumn->GetProvider();
+	if (!Provider.IsValid())
+	{
+		return false;
+	}
+
+	const TSharedPtr<FNavigationTool> Tool = WeakTool.Pin();
+	if (!Tool.IsValid())
+	{
+		return false;
+	}
+
+	const FNavigationToolViewSaveState* const ViewSaveState = Provider->GetViewSaveState(*Tool, ToolViewId);
+	if (!ViewSaveState)
+	{
+		return false;
+	}
+
+	const FName ColumnId = InColumn->GetColumnId();
+	if (const FNavigationToolViewColumnSaveState* const FoundColumnState = ViewSaveState->ColumnsState.Find(ColumnId))
+	{
+		return FoundColumnState->bVisible;
+	}
+
+	return InColumn->ShouldShowColumnByDefault();
+}
+
+void FNavigationToolView::RequestRefresh()
+{
+	bRefreshRequested = true;
+}
+
+void FNavigationToolView::Refresh()
+{
+	UpdateRootVisibleItems();
+
+	bFilterUpdateRequested = true;
+	UpdateFilters();
+
+	UpdateItemExpansions();
+	
+	if (ToolViewWidget.IsValid())
+	{
+		ToolViewWidget->RequestTreeRefresh();
+	}
+	
+	OnToolViewRefreshed.Broadcast();
+}
+
+void FNavigationToolView::SetKeyboardFocus()
+{
+	if (ToolViewWidget.IsValid())
+	{
+		ToolViewWidget->SetKeyboardFocus();
+	}
+}
+
+void FNavigationToolView::UpdateRootVisibleItems()
+{
+	RootVisibleItems.Reset();
+	ReadOnlyItems.Reset();
+	
+	if (WeakTool.IsValid())
+	{
+		GetChildrenOfItem(WeakTool.Pin()->GetTreeRoot(), RootVisibleItems);
+	}
+}
+
+void FNavigationToolView::UpdateItemExpansions()
+{
+	TArray<FNavigationToolItemPtr> Items = RootVisibleItems;
+
+	while (Items.Num() > 0)
+	{
+		FNavigationToolItemPtr Item = Items.Pop();
+		const ENavigationToolItemFlags ItemFlags = GetViewItemFlags(Item);
+		SetItemExpansion(Item, EnumHasAnyFlags(ItemFlags, ENavigationToolItemFlags::Expanded));
+		Items.Append(Item->GetChildren());
+	}
+
+	Items = RootVisibleItems;
+
+	while (Items.Num() > 0)
+	{
+		const FNavigationToolItemPtr Item = Items.Pop();
+		if (ToolViewWidget.IsValid())
+		{
+			ToolViewWidget->UpdateItemExpansions(Item);
+		}
+		Items.Append(Item->GetChildren());
+	}
+}
+
+void FNavigationToolView::NotifyObjectsReplaced()
+{
+	if (ToolViewWidget.IsValid())
+	{
+		ToolViewWidget->Invalidate(EInvalidateWidgetReason::Paint);	
+	}
+}
+
+FNavigationToolItemPtr FNavigationToolView::GetRootItem() const
+{
+	if (WeakTool.IsValid())
+	{
+		return WeakTool.Pin()->GetTreeRoot();
+	}
+	return nullptr;
+}
+
+const TArray<FNavigationToolItemPtr>& FNavigationToolView::GetRootVisibleItems() const
+{
+	return RootVisibleItems;
+}
+
+void FNavigationToolView::SaveViewItemFlags(const FNavigationToolItemPtr& InItem, ENavigationToolItemFlags InFlags)
+{
+	if (!InItem.IsValid())
+	{
+		return;
+	}
+
+	const TSharedPtr<FNavigationToolProvider> Provider = InItem->GetProvider();
+	if (!Provider.IsValid())
+	{
+		return;
+	}
+
+	const TSharedPtr<FNavigationTool> Tool = WeakTool.Pin();
+	if (!Tool.IsValid())
+	{
+		return;
+	}
+
+	FNavigationToolSaveState* const SaveState = Provider->GetSaveState(*Tool);
+	if (!ensure(SaveState))
+	{
+		return;
+	}
+
+	if (InFlags == ENavigationToolItemFlags::None)
+	{
+		SaveState->ToolViewSaveStates[ToolViewId].ViewItemFlags.Remove(InItem->GetItemId().GetStringId());
+	}
+	else
+	{
+		SaveState->ToolViewSaveStates[ToolViewId].ViewItemFlags.Add(InItem->GetItemId().GetStringId(), InFlags);
+	}
+}
+
+ENavigationToolItemFlags FNavigationToolView::GetViewItemFlags(const FNavigationToolItemPtr& InItem) const
+{
+	if (!InItem.IsValid())
+	{
+		return ENavigationToolItemFlags::None;
+	}
+
+	const TSharedPtr<FNavigationToolProvider> Provider = InItem->GetProvider();
+	if (!Provider.IsValid())
+	{
+		return ENavigationToolItemFlags::None;
+	}
+
+	const TSharedPtr<FNavigationTool> Tool = WeakTool.Pin();
+	if (!Tool.IsValid())
+	{
+		return ENavigationToolItemFlags::None;
+	}
+
+	FNavigationToolViewSaveState* const ViewSaveState = Provider->GetViewSaveState(*Tool, ToolViewId);
+	if (!ensure(ViewSaveState))
+	{
+		return ENavigationToolItemFlags::None;
+	}
+
+	if (const ENavigationToolItemFlags* const OverrideFlags = ViewSaveState->ViewItemFlags.Find(InItem->GetItemId().GetStringId()))
+	{
+		return *OverrideFlags;
+	}
+
+	return ENavigationToolItemFlags::None;
+}
+
+void FNavigationToolView::GetChildrenOfItem(const FNavigationToolItemPtr InItem, TArray<FNavigationToolItemPtr>& OutChildren) const
+{
+	static const TSet<FNavigationToolItemPtr> EmptySet;
+	GetChildrenOfItem(InItem, OutChildren, ENavigationToolItemViewMode::ItemTree, EmptySet);
+}
+
+void FNavigationToolView::GetChildrenOfItem(const FNavigationToolItemPtr& InItem
+	, TArray<FNavigationToolItemPtr>& OutChildren
+	, ENavigationToolItemViewMode InViewMode
+	, const TSet<FNavigationToolItemPtr>& InRecursionDisallowedItems) const
+{
+	if (!InItem.IsValid())
+	{
+		return;
+	}
+
+	const TSharedPtr<FNavigationToolProvider> Provider = InItem->GetProvider();
+	if (!Provider.IsValid() && InItem->GetItemId() != FNavigationToolItemId::RootId)
+	{
+		UE_LOG(LogTemp, Warning, TEXT("Navigation Tool Item Id \"%s\" has no provider, but is a root item!"), *InItem->GetItemId().GetStringId());
+		return;
+	}
+
+	for (const FNavigationToolItemPtr& ChildItem : InItem->GetChildren())
+	{
+		if (!ChildItem.IsValid())
+		{
+			continue;
+		}
+
+		const bool bShouldShowItem = ShouldShowItem(ChildItem, true, InViewMode);
+		if (bShouldShowItem)
+		{
+			// If current item is visible in outliner, add it to the children
+			OutChildren.Add(ChildItem);
+		}
+		else if (!InRecursionDisallowedItems.Contains(ChildItem))
+		{
+			TArray<FNavigationToolItemPtr> GrandChildren;
+
+			// For Muted Hierarchy to be in effect, not only does it have to be on
+			// but also the item should be shown (without counting the filter pass)
+			const bool bShouldUseMutedHierarchy = ToolSettings->ShouldUseMutedHierarchy();
+			const bool bShouldShowItemWithoutFilters = ShouldShowItem(ChildItem, false, InViewMode);
+			const bool bShouldMuteItem = bShouldUseMutedHierarchy && bShouldShowItemWithoutFilters;
+
+			// If it's muted hierarchy, there might be ONLY grand children that are just visible in other view modes, 
+			// so instead of just filtering out the child item, check that there are no grand children from other view modes passing filter tests
+			// If it's NOT muted hierarchy, just get the grand children visible in the requested view mode, as this ChildItem is guaranteed to be hidden
+			const ENavigationToolItemViewMode ViewModeToUse = bShouldMuteItem ? ENavigationToolItemViewMode::All : InViewMode;
+
+			GetChildrenOfItem(ChildItem, GrandChildren, ViewModeToUse, InRecursionDisallowedItems);
+
+			if (!GrandChildren.IsEmpty())
+			{
+				if (bShouldMuteItem)
+				{
+					ReadOnlyItems.Add(ChildItem);	
+					OutChildren.Add(ChildItem);
+				}
+				else
+				{
+					// We can append them knowing that the ViewMode to use is the one passed in and there's no
+					// child that leaked from another view mode
+					ensure(ViewModeToUse == InViewMode);
+					OutChildren.Append(GrandChildren);
+				}
+			}
+		}
+	}
+}
+
+FLinearColor FNavigationToolView::GetItemBrushColor(const FNavigationToolItemPtr InItem) const
+{
+	if (InItem.IsValid())
+	{
+		FLinearColor OutColor = InItem->GetItemBrushColor();
+
+		// If NextSelectedItemIntoView is valid, it means we're scrolling items into view with Next/Previous, 
+		// so Make everything that's not the Current Item a bit more translucent to make the Current Item stand out
+		if (SortedSelectedItems.IsValidIndex(NextSelectedItemIntoView)
+			&& SortedSelectedItems[NextSelectedItemIntoView] != InItem)
+		{
+			OutColor.A *= 0.5f;
+		}
+
+		return OutColor;
+	}
+
+	return FStyleColors::White.GetSpecifiedColor();
+}
+
+TArray<FNavigationToolItemPtr> FNavigationToolView::GetSelectedItems() const
+{
+	return SelectedItems;
+}
+
+int32 FNavigationToolView::GetViewSelectedItemCount() const
+{
+	return SelectedItems.Num();
+}
+
+int32 FNavigationToolView::CalculateVisibleItemCount() const
+{
+	TArray<FNavigationToolItemPtr> RemainingItems = RootVisibleItems;
+	
+	int32 VisibleItemCount = RemainingItems.Num();
+
+	while (RemainingItems.Num() > 0)
+	{
+		const FNavigationToolItemPtr Item = RemainingItems.Pop();
+		
+		TArray<FNavigationToolItemPtr> ChildItems;
+		GetChildrenOfItem(Item, ChildItems);
+		VisibleItemCount += ChildItems.Num();
+		RemainingItems.Append(MoveTemp(ChildItems));
+	}
+
+	// Remove the Read Only Items as they are filtered out items that are still shown because of Hierarchy Viz
+	VisibleItemCount -= ReadOnlyItems.Num();
+
+	return VisibleItemCount;
+}
+
+void FNavigationToolView::SelectItems(TArray<FNavigationToolItemPtr> InItems, ENavigationToolItemSelectionFlags InFlags)
+{
+	// Remove Duplicate Items
+	TSet<FNavigationToolItemPtr> SeenItems;
+	SeenItems.Reserve(InItems.Num());
+	for (TArray<FNavigationToolItemPtr>::TIterator Iter(InItems); Iter; ++Iter)
+	{
+		if (SeenItems.Contains(*Iter))
+		{
+			Iter.RemoveCurrent();
+		}
+		else
+		{
+			SeenItems.Add(*Iter);
+		}
+	}
+
+	// Add the Children of the Items given
+	if (EnumHasAnyFlags(InFlags, ENavigationToolItemSelectionFlags::IncludeChildren))
+	{
+		TArray<FNavigationToolItemPtr> ChildItemsRemaining(InItems);
+		while (ChildItemsRemaining.Num() > 0)
+		{
+			if (FNavigationToolItemPtr ChildItem = ChildItemsRemaining.Pop())
+			{
+				TArray<FNavigationToolItemPtr> Children;
+				GetChildrenOfItem(ChildItem, Children);
+
+				InItems.Append(Children);
+				ChildItemsRemaining.Append(Children);
+			}
+		}
+	}
+
+	if (EnumHasAnyFlags(InFlags, ENavigationToolItemSelectionFlags::AppendToCurrentSelection))
+	{
+		// Remove all repeated items to avoid duplicated entries
+		SelectedItems.RemoveAll([&SeenItems](const FNavigationToolItemPtr& Item)
+		{
+			return SeenItems.Contains(Item);
+		});	
+		TArray<FNavigationToolItemPtr> Items = MoveTemp(SelectedItems);
+		Items.Append(MoveTemp(InItems));
+		InItems = MoveTemp(Items);
+	}
+
+	if (!InItems.IsEmpty() && EnumHasAnyFlags(InFlags, ENavigationToolItemSelectionFlags::ScrollIntoView))
+	{
+		ScrollItemIntoView(InItems[0]);
+	}
+
+	const bool bSignalSelectionChange = EnumHasAnyFlags(InFlags, ENavigationToolItemSelectionFlags::SignalSelectionChange);
+	SetItemSelectionImpl(MoveTemp(InItems), bSignalSelectionChange);
+}
+
+void FNavigationToolView::ClearItemSelection(bool bSignalSelectionChange)
+{
+	SetItemSelectionImpl({}, bSignalSelectionChange);
+}
+
+void FNavigationToolView::SetItemSelectionImpl(TArray<FNavigationToolItemPtr>&& InItems, bool bSignalSelectionChange)
+{
+	if (ShouldAutoExpandToSelection())
+	{
+		for (const FNavigationToolItemPtr& Item : InItems)
+		{
+			SetParentItemExpansions(Item, true);
+		}
+	}
+
+	SelectedItems = MoveTemp(InItems);
+
+	if (ToolViewWidget.IsValid())
+	{
+		ToolViewWidget->SetItemSelection(SelectedItems, bSignalSelectionChange);
+	}
+	else if (bSignalSelectionChange)
+	{
+		NotifyItemSelectionChanged(SelectedItems, nullptr, true);
+	}
+	
+	Refresh();
+}
+
+void FNavigationToolView::NotifyItemSelectionChanged(const TArray<FNavigationToolItemPtr>& InSelectedItems
+	, const FNavigationToolItemPtr& InItem
+	, bool bUpdateModeTools)
+{
+	if (bSyncingItemSelection)
+	{
+		return;
+	}
+	TGuardValue<bool> Guard(bSyncingItemSelection, true);
+	
+	SelectedItems            = InSelectedItems;
+	SortedSelectedItems      = InSelectedItems;
+	NextSelectedItemIntoView = -1;
+	
+	FNavigationTool::SortItems(SortedSelectedItems);
+
+	//If we have Pending Items Remaining but we switched Selection via Navigation, treat it as "I want to rename this too"
+	if (bRenamingItems && InItem.IsValid() && InItem != CurrentItemRenaming)
+	{
+		bRequestedRename = true;
+	}
+	
+	const TSharedPtr<FNavigationTool> Tool = WeakTool.Pin();
+	if (bUpdateModeTools && Tool.IsValid())
+	{
+		Tool->SyncSequencerSelection(SelectedItems);
+		Tool->SelectItems(SelectedItems, ENavigationToolItemSelectionFlags::None);
+	}
+}
+
+bool FNavigationToolView::IsItemReadOnly(const FNavigationToolItemPtr& InItem) const
+{
+	return ReadOnlyItems.Contains(InItem);;
+}
+
+bool FNavigationToolView::CanSelectItem(const FNavigationToolItemPtr& InItem) const
+{
+	const bool bIsSelectable = InItem.IsValid() && InItem->IsSelectable();
+	return bIsSelectable && !IsItemReadOnly(InItem);
+}
+
+bool FNavigationToolView::IsItemSelected(const FNavigationToolItemPtr& InItem) const
+{
+	return SelectedItems.Contains(InItem);
+}
+
+bool FNavigationToolView::IsItemExpanded(const FNavigationToolItemPtr& InItem, const bool bInUseFilter) const
+{
+	// Don't continue if Item should be hidden in view.
+	// the tree view still calls OnItemExpansionChanged even if it doesn't contain the item
+	// so this preemptive check is needed
+	if (!ShouldShowItem(InItem, bInUseFilter, ENavigationToolItemViewMode::ItemTree))
+	{
+		return false;
+	}
+
+	if (ToolViewWidget.IsValid())
+	{
+		return ToolViewWidget->IsItemExpanded(InItem);
+	}
+
+	return false;
+}
+
+void FNavigationToolView::SetItemExpansion(const FNavigationToolItemPtr& InItem, const bool bInExpand, const bool bInUseFilter)
+{
+	// Don't continue if Item should be hidden in view.
+	// the tree view still calls OnItemExpansionChanged even if it doesn't contain the item
+	// so this preemptive check is needed
+	if (!ShouldShowItem(InItem, bInUseFilter, ENavigationToolItemViewMode::ItemTree))
+	{
+		return;
+	}
+
+	if (ToolViewWidget.IsValid())
+	{
+		ToolViewWidget->SetItemExpansion(InItem, bInExpand);
+	}
+	else
+	{
+		OnItemExpansionChanged(InItem, bInExpand);
+	}
+}
+
+void FNavigationToolView::SetItemExpansionRecursive(const FNavigationToolItemPtr InItem, const bool bInExpand)
+{
+	SetItemExpansion(InItem, bInExpand, false);
+	
+	for (const FNavigationToolItemPtr& Child : InItem->GetChildren())
+	{
+		if (Child.IsValid())
+		{
+			SetItemExpansionRecursive(Child, bInExpand);
+		}
+	}
+}
+
+void FNavigationToolView::SetParentItemExpansions(const FNavigationToolItemPtr& InItem, const bool bInExpand)
+{
+	if (!InItem.IsValid())
+	{
+		return;
+	}
+	
+	TArray<FNavigationToolItemPtr> ItemsToExpand;
+	
+	// Don't auto expand at all if there's a parent preventing it
+	FNavigationToolItemPtr ParentItem = InItem->GetParent();
+	while (ParentItem.IsValid())
+	{
+		if (!ParentItem->CanAutoExpand())
+		{
+			return;
+		}
+		ItemsToExpand.Add(ParentItem);
+		ParentItem = ParentItem->GetParent();
+	}
+
+	for (const FNavigationToolItemPtr& Item : ItemsToExpand)
+	{
+		SetItemExpansion(Item, bInExpand);
+	}
+}
+
+void FNavigationToolView::OnItemExpansionChanged(const FNavigationToolItemPtr InItem, const bool bInIsExpanded)
+{
+	const ENavigationToolItemFlags CurrentFlags = GetViewItemFlags(InItem);
+	
+	ENavigationToolItemFlags TargetFlags = CurrentFlags;
+	
+	if (bInIsExpanded)
+	{
+		TargetFlags |= ENavigationToolItemFlags::Expanded;
+	}
+	else
+	{
+		TargetFlags &= ~ENavigationToolItemFlags::Expanded;
+	}
+
+	SaveViewItemFlags(InItem, TargetFlags);
+
+	if (CurrentFlags != TargetFlags)
+	{
+		InItem->OnExpansionChanged().Broadcast(SharedThis(this), bInIsExpanded);
+	}
+}
+
+bool FNavigationToolView::ShouldShowItem(const FNavigationToolItemPtr& InItem, bool bInUseFilters
+	, ENavigationToolItemViewMode InViewMode) const
+{
+	if (!InItem.IsValid())
+	{
+		return false;
+	}
+
+	if (InItem->IsA<FNavigationToolTreeRoot>())
+	{
+		return true;
+	}
+
+	const TSharedPtr<FNavigationTool> Tool = WeakTool.Pin();
+	if (!Tool.IsValid())
+	{
+		return false;
+	}
+
+	if (!InItem->IsAllowedInTool())
+	{
+		//UE_LOG(LogTemp, Warning, TEXT("Navigation Tool Item Id \"%s\" is hidden: Not Allowed In Tool"), *InItem->GetItemId().GetStringId());
+		return false;
+	}
+
+	if (!InItem->IsViewModeSupported(InViewMode, *this))
+	{
+		//UE_LOG(LogTemp, Warning, TEXT("Navigation Tool Item Id \"%s\" is hidden: View Mode Not Supported"), *InItem->GetItemId().GetStringId());
+		return false;
+	}
+
+	// Allow providers to determine whether the item should be hidden
+	bool bProviderShouldHideItem = false;
+	FNavigationToolExtender::ForEachToolProvider(*Tool, [&bProviderShouldHideItem, &InItem]
+		(const TSharedRef<FNavigationToolProvider>& InToolProvider)
+		{
+			if (InToolProvider->ShouldHideItem(InItem))
+			{
+				bProviderShouldHideItem = true;
+				return false;
+			}
+			return true;
+		});
+	if (bProviderShouldHideItem)
+	{
+		//UE_LOG(LogTemp, Warning, TEXT("Navigation Tool Item Id \"%s\" is hidden: Provider Should Hide Item"), *InItem->GetItemId().GetStringId());
+		return false;
+	}
+
+	// Extra pass for Non-Item Proxies that are parented under an Item Proxy
+	// Hiding an Item Proxy Type should affect all the rest of the items below it
+	if (!InItem->IsA<FNavigationToolItemProxy>())
+	{
+		FNavigationToolItemPtr ItemParent = InItem->GetParent();
+		while (ItemParent.IsValid())
+		{
+			if (ItemParent->IsA<FNavigationToolItemProxy>())
+			{
+				// Stop at the first Item Proxy parent found
+				break;
+			}
+			ItemParent = ItemParent->GetParent();
+		}
+	}
+
+	/** All global filters must fail to hide the item */
+	bool bGlobalFilterOut = false;
+
+	for (const TSharedPtr<FNavigationToolBuiltInFilter>& GlobalFilter : Tool->GlobalFilters)
+	{
+		if (!GlobalFilter->IsActive() && GlobalFilter->PassesFilter(InItem))
+		{
+			bGlobalFilterOut = true;
+			break;
+		}
+	}
+
+	if (bGlobalFilterOut)
+	{
+		return false;
+	}
+
+	if (bInUseFilters && FilterBar->GetFilterData().IsFilteredOut(InItem))
+	{
+		//UE_LOG(LogTemp, Warning, TEXT("Navigation Tool Item Id \"%s\" is hidden: Filtered Out"), *InItem->GetItemId().GetStringId());
+		return false;
+	}
+
+	return true;
+}
+
+int32 FNavigationToolView::GetVisibleChildIndex(const FNavigationToolItemPtr& InParentItem, const FNavigationToolItemPtr& InChildItem) const
+{
+	if (InParentItem.IsValid())
+	{
+		TArray<FNavigationToolItemPtr> Children;
+		GetChildrenOfItem(InParentItem, Children);
+		return Children.Find(InChildItem);
+	}
+	return INDEX_NONE;
+}
+
+FNavigationToolItemPtr FNavigationToolView::GetVisibleChildAt(const FNavigationToolItemPtr& InParentItem, int32 InChildIndex) const
+{
+	if (InParentItem.IsValid())
+	{
+		TArray<FNavigationToolItemPtr> Children;
+		GetChildrenOfItem(InParentItem, Children);
+		if (Children.IsValidIndex(InChildIndex))
+		{
+			return Children[InChildIndex];
+		}
+	}
+	return nullptr;
+}
+
+bool FNavigationToolView::IsToolLocked() const
+{
+	return WeakTool.IsValid() && WeakTool.Pin()->IsToolLocked();
+}
+
+void FNavigationToolView::ShowColumn(const TSharedPtr<INavigationToolColumn>& InColumn)
+{
+	const FName ColumnId = InColumn->GetColumnId();
+	ShowColumnById(ColumnId);
+}
+
+void FNavigationToolView::ShowColumnById(const FName& InColumnId)
+{
+	if (ToolViewWidget.IsValid())
+	{
+		ToolViewWidget->ShowHideColumn(InColumnId, true);
+	}
+
+	SaveColumnState(InColumnId);
+}
+
+void FNavigationToolView::HideColumn(const TSharedPtr<INavigationToolColumn>& InColumn)
+{
+	const FName ColumnId = InColumn->GetColumnId();
+
+	if (ToolViewWidget.IsValid())
+	{
+		ToolViewWidget->ShowHideColumn(ColumnId, false);
+	}
+
+	SaveColumnState(ColumnId);
+}
+
+bool FNavigationToolView::IsColumnVisible(const TSharedPtr<INavigationToolColumn>& InColumn) const
+{
+	const TSharedPtr<FNavigationToolProvider> Provider = InColumn->GetProvider();
+	if (!Provider.IsValid())
+	{
+		return false;
+	}
+
+	const TSharedPtr<FNavigationTool> Tool = WeakTool.Pin();
+	if (!Tool.IsValid())
+	{
+		return false;
+	}
+
+	FNavigationToolViewSaveState* const ViewSaveState = Provider->GetViewSaveState(*Tool, ToolViewId);
+	if (!ViewSaveState)
+	{
+		return false;
+	}
+
+	const FName ColumnId = InColumn->GetColumnId();
+	if (const FNavigationToolViewColumnSaveState* const FoundColumnState = ViewSaveState->ColumnsState.Find(ColumnId))
+	{
+		return FoundColumnState->bVisible;
+	}
+
+	return false;
+}
+
+ENavigationToolItemViewMode FNavigationToolView::GetItemDefaultViewMode() const
+{
+	return ToolSettings->GetItemDefaultViewMode();
+}
+
+ENavigationToolItemViewMode FNavigationToolView::GetItemProxyViewMode() const
+{
+	return ToolSettings->GetItemProxyViewMode();
+}
+
+void FNavigationToolView::ToggleViewModeSupport(ENavigationToolItemViewMode& InOutViewMode, ENavigationToolItemViewMode InFlags)
+{
+	ToolSettings->ToggleViewModeSupport(InOutViewMode, InFlags);
+	Refresh();
+}
+
+void FNavigationToolView::ToggleItemDefaultViewModeSupport(ENavigationToolItemViewMode InFlags)
+{
+	ToolSettings->ToggleItemDefaultViewModeSupport(InFlags);
+	Refresh();
+}
+
+void FNavigationToolView::ToggleItemProxyViewModeSupport(ENavigationToolItemViewMode InFlags)
+{
+	ToolSettings->ToggleItemProxyViewModeSupport(InFlags);
+	Refresh();
+}
+
+ECheckBoxState FNavigationToolView::GetViewModeCheckState(ENavigationToolItemViewMode InViewMode, ENavigationToolItemViewMode InFlags) const
+{
+	const ENavigationToolItemViewMode Result = InViewMode & InFlags;
+	
+	if (Result == InFlags)
+	{
+		return ECheckBoxState::Checked;
+	}
+	
+	if (Result != ENavigationToolItemViewMode::None)
+	{
+		return ECheckBoxState::Undetermined;
+	}
+	
+	return ECheckBoxState::Unchecked;
+}
+
+ECheckBoxState FNavigationToolView::GetItemDefaultViewModeCheckState(ENavigationToolItemViewMode InFlags) const
+{
+	return GetViewModeCheckState(ToolSettings->GetItemDefaultViewMode(), InFlags);
+}
+
+ECheckBoxState FNavigationToolView::GetItemProxyViewModeCheckState(ENavigationToolItemViewMode InFlags) const
+{
+	return GetViewModeCheckState(ToolSettings->GetItemProxyViewMode(), InFlags);
+}
+
+void FNavigationToolView::ToggleMutedHierarchy()
+{
+	ToolSettings->SetUseMutedHierarchy(!ToolSettings->ShouldUseMutedHierarchy());
+	Refresh();
+}
+
+bool FNavigationToolView::IsMutedHierarchyActive() const
+{
+	return ToolSettings->ShouldUseMutedHierarchy();
+}
+
+void FNavigationToolView::ToggleAutoExpandToSelection()
+{
+	ToolSettings->SetAutoExpandToSelection(!ToolSettings->ShouldAutoExpandToSelection());
+	Refresh();
+}
+
+bool FNavigationToolView::ShouldAutoExpandToSelection() const
+{
+	return ToolSettings->ShouldAutoExpandToSelection();
+}
+
+void FNavigationToolView::ToggleUseShortNames()
+{
+	ToolSettings->SetUseShortNames(!ToolSettings->ShouldUseShortNames());
+	Refresh();
+}
+
+bool FNavigationToolView::ShouldUseShortNames() const
+{
+	return ToolSettings->ShouldUseShortNames();
+}
+
+void FNavigationToolView::ToggleShowItemFilters()
+{
+	//Note: Not Marking Navigation Tool Instance as Modified because this is not saved.
+	bShowItemFilters = !bShowItemFilters;
+}
+
+void FNavigationToolView::ToggleShowItemColumns()
+{
+	bShowItemColumns = !bShowItemColumns;
+}
+
+void FNavigationToolView::SetItemTypeHidden(const FName InItemTypeName, const bool bInHidden)
+{
+	if (IsItemTypeHidden(InItemTypeName) != bInHidden)
+	{
+		if (bInHidden)
+		{
+			HiddenItemTypes.Add(InItemTypeName);
+		}
+		else
+		{
+			HiddenItemTypes.Remove(InItemTypeName);
+		}
+		RequestRefresh();
+	}	
+}
+
+void FNavigationToolView::ToggleHideItemTypes(const FName InItemTypeName)
+{
+	SetItemTypeHidden(InItemTypeName, !IsItemTypeHidden(InItemTypeName));
+}
+
+ECheckBoxState FNavigationToolView::GetToggleHideItemTypesState(const FName InItemTypeName) const
+{
+	return IsItemTypeHidden(InItemTypeName) ? ECheckBoxState::Unchecked : ECheckBoxState::Checked;
+}
+
+bool FNavigationToolView::IsItemTypeHidden(const FName InItemTypeName) const
+{
+	return HiddenItemTypes.Contains(InItemTypeName);
+}
+
+bool FNavigationToolView::IsItemTypeHidden(const FNavigationToolItemPtr& InItem) const
+{
+	return IsItemTypeHidden(InItem->GetTypeId().ToName());
+}
+
+void FNavigationToolView::OnDragEnter(const FDragDropEvent& InDragDropEvent, const FNavigationToolItemPtr InTargetItem)
+{
+	if (!InTargetItem.IsValid() && WeakTool.IsValid())
+	{
+		TSharedRef<FNavigationToolItem> TreeRoot = WeakTool.Pin()->GetTreeRoot();
+		const bool bCanAcceptDrop = TreeRoot->CanAcceptDrop(InDragDropEvent, EItemDropZone::OntoItem).IsSet();
+		SetDragIntoTreeRoot(bCanAcceptDrop);
+	}
+	else
+	{
+		SetDragIntoTreeRoot(false);
+	}
+}
+
+void FNavigationToolView::OnDragLeave(const FDragDropEvent& InDragDropEvent, const FNavigationToolItemPtr InTargetItem)
+{
+	// If drag left an item, set the drag into tree root to false (this will set it back to false if a valid item receives DragEnter)
+	SetDragIntoTreeRoot(InTargetItem.IsValid());
+}
+
+FReply FNavigationToolView::OnDragDetected(const FGeometry& InGeometry
+	, const FPointerEvent& InMouseEvent
+	, const FNavigationToolItemPtr InTargetItem)
+{
+	if (!IsToolLocked())
+	{
+		// Only Select Target if it hasn't already been selected
+		if (!IsItemSelected(InTargetItem))
+		{
+			const ENavigationToolItemSelectionFlags SelectionFlags = InMouseEvent.IsControlDown()
+				? ENavigationToolItemSelectionFlags::AppendToCurrentSelection
+				: ENavigationToolItemSelectionFlags::None;
+
+			SelectItems({ InTargetItem }, SelectionFlags);
+		}
+
+		// Get all Selected Items that are in a state where they can be selected again (i.e. not Read Only)
+		TArray<FNavigationToolItemPtr> Items = GetSelectedItems();
+		Items.RemoveAll([this](const FNavigationToolItemPtr& InItem)->bool
+		{
+			return !CanSelectItem(InItem);
+		});
+
+		if (Items.Num() > 0)
+		{
+			const ENavigationToolDragDropActionType ActionType = InMouseEvent.IsAltDown()
+				? ENavigationToolDragDropActionType::Copy
+				: ENavigationToolDragDropActionType::Move;
+
+			const TSharedRef<FNavigationToolItemDragDropOp> DragDropOp = FNavigationToolItemDragDropOp::New(Items, SharedThis(this), ActionType);
+			return FReply::Handled().BeginDragDrop(DragDropOp);
+		}
+	}
+	return FReply::Unhandled();
+}
+
+FReply FNavigationToolView::OnDrop(const FDragDropEvent& InDragDropEvent, EItemDropZone InDropZone, const FNavigationToolItemPtr InTargetItem)
+{
+	SetDragIntoTreeRoot(false);
+
+	if (InTargetItem.IsValid())
+	{
+		return InTargetItem->AcceptDrop(InDragDropEvent, InDropZone);
+	}
+
+	FNavigationToolItemPtr TreeRoot;
+
+	if (WeakTool.IsValid())
+	{
+		TreeRoot = WeakTool.Pin()->GetTreeRoot();
+	}
+
+	if (TreeRoot.IsValid() && TreeRoot->CanAcceptDrop(InDragDropEvent, EItemDropZone::OntoItem))
+	{
+		return TreeRoot->AcceptDrop(InDragDropEvent, EItemDropZone::OntoItem);
+	}
+	
+	return FReply::Unhandled();
+}
+
+TOptional<EItemDropZone> FNavigationToolView::OnCanDrop(const FDragDropEvent& InDragDropEvent
+	, EItemDropZone InDropZone
+	, const FNavigationToolItemPtr InTargetItem) const
+{	
+	if (!IsToolLocked() && InTargetItem.IsValid() && CanSelectItem(InTargetItem))
+	{
+		return InTargetItem->CanAcceptDrop(InDragDropEvent, InDropZone);
+	}
+	return TOptional<EItemDropZone>();
+}
+
+void FNavigationToolView::SetDragIntoTreeRoot(bool bInIsDraggingIntoTreeRoot)
+{
+	if (ToolViewWidget.IsValid())
+	{
+		ToolViewWidget->SetTreeBorderVisibility(bInIsDraggingIntoTreeRoot);
+	}
+}
+
+void FNavigationToolView::RenameSelected()
+{
+	if (const TSharedPtr<FNavigationTool> Tool = WeakTool.Pin())
+	{
+		TArray<FNavigationToolItemPtr> Items = GetSelectedItems();
+		
+		if (Items.IsEmpty())
+		{
+			return;
+		}
+	
+		// Assume we have an item currently renaming
+		ResetRenaming();
+
+		// Remove Items that are Invalid or can't be renamed
+		Items.RemoveAll([](const FNavigationToolItemPtr& InItem)
+			{
+				if (InItem.IsValid())
+				{
+					if (const IRenameableExtension* const RenameableExtension = InItem->CastTo<IRenameableExtension>())
+					{
+						return !RenameableExtension->CanRename();
+					}
+				}
+				return true;
+			});
+		
+		ItemsRemainingRename = MoveTemp(Items);
+
+		if (ItemsRemainingRename.Num() > 0)
+		{
+			FNavigationTool::SortItems(ItemsRemainingRename);
+			bRenamingItems = true;
+		}
+	}
+}
+
+void FNavigationToolView::ResetRenaming()
+{
+	if (CurrentItemRenaming.IsValid())
+	{
+		CurrentItemRenaming->OnRenameAction().RemoveAll(this);
+		CurrentItemRenaming.Reset();
+	}
+
+	if (ItemsRemainingRename.IsEmpty())
+	{
+		bRenamingItems = false;
+	}
+}
+
+void FNavigationToolView::OnItemRenameAction(const ENavigationToolRenameAction InRenameAction, const TSharedPtr<INavigationToolView>& InToolView)
+{
+	if (InToolView.Get() != this)
+	{
+		return;
+	}
+	
+	switch (InRenameAction)
+	{
+	case ENavigationToolRenameAction::None:
+		break;
+	
+	case ENavigationToolRenameAction::Requested:
+		break;
+	
+	case ENavigationToolRenameAction::Cancelled:
+		ItemsRemainingRename.Reset();
+		ResetRenaming();
+		break;
+	
+	case ENavigationToolRenameAction::Completed:
+		ResetRenaming();
+		break;
+	
+	default:
+		break;
+	}
+}
+
+bool FNavigationToolView::CanRenameSelected() const
+{
+	for (const FNavigationToolItemPtr& Item : SelectedItems)
+	{
+		if (!Item.IsValid())
+		{
+			return false;
+		}
+
+		if (const IRenameableExtension* const RenameableExtension = Item->CastTo<IRenameableExtension>())
+		{
+			if (!RenameableExtension->CanRename())
+			{
+				return false;
+			}
+		}
+	}
+
+	return SelectedItems.Num() > 0;
+}
+
+void FNavigationToolView::DeleteSelected()
+{
+	if (const TSharedPtr<FNavigationTool> Tool = WeakTool.Pin())
+	{
+		TArray<FNavigationToolItemPtr> Items = GetSelectedItems();
+
+		Items.RemoveAll([](const FNavigationToolItemPtr& InItem)
+		{
+			return !InItem.IsValid() || !InItem->CanDelete();
+		});
+
+		if (Items.IsEmpty())
+		{
+			return;
+		}
+
+		Tool->DeleteItems(Items);
+	}
+}
+
+bool FNavigationToolView::CanDeleteSelected() const
+{
+	for (const FNavigationToolItemPtr& Item : GetSelectedItems())
+	{
+		if (Item && Item->CanDelete())
+		{
+			return true;
+		}
+	}
+
+	return false;
+}
+
+void FNavigationToolView::DuplicateSelected()
+{
+	if (const TSharedPtr<FNavigationTool> Tool = WeakTool.Pin())
+	{
+		//Tool->DuplicateItems(GetSelectedItems(), nullptr, TOptional<EItemDropZone>());
+	}
+}
+
+bool FNavigationToolView::CanDuplicateSelected() const
+{
+	for (const FNavigationToolItemPtr& Item : GetSelectedItems())
+	{
+		if (Item && Item->IsA<FNavigationToolActor>())
+		{
+			return true;
+		}
+	}
+	return false;
+}
+
+void FNavigationToolView::SelectChildren(bool bIsRecursive)
+{
+	TArray<FNavigationToolItemPtr> ItemsToSelect;
+	TArray<FNavigationToolItemPtr> RemainingItems = GetSelectedItems();
+
+	while (RemainingItems.Num() > 0)
+	{
+		//Note: Pop here will affect order of Children in Selection
+		const FNavigationToolItemPtr ParentItem = RemainingItems.Pop();
+
+		TArray<FNavigationToolItemPtr> ChildItems;
+		GetChildrenOfItem(ParentItem, ChildItems);
+		if (bIsRecursive)
+		{
+			RemainingItems.Append(ChildItems);
+		}
+		ItemsToSelect.Append(ChildItems);
+	}
+
+	SelectItems(ItemsToSelect
+		, ENavigationToolItemSelectionFlags::AppendToCurrentSelection | ENavigationToolItemSelectionFlags::SignalSelectionChange);
+}
+
+bool FNavigationToolView::CanSelectChildren() const
+{
+	return GetViewSelectedItemCount() > 0;
+}
+
+void FNavigationToolView::SelectParent()
+{
+	const TSet<FNavigationToolItemPtr> Items(GetSelectedItems());
+
+	TSet<FNavigationToolItemPtr> ParentItemsToSelect;
+	ParentItemsToSelect.Reserve(Items.Num());
+
+	const FNavigationToolItemPtr RootItem = GetRootItem();
+	
+	//Add only Valid Parents that are not Root and are not part of the Original Selection!
+	for (const FNavigationToolItemPtr& Item : Items)
+	{
+		if (Item.IsValid())
+		{
+			const FNavigationToolItemPtr ParentItem = Item->GetParent();
+			if (ParentItem.IsValid() && ParentItem != RootItem && !Items.Contains(ParentItem))
+			{
+				ParentItemsToSelect.Add(ParentItem);
+			}
+		}
+	}
+
+	SortAndSelectItems(ParentItemsToSelect.Array());
+}
+
+bool FNavigationToolView::CanSelectParent() const
+{
+	return GetViewSelectedItemCount() == 1;
+}
+
+void FNavigationToolView::SelectFirstChild()
+{
+	const TArray<FNavigationToolItemPtr> Items = GetSelectedItems();
+
+	TSet<FNavigationToolItemPtr> FirstChildItemsToSelect;
+	FirstChildItemsToSelect.Reserve(Items.Num());
+
+	for (const FNavigationToolItemPtr& Item : Items)
+	{
+		if (Item.IsValid())
+		{
+			const FNavigationToolItemPtr FirstChildItem = GetVisibleChildAt(Item, 0);
+
+			//Don't select Component items! (Component items on selection also select their owner actor items, which can cause undesired issues)
+			if (FirstChildItem.IsValid() && !FirstChildItem->IsA<FNavigationToolComponent>())
+			{
+				FirstChildItemsToSelect.Add(FirstChildItem);
+			}
+		}
+	}
+	
+	SortAndSelectItems(FirstChildItemsToSelect.Array());
+}
+
+bool FNavigationToolView::CanSelectFirstChild() const
+{
+	return GetViewSelectedItemCount() == 1;
+}
+
+void FNavigationToolView::SelectSibling(int32 InDeltaIndex)
+{
+	const TArray<FNavigationToolItemPtr> Items = GetSelectedItems();
+
+	TSet<FNavigationToolItemPtr> SiblingItemsToSelect;
+	SiblingItemsToSelect.Reserve(Items.Num());
+
+	for (const FNavigationToolItemPtr& Item : Items)
+	{
+		if (Item.IsValid() && Item->GetParent())
+		{
+			const FNavigationToolItemPtr ParentItem = Item->GetParent();
+
+			const int32 ItemIndex   = GetVisibleChildIndex(ParentItem, Item);
+			const int32 TargetIndex = ItemIndex + InDeltaIndex;
+
+			//Don't try to Normalize Index, if it's Invalid, we won't cycle and just skip that selection
+			const FNavigationToolItemPtr SiblingToSelect = GetVisibleChildAt(ParentItem, TargetIndex);
+
+			//Don't select Component items! (Component items on selection also select their owner actor items, which can cause undesired issues)
+			if (SiblingToSelect.IsValid() && !SiblingToSelect->IsA<FNavigationToolComponent>())
+			{
+				SiblingItemsToSelect.Add(SiblingToSelect);
+			}
+		}
+	}
+	SortAndSelectItems(SiblingItemsToSelect.Array());
+}
+
+bool FNavigationToolView::CanSelectSibling() const
+{
+	return GetViewSelectedItemCount() == 1;
+}
+
+void FNavigationToolView::ExpandAll()
+{
+	for (const FNavigationToolItemPtr& Item : RootVisibleItems)
+	{
+		SetItemExpansionRecursive(Item, true);
+	}
+}
+
+bool FNavigationToolView::CanExpandAll() const
+{
+	return true;
+}
+
+void FNavigationToolView::CollapseAll()
+{
+	for (const FNavigationToolItemPtr& Item : RootVisibleItems)
+	{
+		SetItemExpansionRecursive(Item, false);
+	}
+}
+
+bool FNavigationToolView::CanCollapseAll() const
+{
+	return true;
+}
+
+void FNavigationToolView::ScrollNextIntoView()
+{
+	ScrollDeltaIndexIntoView(+1);
+}
+
+void FNavigationToolView::ScrollPrevIntoView()
+{
+	ScrollDeltaIndexIntoView(-1);
+}
+
+bool FNavigationToolView::CanScrollNextIntoView() const
+{
+	return GetViewSelectedItemCount() > 0;
+}
+
+void FNavigationToolView::ScrollDeltaIndexIntoView(int32 DeltaIndex)
+{
+	const int32 SelectedItemCount = SortedSelectedItems.Num();
+	if (SortedSelectedItems.Num() > 0)
+	{
+		const int32 TargetIndex  = NextSelectedItemIntoView + DeltaIndex;
+		NextSelectedItemIntoView = TargetIndex % SelectedItemCount;
+		if (NextSelectedItemIntoView < 0)
+		{
+			NextSelectedItemIntoView += SelectedItemCount;
+		}
+		ScrollItemIntoView(SortedSelectedItems[NextSelectedItemIntoView]);
+	}
+}
+
+void FNavigationToolView::ScrollItemIntoView(const FNavigationToolItemPtr& InItem)
+{
+	if (InItem.IsValid())
+	{
+		SetParentItemExpansions(InItem, true);
+		if (ToolViewWidget.IsValid() && ToolViewWidget->GetTreeView())
+		{
+			ToolViewWidget->GetTreeView()->FocusOnItem(InItem);
+			ToolViewWidget->ScrollItemIntoView(InItem);
+		}
+	}
+}
+
+void FNavigationToolView::SortAndSelectItems(TArray<FNavigationToolItemPtr> InItemsToSelect)
+{
+	if (!InItemsToSelect.IsEmpty())
+	{
+		FNavigationTool::SortItems(InItemsToSelect);
+
+		SelectItems(InItemsToSelect
+			, ENavigationToolItemSelectionFlags::SignalSelectionChange
+			| ENavigationToolItemSelectionFlags::ScrollIntoView);
+	}
+}
+
+void FNavigationToolView::RefreshTool(bool bInImmediateRefresh)
+{
+	if (const TSharedPtr<INavigationTool> Tool = GetOwnerTool())
+	{
+		if (bInImmediateRefresh)
+		{
+			Tool->Refresh();
+		}
+		else
+		{
+			Tool->RequestRefresh();
+		}
+	}
+}
+
+void FNavigationToolView::OnSequencerOutlinerSelectionChanged()
+{
+	using namespace UE::Sequencer;
+
+	const TSharedPtr<INavigationTool> Tool = GetOwnerTool();
+	if (!Tool.IsValid())
+	{
+		return;
+	}
+
+	const TSharedPtr<ISequencer> Sequencer = Tool->GetSequencer();
+	if (!Sequencer.IsValid())
+	{
+		return;
+	}
+
+	// @TODO: need to stop cyclical selection triggering
+	/*ClearItemSelection(false);
+
+	const TSharedPtr<FSequencerSelection> Selection = Sequencer->GetViewModel()->GetSelection();
+
+	TSet<FNavigationToolItemPtr> ItemsToSelect;
+
+	for (const TViewModelPtr<IOutlinerExtension>& OutlinerItem : Selection->Outliner)
+	{
+		const FNavigationToolItemPtr FoundItem = Tool->FindItem(FNavigationToolItemId(OutlinerItem));
+		if (FoundItem.IsValid())
+		{
+			ItemsToSelect.Add(FoundItem);
+		}
+	}
+
+	for (const TViewModelPtr<FViewModel>& TrackArea : Selection->TrackArea)
+	{
+		const FNavigationToolItemPtr FoundItem = Tool->FindItem(FNavigationToolItemId(TrackArea));
+		if (FoundItem.IsValid())
+		{
+			ItemsToSelect.Add(FoundItem);
+		}
+	}
+
+	SelectItems(ItemsToSelect.Array(), ENavigationToolItemSelectionFlags::ScrollIntoView);*/
+}
+
+void FNavigationToolView::EnsureToolViewCount(const int32 InToolViewId) const
+{
+	const TSharedPtr<FNavigationTool> Tool = WeakTool.Pin();
+	if (!Tool.IsValid())
+	{
+		return;
+	}
+
+	FNavigationTool& ToolRef = *Tool;
+
+	FNavigationToolExtender::ForEachToolProvider(ToolRef, [&ToolRef, InToolViewId]
+		(const TSharedRef<FNavigationToolProvider>& InToolProvider)
+		{
+			InToolProvider->EnsureToolViewCount(ToolRef, InToolViewId);
+			return true;
+		});
+}
+
+void FNavigationToolView::SaveViewState(const TSharedRef<FNavigationToolProvider>& InProvider)
+{
+	const TSharedPtr<FNavigationTool> Tool = WeakTool.Pin();
+	if (!Tool.IsValid())
+	{
+		return;
+	}
+
+	const FNavigationTool& ToolRef = *Tool;
+
+	EnsureToolViewCount(ToolViewId);
+
+	if (FNavigationToolViewSaveState* const ViewSaveState = InProvider->GetViewSaveState(ToolRef, ToolViewId))
+	{
+		// Save view state filters
+		ViewSaveState->ActiveItemFilters.Reset();
+		for (const TSharedRef<FNavigationToolFilter>& ActiveItemFilter : FilterBar->GetActiveFilters())
+		{
+			ViewSaveState->ActiveItemFilters.Add(*ActiveItemFilter->GetDisplayName().ToString());
+		}
+
+		SaveColumnState(NAME_None);
+		SaveToolViewItems(*ViewSaveState);
+	}
+}
+
+void FNavigationToolView::LoadViewState(const TSharedRef<FNavigationToolProvider>& InProvider)
+{
+	const TSharedPtr<FNavigationTool> Tool = WeakTool.Pin();
+	if (!Tool.IsValid())
+	{
+		return;
+	}
+
+	const FNavigationTool& ToolRef = *Tool;
+
+	EnsureToolViewCount(ToolViewId);
+
+	// Disable all filters before load
+	FilterBar->EnableAllFilters(false, {});
+
+	if (FNavigationToolViewSaveState* const ViewSaveState = InProvider->GetViewSaveState(ToolRef, ToolViewId))
+	{
+		LoadFilterState(*ViewSaveState, false, false);
+		LoadToolViewItems(*ViewSaveState);
+	}
+
+	PostLoad();
+
+	FilterBar->RequestFilterUpdate();
+}
+
+void FNavigationToolView::SaveColumnState(const FName InColumnId)
+{
+	if (!ToolViewWidget.IsValid())
+	{
+		return;
+	}
+
+	const TSharedPtr<FNavigationTool> Tool = WeakTool.Pin();
+	if (!Tool.IsValid())
+	{
+		return;
+	}
+
+	const FNavigationTool& ToolRef = *Tool;
+
+	// Save all columns unless a specific column is specified
+	if (InColumnId.IsNone())
+	{
+		// Save each column to their respective provider save data
+		for (const TTuple<FName, TSharedPtr<INavigationToolColumn>>& Column : Columns)
+		{
+			if (!Column.Value.IsValid())
+			{
+				continue;
+			}
+
+			const TSharedPtr<FNavigationToolProvider> Provider = Column.Value->GetProvider();
+			if (!Provider.IsValid())
+			{
+				continue;
+			}
+
+			if (FNavigationToolViewSaveState* const ViewSaveState = Provider->GetViewSaveState(ToolRef, ToolViewId))
+			{
+				const FName ColumnId = Column.Value->GetColumnId();
+				FNavigationToolViewColumnSaveState& FoundColumnState = ViewSaveState->ColumnsState.FindOrAdd(ColumnId);
+				ToolViewWidget->GenerateColumnState(ColumnId, FoundColumnState);
+			}
+		}
+	}
+	else if (ensure(Columns.Contains(InColumnId)))
+	{
+		// Save the specific column to its providers save data
+		const TSharedPtr<INavigationToolColumn>& Column = Columns[InColumnId];
+		if (const TSharedPtr<FNavigationToolProvider> Provider = Column->GetProvider())
+		{
+			if (FNavigationToolViewSaveState* const ViewSaveState = Provider->GetViewSaveState(ToolRef, ToolViewId))
+			{
+				FNavigationToolViewColumnSaveState& FoundColumnState = ViewSaveState->ColumnsState.FindOrAdd(InColumnId);
+				ToolViewWidget->GenerateColumnState(InColumnId, FoundColumnState);
+			}
+		}
+	}
+}
+
+void FNavigationToolView::SaveFilterState(FNavigationToolViewSaveState& OutViewSaveState)
+{
+	OutViewSaveState.ActiveItemFilters.Reset();
+
+	for (const TSharedRef<FNavigationToolFilter>& ActiveItemFilter : FilterBar->GetActiveFilters())
+	{
+		OutViewSaveState.ActiveItemFilters.Add(*ActiveItemFilter->GetDisplayName().ToString());
+	}
+}
+
+void FNavigationToolView::LoadFilterState(const FNavigationToolViewSaveState& InViewSaveState
+	, const bool bInDisableAllFilters
+	, const bool bInRequestFilterUpdate)
+{
+	if (bInDisableAllFilters)
+	{
+		FilterBar->EnableAllFilters(false, {});
+	}
+
+	for (const FName ActiveItemFilterName : InViewSaveState.ActiveItemFilters)
+	{
+		FilterBar->SetFilterActiveByDisplayName(ActiveItemFilterName.ToString(), true, false);
+	}
+
+	if (bInRequestFilterUpdate)
+	{
+		FilterBar->RequestFilterUpdate();
+	}
+}
+
+void FNavigationToolView::SaveToolViewItems(FNavigationToolViewSaveState& OutViewSaveState)
+{
+	const TSharedRef<FNavigationToolItem> TreeRoot = GetOwnerTool()->GetTreeRoot();
+
+	TArray<FNavigationToolItemPtr> ItemsToSave = TreeRoot->GetChildren();
+
+	OutViewSaveState.ViewItemFlags.Reset();
+
+	while (ItemsToSave.Num() > 0)
+	{
+		FNavigationToolItemPtr ItemToSave = ItemsToSave.Pop();
+		if (ItemToSave.IsValid())
+		{
+			// Iteratively also Save children
+			ItemsToSave.Append(ItemToSave->GetChildren());
+
+			const FNavigationToolItemId ItemId = ItemToSave->GetItemId();
+			const FString StringId = ItemId.GetStringId();
+
+			// Save Item State Flags
+			if (const ENavigationToolItemFlags* const ItemFlags = OutViewSaveState.ViewItemFlags.Find(StringId))
+			{
+				OutViewSaveState.ViewItemFlags.Add(StringId, *ItemFlags);
+			}
+			else
+			{
+				OutViewSaveState.ViewItemFlags.Remove(StringId);
+			}
+		}
+	}
+}
+
+void FNavigationToolView::LoadToolViewItems(FNavigationToolViewSaveState& InViewSaveState)
+{
+	const TSharedRef<FNavigationToolItem> TreeRoot = GetOwnerTool()->GetTreeRoot();
+
+	TArray<FNavigationToolItemPtr> ItemsToLoad = TreeRoot->GetChildren();
+
+	while (ItemsToLoad.Num() > 0)
+	{
+		if (const FNavigationToolItemPtr ItemToLoad = ItemsToLoad.Pop())
+		{
+			// Iteratively also Load Children
+			ItemsToLoad.Append(ItemToLoad->GetChildren());
+
+			const FNavigationToolItemId ItemId = ItemToLoad->GetItemId();
+			const FString StringId = ItemId.GetStringId();
+
+			// Load Item Flags
+			if (const ENavigationToolItemFlags* const ItemFlags = InViewSaveState.ViewItemFlags.Find(StringId))
+			{
+				InViewSaveState.ViewItemFlags.Add(StringId, *ItemFlags);
+			}
+			else
+			{
+				InViewSaveState.ViewItemFlags.Remove(StringId);
+			}
+		}
+	}
+}
+
+TSharedRef<SWidget> FNavigationToolView::GetColumnMenuContent(const FName InColumnId)
+{
+	FMenuBuilder MenuBuilder(true, ViewCommandList);
+
+	const FNavigationToolCommands& ToolCommands = FNavigationToolCommands::Get();
+
+	MenuBuilder.AddMenuEntry(
+		LOCTEXT("ResetColumnSize", "Reset Column Size"), 
+		LOCTEXT("ResetColumnSizeTooltip", "Resets the size of this column to the default"), 
+		FSlateIcon(),
+		FUIAction(FExecuteAction::CreateRaw(this, &FNavigationToolView::ResetColumnSize, InColumnId),
+			FCanExecuteAction::CreateRaw(this, &FNavigationToolView::CanResetColumnSize, InColumnId)));
+
+	MenuBuilder.AddMenuSeparator();
+
+	MenuBuilder.AddMenuEntry(ToolCommands.ResetVisibleColumnSizes);
+
+	return MenuBuilder.MakeWidget();
+}
+
+void FNavigationToolView::ResetColumnSize(const FName InColumnId)
+{
+	if (!ToolViewWidget.IsValid())
+	{
+		return;
+	}
+
+	if (!Columns.Contains(InColumnId) || !ToolViewWidget->IsColumnVisible(InColumnId))
+	{
+		return;
+	}
+
+	const TSharedPtr<FNavigationToolProvider> Provider = Columns[InColumnId]->GetProvider();
+	if (!Provider.IsValid())
+	{
+		return;
+	}
+
+	const TSharedPtr<FNavigationTool> Tool = WeakTool.Pin();
+	if (!Tool.IsValid())
+	{
+		return;
+	}
+
+	const FNavigationTool& ToolRef = *Tool;
+
+	FNavigationToolViewSaveState* const ViewSaveState = Provider->GetViewSaveState(ToolRef, ToolViewId);
+	if (!ViewSaveState)
+	{
+		return;
+	}
+
+	const float DefaultSize = Columns[InColumnId]->GetFillWidth();
+	if (DefaultSize <= 0.f)
+	{
+		return;
+	}
+
+	ViewSaveState->ColumnsState.FindOrAdd(InColumnId).Size = DefaultSize;
+
+	ToolViewWidget->SetColumnWidth(InColumnId, DefaultSize);
+
+	ToolViewWidget->GenerateColumnState(InColumnId, ViewSaveState->ColumnsState.FindOrAdd(InColumnId));
+}
+
+bool FNavigationToolView::CanResetColumnSize(const FName InColumnId) const
+{
+	if (!Columns.Contains(InColumnId) || !ToolViewWidget->IsColumnVisible(InColumnId))
+	{
+		return false;
+	}
+
+	const TSharedPtr<FNavigationToolProvider> Provider = Columns[InColumnId]->GetProvider();
+	if (!Provider.IsValid())
+	{
+		return false;
+	}
+
+	const TSharedPtr<FNavigationTool> Tool = WeakTool.Pin();
+	if (!Tool.IsValid())
+	{
+		return false;
+	}
+
+	const FNavigationTool& ToolRef = *Tool;
+
+	const FNavigationToolViewSaveState* const ViewSaveState = Provider->GetViewSaveState(ToolRef, ToolViewId);
+	if (!ViewSaveState || !ViewSaveState->ColumnsState.Contains(InColumnId))
+	{
+		return false;
+	}
+
+	const float DefaultSize = Columns[InColumnId]->GetFillWidth();
+	if (DefaultSize <= 0.f)
+	{
+		return false;
+	}
+
+	return ViewSaveState->ColumnsState[InColumnId].Size != DefaultSize;
+}
+
+void FNavigationToolView::ResetVisibleColumnSizes()
+{
+	if (!ToolViewWidget.IsValid())
+	{
+		return;
+	}
+
+	const TSharedPtr<FNavigationTool> Tool = WeakTool.Pin();
+	if (!Tool.IsValid())
+	{
+		return;
+	}
+
+	const FNavigationTool& ToolRef = *Tool;
+
+	for (const TPair<FName, TSharedPtr<INavigationToolColumn>>& Pair : Columns)
+	{
+		const FName& ColumnId = Pair.Key;
+
+		if (!ToolViewWidget->IsColumnVisible(ColumnId))
+		{
+			continue;
+		}
+
+		const float DefaultSize = Columns[ColumnId]->GetFillWidth();
+		if (DefaultSize <= 0.f)
+		{
+			continue;
+		}
+
+		const TSharedPtr<FNavigationToolProvider> Provider = Pair.Value->GetProvider();
+		if (!Provider.IsValid())
+		{
+			continue;
+		}
+
+		FNavigationToolViewSaveState* const ViewSaveState = Provider->GetViewSaveState(ToolRef, ToolViewId);
+		if (!ViewSaveState)
+		{
+			return;
+		}
+
+		ToolViewWidget->SetColumnWidth(ColumnId, DefaultSize);
+
+		ToolViewWidget->GenerateColumnState(ColumnId, ViewSaveState->ColumnsState.FindOrAdd(ColumnId));
+	}
+}
+
+bool FNavigationToolView::CanResetAllColumnSizes() const
+{
+	return true;
+}
+
+void FNavigationToolView::SaveNewCustomColumnView()
+{
+	TSet<FNavigationToolColumnView>& CustomColumnViews = ToolSettings->GetCustomColumnViews();
+
+	FNavigationToolColumnView NewColumnView;
+	NewColumnView.ViewName = TEXT("Column View 1");
+	for (const TPair<FName, TSharedPtr<INavigationToolColumn>>& Column : GetColumns())
+	{
+		const FName ColumnId = Column.Value->GetColumnId();
+		if (IsColumnVisible(Column.Value))
+		{
+			NewColumnView.VisibleColumns.Add(ColumnId);
+		}
+	}
+
+	FString UserEnteredText;
+
+	FModalTextInputDialog InputDialog;
+	InputDialog.InputLabel = LOCTEXT("NewColumnViewName_InputLabel", "New Column View Name");
+	if (!InputDialog.Open(NewColumnView.ViewName.ToString(), UserEnteredText))
+	{
+		return;
+	}
+
+	NewColumnView.ViewName = FName(UserEnteredText);
+
+	bool bAlreadyExists = false;
+	CustomColumnViews.Add(NewColumnView, &bAlreadyExists);
+
+	if (bAlreadyExists)
+	{
+		FMessageDialog::Open(EAppMsgCategory::Error, EAppMsgType::Ok
+			, LOCTEXT("AlreadyExistsErrorText", "Column view name already exists!"));
+		return;
+	}
+
+	CustomColumnViews.Sort([](const FNavigationToolColumnView& InA, const FNavigationToolColumnView& InB)
+		{
+			return InA.ViewName.FastLess(InB.ViewName);
+		});
+
+	ToolSettings->SaveConfig();
+}
+
+void FNavigationToolView::ApplyCustomColumnView(const FName InColumnViewName)
+{
+	if (!ToolViewWidget.IsValid() || InColumnViewName.IsNone())
+	{
+		return;
+	}
+
+	FNavigationToolColumnView* const SavedColumnView = ToolSettings->FindCustomColumnView(InColumnViewName);
+	if (!SavedColumnView)
+	{
+		return;
+	}
+
+	for (const TPair<FName, TSharedPtr<INavigationToolColumn>>& Column : GetColumns())
+	{
+		const FName ColumnId = Column.Value->GetColumnId();
+		const bool bColumnVisible = SavedColumnView->VisibleColumns.Contains(ColumnId);
+		ToolViewWidget->ShowHideColumn(ColumnId, bColumnVisible);
+	}
+
+	SaveColumnState();
+}
+
+bool FNavigationToolView::CanFocusSingleSelection() const
+{
+	if (SelectedItems.Num() != 1 || !GetOwnerTool().IsValid())
+	{
+		return false;
+	}
+
+	const FNavigationToolSequence* const SequenceItem = SelectedItems[0]->CastTo<FNavigationToolSequence>();
+	if (!SequenceItem || !IsValid(SequenceItem->GetSequence()))
+	{
+		return false;
+	}
+
+	return true;
+}
+
+void FNavigationToolView::FocusSingleSelection()
+{
+	if (SelectedItems.Num() != 1)
+	{
+		return;
+	}
+
+	const FNavigationToolSequence* const SequenceItem = SelectedItems[0]->CastTo<FNavigationToolSequence>();
+	if (!SequenceItem)
+	{
+		return;
+	}
+
+	UMovieSceneSequence* const Sequence = SequenceItem->GetSequence();
+	if (!IsValid(Sequence))
+	{
+		return;
+	}
+
+	const TSharedPtr<INavigationTool> Tool = GetOwnerTool();
+	if (!Tool.IsValid())
+	{
+		return;
+	}
+
+	FocusSequence(*Tool, *Sequence);
+}
+
+bool FNavigationToolView::UpdateFilters()
+{
+	if (!bFilterUpdateRequested)
+	{
+		return false;
+	}
+
+	const TSharedPtr<INavigationTool> Tool = GetOwnerTool();
+	if (!Tool.IsValid())
+	{
+		return false;
+	}
+
+	const FNavigationToolFilterData& PreviousFilterData = FilterBar->GetFilterData();
+	const FNavigationToolFilterData& FilterData = FilterBar->FilterNodes();
+
+	//bFilteringOnNodeGroups = Sequencer.GetFocusedMovieSceneSequence()->GetMovieScene()->GetNodeGroups().HasAnyActiveFilter();
+	bFilterUpdateRequested = false;
+
+	// Return whether the new list of FilteredNodes is different than the previous list
+	return (PreviousFilterData.GetDisplayNodeCount() != FilterData.GetDisplayNodeCount()
+		|| PreviousFilterData != FilterData);
+}
+
+}
+
+#undef LOCTEXT_NAMESPACE
diff --git a/Engine/Source/Editor/Sequencer/Private/Tools/NavigationTool/NavigationToolView.h b/Engine/Source/Editor/Sequencer/Private/Tools/NavigationTool/NavigationToolView.h
new file mode 100644
--- /dev/null
+++ b/Engine/Source/Editor/Sequencer/Private/Tools/NavigationTool/NavigationToolView.h
@@ -0,0 +1,434 @@
+// Copyright Epic Games, Inc. All Rights Reserved.
+
+#pragma once
+
+#include "MovieScene.h"
+#include "Tools/NavigationTool/INavigationToolView.h"
+#include "Tools/NavigationTool/NavigationToolDefines.h"
+#include "Tools/NavigationTool/NavigationToolItemType.h"
+#include "Tools/NavigationTool/Filters/NavigationToolFilterBar.h"
+
+class FDragDropEvent;
+class FReply;
+class FUICommandList;
+class SWidget;
+class UNavigationToolSettings;
+enum class ECheckBoxState : uint8;
+enum class EItemDropZone;
+struct FGeometry;
+struct FNavigationToolSaveState;
+struct FPointerEvent;
+
+namespace UE::Sequencer::NavigationTool
+{
+
+class FNavigationTool;
+class FNavigationToolFilterBar;
+class FNavigationToolItemContextMenu;
+class FNavigationToolProvider;
+class INavigationToolColumn;
+class INavigationToolProvider;
+class SNavigationToolView;
+
+/**
+  * A view instance of the Navigation Tool, that handles viewing a subset of the outliner items based on
+  * item filters, search text, hierarchy type, etc
+  */
+class FNavigationToolView : public INavigationToolView
+{
+private:
+	friend FNavigationToolSaveState;
+	friend FNavigationTool;
+
+	struct FPrivateToken { explicit FPrivateToken() = default; };
+
+public:
+	explicit FNavigationToolView(FPrivateToken);
+	virtual ~FNavigationToolView() override;
+
+	/** Creates columns from a specific provider */
+	void CreateColumns(const TSharedRef<FNavigationToolProvider>& InProvider);
+
+	void CreateDefaultColumnViews(const TSharedRef<FNavigationToolProvider>& InProvider);
+
+	/**
+	 * Creates a Navigation Tool View Instance and register it to the Navigation Tool
+	 * @param InToolViewId the Id assigned to the Instance created
+	 * @param InTool the Navigation Tool this Instance View belongs to
+	 * @param InBaseCommandList the base command list
+	 */
+	static TSharedRef<FNavigationToolView> CreateInstance(const int32 InToolViewId
+		, const TSharedRef<FNavigationTool>& InTool
+		, const TSharedPtr<FUICommandList>& InBaseCommandList);
+
+	void PostLoad();
+
+	/** Called when the UNavigationToolSettings has a property change */
+	void OnToolSettingsChanged(UObject* InObject, FPropertyChangedEvent& InPropertyChangedEvent);
+
+	void Tick(float InDeltaTime);
+
+	void BindCommands(const TSharedPtr<FUICommandList>& InBaseCommandList);
+
+	TSharedPtr<FUICommandList> GetBaseCommandList() const;
+	TSharedPtr<FUICommandList> GetViewCommandList() const { return ViewCommandList; }
+
+	/** Notifies FNavigationTool so that this view instance becomes the most recent interacted outliner view */
+	void UpdateRecentViews();
+
+	/** Returns whether this view instance is the most recently interacted between the instances in FNavigationTool */
+	bool IsMostRecentToolView() const;
+
+	int32 GetToolViewId() const { return ToolViewId; }
+
+	//~ Begin INavigationToolView
+	virtual TSharedPtr<INavigationTool> GetOwnerTool() const override;
+	virtual TSharedPtr<SWidget> GetToolWidget() const override;
+	virtual TSharedPtr<ISequencer> GetSequencer() const override;
+	virtual void RequestRefresh() override;
+	virtual FOnToolViewRefreshed& GetOnToolViewRefreshed() override { return OnToolViewRefreshed; }
+	virtual void SetKeyboardFocus() override;
+	virtual ENavigationToolItemViewMode GetItemDefaultViewMode() const override;
+	virtual ENavigationToolItemViewMode GetItemProxyViewMode() const override;
+	virtual TSharedRef<SWidget> GetColumnMenuContent(const FName InColumnId) override;
+	virtual void GetChildrenOfItem(const FNavigationToolItemPtr InItem, TArray<FNavigationToolItemPtr>& OutChildren) const override;
+	virtual void GetChildrenOfItem(const FNavigationToolItemPtr& InItem
+		, TArray<FNavigationToolItemPtr>& OutChildren
+		, ENavigationToolItemViewMode InViewMode
+		, const TSet<FNavigationToolItemPtr>& InRecursionDisallowedItems) const override;
+	virtual bool IsItemReadOnly(const FNavigationToolItemPtr& InItem) const override;
+	virtual bool CanSelectItem(const FNavigationToolItemPtr& InItem) const override;
+	virtual void SelectItems(TArray<FNavigationToolItemPtr> InItems, const ENavigationToolItemSelectionFlags InFlags = ENavigationToolItemSelectionFlags::SignalSelectionChange) override;
+	virtual void ClearItemSelection(bool bSignalSelectionChange = true) override;
+	virtual bool IsItemSelected(const FNavigationToolItemPtr& InItem) const override;
+	virtual TArray<FNavigationToolItemPtr> GetSelectedItems() const override;
+	virtual FReply OnDragDetected(const FGeometry& InGeometry, const FPointerEvent& InMouseEvent, const FNavigationToolItemPtr InTargetItem) override;
+	virtual FReply OnDrop(const FDragDropEvent& InDragDropEvent, EItemDropZone InDropZone, const FNavigationToolItemPtr InTargetItem) override;
+	virtual TOptional<EItemDropZone> OnCanDrop(const FDragDropEvent& InDragDropEvent
+		, EItemDropZone InDropZone
+		, const FNavigationToolItemPtr InTargetItem) const override;
+	virtual bool IsItemExpanded(const FNavigationToolItemPtr& InItem, const bool bInUseFilter = true) const override;
+	virtual void SetItemExpansion(const FNavigationToolItemPtr& InItem, const bool bInExpand, const bool bInUseFilter = true) override;
+	virtual void SetItemExpansionRecursive(const FNavigationToolItemPtr InItem, const bool bInExpand) override;
+	virtual void SetParentItemExpansions(const FNavigationToolItemPtr& InItem, const bool bInExpand) override;
+	virtual bool CanExpandAll() const override;
+	virtual void ExpandAll() override;
+	virtual void CollapseAll() override;
+	virtual bool CanCollapseAll() const override;
+	//~ End INavigationToolView
+
+	TSharedPtr<SWidget> CreateItemContextMenu();
+
+	/** Get the Columns currently allocated in this Navigation Tool Instance */
+	const TMap<FName, TSharedPtr<INavigationToolColumn>>& GetColumns() const { return Columns; }
+
+	/** Whether the given column should be shown by default. This is only used at the start when Creating the Columns */
+	bool ShouldShowColumnByDefault(const TSharedPtr<INavigationToolColumn>& InColumn) const;
+
+	/** Refreshes the Items visible in this View, and refreshes the Navigation Tool Widget if it was created */
+	void Refresh();
+
+	/** Refreshes the Items that will be at the Top Level of the Tree (e.g. Actors with no parent item) */
+	void UpdateRootVisibleItems();
+
+	/** Updates the Item Expansions in the Tree Widget based on whether the Items have the Expanded flag or not */
+	void UpdateItemExpansions();
+
+	/** Called when Object Replacement has taken place. Used to invalidate the widget for painting */
+	void NotifyObjectsReplaced();
+
+	/** The Root Item of the Tree. Convenience function to get the RootItem from FNavigationTool */
+	FNavigationToolItemPtr GetRootItem() const;
+
+	/** Gets the Top Level Items that should be visible in the Tree. This can differ between views depending on Filters and other factors */
+	const TArray<FNavigationToolItemPtr>& GetRootVisibleItems() const;
+
+	void SaveViewItemFlags(const FNavigationToolItemPtr& InItem, ENavigationToolItemFlags InFlags);
+	
+	ENavigationToolItemFlags GetViewItemFlags(const FNavigationToolItemPtr& InItem) const;
+
+	FLinearColor GetItemBrushColor(const FNavigationToolItemPtr InItem) const;
+
+	/** Gets the Currently Selected Item Count in the Tree View */
+	int32 GetViewSelectedItemCount() const;
+
+	/** Calculates the amount of Items that are visible in this Navigation Tool View */
+	int32 CalculateVisibleItemCount() const;
+
+	/** Whether Sync Selection is currently taking place */
+	bool IsSyncingItemSelection() const { return bSyncingItemSelection; }
+
+	TSharedPtr<FNavigationToolFilterBar> GetFilterBar() const { return FilterBar; }
+
+	/** Called when Item selection has changed. This can be called by the Navigation Tool Widget or by the View itself if there's no ToolWidget */
+	void NotifyItemSelectionChanged(const TArray<FNavigationToolItemPtr>& InSelectedItems
+		, const FNavigationToolItemPtr& InItem
+		, bool bUpdateModeTools);
+
+	/** Return whether the given item should be visible in this Navigation Tool Instance */
+	bool ShouldShowItem(const FNavigationToolItemPtr& InItem, bool bInUseFilters, ENavigationToolItemViewMode InViewMode) const;
+
+	/** Returns the Index of the Child from the Parent's Children Visible List (Visible List can vary between Tool Instances) */
+	int32 GetVisibleChildIndex(const FNavigationToolItemPtr& InParentItem, const FNavigationToolItemPtr& InChildItem) const;
+
+	/** Returns the Child from the Parent's Children Visible List at the given Index (Visible List can vary between Tool Instances) */
+	FNavigationToolItemPtr GetVisibleChildAt(const FNavigationToolItemPtr& InParentItem, int32 InChildIndex) const;
+
+	/** Returns whether the Tool is Locked. See FNavigationTool::IsToolLocked */
+	bool IsToolLocked() const;
+
+	/**
+	 * Activate an Item Column
+	 * @param InColumn the column to activate
+	 */
+	void ShowColumn(const TSharedPtr<INavigationToolColumn>& InColumn);
+
+	/**
+	 * Activate the Item Column with the given ID
+	 * @param InColumnId the filter id associated with the Filter to activate
+	 */
+	void ShowColumnById(const FName& InColumnId);
+
+	/**
+	 * Deactivates the given Item Column
+	 * @param InColumn the filter to remove from the active filter list
+	 */
+	void HideColumn(const TSharedPtr<INavigationToolColumn>& InColumn);
+
+	/** Whether the given Column is in the Active Column List */
+	virtual bool IsColumnVisible(const TSharedPtr<INavigationToolColumn>& InColumn) const override;
+
+	void ToggleViewModeSupport(ENavigationToolItemViewMode& InOutViewMode, ENavigationToolItemViewMode InFlags);
+	void ToggleItemDefaultViewModeSupport(ENavigationToolItemViewMode InFlags);
+	void ToggleItemProxyViewModeSupport(ENavigationToolItemViewMode InFlags);
+
+	ECheckBoxState GetViewModeCheckState(ENavigationToolItemViewMode InViewMode, ENavigationToolItemViewMode InFlags) const;
+	ECheckBoxState GetItemDefaultViewModeCheckState(ENavigationToolItemViewMode InFlags) const;
+	ECheckBoxState GetItemProxyViewModeCheckState(ENavigationToolItemViewMode InFlags) const;
+	
+	/**
+	 * Action to turn Muted Hierarchy on or off
+	 * @see bUseMutedHierarchy
+	 */
+	void ToggleMutedHierarchy();
+	bool CanToggleMutedHierarchy() const { return true; }
+	bool IsMutedHierarchyActive() const;
+
+	/**
+	 * Action to turn Auto Expand to Selection on or off
+	 * @see bAutoExpandToSelection
+	 */
+	void ToggleAutoExpandToSelection();
+	bool CanToggleAutoExpandToSelection() const { return true; }
+	bool ShouldAutoExpandToSelection() const;
+
+	void ToggleUseShortNames();
+	bool CanToggleUseShortNames() const { return true; }
+	bool ShouldUseShortNames() const;
+
+	void ToggleShowItemFilters();
+	bool CanToggleShowItemFilters() const { return true; }
+	bool ShouldShowItemFilters() const { return bShowItemFilters; }
+
+	void ToggleShowItemColumns();
+	bool CanToggleShowItemColumns() const { return true; }
+	bool ShouldShowItemColumns() const { return bShowItemColumns; }
+
+	/** Sets whether the given item type should be hidden or not */
+	void SetItemTypeHidden(const FName InItemTypeName, const bool bInHidden);
+	
+	/** Toggles the given Item Types to Hide/Show */
+	void ToggleHideItemTypes(const FName InItemTypeName);
+	ECheckBoxState GetToggleHideItemTypesState(const FName InItemTypeName) const;
+	
+	/** Hides the given Outliner Item Type from showing in this Outliner View (e.g. Hide Components is used) */
+	template<typename InItemType, typename = typename TEnableIf<TIsDerivedFrom<InItemType, INavigationToolItem>::IsDerived>::Type>
+	void HideItemType()
+	{
+		SetItemTypeHidden(TNavigationToolItemType<InItemType>::GetTypeId().ToName(), true);
+	}
+	
+	/** Shows the given Outliner Item Type if it was hidden in this Outliner View (e.g. Showing Components again) */
+	template<typename InItemType, typename = typename TEnableIf<TIsDerivedFrom<InItemType, INavigationToolItem>::IsDerived>::Type>
+	void ExposeItemType()
+	{
+		SetItemTypeHidden(TNavigationToolItemType<InItemType>::GetTypeId().ToName(), false);
+	}
+	
+	/** Whether the given Outliner Item Type is currently hidden in this Outliner View */
+	bool IsItemTypeHidden(const FName InItemTypeName) const;
+	bool IsItemTypeHidden(const FNavigationToolItemPtr& InItem) const;
+	
+	template<typename InItemType, typename = typename TEnableIf<TIsDerivedFrom<InItemType, INavigationToolItem>::IsDerived>::Type>
+	bool IsItemTypeHidden() const
+	{
+		return IsItemTypeHidden(InItemType::GetStaticTypeName());
+	}
+
+	/** Called when a Drag enters the Navigation Tool Widgets for the given Item (if null, treat as root) */
+	void OnDragEnter(const FDragDropEvent& InDragDropEvent, const FNavigationToolItemPtr InTargetItem);
+
+	/** Called when a Drag leaves the Navigation Tool Widgets for the given Item (if null, treat as root) */
+	void OnDragLeave(const FDragDropEvent& InDragDropEvent, const FNavigationToolItemPtr InTargetItem);
+
+	/** Called when there's a Drag being started/ended to the Navigation Tool Widget as a whole (i.e. root) and not an individual item */
+	void SetDragIntoTreeRoot(bool bInIsDraggingIntoTreeRoot);
+	
+	void RenameSelected();
+	void ResetRenaming();
+	void OnItemRenameAction(const ENavigationToolRenameAction InRenameAction, const TSharedPtr<INavigationToolView>& InToolView);
+	bool CanRenameSelected() const;
+
+	void DeleteSelected();
+	bool CanDeleteSelected() const;
+
+	void DuplicateSelected();
+	bool CanDuplicateSelected() const;
+
+	void SelectChildren(bool bIsRecursive);
+	bool CanSelectChildren() const;
+
+	void SelectParent();
+	bool CanSelectParent() const;
+
+	void SelectFirstChild();
+	bool CanSelectFirstChild() const;
+
+	void SelectSibling(int32 InDeltaIndex);
+	bool CanSelectSibling() const;
+
+	/** Called when the given item's expansion state (expanded/collapsed) has changed */
+	void OnItemExpansionChanged(const FNavigationToolItemPtr InItem, const bool bInIsExpanded);
+
+	void ScrollNextIntoView();
+	void ScrollPrevIntoView();
+	bool CanScrollNextIntoView() const;
+	void ScrollDeltaIndexIntoView(int32 DeltaIndex);
+	void ScrollItemIntoView(const FNavigationToolItemPtr& InItem);
+
+	/** Sorts the given Items in the user-defined order and selects them, and if the Widget is valid focusing to the item in the Widget */
+	void SortAndSelectItems(TArray<FNavigationToolItemPtr> InItemsToSelect);
+
+	/** Save the current state of the Navigation Tool */
+	void SaveViewState(const TSharedRef<FNavigationToolProvider>& InProvider);
+	/** Load the current state of the Navigation Tool */
+	void LoadViewState(const TSharedRef<FNavigationToolProvider>& InProvider);
+
+	/**  */
+	void SaveColumnState(const FName InColumnId = NAME_None);
+
+	void ResetColumnSize(const FName InColumnId);
+	bool CanResetColumnSize(const FName InColumnId) const;
+
+	void ResetVisibleColumnSizes();
+	bool CanResetAllColumnSizes() const;
+
+	/** Display a modal dialog for saving a new custom column view */
+	void SaveNewCustomColumnView();
+
+	/** Apply a saved custom column view by name */
+	void ApplyCustomColumnView(const FName InColumnViewName);
+
+protected:
+	/** Initializes the Navigation Tool Instance. Only executed once */
+	void Init(const TSharedRef<FNavigationTool>& InTool, const TSharedPtr<FUICommandList>& InBaseCommandList);
+
+	/** Sets the Selected Items to be the given Items in this List. */
+	void SetItemSelectionImpl(TArray<FNavigationToolItemPtr>&& InItems
+		, bool bSignalSelectionChange);
+
+	/** Triggers a Refresh on the FNavigationTool */
+	void RefreshTool(bool bInImmediateRefresh);
+
+	void OnSequencerOutlinerSelectionChanged();
+
+	void EnsureToolViewCount(const int32 InToolViewId) const;
+
+	/**  */
+	void SaveFilterState(FNavigationToolViewSaveState& OutViewSaveState);
+	/**  */
+	void LoadFilterState(const FNavigationToolViewSaveState& InViewSaveState
+		, const bool bInDisableAllFilters = true
+		, const bool bInRequestFilterUpdate = true);
+
+	/** Save the Item State in the Navigation Tool Widget (e.g. Item Scoped Flags) */
+	void SaveToolViewItems(FNavigationToolViewSaveState& OutViewSaveState);
+	/** Load the Item State in the Navigation Tool Widget (e.g. Item Scoped Flags) */
+	void LoadToolViewItems(FNavigationToolViewSaveState& InViewSaveState);
+
+	bool CanFocusSingleSelection() const;
+	void FocusSingleSelection();
+
+	bool UpdateFilters();
+
+	/** Local Identifier of this Instance */
+	int32 ToolViewId = -1;
+
+	/** Weak pointer to the outliner this is a view of */
+	TWeakPtr<FNavigationTool> WeakTool;
+
+	/** Widget showing the View. Can be null if instanced for testing (i.e. bCreateToolWidget is false in Init) */
+	TSharedPtr<SNavigationToolView> ToolViewWidget;
+
+	/** Command List mapped to the View to handle things like Selected Items */
+	TSharedPtr<FUICommandList> ViewCommandList;
+
+	/**  */
+	const TObjectPtr<UNavigationToolSettings> ToolSettings;
+
+	/** Root Items from Navigation Tool visible to this Instance. i.e. a Subset of all the Items in FNavigationTool */
+	TArray<FNavigationToolItemPtr> RootVisibleItems;
+
+	/** A list of the Selected Items in this Navigation Tool View */
+	TArray<FNavigationToolItemPtr> SelectedItems;
+
+	/** Set of Items that are Currently Read Only in this Instance. Mutable as this changes on GetChildrenOfItem const func */
+	mutable TSet<FNavigationToolItemPtr> ReadOnlyItems;
+
+	/** A set of the Navigation Tool Item Type Names that should be hidden */
+	TSet<FName> HiddenItemTypes;
+
+	TSharedPtr<FNavigationToolFilterBar> FilterBar;
+
+	/** Delegate called at the end of the Navigation Tool View Refresh */
+	FOnToolViewRefreshed OnToolViewRefreshed;
+
+	/** Map of the Column Ids and the Column with that Id. This is filled in FNavigationToolView::CreateColumns */
+	TMap<FName, TSharedPtr<INavigationToolColumn>> Columns;
+
+	/** Right click context menu for items */
+	TSharedPtr<FNavigationToolItemContextMenu> ItemContextMenu;
+
+	/** The Index of the Item that will be scrolled into View */
+	int32 NextSelectedItemIntoView = -1;
+
+	/** Selected Items Sorted from Top To Bottom */
+	TArray<FNavigationToolItemPtr> SortedSelectedItems;
+
+	/** The list of items that need renaming and are waiting for their turn to be renamed */
+	TArray<FNavigationToolItemPtr> ItemsRemainingRename;
+	
+	/** The current item in the process of renaming. Null if no renaming is taking place */
+	FNavigationToolItemPtr CurrentItemRenaming;
+	
+	/** Flag to call FNavigationToolView::Refresh next tick */
+	bool bRefreshRequested = false;
+	
+	/** Whether Item renaming is taking place in this View */
+	bool bRenamingItems = false;
+	
+	/** Whether Renaming Items should be processed next tick */
+	bool bRequestedRename = false;
+
+	/** Whether Items are currently being synced to what the Widget or Interface sent. Used as a re-enter guard */
+	bool bSyncingItemSelection = false;
+	
+	/** Flag used for the Navigation Tool Widget to determine if the Item Filter Bar should open and show the Item Filters */
+	bool bShowItemFilters = false;
+
+	bool bShowItemColumns = false;
+
+	bool bFilterUpdateRequested = false;
+};
+
+}
diff --git a/Engine/Source/Editor/Sequencer/Private/Tools/NavigationTool/Utils/NavigationToolMiscUtils.cpp b/Engine/Source/Editor/Sequencer/Private/Tools/NavigationTool/Utils/NavigationToolMiscUtils.cpp
new file mode 100644
--- /dev/null
+++ b/Engine/Source/Editor/Sequencer/Private/Tools/NavigationTool/Utils/NavigationToolMiscUtils.cpp
@@ -0,0 +1,81 @@
+﻿// Copyright Epic Games, Inc. All Rights Reserved.
+
+#include "NavigationToolMiscUtils.h"
+#include "ISequencer.h"
+#include "ISourceControlModule.h"
+#include "ISourceControlProvider.h"
+#include "Sections/MovieSceneSubSection.h"
+#include "Tools/NavigationTool/INavigationTool.h"
+#include "UObject/Package.h"
+
+namespace UE::Sequencer::NavigationTool
+{
+
+FSourceControlStatePtr GetSourceControlState(UPackage* const InPackage)
+{
+	if (!IsValid(InPackage))
+	{
+		return nullptr;
+	}
+
+	ISourceControlModule& SourceControlModule = ISourceControlModule::Get();
+	if (!SourceControlModule.IsEnabled())
+	{
+		return nullptr;
+	}
+
+	ISourceControlProvider& SourceControlProvider = SourceControlModule.GetProvider();
+	if (!SourceControlProvider.IsAvailable())
+	{
+		return nullptr;
+	}
+
+	return SourceControlProvider.GetState({ InPackage }, EStateCacheUsage::Use);
+}
+
+const FSlateBrush* GetSourceControlStatusBrush(UPackage* const InPackage)
+{
+	if (const FSourceControlStatePtr SourceControlState = GetSourceControlState(InPackage))
+	{
+		const FSlateIcon Icon = SourceControlState->GetIcon();
+		if (Icon.IsSet())
+		{
+			return Icon.GetIcon();
+		}
+	}
+	return nullptr;
+}
+
+FText GetSourceControlStatusText(UPackage* const InPackage)
+{
+	if (const FSourceControlStatePtr SourceControlState = GetSourceControlState(InPackage))
+	{
+		return SourceControlState->GetDisplayTooltip();
+	}
+	return FText::GetEmpty();
+}
+
+void FocusSequence(const INavigationTool& InTool, UMovieSceneSequence& InSequence)
+{
+	const TSharedPtr<ISequencer> Sequencer = InTool.GetSequencer();
+	if (!Sequencer.IsValid())
+	{
+		return;
+	}
+
+	FMovieSceneEvaluationState* const EvaluationState = Sequencer->GetEvaluationState();
+	if (!EvaluationState)
+	{
+		return;
+	}
+
+	const FMovieSceneSequenceIDRef SequenceID = EvaluationState->FindSequenceId(&InSequence);
+
+	UMovieSceneSubSection* const SubSection = Sequencer->FindSubSection(SequenceID);
+	if (IsValid(SubSection))
+	{
+		Sequencer->FocusSequenceInstance(*SubSection);
+	}
+}
+
+}
diff --git a/Engine/Source/Editor/Sequencer/Private/Tools/NavigationTool/Utils/NavigationToolMiscUtils.h b/Engine/Source/Editor/Sequencer/Private/Tools/NavigationTool/Utils/NavigationToolMiscUtils.h
new file mode 100644
--- /dev/null
+++ b/Engine/Source/Editor/Sequencer/Private/Tools/NavigationTool/Utils/NavigationToolMiscUtils.h
@@ -0,0 +1,25 @@
+﻿// Copyright Epic Games, Inc. All Rights Reserved.
+
+#pragma once
+
+#include "ISourceControlState.h"
+
+class UMovieSceneSequence;
+class UPackage;
+
+namespace UE::Sequencer::NavigationTool
+{
+
+class INavigationTool;
+
+/**  */
+FSourceControlStatePtr GetSourceControlState(UPackage* const InPackage);
+/**  */
+const FSlateBrush* GetSourceControlStatusBrush(UPackage* const InPackage);
+/**  */
+FText GetSourceControlStatusText(UPackage* const InPackage);
+
+/** Focuses a sub-movie scene (MovieScene within a MovieScene) in the sequencer. */
+void FocusSequence(const INavigationTool& InTool, UMovieSceneSequence& InSequence);
+
+}
diff --git a/Engine/Source/Editor/Sequencer/Private/Tools/NavigationTool/Utils/NavigationToolMovieSceneUtils.cpp b/Engine/Source/Editor/Sequencer/Private/Tools/NavigationTool/Utils/NavigationToolMovieSceneUtils.cpp
new file mode 100644
--- /dev/null
+++ b/Engine/Source/Editor/Sequencer/Private/Tools/NavigationTool/Utils/NavigationToolMovieSceneUtils.cpp
@@ -0,0 +1,111 @@
+﻿// Copyright Epic Games, Inc. All Rights Reserved.
+
+#include "NavigationToolMovieSceneUtils.h"
+#include "MovieScene.h"
+#include "Sections/MovieSceneSubSection.h"
+#include "Sequencer.h"
+#include "SequencerUtilities.h"
+
+namespace UE::Sequencer::NavigationTool
+{
+
+UMovieSceneSubSection* FindSequenceSubSection(ISequencer& InSequencer, UMovieSceneSequence* const InSequence)
+{
+	FMovieSceneEvaluationState* const EvaluationState = InSequencer.GetEvaluationState();
+	if (!EvaluationState)
+	{
+		return nullptr;
+	}
+
+	const FMovieSceneSequenceIDRef SequenceID = EvaluationState->FindSequenceId(InSequence);
+
+	return InSequencer.FindSubSection(SequenceID);
+}
+
+bool IsGloballyMarkedFramesForSequence(UMovieSceneSequence* const InSequence)
+{
+	if (!IsValid(InSequence))
+	{
+		return false;
+	}
+
+	UMovieScene* const MovieScene = InSequence->GetMovieScene();
+	if (!IsValid(MovieScene))
+	{
+		return false;
+	}
+
+	return MovieScene->GetGloballyShowMarkedFrames();
+}
+
+void ShowGloballyMarkedFramesForSequence(ISequencer& InSequencer, UMovieSceneSequence* const InSequence, const bool bInVisible)
+{
+	if (!IsValid(InSequence))
+	{
+		return;
+	}
+
+	UMovieScene* const MovieScene = InSequence->GetMovieScene();
+	if (!IsValid(MovieScene))
+	{
+		return;
+	}
+
+	if (MovieScene->IsReadOnly())
+	{
+		FSequencerUtilities::ShowReadOnlyError();
+		return;
+	}
+
+	InSequence->Modify();
+	MovieScene->Modify();
+
+	MovieScene->SetGloballyShowMarkedFrames(bInVisible);
+
+	if (FSequencer* const Sequencer = static_cast<FSequencer*>(&InSequencer))
+	{
+		Sequencer->InvalidateGlobalMarkedFramesCache();
+	}
+}
+
+void ModifySequenceAndMovieScene(UMovieSceneSequence* const InSequence)
+{
+	if (!IsValid(InSequence))
+	{
+		return;
+	}
+
+	UMovieScene* const MovieScene = InSequence->GetMovieScene();
+	if (!IsValid(MovieScene))
+	{
+		return;
+	}
+
+	InSequence->Modify();
+	MovieScene->Modify();
+}
+
+void GetSequenceSubSections(UMovieSceneSequence* const InSequence, TArray<UMovieSceneSubSection*>& OutSubSections)
+{
+	if (!IsValid(InSequence))
+	{
+		return;
+	}
+
+	UMovieScene* const MovieScene = InSequence->GetMovieScene();
+	if (!IsValid(MovieScene))
+	{
+		return;
+	}
+
+	for (UMovieSceneSection* const Section : MovieScene->GetAllSections())
+	{
+		UMovieSceneSubSection* const SubSection = Cast<UMovieSceneSubSection>(Section);
+		if (IsValid(SubSection))
+		{
+			OutSubSections.Add(SubSection);
+		}
+	}
+}
+
+}
diff --git a/Engine/Source/Editor/Sequencer/Private/Tools/NavigationTool/Widgets/Columns/SNavigationToolComment.cpp b/Engine/Source/Editor/Sequencer/Private/Tools/NavigationTool/Widgets/Columns/SNavigationToolComment.cpp
new file mode 100644
--- /dev/null
+++ b/Engine/Source/Editor/Sequencer/Private/Tools/NavigationTool/Widgets/Columns/SNavigationToolComment.cpp
@@ -0,0 +1,120 @@
+// Copyright Epic Games, Inc. All Rights Reserved.
+
+#include "SNavigationToolComment.h"
+#include "Tools/NavigationTool/Items/INavigationToolItem.h"
+#include "Tools/NavigationTool/Items/NavigationToolItemUtils.h"
+#include "Tools/NavigationTool/Items/NavigationToolSequence.h"
+#include "Tools/NavigationTool/NavigationToolView.h"
+#include "Tools/NavigationTool/Widgets/SNavigationToolTreeRow.h"
+#include "TrackEditors/SubTrackEditorBase.h"
+#include "Widgets/Input/SEditableTextBox.h"
+#include "Widgets/Layout/SScrollBox.h"
+
+#define LOCTEXT_NAMESPACE "SNavigationToolComment"
+
+namespace UE::Sequencer::NavigationTool
+{
+
+void SNavigationToolComment::Construct(const FArguments& InArgs
+	, const FNavigationToolItemPtr& InItem
+	, const TSharedRef<INavigationToolView>& InView
+	, const TSharedRef<SNavigationToolTreeRow>& InRowWidget)
+{
+	WeakItem = InItem;
+	WeakView = InView;
+	WeakRowWidget = InRowWidget;
+
+	auto EmptyContent = [this]()
+	{
+		ChildSlot
+		[
+			SNullWidget::NullWidget
+		];
+	};
+
+	if (!WeakItem.IsValid() || !WeakView.IsValid())
+	{
+		EmptyContent();
+		return;
+	}
+
+	const FNavigationToolSequence* const SequenceItem = WeakItem.Pin()->CastTo<FNavigationToolSequence>();
+	if (!SequenceItem || !IsValid(SequenceItem->GetSubSection()))
+	{
+		EmptyContent();
+		return;
+	}
+
+	ChildSlot
+		[
+			SNew(SEditableTextBox)
+			.Text(this, &SNavigationToolComment::GetCommentText)
+			.OnTextChanged(this, &SNavigationToolComment::OnCommentTextChanged)
+			.OnTextCommitted(this, &SNavigationToolComment::OnCommentTextCommitted)
+		];
+}
+
+SNavigationToolComment::~SNavigationToolComment()
+{
+}
+
+FString SNavigationToolComment::GetMetaDataComment() const
+{
+	UMovieSceneMetaData* const SequenceMetaData = GetSequenceItemMetaData(WeakItem.Pin());
+	if (IsValid(SequenceMetaData))
+	{
+		return SequenceMetaData->GetNotes();
+	}
+	return FString();
+}
+
+FText SNavigationToolComment::GetCommentText() const
+{
+	return FText::FromString(GetMetaDataComment());
+}
+
+void SNavigationToolComment::OnCommentTextChanged(const FText& InNewText)
+{
+	OnCommentTextCommitted(InNewText, ETextCommit::Default);
+}
+
+void SNavigationToolComment::OnCommentTextCommitted(const FText& InNewText, const ETextCommit::Type InCommitType)
+{
+	const FNavigationToolSequence* const SequenceItem = WeakItem.Pin()->CastTo<FNavigationToolSequence>();
+	if (!SequenceItem)
+	{
+		return;
+	}
+
+	UMovieSceneSequence* const Sequence = SequenceItem->GetSequence();
+	if (!IsValid(Sequence))
+	{
+		return;
+	}
+
+	UMovieSceneMetaData* const SequenceMetaData = FSubTrackEditorUtil::FindOrAddMetaData(Sequence);
+	if (!IsValid(SequenceMetaData))
+	{
+		return;
+	}
+
+	if (InNewText.ToString().Equals(SequenceMetaData->GetNotes(), ESearchCase::CaseSensitive))
+	{
+		return;
+	}
+
+	const bool bShouldTransact = (InCommitType != ETextCommit::Default);
+	const FScopedTransaction Transaction(GetTransactionText(), bShouldTransact);
+
+	Sequence->Modify();
+	SequenceMetaData->SetNotes(InNewText.ToString());
+}
+
+FText SNavigationToolComment::GetTransactionText() const
+{
+	return LOCTEXT("SetSequenceCommentTransaction", "Set Sequence Comment");
+}
+
+}
+
+#undef LOCTEXT_NAMESPACE
