[Sequencer] Navigation Tool initial submit

diff --git a/Engine/Source/Editor/Sequencer/Private/Tools/NavigationTool/Widgets/Columns/SNavigationToolDeactiveState.cpp b/Engine/Source/Editor/Sequencer/Private/Tools/NavigationTool/Widgets/Columns/SNavigationToolDeactiveState.cpp
new file mode 100644
--- /dev/null
+++ b/Engine/Source/Editor/Sequencer/Private/Tools/NavigationTool/Widgets/Columns/SNavigationToolDeactiveState.cpp
@@ -0,0 +1,238 @@
+// Copyright Epic Games, Inc. All Rights Reserved.
+
+#include "SNavigationToolDeactiveState.h"
+#include "Editor.h"
+#include "ScopedTransaction.h"
+#include "Tools/NavigationTool/Columns/NavigationToolDeactiveStateColumn.h"
+#include "Tools/NavigationTool/Extensions/ISequenceInactivableExtension.h"
+#include "Tools/NavigationTool/Items/INavigationToolItem.h"
+#include "Tools/NavigationTool/NavigationToolView.h"
+#include "Tools/NavigationTool/Widgets/SNavigationToolTreeRow.h"
+
+#define LOCTEXT_NAMESPACE "SNavigationToolDeactiveState"
+
+namespace UE::Sequencer::NavigationTool
+{
+
+class FEvaluationStateDragDropOp : public FDragDropOperation
+{
+public:
+	DRAG_DROP_OPERATOR_TYPE(FVisibilityDragDropOp, FDragDropOperation)
+
+	/** Flag which defines whether to hide destination items or not */
+	EItemSequenceInactiveState InactiveState;
+
+	/** Undo transaction stolen from the gutter which is kept alive for the duration of the drag */
+	TUniquePtr<FScopedTransaction> UndoTransaction;
+
+	/** The widget decorator to use */
+	virtual TSharedPtr<SWidget> GetDefaultDecorator() const override
+	{
+		return SNullWidget::NullWidget;
+	}
+
+	/** Create a new drag and drop operation out of the specified flag */
+	static TSharedRef<FEvaluationStateDragDropOp> New(const EItemSequenceInactiveState InEvaluationState, TUniquePtr<FScopedTransaction>& ScopedTransaction)
+	{
+		TSharedRef<FEvaluationStateDragDropOp> Operation = MakeShared<FEvaluationStateDragDropOp>();
+		Operation->InactiveState = InEvaluationState;
+		Operation->UndoTransaction = MoveTemp(ScopedTransaction);
+		Operation->Construct();
+		return Operation;
+	}
+};
+
+EItemSequenceInactiveState GetItemInactiveState(const FNavigationToolItemPtr& InItem)
+{
+	if (!InItem.IsValid())
+	{
+		return EItemSequenceInactiveState::None;
+	}
+
+	ISequenceInactivableExtension* const InactivatableItem = InItem->CastTo<ISequenceInactivableExtension>();
+	if (!InactivatableItem)
+	{
+		return EItemSequenceInactiveState::None;
+	}
+
+	return InactivatableItem->GetInactiveState();
+}
+
+void SetItemInactive(const FNavigationToolItemPtr& InItem, const bool bInInactive)
+{
+	if (!InItem.IsValid())
+	{
+		return;
+	}
+
+	if (ISequenceInactivableExtension* const InactivatableItem = InItem->CastTo<ISequenceInactivableExtension>())
+	{
+		InactivatableItem->SetIsInactive(bInInactive);
+	}
+}
+
+void SNavigationToolDeactiveState::Construct(const FArguments& InArgs
+	, const TSharedRef<FNavigationToolDeactiveStateColumn>& InColumn
+	, const TSharedRef<INavigationToolView>& InView
+	, const FNavigationToolItemPtr& InItem
+	, const TSharedRef<SNavigationToolTreeRow>& InRowWidget)
+{
+	WeakColumn = InColumn;
+	WeakItem = InItem;
+	WeakView = InView;
+	WeakRowWidget = InRowWidget;
+
+	SImage::Construct(SImage::FArguments()
+		.IsEnabled(this, &SNavigationToolDeactiveState::IsVisibilityWidgetEnabled)
+		.ColorAndOpacity(this, &SNavigationToolDeactiveState::GetForegroundColor)
+		.Image(this, &SNavigationToolDeactiveState::GetBrush));
+}
+
+FReply SNavigationToolDeactiveState::OnDragDetected(const FGeometry& MyGeometry, const FPointerEvent& MouseEvent)
+{
+	if (MouseEvent.IsMouseButtonDown(EKeys::LeftMouseButton))
+	{
+		const EItemSequenceInactiveState NewState = (GetInactiveState() != EItemSequenceInactiveState::None)
+			? EItemSequenceInactiveState::Inactive : EItemSequenceInactiveState::None;
+		return FReply::Handled().BeginDragDrop(FEvaluationStateDragDropOp::New(NewState, UndoTransaction));
+	}
+	return FReply::Unhandled();
+}
+
+void SNavigationToolDeactiveState::OnDragEnter(const FGeometry& MyGeometry, const FDragDropEvent& DragDropEvent)
+{
+	const TSharedPtr<FEvaluationStateDragDropOp> DragDropOp = DragDropEvent.GetOperationAs<FEvaluationStateDragDropOp>();
+	if (DragDropOp.IsValid())
+	{
+		const bool bItemInactive = (DragDropOp->InactiveState != EItemSequenceInactiveState::None);
+		SetIsInactive(bItemInactive);
+	}
+}
+
+FReply SNavigationToolDeactiveState::HandleClick()
+{
+	if (!IsVisibilityWidgetEnabled())
+	{
+		return FReply::Unhandled();
+	}
+
+	const TSharedPtr<INavigationToolView> ToolView = WeakView.Pin();
+	const FNavigationToolItemPtr Item = WeakItem.Pin();
+	const TSharedPtr<FNavigationToolDeactiveStateColumn> Column = WeakColumn.Pin();
+
+	if (!ToolView.IsValid() || !Item.IsValid() || !Column.IsValid())
+	{
+		return FReply::Unhandled();
+	}
+
+	// Open an undo transaction
+	UndoTransaction = MakeUnique<FScopedTransaction>(LOCTEXT("SetNavigationToolItemInactiveState", "Set Item Inactive State"));
+
+	const bool bNewIsInactive = (GetInactiveState() != EItemSequenceInactiveState::Inactive);
+
+	// We operate on all the selected items if the specified item is selected
+	if (ToolView->IsItemSelected(Item.ToSharedRef()))
+	{
+		for (FNavigationToolItemPtr& SelectedItem : ToolView->GetSelectedItems())
+		{
+			SetItemInactive(SelectedItem, bNewIsInactive);
+		}
+	}
+	else
+	{
+		SetIsInactive(bNewIsInactive);
+	}
+
+	return FReply::Handled().DetectDrag(SharedThis(this), EKeys::LeftMouseButton);
+}
+
+FReply SNavigationToolDeactiveState::OnMouseButtonDoubleClick(const FGeometry& InMyGeometry, const FPointerEvent& InMouseEvent)
+{
+	return HandleClick();
+}
+
+FReply SNavigationToolDeactiveState::OnMouseButtonDown(const FGeometry& MyGeometry, const FPointerEvent& MouseEvent)
+{
+	if (MouseEvent.GetEffectingButton() != EKeys::LeftMouseButton)
+	{
+		return FReply::Unhandled();
+	}
+
+	return HandleClick();
+}
+
+FReply SNavigationToolDeactiveState::OnMouseButtonUp(const FGeometry& MyGeometry, const FPointerEvent& MouseEvent)
+{
+	if (MouseEvent.GetEffectingButton() == EKeys::LeftMouseButton)
+	{
+		UndoTransaction.Reset();
+		return FReply::Handled();
+	}
+
+	return FReply::Unhandled();
+}
+
+void SNavigationToolDeactiveState::OnMouseCaptureLost(const FCaptureLostEvent& CaptureLostEvent)
+{
+	UndoTransaction.Reset();
+}
+
+const FSlateBrush* SNavigationToolDeactiveState::GetBrush() const
+{
+	return FAppStyle::GetBrush(TEXT("Sequencer.Column.Mute"));
+}
+
+FSlateColor SNavigationToolDeactiveState::GetForegroundColor() const
+{
+	const FNavigationToolItemPtr Item = WeakItem.Pin();
+	if (!Item.IsValid())
+	{
+		return FStyleColors::Transparent;
+	}
+
+	const bool bIsItemHovered = WeakRowWidget.IsValid() && WeakRowWidget.Pin()->IsHovered();
+	if (bIsItemHovered || IsHovered())
+	{
+		switch (GetInactiveState())
+		{
+		case EItemSequenceInactiveState::None:
+			return FStyleColors::White25;
+
+		case EItemSequenceInactiveState::PartiallyInactive:
+			return FStyleColors::ForegroundHover;
+
+		case EItemSequenceInactiveState::Inactive:
+			return FStyleColors::ForegroundHover;
+		}
+	}
+	else
+	{
+		switch (GetInactiveState())
+		{
+		case EItemSequenceInactiveState::None:
+			return FStyleColors::Transparent;
+
+		case EItemSequenceInactiveState::PartiallyInactive:
+			return FStyleColors::White25;
+
+		case EItemSequenceInactiveState::Inactive:
+			return FStyleColors::Foreground;
+		}
+	}
+
+	return FStyleColors::Transparent;
+}
+
+EItemSequenceInactiveState SNavigationToolDeactiveState::GetInactiveState() const
+{
+	return GetItemInactiveState(WeakItem.Pin());
+}
+
+void SNavigationToolDeactiveState::SetIsInactive(const bool bInIsInactive)
+{
+	SetItemInactive(WeakItem.Pin(), bInIsInactive);
+}
+
+}
+
+#undef LOCTEXT_NAMESPACE
diff --git a/Engine/Source/Editor/Sequencer/Private/Tools/NavigationTool/Widgets/Columns/SNavigationToolDeactiveState.h b/Engine/Source/Editor/Sequencer/Private/Tools/NavigationTool/Widgets/Columns/SNavigationToolDeactiveState.h
new file mode 100644
--- /dev/null
+++ b/Engine/Source/Editor/Sequencer/Private/Tools/NavigationTool/Widgets/Columns/SNavigationToolDeactiveState.h
@@ -0,0 +1,67 @@
+// Copyright Epic Games, Inc. All Rights Reserved.
+
+#pragma once
+
+#include "ScopedTransaction.h"
+#include "Templates/UniquePtr.h"
+#include "Tools/NavigationTool/NavigationToolDefines.h"
+#include "Widgets/Images/SImage.h"
+
+namespace UE::Sequencer::NavigationTool
+{
+
+class FNavigationToolItem;
+class FNavigationToolDeactiveStateColumn;
+class FNavigationToolView;
+class INavigationToolItem;
+class INavigationToolView;
+class SNavigationToolTreeRow;
+enum class EItemSequenceInactiveState;
+
+/** Widget responsible for managing the visibility for a single item */
+class SNavigationToolDeactiveState : public SImage
+{
+public:
+	SLATE_BEGIN_ARGS(SNavigationToolDeactiveState) {}
+	SLATE_END_ARGS()
+
+	void Construct(const FArguments& InArgs
+		, const TSharedRef<FNavigationToolDeactiveStateColumn>& InColumn
+		, const TSharedRef<INavigationToolView>& InView
+		, const FNavigationToolItemPtr& InItem
+		, const TSharedRef<SNavigationToolTreeRow>& InRowWidget);
+
+private:
+	/** Returns whether the widget is enabled or not */
+	virtual bool IsVisibilityWidgetEnabled() const { return true; }
+
+	virtual const FSlateBrush* GetBrush() const;
+
+	//~ Begin SWidget
+	virtual FReply OnDragDetected(const FGeometry& MyGeometry, const FPointerEvent& MouseEvent) override;
+	virtual void OnDragEnter(const FGeometry& MyGeometry, const FDragDropEvent& DragDropEvent) override;
+	virtual FReply OnMouseButtonDoubleClick(const FGeometry& InMyGeometry, const FPointerEvent& InMouseEvent) override;
+	virtual FReply OnMouseButtonDown(const FGeometry& MyGeometry, const FPointerEvent& MouseEvent) override;
+	virtual FReply OnMouseButtonUp(const FGeometry& MyGeometry, const FPointerEvent& MouseEvent) override;
+	virtual void OnMouseCaptureLost(const FCaptureLostEvent& CaptureLostEvent) override;
+	//~ End SWidget
+
+	FReply HandleClick();
+
+	virtual FSlateColor GetForegroundColor() const override;
+
+	EItemSequenceInactiveState GetInactiveState() const;
+	void SetIsInactive(const bool bInIsInactive);
+
+	TWeakPtr<FNavigationToolDeactiveStateColumn> WeakColumn;
+
+	FNavigationToolItemWeakPtr WeakItem;
+
+	TWeakPtr<INavigationToolView> WeakView;
+
+	TWeakPtr<SNavigationToolTreeRow> WeakRowWidget;
+
+	TUniquePtr<FScopedTransaction> UndoTransaction;
+};
+
+}
diff --git a/Engine/Source/Editor/Sequencer/Private/Tools/NavigationTool/Widgets/Columns/SNavigationToolEndFrameOffset.cpp b/Engine/Source/Editor/Sequencer/Private/Tools/NavigationTool/Widgets/Columns/SNavigationToolEndFrameOffset.cpp
new file mode 100644
--- /dev/null
+++ b/Engine/Source/Editor/Sequencer/Private/Tools/NavigationTool/Widgets/Columns/SNavigationToolEndFrameOffset.cpp
@@ -0,0 +1,137 @@
+// Copyright Epic Games, Inc. All Rights Reserved.
+
+#include "SNavigationToolEndFrameOffset.h"
+#include "Sequencer.h"
+#include "Tools/NavigationTool/Items/NavigationToolItemUtils.h"
+#include "Tools/NavigationTool/Items/NavigationToolSequence.h"
+#include "Tools/NavigationTool/NavigationToolView.h"
+#include "Tools/NavigationTool/Widgets/SNavigationToolTreeRow.h"
+#include "TrackEditors/SubTrackEditorBase.h"
+#include "Widgets/Input/SSpinBox.h"
+
+using namespace UE::Sequencer::NavigationTool;
+
+#define LOCTEXT_NAMESPACE "SNavigationToolEndFrameOffset"
+
+void SNavigationToolEndFrameOffset::Construct(const FArguments& InArgs
+	, const FNavigationToolItemPtr& InItem
+	, const TSharedRef<INavigationToolView>& InView
+	, const TSharedRef<SNavigationToolTreeRow>& InRowWidget)
+{
+	WeakItem = InItem;
+	WeakView = InView;
+	WeakRowWidget = InRowWidget;
+
+	auto EmptyContent = [this]()
+	{
+		ChildSlot
+		[
+			SNullWidget::NullWidget
+		];
+	};
+
+	if (!WeakItem.IsValid() || !WeakView.IsValid())
+	{
+		EmptyContent();
+		return;
+	}
+
+	const FNavigationToolSequence* const SequenceItem = WeakItem.Pin()->CastTo<FNavigationToolSequence>();
+	if (!SequenceItem || !IsValid(SequenceItem->GetSubSection()))
+	{
+		EmptyContent();
+		return;
+	}
+
+	ChildSlot
+		[
+			SNew(SSpinBox<int32>)
+			.Justification(ETextJustify::Center)
+			.Value(this, &SNavigationToolEndFrameOffset::GetValue)
+			.OnValueChanged(this, &SNavigationToolEndFrameOffset::OnValueChanged)
+			.OnValueCommitted(this, &SNavigationToolEndFrameOffset::OnValueCommitted)
+			.OnBeginSliderMovement(this, &SNavigationToolEndFrameOffset::OnBeginSliderMovement)
+			.OnEndSliderMovement(this, &SNavigationToolEndFrameOffset::OnEndSliderMovement)
+		];
+}
+
+int32 SNavigationToolEndFrameOffset::GetValue() const
+{
+	const TSharedPtr<INavigationToolView> ToolView = WeakView.Pin();
+	if (!ToolView.IsValid())
+	{
+		return 0;
+	}
+
+	const TSharedPtr<ISequencer> Sequencer = ToolView->GetSequencer();
+	if (!Sequencer.IsValid())
+	{
+		return 0;
+	}
+
+	UMovieSceneSubSection* const SubSection = Cast<UMovieSceneSubSection>(GetSequenceItemSubSection(WeakItem.Pin()));
+	if (!IsValid(SubSection))
+	{
+		return 0;
+	}
+
+	const FFrameNumber Offset = ConvertFrameTime(SubSection->Parameters.EndFrameOffset.Value
+		, Sequencer->GetFocusedTickResolution()
+		, Sequencer->GetFocusedDisplayRate()).FrameNumber;
+	return Offset.Value;
+}
+
+void SNavigationToolEndFrameOffset::OnValueChanged(const int32 InNewValue)
+{
+	const TSharedPtr<INavigationToolView> ToolView = WeakView.Pin();
+	if (!ToolView.IsValid())
+	{
+		return;
+	}
+
+	const TSharedPtr<ISequencer> Sequencer = ToolView->GetSequencer();
+	if (!Sequencer.IsValid())
+	{
+		return;
+	}
+
+	UMovieSceneSubSection* const SubSection = Cast<UMovieSceneSubSection>(GetSequenceItemSubSection(WeakItem.Pin()));
+	if (!IsValid(SubSection))
+	{
+		return;
+	}
+
+	const FFrameNumber Offset = ConvertFrameTime(SubSection->Parameters.EndFrameOffset.Value
+		, Sequencer->GetFocusedDisplayRate()
+		, Sequencer->GetFocusedTickResolution()).FrameNumber;
+
+	const FScopedTransaction Transaction(GetTransactionText(), !UndoTransaction.IsValid());
+
+	SubSection->Modify();
+	SubSection->Parameters.EndFrameOffset = Offset;
+}
+
+void SNavigationToolEndFrameOffset::OnValueCommitted(const int32 InNewValue, const ETextCommit::Type InCommitType)
+{
+	OnValueChanged(InNewValue);
+}
+
+void SNavigationToolEndFrameOffset::OnBeginSliderMovement()
+{
+	if (!UndoTransaction.IsValid())
+	{
+		UndoTransaction = MakeUnique<FScopedTransaction>(GetTransactionText());
+	}
+}
+
+void SNavigationToolEndFrameOffset::OnEndSliderMovement(const int32 InNewValue)
+{
+	UndoTransaction.Reset();
+}
+
+FText SNavigationToolEndFrameOffset::GetTransactionText() const
+{
+	return LOCTEXT("SetSequenceEndFrameOffsetTransaction", "Set Sequence End Frame Offset");
+}
+
+#undef LOCTEXT_NAMESPACE
diff --git a/Engine/Source/Editor/Sequencer/Private/Tools/NavigationTool/Widgets/Columns/SNavigationToolEndFrameOffset.h b/Engine/Source/Editor/Sequencer/Private/Tools/NavigationTool/Widgets/Columns/SNavigationToolEndFrameOffset.h
new file mode 100644
--- /dev/null
+++ b/Engine/Source/Editor/Sequencer/Private/Tools/NavigationTool/Widgets/Columns/SNavigationToolEndFrameOffset.h
@@ -0,0 +1,47 @@
+// Copyright Epic Games, Inc. All Rights Reserved.
+
+#pragma once
+
+#include "ScopedTransaction.h"
+#include "Templates/UniquePtr.h"
+#include "Tools/NavigationTool/NavigationToolDefines.h"
+#include "Widgets/SCompoundWidget.h"
+
+namespace UE::Sequencer::NavigationTool
+{
+
+class INavigationToolView;
+class SNavigationToolTreeRow;
+
+class SNavigationToolEndFrameOffset : public SCompoundWidget
+{
+public:
+	SLATE_BEGIN_ARGS(SNavigationToolEndFrameOffset) {}
+	SLATE_END_ARGS()
+
+	virtual ~SNavigationToolEndFrameOffset() override {}
+
+	void Construct(const FArguments& InArgs
+		, const FNavigationToolItemPtr& InItem
+		, const TSharedRef<INavigationToolView>& InView
+		, const TSharedRef<SNavigationToolTreeRow>& InRowWidget);
+
+private:
+	int32 GetValue() const;
+	void OnValueChanged(const int32 InNewValue);
+	void OnValueCommitted(const int32 InNewValue, const ETextCommit::Type InCommitType);
+	void OnBeginSliderMovement();
+	void OnEndSliderMovement(const int32 InNewValue);
+
+	FText GetTransactionText() const;
+
+	FNavigationToolItemWeakPtr WeakItem;
+
+	TWeakPtr<INavigationToolView> WeakView;
+
+	TWeakPtr<SNavigationToolTreeRow> WeakRowWidget;
+
+	TUniquePtr<FScopedTransaction> UndoTransaction;
+};
+
+}
diff --git a/Engine/Source/Editor/Sequencer/Private/Tools/NavigationTool/Widgets/Columns/SNavigationToolHBias.cpp b/Engine/Source/Editor/Sequencer/Private/Tools/NavigationTool/Widgets/Columns/SNavigationToolHBias.cpp
new file mode 100644
--- /dev/null
+++ b/Engine/Source/Editor/Sequencer/Private/Tools/NavigationTool/Widgets/Columns/SNavigationToolHBias.cpp
@@ -0,0 +1,124 @@
+// Copyright Epic Games, Inc. All Rights Reserved.
+
+#include "SNavigationToolHBias.h"
+#include "ScopedTransaction.h"
+#include "Tools/NavigationTool/Items/NavigationToolItemUtils.h"
+#include "Tools/NavigationTool/Items/NavigationToolSequence.h"
+#include "Tools/NavigationTool/NavigationTool.h"
+#include "Tools/NavigationTool/NavigationToolView.h"
+#include "Tools/NavigationTool/Widgets/SNavigationToolTreeRow.h"
+#include "TrackEditors/SubTrackEditorBase.h"
+#include "Widgets/Input/SSpinBox.h"
+
+#define LOCTEXT_NAMESPACE "SNavigationToolHBias"
+
+namespace UE::Sequencer::NavigationTool
+{
+
+void SNavigationToolHBias::Construct(const FArguments& InArgs
+	, const FNavigationToolItemPtr& InItem
+	, const TSharedRef<INavigationToolView>& InView
+	, const TSharedRef<SNavigationToolTreeRow>& InRowWidget)
+{
+	WeakItem = InItem;
+	WeakView = InView;
+	WeakRowWidget = InRowWidget;
+
+	auto EmptyContent = [this]()
+	{
+		ChildSlot
+		[
+			SNullWidget::NullWidget
+		];
+	};
+
+	if (!WeakItem.IsValid() || !WeakView.IsValid())
+	{
+		EmptyContent();
+		return;
+	}
+
+	const FNavigationToolSequence* const SequenceItem = WeakItem.Pin()->CastTo<FNavigationToolSequence>();
+	if (!SequenceItem)
+	{
+		EmptyContent();
+		return;
+	}
+
+	UMovieSceneSubSection* const SubSection = SequenceItem->GetSubSection();
+	if (!IsValid(SubSection))
+	{
+		EmptyContent();
+		return;
+	}
+
+	ChildSlot
+		[
+			SNew(SSpinBox<int32>)
+			.Justification(ETextJustify::Center)
+			.Style(&FAppStyle::Get(), TEXT("Sequencer.NavigationTool.SpinBox"))
+			.Value(this, &SNavigationToolHBias::GetValue)
+			.OnValueChanged(this, &SNavigationToolHBias::OnValueChanged)
+			.OnValueCommitted(this, &SNavigationToolHBias::OnValueCommitted)
+			.OnBeginSliderMovement(this, &SNavigationToolHBias::OnBeginSliderMovement)
+			.OnEndSliderMovement(this, &SNavigationToolHBias::OnEndSliderMovement)
+		];
+}
+
+int32 SNavigationToolHBias::GetValue() const
+{
+	UMovieSceneSubSection* const SubSection = Cast<UMovieSceneSubSection>(GetSequenceItemSubSection(WeakItem.Pin()));
+	if (!IsValid(SubSection))
+	{
+		return 0;
+	}
+
+	return SubSection->Parameters.HierarchicalBias;
+}
+
+void SNavigationToolHBias::OnValueChanged(const int32 InNewValue)
+{
+	OnValueCommitted(InNewValue, ETextCommit::Default);
+}
+
+void SNavigationToolHBias::OnValueCommitted(const int32 InNewValue, const ETextCommit::Type InCommitType)
+{
+	UMovieSceneSubSection* const SubSection = Cast<UMovieSceneSubSection>(GetSequenceItemSubSection(WeakItem.Pin()));
+	if (!IsValid(SubSection))
+	{
+		return;
+	}
+
+	if (InNewValue == SubSection->Parameters.HierarchicalBias)
+	{
+		return;
+	}
+
+	const bool bShouldTransact = !UndoTransaction.IsValid() && (InCommitType == ETextCommit::OnEnter);
+	const FScopedTransaction Transaction(GetTransactionText(), bShouldTransact);
+
+	SubSection->Modify();
+	SubSection->Parameters.HierarchicalBias = InNewValue;
+}
+
+void SNavigationToolHBias::OnBeginSliderMovement()
+{
+	if (!UndoTransaction.IsValid())
+	{
+		UndoTransaction = MakeUnique<FScopedTransaction>(GetTransactionText());
+	}
+}
+
+void SNavigationToolHBias::OnEndSliderMovement(const int32 InNewValue)
+{
+	UndoTransaction.Reset();
+}
+
+FText SNavigationToolHBias::GetTransactionText() const
+{
+	return LOCTEXT("SetSequenceHBiasTransaction", "Set Sequence HBias");
+}
+
+}
+
+#undef LOCTEXT_NAMESPACE
diff --git a/Engine/Source/Editor/Sequencer/Private/Tools/NavigationTool/Widgets/Columns/SNavigationToolInTime.cpp b/Engine/Source/Editor/Sequencer/Private/Tools/NavigationTool/Widgets/Columns/SNavigationToolInTime.cpp
new file mode 100644
--- /dev/null
+++ b/Engine/Source/Editor/Sequencer/Private/Tools/NavigationTool/Widgets/Columns/SNavigationToolInTime.cpp
@@ -0,0 +1,149 @@
+// Copyright Epic Games, Inc. All Rights Reserved.
+
+#include "SNavigationToolInTime.h"
+#include "Sequencer.h"
+#include "Tools/NavigationTool/Items/NavigationToolItemUtils.h"
+#include "Tools/NavigationTool/Items/NavigationToolSequence.h"
+#include "Tools/NavigationTool/NavigationTool.h"
+#include "Tools/NavigationTool/NavigationToolView.h"
+#include "Tools/NavigationTool/Widgets/SNavigationToolTreeRow.h"
+#include "TrackEditors/SubTrackEditorBase.h"
+#include "Widgets/Input/SSpinBox.h"
+
+#define LOCTEXT_NAMESPACE "SNavigationToolInTime"
+
+namespace UE::Sequencer::NavigationTool
+{
+
+void SNavigationToolInTime::Construct(const FArguments& InArgs
+	, const FNavigationToolItemPtr& InItem
+	, const TSharedRef<INavigationToolView>& InView
+	, const TSharedRef<SNavigationToolTreeRow>& InRowWidget)
+{
+	WeakItem = InItem;
+	WeakView = InView;
+	WeakRowWidget = InRowWidget;
+
+	auto EmptyContent = [this]()
+	{
+		ChildSlot
+		[
+			SNullWidget::NullWidget
+		];
+	};
+
+	if (!WeakItem.IsValid() || !WeakView.IsValid())
+	{
+		EmptyContent();
+		return;
+	}
+
+	const FNavigationToolSequence* const SequenceItem = WeakItem.Pin()->CastTo<FNavigationToolSequence>();
+	if (!SequenceItem || !IsValid(SequenceItem->GetSubSection()))
+	{
+		EmptyContent();
+		return;
+	}
+
+	ChildSlot
+		[
+			SNew(SSpinBox<int32>)
+			.Justification(ETextJustify::Center)
+			.Style(&FAppStyle::GetWidgetStyle<FSpinBoxStyle>(TEXT("Sequencer.NavigationTool.InTime.SpinBox")))
+			.Value(this, &SNavigationToolInTime::GetFrameTimeValue)
+			.OnValueChanged(this, &SNavigationToolInTime::OnFrameTimeValueChanged)
+			.OnValueCommitted(this, &SNavigationToolInTime::OnFrameTimeValueCommitted)
+			.OnBeginSliderMovement(this, &SNavigationToolInTime::OnBeginSliderMovement)
+			.OnEndSliderMovement(this, &SNavigationToolInTime::OnEndSliderMovement)
+		];
+}
+
+int32 SNavigationToolInTime::GetFrameTimeValue() const
+{
+	const TSharedPtr<INavigationToolView> ToolView = WeakView.Pin();
+	if (!ToolView.IsValid())
+	{
+		return 0;
+	}
+
+	const TSharedPtr<ISequencer> Sequencer = ToolView->GetSequencer();
+	if (!Sequencer.IsValid())
+	{
+		return 0;
+	}
+
+	UMovieSceneSubSection* const SubSection = Cast<UMovieSceneSubSection>(GetSequenceItemSubSection(WeakItem.Pin()));
+	if (!IsValid(SubSection))
+	{
+		return 0;
+	}
+
+	const FFrameNumber FrameTime = ConvertFrameTime(SubSection->SectionRange.Value.GetLowerBoundValue()
+		, Sequencer->GetFocusedTickResolution()
+		, Sequencer->GetFocusedDisplayRate()).FrameNumber;
+
+	return FrameTime.Value;
+}
+
+void SNavigationToolInTime::OnFrameTimeValueChanged(const int32 InNewValue)
+{
+	OnFrameTimeValueCommitted(InNewValue, ETextCommit::Default);
+}
+
+void SNavigationToolInTime::OnFrameTimeValueCommitted(const int32 InNewValue, const ETextCommit::Type InCommitType)
+{
+	const TSharedPtr<INavigationToolView> ToolView = WeakView.Pin();
+	if (!ToolView.IsValid())
+	{
+		return;
+	}
+
+	const TSharedPtr<ISequencer> Sequencer = ToolView->GetSequencer();
+	if (!Sequencer.IsValid())
+	{
+		return;
+	}
+
+	UMovieSceneSubSection* const SubSection = Cast<UMovieSceneSubSection>(GetSequenceItemSubSection(WeakItem.Pin()));
+	if (!IsValid(SubSection))
+	{
+		return;
+	}
+
+	const FFrameNumber NewFrameTime = ConvertFrameTime(InNewValue
+		, Sequencer->GetFocusedDisplayRate()
+		, Sequencer->GetFocusedTickResolution()).FrameNumber;
+
+	if (SubSection->SectionRange.Value.GetLowerBoundValue() == NewFrameTime)
+	{
+		return;
+	}
+
+	const bool bShouldTransact = !UndoTransaction.IsValid() && (InCommitType == ETextCommit::OnEnter);
+	const FScopedTransaction Transaction(GetTransactionText(), bShouldTransact);
+
+	SubSection->Modify();
+	SubSection->SectionRange.Value.SetLowerBoundValue(NewFrameTime);
+}
+
+void SNavigationToolInTime::OnBeginSliderMovement()
+{
+	if (!UndoTransaction.IsValid())
+	{
+		UndoTransaction = MakeUnique<FScopedTransaction>(GetTransactionText());
+	}
+}
+
+void SNavigationToolInTime::OnEndSliderMovement(const int32 InNewValue)
+{
+	UndoTransaction.Reset();
+}
+
+FText SNavigationToolInTime::GetTransactionText() const
+{
+	return LOCTEXT("SetSequenceInTimeTransaction", "Set Sequence In Time");
+}
+
+}
+
+#undef LOCTEXT_NAMESPACE
diff --git a/Engine/Source/Editor/Sequencer/Private/Tools/NavigationTool/Widgets/Columns/SNavigationToolInTime.h b/Engine/Source/Editor/Sequencer/Private/Tools/NavigationTool/Widgets/Columns/SNavigationToolInTime.h
new file mode 100644
--- /dev/null
+++ b/Engine/Source/Editor/Sequencer/Private/Tools/NavigationTool/Widgets/Columns/SNavigationToolInTime.h
@@ -0,0 +1,47 @@
+// Copyright Epic Games, Inc. All Rights Reserved.
+
+#pragma once
+
+#include "ScopedTransaction.h"
+#include "Templates/UniquePtr.h"
+#include "Tools/NavigationTool/NavigationToolDefines.h"
+#include "Widgets/SCompoundWidget.h"
+
+namespace UE::Sequencer::NavigationTool
+{
+
+class INavigationToolView;
+class SNavigationToolTreeRow;
+
+class SNavigationToolInTime : public SCompoundWidget
+{
+public:
+	SLATE_BEGIN_ARGS(SNavigationToolInTime) {}
+	SLATE_END_ARGS()
+
+	virtual ~SNavigationToolInTime() override {}
+
+	void Construct(const FArguments& InArgs
+		, const FNavigationToolItemPtr& InItem
+		, const TSharedRef<INavigationToolView>& InView
+		, const TSharedRef<SNavigationToolTreeRow>& InRowWidget);
+
+private:
+	int32 GetFrameTimeValue() const;
+	void OnFrameTimeValueChanged(const int32 InNewValue);
+	void OnFrameTimeValueCommitted(const int32 InNewValue, const ETextCommit::Type InCommitType);
+	void OnBeginSliderMovement();
+	void OnEndSliderMovement(const int32 InNewValue);
+
+	FText GetTransactionText() const;
+
+	FNavigationToolItemWeakPtr WeakItem;
+
+	TWeakPtr<INavigationToolView> WeakView;
+
+	TWeakPtr<SNavigationToolTreeRow> WeakRowWidget;
+
+	TUniquePtr<FScopedTransaction> UndoTransaction;
+};
+
+}
diff --git a/Engine/Source/Editor/Sequencer/Private/Tools/NavigationTool/Widgets/Columns/SNavigationToolItemList.cpp b/Engine/Source/Editor/Sequencer/Private/Tools/NavigationTool/Widgets/Columns/SNavigationToolItemList.cpp
new file mode 100644
--- /dev/null
+++ b/Engine/Source/Editor/Sequencer/Private/Tools/NavigationTool/Widgets/Columns/SNavigationToolItemList.cpp
@@ -0,0 +1,237 @@
+// Copyright Epic Games, Inc. All Rights Reserved.
+
+#include "SNavigationToolItemList.h"
+#include "Tools/NavigationTool/Items/NavigationToolItem.h"
+#include "Tools/NavigationTool/NavigationToolView.h"
+#include "Tools/NavigationTool/Widgets/Columns/SNavigationToolItemChip.h"
+#include "Tools/NavigationTool/Widgets/SNavigationToolTreeRow.h"
+#include "Widgets/Layout/SScrollBox.h"
+
+#define LOCTEXT_NAMESPACE "SNavigationToolItemList"
+
+namespace UE::Sequencer::NavigationTool
+{
+
+void SNavigationToolItemList::Construct(const FArguments& InArgs
+	, const FNavigationToolItemPtr& InItem
+	, const TSharedRef<INavigationToolView>& InView
+	, const TSharedRef<SNavigationToolTreeRow>& InRowWidget)
+{
+	WeakItem = InItem;
+	WeakView = InView;
+	WeakRowWidget = InRowWidget;
+
+	InView->GetOnToolViewRefreshed().AddSP(this, &SNavigationToolItemList::Refresh);
+	InItem->OnExpansionChanged().AddSP(this, &SNavigationToolItemList::OnItemExpansionChanged);
+
+	ChildSlot
+		.HAlign(HAlign_Left)
+		[
+			SAssignNew(ItemListBox, SScrollBox)
+			.ConsumeMouseWheel(EConsumeMouseWheel::Always)
+			.Orientation(Orient_Horizontal)
+			.ScrollBarThickness(FVector2D(2.f))
+		];
+
+	Refresh();
+}
+
+SNavigationToolItemList::~SNavigationToolItemList()
+{
+	if (WeakView.IsValid())
+	{
+		WeakView.Pin()->GetOnToolViewRefreshed().RemoveAll(this);
+	}
+	if (WeakItem.IsValid())
+	{
+		WeakItem.Pin()->OnExpansionChanged().RemoveAll(this);
+	}
+}
+
+void SNavigationToolItemList::OnItemExpansionChanged(const TSharedPtr<INavigationToolView>& InToolView, bool bInIsExpanded)
+{
+	Refresh();
+}
+
+void SNavigationToolItemList::Refresh()
+{
+	ItemListBox->ClearChildren();
+	WeakChildItems.Reset();
+
+	const FNavigationToolItemPtr ParentItem = WeakItem.Pin();
+	const TSharedPtr<SNavigationToolTreeRow> TreeRow = WeakRowWidget.Pin();
+	const TSharedPtr<FNavigationToolView> ToolView = StaticCastSharedPtr<FNavigationToolView>(WeakView.Pin());
+
+	if (!ParentItem.IsValid() || !ToolView.IsValid() || !TreeRow.IsValid())
+	{
+		return;
+	}
+
+	TSet<FNavigationToolItemPtr> DisallowedItems;
+	{
+		TArray<FNavigationToolItemPtr> ItemsToDisallow;
+
+		// First get the Children in the Navigation Tool View to see if we need to Disallow some of the Items
+		// to avoid redundancy or unnecessarily showing other item's children (e.g. an actor's)
+		static const TSet<FNavigationToolItemPtr> EmptySet;
+		ToolView->GetChildrenOfItem(ParentItem, ItemsToDisallow, ENavigationToolItemViewMode::ItemTree, EmptySet);
+
+		// If Parent item is Collapsed, only disallow items that are top levels (as these should deal with their own item list)
+		// Items that can't be top level should be visualized by the Parent Item when Collapsed
+		if (!EnumHasAllFlags(ToolView->GetViewItemFlags(ParentItem), ENavigationToolItemFlags::Expanded))
+		{
+			// keep only top level items
+			ItemsToDisallow.RemoveAllSwap([](const FNavigationToolItemPtr& Child)
+				{
+					return !Child.IsValid() || !Child->CanBeTopLevel();
+				}
+				, EAllowShrinking::No);
+		}
+
+		DisallowedItems.Append(MoveTemp(ItemsToDisallow));
+	}
+
+	TArray<FNavigationToolItemPtr> Children;
+	ToolView->GetChildrenOfItem(ParentItem, Children, ENavigationToolItemViewMode::HorizontalItemList, DisallowedItems);
+
+	// TArray<>::Pop will be removing the item from the end which results in the Item List being in reverse
+	// so instead of immediately adding it to the slot , we will add it to this array and reverse it
+	TArray<FNavigationToolItemPtr> ItemsToAdd;
+
+	const INavigationToolView& ToolViewConstRef = *ToolView;
+
+	while (!Children.IsEmpty())
+	{
+		const FNavigationToolItemPtr Child = Children.Pop();
+
+		if (!Child.IsValid() || DisallowedItems.Contains(Child))
+		{
+			continue;
+		}
+
+		if (Child->IsViewModeSupported(ENavigationToolItemViewMode::HorizontalItemList, ToolViewConstRef))
+		{
+			ItemsToAdd.Add(Child);
+		}
+		else
+		{
+			ToolView->GetChildrenOfItem(Child, Children, ENavigationToolItemViewMode::HorizontalItemList, DisallowedItems);
+		}
+	}
+
+	WeakChildItems.Reserve(ItemsToAdd.Num());
+
+	for (int32 ItemIndex = ItemsToAdd.Num() - 1; ItemIndex >= 0; --ItemIndex)
+	{
+		const FNavigationToolItemPtr& Item = ItemsToAdd[ItemIndex];
+		WeakChildItems.Add(Item);
+		ItemListBox->AddSlot()
+			.Padding(0.f, 1.f)
+			[
+				SNew(SNavigationToolItemChip, Item.ToSharedRef(), ToolView)
+				.OnItemChipClicked(this, &SNavigationToolItemList::OnItemChipSelected)
+				.ChipStyle(TreeRow->GetStyle())
+				.OnValidDragOver(this, &SNavigationToolItemList::OnItemChipValidDragOver)
+			];
+	}
+}
+
+FReply SNavigationToolItemList::OnItemChipSelected(const FNavigationToolItemPtr& InItem, const FPointerEvent& InMouseEvent)
+{
+	const TSharedPtr<INavigationToolView> ToolView = WeakView.Pin();
+	const FNavigationToolItemPtr ParentItem = WeakItem.Pin();
+
+	if (!InItem.IsValid() || !ToolView.IsValid() || !ParentItem.IsValid())
+	{
+		return FReply::Unhandled();
+	}
+
+	TArray<FNavigationToolItemPtr> SelectedItems = ToolView->GetSelectedItems();
+
+	ENavigationToolItemSelectionFlags SelectionFlags = ENavigationToolItemSelectionFlags::SignalSelectionChange;
+
+	if (InMouseEvent.IsControlDown())
+	{
+		// Deselect if the Item in question is already selected (i.e. CTRL also behaves like a Toggle)
+		if (ToolView->IsItemSelected(InItem))
+		{
+			SelectedItems.Remove(InItem);
+			ToolView->SelectItems(MoveTemp(SelectedItems), SelectionFlags);
+			return FReply::Handled();
+		}
+
+		SelectionFlags |= ENavigationToolItemSelectionFlags::AppendToCurrentSelection;
+	}
+
+	// add the given item since it will be selected regardless if Shift (Range Selection) is pressed or not
+	TArray<FNavigationToolItemPtr> ItemsToSelect { InItem };
+
+	if (InMouseEvent.IsShiftDown() && !SelectedItems.IsEmpty())
+	{
+		// Similar to Control, Shift should always append Selections, never remove
+		SelectionFlags |= ENavigationToolItemSelectionFlags::AppendToCurrentSelection;
+
+		int32 StartIndex = INDEX_NONE;
+
+		// Find whether there's a Child in the Selected Item List.
+		// Iterate in Reverse order since we want more recent Selections to take precedence
+		for (int32 ItemIndex = SelectedItems.Num() - 1; ItemIndex >= 0; --ItemIndex)
+		{
+			StartIndex = WeakChildItems.Find(SelectedItems[ItemIndex]);
+			if (StartIndex != INDEX_NONE)
+			{
+				break;
+			}
+		}
+
+		// If we manage to find an already selected child, start the range selection
+		if (StartIndex != INDEX_NONE)
+		{
+			int32 TargetIndex = WeakChildItems.Find(InItem);
+
+			StartIndex  = FMath::Clamp(StartIndex, 0, WeakChildItems.Num());
+			TargetIndex = FMath::Clamp(TargetIndex, 0, WeakChildItems.Num());
+
+			// If Target Index comes before, just iterate normally, and reverse later
+			if (TargetIndex < StartIndex)
+			{
+				Swap(StartIndex, TargetIndex);
+			}
+
+			// Empty List as we're going to add the InItem anyway through the List Iteration
+			ItemsToSelect.Empty(TargetIndex - StartIndex);
+
+			for (int32 ItemIndex = StartIndex; ItemIndex <= TargetIndex; ++ItemIndex)
+			{
+				ItemsToSelect.Add(WeakChildItems[ItemIndex].Pin());
+			}
+
+			if (TargetIndex < StartIndex)
+			{
+				Algo::Reverse(ItemsToSelect);
+			}
+		}
+	}
+
+	ToolView->SelectItems(MoveTemp(ItemsToSelect), SelectionFlags);
+	return FReply::Handled();
+}
+
+FReply SNavigationToolItemList::OnItemChipValidDragOver(const FGeometry& InGeometry, const FDragDropEvent& InDragDropEvent)
+{
+	TSharedPtr<SNavigationToolTreeRow> TreeRow = WeakRowWidget.Pin();
+	if (!TreeRow.IsValid())
+	{
+		return FReply::Unhandled();
+	}
+
+	// When Item Chip has a Valid Drag Over (i.e. a Supported Drag/Drop), make sure the Tree Row Holding this Item Chips simulates a Drag Leave
+	// Slate App won't do it as it will still find the Widget under the mouse
+	TreeRow->OnDragLeave(InDragDropEvent);
+
+	return FReply::Handled();
+}
+
+}
+
+#undef LOCTEXT_NAMESPACE
diff --git a/Engine/Source/Editor/Sequencer/Private/Tools/NavigationTool/Widgets/Columns/SNavigationToolLength.cpp b/Engine/Source/Editor/Sequencer/Private/Tools/NavigationTool/Widgets/Columns/SNavigationToolLength.cpp
new file mode 100644
--- /dev/null
+++ b/Engine/Source/Editor/Sequencer/Private/Tools/NavigationTool/Widgets/Columns/SNavigationToolLength.cpp
@@ -0,0 +1,91 @@
+// Copyright Epic Games, Inc. All Rights Reserved.
+
+#include "SNavigationToolLength.h"
+#include "ISequencer.h"
+#include "Sections/MovieSceneSubSection.h"
+#include "Tools/NavigationTool/Items/NavigationToolItemUtils.h"
+#include "Tools/NavigationTool/Items/NavigationToolSequence.h"
+#include "Tools/NavigationTool/NavigationTool.h"
+#include "Tools/NavigationTool/NavigationToolView.h"
+#include "Tools/NavigationTool/Widgets/SNavigationToolTreeRow.h"
+#include "Widgets/Input/SEditableTextBox.h"
+
+#define LOCTEXT_NAMESPACE "SNavigationToolLength"
+
+namespace UE::Sequencer::NavigationTool
+{
+
+void SNavigationToolLength::Construct(const FArguments& InArgs
+	, const FNavigationToolItemPtr& InItem
+	, const TSharedRef<INavigationToolView>& InView
+	, const TSharedRef<SNavigationToolTreeRow>& InRowWidget)
+{
+	WeakItem = InItem;
+	WeakView = InView;
+	WeakRowWidget = InRowWidget;
+
+	auto EmptyContent = [this]()
+	{
+		ChildSlot
+		[
+			SNullWidget::NullWidget
+		];
+	};
+
+	if (!WeakItem.IsValid() || !WeakView.IsValid())
+	{
+		EmptyContent();
+		return;
+	}
+
+	const FNavigationToolSequence* const SequenceItem = WeakItem.Pin()->CastTo<FNavigationToolSequence>();
+	if (!SequenceItem || !IsValid(SequenceItem->GetSubSection()))
+	{
+		EmptyContent();
+		return;
+	}
+
+	ChildSlot
+		[
+			SNew(SEditableTextBox)
+			.Style(FAppStyle::Get(), TEXT("Sequencer.NavigationTool.NonEditableTextBox"))
+			.Justification(ETextJustify::Center)
+			.IsReadOnly(true)
+			.Text_Lambda([this]()
+				{
+					return FText::FromString(FString::FromInt(GetLengthValue()));
+				})
+		];
+}
+
+int32 SNavigationToolLength::GetLengthValue() const
+{
+	const TSharedPtr<INavigationToolView> ToolView = WeakView.Pin();
+	if (!ToolView.IsValid())
+	{
+		return 0;
+	}
+
+	const TSharedPtr<ISequencer> Sequencer = ToolView->GetSequencer();
+	if (!Sequencer.IsValid())
+	{
+		return 0;
+	}
+
+	UMovieSceneSubSection* const SubSection = Cast<UMovieSceneSubSection>(GetSequenceItemSubSection(WeakItem.Pin()));
+	if (!IsValid(SubSection))
+	{
+		return 0;
+	}
+
+	const FFrameNumber FrameLength = SubSection->GetRange().Size<FFrameNumber>();
+	const FFrameNumber ConvertedFrameLength = ConvertFrameTime(FrameLength
+		, Sequencer->GetFocusedTickResolution()
+		, Sequencer->GetFocusedDisplayRate()).FrameNumber;
+
+	return ConvertedFrameLength.Value;
+}
+
+}
+
+#undef LOCTEXT_NAMESPACE
diff --git a/Engine/Source/Editor/Sequencer/Private/Tools/NavigationTool/Widgets/Columns/SNavigationToolLength.h b/Engine/Source/Editor/Sequencer/Private/Tools/NavigationTool/Widgets/Columns/SNavigationToolLength.h
new file mode 100644
--- /dev/null
+++ b/Engine/Source/Editor/Sequencer/Private/Tools/NavigationTool/Widgets/Columns/SNavigationToolLength.h
@@ -0,0 +1,37 @@
+// Copyright Epic Games, Inc. All Rights Reserved.
+
+#pragma once
+
+#include "Tools/NavigationTool/NavigationToolDefines.h"
+#include "Widgets/SCompoundWidget.h"
+
+namespace UE::Sequencer::NavigationTool
+{
+
+class INavigationToolView;
+class SNavigationToolTreeRow;
+
+class SNavigationToolLength : public SCompoundWidget
+{
+public:
+	SLATE_BEGIN_ARGS(SNavigationToolLength) {}
+	SLATE_END_ARGS()
+
+	virtual ~SNavigationToolLength() override {}
+
+	void Construct(const FArguments& InArgs
+		, const FNavigationToolItemPtr& InItem
+		, const TSharedRef<INavigationToolView>& InView
+		, const TSharedRef<SNavigationToolTreeRow>& InRowWidget);
+
+private:
+	int32 GetLengthValue() const;
+
+	FNavigationToolItemWeakPtr WeakItem;
+
+	TWeakPtr<INavigationToolView> WeakView;
+
+	TWeakPtr<SNavigationToolTreeRow> WeakRowWidget;
+};
+
+}
diff --git a/Engine/Source/Editor/Sequencer/Private/Tools/NavigationTool/Widgets/Columns/SNavigationToolLock.cpp b/Engine/Source/Editor/Sequencer/Private/Tools/NavigationTool/Widgets/Columns/SNavigationToolLock.cpp
new file mode 100644
--- /dev/null
+++ b/Engine/Source/Editor/Sequencer/Private/Tools/NavigationTool/Widgets/Columns/SNavigationToolLock.cpp
@@ -0,0 +1,239 @@
+// Copyright Epic Games, Inc. All Rights Reserved.
+
+#include "SNavigationToolLock.h"
+#include "ScopedTransaction.h"
+#include "Tools/NavigationTool/Items/INavigationToolItem.h"
+#include "Tools/NavigationTool/NavigationToolSettings.h"
+#include "Tools/NavigationTool/NavigationToolView.h"
+#include "Tools/NavigationTool/Items/NavigationToolItemUtils.h"
+#include "Tools/NavigationTool/Widgets/SNavigationToolTreeRow.h"
+
+#define LOCTEXT_NAMESPACE "SNavigationToolLock"
+
+namespace UE::Sequencer::NavigationTool
+{
+
+class FLockDragDropOp : public FDragDropOperation
+{
+public:
+	DRAG_DROP_OPERATOR_TYPE(FLockDragDropOp, FDragDropOperation)
+
+	/** Flag which defines whether to lock destination items or not */
+	bool bShouldLock;
+
+	/** Undo transaction stolen from the gutter which is kept alive for the duration of the drag */
+	TUniquePtr<FScopedTransaction> UndoTransaction;
+
+	/** The widget decorator to use */
+	virtual TSharedPtr<SWidget> GetDefaultDecorator() const override
+	{
+		return SNullWidget::NullWidget;
+	}
+
+	/** Create a new drag and drop operation out of the specified flag */
+	static TSharedRef<FLockDragDropOp> New(const bool bShouldLock, TUniquePtr<FScopedTransaction>& ScopedTransaction)
+	{
+		const TSharedRef<FLockDragDropOp> Operation = MakeShared<FLockDragDropOp>();
+		Operation->bShouldLock     = bShouldLock;
+		Operation->UndoTransaction = MoveTemp(ScopedTransaction);
+		Operation->Construct();
+		return Operation;
+	}
+};
+
+EItemSequenceLockState GetItemLockState(const FNavigationToolItemPtr& InItem)
+{
+	if (!InItem.IsValid())
+	{
+		return EItemSequenceLockState::None;
+	}
+
+	ISequenceLockableExtension* const LockableItem = InItem->CastTo<ISequenceLockableExtension>();
+	if (!LockableItem)
+	{
+		return EItemSequenceLockState::None;
+	}
+
+	return LockableItem->GetLockState();
+}
+
+void SetItemLocked(const FNavigationToolItemPtr& InItem, const bool bInIsLocked)
+{
+	if (!InItem.IsValid())
+	{
+		return;
+	}
+
+	if (ISequenceLockableExtension* const LockableItem = InItem->CastTo<ISequenceLockableExtension>())
+	{
+		LockableItem->SetIsLocked(bInIsLocked);
+	}
+}
+
+void SNavigationToolLock::Construct(const FArguments& InArgs
+	, const FNavigationToolItemPtr& InItem
+	, const TSharedRef<INavigationToolView>& InView
+	, const TSharedRef<SNavigationToolTreeRow>& InRowWidget)
+{
+	WeakItem = InItem;
+	WeakView = InView;
+	WeakRowWidget = InRowWidget;
+
+	SImage::Construct(SImage::FArguments()
+		.ColorAndOpacity(this, &SNavigationToolLock::GetForegroundColor)
+		.Image(this, &SNavigationToolLock::GetBrush));
+}
+
+FSlateColor SNavigationToolLock::GetForegroundColor() const
+{
+	const FNavigationToolItemPtr Item = WeakItem.Pin();
+	if (!Item.IsValid())
+	{
+		return FSlateColor::UseForeground();
+	}
+
+	const bool bIsItemSelected = WeakView.IsValid() && WeakView.Pin()->IsItemSelected(Item.ToSharedRef());
+	const bool bIsItemHovered = WeakRowWidget.IsValid() && WeakRowWidget.Pin()->IsHovered();
+	const bool bAlwaysShowLock = GetDefault<UNavigationToolSettings>()->ShouldAlwaysShowLockState();
+
+	// We can hide the brush if settings for "Always Showing Lock State" is OFF
+	// and item is not locked while also not being selected or hovered
+	if (!bAlwaysShowLock
+		&& GetLockState() == EItemSequenceLockState::None
+		&& !bIsItemSelected
+		&& !bIsItemHovered)
+	{
+		return FLinearColor::Transparent;
+	}
+
+	if (bIsItemHovered || IsHovered())
+	{
+		switch (GetLockState())
+		{
+		case EItemSequenceLockState::None:
+			return FStyleColors::White25;
+
+		case EItemSequenceLockState::PartiallyLocked:
+			return FStyleColors::ForegroundHover;
+
+		case EItemSequenceLockState::Locked:
+			return FStyleColors::ForegroundHover;
+		}
+	}
+	else
+	{
+		switch (GetLockState())
+		{
+		case EItemSequenceLockState::None:
+			return FStyleColors::Transparent;
+
+		case EItemSequenceLockState::PartiallyLocked:
+			return FStyleColors::White25;
+
+		case EItemSequenceLockState::Locked:
+			return FStyleColors::Foreground;
+		}
+	}
+
+	return FSlateColor::UseForeground();
+}
+
+const FSlateBrush* SNavigationToolLock::GetBrush() const
+{
+	return (GetLockState() == EItemSequenceLockState::None)
+		? FAppStyle::GetBrush(TEXT("Icons.Unlock"))
+		: FAppStyle::GetBrush(TEXT("Icons.Lock"));
+}
+
+FReply SNavigationToolLock::OnDragDetected(const FGeometry& MyGeometry, const FPointerEvent& MouseEvent)
+{
+	if (MouseEvent.IsMouseButtonDown(EKeys::LeftMouseButton))
+	{
+		const bool bShouldLock = GetLockState() == EItemSequenceLockState::Locked;
+		return FReply::Handled().BeginDragDrop(FLockDragDropOp::New(bShouldLock, UndoTransaction));
+	}
+	return FReply::Unhandled();
+}
+
+void SNavigationToolLock::OnDragEnter(const FGeometry& MyGeometry, const FDragDropEvent& DragDropEvent)
+{
+	if (const TSharedPtr<FLockDragDropOp> DragDropOp = DragDropEvent.GetOperationAs<FLockDragDropOp>())
+	{
+		SetIsLocked(DragDropOp->bShouldLock);
+	}
+}
+
+FReply SNavigationToolLock::HandleClick()
+{
+	const TSharedPtr<INavigationToolView> ToolView = WeakView.Pin();
+	const FNavigationToolItemPtr Item = WeakItem.Pin();
+
+	if (!ToolView.IsValid() || !Item.IsValid())
+	{
+		return FReply::Unhandled();
+	}
+
+	UndoTransaction = MakeUnique<FScopedTransaction>(LOCTEXT("SetNavigationToolItemLock", "Set Item Lock"));
+
+	const bool bNewIsLocked = (GetLockState() != EItemSequenceLockState::Locked);
+
+	// We operate on all the selected items if the specified item is selected
+	if (ToolView->IsItemSelected(Item.ToSharedRef()))
+	{
+		for (FNavigationToolItemPtr& SelectedItem : ToolView->GetSelectedItems())
+		{
+			SetItemLocked(SelectedItem, bNewIsLocked);
+		}
+	}
+	else
+	{
+		SetIsLocked(bNewIsLocked);
+	}
+
+	return FReply::Handled().DetectDrag(SharedThis(this), EKeys::LeftMouseButton);
+}
+
+FReply SNavigationToolLock::OnMouseButtonDoubleClick(const FGeometry& InMyGeometry, const FPointerEvent& InMouseEvent)
+{
+	return HandleClick();
+}
+
+FReply SNavigationToolLock::OnMouseButtonDown(const FGeometry& MyGeometry, const FPointerEvent& MouseEvent)
+{
+	if (MouseEvent.GetEffectingButton() != EKeys::LeftMouseButton)
+	{
+		return FReply::Unhandled();
+	}
+	return HandleClick();
+}
+
+FReply SNavigationToolLock::OnMouseButtonUp(const FGeometry& MyGeometry, const FPointerEvent& MouseEvent)
+{
+	if (MouseEvent.GetEffectingButton() == EKeys::LeftMouseButton)
+	{
+		UndoTransaction.Reset();
+
+		return FReply::Handled();
+	}
+
+	return FReply::Unhandled();
+}
+
+void SNavigationToolLock::OnMouseCaptureLost(const FCaptureLostEvent& CaptureLostEvent)
+{
+	UndoTransaction.Reset();
+}
+
+EItemSequenceLockState SNavigationToolLock::GetLockState() const
+{
+	return GetItemLockState(WeakItem.Pin());
+}
+
+void SNavigationToolLock::SetIsLocked(const bool bInIsLocked)
+{
+	SetItemLocked(WeakItem.Pin(), bInIsLocked);
+}
+
+}
+
+#undef LOCTEXT_NAMESPACE
diff --git a/Engine/Source/Editor/Sequencer/Private/Tools/NavigationTool/Widgets/Columns/SNavigationToolLock.h b/Engine/Source/Editor/Sequencer/Private/Tools/NavigationTool/Widgets/Columns/SNavigationToolLock.h
new file mode 100644
--- /dev/null
+++ b/Engine/Source/Editor/Sequencer/Private/Tools/NavigationTool/Widgets/Columns/SNavigationToolLock.h
@@ -0,0 +1,74 @@
+// Copyright Epic Games, Inc. All Rights Reserved.
+
+#pragma once
+
+#include "ScopedTransaction.h"
+#include "Templates/UniquePtr.h"
+#include "Tools/NavigationTool/Extensions/ISequenceLockableExtension.h"
+#include "Tools/NavigationTool/NavigationToolDefines.h"
+#include "Widgets/Images/SImage.h"
+
+class FDragDropEvent;
+class FReply;
+struct FCaptureLostEvent;
+struct FGeometry;
+struct FPointerEvent;
+struct FSlateBrush;
+struct FSlateColor;
+
+namespace UE::Sequencer::NavigationTool
+{
+
+class INavigationToolView;
+class SNavigationToolTreeRow;
+enum class ESequenceLockState;
+
+class SNavigationToolLock : public SImage
+{
+public:
+	SLATE_BEGIN_ARGS(SNavigationToolLock) {}
+	SLATE_END_ARGS()
+
+	virtual ~SNavigationToolLock() override {}
+
+	void Construct(const FArguments& InArgs
+		, const FNavigationToolItemPtr& InItem
+		, const TSharedRef<INavigationToolView>& InView
+		, const TSharedRef<SNavigationToolTreeRow>& InRowWidget);
+
+	virtual FSlateColor GetForegroundColor() const override;
+
+	EItemSequenceLockState GetLockState() const;
+
+	void SetIsLocked(const bool bInIsLocked);
+
+	const FSlateBrush* GetBrush() const;
+
+	virtual FReply OnDragDetected(const FGeometry& MyGeometry, const FPointerEvent& MouseEvent) override;
+
+	/** If a lock drag drop operation has entered this widget, set its item to the new lock state */
+	virtual void OnDragEnter(const FGeometry& MyGeometry, const FDragDropEvent& DragDropEvent) override;
+
+	FReply HandleClick();
+
+	virtual FReply OnMouseButtonDoubleClick(const FGeometry& InMyGeometry, const FPointerEvent& InMouseEvent) override;
+
+	virtual FReply OnMouseButtonDown(const FGeometry& MyGeometry, const FPointerEvent& MouseEvent) override;
+
+	virtual FReply OnMouseButtonUp(const FGeometry& MyGeometry, const FPointerEvent& MouseEvent) override;
+
+	virtual void OnMouseCaptureLost(const FCaptureLostEvent& CaptureLostEvent) override;
+
+	//void OnSetItemLockedState(const FNavigationToolItemPtr& InItem, const bool bInIsLocked);
+
+private:
+	FNavigationToolItemWeakPtr WeakItem;
+
+	TWeakPtr<INavigationToolView> WeakView;
+
+	TWeakPtr<SNavigationToolTreeRow> WeakRowWidget;
+
+	TUniquePtr<FScopedTransaction> UndoTransaction;
+};
+
+}
diff --git a/Engine/Source/Editor/Sequencer/Private/Tools/NavigationTool/Widgets/Columns/SNavigationToolMarkerVisibility.cpp b/Engine/Source/Editor/Sequencer/Private/Tools/NavigationTool/Widgets/Columns/SNavigationToolMarkerVisibility.cpp
new file mode 100644
--- /dev/null
+++ b/Engine/Source/Editor/Sequencer/Private/Tools/NavigationTool/Widgets/Columns/SNavigationToolMarkerVisibility.cpp
@@ -0,0 +1,241 @@
+// Copyright Epic Games, Inc. All Rights Reserved.
+
+#include "SNavigationToolMarkerVisibility.h"
+#include "Editor.h"
+#include "ScopedTransaction.h"
+#include "Styling/StyleColors.h"
+#include "Tools/NavigationTool/Columns/NavigationToolMarkerVisibilityColumn.h"
+#include "Tools/NavigationTool/Items/INavigationToolItem.h"
+#include "Tools/NavigationTool/Items/NavigationToolItemUtils.h"
+#include "Tools/NavigationTool/NavigationToolView.h"
+#include "Tools/NavigationTool/Widgets/SNavigationToolTreeRow.h"
+
+#define LOCTEXT_NAMESPACE "SNavigationToolMarkerVisibility"
+
+namespace UE::Sequencer::NavigationTool
+{
+
+class FMarkerVisibilityDragDropOp : public FDragDropOperation
+{
+public:
+	DRAG_DROP_OPERATOR_TYPE(FMarkerVisibilityDragDropOp, FDragDropOperation)
+
+	/** Flag which defines whether to hide destination items or not */
+	EItemMarkerVisibility MarkerVisibility;
+
+	/** Undo transaction stolen from the gutter which is kept alive for the duration of the drag */
+	TUniquePtr<FScopedTransaction> UndoTransaction;
+
+	/** The widget decorator to use */
+	virtual TSharedPtr<SWidget> GetDefaultDecorator() const override
+	{
+		return SNullWidget::NullWidget;
+	}
+
+	/** Create a new drag and drop operation out of the specified flag */
+	static TSharedRef<FMarkerVisibilityDragDropOp> New(const EItemMarkerVisibility InMarkerVisibility, TUniquePtr<FScopedTransaction>& ScopedTransaction)
+	{
+		TSharedRef<FMarkerVisibilityDragDropOp> Operation = MakeShared<FMarkerVisibilityDragDropOp>();
+		Operation->MarkerVisibility = InMarkerVisibility;
+		Operation->UndoTransaction = MoveTemp(ScopedTransaction);
+		Operation->Construct();
+		return Operation;
+	}
+};
+
+EItemMarkerVisibility GetItemMarkerVisibility(const FNavigationToolItemPtr& InItem)
+{
+	if (!InItem.IsValid())
+	{
+		return EItemMarkerVisibility::None;
+	}
+
+	IMarkerVisibilityExtension* const MarkerVisibilityItem = InItem->CastTo<IMarkerVisibilityExtension>();
+	if (!MarkerVisibilityItem)
+	{
+		return EItemMarkerVisibility::None;
+	}
+
+	return MarkerVisibilityItem->GetMarkerVisibility();
+}
+
+void SetItemMarkerVisibility(const FNavigationToolItemPtr& InItem, const bool bInMarkersVisible)
+{
+	if (!InItem.IsValid())
+	{
+		return;
+	}
+
+	if (IMarkerVisibilityExtension* const MarkerVisibilityItem = InItem->CastTo<IMarkerVisibilityExtension>())
+	{
+		MarkerVisibilityItem->SetMarkerVisibility(bInMarkersVisible);
+	}
+}
+
+void SNavigationToolMarkerVisibility::Construct(const FArguments& InArgs
+	, const TSharedRef<FNavigationToolMarkerVisibilityColumn>& InColumn
+	, const FNavigationToolItemPtr& InItem
+	, const TSharedRef<INavigationToolView>& InView
+	, const TSharedRef<SNavigationToolTreeRow>& InRowWidget)
+{
+	WeakColumn = InColumn;
+	WeakItem = InItem;
+	WeakView = InView;
+	WeakRowWidget = InRowWidget;
+
+	SImage::Construct(SImage::FArguments()
+		.IsEnabled(this, &SNavigationToolMarkerVisibility::IsVisibilityWidgetEnabled)
+		.ColorAndOpacity(this, &SNavigationToolMarkerVisibility::GetForegroundColor)
+		.Image(this, &SNavigationToolMarkerVisibility::GetBrush));
+}
+
+FReply SNavigationToolMarkerVisibility::OnDragDetected(const FGeometry& MyGeometry, const FPointerEvent& MouseEvent)
+{
+	if (MouseEvent.IsMouseButtonDown(EKeys::LeftMouseButton))
+	{
+		return FReply::Handled().BeginDragDrop(FMarkerVisibilityDragDropOp::New(GetMarkerVisibility(), UndoTransaction));
+	}
+	return FReply::Unhandled();
+}
+
+void SNavigationToolMarkerVisibility::OnDragEnter(const FGeometry& MyGeometry, const FDragDropEvent& DragDropEvent)
+{
+	const TSharedPtr<FMarkerVisibilityDragDropOp> DragDropOp = DragDropEvent.GetOperationAs<FMarkerVisibilityDragDropOp>();
+	if (DragDropOp.IsValid())
+	{
+		const bool bNewVisible = (DragDropOp->MarkerVisibility != EItemMarkerVisibility::Visible);
+		SetMarkersVisible(bNewVisible);
+	}
+}
+
+FReply SNavigationToolMarkerVisibility::HandleClick()
+{
+	if (!IsVisibilityWidgetEnabled())
+	{
+		return FReply::Unhandled();
+	}
+
+	const TSharedPtr<INavigationToolView> ToolView = WeakView.Pin();
+	const FNavigationToolItemPtr Item = WeakItem.Pin();
+	const TSharedPtr<FNavigationToolMarkerVisibilityColumn> Column = WeakColumn.Pin();
+
+	if (!ToolView.IsValid() || !Item.IsValid() || !Column.IsValid())
+	{
+		return FReply::Unhandled();
+	}
+
+	// Open an undo transaction
+	UndoTransaction = MakeUnique<FScopedTransaction>(LOCTEXT("SetNavigationToolMarkerVisibility", "Set Marker Visibility"));
+
+	const bool bNewVisible = (GetMarkerVisibility() != EItemMarkerVisibility::Visible);
+
+	// We operate on all the selected items if the specified item is selected
+	if (ToolView->IsItemSelected(Item.ToSharedRef()))
+	{
+		for (const FNavigationToolItemPtr& SelectedItem : ToolView->GetSelectedItems())
+		{
+			if (IMarkerVisibilityExtension* const MarkerVisibilityItem = SelectedItem->CastTo<IMarkerVisibilityExtension>())
+			{
+				MarkerVisibilityItem->SetMarkerVisibility(bNewVisible);
+			}
+		}
+	}
+	else
+	{
+		SetMarkersVisible(bNewVisible);
+	}
+
+	return FReply::Handled().DetectDrag(SharedThis(this), EKeys::LeftMouseButton);
+}
+
+FReply SNavigationToolMarkerVisibility::OnMouseButtonDoubleClick(const FGeometry& InMyGeometry, const FPointerEvent& InMouseEvent)
+{
+	return HandleClick();
+}
+
+FReply SNavigationToolMarkerVisibility::OnMouseButtonDown(const FGeometry& MyGeometry, const FPointerEvent& MouseEvent)
+{
+	if (MouseEvent.GetEffectingButton() != EKeys::LeftMouseButton)
+	{
+		return FReply::Unhandled();
+	}
+
+	return HandleClick();
+}
+
+FReply SNavigationToolMarkerVisibility::OnMouseButtonUp(const FGeometry& MyGeometry, const FPointerEvent& MouseEvent)
+{
+	if (MouseEvent.GetEffectingButton() == EKeys::LeftMouseButton)
+	{
+		UndoTransaction.Reset();
+		return FReply::Handled();
+	}
+
+	return FReply::Unhandled();
+}
+
+void SNavigationToolMarkerVisibility::OnMouseCaptureLost(const FCaptureLostEvent& CaptureLostEvent)
+{
+	UndoTransaction.Reset();
+}
+
+const FSlateBrush* SNavigationToolMarkerVisibility::GetBrush() const
+{
+	return FAppStyle::GetBrush(TEXT("AnimTimeline.SectionMarker"));
+}
+
+FSlateColor SNavigationToolMarkerVisibility::GetForegroundColor() const
+{
+	const IMarkerVisibilityExtension* const MarkerVisibilityExt = WeakItem.Pin()->CastTo<IMarkerVisibilityExtension>();
+	if (!MarkerVisibilityExt)
+	{
+		return FLinearColor::Transparent;
+	}
+
+	const bool bIsItemHovered = WeakRowWidget.IsValid() && WeakRowWidget.Pin()->IsHovered();
+
+	if (IsHovered() || bIsItemHovered)
+	{
+		switch (GetMarkerVisibility())
+		{
+		case EItemMarkerVisibility::None:
+			return FStyleColors::White25;
+
+		case EItemMarkerVisibility::PartiallyVisible:
+			return FStyleColors::ForegroundHover;
+
+		case EItemMarkerVisibility::Visible:
+			return FStyleColors::ForegroundHover;
+		}
+	}
+	else
+	{
+		switch (GetMarkerVisibility())
+		{
+		case EItemMarkerVisibility::None:
+			return FStyleColors::Transparent;
+
+		case EItemMarkerVisibility::PartiallyVisible:
+			return FStyleColors::White25;
+
+		case EItemMarkerVisibility::Visible:
+			return FStyleColors::Foreground;
+		}
+	}
+
+	return FStyleColors::Transparent;
+}
+
+EItemMarkerVisibility SNavigationToolMarkerVisibility::GetMarkerVisibility() const
+{
+	return GetItemMarkerVisibility(WeakItem.Pin());
+}
+
+void SNavigationToolMarkerVisibility::SetMarkersVisible(const bool bInVisible)
+{
+	SetItemMarkerVisibility(WeakItem.Pin(), bInVisible);
+}
+
+}
+
+#undef LOCTEXT_NAMESPACE
diff --git a/Engine/Source/Editor/Sequencer/Private/Tools/NavigationTool/Widgets/Columns/SNavigationToolMarkerVisibility.h b/Engine/Source/Editor/Sequencer/Private/Tools/NavigationTool/Widgets/Columns/SNavigationToolMarkerVisibility.h
new file mode 100644
--- /dev/null
+++ b/Engine/Source/Editor/Sequencer/Private/Tools/NavigationTool/Widgets/Columns/SNavigationToolMarkerVisibility.h
@@ -0,0 +1,64 @@
+// Copyright Epic Games, Inc. All Rights Reserved.
+
+#pragma once
+
+#include "ScopedTransaction.h"
+#include "Templates/UniquePtr.h"
+#include "Tools/NavigationTool/Extensions/IMarkerVisibilityExtension.h"
+#include "Tools/NavigationTool/NavigationToolDefines.h"
+#include "Widgets/Images/SImage.h"
+
+namespace UE::Sequencer::NavigationTool
+{
+
+class FNavigationToolMarkerVisibilityColumn;
+class INavigationToolView;
+class SNavigationToolTreeRow;
+
+/** Widget responsible for managing the visibility for a single item */
+class SNavigationToolMarkerVisibility : public SImage
+{
+public:
+	SLATE_BEGIN_ARGS(SNavigationToolMarkerVisibility) {}
+	SLATE_END_ARGS()
+
+	void Construct(const FArguments& InArgs
+		, const TSharedRef<FNavigationToolMarkerVisibilityColumn>& InColumn
+		, const FNavigationToolItemPtr& InItem
+		, const TSharedRef<INavigationToolView>& InView
+		, const TSharedRef<SNavigationToolTreeRow>& InRowWidget);
+
+private:
+	/** Returns whether the widget is enabled or not */
+	virtual bool IsVisibilityWidgetEnabled() const { return true; }
+
+	virtual const FSlateBrush* GetBrush() const;
+
+	virtual FReply OnDragDetected(const FGeometry& MyGeometry, const FPointerEvent& MouseEvent) override;
+	virtual void OnDragEnter(const FGeometry& MyGeometry, const FDragDropEvent& DragDropEvent) override;
+
+	FReply HandleClick();
+
+	virtual FReply OnMouseButtonDoubleClick(const FGeometry& InMyGeometry, const FPointerEvent& InMouseEvent) override;
+	virtual FReply OnMouseButtonDown(const FGeometry& MyGeometry, const FPointerEvent& MouseEvent) override;
+	virtual FReply OnMouseButtonUp(const FGeometry& MyGeometry, const FPointerEvent& MouseEvent) override;
+	virtual void OnMouseCaptureLost(const FCaptureLostEvent& CaptureLostEvent) override;
+
+	virtual FSlateColor GetForegroundColor() const override;
+	
+	EItemMarkerVisibility GetMarkerVisibility() const;
+
+	void SetMarkersVisible(const bool bInVisible);
+
+	TWeakPtr<FNavigationToolMarkerVisibilityColumn> WeakColumn;
+
+	FNavigationToolItemWeakPtr WeakItem;
+
+	TWeakPtr<INavigationToolView> WeakView;
+
+	TWeakPtr<SNavigationToolTreeRow> WeakRowWidget;
+
+	TUniquePtr<FScopedTransaction> UndoTransaction;
+};
+
+}
diff --git a/Engine/Source/Editor/Sequencer/Private/Tools/NavigationTool/Widgets/Columns/SNavigationToolOutTime.cpp b/Engine/Source/Editor/Sequencer/Private/Tools/NavigationTool/Widgets/Columns/SNavigationToolOutTime.cpp
new file mode 100644
--- /dev/null
+++ b/Engine/Source/Editor/Sequencer/Private/Tools/NavigationTool/Widgets/Columns/SNavigationToolOutTime.cpp
@@ -0,0 +1,150 @@
+// Copyright Epic Games, Inc. All Rights Reserved.
+
+#include "SNavigationToolOutTime.h"
+#include "ScopedTransaction.h"
+#include "Sequencer.h"
+#include "Tools/NavigationTool/Items/NavigationToolItemUtils.h"
+#include "Tools/NavigationTool/NavigationTool.h"
+#include "Tools/NavigationTool/NavigationToolView.h"
+#include "Tools/NavigationTool/Items/NavigationToolSequence.h"
+#include "Tools/NavigationTool/Widgets/SNavigationToolTreeRow.h"
+#include "TrackEditors/SubTrackEditorBase.h"
+#include "Widgets/Input/SSpinBox.h"
+
+#define LOCTEXT_NAMESPACE "SNavigationToolOutTime"
+
+namespace UE::Sequencer::NavigationTool
+{
+
+void SNavigationToolOutTime::Construct(const FArguments& InArgs
+	, const FNavigationToolItemPtr& InItem
+	, const TSharedRef<INavigationToolView>& InView
+	, const TSharedRef<SNavigationToolTreeRow>& InRowWidget)
+{
+	WeakItem = InItem;
+	WeakView = InView;
+	WeakRowWidget = InRowWidget;
+
+	auto EmptyContent = [this]()
+	{
+		ChildSlot
+		[
+			SNullWidget::NullWidget
+		];
+	};
+
+	if (!WeakItem.IsValid() || !WeakView.IsValid())
+	{
+		EmptyContent();
+		return;
+	}
+
+	const FNavigationToolSequence* const SequenceItem = WeakItem.Pin()->CastTo<FNavigationToolSequence>();
+	if (!SequenceItem || !IsValid(SequenceItem->GetSubSection()))
+	{
+		EmptyContent();
+		return;
+	}
+
+	ChildSlot
+		[
+			SNew(SSpinBox<int32>)
+			.Justification(ETextJustify::Center)
+			.Style(&FAppStyle::GetWidgetStyle<FSpinBoxStyle>(TEXT("Sequencer.NavigationTool.OutTime.SpinBox")))
+			.Value(this, &SNavigationToolOutTime::GetFrameTimeValue)
+			.OnValueChanged(this, &SNavigationToolOutTime::OnFrameTimeValueChanged)
+			.OnValueCommitted(this, &SNavigationToolOutTime::OnFrameTimeValueCommitted)
+			.OnBeginSliderMovement(this, &SNavigationToolOutTime::OnBeginSliderMovement)
+			.OnEndSliderMovement(this, &SNavigationToolOutTime::OnEndSliderMovement)
+		];
+}
+
+int32 SNavigationToolOutTime::GetFrameTimeValue() const
+{
+	const TSharedPtr<INavigationToolView> ToolView = WeakView.Pin();
+	if (!ToolView.IsValid())
+	{
+		return 0;
+	}
+
+	const TSharedPtr<ISequencer> Sequencer = ToolView->GetSequencer();
+	if (!Sequencer.IsValid())
+	{
+		return 0;
+	}
+
+	UMovieSceneSubSection* const SubSection = Cast<UMovieSceneSubSection>(GetSequenceItemSubSection(WeakItem.Pin()));
+	if (!IsValid(SubSection))
+	{
+		return 0;
+	}
+
+	const FFrameNumber FrameTime = ConvertFrameTime(SubSection->SectionRange.Value.GetUpperBoundValue()
+		, Sequencer->GetFocusedTickResolution()
+		, Sequencer->GetFocusedDisplayRate()).FrameNumber;
+
+	return FrameTime.Value;
+}
+
+void SNavigationToolOutTime::OnFrameTimeValueChanged(const int32 InNewValue)
+{
+	OnFrameTimeValueCommitted(InNewValue, ETextCommit::Default);
+}
+
+void SNavigationToolOutTime::OnFrameTimeValueCommitted(const int32 InNewValue, const ETextCommit::Type InCommitType)
+{
+	const TSharedPtr<INavigationToolView> ToolView = WeakView.Pin();
+	if (!ToolView.IsValid())
+	{
+		return;
+	}
+
+	const TSharedPtr<ISequencer> Sequencer = ToolView->GetSequencer();
+	if (!Sequencer.IsValid())
+	{
+		return;
+	}
+
+	UMovieSceneSubSection* const SubSection = Cast<UMovieSceneSubSection>(GetSequenceItemSubSection(WeakItem.Pin()));
+	if (!IsValid(SubSection))
+	{
+		return;
+	}
+
+	const FFrameNumber NewFrameTime = ConvertFrameTime(InNewValue
+		, Sequencer->GetFocusedDisplayRate()
+		, Sequencer->GetFocusedTickResolution()).FrameNumber;
+
+	if (SubSection->SectionRange.Value.GetUpperBoundValue() == NewFrameTime)
+	{
+		return;
+	}
+
+	const bool bShouldTransact = !UndoTransaction.IsValid() && (InCommitType == ETextCommit::OnEnter);
+	const FScopedTransaction Transaction(GetTransactionText(), bShouldTransact);
+
+	SubSection->Modify();
+	SubSection->SectionRange.Value.SetUpperBoundValue(NewFrameTime);
+}
+
+void SNavigationToolOutTime::OnBeginSliderMovement()
+{
+	if (!UndoTransaction.IsValid())
+	{
+		UndoTransaction = MakeUnique<FScopedTransaction>(GetTransactionText());
+	}
+}
+
+void SNavigationToolOutTime::OnEndSliderMovement(const int32 InNewValue)
+{
+	UndoTransaction.Reset();
+}
+
+FText SNavigationToolOutTime::GetTransactionText() const
+{
+	return LOCTEXT("SetSequenceOutTimeTransaction", "Set Sequence Out Time");
+}
+
+}
+
+#undef LOCTEXT_NAMESPACE
diff --git a/Engine/Source/Editor/Sequencer/Private/Tools/NavigationTool/Widgets/Columns/SNavigationToolOutTime.h b/Engine/Source/Editor/Sequencer/Private/Tools/NavigationTool/Widgets/Columns/SNavigationToolOutTime.h
new file mode 100644
--- /dev/null
+++ b/Engine/Source/Editor/Sequencer/Private/Tools/NavigationTool/Widgets/Columns/SNavigationToolOutTime.h
@@ -0,0 +1,47 @@
+// Copyright Epic Games, Inc. All Rights Reserved.
+
+#pragma once
+
+#include "ScopedTransaction.h"
+#include "Templates/UniquePtr.h"
+#include "Tools/NavigationTool/NavigationToolDefines.h"
+#include "Widgets/SCompoundWidget.h"
+
+namespace UE::Sequencer::NavigationTool
+{
+
+class INavigationToolView;
+class SNavigationToolTreeRow;
+
+class SNavigationToolOutTime : public SCompoundWidget
+{
+public:
+	SLATE_BEGIN_ARGS(SNavigationToolOutTime) {}
+	SLATE_END_ARGS()
+
+	virtual ~SNavigationToolOutTime() override {}
+
+	void Construct(const FArguments& InArgs
+		, const FNavigationToolItemPtr& InItem
+		, const TSharedRef<INavigationToolView>& InView
+		, const TSharedRef<SNavigationToolTreeRow>& InRowWidget);
+
+private:
+	int32 GetFrameTimeValue() const;
+	void OnFrameTimeValueChanged(const int32 InNewValue);
+	void OnFrameTimeValueCommitted(const int32 InNewValue, const ETextCommit::Type InCommitType);
+	void OnBeginSliderMovement();
+	void OnEndSliderMovement(const int32 InNewValue);
+
+	FText GetTransactionText() const;
+
+	FNavigationToolItemWeakPtr WeakItem;
+
+	TWeakPtr<INavigationToolView> WeakView;
+
+	TWeakPtr<SNavigationToolTreeRow> WeakRowWidget;
+
+	TUniquePtr<FScopedTransaction> UndoTransaction;
+};
+
+}
diff --git a/Engine/Source/Editor/Sequencer/Private/Tools/NavigationTool/Widgets/Columns/SNavigationToolRevisionControl.cpp b/Engine/Source/Editor/Sequencer/Private/Tools/NavigationTool/Widgets/Columns/SNavigationToolRevisionControl.cpp
new file mode 100644
--- /dev/null
+++ b/Engine/Source/Editor/Sequencer/Private/Tools/NavigationTool/Widgets/Columns/SNavigationToolRevisionControl.cpp
@@ -0,0 +1,174 @@
+// Copyright Epic Games, Inc. All Rights Reserved.
+
+#include "SNavigationToolRevisionControl.h"
+#include "Styling/StyleColors.h"
+#include "Tools/NavigationTool/Columns/NavigationToolRevisionControlColumn.h"
+#include "Tools/NavigationTool/Extensions/IRevisionControlExtension.h"
+#include "Tools/NavigationTool/Items/INavigationToolItem.h"
+#include "Tools/NavigationTool/NavigationToolView.h"
+#include "Tools/NavigationTool/Widgets/SNavigationToolTreeRow.h"
+
+#define LOCTEXT_NAMESPACE "SNavigationToolRevisionControl"
+
+namespace UE::Sequencer::NavigationTool
+{
+
+EItemRevisionControlState GetRevisionControlState(const FNavigationToolItemPtr& InItem)
+{
+	if (const IRevisionControlExtension* const RevisionControlExt = InItem->CastTo<IRevisionControlExtension>())
+	{
+		return RevisionControlExt->GetRevisionControlState();
+	}
+	return EItemRevisionControlState::None;
+}
+
+void SetIsRevisionControlled(const FNavigationToolItemPtr& InItem, const bool bInIsRevisionControlled)
+{
+	if (IRevisionControlExtension* const RevisionControlExt = InItem->CastTo<IRevisionControlExtension>())
+	{
+		const bool bItemRevisionControlled = RevisionControlExt->GetRevisionControlState() == EItemRevisionControlState::CheckedOut;
+		if (bItemRevisionControlled != bInIsRevisionControlled)
+		{
+			RevisionControlExt->SetIsRevisionControlled(bInIsRevisionControlled);
+		}
+	}
+}
+
+void SNavigationToolRevisionControl::Construct(const FArguments& InArgs
+	, const TSharedRef<FNavigationToolRevisionControlColumn>& InColumn
+	, const FNavigationToolItemPtr& InItem
+	, const TSharedRef<INavigationToolView>& InView
+	, const TSharedRef<SNavigationToolTreeRow>& InRowWidget)
+{
+	WeakColumn = InColumn;
+	WeakItem = InItem;
+	WeakView = InView;
+	WeakRowWidget = InRowWidget;
+
+	SetToolTipText(TAttribute<FText>::CreateSP(this, &SNavigationToolRevisionControl::GetToolTipText));
+
+	SImage::Construct(SImage::FArguments()
+		.Image(this, &SNavigationToolRevisionControl::GetBrush)
+		.ColorAndOpacity(this, &SNavigationToolRevisionControl::GetForegroundColor));
+}
+
+FSlateColor SNavigationToolRevisionControl::GetForegroundColor() const
+{
+	const IRevisionControlExtension* const RevisionControlExt = WeakItem.Pin()->CastTo<IRevisionControlExtension>();
+	if (!RevisionControlExt)
+	{
+		return FStyleColors::Transparent;
+	}
+
+	switch (GetRevisionControlState())
+	{
+	case EItemRevisionControlState::None:
+		return FStyleColors::Transparent;
+	case EItemRevisionControlState::PartiallyCheckedOut:
+		return FStyleColors::White25;
+	case EItemRevisionControlState::CheckedOut:
+		return FStyleColors::Foreground;
+	}
+
+	return FStyleColors::Transparent;
+}
+
+const FSlateBrush* SNavigationToolRevisionControl::GetBrush() const
+{
+	const FNavigationToolItemPtr Item = WeakItem.Pin();
+	if (!Item.IsValid())
+	{
+		return nullptr;
+	}
+
+	const IRevisionControlExtension* const RevisionControlExt = Item->CastTo<IRevisionControlExtension>();
+	if (!RevisionControlExt)
+	{
+		return nullptr;
+	}
+
+	return RevisionControlExt->GetRevisionControlStatusIcon();
+}
+
+FText SNavigationToolRevisionControl::GetToolTipText() const
+{
+	const FNavigationToolItemPtr Item = WeakItem.Pin();
+	if (!Item.IsValid())
+	{
+		return FText::GetEmpty();
+	}
+
+	const IRevisionControlExtension* const RevisionControlExt = Item->CastTo<IRevisionControlExtension>();
+	if (!RevisionControlExt)
+	{
+		return FText::GetEmpty();
+	}
+
+	return RevisionControlExt->GetRevisionControlStatusText();
+}
+
+EItemRevisionControlState SNavigationToolRevisionControl::GetRevisionControlState() const
+{
+	const FNavigationToolItemPtr Item = WeakItem.Pin();
+	if (!Item.IsValid())
+	{
+		return EItemRevisionControlState::None;
+	}
+
+	const IRevisionControlExtension* const RevisionControlExt = Item->CastTo<IRevisionControlExtension>();
+	if (!RevisionControlExt)
+	{
+		return EItemRevisionControlState::None;
+	}
+
+	return RevisionControlExt->GetRevisionControlState();
+}
+
+void SNavigationToolRevisionControl::SetIsRevisionControlled(const bool bInRevisionControlled)
+{
+	const FNavigationToolItemPtr Item = WeakItem.Pin();
+	if (!Item.IsValid())
+	{
+		return;
+	}
+
+	const IRevisionControlExtension* const RevisionControlExt = Item->CastTo<IRevisionControlExtension>();
+	if (!RevisionControlExt)
+	{
+		return;
+	}
+
+	const bool bIsRevisionControlled = RevisionControlExt->GetRevisionControlState() == EItemRevisionControlState::CheckedOut;
+	if (bIsRevisionControlled == bInRevisionControlled)
+	{
+		return;
+	}
+
+	OnSetIsRevisionControlled(Item, bInRevisionControlled);
+
+	if (const TSharedPtr<INavigationToolView> ToolView = WeakView.Pin())
+	{
+		ToolView->RequestRefresh();
+	}
+}
+
+void SNavigationToolRevisionControl::OnSetIsRevisionControlled(const FNavigationToolItemPtr& InItem, const bool bInIsRevisionControlled)
+{
+	if (IRevisionControlExtension* const RevisionControlExt = InItem->CastTo<IRevisionControlExtension>())
+	{
+		RevisionControlExt->SetIsRevisionControlled(bInIsRevisionControlled);
+	}
+
+	// Apply the same visibility to the children
+	for (const FNavigationToolItemPtr& Child : InItem->GetChildren())
+	{
+		if (IRevisionControlExtension* const RevisionControlExt = Child->CastTo<IRevisionControlExtension>())
+		{
+			RevisionControlExt->SetIsRevisionControlled(bInIsRevisionControlled);
+		}
+	}
+}
+
+}
+
+#undef LOCTEXT_NAMESPACE
diff --git a/Engine/Source/Editor/Sequencer/Private/Tools/NavigationTool/Widgets/Columns/SNavigationToolRevisionControl.h b/Engine/Source/Editor/Sequencer/Private/Tools/NavigationTool/Widgets/Columns/SNavigationToolRevisionControl.h
new file mode 100644
--- /dev/null
+++ b/Engine/Source/Editor/Sequencer/Private/Tools/NavigationTool/Widgets/Columns/SNavigationToolRevisionControl.h
@@ -0,0 +1,50 @@
+// Copyright Epic Games, Inc. All Rights Reserved.
+
+#pragma once
+
+#include "ScopedTransaction.h"
+#include "Templates/UniquePtr.h"
+#include "Tools/NavigationTool/NavigationToolDefines.h"
+#include "Widgets/Images/SImage.h"
+
+namespace UE::Sequencer::NavigationTool
+{
+
+class FNavigationToolRevisionControlColumn;
+class INavigationToolView;
+class SNavigationToolTreeRow;
+enum class EItemRevisionControlState;
+
+class SNavigationToolRevisionControl : public SImage
+{
+public:
+	SLATE_BEGIN_ARGS(SNavigationToolRevisionControl) {}
+	SLATE_END_ARGS()
+
+	void Construct(const FArguments& InArgs
+		, const TSharedRef<FNavigationToolRevisionControlColumn>& InColumn
+		, const FNavigationToolItemPtr& InItem
+		, const TSharedRef<INavigationToolView>& InView
+		, const TSharedRef<SNavigationToolTreeRow>& InRowWidget);
+
+private:
+	virtual FSlateColor GetForegroundColor() const override;
+	
+	virtual const FSlateBrush* GetBrush() const;
+	virtual FText GetToolTipText() const;
+	
+	virtual void OnSetIsRevisionControlled(const FNavigationToolItemPtr& InItem, const bool bInIsRevisionControlled);
+
+	EItemRevisionControlState GetRevisionControlState() const;
+	void SetIsRevisionControlled(const bool bInRevisionControlled);
+
+	TWeakPtr<FNavigationToolRevisionControlColumn> WeakColumn;
+
+	FNavigationToolItemWeakPtr WeakItem;
+
+	TWeakPtr<INavigationToolView> WeakView;
+
+	TWeakPtr<SNavigationToolTreeRow> WeakRowWidget;
+};
+
+}
diff --git a/Engine/Source/Editor/Sequencer/Private/Tools/NavigationTool/Widgets/Columns/SNavigationToolStartFrameOffset.cpp b/Engine/Source/Editor/Sequencer/Private/Tools/NavigationTool/Widgets/Columns/SNavigationToolStartFrameOffset.cpp
new file mode 100644
--- /dev/null
+++ b/Engine/Source/Editor/Sequencer/Private/Tools/NavigationTool/Widgets/Columns/SNavigationToolStartFrameOffset.cpp
@@ -0,0 +1,137 @@
+// Copyright Epic Games, Inc. All Rights Reserved.
+
+#include "SNavigationToolStartFrameOffset.h"
+#include "Sequencer.h"
+#include "Tools/NavigationTool/Items/NavigationToolItemUtils.h"
+#include "Tools/NavigationTool/NavigationToolView.h"
+#include "Tools/NavigationTool/Widgets/SNavigationToolTreeRow.h"
+#include "TrackEditors/SubTrackEditorBase.h"
+#include "Widgets/Input/SSpinBox.h"
+
+#define LOCTEXT_NAMESPACE "SNavigationToolStartFrameOffset"
+
+namespace UE::Sequencer::NavigationTool
+{
+
+void SNavigationToolStartFrameOffset::Construct(const FArguments& InArgs
+	, const FNavigationToolItemPtr& InItem
+	, const TSharedRef<INavigationToolView>& InView
+	, const TSharedRef<SNavigationToolTreeRow>& InRowWidget)
+{
+	WeakItem = InItem;
+	WeakView = InView;
+	WeakRowWidget = InRowWidget;
+
+	UMovieSceneSubSection* const SubSection = Cast<UMovieSceneSubSection>(GetSequenceItemSubSection(WeakItem.Pin()));
+
+	if (!WeakView.IsValid() || !IsValid(SubSection))
+	{
+		ChildSlot
+		[
+			SNullWidget::NullWidget
+		];
+
+		return;
+	}
+
+	ChildSlot
+		[
+			SNew(SSpinBox<int32>)
+			.Justification(ETextJustify::Center)
+			.Value(this, &SNavigationToolStartFrameOffset::GetValue)
+			.OnValueChanged(this, &SNavigationToolStartFrameOffset::OnValueChanged)
+			.OnValueCommitted(this, &SNavigationToolStartFrameOffset::OnValueCommitted)
+			.OnBeginSliderMovement(this, &SNavigationToolStartFrameOffset::OnBeginSliderMovement)
+			.OnEndSliderMovement(this, &SNavigationToolStartFrameOffset::OnEndSliderMovement)
+		];
+}
+
+int32 SNavigationToolStartFrameOffset::GetValue() const
+{
+	const TSharedPtr<INavigationToolView> ToolView = WeakView.Pin();
+	if (!ToolView.IsValid())
+	{
+		return 0;
+	}
+
+	const TSharedPtr<ISequencer> Sequencer = ToolView->GetSequencer();
+	if (!Sequencer.IsValid())
+	{
+		return 0;
+	}
+
+	UMovieSceneSubSection* const SubSection = Cast<UMovieSceneSubSection>(GetSequenceItemSubSection(WeakItem.Pin()));
+	if (!IsValid(SubSection))
+	{
+		return 0;
+	}
+
+	const FFrameNumber Offset = ConvertFrameTime(SubSection->Parameters.StartFrameOffset
+		, Sequencer->GetFocusedTickResolution()
+		, Sequencer->GetFocusedDisplayRate()).FrameNumber;
+
+	return Offset.Value;
+}
+
+void SNavigationToolStartFrameOffset::OnValueChanged(const int32 InNewValue)
+{
+	OnValueCommitted(InNewValue, ETextCommit::Default);
+}
+
+void SNavigationToolStartFrameOffset::OnValueCommitted(const int32 InNewValue, const ETextCommit::Type InCommitType)
+{
+	const TSharedPtr<INavigationToolView> ToolView = WeakView.Pin();
+	if (!ToolView.IsValid())
+	{
+		return;
+	}
+
+	const TSharedPtr<ISequencer> Sequencer = ToolView->GetSequencer();
+	if (!Sequencer.IsValid())
+	{
+		return;
+	}
+
+	UMovieSceneSubSection* const SubSection = Cast<UMovieSceneSubSection>(GetSequenceItemSubSection(WeakItem.Pin()));
+	if (!IsValid(SubSection))
+	{
+		return;
+	}
+
+	const FFrameNumber FrameTime = ConvertFrameTime(InNewValue
+		, Sequencer->GetFocusedDisplayRate()
+		, Sequencer->GetFocusedTickResolution()).FrameNumber;
+
+	if (InNewValue == FrameTime.Value)
+	{
+		return;
+	}
+
+	const bool bShouldTransact = !UndoTransaction.IsValid() && (InCommitType == ETextCommit::OnEnter);
+	const FScopedTransaction Transaction(GetTransactionText(), bShouldTransact);
+
+	SubSection->Modify();
+	SubSection->Parameters.StartFrameOffset = FrameTime;
+}
+
+void SNavigationToolStartFrameOffset::OnBeginSliderMovement()
+{
+	if (!UndoTransaction.IsValid())
+	{
+		UndoTransaction = MakeUnique<FScopedTransaction>(GetTransactionText());
+	}
+}
+
+void SNavigationToolStartFrameOffset::OnEndSliderMovement(const int32 InNewValue)
+{
+	UndoTransaction.Reset();
+}
+
+FText SNavigationToolStartFrameOffset::GetTransactionText() const
+{
+	return LOCTEXT("SetSequenceStartFrameOffsetTransaction", "Set Sequence Start Frame Offset");
+}
+
+}
+
+#undef LOCTEXT_NAMESPACE
diff --git a/Engine/Source/Editor/Sequencer/Private/Tools/NavigationTool/Widgets/Columns/SNavigationToolStartFrameOffset.h b/Engine/Source/Editor/Sequencer/Private/Tools/NavigationTool/Widgets/Columns/SNavigationToolStartFrameOffset.h
new file mode 100644
--- /dev/null
+++ b/Engine/Source/Editor/Sequencer/Private/Tools/NavigationTool/Widgets/Columns/SNavigationToolStartFrameOffset.h
@@ -0,0 +1,47 @@
+// Copyright Epic Games, Inc. All Rights Reserved.
+
+#pragma once
+
+#include "ScopedTransaction.h"
+#include "Templates/UniquePtr.h"
+#include "Tools/NavigationTool/NavigationToolDefines.h"
+#include "Widgets/SCompoundWidget.h"
+
+namespace UE::Sequencer::NavigationTool
+{
+
+class INavigationToolView;
+class SNavigationToolTreeRow;
+
+class SNavigationToolStartFrameOffset : public SCompoundWidget
+{
+public:
+	SLATE_BEGIN_ARGS(SNavigationToolStartFrameOffset) {}
+	SLATE_END_ARGS()
+
+	virtual ~SNavigationToolStartFrameOffset() override {}
+
+	void Construct(const FArguments& InArgs
+		, const FNavigationToolItemPtr& InItem
+		, const TSharedRef<INavigationToolView>& InView
+		, const TSharedRef<SNavigationToolTreeRow>& InRowWidget);
+
+private:
+	int32 GetValue() const;
+	void OnValueChanged(const int32 InNewValue);
+	void OnValueCommitted(const int32 InNewValue, const ETextCommit::Type InCommitType);
+	void OnBeginSliderMovement();
+	void OnEndSliderMovement(const int32 InNewValue);
+
+	FText GetTransactionText() const;
+
+	FNavigationToolItemWeakPtr WeakItem;
+
+	TWeakPtr<INavigationToolView> WeakView;
+
+	TWeakPtr<SNavigationToolTreeRow> WeakRowWidget;
+
+	TUniquePtr<FScopedTransaction> UndoTransaction;
+};
+
+}
diff --git a/Engine/Source/Editor/Sequencer/Private/Tools/NavigationTool/Widgets/Columns/SNavigationToolTake.cpp b/Engine/Source/Editor/Sequencer/Private/Tools/NavigationTool/Widgets/Columns/SNavigationToolTake.cpp
new file mode 100644
--- /dev/null
+++ b/Engine/Source/Editor/Sequencer/Private/Tools/NavigationTool/Widgets/Columns/SNavigationToolTake.cpp
@@ -0,0 +1,400 @@
+// Copyright Epic Games, Inc. All Rights Reserved.
+
+#include "SNavigationToolTake.h"
+#include "Framework/MultiBox/MultiBoxBuilder.h"
+#include "MovieSceneSequence.h"
+#include "MovieSceneToolHelpers.h"
+#include "ScopedTransaction.h"
+#include "Sequencer.h"
+#include "Styling/StyleColors.h"
+#include "Tools/NavigationTool/Items/NavigationToolItemUtils.h"
+#include "Tools/NavigationTool/NavigationTool.h"
+#include "Tools/NavigationTool/NavigationToolView.h"
+#include "Tools/NavigationTool/Widgets/SNavigationToolTreeRow.h"
+#include "Widgets/Input/SComboBox.h"
+#include "Widgets/Layout/SBox.h"
+
+using namespace UE::Sequencer::NavigationTool;
+
+#define LOCTEXT_NAMESPACE "SNavigationToolTake"
+
+namespace UE::Sequencer::NavigationTool
+{
+
+DECLARE_DELEGATE_RetVal_OneParam(FReply, FOnTakeEntrySelected, const TSharedPtr<SNavigationToolTake::FTakeItemInfo> /*InTakeInfo*/);
+
+class SNavigationToolTakeEntry : public SCompoundWidget
+{
+public:
+	SLATE_BEGIN_ARGS(SNavigationToolTakeEntry) {}
+		SLATE_EVENT(FOnTakeEntrySelected, OnEntrySelected)
+	SLATE_END_ARGS()
+
+	void Construct(const FArguments& InArgs, const TSharedPtr<SNavigationToolTake::FTakeItemInfo>& InTakeEntry)
+	{
+		TakeEntry  = InTakeEntry;
+		MenuButtonStyle = &FAppStyle::Get().GetWidgetStyle<FButtonStyle>("Menu.Button");
+
+		OnTakeEntrySelected = InArgs._OnEntrySelected;
+
+		ChildSlot
+		[
+			SNew(SBox)
+			.WidthOverride(120.f)
+			[
+				SNew(SBorder)
+				.BorderImage(this, &SNavigationToolTakeEntry::GetBorderImage)
+				.Padding(FMargin(12.f, 1.f, 12.f, 1.f))
+				[
+					SNew(STextBlock)
+					.Text(FText::FromString(TakeEntry->DisplayName))
+					.ColorAndOpacity(FLinearColor::White)
+				]
+
+			]
+		];
+	}
+
+	const FSlateBrush* GetBorderImage() const
+	{
+		if (IsHovered())
+		{
+			return &MenuButtonStyle->Hovered;
+		}
+		return &MenuButtonStyle->Normal;
+	}
+
+	virtual FReply OnMouseButtonUp(const FGeometry& InGeometry, const FPointerEvent& InMouseEvent) override
+	{
+		if (OnTakeEntrySelected.IsBound())
+		{
+			return OnTakeEntrySelected.Execute(TakeEntry);
+		}
+		return FReply::Handled();
+	}
+
+protected:
+	TSharedPtr<SNavigationToolTake::FTakeItemInfo> TakeEntry;
+	FOnTakeEntrySelected OnTakeEntrySelected;
+	const FButtonStyle* MenuButtonStyle = nullptr;
+};
+
+class FTakeDragDropOp : public FDragDropOperation
+{
+public:
+	DRAG_DROP_OPERATOR_TYPE(FTakeDragDropOp, FDragDropOperation)
+
+	virtual TSharedPtr<SWidget> GetDefaultDecorator() const override
+	{
+		return SNullWidget::NullWidget;
+	}
+
+	static TSharedRef<FTakeDragDropOp> New(const TSharedPtr<SNavigationToolTake::FTakeItemInfo>& InTakeInfo)
+	{
+		const TSharedRef<FTakeDragDropOp> Operation = MakeShared<FTakeDragDropOp>();
+		Operation->TakeInfo = InTakeInfo;
+		Operation->Construct();
+		return Operation;
+	}
+
+	TSharedPtr<SNavigationToolTake::FTakeItemInfo> TakeInfo;
+};
+
+void SNavigationToolTake::Construct(const FArguments& InArgs
+	, const FNavigationToolItemPtr& InItem
+	, const TSharedRef<INavigationToolView>& InView
+	, const TSharedRef<SNavigationToolTreeRow>& InRowWidget)
+{
+	WeakItem = InItem;
+	WeakView = InView;
+	WeakRowWidget = InRowWidget;
+
+	WeakTool = InView->GetOwnerTool();
+
+	CacheTakes();
+
+	if (!ActiveTakeInfo.IsValid())
+	{
+		ChildSlot
+		[
+			SNullWidget::NullWidget
+		];
+
+		return;
+	}
+
+	ChildSlot
+	[
+		SNew(SComboBox<TSharedPtr<SNavigationToolTake::FTakeItemInfo>>)
+		.OptionsSource(&CachedTakes)
+		.InitiallySelectedItem(ActiveTakeInfo)
+		.OnSelectionChanged(this, &SNavigationToolTake::OnSelectionChanged)
+		.OnGenerateWidget(this, &SNavigationToolTake::GenerateTakeWidget)
+		[
+			GenerateTakeWidget(ActiveTakeInfo)
+		]
+	];
+}
+
+TSharedRef<SWidget> SNavigationToolTake::GenerateTakeWidget(const TSharedPtr<FTakeItemInfo> InTakeInfo)
+{
+	const FText TakeInfoText = FText::Format(LOCTEXT("TakeNumberLabel", "({0}/{1})")
+		, FText::FromString(FString::FromInt(InTakeInfo->TakeIndex + 1))
+		, FText::FromString(FString::FromInt(CachedTakes.Num())));
+
+	return SNew(SHorizontalBox)
+		+ SHorizontalBox::Slot()
+		.HAlign(HAlign_Left)
+		.VAlign(VAlign_Center)
+		.AutoWidth()
+		[
+			SNew(STextBlock)
+			.Text(FText::FromString(FString::FromInt(InTakeInfo->TakeNumber)))
+		]
+		+ SHorizontalBox::Slot()
+		.FillWidth(1.f)
+		.HAlign(HAlign_Right)
+		.VAlign(VAlign_Center)
+		.Padding(5.f, 0.f, 0.f, 0.f)
+		[
+			SNew(STextBlock)
+			.TextStyle(FAppStyle::Get(), TEXT("SmallText"))
+			.ColorAndOpacity(FStyleColors::Hover)
+			.Text(TakeInfoText)
+		];
+}
+
+void SNavigationToolTake::RemoveItemColor() const
+{
+	if (const TSharedPtr<INavigationTool> Tool = WeakTool.Pin())
+	{
+		Tool->RemoveItemColor(WeakItem.Pin());
+	}
+}
+
+FSlateColor SNavigationToolTake::GetBorderColor() const
+{
+	if (IsHovered())
+	{
+		return FSlateColor::UseForeground();
+	}
+	return FSlateColor::UseSubduedForeground();
+}
+
+FReply SNavigationToolTake::OnTakeEntrySelected(const TSharedPtr<FTakeItemInfo> InTakeInfo)
+{
+	if (const TSharedPtr<INavigationTool> Tool = WeakTool.Pin())
+	{
+		SetActiveTake(InTakeInfo->WeakSequence.Get());
+
+		return FReply::Handled();
+	}
+
+	return FReply::Unhandled();
+}
+
+void SNavigationToolTake::Press()
+{
+	bIsPressed = true;
+}
+
+void SNavigationToolTake::Release()
+{
+	bIsPressed = false;
+}
+
+FReply SNavigationToolTake::OnMouseButtonDown(const FGeometry& MyGeometry, const FPointerEvent& MouseEvent)
+{
+	if (MouseEvent.IsMouseButtonDown(EKeys::LeftMouseButton))
+	{
+		Press();
+
+		const TSharedRef<SNavigationToolTake> This = SharedThis(this);
+		return FReply::Handled().CaptureMouse(This).DetectDrag(This, EKeys::LeftMouseButton);
+	}
+
+	return SCompoundWidget::OnMouseButtonDown(MyGeometry, MouseEvent);
+}
+
+FReply SNavigationToolTake::OnMouseButtonUp(const FGeometry& MyGeometry, const FPointerEvent& MouseEvent)
+{
+	FReply Reply = FReply::Unhandled();
+
+	if (bIsPressed && IsHovered() && MouseEvent.GetEffectingButton() == EKeys::LeftMouseButton)
+	{
+		Release();
+		Reply = FReply::Handled();
+	}
+
+	if (Reply.GetMouseCaptor().IsValid() == false && HasMouseCapture())
+	{
+		Reply.ReleaseMouseCapture();
+	}
+
+	return Reply;
+}
+
+void SNavigationToolTake::OnFocusLost(const FFocusEvent& InFocusEvent)
+{
+	Release();
+
+	SCompoundWidget::OnFocusLost(InFocusEvent);
+}
+
+FReply SNavigationToolTake::OnDragDetected(const FGeometry& MyGeometry, const FPointerEvent& MouseEvent)
+{
+	if (MouseEvent.IsMouseButtonDown(EKeys::LeftMouseButton))
+	{
+		return FReply::Handled().BeginDragDrop(FTakeDragDropOp::New(ActiveTakeInfo));
+	}
+
+	return SCompoundWidget::OnDragDetected(MyGeometry, MouseEvent);
+}
+
+void SNavigationToolTake::OnDragEnter(const FGeometry& MyGeometry, const FDragDropEvent& DragDropEvent)
+{
+	if (const TSharedPtr<FTakeDragDropOp> DragOp = DragDropEvent.GetOperationAs<FTakeDragDropOp>())
+	{
+		OnTakeEntrySelected(DragOp->TakeInfo);
+	}
+
+	SCompoundWidget::OnDragEnter(MyGeometry, DragDropEvent);
+}
+
+void SNavigationToolTake::OnSelectionChanged(const TSharedPtr<FTakeItemInfo> InTakeInfo, const ESelectInfo::Type InSelectType)
+{
+	if (InTakeInfo.IsValid() && InTakeInfo->WeakSequence.IsValid())
+	{
+		SetActiveTake(InTakeInfo->WeakSequence.Get());
+	}
+}
+
+void SNavigationToolTake::CacheTakes()
+{
+	CachedTakes.Reset();
+
+	UMovieSceneSubSection* const SubSection = Cast<UMovieSceneSubSection>(GetSequenceItemSubSection(WeakItem.Pin()));
+	if (!IsValid(SubSection))
+	{
+		return;
+	}
+
+	TArray<FAssetData> AssetData;
+	uint32 CurrentTakeNumber = INDEX_NONE;
+	MovieSceneToolHelpers::GatherTakes(SubSection, AssetData, CurrentTakeNumber);
+
+	AssetData.Sort([SubSection](const FAssetData& InA, const FAssetData& InB)
+		{
+			uint32 TakeNumberA = INDEX_NONE;
+			uint32 TakeNumberB = INDEX_NONE;
+			if (MovieSceneToolHelpers::GetTakeNumber(SubSection, InA, TakeNumberA)
+				&& MovieSceneToolHelpers::GetTakeNumber(SubSection, InB, TakeNumberB))
+			{
+				return TakeNumberA < TakeNumberB;
+			}
+			return true;
+		});
+
+	uint32 TakeIndex = 0;
+	for (const FAssetData& ThisAssetData : AssetData)
+	{
+		uint32 TakeNumber = INDEX_NONE;
+		if (MovieSceneToolHelpers::GetTakeNumber(SubSection, ThisAssetData, TakeNumber))
+		{
+			UMovieSceneSequence* const Sequence = Cast<UMovieSceneSequence>(ThisAssetData.GetAsset());
+			if (IsValid(Sequence))
+			{
+				const TSharedRef<FTakeItemInfo> NewTakeInfo = MakeShared<FTakeItemInfo>();
+				NewTakeInfo->TakeIndex = TakeIndex;
+				NewTakeInfo->TakeNumber = TakeNumber;
+				NewTakeInfo->DisplayName = Sequence->GetDisplayName().ToString();
+				NewTakeInfo->WeakSequence = Sequence;
+				CachedTakes.Add(NewTakeInfo);
+
+				if (TakeNumber == CurrentTakeNumber)
+				{
+					ActiveTakeInfo = NewTakeInfo;
+				}
+
+				++TakeIndex;
+			}
+		}
+	}
+}
+
+void SNavigationToolTake::SetActiveTake(UMovieSceneSequence* const InSequence)
+{
+	const TSharedPtr<INavigationTool> Tool = WeakTool.Pin();
+	if (!Tool.IsValid())
+	{
+		return;
+	}
+
+	//FSequencer& Sequencer = Tool->GetSequencer();
+	const TSharedPtr<ISequencer> Sequencer = Tool->GetSequencer();
+	if (!Sequencer.IsValid())
+	{
+		return;
+	}
+
+	bool bChangedTake = false;
+
+	const FScopedTransaction Transaction(LOCTEXT("ChangeTake_Transaction", "Change Take"));
+
+	TArray<UMovieSceneSection*> Sections;
+	Sequencer->GetSelectedSections(Sections);
+
+	for (int32 SectionIndex = 0; SectionIndex < Sections.Num(); ++SectionIndex)
+	{
+		if (!Sections[SectionIndex]->IsA<UMovieSceneSubSection>())
+		{
+			continue;
+		}
+
+		UMovieSceneSubSection* Section = Cast<UMovieSceneSubSection>(Sections[SectionIndex]);
+		UMovieSceneSubTrack* SubTrack = CastChecked<UMovieSceneSubTrack>(Section->GetOuter());
+
+		TRange<FFrameNumber> NewSectionRange = Section->GetRange();
+		FFrameNumber		 NewSectionStartOffset = Section->Parameters.StartFrameOffset;
+		int32                NewSectionPrerollFrames = Section->GetPreRollFrames();
+		int32                NewRowIndex = Section->GetRowIndex();
+		FFrameNumber         NewSectionStartTime = NewSectionRange.GetLowerBound().IsClosed() ? UE::MovieScene::DiscreteInclusiveLower(NewSectionRange) : 0;
+		int32                NewSectionRowIndex = Section->GetRowIndex();
+		FColor               NewSectionColorTint = Section->GetColorTint();
+
+		const int32 Duration = (NewSectionRange.GetLowerBound().IsClosed() && NewSectionRange.GetUpperBound().IsClosed()) ? UE::MovieScene::DiscreteSize(NewSectionRange) : 1;
+		UMovieSceneSubSection* NewSection = SubTrack->AddSequence(InSequence, NewSectionStartTime, Duration);
+
+		if (NewSection != nullptr)
+		{
+			SubTrack->RemoveSection(*Section);
+
+			NewSection->SetRange(NewSectionRange);
+			NewSection->Parameters.StartFrameOffset = NewSectionStartOffset;
+			NewSection->Parameters.TimeScale = Section->Parameters.TimeScale.DeepCopy(NewSection);
+			NewSection->SetPreRollFrames(NewSectionPrerollFrames);
+			NewSection->SetRowIndex(NewSectionRowIndex);
+			NewSection->SetColorTint(NewSectionColorTint);
+
+			UMovieSceneCinematicShotSection* ShotSection = Cast<UMovieSceneCinematicShotSection>(Section);
+			UMovieSceneCinematicShotSection* NewShotSection = Cast<UMovieSceneCinematicShotSection>(NewSection);
+
+			// If the old shot's name is not the same as the sequence's name, assume the user had customized the shot name, so carry it over
+			if (ShotSection && NewShotSection && ShotSection->GetSequence() && ShotSection->GetShotDisplayName() != ShotSection->GetSequence()->GetName())
+			{
+				NewShotSection->SetShotDisplayName(ShotSection->GetShotDisplayName());
+			}
+
+			bChangedTake = true;
+		}
+	}
+
+	if (bChangedTake)
+	{
+		Sequencer->NotifyMovieSceneDataChanged(EMovieSceneDataChangeType::MovieSceneStructureItemsChanged);
+	}
+}
+
+}
+
+#undef LOCTEXT_NAMESPACE
diff --git a/Engine/Source/Editor/Sequencer/Private/Tools/NavigationTool/Widgets/Columns/SNavigationToolTake.h b/Engine/Source/Editor/Sequencer/Private/Tools/NavigationTool/Widgets/Columns/SNavigationToolTake.h
new file mode 100644
--- /dev/null
+++ b/Engine/Source/Editor/Sequencer/Private/Tools/NavigationTool/Widgets/Columns/SNavigationToolTake.h
@@ -0,0 +1,82 @@
+// Copyright Epic Games, Inc. All Rights Reserved.
+
+#pragma once
+
+#include "ScopedTransaction.h"
+#include "Templates/UniquePtr.h"
+#include "Tools/NavigationTool/NavigationToolDefines.h"
+#include "Tools/NavigationTool/Columns/NavigationToolTakeColumn.h"
+#include "Widgets/SCompoundWidget.h"
+
+class SMenuAnchor;
+class UMovieSceneSequence;
+
+namespace UE::Sequencer::NavigationTool
+{
+
+class INavigationTool;
+class INavigationToolItem;
+class INavigationToolView;
+class SNavigationToolTreeRow;
+
+class SNavigationToolTake : public SCompoundWidget
+{
+public:
+	struct FTakeItemInfo
+	{
+		uint32 TakeIndex = 0;
+		uint32 TakeNumber = 0;
+		FString DisplayName;
+		TWeakObjectPtr<UMovieSceneSequence> WeakSequence;
+	};
+
+	SLATE_BEGIN_ARGS(SNavigationToolTake) {}
+	SLATE_END_ARGS()
+
+	void Construct(const FArguments& InArgs
+		, const FNavigationToolItemPtr& InItem
+		, const TSharedRef<INavigationToolView>& InView
+		, const TSharedRef<SNavigationToolTreeRow>& InRowWidget);
+
+	void RemoveItemColor() const;
+
+	FSlateColor GetBorderColor() const;
+
+	FReply OnTakeEntrySelected(const TSharedPtr<FTakeItemInfo> InTakeInfo);
+
+	void Press();
+	void Release();
+
+	virtual FReply OnMouseButtonDown(const FGeometry& MyGeometry, const FPointerEvent& MouseEvent) override;
+	virtual FReply OnMouseButtonUp(const FGeometry& MyGeometry, const FPointerEvent& MouseEvent) override;
+	virtual void OnFocusLost(const FFocusEvent& InFocusEvent) override;
+
+	virtual FReply OnDragDetected(const FGeometry& MyGeometry, const FPointerEvent& MouseEvent) override;
+	virtual void OnDragEnter(const FGeometry& MyGeometry, const FDragDropEvent& DragDropEvent) override;
+
+private:
+	TSharedRef<SWidget> GenerateTakeWidget(const TSharedPtr<FTakeItemInfo> InTakeInfo);
+
+	void OnSelectionChanged(const TSharedPtr<FTakeItemInfo> InTakeInfo, const ESelectInfo::Type InSelectType);
+
+	void CacheTakes();
+	void SetActiveTake(UMovieSceneSequence* const InSequence);
+
+	FNavigationToolItemWeakPtr WeakItem;
+
+	TWeakPtr<INavigationToolView> WeakView;
+
+	TWeakPtr<SNavigationToolTreeRow> WeakRowWidget;
+
+	TWeakPtr<INavigationTool> WeakTool;
+
+	TArray<TSharedPtr<FTakeItemInfo>> CachedTakes;
+
+	TSharedPtr<FTakeItemInfo> ActiveTakeInfo;
+
+	bool bIsPressed = false;
+
+	TUniquePtr<FScopedTransaction> UndoTransaction;
+};
+
+}
diff --git a/Engine/Source/Editor/Sequencer/Private/Tools/NavigationTool/Widgets/ModalTextInputDialog.h b/Engine/Source/Editor/Sequencer/Private/Tools/NavigationTool/Widgets/ModalTextInputDialog.h
new file mode 100644
--- /dev/null
+++ b/Engine/Source/Editor/Sequencer/Private/Tools/NavigationTool/Widgets/ModalTextInputDialog.h
@@ -0,0 +1,33 @@
+﻿// Copyright Epic Games, Inc. All Rights Reserved.
+
+#pragma once
+
+#include "Internationalization/Text.h"
+#include "Templates/SharedPointer.h"
+
+class FString;
+class SWindow;
+
+namespace UE::Sequencer::NavigationTool
+{
+
+class FModalTextInputDialog
+{
+public:
+	FModalTextInputDialog();
+
+	/** @return True if the text input was accepted. False if cancel was clicked. */
+	bool Open(const FString& InDefaultText, FString& OutText, const TSharedPtr<SWindow>& InParentWindow = nullptr);
+
+	bool IsOpen();
+
+	void Close();
+
+	FText DialogTitle;
+	FText InputLabel;
+
+protected:
+	TSharedPtr<SWindow> WindowWidget;
+};
+
+}
diff --git a/Engine/Source/Editor/Sequencer/Private/Tools/NavigationTool/Widgets/SNavigationToolTreeRow.cpp b/Engine/Source/Editor/Sequencer/Private/Tools/NavigationTool/Widgets/SNavigationToolTreeRow.cpp
new file mode 100644
--- /dev/null
+++ b/Engine/Source/Editor/Sequencer/Private/Tools/NavigationTool/Widgets/SNavigationToolTreeRow.cpp
@@ -0,0 +1,117 @@
+// Copyright Epic Games, Inc. All Rights Reserved.
+
+#include "Tools/NavigationTool/Widgets/SNavigationToolTreeRow.h"
+#include "Tools/NavigationTool/NavigationTool.h"
+#include "Tools/NavigationTool/Columns/INavigationToolColumn.h"
+#include "Tools/NavigationTool/Items/INavigationToolItem.h"
+#include "Tools/NavigationTool/NavigationToolView.h"
+#include "Tools/NavigationTool/Widgets/SNavigationToolTreeView.h"
+
+#define LOCTEXT_NAMESPACE "SNavigationToolTreeRow"
+
+namespace UE::Sequencer::NavigationTool
+{
+
+void SNavigationToolTreeRow::Construct(const FArguments& InArgs
+	, const TSharedRef<FNavigationToolView>& InToolView
+	, const TSharedPtr<SNavigationToolTreeView>& InTreeView
+	, const FNavigationToolItemPtr& InItem)
+{
+	WeakToolView = InToolView;
+	WeakTreeView = InTreeView;
+	Item = InItem;
+	HighlightText = InArgs._HighlightText;
+
+	SetColorAndOpacity(TAttribute<FLinearColor>::CreateSP(&*InToolView, &FNavigationToolView::GetItemBrushColor, Item));
+
+	SMultiColumnTableRow::Construct(FSuperRowType::FArguments()
+			.Style(&FAppStyle::Get().GetWidgetStyle<FTableRowStyle>(TEXT("Sequencer.NavigationTool.TableViewRow")))
+			.OnCanAcceptDrop(InToolView, &FNavigationToolView::OnCanDrop)
+			.OnDragDetected(InToolView, &FNavigationToolView::OnDragDetected, Item)
+			.OnDragEnter(InToolView, &FNavigationToolView::OnDragEnter, Item)
+			.OnDragLeave(InToolView, &FNavigationToolView::OnDragLeave, Item)
+			.OnAcceptDrop(InToolView, &FNavigationToolView::OnDrop)
+			.OnDrop(this, &SNavigationToolTreeRow::OnDefaultDrop)
+		, InTreeView.ToSharedRef());
+}
+
+TSharedRef<SWidget> SNavigationToolTreeRow::GenerateWidgetForColumn(const FName& InColumnName)
+{
+	if (Item.IsValid())
+	{
+		const TSharedPtr<FNavigationToolView> ToolView = WeakToolView.Pin();
+		check(ToolView.IsValid());
+		
+		if (const TSharedPtr<INavigationToolColumn> Column = ToolView->GetColumns().FindRef(InColumnName))
+		{
+			return Column->ConstructRowWidget(Item.ToSharedRef(), ToolView.ToSharedRef(), SharedThis(this));
+		}
+	}
+	return SNullWidget::NullWidget;
+}
+
+FReply SNavigationToolTreeRow::OnMouseButtonUp(const FGeometry& MyGeometry, const FPointerEvent& MouseEvent)
+{
+	const TSharedPtr<FNavigationToolView> ToolView = WeakToolView.Pin();
+	if (!ToolView.IsValid())
+	{
+		return FReply::Unhandled();
+	}
+
+	//Select Item and the Tree of Children it contains
+	if (MouseEvent.GetEffectingButton() == EKeys::MiddleMouseButton)
+	{
+		ENavigationToolItemSelectionFlags Flags = ENavigationToolItemSelectionFlags::IncludeChildren
+			| ENavigationToolItemSelectionFlags::SignalSelectionChange
+			| ENavigationToolItemSelectionFlags::ScrollIntoView;
+
+		if (MouseEvent.IsControlDown())
+		{
+			Flags |= ENavigationToolItemSelectionFlags::AppendToCurrentSelection;
+		}
+
+		ToolView->SelectItems({Item}, Flags);
+
+		Item->OnSelect();
+
+		return FReply::Handled();
+	}
+	
+	return SMultiColumnTableRow::OnMouseButtonUp(MyGeometry, MouseEvent);
+}
+
+FReply SNavigationToolTreeRow::OnMouseButtonDoubleClick(const FGeometry& InGeometry, const FPointerEvent& InMouseEvent)
+{
+	// Select Item and the Tree of Children it contains
+	if (InMouseEvent.GetEffectingButton() == EKeys::LeftMouseButton)
+	{
+		if (const TSharedPtr<FNavigationToolView> ToolView = WeakToolView.Pin())
+		{
+			Item->OnDoubleClick();
+
+			return FReply::Handled();
+		}
+	}
+
+	return SMultiColumnTableRow<FNavigationToolItemPtr>::OnMouseButtonDoubleClick(InGeometry, InMouseEvent);
+}
+
+TSharedPtr<FNavigationToolView> SNavigationToolTreeRow::GetToolView() const
+{
+	return WeakToolView.Pin();
+}
+
+FReply SNavigationToolTreeRow::OnDefaultDrop(const FDragDropEvent& InDragDropEvent) const
+{
+	if (const TSharedPtr<FNavigationToolView> ToolView = WeakToolView.Pin())
+	{
+		ToolView->SetDragIntoTreeRoot(false);
+	}
+
+	// Always return handled as no action should take place if the Drop wasn't accepted
+	return FReply::Handled();
+}
+
+}
+
+#undef LOCTEXT_NAMESPACE
