[Sequencer] Navigation Tool initial submit

diff --git a/Engine/Source/Editor/Sequencer/Private/Tools/NavigationTool/Widgets/SNavigationToolTreeView.cpp b/Engine/Source/Editor/Sequencer/Private/Tools/NavigationTool/Widgets/SNavigationToolTreeView.cpp
new file mode 100644
--- /dev/null
+++ b/Engine/Source/Editor/Sequencer/Private/Tools/NavigationTool/Widgets/SNavigationToolTreeView.cpp
@@ -0,0 +1,279 @@
+// Copyright Epic Games, Inc. All Rights Reserved.
+
+#include "SNavigationToolTreeView.h"
+#include "Framework/Commands/UICommandList.h"
+#include "Tools/NavigationTool/Items/INavigationToolItem.h"
+#include "Tools/NavigationTool/NavigationToolView.h"
+
+namespace UE::Sequencer::NavigationTool
+{
+
+void SNavigationToolTreeView::Construct(const FArguments& InArgs, const TSharedPtr<FNavigationToolView>& InToolView)
+{
+	WeakToolView = InToolView;
+
+	if (InToolView.IsValid())
+	{
+		InToolView->SaveColumnState();
+	}
+
+	STreeView::Construct(InArgs._TreeViewArgs);
+}
+
+int32 SNavigationToolTreeView::GetItemIndex(const FNavigationToolItemPtr& Item) const
+{
+	if (SListView<FNavigationToolItemPtr>::HasValidItemsSource())
+	{
+		return SListView<FNavigationToolItemPtr>::GetItems().Find(TListTypeTraits<FNavigationToolItemPtr>::NullableItemTypeConvertToItemType(Item));
+	}
+	return INDEX_NONE;
+}
+
+void SNavigationToolTreeView::FocusOnItem(const FNavigationToolItemPtr& InItem)
+{
+	SelectorItem = InItem;
+	RangeSelectionStart = InItem;
+}
+
+void SNavigationToolTreeView::ScrollItemIntoView(const FNavigationToolItemPtr& InItem)
+{
+	RequestScrollIntoView(InItem, 0);
+}
+
+void SNavigationToolTreeView::UpdateItemExpansions(const FNavigationToolItemPtr& InItem)
+{
+	if (InItem.IsValid())
+	{
+		const FSparseItemInfo* const SparseItemInfo = SparseItemInfos.Find(InItem);
+
+		bool bIsExpanded = false;
+		bool bHasExpandedChildren = false;
+
+		if (SparseItemInfo)
+		{
+			bIsExpanded = SparseItemInfo->bIsExpanded;
+			bHasExpandedChildren = SparseItemInfo->bHasExpandedChildren;
+		}
+
+		//Skip to avoid redundancy
+		if (bIsExpanded || bHasExpandedChildren)
+		{
+			return;
+		}
+
+		TArray<FNavigationToolItemPtr> ItemsToCheck = InItem->GetChildren();
+		while (ItemsToCheck.Num() > 0)
+		{
+			FNavigationToolItemPtr ItemToCheck = ItemsToCheck.Pop();
+
+			if (ItemToCheck.IsValid())
+			{
+				if (IsItemExpanded(ItemToCheck))
+				{
+					bHasExpandedChildren = true;
+					break;
+				}
+				ItemsToCheck.Append(ItemToCheck->GetChildren());
+			}
+		}
+
+		if (bIsExpanded || bHasExpandedChildren)
+		{
+			SparseItemInfos.Add(InItem, FSparseItemInfo(bIsExpanded, bHasExpandedChildren));
+		}
+	}
+}
+
+void SNavigationToolTreeView::Private_SetItemSelection(const FNavigationToolItemPtr InItem, bool bShouldBeSelected, bool bWasUserDirected)
+{
+	const TSharedPtr<FNavigationToolView> ToolView = WeakToolView.Pin();
+	if (ToolView.IsValid() && !ToolView->IsToolLocked() && ToolView->CanSelectItem(InItem))
+	{
+		STreeView::Private_SetItemSelection(InItem, bShouldBeSelected, bWasUserDirected);
+	}
+}
+
+void SNavigationToolTreeView::Private_ClearSelection()
+{
+	const TSharedPtr<FNavigationToolView> ToolView = WeakToolView.Pin();
+	if (ToolView.IsValid() && !ToolView->IsToolLocked())
+	{
+		STreeView::Private_ClearSelection();
+	}
+}
+
+void SNavigationToolTreeView::Private_SignalSelectionChanged(ESelectInfo::Type SelectInfo)
+{
+	STreeView::Private_SignalSelectionChanged(SelectInfo);
+
+	const TItemSet AddedItems = SelectedItems.Difference(PreviousSelectedItems);
+	for (const FNavigationToolItemPtr& AddedItem : AddedItems)
+	{
+		if (AddedItem.IsValid())
+		{
+			AddedItem->OnItemSelectionChanged(true);
+		}
+	}
+
+	const TItemSet RemovedItems = PreviousSelectedItems.Difference(SelectedItems);
+	for (const FNavigationToolItemPtr& RemovedItem : RemovedItems)
+	{
+		if (RemovedItem.IsValid())
+		{
+			RemovedItem->OnItemSelectionChanged(false);
+		}
+	}
+	PreviousSelectedItems = SelectedItems;
+}
+
+void SNavigationToolTreeView::Private_UpdateParentHighlights()
+{
+	this->Private_ClearHighlightedItems();
+	for (TItemSet::TConstIterator SelectedItemIt(this->SelectedItems); SelectedItemIt; ++SelectedItemIt)
+	{
+		const FNavigationToolItemPtr SelectedItem(*SelectedItemIt);
+
+		// Sometimes selection events can come through before the Linearized List is built, so the item may not exist yet.
+		const int32 ItemIndex = LinearizedItems.Find(SelectedItem);
+		if (ItemIndex == INDEX_NONE)
+		{
+			if (SelectedItem.IsValid())
+			{
+				FNavigationToolItemPtr ParentItem = SelectedItem->GetParent();
+				while (ParentItem.IsValid())
+				{
+					if (LinearizedItems.Contains(ParentItem))
+					{
+						this->SetItemHighlighted(ParentItem, true);
+					}
+					ParentItem = ParentItem->GetParent();
+				}
+			}
+			continue;
+		}
+
+		if (DenseItemInfos.IsValidIndex(ItemIndex))
+		{
+			const FItemInfo& ItemInfo = DenseItemInfos[ItemIndex];
+			
+			int32 ParentIndex = ItemInfo.ParentIndex;
+			
+			while (ParentIndex != INDEX_NONE)
+			{
+				const FNavigationToolItemPtr& ParentItem = this->LinearizedItems[ParentIndex];
+				this->Private_SetItemHighlighted(ParentItem, true);
+
+				const FItemInfo& ParentItemInfo = DenseItemInfos[ParentIndex];
+				
+				ParentIndex = ParentItemInfo.ParentIndex;
+			}
+		}
+	}
+}
+
+FCursorReply SNavigationToolTreeView::OnCursorQuery(const FGeometry& MyGeometry, const FPointerEvent& CursorEvent) const
+{
+	if (IsRightClickScrolling() && CursorEvent.IsMouseButtonDown(EKeys::MiddleMouseButton))
+	{
+		// We hide the native cursor as we'll be drawing the software EMouseCursor::GrabHandClosed cursor
+		return FCursorReply::Cursor(EMouseCursor::None);
+	}
+	return STreeView::OnCursorQuery(MyGeometry, CursorEvent);
+}
+
+FReply SNavigationToolTreeView::OnMouseMove(const FGeometry& MyGeometry, const FPointerEvent& MouseEvent)
+{
+	if (MouseEvent.IsMouseButtonDown(EKeys::MiddleMouseButton) && !MouseEvent.IsTouchEvent())
+	{
+		// We only care about deltas along the scroll axis
+		FTableViewDimensions CursorDeltaDimensions(Orientation, MouseEvent.GetCursorDelta());
+		CursorDeltaDimensions.LineAxis = 0.f;
+
+		const float ScrollByAmount = CursorDeltaDimensions.ScrollAxis / MyGeometry.Scale;
+
+		// If scrolling with the right mouse button, we need to remember how much we scrolled.
+		// If we did not scroll at all, we will bring up the context menu when the mouse is released.
+		AmountScrolledWhileRightMouseDown += FMath::Abs(ScrollByAmount);
+
+		// Has the mouse moved far enough with the right mouse button held down to start capturing
+		// the mouse and dragging the view?
+		if (IsRightClickScrolling())
+		{
+			// Make sure the active timer is registered to update the inertial scroll
+			if (!bIsScrollingActiveTimerRegistered)
+			{
+				bIsScrollingActiveTimerRegistered = true;
+				RegisterActiveTimer(0.f, FWidgetActiveTimerDelegate::CreateSP(this, &SNavigationToolTreeView::UpdateInertialScroll));
+			}
+
+			TickScrollDelta -= ScrollByAmount;
+
+			const float AmountScrolled = this->ScrollBy(MyGeometry, -ScrollByAmount, AllowOverscroll);
+
+			FReply Reply = FReply::Handled();
+
+			// The mouse moved enough that we're now dragging the view. Capture the mouse
+			// so the user does not have to stay within the bounds of the list while dragging.
+			if (this->HasMouseCapture() == false)
+			{
+				Reply.CaptureMouse(AsShared()).UseHighPrecisionMouseMovement(AsShared());
+				SoftwareCursorPosition = MyGeometry.AbsoluteToLocal(MouseEvent.GetScreenSpacePosition());
+				bShowSoftwareCursor    = true;
+			}
+
+			// Check if the mouse has moved.
+			if (AmountScrolled != 0)
+			{
+				SoftwareCursorPosition += CursorDeltaDimensions.ToVector2D();
+			}
+
+			return Reply;
+		}
+	}
+
+	return STreeView::OnMouseMove(MyGeometry, MouseEvent);
+}
+
+FReply SNavigationToolTreeView::OnMouseButtonUp(const FGeometry& MyGeometry, const FPointerEvent& MouseEvent)
+{
+	if (MouseEvent.GetEffectingButton() == EKeys::MiddleMouseButton)
+	{
+		FReply Reply = FReply::Handled().ReleaseMouseCapture();
+		AmountScrolledWhileRightMouseDown = 0;
+		bShowSoftwareCursor = false;
+
+		// If we have mouse capture, snap the mouse back to the closest location that is within the list's bounds
+		if (HasMouseCapture())
+		{
+			const FSlateRect ListScreenSpaceRect = MyGeometry.GetLayoutBoundingRect();
+			const FVector2D CursorPosition = MyGeometry.LocalToAbsolute(SoftwareCursorPosition);
+
+			const FIntPoint BestPositionInList(
+					FMath::RoundToInt(FMath::Clamp(CursorPosition.X, ListScreenSpaceRect.Left, ListScreenSpaceRect.Right)),
+					FMath::RoundToInt(FMath::Clamp(CursorPosition.Y, ListScreenSpaceRect.Top, ListScreenSpaceRect.Bottom))
+				);
+
+			Reply.SetMousePos(BestPositionInList);
+		}
+
+		return Reply;
+	}
+	return STreeView::OnMouseButtonUp(MyGeometry, MouseEvent);
+}
+
+FReply SNavigationToolTreeView::OnKeyDown(const FGeometry& MyGeometry, const FKeyEvent& InKeyEvent)
+{
+	if (const TSharedPtr<FNavigationToolView> ToolView = WeakToolView.Pin())
+	{
+		ToolView->UpdateRecentViews();
+
+		TSharedPtr<FUICommandList> CommandList = ToolView->GetViewCommandList();
+		if (CommandList.IsValid() && CommandList->ProcessCommandBindings(InKeyEvent))
+		{
+			return FReply::Handled();	
+		}
+	}
+	return STreeView::OnKeyDown(MyGeometry, InKeyEvent);
+}
+
+}
diff --git a/Engine/Source/Editor/Sequencer/Private/Tools/NavigationTool/Widgets/SNavigationToolTreeView.h b/Engine/Source/Editor/Sequencer/Private/Tools/NavigationTool/Widgets/SNavigationToolTreeView.h
new file mode 100644
--- /dev/null
+++ b/Engine/Source/Editor/Sequencer/Private/Tools/NavigationTool/Widgets/SNavigationToolTreeView.h
@@ -0,0 +1,53 @@
+// Copyright Epic Games, Inc. All Rights Reserved.
+
+#pragma once
+
+#include "Tools/NavigationTool/NavigationToolDefines.h"
+#include "Widgets/Views/STreeView.h"
+
+namespace UE::Sequencer::NavigationTool
+{
+
+class FNavigationToolView;
+
+class SNavigationToolTreeView : public STreeView<FNavigationToolItemPtr>
+{
+public:
+	SLATE_BEGIN_ARGS(SNavigationToolTreeView) {}
+		SLATE_ARGUMENT(STreeView<FNavigationToolItemPtr>::FArguments, TreeViewArgs)
+	SLATE_END_ARGS()
+
+	void Construct(const FArguments& InArgs, const TSharedPtr<FNavigationToolView>& InToolView);
+
+	int32 GetItemIndex(const FNavigationToolItemPtr& Item) const;
+
+	void FocusOnItem(const FNavigationToolItemPtr& InItem);
+
+	void ScrollItemIntoView(const FNavigationToolItemPtr& InItem);
+
+	void UpdateItemExpansions(const FNavigationToolItemPtr& InItem);
+
+	//~ Begin STreeView
+	virtual void Private_UpdateParentHighlights() override;
+	//~ End STreeView
+
+	//~ Begin ITypedTableView
+	virtual void Private_SetItemSelection(const FNavigationToolItemPtr InItem, bool bShouldBeSelected, bool bWasUserDirected = false) override;
+	virtual void Private_SignalSelectionChanged(ESelectInfo::Type SelectInfo) override;
+	virtual void Private_ClearSelection() override;
+	//~ End SListView
+
+	//~ Begin SWidget
+	virtual FCursorReply OnCursorQuery(const FGeometry& MyGeometry, const FPointerEvent& CursorEvent) const override;
+	virtual FReply OnMouseMove(const FGeometry& MyGeometry, const FPointerEvent& MouseEvent) override;
+	virtual FReply OnMouseButtonUp(const FGeometry& MyGeometry, const FPointerEvent& MouseEvent) override;
+	virtual FReply OnKeyDown(const FGeometry& MyGeometry, const FKeyEvent& InKeyEvent) override;
+	//~ End SWidget
+
+private:
+	TWeakPtr<FNavigationToolView> WeakToolView;
+	
+	SNavigationToolTreeView::TItemSet PreviousSelectedItems;
+};
+
+}
diff --git a/Engine/Source/Editor/Sequencer/Private/Tools/NavigationTool/Widgets/SNavigationToolView.h b/Engine/Source/Editor/Sequencer/Private/Tools/NavigationTool/Widgets/SNavigationToolView.h
new file mode 100644
--- /dev/null
+++ b/Engine/Source/Editor/Sequencer/Private/Tools/NavigationTool/Widgets/SNavigationToolView.h
@@ -0,0 +1,116 @@
+﻿// Copyright Epic Games, Inc. All Rights Reserved.
+
+#pragma once
+
+#include "Filters/SBasicFilterBar.h"
+#include "Tools/NavigationTool/Items/INavigationToolItem.h"
+#include "Tools/NavigationTool/NavigationToolDefines.h"
+#include "Widgets/SCompoundWidget.h"
+
+class ITableRow;
+class SAssetFilterBar;
+class SAssetSearchBox;
+class SBorder;
+class SBox;
+class SComboButton;
+class SHeaderRow;
+class SHorizontalBox;
+class SSequencerSearchBox;
+class STableViewBase;
+class UToolMenu;
+struct FAssetSearchBoxSuggestion;
+struct FCustomTextFilterData;
+
+namespace UE::Sequencer::NavigationTool
+{
+
+class FNavigationToolFilterBar;
+class FNavigationToolItemDragDropOp;
+class FNavigationToolToolbarMenu;
+class FNavigationToolView;
+class SNavigationToolFilterBar;
+class SNavigationToolItemFilters;
+class SNavigationToolTreeView;
+
+class SNavigationToolView : public SCompoundWidget
+{
+public:
+	SLATE_BEGIN_ARGS(SNavigationToolView) {}
+	SLATE_END_ARGS()
+
+	void Construct(const FArguments& InArgs, const TSharedRef<FNavigationToolView>& InToolView);
+
+	virtual ~SNavigationToolView() override;
+
+	TSharedPtr<SNavigationToolTreeView> GetTreeView() const { return TreeView; }
+
+	/** Rebuilds the entire widget. Called when the filter bar layout is changed. */
+	void RebuildWidget();
+
+	void ReconstructColumns();
+
+	bool IsColumnVisible(const FName InColumnId) const;
+	void ShowHideColumn(const FName InColumnId, const bool bInVisible);
+
+	void SetColumnWidth(const FName InColumnId, const float InWidth);
+
+	void RequestTreeRefresh();
+
+	void SetItemSelection(const TArray<FNavigationToolItemPtr>& InItems, bool bSignalSelectionChange);
+
+	void OnItemSelectionChanged(const FNavigationToolItemPtr InItem, ESelectInfo::Type InSelectionType);
+
+	void ScrollItemIntoView(const FNavigationToolItemPtr& InItem) const;
+
+	bool IsItemExpanded(const FNavigationToolItemPtr& InItem) const;
+	void SetItemExpansion(const FNavigationToolItemPtr& Item, const bool bInExpand) const;
+
+	void UpdateItemExpansions(const FNavigationToolItemPtr& Item) const;
+
+	TSharedRef<ITableRow> OnItemGenerateRow(const FNavigationToolItemPtr Item, const TSharedRef<STableViewBase>& OwnerTable);
+
+	void SetKeyboardFocus() const;
+
+	virtual void OnMouseEnter(const FGeometry& MyGeometry, const FPointerEvent& MouseEvent) override;
+	virtual FReply OnKeyDown(const FGeometry& MyGeometry, const FKeyEvent& InKeyEvent) override;
+	virtual FReply OnDrop(const FGeometry& MyGeometry, const FDragDropEvent& DragDropEvent) override;
+	virtual void OnDragEnter(const FGeometry& MyGeometry, const FDragDropEvent& DragDropEvent) override;
+	virtual void OnDragLeave(const FDragDropEvent& DragDropEvent) override;
+
+	void SetTreeBorderVisibility(bool bVisible);
+
+	void GenerateColumnState(const FName InColumnId, FNavigationToolViewColumnSaveState& OutColumnState);
+	void GenerateColumnStateMap(TMap<FName, FNavigationToolViewColumnSaveState>& OutStateMap);
+
+	EFilterBarLayout GetFilterBarLayout() const;
+	void SetFilterBarLayout(const EFilterBarLayout InLayout);
+
+private:
+	TSharedRef<SWidget> ConstructSplitterContent();
+	TSharedRef<SWidget> ConstructMainContent();
+
+	void RebuildSearchAndFilterRow();
+
+	void OnFilterBarStateChanged(const bool InIsVisible, const EFilterBarLayout InNewLayout);
+
+	TSharedPtr<FNavigationToolFilterBar> GetFilterBar() const;
+
+	TWeakPtr<FNavigationToolView> WeakToolView;
+
+	TSharedPtr<FNavigationToolToolbarMenu> ToolbarMenu;
+
+	TSharedPtr<SVerticalBox> SearchAndFilterRow;
+	TSharedPtr<SHeaderRow> HeaderRowWidget;
+	TSharedPtr<SBorder> TreeBorder;
+	TSharedPtr<SNavigationToolTreeView> TreeView;
+	TSharedPtr<SNavigationToolFilterBar> FilterBarWidget;
+
+	/** If true the SuggestionList shouldn't appear since we selected what we wanted */
+	bool bIsEnterLastKeyPressed = false;
+
+	bool bSelectingItems = false;
+
+	TSet<TWeakPtr<FNavigationToolItemDragDropOp>> ItemDragDropOps;
+};
+
+}
diff --git a/Engine/Source/Editor/Sequencer/Public/ISequencer.h b/Engine/Source/Editor/Sequencer/Public/ISequencer.h
--- a/Engine/Source/Editor/Sequencer/Public/ISequencer.h
+++ b/Engine/Source/Editor/Sequencer/Public/ISequencer.h
@@ -95,6 +95,13 @@ struct FCreateBindingParams;
 struct FTimeDomainOverride;
 struct FSequencerNumericTypeInterface;
 
+namespace NavigationTool
+{
+
+class INavigationTool;
+
+}
+
 } // namespace Sequencer
 } // namespace UE
 
@@ -991,6 +998,8 @@ public:
 
 	virtual bool TrackSupportsConditions(const UMovieSceneTrack* Track) const = 0;
 
+	virtual TSharedPtr<UE::Sequencer::NavigationTool::INavigationTool> GetNavigationTool() const = 0;
+
 protected:
 	FOnInitializeDetailsPanel InitializeDetailsPanelEvent;
 	FOnCameraAddedToSequencer CameraAddedToSequencer;
diff --git a/Engine/Source/Editor/Sequencer/Public/ISequencerModule.h b/Engine/Source/Editor/Sequencer/Public/ISequencerModule.h
--- a/Engine/Source/Editor/Sequencer/Public/ISequencerModule.h
+++ b/Engine/Source/Editor/Sequencer/Public/ISequencerModule.h
@@ -30,6 +30,11 @@ namespace UE::Sequencer
 	class IObjectSchema;
 	class IOutlinerColumn;
 	class IOutlinerIndicatorBuilder;
+
+	namespace NavigationTool
+	{
+		class FNavigationToolExtender;
+	}
 } // namespace UE::Sequencer
 
 enum class ECurveEditorTreeFilterType : uint32;
@@ -160,6 +165,9 @@ struct FSequencerHostCapabilities
 	/** Do we support selectability options in the viewport */
 	bool bSupportsViewportSelectability;
 
+	/** Do we support the Navigation Tool */
+	bool bSupportsNavigationTool;
+
 	FSequencerHostCapabilities()
 		: bSupportsSaveMovieSceneAsset(false)
 		, bSupportsCurveEditor(false)
@@ -168,6 +176,7 @@ struct FSequencerHostCapabilities
 		, bSupportsAddFromContentBrowser(false)
 		, bSupportsSidebar(false)
 		, bSupportsViewportSelectability(false)
+		, bSupportsNavigationTool(false)
 	{}
 };
 
@@ -526,6 +535,9 @@ public:
 	 */
 	virtual TArray<FString> GetMovieRendererNames() = 0;
 
+	/** Gets the extender for managing Navigation Tool customizations */
+	virtual UE::Sequencer::NavigationTool::FNavigationToolExtender& GetNavigationToolExtender() = 0;
+
 private:
 
 	/** Map of sequencer interfaces for movie scene channel types, keyed on channel UStruct name */
diff --git a/Engine/Source/Editor/Sequencer/Public/SequencerCommands.h b/Engine/Source/Editor/Sequencer/Public/SequencerCommands.h
--- a/Engine/Source/Editor/Sequencer/Public/SequencerCommands.h
+++ b/Engine/Source/Editor/Sequencer/Public/SequencerCommands.h
@@ -481,6 +481,9 @@ public:
 	/** Aligns layer bars, sections, or key selections to the current playhead location */
 	TSharedPtr< FUICommandInfo > AlignSelectionToPlayhead;
 
+	/** Toggles whether or not the Navigation Tool should be visible. */
+	TSharedPtr< FUICommandInfo > ToggleNavigationToolVisible;
+
 	/**
 	 * Initialize commands
 	 */
diff --git a/Engine/Source/Editor/Sequencer/Public/SequencerSettings.h b/Engine/Source/Editor/Sequencer/Public/SequencerSettings.h
--- a/Engine/Source/Editor/Sequencer/Public/SequencerSettings.h
+++ b/Engine/Source/Editor/Sequencer/Public/SequencerSettings.h
@@ -548,6 +548,11 @@ public:
 	/** Sets how the asset thumbnail is captured when the sequence is saved.  */
 	void SetThumbnailCaptureSettings(const FSequencerThumbnailCaptureSettings& InNewValue);
 
+	/** @return True if the Navigation Tool is visible */
+	bool IsNavigationToolVisible() const { return bNavigationToolVisible; }
+	/** Sets the visibility of the Navigation Tool */
+	void SetNavigationToolVisible(const bool bInVisible);
+
 protected:
 
 	/** The auto change mode (auto-key, auto-track or none). */
@@ -875,6 +880,10 @@ protected:
 	UPROPERTY(config, EditAnywhere, Category = General, meta = (EditCondition="ShouldShowThumbnailCaptureSettings()", EditConditionHides))
 	FSequencerThumbnailCaptureSettings ThumbnailCaptureSettings;
 
+	/** Last saved visibility of the Navigation Tool to restore after closed */
+	UPROPERTY(config, EditAnywhere, Category = General)
+	bool bNavigationToolVisible;
+
 	UFUNCTION()
 	static bool ShouldShowThumbnailCaptureSettings();
 
diff --git a/Engine/Source/Editor/Sequencer/Public/Tools/NavigationTool/Columns/INavigationToolColumn.h b/Engine/Source/Editor/Sequencer/Public/Tools/NavigationTool/Columns/INavigationToolColumn.h
new file mode 100644
--- /dev/null
+++ b/Engine/Source/Editor/Sequencer/Public/Tools/NavigationTool/Columns/INavigationToolColumn.h
@@ -0,0 +1,46 @@
+// Copyright Epic Games, Inc. All Rights Reserved.
+
+#pragma once
+
+#include "Tools/NavigationTool/NavigationToolDefines.h"
+#include "Tools/NavigationTool/NavigationToolItemType.h"
+#include "Widgets/Views/SHeaderRow.h"
+
+namespace UE::Sequencer::NavigationTool
+{
+
+class FNavigationToolProvider;
+class INavigationToolView;
+class SNavigationToolTreeRow;
+
+class INavigationToolColumn : public INavigationToolItemTypeCastable, public TSharedFromThis<INavigationToolColumn>
+{
+public:
+	FName GetColumnId() const { return GetTypeId().ToName(); }
+
+	virtual TSharedPtr<FNavigationToolProvider> GetProvider() const = 0;
+
+	virtual FText GetColumnDisplayNameText() const = 0;
+
+	virtual const FSlateBrush* GetIconBrush() const { return nullptr; }
+
+	virtual SHeaderRow::FColumn::FArguments ConstructHeaderRowColumn(const TSharedRef<INavigationToolView>& InToolView, const float InFillSize) = 0;
+
+	virtual float GetFillWidth() const { return 0.f; }
+
+	/*
+	 * Determines whether the Column should be Showing by Default while still be able to toggle it on/off.
+	 * Used when calling SHeaderRow::SetShowGeneratedColumn (requires ShouldGenerateWidget to not be set).
+	 */
+	virtual bool ShouldShowColumnByDefault() const { return false; }
+
+	virtual bool CanHideColumn(const FName InColumnId) const { return true; }
+
+	virtual TSharedRef<SWidget> ConstructRowWidget(const FNavigationToolItemPtr& InItem
+		, const TSharedRef<INavigationToolView>& InToolView
+		, const TSharedRef<SNavigationToolTreeRow>& InRow) = 0;
+
+	virtual void Tick(const float InDeltaTime) {}
+};
+
+}
diff --git a/Engine/Source/Editor/Sequencer/Public/Tools/NavigationTool/Columns/NavigationToolColorColumn.h b/Engine/Source/Editor/Sequencer/Public/Tools/NavigationTool/Columns/NavigationToolColorColumn.h
new file mode 100644
--- /dev/null
+++ b/Engine/Source/Editor/Sequencer/Public/Tools/NavigationTool/Columns/NavigationToolColorColumn.h
@@ -0,0 +1,28 @@
+// Copyright Epic Games, Inc. All Rights Reserved.
+
+#pragma once
+
+#include "Tools/NavigationTool/Columns/NavigationToolColumn.h"
+
+namespace UE::Sequencer::NavigationTool
+{
+
+class INavigationToolView;
+
+class FNavigationToolColorColumn : public FNavigationToolColumn
+{
+public:
+	UE_NAVIGATIONTOOL_INHERITS(FNavigationToolColorColumn, FNavigationToolColumn);
+
+protected:
+	//~ Begin INavigationToolColumn
+	SEQUENCER_API virtual FText GetColumnDisplayNameText() const override;
+	virtual bool ShouldShowColumnByDefault() const override { return true; }
+	SEQUENCER_API virtual SHeaderRow::FColumn::FArguments ConstructHeaderRowColumn(const TSharedRef<INavigationToolView>& InToolView, const float InFillSize) override;
+	SEQUENCER_API virtual TSharedRef<SWidget> ConstructRowWidget(const FNavigationToolItemPtr& InItem
+		, const TSharedRef<INavigationToolView>& InToolView
+		, const TSharedRef<SNavigationToolTreeRow>& InRow) override;
+	//~ End INavigationToolColumn
+};
+
+}
diff --git a/Engine/Source/Editor/Sequencer/Public/Tools/NavigationTool/Columns/NavigationToolColumn.h b/Engine/Source/Editor/Sequencer/Public/Tools/NavigationTool/Columns/NavigationToolColumn.h
new file mode 100644
--- /dev/null
+++ b/Engine/Source/Editor/Sequencer/Public/Tools/NavigationTool/Columns/NavigationToolColumn.h
@@ -0,0 +1,26 @@
+﻿// Copyright Epic Games, Inc. All Rights Reserved.
+
+#pragma once
+
+#include "Tools/NavigationTool/Columns/INavigationToolColumn.h"
+
+namespace UE::Sequencer::NavigationTool
+{
+
+/** Base Navigation Tool Column to extend from */
+class FNavigationToolColumn : public INavigationToolColumn
+{
+public:
+	//~ Begin INavigationToolColumn
+	SEQUENCER_API virtual TSharedPtr<FNavigationToolProvider> GetProvider() const final override;
+	//~ End INavigationToolColumn
+
+protected:
+	friend class FNavigationToolView;
+
+	void SetProvider(const TWeakPtr<FNavigationToolProvider>& InWeakProvider);
+
+	TWeakPtr<FNavigationToolProvider> WeakProvider;
+};
+
+}
diff --git a/Engine/Source/Editor/Sequencer/Public/Tools/NavigationTool/Columns/NavigationToolCommentColumn.h b/Engine/Source/Editor/Sequencer/Public/Tools/NavigationTool/Columns/NavigationToolCommentColumn.h
new file mode 100644
--- /dev/null
+++ b/Engine/Source/Editor/Sequencer/Public/Tools/NavigationTool/Columns/NavigationToolCommentColumn.h
@@ -0,0 +1,30 @@
+// Copyright Epic Games, Inc. All Rights Reserved.
+
+#pragma once
+
+#include "Tools/NavigationTool/Columns/NavigationToolColumn.h"
+
+namespace UE::Sequencer::NavigationTool
+{
+
+class INavigationToolView;
+
+class FNavigationToolCommentColumn : public FNavigationToolColumn
+{
+public:
+	UE_NAVIGATIONTOOL_INHERITS(FNavigationToolCommentColumn, FNavigationToolColumn);
+
+protected:
+	//~ Begin INavigationToolColumn
+	SEQUENCER_API virtual FText GetColumnDisplayNameText() const override;
+	SEQUENCER_API virtual const FSlateBrush* GetIconBrush() const override;
+	virtual bool ShouldShowColumnByDefault() const override { return false; }
+	virtual float GetFillWidth() const override { return 5.f; }
+	SEQUENCER_API virtual SHeaderRow::FColumn::FArguments ConstructHeaderRowColumn(const TSharedRef<INavigationToolView>& InToolView, const float InFillSize) override;
+	SEQUENCER_API virtual TSharedRef<SWidget> ConstructRowWidget(const FNavigationToolItemPtr& InItem
+		, const TSharedRef<INavigationToolView>& InToolView
+		, const TSharedRef<SNavigationToolTreeRow>& InRow) override;
+	//~ End INavigationToolColumn
+};
+
+}
diff --git a/Engine/Source/Editor/Sequencer/Public/Tools/NavigationTool/Columns/NavigationToolDeactiveStateColumn.h b/Engine/Source/Editor/Sequencer/Public/Tools/NavigationTool/Columns/NavigationToolDeactiveStateColumn.h
new file mode 100644
--- /dev/null
+++ b/Engine/Source/Editor/Sequencer/Public/Tools/NavigationTool/Columns/NavigationToolDeactiveStateColumn.h
@@ -0,0 +1,29 @@
+// Copyright Epic Games, Inc. All Rights Reserved.
+
+#pragma once
+
+#include "Tools/NavigationTool/Columns/NavigationToolColumn.h"
+
+namespace UE::Sequencer::NavigationTool
+{
+
+class INavigationToolView;
+
+class FNavigationToolDeactiveStateColumn : public FNavigationToolColumn
+{
+public:
+	UE_NAVIGATIONTOOL_INHERITS(FNavigationToolDeactiveStateColumn, FNavigationToolColumn);
+
+	//~ Begin INavigationToolColumn
+	SEQUENCER_API virtual FText GetColumnDisplayNameText() const override;
+	SEQUENCER_API virtual const FSlateBrush* GetIconBrush() const override;
+	SEQUENCER_API virtual SHeaderRow::FColumn::FArguments ConstructHeaderRowColumn(const TSharedRef<INavigationToolView>& InToolView, const float InFillSize) override;
+	virtual bool ShouldShowColumnByDefault() const override { return true; }
+	SEQUENCER_API virtual TSharedRef<SWidget> ConstructRowWidget(const FNavigationToolItemPtr& InItem
+		, const TSharedRef<INavigationToolView>& InToolView
+		, const TSharedRef<SNavigationToolTreeRow>& InRow) override;
+	//~ End INavigationToolColumn
+};
+
+}
+
diff --git a/Engine/Source/Editor/Sequencer/Public/Tools/NavigationTool/Columns/NavigationToolHBiasColumn.h b/Engine/Source/Editor/Sequencer/Public/Tools/NavigationTool/Columns/NavigationToolHBiasColumn.h
new file mode 100644
--- /dev/null
+++ b/Engine/Source/Editor/Sequencer/Public/Tools/NavigationTool/Columns/NavigationToolHBiasColumn.h
@@ -0,0 +1,30 @@
+// Copyright Epic Games, Inc. All Rights Reserved.
+
+#pragma once
+
+#include "Tools/NavigationTool/Columns/NavigationToolColumn.h"
+
+namespace UE::Sequencer::NavigationTool
+{
+
+class INavigationToolView;
+
+class FNavigationToolHBiasColumn : public FNavigationToolColumn
+{
+public:
+	UE_NAVIGATIONTOOL_INHERITS(FNavigationToolHBiasColumn, FNavigationToolColumn);
+
+protected:
+	//~ Begin INavigationToolColumn
+	SEQUENCER_API virtual FText GetColumnDisplayNameText() const override;
+	SEQUENCER_API virtual const FSlateBrush* GetIconBrush() const override;
+	virtual bool ShouldShowColumnByDefault() const override { return false; }
+	virtual float GetFillWidth() const override { return 5.f; }
+	SEQUENCER_API virtual SHeaderRow::FColumn::FArguments ConstructHeaderRowColumn(const TSharedRef<INavigationToolView>& InToolView, const float InFillSize) override;
+	SEQUENCER_API virtual TSharedRef<SWidget> ConstructRowWidget(const FNavigationToolItemPtr& InItem
+		, const TSharedRef<INavigationToolView>& InToolView
+		, const TSharedRef<SNavigationToolTreeRow>& InRow) override;
+	//~ End INavigationToolColumn
+};
+
+}
diff --git a/Engine/Source/Editor/Sequencer/Public/Tools/NavigationTool/Columns/NavigationToolInTimeColumn.h b/Engine/Source/Editor/Sequencer/Public/Tools/NavigationTool/Columns/NavigationToolInTimeColumn.h
new file mode 100644
--- /dev/null
+++ b/Engine/Source/Editor/Sequencer/Public/Tools/NavigationTool/Columns/NavigationToolInTimeColumn.h
@@ -0,0 +1,30 @@
+// Copyright Epic Games, Inc. All Rights Reserved.
+
+#pragma once
+
+#include "Tools/NavigationTool/Columns/NavigationToolColumn.h"
+
+namespace UE::Sequencer::NavigationTool
+{
+
+class INavigationToolView;
+
+class FNavigationToolInTimeColumn : public FNavigationToolColumn
+{
+public:
+	UE_NAVIGATIONTOOL_INHERITS(FNavigationToolInTimeColumn, FNavigationToolColumn);
+
+protected:
+	//~ Begin INavigationToolColumn
+	SEQUENCER_API virtual FText GetColumnDisplayNameText() const override;
+	SEQUENCER_API virtual const FSlateBrush* GetIconBrush() const override;
+	virtual bool ShouldShowColumnByDefault() const override { return true; }
+	virtual float GetFillWidth() const override { return 5.f; }
+	SEQUENCER_API virtual SHeaderRow::FColumn::FArguments ConstructHeaderRowColumn(const TSharedRef<INavigationToolView>& InToolView, const float InFillSize) override;
+	SEQUENCER_API virtual TSharedRef<SWidget> ConstructRowWidget(const FNavigationToolItemPtr& InItem
+		, const TSharedRef<INavigationToolView>& InToolView
+		, const TSharedRef<SNavigationToolTreeRow>& InRow) override;
+	//~ End INavigationToolColumn
+};
+
+}
diff --git a/Engine/Source/Editor/Sequencer/Public/Tools/NavigationTool/Columns/NavigationToolItemsColumn.h b/Engine/Source/Editor/Sequencer/Public/Tools/NavigationTool/Columns/NavigationToolItemsColumn.h
new file mode 100644
--- /dev/null
+++ b/Engine/Source/Editor/Sequencer/Public/Tools/NavigationTool/Columns/NavigationToolItemsColumn.h
@@ -0,0 +1,30 @@
+// Copyright Epic Games, Inc. All Rights Reserved.
+
+#pragma once
+
+#include "Tools/NavigationTool/Columns/NavigationToolColumn.h"
+
+namespace UE::Sequencer::NavigationTool
+{
+
+class INavigationToolView;
+
+class FNavigationToolItemsColumn : public FNavigationToolColumn
+{
+public:
+	UE_NAVIGATIONTOOL_INHERITS(FNavigationToolItemsColumn, FNavigationToolColumn);
+
+protected:
+	//~ Begin INavigationToolColumn
+	SEQUENCER_API virtual FText GetColumnDisplayNameText() const override;
+	SEQUENCER_API virtual const FSlateBrush* GetIconBrush() const override;
+	virtual bool ShouldShowColumnByDefault() const override { return true; }
+	virtual float GetFillWidth() const override { return 3.f; }
+	SEQUENCER_API virtual SHeaderRow::FColumn::FArguments ConstructHeaderRowColumn(const TSharedRef<INavigationToolView>& InToolView, const float InFillSize) override;
+	SEQUENCER_API virtual TSharedRef<SWidget> ConstructRowWidget(const FNavigationToolItemPtr& InItem
+		, const TSharedRef<INavigationToolView>& InToolView
+		, const TSharedRef<SNavigationToolTreeRow>& InRow) override;
+	//~ End INavigationToolColumn
+};
+
+}
diff --git a/Engine/Source/Editor/Sequencer/Public/Tools/NavigationTool/Columns/NavigationToolLabelColumn.h b/Engine/Source/Editor/Sequencer/Public/Tools/NavigationTool/Columns/NavigationToolLabelColumn.h
new file mode 100644
--- /dev/null
+++ b/Engine/Source/Editor/Sequencer/Public/Tools/NavigationTool/Columns/NavigationToolLabelColumn.h
@@ -0,0 +1,30 @@
+// Copyright Epic Games, Inc. All Rights Reserved.
+
+#pragma once
+
+#include "Tools/NavigationTool/Columns/NavigationToolColumn.h"
+
+namespace UE::Sequencer::NavigationTool
+{
+
+class INavigationToolView;
+
+class FNavigationToolLabelColumn : public FNavigationToolColumn
+{
+public:
+	UE_NAVIGATIONTOOL_INHERITS(FNavigationToolLabelColumn, FNavigationToolColumn);
+
+protected:
+	//~ Begin INavigationToolColumn
+	SEQUENCER_API virtual FText GetColumnDisplayNameText() const override;
+	SEQUENCER_API virtual const FSlateBrush* GetIconBrush() const override;
+	SEQUENCER_API virtual SHeaderRow::FColumn::FArguments ConstructHeaderRowColumn(const TSharedRef<INavigationToolView>& InToolView, const float InFillSize) override;
+	virtual float GetFillWidth() const override { return 10.f; }
+	virtual bool CanHideColumn(const FName InColumnId) const override { return false; }
+	SEQUENCER_API virtual TSharedRef<SWidget> ConstructRowWidget(const FNavigationToolItemPtr& InItem
+		, const TSharedRef<INavigationToolView>& InToolView
+		, const TSharedRef<SNavigationToolTreeRow>& InRow) override;
+	//~ End INavigationToolColumn
+};
+
+}
diff --git a/Engine/Source/Editor/Sequencer/Public/Tools/NavigationTool/Columns/NavigationToolLengthColumn.h b/Engine/Source/Editor/Sequencer/Public/Tools/NavigationTool/Columns/NavigationToolLengthColumn.h
new file mode 100644
--- /dev/null
+++ b/Engine/Source/Editor/Sequencer/Public/Tools/NavigationTool/Columns/NavigationToolLengthColumn.h
@@ -0,0 +1,30 @@
+// Copyright Epic Games, Inc. All Rights Reserved.
+
+#pragma once
+
+#include "Tools/NavigationTool/Columns/NavigationToolColumn.h"
+
+namespace UE::Sequencer::NavigationTool
+{
+
+class INavigationToolView;
+
+class FNavigationToolLengthColumn : public FNavigationToolColumn
+{
+public:
+	UE_NAVIGATIONTOOL_INHERITS(FNavigationToolLengthColumn, FNavigationToolColumn);
+
+protected:
+	//~ Begin INavigationToolColumn
+	SEQUENCER_API virtual FText GetColumnDisplayNameText() const override;
+	SEQUENCER_API virtual const FSlateBrush* GetIconBrush() const override;
+	virtual bool ShouldShowColumnByDefault() const override { return false; }
+	virtual float GetFillWidth() const override { return 5.f; }
+	SEQUENCER_API virtual SHeaderRow::FColumn::FArguments ConstructHeaderRowColumn(const TSharedRef<INavigationToolView>& InToolView, const float InFillSize) override;
+	SEQUENCER_API virtual TSharedRef<SWidget> ConstructRowWidget(const FNavigationToolItemPtr& InItem
+		, const TSharedRef<INavigationToolView>& InToolView
+		, const TSharedRef<SNavigationToolTreeRow>& InRow) override;
+	//~ End INavigationToolColumn
+};
+
+}
diff --git a/Engine/Source/Editor/Sequencer/Public/Tools/NavigationTool/Columns/NavigationToolLockColumn.h b/Engine/Source/Editor/Sequencer/Public/Tools/NavigationTool/Columns/NavigationToolLockColumn.h
new file mode 100644
--- /dev/null
+++ b/Engine/Source/Editor/Sequencer/Public/Tools/NavigationTool/Columns/NavigationToolLockColumn.h
@@ -0,0 +1,29 @@
+// Copyright Epic Games, Inc. All Rights Reserved.
+
+#pragma once
+
+#include "Tools/NavigationTool/Columns/NavigationToolColumn.h"
+
+namespace UE::Sequencer::NavigationTool
+{
+
+class INavigationToolView;
+
+class FNavigationToolLockColumn : public FNavigationToolColumn
+{
+public:
+	UE_NAVIGATIONTOOL_INHERITS(FNavigationToolLockColumn, FNavigationToolColumn);
+
+protected:
+	//~ Begin INavigationToolColumn
+	SEQUENCER_API virtual FText GetColumnDisplayNameText() const override;
+	SEQUENCER_API virtual const FSlateBrush* GetIconBrush() const override;
+	virtual bool ShouldShowColumnByDefault() const override { return true; }
+	SEQUENCER_API virtual SHeaderRow::FColumn::FArguments ConstructHeaderRowColumn(const TSharedRef<INavigationToolView>& InToolView, const float InFillSize) override;
+	SEQUENCER_API virtual TSharedRef<SWidget> ConstructRowWidget(const FNavigationToolItemPtr& InItem
+		, const TSharedRef<INavigationToolView>& InToolView
+		, const TSharedRef<SNavigationToolTreeRow>& InRow) override;
+	//~ End INavigationToolColumn
+};
+
+}
diff --git a/Engine/Source/Editor/Sequencer/Public/Tools/NavigationTool/Columns/NavigationToolMarkerVisibilityColumn.h b/Engine/Source/Editor/Sequencer/Public/Tools/NavigationTool/Columns/NavigationToolMarkerVisibilityColumn.h
new file mode 100644
--- /dev/null
+++ b/Engine/Source/Editor/Sequencer/Public/Tools/NavigationTool/Columns/NavigationToolMarkerVisibilityColumn.h
@@ -0,0 +1,28 @@
+// Copyright Epic Games, Inc. All Rights Reserved.
+
+#pragma once
+
+#include "Tools/NavigationTool/Columns/NavigationToolColumn.h"
+
+namespace UE::Sequencer::NavigationTool
+{
+
+class INavigationToolView;
+
+class FNavigationToolMarkerVisibilityColumn : public FNavigationToolColumn
+{
+public:
+	UE_NAVIGATIONTOOL_INHERITS(FNavigationToolMarkerVisibilityColumn, FNavigationToolColumn);
+
+	//~ Begin INavigationToolColumn
+	SEQUENCER_API virtual FText GetColumnDisplayNameText() const override;
+	SEQUENCER_API virtual const FSlateBrush* GetIconBrush() const override;
+	SEQUENCER_API virtual SHeaderRow::FColumn::FArguments ConstructHeaderRowColumn(const TSharedRef<INavigationToolView>& InToolView, const float InFillSize) override;
+	virtual bool ShouldShowColumnByDefault() const override { return true; }
+	SEQUENCER_API virtual TSharedRef<SWidget> ConstructRowWidget(const FNavigationToolItemPtr& InItem
+		, const TSharedRef<INavigationToolView>& InToolView
+		, const TSharedRef<SNavigationToolTreeRow>& InRow) override;
+	//~ End INavigationToolColumn
+};
+
+}
diff --git a/Engine/Source/Editor/Sequencer/Public/Tools/NavigationTool/Columns/NavigationToolOutTimeColumn.h b/Engine/Source/Editor/Sequencer/Public/Tools/NavigationTool/Columns/NavigationToolOutTimeColumn.h
new file mode 100644
--- /dev/null
+++ b/Engine/Source/Editor/Sequencer/Public/Tools/NavigationTool/Columns/NavigationToolOutTimeColumn.h
@@ -0,0 +1,30 @@
+// Copyright Epic Games, Inc. All Rights Reserved.
+
+#pragma once
+
+#include "Tools/NavigationTool/Columns/NavigationToolColumn.h"
+
+namespace UE::Sequencer::NavigationTool
+{
+
+class INavigationToolView;
+
+class FNavigationToolOutTimeColumn : public FNavigationToolColumn
+{
+public:
+	UE_NAVIGATIONTOOL_INHERITS(FNavigationToolOutTimeColumn, FNavigationToolColumn);
+
+protected:
+	//~ Begin INavigationToolColumn
+	SEQUENCER_API virtual FText GetColumnDisplayNameText() const override;
+	SEQUENCER_API virtual const FSlateBrush* GetIconBrush() const override;
+	virtual bool ShouldShowColumnByDefault() const override { return true; }
+	virtual float GetFillWidth() const override { return 5.f; }
+	SEQUENCER_API virtual SHeaderRow::FColumn::FArguments ConstructHeaderRowColumn(const TSharedRef<INavigationToolView>& InToolView, const float InFillSize) override;
+	SEQUENCER_API virtual TSharedRef<SWidget> ConstructRowWidget(const FNavigationToolItemPtr& InItem
+		, const TSharedRef<INavigationToolView>& InToolView
+		, const TSharedRef<SNavigationToolTreeRow>& InRow) override;
+	//~ End INavigationToolColumn
+};
+
+}
diff --git a/Engine/Source/Editor/Sequencer/Public/Tools/NavigationTool/Columns/NavigationToolRevisionControlColumn.h b/Engine/Source/Editor/Sequencer/Public/Tools/NavigationTool/Columns/NavigationToolRevisionControlColumn.h
new file mode 100644
--- /dev/null
+++ b/Engine/Source/Editor/Sequencer/Public/Tools/NavigationTool/Columns/NavigationToolRevisionControlColumn.h
@@ -0,0 +1,28 @@
+// Copyright Epic Games, Inc. All Rights Reserved.
+
+#pragma once
+
+#include "Tools/NavigationTool/Columns/NavigationToolColumn.h"
+
+namespace UE::Sequencer::NavigationTool
+{
+
+class INavigationToolView;
+
+class FNavigationToolRevisionControlColumn : public FNavigationToolColumn
+{
+public:
+	UE_NAVIGATIONTOOL_INHERITS(FNavigationToolRevisionControlColumn, FNavigationToolColumn);
+
+	//~ Begin INavigationToolColumn
+	SEQUENCER_API virtual FText GetColumnDisplayNameText() const override;
+	SEQUENCER_API virtual const FSlateBrush* GetIconBrush() const override;
+	SEQUENCER_API virtual SHeaderRow::FColumn::FArguments ConstructHeaderRowColumn(const TSharedRef<INavigationToolView>& InToolView, const float InFillSize) override;
+	virtual bool ShouldShowColumnByDefault() const override { return true; }
+	SEQUENCER_API virtual TSharedRef<SWidget> ConstructRowWidget(const FNavigationToolItemPtr& InItem
+		, const TSharedRef<INavigationToolView>& InToolView
+		, const TSharedRef<SNavigationToolTreeRow>& InRow) override;
+	//~ End INavigationToolColumn
+};
+
+}
diff --git a/Engine/Source/Editor/Sequencer/Public/Tools/NavigationTool/Columns/NavigationToolStartFrameOffsetColumn.h b/Engine/Source/Editor/Sequencer/Public/Tools/NavigationTool/Columns/NavigationToolStartFrameOffsetColumn.h
new file mode 100644
--- /dev/null
+++ b/Engine/Source/Editor/Sequencer/Public/Tools/NavigationTool/Columns/NavigationToolStartFrameOffsetColumn.h
@@ -0,0 +1,30 @@
+// Copyright Epic Games, Inc. All Rights Reserved.
+
+#pragma once
+
+#include "Tools/NavigationTool/Columns/NavigationToolColumn.h"
+
+namespace UE::Sequencer::NavigationTool
+{
+
+class INavigationToolView;
+
+class FNavigationToolStartFrameOffsetColumn : public FNavigationToolColumn
+{
+public:
+	UE_NAVIGATIONTOOL_INHERITS(FNavigationToolStartFrameOffsetColumn, FNavigationToolColumn);
+
+protected:
+	//~ Begin INavigationToolColumn
+	SEQUENCER_API virtual FText GetColumnDisplayNameText() const override;
+	SEQUENCER_API virtual const FSlateBrush* GetIconBrush() const override;
+	virtual bool ShouldShowColumnByDefault() const override { return false; }
+	virtual float GetFillWidth() const override { return 5.f; }
+	SEQUENCER_API virtual SHeaderRow::FColumn::FArguments ConstructHeaderRowColumn(const TSharedRef<INavigationToolView>& InToolView, const float InFillSize) override;
+	SEQUENCER_API virtual TSharedRef<SWidget> ConstructRowWidget(const FNavigationToolItemPtr& InItem
+		, const TSharedRef<INavigationToolView>& InToolView
+		, const TSharedRef<SNavigationToolTreeRow>& InRow) override;
+	//~ End INavigationToolColumn
+};
+
+}
diff --git a/Engine/Source/Editor/Sequencer/Public/Tools/NavigationTool/Columns/NavigationToolTakeColumn.h b/Engine/Source/Editor/Sequencer/Public/Tools/NavigationTool/Columns/NavigationToolTakeColumn.h
new file mode 100644
--- /dev/null
+++ b/Engine/Source/Editor/Sequencer/Public/Tools/NavigationTool/Columns/NavigationToolTakeColumn.h
@@ -0,0 +1,30 @@
+// Copyright Epic Games, Inc. All Rights Reserved.
+
+#pragma once
+
+#include "Tools/NavigationTool/Columns/NavigationToolColumn.h"
+
+namespace UE::Sequencer::NavigationTool
+{
+
+class INavigationToolView;
+
+class FNavigationToolTakeColumn : public FNavigationToolColumn
+{
+public:
+	UE_NAVIGATIONTOOL_INHERITS(FNavigationToolTakeColumn, FNavigationToolColumn);
+
+protected:
+	//~ Begin INavigationToolColumn
+	SEQUENCER_API virtual FText GetColumnDisplayNameText() const override;
+	SEQUENCER_API virtual const FSlateBrush* GetIconBrush() const override;
+	virtual bool ShouldShowColumnByDefault() const override { return false; }
+	virtual float GetFillWidth() const override { return 5.f; }
+	SEQUENCER_API virtual SHeaderRow::FColumn::FArguments ConstructHeaderRowColumn(const TSharedRef<INavigationToolView>& InToolView, const float InFillSize) override;
+	SEQUENCER_API virtual TSharedRef<SWidget> ConstructRowWidget(const FNavigationToolItemPtr& InItem
+		, const TSharedRef<INavigationToolView>& InToolView
+		, const TSharedRef<SNavigationToolTreeRow>& InRow) override;
+	//~ End INavigationToolColumn
+};
+
+}
diff --git a/Engine/Source/Editor/Sequencer/Public/Tools/NavigationTool/Customization/INavigationToolIconCustomization.h b/Engine/Source/Editor/Sequencer/Public/Tools/NavigationTool/Customization/INavigationToolIconCustomization.h
new file mode 100644
--- /dev/null
+++ b/Engine/Source/Editor/Sequencer/Public/Tools/NavigationTool/Customization/INavigationToolIconCustomization.h
@@ -0,0 +1,27 @@
+// Copyright Epic Games, Inc. All Rights Reserved.
+
+#pragma once
+
+#include "Templates/SharedPointer.h"
+#include "UObject/NameTypes.h"
+
+struct FSlateIcon;
+
+namespace UE::Sequencer::NavigationTool
+{
+
+class FNavigationToolItem;
+
+class INavigationToolIconCustomization
+{
+public:
+	virtual ~INavigationToolIconCustomization() = default;
+
+	virtual FName GetToolItemIdentifier() const = 0;
+
+	virtual bool HasOverrideIcon(TSharedPtr<const FNavigationToolItem> InToolItem) const = 0;
+
+	virtual FSlateIcon GetOverrideIcon(TSharedPtr<const FNavigationToolItem> InToolItem) const = 0;
+};
+
+}
diff --git a/Engine/Source/Editor/Sequencer/Public/Tools/NavigationTool/Customization/NavigationToolExtender.h b/Engine/Source/Editor/Sequencer/Public/Tools/NavigationTool/Customization/NavigationToolExtender.h
new file mode 100644
--- /dev/null
+++ b/Engine/Source/Editor/Sequencer/Public/Tools/NavigationTool/Customization/NavigationToolExtender.h
@@ -0,0 +1,210 @@
+// Copyright Epic Games, Inc. All Rights Reserved.
+
+#pragma once
+
+#include "Tools/NavigationTool/DragDropOps/NavigationToolItemDragDropOp.h"
+#include "Tools/NavigationTool/ItemProxies/NavigationToolItemProxyRegistry.h"
+#include "Tools/NavigationTool/Providers/NavigationToolProvider.h"
+
+namespace UE::Sequencer::NavigationTool
+{
+
+class INavigationTool;
+class INavigationToolIconCustomization;
+
+enum class ENavigationToolProvidersChangeType
+{
+	Add,
+	Remove
+};
+
+DECLARE_MULTICAST_DELEGATE_ThreeParams(FOnNavigationToolProvidersChanged
+	, const TSharedRef<INavigationTool>& /*InTool*/
+	, const TSharedRef<FNavigationToolProvider>& /*InProvider*/
+	, const ENavigationToolProvidersChangeType /*InChangeType*/);
+
+/** Called when the FNavigationToolItemDragDropOp has been created and Initialized in FNavigationToolItemDragDropOp::Init */
+DECLARE_MULTICAST_DELEGATE_OneParam(FOnItemDragDropOpInitialized, FNavigationToolItemDragDropOp&);
+
+/** Singleton class for extending the Sequencer Navigation Tool */
+class FNavigationToolExtender
+{
+public:
+	FNavigationToolExtender();
+	virtual ~FNavigationToolExtender();
+
+	/**
+	 * Retrieves the unique name for the provided Navigation Tool instance.
+	 * The Navigation Tool instance name is based on its associated sequencer settings
+	 * object name. This provides a unique instance name for each Sequencer type.
+	 * @param InTool The Navigation Tool instance for which the name is being retrieved.
+	 * @return The name of the Navigation Tool instance if available; otherwise, NAME_None.
+	 */
+	SEQUENCER_API static FName GetToolInstanceName(const INavigationTool& InTool);
+
+	/**
+	 * Registers a Navigation Tool provider for a specified Navigation Tool instance.
+	 * This method ensures that the provider is associated with the given tool and replaces any
+	 * existing provider with the same identifier, if one exists.
+	 * @param InTool The Navigation Tool instance for which the provider will be registered.
+	 * @param InProvider The Navigation Tool provider being registered.
+	 * @return True if the provider was successfully registered, false otherwise.
+	 */
+	SEQUENCER_API static bool RegisterToolProvider(const TSharedRef<INavigationTool>& InTool
+		, const TSharedRef<FNavigationToolProvider>& InProvider);
+
+	/**
+	 * Registers a Navigation Tool provider for a specified Navigation Tool instance.
+	 * Unregisters all other providers before registering the given provider.
+	 * @param InTool The Navigation Tool instance for which the provider will be registered.
+	 * @param InProvider The Navigation Tool provider being registered.
+	 * @return True if the provider was successfully registered, false otherwise.
+	 */
+	SEQUENCER_API static bool RegisterOnlyToolProvider(const TSharedRef<INavigationTool>& InTool
+		, const TSharedRef<FNavigationToolProvider>& InProvider);
+
+	/**
+	 * Unregisters a Navigation Tool provider associated with a specific instance and identifier.
+	 * @param InTool A shared reference to the Navigation Tool from which the provider needs to be unregistered.
+	 * @param InProviderId The identifier of the provider to be unregistered.
+	 * @return True if the provider was successfully unregistered; false otherwise.
+	 */
+	SEQUENCER_API static bool UnregisterToolProvider(const TSharedRef<INavigationTool>& InTool, const FName InProviderId);
+
+
+	/**
+	 * Unregisters all Navigation Tool providers associated with a given instance.
+	 * @param InTool The Navigation Tool for which all associated providers should be unregistered.
+	 */
+	SEQUENCER_API static void UnregisterAllToolProviders(const TSharedRef<INavigationTool>& InTool);
+
+	/**
+	 * Finds the tool providers associated with the given Navigation Tool instance.
+	 * @param InTool The Navigation Tool instance for which to find associated providers.
+	 * @param OutProviders A set that will be populated with the found tool providers as weak pointers.
+	 * @return True if tool providers were found and populated into OutProviders; otherwise, false.
+	 */
+	SEQUENCER_API static bool FindToolProviders(const INavigationTool& InTool
+		, TSet<TWeakPtr<FNavigationToolProvider>>& OutProviders);
+
+	/**
+	 * Iterates over each Navigation Tool provider and executes the predicate function.
+	 * The provided predicate function determines whether the iteration continues or stops.
+	 * @param InPredicate A function that accepts a reference to an INavigationTool and a shared reference
+	 *                    to FNavigationToolProvider. Should return true to continue iteration or false to stop.
+	 */
+	SEQUENCER_API static void ForEachProvider(const TFunction<bool(INavigationTool& InTool
+		, const TSharedRef<FNavigationToolProvider>& InProvider)>& InPredicate);
+
+	/**
+	 * Iterates over each Navigation Tool provider and executes the predicate function associated with the given
+	 * Navigation Tool instance. The iteration is controlled by the provided predicate function, which determines
+	 * whether it continues or stops.
+	 * @param InTool The Navigation Tool instance whose providers are to be iterated over.
+	 * @param InPredicate A predicate function that takes a shared reference to a Navigation Tool provider as input.
+	 *                    The function should return true to continue iteration or false to stop it.
+	 */
+	SEQUENCER_API static void ForEachToolProvider(INavigationTool& InTool
+		, const TFunction<bool(const TSharedRef<FNavigationToolProvider>& InProvider)>& InPredicate);
+
+	/** Event called when a provider has been added or removed from a tool instance. */
+	SEQUENCER_API static FOnNavigationToolProvidersChanged& OnProvidersChanged()
+	{
+		return Get().ProvidersChangedDelegate;
+	}
+
+	/** Event called when the FNavigationToolItemDragDropOp has been created and initialized in FNavigationToolItemDragDropOp::Init. */
+	SEQUENCER_API static FOnItemDragDropOpInitialized& OnItemDragDropOpInitialized()
+	{
+		return Get().ItemDragDropOpInitializedDelegate;
+	}
+
+	SEQUENCER_API static FNavigationToolItemProxyRegistry& GetItemProxyRegistry()
+	{
+		return Get().ItemProxyRegistry;
+	}
+
+	/**
+	 * Get the custom icon for the item if a customization for it exists.
+	 * @param InItem The item to find the icon of
+	 * @return The custom icon for the item if a customization is found, FSlateIcon() otherwise
+	 */
+	SEQUENCER_API static FSlateIcon FindOverrideIcon(TSharedPtr<const FNavigationToolItem> InItem);
+
+protected:
+	struct FNavigationToolProviderEntry
+	{
+		FNavigationToolProviderEntry() = delete;
+		FNavigationToolProviderEntry(const TSharedRef<INavigationTool>& InTool
+			, const TSet<TWeakPtr<FNavigationToolProvider>>& InWeakProviders)
+			: WeakTool(InTool)
+			, WeakProviders(InWeakProviders)
+		{}
+
+		bool operator==(const FNavigationToolProviderEntry& InOther) const
+		{
+			return (WeakTool.IsValid() && InOther.WeakTool.IsValid())
+				&& (WeakTool.Pin() == InOther.WeakTool.Pin());
+		}
+
+		friend uint32 GetTypeHash(const FNavigationToolProviderEntry& InEntry)
+		{
+			return GetTypeHash(InEntry.WeakTool);
+		}
+
+		TWeakPtr<INavigationTool> WeakTool;
+
+		TSet<TWeakPtr<FNavigationToolProvider>> WeakProviders;
+	};
+
+	static FNavigationToolExtender& Get();
+
+	static FNavigationToolProviderEntry* FindToolProvidersEntry(const INavigationTool& InTool);
+
+	static TSharedPtr<FNavigationToolProvider> FindToolProvider(const INavigationTool& InTool, const FName InProviderId);
+
+	void RegisterOverridenIcon_Internal(const FNavigationToolItemTypeId& InItemTypeId, const TSharedRef<INavigationToolIconCustomization>& InNavigationToolIconCustomization);
+	void UnregisterOverriddenIcon_Internal(const FNavigationToolItemTypeId& InItemTypeId, const FName& InSpecializationIdentifier);
+
+	/**
+	 * Get the customization for the given item if any.
+	 * @param InItem Item to search the customization for.
+	 * @return The customization for the item if any; nullptr otherwise.
+	 */
+	TSharedPtr<INavigationToolIconCustomization> GetCustomizationForItem(const TSharedPtr<const FNavigationToolItem>& InItem) const;
+
+	/** Providers from modules that specify their extensions */
+	TSet<FNavigationToolProviderEntry> ToolProviders;
+
+	FOnNavigationToolProvidersChanged ProvidersChangedDelegate;
+
+	FOnItemDragDropOpInitialized ItemDragDropOpInitializedDelegate;
+
+	FNavigationToolItemProxyRegistry ItemProxyRegistry;
+
+	/** Hold the key of the map containing the IconCustomizations */
+	struct FIconCustomizationKey
+	{
+		/** Navigation Tool item class FName (ex. NavigationToolItem) */
+		FNavigationToolItemTypeId ItemTypeId = FNavigationToolItemTypeId::Invalid();
+
+		/** Specialization identifier (ex. AvaShapeActor class FName for ActorIcon customization) */
+		FName CustomizationSpecializationIdentifier;
+
+		bool operator==(const FIconCustomizationKey& InOther) const
+		{
+			return ItemTypeId == InOther.ItemTypeId
+				&& CustomizationSpecializationIdentifier == InOther.CustomizationSpecializationIdentifier;
+		}
+
+		friend uint32 GetTypeHash(FIconCustomizationKey const& InCustomizationKey)
+		{
+			return HashCombine(GetTypeHash(InCustomizationKey.ItemTypeId)
+				, GetTypeHash(InCustomizationKey.CustomizationSpecializationIdentifier));
+		}
+	};
+
+	TMap<FIconCustomizationKey, TSharedPtr<INavigationToolIconCustomization>> IconRegistry;
+};
+
+}
diff --git a/Engine/Source/Editor/Sequencer/Public/Tools/NavigationTool/Customization/NavigationToolIconCustomization.h b/Engine/Source/Editor/Sequencer/Public/Tools/NavigationTool/Customization/NavigationToolIconCustomization.h
new file mode 100644
--- /dev/null
+++ b/Engine/Source/Editor/Sequencer/Public/Tools/NavigationTool/Customization/NavigationToolIconCustomization.h
@@ -0,0 +1,27 @@
+// Copyright Epic Games, Inc. All Rights Reserved.
+
+#pragma once
+
+#include "Templates/SharedPointer.h"
+#include "UObject/NameTypes.h"
+
+struct FSlateIcon;
+
+namespace UE::Sequencer::NavigationTool
+{
+
+class FNavigationToolItem;
+
+class INavigationToolIconCustomization
+{
+public:
+	virtual ~INavigationToolIconCustomization() = default;
+
+	virtual FName GetToolItemIdentifier() const = 0;
+
+	virtual bool HasOverrideIcon(TSharedPtr<const FNavigationToolItem> InToolItem) const = 0;
+
+	virtual FSlateIcon GetOverrideIcon(TSharedPtr<const FNavigationToolItem> InToolItem) const = 0;
+};
+
+}
diff --git a/Engine/Source/Editor/Sequencer/Public/Tools/NavigationTool/Customization/NavigationToolObjectIconCustomization.h b/Engine/Source/Editor/Sequencer/Public/Tools/NavigationTool/Customization/NavigationToolObjectIconCustomization.h
new file mode 100644
--- /dev/null
+++ b/Engine/Source/Editor/Sequencer/Public/Tools/NavigationTool/Customization/NavigationToolObjectIconCustomization.h
@@ -0,0 +1,36 @@
+// Copyright Epic Games, Inc. All Rights Reserved.
+
+#pragma once
+
+#include "NavigationToolIconCustomization.h"
+#include "Delegates/Delegate.h"
+#include "Delegates/DelegateCombinations.h"
+#include "Templates/SharedPointer.h"
+
+class UClass;
+
+namespace UE::Sequencer::NavigationTool
+{
+
+DECLARE_DELEGATE_RetVal_OneParam(FSlateIcon /*OutIconForItem*/, FOnGetOverriddenObjectIcon, TSharedPtr<const FNavigationToolItem> /*InItemToCustomize*/)
+
+class FNavigationToolObjectIconCustomization : public INavigationToolIconCustomization
+{
+public:
+	SEQUENCER_API FNavigationToolObjectIconCustomization(const UClass* InSupportedClass);
+
+	SEQUENCER_API void SetOverriddenIcon(const FOnGetOverriddenObjectIcon& InOverriddenIcon);
+
+private:
+	//~ Begin INavigationToolIconCustomization
+	SEQUENCER_API virtual FName GetToolItemIdentifier() const override;
+	SEQUENCER_API virtual bool HasOverrideIcon(TSharedPtr<const FNavigationToolItem> InToolItem) const override;
+	SEQUENCER_API virtual FSlateIcon GetOverrideIcon(TSharedPtr<const FNavigationToolItem> InToolItem) const override;
+	//~ End INavigationToolIconCustomization
+
+	FName SupportedClassName;
+
+	FOnGetOverriddenObjectIcon OnGetOverriddenIcon;
+};
+
+}
diff --git a/Engine/Source/Editor/Sequencer/Public/Tools/NavigationTool/DragDropOps/Handlers/NavigationToolItemDropHandler.h b/Engine/Source/Editor/Sequencer/Public/Tools/NavigationTool/DragDropOps/Handlers/NavigationToolItemDropHandler.h
new file mode 100644
--- /dev/null
+++ b/Engine/Source/Editor/Sequencer/Public/Tools/NavigationTool/DragDropOps/Handlers/NavigationToolItemDropHandler.h
@@ -0,0 +1,74 @@
+﻿// Copyright Epic Games, Inc. All Rights Reserved.
+
+#pragma once
+
+#include "Containers/Array.h"
+#include "Templates/EnableIf.h"
+#include "Templates/SharedPointer.h"
+#include "Tools/NavigationTool/Items/INavigationToolItem.h"
+
+enum class EItemDropZone;
+template<typename OptionalType> struct TOptional;
+
+namespace UE::Sequencer::NavigationTool
+{
+
+class FNavigationToolItemDragDropOp;
+
+/**
+ * Base Class to Handle Dropping Navigation Tool Items into a Target Navigation Tool Item
+ * @see built-in example FNavigationToolActorDropHandler
+ */
+class FNavigationToolItemDropHandler : public INavigationToolItemTypeCastable, public TSharedFromThis<FNavigationToolItemDropHandler>
+{
+	friend FNavigationToolItemDragDropOp;
+
+	SEQUENCER_API void Initialize(const FNavigationToolItemDragDropOp& InDragDropOp);
+
+public:
+	UE_NAVIGATIONTOOL_INHERITS(FNavigationToolItemDropHandler, INavigationToolItemTypeCastable);
+
+	TConstArrayView<TSharedPtr<INavigationToolItem>> GetItems() const
+	{
+		return Items;
+	}	
+
+protected:
+	virtual bool IsDraggedItemSupported(const FNavigationToolItemPtr& InDraggedItem) const = 0;
+
+	virtual TOptional<EItemDropZone> CanDrop(EItemDropZone InDropZone, const FNavigationToolItemPtr& InTargetItem) const = 0;
+
+	virtual bool Drop(EItemDropZone InDropZone, const FNavigationToolItemPtr& InTargetItem) = 0;
+
+	enum class EIterationResult
+	{
+		Continue,
+		Break,
+	};
+	template<typename InItemType, typename = typename TEnableIf<TIsDerivedFrom<InItemType, INavigationToolItem>::Value>::Type>
+	void ForEachItem(const TFunctionRef<EIterationResult(InItemType&)>& InFunc) const
+	{
+		for (const TSharedPtr<INavigationToolItem>& Item : Items)
+		{
+			if (!Item.IsValid())
+			{
+				continue;
+			}
+
+			if (InItemType* const CastedItem = Item->CastTo<InItemType>())
+			{
+				EIterationResult IterationResult = InFunc(*CastedItem);
+				if (IterationResult == EIterationResult::Break)
+				{
+					break;
+				}
+			}
+		}
+	}
+
+	TArray<FNavigationToolItemPtr> Items;
+
+	ENavigationToolDragDropActionType ActionType = ENavigationToolDragDropActionType::Move;
+};
+
+}
diff --git a/Engine/Source/Editor/Sequencer/Public/Tools/NavigationTool/DragDropOps/NavigationToolItemDragDropOp.h b/Engine/Source/Editor/Sequencer/Public/Tools/NavigationTool/DragDropOps/NavigationToolItemDragDropOp.h
new file mode 100644
--- /dev/null
+++ b/Engine/Source/Editor/Sequencer/Public/Tools/NavigationTool/DragDropOps/NavigationToolItemDragDropOp.h
@@ -0,0 +1,73 @@
+﻿// Copyright Epic Games, Inc. All Rights Reserved.
+
+#pragma once
+
+#include "DragAndDrop/DecoratedDragDropOp.h"
+#include "Handlers/NavigationToolItemDropHandler.h"
+
+enum class EItemDropZone;
+
+namespace UE::Sequencer::NavigationTool
+{
+
+class INavigationToolView;
+class FNavigationToolView;
+class FNavigationToolItemDropHandler;
+
+/** Drag Drop Operation for Navigation Tool Items. Customized behavior can be added in via the AddDropHandler function */
+class FNavigationToolItemDragDropOp : public FDecoratedDragDropOp
+{
+public:
+	DRAG_DROP_OPERATOR_TYPE(FNavigationToolItemDragDropOp, FDecoratedDragDropOp)
+
+	static TSharedRef<FNavigationToolItemDragDropOp> New(const TArray<FNavigationToolItemPtr>& InItems
+		, const TSharedPtr<FNavigationToolView>& InToolView
+		, ENavigationToolDragDropActionType InActionType);
+
+	TSharedPtr<INavigationToolView> GetToolView() const
+	{
+		return WeakToolView.Pin();
+	}
+
+	TConstArrayView<FNavigationToolItemPtr> GetItems() const
+	{
+		return Items;
+	}
+
+	ENavigationToolDragDropActionType GetActionType() const
+	{
+		return ActionType;
+	}
+
+	/** Retrieves the Actors found in the DragDropOp, assuming this is of type FNavigationToolItemDragDropOp */
+	SEQUENCER_API void GetDragDropOpActors(TArray<TWeakObjectPtr<AActor>>& OutActors) const;
+
+	template<typename InDropHandlerType
+		, typename = typename TEnableIf<TIsDerivedFrom<InDropHandlerType, FNavigationToolItemDropHandler>::Value>::Type
+		, typename... InArgTypes>
+	void AddDropHandler(InArgTypes&&... InArgs)
+	{
+		TSharedRef<FNavigationToolItemDropHandler> DropHandler = MakeShared<InDropHandlerType>(Forward<InArgTypes>(InArgs)...);
+		DropHandler->Initialize(*this);
+		DropHandlers.Add(DropHandler);
+	}
+
+	FReply Drop(EItemDropZone InDropZone, const FNavigationToolItemPtr& InTargetItem);
+
+	TOptional<EItemDropZone> CanDrop(EItemDropZone InDropZone, const FNavigationToolItemPtr& InTargetItem) const;
+
+protected:
+	void Init(const TArray<FNavigationToolItemPtr>& InItems
+		, const TSharedPtr<FNavigationToolView>& InToolView
+		, ENavigationToolDragDropActionType InActionType);
+
+	TArray<FNavigationToolItemPtr> Items;
+
+	TArray<TSharedRef<FNavigationToolItemDropHandler>> DropHandlers;
+
+	TWeakPtr<INavigationToolView> WeakToolView;
+
+	ENavigationToolDragDropActionType ActionType = ENavigationToolDragDropActionType::Move;
+};
+
+}
diff --git a/Engine/Source/Editor/Sequencer/Public/Tools/NavigationTool/Extensions/IColorExtension.h b/Engine/Source/Editor/Sequencer/Public/Tools/NavigationTool/Extensions/IColorExtension.h
new file mode 100644
--- /dev/null
+++ b/Engine/Source/Editor/Sequencer/Public/Tools/NavigationTool/Extensions/IColorExtension.h
@@ -0,0 +1,21 @@
+﻿// Copyright Epic Games, Inc. All Rights Reserved.
+
+#pragma once
+
+#include "Misc/Optional.h"
+#include "Tools/NavigationTool/NavigationToolItemType.h"
+
+namespace UE::Sequencer::NavigationTool
+{
+
+class SEQUENCER_API IColorExtension : public INavigationToolItemTypeCastable
+{
+public:
+	UE_NAVIGATIONTOOL_INHERITS(IColorExtension, INavigationToolItemTypeCastable)
+
+	virtual TOptional<FColor> GetColor() const = 0;
+
+	virtual void SetColor(const TOptional<FColor>& InColor) = 0;
+};
+
+}
diff --git a/Engine/Source/Editor/Sequencer/Public/Tools/NavigationTool/Extensions/IMarkerVisibilityExtension.h b/Engine/Source/Editor/Sequencer/Public/Tools/NavigationTool/Extensions/IMarkerVisibilityExtension.h
new file mode 100644
--- /dev/null
+++ b/Engine/Source/Editor/Sequencer/Public/Tools/NavigationTool/Extensions/IMarkerVisibilityExtension.h
@@ -0,0 +1,28 @@
+﻿// Copyright Epic Games, Inc. All Rights Reserved.
+
+#pragma once
+
+#include "Tools/NavigationTool/NavigationToolItemType.h"
+
+namespace UE::Sequencer::NavigationTool
+{
+
+enum class EItemMarkerVisibility
+{
+	None             = 0,
+	Visible          = 1 << 0,
+	PartiallyVisible = 1 << 1,
+};
+ENUM_CLASS_FLAGS(EItemMarkerVisibility)
+
+class SEQUENCER_API IMarkerVisibilityExtension : public INavigationToolItemTypeCastable
+{
+public:
+	UE_NAVIGATIONTOOL_INHERITS(IMarkerVisibilityExtension, INavigationToolItemTypeCastable)
+
+	virtual EItemMarkerVisibility GetMarkerVisibility() const = 0;
+
+	virtual void SetMarkerVisibility(const bool bInVisible) = 0;
+};
+
+}
diff --git a/Engine/Source/Editor/Sequencer/Public/Tools/NavigationTool/Extensions/IRenameableExtension.h b/Engine/Source/Editor/Sequencer/Public/Tools/NavigationTool/Extensions/IRenameableExtension.h
new file mode 100644
--- /dev/null
+++ b/Engine/Source/Editor/Sequencer/Public/Tools/NavigationTool/Extensions/IRenameableExtension.h
@@ -0,0 +1,23 @@
+﻿// Copyright Epic Games, Inc. All Rights Reserved.
+
+#pragma once
+
+#include "Misc/Optional.h"
+#include "Tools/NavigationTool/NavigationToolItemType.h"
+
+namespace UE::Sequencer::NavigationTool
+{
+
+class SEQUENCER_API IRenameableExtension : public INavigationToolItemTypeCastable
+{
+public:
+	UE_NAVIGATIONTOOL_INHERITS(IRenameableExtension, INavigationToolItemTypeCastable)
+
+	/** @return True if this item can be renamed */
+	virtual bool CanRename() const = 0;
+
+	/** The implementation to rename the item */
+	virtual bool Rename(const FString& InName) = 0;
+};
+
+}
diff --git a/Engine/Source/Editor/Sequencer/Public/Tools/NavigationTool/Extensions/IRevisionControlExtension.h b/Engine/Source/Editor/Sequencer/Public/Tools/NavigationTool/Extensions/IRevisionControlExtension.h
new file mode 100644
--- /dev/null
+++ b/Engine/Source/Editor/Sequencer/Public/Tools/NavigationTool/Extensions/IRevisionControlExtension.h
@@ -0,0 +1,34 @@
+﻿// Copyright Epic Games, Inc. All Rights Reserved.
+
+#pragma once
+
+#include "Tools/NavigationTool/NavigationToolItemType.h"
+
+struct FSlateBrush;
+
+namespace UE::Sequencer::NavigationTool
+{
+
+enum class EItemRevisionControlState
+{
+	None                = 0,
+	CheckedOut          = 1 << 0,
+	PartiallyCheckedOut = 1 << 1,
+};
+ENUM_CLASS_FLAGS(EItemRevisionControlState)
+
+class SEQUENCER_API IRevisionControlExtension : public INavigationToolItemTypeCastable
+{
+public:
+	UE_NAVIGATIONTOOL_INHERITS(IRevisionControlExtension, INavigationToolItemTypeCastable)
+
+	virtual EItemRevisionControlState GetRevisionControlState() const = 0;
+
+	virtual void SetIsRevisionControlled(const bool bInRevisionControlled) = 0;
+
+	virtual const FSlateBrush* GetRevisionControlStatusIcon() const = 0;
+
+	virtual FText GetRevisionControlStatusText() const = 0;
+};
+
+}
diff --git a/Engine/Source/Editor/Sequencer/Public/Tools/NavigationTool/Extensions/ISequenceInactivableExtension.h b/Engine/Source/Editor/Sequencer/Public/Tools/NavigationTool/Extensions/ISequenceInactivableExtension.h
new file mode 100644
--- /dev/null
+++ b/Engine/Source/Editor/Sequencer/Public/Tools/NavigationTool/Extensions/ISequenceInactivableExtension.h
@@ -0,0 +1,28 @@
+﻿// Copyright Epic Games, Inc. All Rights Reserved.
+
+#pragma once
+
+#include "Tools/NavigationTool/NavigationToolItemType.h"
+
+namespace UE::Sequencer::NavigationTool
+{
+
+enum class EItemSequenceInactiveState
+{
+	None              = 0,
+	Inactive          = 1 << 0,
+	PartiallyInactive = 1 << 1,
+};
+ENUM_CLASS_FLAGS(EItemSequenceInactiveState)
+
+class SEQUENCER_API ISequenceInactivableExtension : public INavigationToolItemTypeCastable
+{
+public:
+	UE_NAVIGATIONTOOL_INHERITS(ISequenceInactivableExtension, INavigationToolItemTypeCastable)
+
+	virtual EItemSequenceInactiveState GetInactiveState() const = 0;
+
+	virtual void SetIsInactive(const bool bInIsMuted) = 0;
+};
+
+}
diff --git a/Engine/Source/Editor/Sequencer/Public/Tools/NavigationTool/Extensions/ISequenceLockableExtension.h b/Engine/Source/Editor/Sequencer/Public/Tools/NavigationTool/Extensions/ISequenceLockableExtension.h
new file mode 100644
--- /dev/null
+++ b/Engine/Source/Editor/Sequencer/Public/Tools/NavigationTool/Extensions/ISequenceLockableExtension.h
@@ -0,0 +1,28 @@
+﻿// Copyright Epic Games, Inc. All Rights Reserved.
+
+#pragma once
+
+#include "Tools/NavigationTool/NavigationToolItemType.h"
+
+namespace UE::Sequencer::NavigationTool
+{
+
+enum class EItemSequenceLockState
+{
+	None            = 0,
+	Locked          = 1 << 0,
+	PartiallyLocked = 1 << 1,
+};
+ENUM_CLASS_FLAGS(EItemSequenceLockState)
+
+class SEQUENCER_API ISequenceLockableExtension : public INavigationToolItemTypeCastable
+{
+public:
+	UE_NAVIGATIONTOOL_INHERITS(ISequenceLockableExtension, INavigationToolItemTypeCastable)
+
+	virtual EItemSequenceLockState GetLockState() const = 0;
+
+	virtual void SetIsLocked(const bool bInIsLocked) = 0;
+};
+
+}
diff --git a/Engine/Source/Editor/Sequencer/Public/Tools/NavigationTool/Extensions/ISequenceMutableExtension.h b/Engine/Source/Editor/Sequencer/Public/Tools/NavigationTool/Extensions/ISequenceMutableExtension.h
new file mode 100644
--- /dev/null
+++ b/Engine/Source/Editor/Sequencer/Public/Tools/NavigationTool/Extensions/ISequenceMutableExtension.h
@@ -0,0 +1,28 @@
+﻿// Copyright Epic Games, Inc. All Rights Reserved.
+
+#pragma once
+
+#include "Tools/NavigationTool/NavigationToolItemType.h"
+
+namespace UE::Sequencer::NavigationTool
+{
+
+enum class EItemSequenceMuteState
+{
+	None           = 0,
+	Muted          = 1 << 0,
+	PartiallyMuted = 1 << 1,
+};
+ENUM_CLASS_FLAGS(EItemSequenceMuteState)
+
+class SEQUENCER_API ISequenceMutableExtension : public INavigationToolItemTypeCastable
+{
+public:
+	UE_NAVIGATIONTOOL_INHERITS(ISequenceMutableExtension, INavigationToolItemTypeCastable)
+
+	virtual EItemSequenceMuteState GetMuteState() const = 0;
+
+	virtual void SetIsMuted(const bool bInIsMuted) = 0;
+};
+
+}
diff --git a/Engine/Source/Editor/Sequencer/Public/Tools/NavigationTool/Filters/Filters/NavigationToolFilterBase.h b/Engine/Source/Editor/Sequencer/Public/Tools/NavigationTool/Filters/Filters/NavigationToolFilterBase.h
new file mode 100644
--- /dev/null
+++ b/Engine/Source/Editor/Sequencer/Public/Tools/NavigationTool/Filters/Filters/NavigationToolFilterBase.h
@@ -0,0 +1,65 @@
+﻿// Copyright Epic Games, Inc. All Rights Reserved.
+
+#pragma once
+
+#include "Filters/FilterBase.h"
+#include "Filters/SequencerFilterBase.h"
+#include "GameFramework/Actor.h"
+#include "Misc/IFilter.h"
+#include "Textures/SlateIcon.h"
+#include "Tools/NavigationTool/Filters/INavigationToolFilterBar.h"
+#include "Tools/NavigationTool/Items/NavigationToolTrack.h"
+#include "Tools/NavigationTool/NavigationToolDefines.h"
+#include "NavigationToolFilterBase.generated.h"
+
+UENUM()
+enum class ENavigationToolFilterMode : uint8
+{
+	None = 0,
+
+	/** Navigation Tool Item type matches the Filter Type */
+	MatchesType = 1 << 0,
+
+	/** Navigation Tool Item contains an item (as a descendant/child) of that type*/
+	ContainerOfType = 1 << 1,
+};
+ENUM_CLASS_FLAGS(ENavigationToolFilterMode);
+
+namespace UE::Sequencer::NavigationTool
+{
+
+class FNavigationToolFilter : public FSequencerFilterBase<FNavigationToolItemPtr>
+{
+public:
+	SEQUENCER_API FNavigationToolFilter(INavigationToolFilterBar& InOutFilterInterface, TSharedPtr<FFilterCategory>&& InCategory = nullptr);
+
+	//~ Begin IFilter
+	SEQUENCER_API virtual bool PassesFilter(const FNavigationToolItemPtr InItem) const override { return true; };
+	//~ End IFilter
+
+	//~ Begin FSequencerFilterBase
+	SEQUENCER_API virtual INavigationToolFilterBar& GetFilterInterface() const override;
+	//~ End FSequencerFilterBase
+};
+
+//////////////////////////////////////////////////////////////////////////
+//
+
+/** Base filter for filtering items based on their item class */
+template<typename InItemType>
+class FNavigationToolFilter_ItemType : public FNavigationToolFilter
+{
+public:
+	FNavigationToolFilter_ItemType(INavigationToolFilterBar& InOutFilterInterface, TSharedPtr<FFilterCategory> InCategory)
+		: FNavigationToolFilter(InOutFilterInterface, MoveTemp(InCategory))
+	{}
+
+	//~ Begin IFilter
+	virtual bool PassesFilter(const FNavigationToolItemPtr InItem) const override
+	{
+		return InItem->CastTo<InItemType>() != nullptr;
+	}
+	//~ End IFilter
+};
+
+}
diff --git a/Engine/Source/Editor/Sequencer/Public/Tools/NavigationTool/Filters/Filters/NavigationToolFilterCollection.h b/Engine/Source/Editor/Sequencer/Public/Tools/NavigationTool/Filters/Filters/NavigationToolFilterCollection.h
new file mode 100644
--- /dev/null
+++ b/Engine/Source/Editor/Sequencer/Public/Tools/NavigationTool/Filters/Filters/NavigationToolFilterCollection.h
@@ -0,0 +1,99 @@
+﻿// Copyright Epic Games, Inc. All Rights Reserved.
+
+#pragma once
+
+#include "Misc/FilterCollection.h"
+#include "Misc/IFilter.h"
+#include "Tools/NavigationTool/Items/NavigationToolItem.h"
+
+class FFilterCategory;
+class ISequencerFilterBar;
+
+namespace UE::Sequencer::NavigationTool
+{
+
+class FNavigationToolFilter;
+
+/** Some of this classes functionality could be moved to the TFilterCollection base class so other places could take advantage */
+class FNavigationToolFilterCollection : public TFilterCollection<FNavigationToolItemPtr>
+{
+public:
+	FNavigationToolFilterCollection(ISequencerFilterBar& InFilterInterface);
+	virtual ~FNavigationToolFilterCollection() {};
+
+	// @todo Maybe this should get moved in to TFilterCollection
+	bool ContainsFilter(const TSharedRef<FNavigationToolFilter>& InItem) const;
+
+	// @todo Maybe this should get moved in to TFilterCollection
+	void RemoveAll();
+
+	/**
+	 * Adds the specified Filter to the collection
+	 * 
+	 * @param InFilter The filter object to add to the collection
+	 * 
+	 * @return The index in the collection at which the filter was added
+	 */
+	int32 Add(const TSharedRef<FNavigationToolFilter>& InFilter);
+
+	/**
+	 * Removes as many instances of the specified Filter as there are in the collection
+	 * 
+	 * @param InFilter The filter object to remove from the collection
+	 * 
+	 * @return The number of Filters removed from the collection
+	 */
+	int32 Remove(const TSharedRef<FNavigationToolFilter>& InFilter);
+
+	/**
+	 * Gets the filter at the specified index
+	 * 
+	 * @param InIndex The index of the requested filter in the ChildFilters array
+	 * 
+	 * @return Filter at the specified index
+	 */
+	TSharedRef<FNavigationToolFilter> GetFilterAtIndex(const int32 InIndex);
+
+	/** Returns the number of Filters in the collection */
+	int32 Num() const;
+
+	bool IsEmpty() const;
+
+	/** Sorts the filters by display string */
+	void Sort();
+
+	/** Gets all the available track filter names */
+	TArray<FText> GetFilterDisplayNames() const;
+
+	TArray<TSharedRef<FNavigationToolFilter>> GetAllFilters(const bool bInCheckSupportsSequence
+		, const TArray<TSharedRef<FFilterCategory>>& InCategories = {}) const;
+
+	TSet<TSharedRef<FFilterCategory>> GetCategories(const TSet<TSharedRef<FNavigationToolFilter>>* InFilters = nullptr) const;
+
+	/** @return List of filters that match the specified category */
+	TArray<TSharedRef<FNavigationToolFilter>> GetCategoryFilters(const TSharedRef<FFilterCategory>& InCategory) const;
+
+	void ForEachFilter(const TFunctionRef<bool(const TSharedRef<FNavigationToolFilter>&)>& InFunction
+		, const TArray<TSharedRef<FFilterCategory>>& InCategories = {}) const;
+
+protected:
+	/**
+	 *	Called when a child Filter restrictions change and broadcasts the FilterChanged delegate
+	 *	for the collection
+	 */
+	void OnChildFilterChanged();
+
+	ISequencerFilterBar& FilterInterface;
+
+public:
+	/**
+	 * DO NOT USE DIRECTLY
+	 * STL-like iterators to enable range-based for loop support.
+	 */
+	FORCEINLINE TArray<TSharedPtr<IFilter<FNavigationToolItemPtr>>>::RangedForIteratorType      begin()       { return ChildFilters.begin(); }
+	FORCEINLINE TArray<TSharedPtr<IFilter<FNavigationToolItemPtr>>>::RangedForConstIteratorType begin() const { return ChildFilters.begin(); }
+	FORCEINLINE TArray<TSharedPtr<IFilter<FNavigationToolItemPtr>>>::RangedForIteratorType      end()	      { return ChildFilters.end(); }
+	FORCEINLINE TArray<TSharedPtr<IFilter<FNavigationToolItemPtr>>>::RangedForConstIteratorType end()   const { return ChildFilters.end(); }
+};
+
+}
diff --git a/Engine/Source/Editor/Sequencer/Public/Tools/NavigationTool/Filters/Filters/NavigationToolFilterExtension.h b/Engine/Source/Editor/Sequencer/Public/Tools/NavigationTool/Filters/Filters/NavigationToolFilterExtension.h
new file mode 100644
--- /dev/null
+++ b/Engine/Source/Editor/Sequencer/Public/Tools/NavigationTool/Filters/Filters/NavigationToolFilterExtension.h
@@ -0,0 +1,34 @@
+﻿// Copyright Epic Games, Inc. All Rights Reserved.
+
+#pragma once
+
+#include "NavigationToolFilterExtension.generated.h"
+
+class FFilterCategory;
+
+namespace UE::Sequencer::NavigationTool
+{
+	class FNavigationToolFilter;
+	class INavigationToolFilterBar;
+}
+
+/**
+ * Derive from this class to make additional filters available in the Navigation Tool.
+ */
+UCLASS(Abstract)
+class SEQUENCER_API UNavigationToolFilterExtension : public UObject
+{
+	GENERATED_BODY()
+
+public:
+	/**
+	 * Override to add additional Navigation Tool filters.
+	 * @param InOutFilterInterface The filter interface to extend
+	 * @param InPreferredCategory  The preferred filter bar category to place any newly added filters to, but can be any category
+	 * @param InOutFilterList      Filter list to add additional filters to
+	 */
+	virtual void AddFilterExtensions(UE::Sequencer::NavigationTool::INavigationToolFilterBar& InOutFilterInterface
+		, const TSharedRef<FFilterCategory>& InPreferredCategory
+		, TArray<TSharedRef<UE::Sequencer::NavigationTool::FNavigationToolFilter>>& InOutFilterList) const
+	{}
+};
diff --git a/Engine/Source/Editor/Sequencer/Public/Tools/NavigationTool/Filters/Filters/NavigationToolFilter_Marks.h b/Engine/Source/Editor/Sequencer/Public/Tools/NavigationTool/Filters/Filters/NavigationToolFilter_Marks.h
new file mode 100644
--- /dev/null
+++ b/Engine/Source/Editor/Sequencer/Public/Tools/NavigationTool/Filters/Filters/NavigationToolFilter_Marks.h
@@ -0,0 +1,34 @@
+﻿// Copyright Epic Games, Inc. All Rights Reserved.
+
+#pragma once
+
+#include "Tools/NavigationTool/Filters/Filters/NavigationToolFilterBase.h"
+#include "Tools/NavigationTool/Filters/NavigationToolFilterData.h"
+
+namespace UE::Sequencer::NavigationTool
+{
+
+class FNavigationToolFilter_Marks : public FNavigationToolFilter
+{
+public:
+	static FString StaticName() { return TEXT("Marks"); }
+
+	FNavigationToolFilter_Marks(INavigationToolFilterBar& InFilterInterface, TSharedPtr<FFilterCategory> InCategory = nullptr);
+
+	//~ Begin FNavigationToolFilter
+	virtual FText GetDefaultToolTipText() const override;
+	virtual TSharedPtr<FUICommandInfo> GetToggleCommand() const override;
+	//~ End FNavigationToolFilter
+
+	//~ Begin FFilterBase
+	virtual FText GetDisplayName() const override;
+	virtual FSlateIcon GetIcon() const override;
+	//~ End FFilterBase
+
+	//~ Begin IFilter
+	virtual FString GetName() const override;
+	virtual bool PassesFilter(FNavigationToolFilterType InItem) const override;
+	//~ End IFilter
+};
+
+}
diff --git a/Engine/Source/Editor/Sequencer/Public/Tools/NavigationTool/Filters/Filters/NavigationToolFilter_Unbound.h b/Engine/Source/Editor/Sequencer/Public/Tools/NavigationTool/Filters/Filters/NavigationToolFilter_Unbound.h
new file mode 100644
--- /dev/null
+++ b/Engine/Source/Editor/Sequencer/Public/Tools/NavigationTool/Filters/Filters/NavigationToolFilter_Unbound.h
@@ -0,0 +1,34 @@
+﻿// Copyright Epic Games, Inc. All Rights Reserved.
+
+#pragma once
+
+#include "Tools/NavigationTool/Filters/Filters/NavigationToolFilterBase.h"
+#include "Tools/NavigationTool/Filters/NavigationToolFilterData.h"
+
+namespace UE::Sequencer::NavigationTool
+{
+
+class FNavigationToolFilter_Unbound : public FNavigationToolFilter
+{
+public:
+	static FString StaticName() { return TEXT("Unbound"); }
+
+	FNavigationToolFilter_Unbound(INavigationToolFilterBar& InFilterInterface, TSharedPtr<FFilterCategory> InCategory = nullptr);
+
+	//~ Begin FNavigationToolFilter
+	virtual FText GetDefaultToolTipText() const override;
+	virtual TSharedPtr<FUICommandInfo> GetToggleCommand() const override;
+	//~ End FNavigationToolFilter
+
+	//~ Begin FFilterBase
+	virtual FText GetDisplayName() const override;
+	virtual FSlateIcon GetIcon() const override;
+	//~ End FFilterBase
+
+	//~ Begin IFilter
+	virtual FString GetName() const override;
+	virtual bool PassesFilter(FNavigationToolFilterType InItem) const override;
+	//~ End IFilter
+};
+
+}
diff --git a/Engine/Source/Editor/Sequencer/Public/Tools/NavigationTool/Filters/INavigationToolFilterBar.h b/Engine/Source/Editor/Sequencer/Public/Tools/NavigationTool/Filters/INavigationToolFilterBar.h
new file mode 100644
--- /dev/null
+++ b/Engine/Source/Editor/Sequencer/Public/Tools/NavigationTool/Filters/INavigationToolFilterBar.h
@@ -0,0 +1,46 @@
+﻿// Copyright Epic Games, Inc. All Rights Reserved.
+
+#pragma once
+
+#include "Filters/CustomTextFilters.h"
+#include "Filters/ISequencerFilterBar.h"
+
+class FFilterCategory;
+
+namespace UE::Sequencer::NavigationTool
+{
+
+class FNavigationToolFilter;
+class FNavigationToolFilter_CustomText;
+struct FNavigationToolFilterData;
+
+class INavigationToolFilterBar : public ISequencerFilterBar
+{
+public:
+	virtual ~INavigationToolFilterBar() override = default;
+
+	virtual TSet<TSharedRef<FFilterCategory>> GetFilterCategories(const TSet<TSharedRef<FNavigationToolFilter>>* InFilters = nullptr) const = 0;
+
+	virtual TArray<TSharedRef<FNavigationToolFilter>> GetCommonFilters(const TArray<TSharedRef<FFilterCategory>>& InCategories = {}) const = 0;
+
+	SEQUENCER_API virtual bool IsFilterEnabled(TSharedRef<FNavigationToolFilter> InFilter) const = 0;
+	SEQUENCER_API virtual bool SetFilterEnabled(const TSharedRef<FNavigationToolFilter> InFilter, const bool bInEnabled, const bool bInRequestFilterUpdate = true) = 0;
+
+	SEQUENCER_API virtual bool IsFilterActive(const TSharedRef<FNavigationToolFilter> InFilter) const = 0;
+	SEQUENCER_API virtual bool SetFilterActive(const TSharedRef<FNavigationToolFilter>& InFilter, const bool bInActive, const bool bInRequestFilterUpdate = true) = 0;
+
+	SEQUENCER_API virtual bool AddCustomTextFilter(const TSharedRef<FNavigationToolFilter_CustomText>& InFilter, const bool bInAddToConfig) = 0;
+	SEQUENCER_API virtual bool RemoveCustomTextFilter(const TSharedRef<FNavigationToolFilter_CustomText>& InFilter, const bool bInAddToConfig) = 0;
+
+	SEQUENCER_API virtual TSharedPtr<FNavigationToolFilter> FindFilterByDisplayName(const FString& InFilterName) const = 0;
+	SEQUENCER_API virtual TSharedPtr<FNavigationToolFilter_CustomText> FindCustomTextFilterByDisplayName(const FString& InFilterName) const = 0;
+
+	virtual FNavigationToolFilterData& GetFilterData() = 0;
+
+	SEQUENCER_API virtual void EnableCustomTextFilters(const bool bInEnable, const TArray<TSharedRef<FNavigationToolFilter_CustomText>> InExceptions = {}) = 0;
+
+	SEQUENCER_API virtual void CreateWindow_AddCustomTextFilter(const FCustomTextFilterData& InCustomTextFilterData = FCustomTextFilterData()) = 0;
+	SEQUENCER_API virtual void CreateWindow_EditCustomTextFilter(const TSharedPtr<FNavigationToolFilter_CustomText>& InCustomTextFilter) = 0;
+};
+
+}
diff --git a/Engine/Source/Editor/Sequencer/Public/Tools/NavigationTool/Filters/NavigationToolFilterData.h b/Engine/Source/Editor/Sequencer/Public/Tools/NavigationTool/Filters/NavigationToolFilterData.h
new file mode 100644
--- /dev/null
+++ b/Engine/Source/Editor/Sequencer/Public/Tools/NavigationTool/Filters/NavigationToolFilterData.h
@@ -0,0 +1,57 @@
+﻿// Copyright Epic Games, Inc. All Rights Reserved.
+
+#pragma once
+
+#include "Tools/NavigationTool/NavigationToolDefines.h"
+#include "UObject/WeakObjectPtr.h"
+
+class ISequencer;
+
+namespace UE::Sequencer::NavigationTool
+{
+
+using FNavigationToolFilterType = TSharedPtr<INavigationToolItem>;
+
+/** Represents a cache between nodes for a filter operation. */
+struct SEQUENCER_API FNavigationToolFilterData
+{
+	FNavigationToolFilterData(const FString& InRawFilterText);
+
+	bool operator==(const FNavigationToolFilterData& InRhs) const;
+	bool operator!=(const FNavigationToolFilterData& InRhs) const;
+
+	void Reset();
+
+	FString GetRawFilterText() const;
+
+	uint32 GetDisplayNodeCount() const;
+	uint32 GetTotalNodeCount() const;
+
+	uint32 GetFilterInCount() const;
+	uint32 GetFilterOutCount() const;
+
+	void IncrementTotalNodeCount();
+
+	void FilterInNode(const FNavigationToolItemPtr& InNode);
+	void FilterOutNode(const FNavigationToolItemPtr& InNode);
+
+	void FilterInParentChildNodes(const FNavigationToolItemPtr& InNode
+		, const bool bInIncludeSelf, const bool bInIncludeParents, const bool bInIncludeChildren = false);
+
+	void FilterInNodeWithAncestors(const FNavigationToolItemPtr& InNode);
+
+	bool ContainsFilterInNodes(const FNavigationToolFilterData& InOtherData) const;
+
+	bool IsFilteredIn(const FNavigationToolItemPtr& InNode) const;
+	bool IsFilteredOut(const FNavigationToolItemPtr& InNode) const;
+
+protected:
+	FString RawFilterText;
+
+	uint32 TotalNodeCount = 0;
+
+	/** Nodes to be displayed in the UI */
+	TSet<FNavigationToolItemWeakPtr> FilterInNodes;
+};
+
+}
diff --git a/Engine/Source/Editor/Sequencer/Public/Tools/NavigationTool/Filters/TextExpressions/NavigationToolFilterTextExpressionContext.h b/Engine/Source/Editor/Sequencer/Public/Tools/NavigationTool/Filters/TextExpressions/NavigationToolFilterTextExpressionContext.h
new file mode 100644
--- /dev/null
+++ b/Engine/Source/Editor/Sequencer/Public/Tools/NavigationTool/Filters/TextExpressions/NavigationToolFilterTextExpressionContext.h
@@ -0,0 +1,44 @@
+// Copyright Epic Games, Inc. All Rights Reserved.
+
+#pragma once
+
+#include "Filters/ISequencerTextFilterExpressionContext.h"
+#include "Filters/SequencerTrackFilterBase.h"
+#include "Tools/NavigationTool/Filters/INavigationToolFilterBar.h"
+#include "Tools/NavigationTool/NavigationToolDefines.h"
+
+class UMovieScene;
+class UMovieSceneSequence;
+
+namespace UE::Sequencer::NavigationTool
+{
+
+class FNavigationToolFilterTextExpressionContext : public ISequencerTextFilterExpressionContext
+{ 
+public:
+	SEQUENCER_API FNavigationToolFilterTextExpressionContext(INavigationToolFilterBar& InFilterInterface);
+
+	void SetFilterItem(const FNavigationToolItemPtr& InFilterItem);
+
+	//~ Begin ITextFilterExpressionContext
+
+	SEQUENCER_API virtual bool TestBasicStringExpression(const FTextFilterString& InValue
+		, const ETextFilterTextComparisonMode InTextComparisonMode) const override;
+
+	SEQUENCER_API virtual bool TestComplexExpression(const FName& InKey
+		, const FTextFilterString& InValue
+		, const ETextFilterComparisonOperation InComparisonOperation
+		, const ETextFilterTextComparisonMode InTextComparisonMode) const override;
+
+	//~ End ITextFilterExpressionContext
+
+protected:
+	SEQUENCER_API UMovieSceneSequence* GetFocusedMovieSceneSequence() const;
+	SEQUENCER_API UMovieScene* GetFocusedGetMovieScene() const;
+
+	INavigationToolFilterBar& FilterInterface;
+
+	FNavigationToolItemWeakPtr WeakFilterItem;
+};
+
+}
diff --git a/Engine/Source/Editor/Sequencer/Public/Tools/NavigationTool/Filters/TextExpressions/NavigationToolFilterTextExpressionExtension.h b/Engine/Source/Editor/Sequencer/Public/Tools/NavigationTool/Filters/TextExpressions/NavigationToolFilterTextExpressionExtension.h
new file mode 100644
--- /dev/null
+++ b/Engine/Source/Editor/Sequencer/Public/Tools/NavigationTool/Filters/TextExpressions/NavigationToolFilterTextExpressionExtension.h
@@ -0,0 +1,34 @@
+﻿// Copyright Epic Games, Inc. All Rights Reserved.
+
+#pragma once
+
+#include "NavigationToolFilterTextExpressionExtension.generated.h"
+
+class FNavigationToolFilterTextExpressionContext;
+
+namespace UE::Sequencer::NavigationTool
+{
+
+class INavigationToolFilterBar;
+class FNavigationToolFilterTextExpressionContext;
+
+}
+
+/**
+ * Derive from this class to make additional track filter text expressions available in Sequencer.
+ */
+UCLASS(Abstract)
+class SEQUENCER_API UNavigationToolFilterTextExpressionExtension : public UObject
+{
+	GENERATED_BODY()
+
+public:
+	/**
+	 * Override to add additional Sequencer track filter text expressions.
+	 * @param InOutFilterInterface The filter interface to extend
+	 * @param InOutExpressionList  Expression list to add additional text expressions to
+	 */
+	virtual void AddFilterTextExpressionExtensions(UE::Sequencer::NavigationTool::INavigationToolFilterBar& InOutFilterInterface
+		, TArray<TSharedRef<UE::Sequencer::NavigationTool::FNavigationToolFilterTextExpressionContext>>& InOutExpressionList) const
+	{}
+};
diff --git a/Engine/Source/Editor/Sequencer/Public/Tools/NavigationTool/INavigationTool.h b/Engine/Source/Editor/Sequencer/Public/Tools/NavigationTool/INavigationTool.h
new file mode 100644
--- /dev/null
+++ b/Engine/Source/Editor/Sequencer/Public/Tools/NavigationTool/INavigationTool.h
@@ -0,0 +1,250 @@
+﻿// Copyright Epic Games, Inc. All Rights Reserved.
+
+#pragma once
+
+#include "Containers/ArrayView.h"
+#include "Containers/ContainersFwd.h"
+#include "Containers/StringFwd.h"
+#include "Delegates/Delegate.h"
+#include "Modules/ModuleManager.h"
+#include "Tools/NavigationTool/Customization/NavigationToolExtender.h"
+#include "Tools/NavigationTool/ItemProxies/NavigationToolItemProxyRegistry.h"
+#include "Tools/NavigationTool/ItemProxies/INavigationToolItemProxyFactory.h"
+#include "Tools/NavigationTool/Items/NavigationToolItemId.h"
+#include "Tools/NavigationTool/Items/NavigationToolItemProxy.h"
+#include "Tools/NavigationTool/Items/INavigationToolItem.h"
+#include "Tools/NavigationTool/NavigationToolDefines.h"
+#include "Templates/SharedPointer.h"
+
+class AActor;
+class FArchive;
+class FSequencer;
+class FUICommandList;
+class UMovieSceneSequence;
+enum class EItemDropZone;
+struct FAttachmentTransformRules;
+template<typename OptionalType> struct TOptional;
+
+namespace UE::Sequencer::NavigationTool
+{
+
+class FNavigationToolEditorSelection;
+class FNavigationToolProvider;
+class INavigationToolAction;
+class INavigationToolFilterBar;
+class INavigationToolView;
+struct FNavigationToolItemId;
+
+/** 
+ * The Navigation Tool Object that is commonly instanced once per Sequencer
+ * (unless for advanced use where there are different Navigation Tool instances with different item ordering and behaviors).
+ * This is the object that dictates core Navigation Tool behavior like how items are sorted, which items are allowed, etc.
+ * Views are the objects that take this core behavior and show a part of it (e.g. through filters).
+ */
+class INavigationTool : public TSharedFromThis<INavigationTool>
+{
+public:
+	virtual ~INavigationTool() = default;
+
+	DECLARE_MULTICAST_DELEGATE(FOnToolLoaded);
+	virtual FOnToolLoaded& GetOnToolLoaded() = 0;
+
+	virtual TSharedPtr<ISequencer> GetSequencer() const = 0;
+
+	SEQUENCER_API virtual bool IsToolTabVisible() const = 0;
+	SEQUENCER_API virtual void ShowHideToolTab(const bool bInVisible) = 0;
+	SEQUENCER_API virtual void ToggleToolTabVisible() = 0;
+
+	/** Sets the Command List that the Navigation Tool Views will use to append their Command Lists to */
+	virtual void SetBaseCommandList(const TSharedPtr<FUICommandList>& InBaseCommandList) = 0;
+
+	/** Register a new Navigation Tool View to the Navigation Tool to the given id, replacing the old view that was bound to the given id */
+	virtual TSharedPtr<INavigationToolView> RegisterToolView(const int32 InToolViewId) = 0;
+
+	/** Gets the Navigation Tool View bound to the given id */
+	virtual TSharedPtr<INavigationToolView> GetToolView(const int32 InToolViewId) const = 0;
+
+	/** Gets the outliner view that was most recently used (i.e. called FNavigationTool::UpdateRecentToolViews) */
+	virtual TSharedPtr<INavigationToolView> GetMostRecentToolView() const = 0;
+
+	/** Instantiates a new Item and automatically registers it to the Navigation Tool */
+	template<typename InItemType, typename = typename TEnableIf<TIsDerivedFrom<InItemType, INavigationToolItem>::IsDerived>::Type, typename ...InArgTypes>
+	TSharedRef<InItemType> FindOrAdd(const TSharedPtr<FNavigationToolProvider>& InProvider, InArgTypes&&... InArgs)
+	{
+		const TSharedRef<InItemType> Item = MakeShared<InItemType>(*this, Forward<InArgTypes>(InArgs)...);
+		Item->SetProvider(InProvider);
+
+		// If an existing item already exists and has a valid state, use that and forget about the newly created
+		const FNavigationToolItemId ItemId = Item->GetItemId();
+		const FNavigationToolItemPtr ExistingItem = FindItem(ItemId);
+		if (ExistingItem.IsValid() && ExistingItem->IsItemValid() && ExistingItem->IsA<InItemType>())
+		{
+			return StaticCastSharedPtr<InItemType>(ExistingItem).ToSharedRef();
+		}
+
+		if (Item->IsAllowedInTool())
+		{
+			RegisterItem(Item);
+		}
+
+		return Item;
+	}
+
+	/** Tries to find the Item Proxy Factory for the given Item Proxy Type Name */
+	template<typename InItemProxyType, typename = typename TEnableIf<TIsDerivedFrom<InItemProxyType, FNavigationToolItemProxy>::IsDerived>::Type>
+	INavigationToolItemProxyFactory* GetItemProxyFactory() const
+	{
+		// First look for the Registry in Navigation Tool
+		if (INavigationToolItemProxyFactory* Factory = GetItemProxyRegistry().GetItemProxyFactory<InItemProxyType>())
+		{
+			return Factory;
+		}
+
+		// Fallback to finding the Factory in the Module if the Navigation Tool did not find it
+		return FNavigationToolExtender::GetItemProxyRegistry().GetItemProxyFactory<InItemProxyType>();
+	}
+
+	/**
+	 * Tries to get the Item Proxy Factory for the given Item Proxy type, first trying the Navigation Tool Registry then the Module's
+	 * then returns an existing item proxy created via the factory, or creates one if there's no existing item proxy
+	 * @returns the Item Proxy created by the Factory. Can be null if no factory was found or if the factory intentionally returns null
+	 */
+	template<typename InItemProxyType, typename = typename TEnableIf<TIsDerivedFrom<InItemProxyType, FNavigationToolItemProxy>::IsDerived>::Type, typename ...InArgTypes>
+	TSharedPtr<FNavigationToolItemProxy> GetOrCreateItemProxy(const FNavigationToolItemPtr& InParentItem, InArgTypes&&... InArgs)
+	{
+		if (!InParentItem.IsValid() || !InParentItem->IsAllowedInTool())
+		{
+			return nullptr;
+		}
+		
+		INavigationToolItemProxyFactory* const Factory = GetItemProxyFactory<InItemProxyType>();
+		if (!Factory)
+		{
+			return nullptr;
+		}
+		
+		TSharedPtr<FNavigationToolItemProxy> OutItemProxy;
+		if (const FNavigationToolItemPtr ExistingItemProxy = FindItem(FNavigationToolItemId(InParentItem, *Factory)))
+		{
+			check(ExistingItemProxy->IsA<FNavigationToolItemProxy>());
+			ExistingItemProxy->SetParent(InParentItem);
+			OutItemProxy = StaticCastSharedPtr<FNavigationToolItemProxy>(ExistingItemProxy);
+		}
+		else
+		{
+			OutItemProxy = Factory->CreateItemProxy(*this, InParentItem);
+		}
+		RegisterItem(OutItemProxy);
+		return OutItemProxy;
+	}
+
+	/** Gathers the Type Names of all the Item Proxies that are registered both in the outliner proxy registry and the module's */
+	virtual TArray<FName> GetRegisteredItemProxyTypeNames() const = 0;
+
+	/**
+	 * Determines whether the given actor can be presented in the Navigation Tool, at all.
+	 * This is a permanent check unlike filters that are temporary.
+	 */
+	virtual bool IsObjectAllowedInTool(const UObject* const InObject) const = 0;
+
+	/** Registers the given Item, replacing the old one. */
+	virtual void RegisterItem(const FNavigationToolItemPtr& InItem) = 0;
+
+	/** Unregisters the Item having the given ItemId */
+	virtual void UnregisterItem(const FNavigationToolItemId& InItemId) = 0;
+
+	/** Ensures that the next time Refresh is called in tick, Refresh will be called */
+	virtual void RequestRefresh() = 0;
+
+	/**
+	 * Flushes the Pending Actions from the Queue while also updating the state of the Navigation Tool.
+	 * Calling it directly is forcing it to happen.
+	 * If a refresh is needed it will be called on the next tick automatically.
+	 */
+	virtual void Refresh() = 0;
+
+	/** Gets the Tree Root Item of the Navigation Tool */
+	virtual TSharedRef<FNavigationToolItem> GetTreeRoot() const = 0;
+
+	/**
+	 * Finds the Registered Item that has the given Id
+	 * @returns the Item with the given Id, or null if the item does not exist or was not registered to the Navigation Tool
+	 */
+	virtual FNavigationToolItemPtr FindItem(const FNavigationToolItemId& InItemId) const = 0;
+
+	/**
+ 	 * Adds or Removes the Ignore Notify Flags to prevent certain actions from automatically happening when they're triggered
+ 	 * @param InFlag the ignore flag to add or remove
+ 	 * @param bIgnore whether to add (true) or remove (false) the flag
+ 	 */
+	virtual void SetIgnoreNotify(const ENavigationToolIgnoreNotifyFlags InFlag, const bool bInIgnore) = 0;
+
+	/** Called when the objects have been selected and notified through USelection Instances in Mode Tools */
+	virtual void OnObjectSelectionChanged(const FNavigationToolEditorSelection& InEditorSelection) = 0;
+
+	/** Returns the currently selected items in the most recent outliner view (since this list can vary between outliner views) */
+	virtual TArray<FNavigationToolItemPtr> GetSelectedItems(const bool bInNormalizeToTopLevelSelections = false) const = 0;
+
+	/**
+	 * Selects the given Items on all Navigation Tool Views
+	 * @param InItems the items to select
+	 * @param InFlags how the items should be selected (appended, notify of selections, etc)
+	 */
+	virtual void SelectItems(const TArray<FNavigationToolItemPtr>& InItems, const ENavigationToolItemSelectionFlags InFlags = ENavigationToolItemSelectionFlags::SignalSelectionChange) const = 0;
+
+	/**
+	 * Clears the Item Selection from all Navigation Tool Views
+	 * @param bInSignalSelectionChange whether to notify the change in selection
+	 */
+	virtual void ClearItemSelection(const bool bInSignalSelectionChange) const = 0;
+
+	/** Pairs the Item with the given color name, overriding the inherited color if different */
+	virtual void SetItemColor(const FNavigationToolItemPtr& InItem, const FColor& InColor) = 0;
+
+	/** Removes the Color pairing of the given Item (can still have an inherited color though) */
+	virtual void RemoveItemColor(const FNavigationToolItemPtr& InItem) = 0;
+
+	/**
+	 * Gets the color related to the Item
+	 * @param InItem the item to query
+	 * @param bRecurseParent whether to get the color of the parent (recursively) if the given item does not have a color by itself
+	 * @returns the matching color pair or unset if item is invalid or no color could be found.
+	 */
+	virtual TOptional<FColor> FindItemColor(const FNavigationToolItemPtr& InItem, bool bRecurseParent = true) const = 0;
+
+	/**
+	 * Instantiates a new item action without adding it to the Pending Actions Queue.
+	 * This should only be used directly when planning to enqueue multiple actions.
+	 * @see FNavigationTool::EnqueueItemActions
+	 */
+	template<typename InItemActionType, typename = typename TEnableIf<TIsDerivedFrom<InItemActionType, INavigationToolAction>::IsDerived>::Type, typename ...InArgTypes>
+	TSharedRef<InItemActionType> NewItemAction(InArgTypes&&... InArgs)
+	{
+		return MakeShared<InItemActionType>(Forward<InArgTypes>(InArgs)...);
+	}
+
+	/**
+	 * Instantiates a single new item action and immediately adds it to the Pending Actions Queue.
+	 * Ideal for when dealing with a single action.
+	 * For multiple actions use FNavigationTool::EnqueueItemActions.
+	 */
+	template<typename InItemActionType, typename = typename TEnableIf<TIsDerivedFrom<InItemActionType, INavigationToolAction>::IsDerived>::Type, typename ...InArgTypes>
+	void EnqueueItemAction(InArgTypes&&... InArgs)
+	{
+		EnqueueItemActions({ NewItemAction<InItemActionType>(Forward<InArgTypes>(InArgs)...) });
+	}
+
+	/** Adds the given actions to the Pending Action Queue */
+	virtual void EnqueueItemActions(TArray<TSharedPtr<INavigationToolAction>>&& InItemActions) noexcept = 0;
+
+	/** Called when an Item has been Renamed */
+	virtual void NotifyToolItemRenamed(const FNavigationToolItemPtr& InItem) = 0;
+
+	/** Called when an Item has been deleted */
+	virtual void NotifyToolItemDeleted(const FNavigationToolItemPtr& InItem) = 0;
+
+protected:
+	virtual const FNavigationToolItemProxyRegistry& GetItemProxyRegistry() const = 0;
+};
+
+}
