[Sequencer] Navigation Tool initial submit

diff --git a/Engine/Source/Editor/Sequencer/Public/Tools/NavigationTool/INavigationToolView.h b/Engine/Source/Editor/Sequencer/Public/Tools/NavigationTool/INavigationToolView.h
new file mode 100644
--- /dev/null
+++ b/Engine/Source/Editor/Sequencer/Public/Tools/NavigationTool/INavigationToolView.h
@@ -0,0 +1,125 @@
+// Copyright Epic Games, Inc. All Rights Reserved.
+
+#pragma once
+
+#include "Misc/Optional.h"
+#include "NavigationToolDefines.h"
+#include "Templates/SharedPointer.h"
+
+class FDragDropEvent;
+class FReply;
+class ISequencer;
+class SWidget;
+enum class EItemDropZone;
+struct FGeometry;
+struct FPointerEvent;
+enum class ENavigationToolItemViewMode : uint8;
+
+namespace UE::Sequencer::NavigationTool
+{
+
+class INavigationTool;
+class INavigationToolColumn;
+
+class INavigationToolView : public TSharedFromThis<INavigationToolView>
+{
+public:
+	virtual ~INavigationToolView() = default;
+
+	virtual TSharedPtr<INavigationTool> GetOwnerTool() const  = 0;
+
+	/** Returns the Navigation Tool Widget. Can be null widget */
+	virtual TSharedPtr<SWidget> GetToolWidget() const = 0;
+
+	virtual TSharedPtr<ISequencer> GetSequencer() const = 0;
+
+	/** Marks the Navigation Tool View to be refreshed on Next Tick */
+	virtual void RequestRefresh() = 0;
+
+	DECLARE_MULTICAST_DELEGATE(FOnToolViewRefreshed);
+	virtual FOnToolViewRefreshed& GetOnToolViewRefreshed() = 0;
+
+	virtual void SetKeyboardFocus() = 0;
+
+	virtual ENavigationToolItemViewMode GetItemDefaultViewMode() const = 0;
+
+	virtual ENavigationToolItemViewMode GetItemProxyViewMode() const = 0;
+
+	virtual bool IsColumnVisible(const TSharedPtr<INavigationToolColumn>& InColumn) const = 0;
+
+	virtual TSharedRef<SWidget> GetColumnMenuContent(const FName InColumnId) = 0;
+
+	virtual void GetChildrenOfItem(const FNavigationToolItemPtr InItem, TArray<FNavigationToolItemPtr>& OutChildren) const = 0;
+
+	/**
+	 * Gets the Children of a given Item. Can recurse if the immediate child is hidden (the children of these hidden items should still be given the opportunity to show up)
+	 * @param InItem the Item to get the children of
+	 * @param OutChildren the visible children in the give view mode.
+	 * @param InViewMode the view mode(s) that the children should support to be added to OutChildren
+	 * @param InRecursionDisallowedItems the items where recursion should not be performed
+	 */
+	virtual void GetChildrenOfItem(const FNavigationToolItemPtr& InItem
+		, TArray<FNavigationToolItemPtr>& OutChildren
+		, ENavigationToolItemViewMode InViewMode
+		, const TSet<FNavigationToolItemPtr>& InRecursionDisallowedItems) const = 0;
+
+	/** Whether the given Item is explicitly marked as Read-only in the Navigation Tool */
+	virtual bool IsItemReadOnly(const FNavigationToolItemPtr& InItem) const = 0;
+
+	/** Selection State */
+
+	/** Whether the given Item can be selected in this Navigation Tool View / Widget */
+	virtual bool CanSelectItem(const FNavigationToolItemPtr& InItem) const = 0;
+
+	/** Selects the Item in this Navigation Tool View */
+	virtual void SelectItems(TArray<FNavigationToolItemPtr> InItems, const ENavigationToolItemSelectionFlags InFlags = ENavigationToolItemSelectionFlags::SignalSelectionChange) = 0;
+
+	/** Clears the currently Selected Items in the Navigation Tool View */
+	virtual void ClearItemSelection(bool bSignalSelectionChange = true) = 0;
+
+	/** Whether the current item is selected in this Navigation Tool View */
+	virtual bool IsItemSelected(const FNavigationToolItemPtr& InItem) const = 0;
+
+	/** Gets the Currently Selected Items in the Tree View */
+	virtual TArray<FNavigationToolItemPtr> GetSelectedItems() const = 0;
+
+	/** Drag Drop */
+
+	/** Called a Drag is attempted for the selected items in view */
+	virtual FReply OnDragDetected(const FGeometry& InGeometry, const FPointerEvent& InMouseEvent, const FNavigationToolItemPtr InTargetItem) = 0;
+
+	/** Processes the Drag and Drop for the given Item if valid, else it will process it for the Root Item */
+	virtual FReply OnDrop(const FDragDropEvent& InDragDropEvent, EItemDropZone InDropZone, const FNavigationToolItemPtr InTargetItem) = 0;
+
+	/** Determines whether the Drag and Drop can be processed by the given item */
+	virtual TOptional<EItemDropZone> OnCanDrop(const FDragDropEvent& InDragDropEvent
+		, EItemDropZone InDropZone
+		, const FNavigationToolItemPtr InTargetItem) const = 0;
+
+	/** Expansion State */
+
+	virtual bool IsItemExpanded(const FNavigationToolItemPtr& InItem, const bool bInUseFilter = true) const = 0;
+
+	/** Expands / Collapses the given Item */
+	virtual void SetItemExpansion(const FNavigationToolItemPtr& InItem, const bool bInExpand, const bool bInUseFilter = true) = 0;
+
+	/** Expands / Collapses the given Item and its children recursively */
+	virtual void SetItemExpansionRecursive(const FNavigationToolItemPtr InItem, const bool bInExpand) = 0;
+
+	/** Change the expansion state of the parents (recursively) of the given item */
+	virtual void SetParentItemExpansions(const FNavigationToolItemPtr& InItem, const bool bInExpand) = 0;
+
+	/** @return True if all items in the tool tree view can be expanded */
+	virtual bool CanExpandAll() const = 0;
+
+	/** Expands all items in the tool tree view */
+	virtual void ExpandAll() = 0;
+
+	/** @return True if all items in the tool tree view can be collapsed */
+	virtual bool CanCollapseAll() const = 0;
+
+	/** Collapses all items in the tool tree view */
+	virtual void CollapseAll() = 0;
+};
+
+}
diff --git a/Engine/Source/Editor/Sequencer/Public/Tools/NavigationTool/ItemActions/NavigationToolAddItem.h b/Engine/Source/Editor/Sequencer/Public/Tools/NavigationTool/ItemActions/NavigationToolAddItem.h
new file mode 100644
--- /dev/null
+++ b/Engine/Source/Editor/Sequencer/Public/Tools/NavigationTool/ItemActions/NavigationToolAddItem.h
@@ -0,0 +1,34 @@
+// Copyright Epic Games, Inc. All Rights Reserved.
+
+#pragma once
+
+#include "NavigationToolItemAction.h"
+#include "Tools/NavigationTool/Items/NavigationToolItemParameters.h"
+
+namespace UE::Sequencer::NavigationTool
+{
+
+class INavigationTool;
+
+/**
+ * Item action responsible for adding an item to the Tree under a given optional Parent.
+ * If Parent is null, it is added as a Top Level Item.
+ */
+class SEQUENCER_API FNavigationToolAddItem : public INavigationToolAction
+{
+public:
+	UE_NAVIGATIONTOOL_INHERITS(FNavigationToolAddItem, INavigationToolAction);
+
+	FNavigationToolAddItem(const FNavigationToolAddItemParams& InAddItemParams);
+
+	//~ Begin INavigationToolAction
+	virtual bool ShouldTransact() const override;
+	virtual void Execute(INavigationTool& InTool) override;
+	virtual void OnObjectsReplaced(const TMap<UObject*, UObject*>& InReplacementMap, bool bRecursive) override;
+	//~ End INavigationToolAction
+
+protected:
+	FNavigationToolAddItemParams AddParams;
+};
+
+}
diff --git a/Engine/Source/Editor/Sequencer/Public/Tools/NavigationTool/ItemActions/NavigationToolItemAction.h b/Engine/Source/Editor/Sequencer/Public/Tools/NavigationTool/ItemActions/NavigationToolItemAction.h
new file mode 100644
--- /dev/null
+++ b/Engine/Source/Editor/Sequencer/Public/Tools/NavigationTool/ItemActions/NavigationToolItemAction.h
@@ -0,0 +1,32 @@
+// Copyright Epic Games, Inc. All Rights Reserved.
+
+#pragma once
+
+#include "Tools/NavigationTool/NavigationToolItemType.h"
+
+class UObject;
+
+namespace UE::Sequencer::NavigationTool
+{
+
+class INavigationTool;
+
+/**
+ * Interface class for an Action in the Navigation Tool (e.g. Add/Delete/Move Tree Item)
+ */ 
+class INavigationToolAction : public INavigationToolItemTypeCastable
+{
+public:
+	UE_NAVIGATIONTOOL_INHERITS(INavigationToolAction, INavigationToolItemTypeCastable);
+
+	/** Determines whether the given action modifies its objects and should transact */
+	virtual bool ShouldTransact() const { return false; }
+
+	/** The Action to execute on the given Navigation Tool */
+	virtual void Execute(INavigationTool& InTool) = 0;
+
+	/** Replace any Objects that might be held in this Action that has been killed and replaced by a new object (e.g. BP Components) */
+	virtual void OnObjectsReplaced(const TMap<UObject*, UObject*>& InReplacementMap, const bool bInRecursive) = 0;
+};
+
+}
diff --git a/Engine/Source/Editor/Sequencer/Public/Tools/NavigationTool/ItemActions/NavigationToolRemoveItem.h b/Engine/Source/Editor/Sequencer/Public/Tools/NavigationTool/ItemActions/NavigationToolRemoveItem.h
new file mode 100644
--- /dev/null
+++ b/Engine/Source/Editor/Sequencer/Public/Tools/NavigationTool/ItemActions/NavigationToolRemoveItem.h
@@ -0,0 +1,32 @@
+// Copyright Epic Games, Inc. All Rights Reserved.
+
+#pragma once
+
+#include "NavigationToolItemAction.h"
+#include "Tools/NavigationTool/Items/NavigationToolItemParameters.h"
+
+namespace UE::Sequencer::NavigationTool
+{
+
+class FNavigationTool;
+
+/**
+ * Item Action responsible of removing/unregistering Items from the Tree
+ */
+class SEQUENCER_API FNavigationToolRemoveItem : public INavigationToolAction
+{
+public:
+	UE_NAVIGATIONTOOL_INHERITS(FNavigationToolRemoveItem, INavigationToolAction);
+
+	FNavigationToolRemoveItem(const FNavigationToolRemoveItemParams& InRemoveItemParams);
+
+	//~ Begin INavigationToolAction
+	virtual void Execute(INavigationTool& InTool) override;
+	virtual void OnObjectsReplaced(const TMap<UObject*, UObject*>& InReplacementMap, bool bInRecursive) override;
+	//~ End INavigationToolAction
+
+protected:
+	FNavigationToolRemoveItemParams RemoveParams;
+};
+
+}
diff --git a/Engine/Source/Editor/Sequencer/Public/Tools/NavigationTool/ItemProxies/INavigationToolItemProxyFactory.h b/Engine/Source/Editor/Sequencer/Public/Tools/NavigationTool/ItemProxies/INavigationToolItemProxyFactory.h
new file mode 100644
--- /dev/null
+++ b/Engine/Source/Editor/Sequencer/Public/Tools/NavigationTool/ItemProxies/INavigationToolItemProxyFactory.h
@@ -0,0 +1,33 @@
+// Copyright Epic Games, Inc. All Rights Reserved.
+
+#pragma once
+
+#include "Templates/SharedPointer.h"
+#include "Tools/NavigationTool/NavigationToolDefines.h"
+
+class FName;
+
+namespace UE::Sequencer::NavigationTool
+{
+
+class FNavigationToolItem;
+class FNavigationToolItemProxy;
+class INavigationTool;
+
+/**
+ * Item Proxy Factories are the classes that instance or get the existing Navigation Tool Item Proxies for a given Item
+ * @see INavigationToolItemProxy
+ */
+class INavigationToolItemProxyFactory
+{
+public:
+	virtual ~INavigationToolItemProxyFactory() = default;
+
+	/** Gets the Type Name of the Item Proxy the Factory creates */
+	virtual FName GetItemProxyTypeName() const = 0;
+
+	/** Returns a newly created instance of the Relevant Item Proxy if successful */
+	virtual TSharedPtr<FNavigationToolItemProxy> CreateItemProxy(INavigationTool& InTool, const FNavigationToolItemPtr& InParentItem) = 0;
+};
+
+}
diff --git a/Engine/Source/Editor/Sequencer/Public/Tools/NavigationTool/ItemProxies/NavigationToolComponentProxy.h b/Engine/Source/Editor/Sequencer/Public/Tools/NavigationTool/ItemProxies/NavigationToolComponentProxy.h
new file mode 100644
--- /dev/null
+++ b/Engine/Source/Editor/Sequencer/Public/Tools/NavigationTool/ItemProxies/NavigationToolComponentProxy.h
@@ -0,0 +1,29 @@
+// Copyright Epic Games, Inc. All Rights Reserved.
+
+#pragma once
+
+#include "Tools/NavigationTool/Items/NavigationToolItemProxy.h"
+
+namespace UE::Sequencer::NavigationTool
+{
+
+class SEQUENCER_API FNavigationToolComponentProxy : public FNavigationToolItemProxy
+{
+public:
+	UE_NAVIGATIONTOOL_INHERITS_WITH_SUPER(FNavigationToolComponentProxy, FNavigationToolItemProxy)
+
+	FNavigationToolComponentProxy(INavigationTool& InTool, const FNavigationToolItemPtr& InParentItem);
+	
+	//~ Begin INavigationToolItem
+	virtual FText GetDisplayName() const override;
+	virtual FSlateIcon GetIcon() const override;
+	virtual FText GetIconTooltipText() const override;
+	virtual ENavigationToolItemViewMode GetSupportedViewModes(const INavigationToolView& InToolView) const override;
+	//~ End INavigationToolItem
+	
+	//~ Begin FNavigationToolItemProxy
+	virtual void GetProxiedItems(const TSharedRef<INavigationToolItem>& InParent, TArray<FNavigationToolItemPtr>& OutChildren, bool bInRecursive) override;
+	//~ End FNavigationToolItemProxy
+};
+
+}
diff --git a/Engine/Source/Editor/Sequencer/Public/Tools/NavigationTool/Items/INavigationToolItem.h b/Engine/Source/Editor/Sequencer/Public/Tools/NavigationTool/Items/INavigationToolItem.h
new file mode 100644
--- /dev/null
+++ b/Engine/Source/Editor/Sequencer/Public/Tools/NavigationTool/Items/INavigationToolItem.h
@@ -0,0 +1,507 @@
+// Copyright Epic Games, Inc. All Rights Reserved.
+
+#pragma once
+
+#include "Delegates/Delegate.h"
+#include "EditorSequenceNavigationDefs.h"
+#include "NavigationToolItemId.h"
+#include "Tools/NavigationTool/NavigationToolDefines.h"
+#include "Tools/NavigationTool/NavigationToolItemType.h"
+
+class FDragDropEvent;
+class FReply;
+class SWidget;
+class UObject;
+enum class EItemDropZone;
+enum class ENavigationToolItemViewMode : uint8;
+struct FLinearColor;
+struct FSlateBrush;
+struct FSlateIcon;
+
+namespace UE::Sequencer::NavigationTool
+{
+
+class FNavigationToolItemProxy;
+class FNavigationToolProvider;
+class FNavigationToolScopedSelection;
+class INavigationTool;
+class INavigationToolView;
+class SNavigationToolTreeRow;
+struct FNavigationToolAddItemParams;
+struct FNavigationToolItemId;
+struct FNavigationToolRemoveItemParams;
+
+/*
+ * An Navigation Tool Item is the class that represents a Single Element (i.e. Node) in the Navigation Tool Tree.
+ * This can be an Item that represents an Object (e.g. Actor,Component) or a Folder, or something else.
+*/
+class INavigationToolItem : public INavigationToolItemTypeCastable, public TSharedFromThis<INavigationToolItem>
+{
+protected:
+	using IndexType = TArray<INavigationToolItem>::SizeType;
+
+public:
+	UE_NAVIGATIONTOOL_INHERITS(INavigationToolItem, INavigationToolItemTypeCastable);
+
+	/** Gets the module provider responsible for the creation of this item */
+	virtual TSharedPtr<FNavigationToolProvider> GetProvider() const = 0;
+
+	/** Used to signal the Scoped Selection that this Item should be Selected */
+	virtual void Select(FNavigationToolScopedSelection& InSelection) const {}
+
+	/** Determines whether the given Item is selected in the given Scoped Selection */
+	virtual bool IsSelected(const FNavigationToolScopedSelection& InSelection) const { return false; }
+
+	/** Whether the Item can be selected in Navigation Tool at all */
+	virtual bool IsSelectable() const { return true; }
+
+	/** Called when an Item is selected */
+	virtual void OnSelect() {}
+
+	/** Called when an Item is double-clicked with the mouse */
+	virtual void OnDoubleClick() {}
+
+	/** Gets the Navigation Tool that owns this item */
+	virtual INavigationTool& GetOwnerTool() const = 0;
+
+	/** Determines whether the Item properties are in a valid state */
+	virtual bool IsItemValid() const = 0;
+
+	/** Called when the Item has been registered into the Navigation Tool */
+	virtual void OnItemRegistered() {}
+
+	/** Called when the Item has been unregistered from the Navigation Tool */
+	virtual void OnItemUnregistered() {}
+
+	/** Called when the Item been selected/deselected from the Tree View */
+	virtual void OnItemSelectionChanged(const bool bInIsSelected) {}
+
+	/** Refreshes what the Parent and Children are of this Item. (not recursive!) */
+	virtual void RefreshChildren() = 0;
+
+	/** Resets both the Cached Visible Children and Children (before doing so, sets all child's parents to null) */
+	virtual void ResetChildren()
+	{
+		for (const FNavigationToolItemPtr& Item : GetChildren())
+		{
+			if (Item.IsValid())
+			{
+				Item->SetParent(nullptr);
+			}
+		}
+		GetChildrenMutable().Reset();	
+	}
+
+	/**
+	 * Determines whether this item can be sorted by the Navigation Tool or not.
+	 * Unsorted Items usually mean that they have their own way of sorting that Navigation Tool's Item Sorting Data should not interfere with.
+	 * Note: Unsorted Children go before Sorted Items (e.g. Item Proxies go first before Actors below a Parent)
+	 */
+	virtual bool ShouldSort() const = 0;
+
+	/** Determines whether the given Child is supported and can be added under this Item */
+	virtual bool CanAddChild(const FNavigationToolItemPtr& InChild) const = 0;
+
+	/**
+	 * Adds another Child under this Item if such Item is supported.
+	 * Returns true if it did, false if it could not add it (e.g. item not supported).
+	 */
+	virtual bool AddChild(const FNavigationToolAddItemParams& InAddItemParams) = 0;
+
+	/**
+	 * Removes the given child from this Item if it was ever indeed a child.
+	 * Returns true if the removal did happen.
+	 */
+	virtual bool RemoveChild(const FNavigationToolRemoveItemParams& InRemoveItemParams) = 0;
+
+	/** Figures out the Children of this Item. This is only relevant for items that do have that functionality (e.g. Components or Actors) */
+	virtual void FindChildren(TArray<FNavigationToolItemPtr>& OutChildren, const bool bInRecursive) = 0;
+
+	virtual void FindValidChildren(TArray<FNavigationToolItemPtr>& OutChildren, const bool bInRecursive)
+	{
+		FindChildren(OutChildren, bInRecursive);
+
+		OutChildren.RemoveAll(
+			[](const FNavigationToolItemPtr& InItem)
+			{
+				return !InItem.IsValid() || !InItem->IsAllowedInTool();
+			});
+	}
+
+	/** Gets the Item Proxies for this Item (e.g. Component Item that represent Primitives add in a Material Proxy to display) */
+	virtual void GetItemProxies(TArray<TSharedPtr<FNavigationToolItemProxy>>& OutItemProxies) {}
+
+	/*
+	 * Tries to Find the First Path of Descendants (not including self) that lead to a given Item in the Set.
+	 * The last item is the Item in the Set that was found so path might be A/B/C/.../ItemInSet  where A is Child of This.
+	 * Returns Empty array if failed.
+	*/
+	virtual TArray<FNavigationToolItemPtr> FindPath(const TArray<FNavigationToolItemPtr>& Items) const
+	{
+		const TSharedPtr<const INavigationToolItem> This = SharedThis(this);
+		TArray<FNavigationToolItemPtr> Path;
+
+		for (FNavigationToolItemPtr Item : Items)
+		{
+			Path.Reset();
+			FNavigationToolItemPtr CurrentItem = Item;
+
+			while (CurrentItem.IsValid())
+			{
+				if (This == CurrentItem)
+				{
+					Algo::Reverse(Path);
+					return Path;
+				}
+
+				Path.Add(CurrentItem);
+				CurrentItem = CurrentItem->GetParent();
+			}
+		}
+
+		return TArray<FNavigationToolItemPtr>();
+	}
+
+	/** Returns the path to this item in the tree. Ex. "RootId,ChildId,ChildId" */
+	FString GetFullPath() const
+	{
+		FString OutPath;
+
+		FNavigationToolItemPtr CurrentParent = GetParent();
+		while (CurrentParent.IsValid())
+		{
+			if (OutPath.IsEmpty())
+			{
+				OutPath = CurrentParent->GetItemId().GetStringId();
+			}
+			else
+			{
+				OutPath = CurrentParent->GetItemId().GetStringId() + TEXT(",") + OutPath;
+			}
+
+			CurrentParent = CurrentParent->GetParent();
+		}
+
+		return OutPath;
+	}
+
+	/** Gets the current Child Items of this Item */
+	virtual const TArray<FNavigationToolItemPtr>& GetChildren() const = 0;
+	
+	/** Gets the current Child Items of this Item */
+	virtual TArray<FNavigationToolItemPtr>& GetChildrenMutable() = 0;
+
+	/**
+	 * Gets the Index that the given Child Item is at.
+	 * NOTE: This includes HIDDEN items as Item Visibility is relative to each Navigation Tool View.
+	 * To consider only Visible Items, use FNavigationToolView::GetVisibleChildAt.
+	 */
+	virtual IndexType GetChildIndex(const FNavigationToolItemPtr& ChildItem) const
+	{
+		return GetChildren().Find(ChildItem);
+	}
+
+	/*
+	 * Gets the Child Item at the given Index.
+	 * NOTE: This includes HIDDEN items as Item Visibility is relative to each Navigation Tool View.
+	 * To consider only Visible Items, use FNavigationToolView::GetVisibleChildAt.
+	 */
+	virtual FNavigationToolItemPtr GetChildAt(const IndexType Index) const
+	{
+		const TArray<FNavigationToolItemPtr>& ChildItems = GetChildren();
+		if (ChildItems.IsValidIndex(Index))
+		{
+			return ChildItems[Index];
+		}
+		return nullptr;
+	};
+
+	/** Gets the Parent of this Item. Should only be null prior to registering it in Navigation Tool or if its Root Item */
+	virtual FNavigationToolItemPtr GetParent() const = 0;
+
+	/** Gets the list of all parents up the tree from this Item */
+	virtual TSet<FNavigationToolItemPtr> GetParents() const
+	{
+		TSet<FNavigationToolItemPtr> OutParents;
+
+		FNavigationToolItemPtr CurrentParent = GetParent();
+		while (CurrentParent.IsValid())
+		{
+			OutParents.Add(CurrentParent);
+
+			CurrentParent = CurrentParent->GetParent();
+		}
+
+		return OutParents;
+	}
+
+	/** Sets the Parent. Note that the Parent must've already have this instance as a child (check is done) */
+	virtual void SetParent(FNavigationToolItemPtr InParent) = 0;
+
+	template<typename InItemType>
+	TArray<InItemType*> GetDescendantsOfType(const bool bInIncludeThis = false) const
+	{
+		TArray<InItemType*> OutDescendants;
+
+		if (bInIncludeThis)
+		{
+			if (InItemType* const CastedChild = const_cast<InItemType*>(CastTo<const InItemType>()))
+			{
+				OutDescendants.Add(CastedChild);
+			}
+		}
+
+		// Process each child of the current item recursively
+		for (const FNavigationToolItemPtr& Child : GetChildren())
+		{
+			// If the child is of the requested type, cast it and add to the result
+			if (InItemType* CastedChild = Child->CastTo<InItemType>())
+			{
+				OutDescendants.Add(CastedChild);
+			}
+
+			// Recursively check the child's children
+			if (Child.IsValid())
+			{
+				OutDescendants.Append(Child->GetDescendantsOfType<InItemType>());
+			}
+		}
+
+		return OutDescendants;
+	}
+
+	template<typename InItemType>
+	TArray<InItemType*> GetChildrenOfType(const bool bInIncludeThis = false) const
+	{
+		TArray<InItemType*> OutChildren;
+
+		if (bInIncludeThis)
+		{
+			if (InItemType* const CastedChild = const_cast<InItemType*>(CastTo<const InItemType>()))
+			{
+				OutChildren.Add(CastedChild);
+			}
+		}
+
+		for (const FNavigationToolItemPtr& Child : GetChildren())
+		{
+			if (InItemType* const CastedChild = Child->CastTo<InItemType>())
+			{
+				OutChildren.Add(CastedChild);
+			}
+		}
+
+		return OutChildren;
+	};
+
+	/*template<typename InItemType>
+	TArray<InItemType*> GetDescendantsOfType(const bool bInIncludeThis = false) const
+	{
+		TArray<InItemType*> OutDescendants;
+
+		for (InItemType* const Child : GetChildrenOfType<InItemType>(bInIncludeThis))
+		{
+			OutDescendants.Add(Child);
+
+			const TArray<InItemType*> ThisChildDescendants = Child->GetDescendantsOfType<InItemType>();
+			OutDescendants.Append(ThisChildDescendants);
+		}
+
+		return OutDescendants;
+	}*/
+	
+	/**
+	 * Finds the first ancestor view model that implements the given type
+	 */
+	template<typename InItemType>
+	InItemType* FindAncestorOfType(const bool bInIncludeThis = false)
+	{
+		if (bInIncludeThis && IsA<InItemType>())
+		{
+			return CastTo<InItemType>();
+		}
+
+		FNavigationToolItemPtr CurrentParent = GetParent();
+		while (CurrentParent.IsValid())
+		{
+			if (CurrentParent->IsA<InItemType>())
+			{
+				return CurrentParent->CastTo<InItemType>();
+			}
+
+			CurrentParent = CurrentParent->GetParent();
+		}
+
+		return nullptr;
+	}
+
+	/**
+	 * Whether this Item can be at the Top Level just beneath the Root, or it needs to always be under some other Item
+	 * E.g. Actors can be Top Level, but Components or Materials can't
+	 */
+	virtual bool CanBeTopLevel() const = 0;
+
+	/** Gets the Id of this Item */
+	virtual FNavigationToolItemId GetItemId() const = 0;
+	
+	/** Returns whether this Item (and what it represents) should be allowed to be registered in Navigation Tool */
+	virtual bool IsAllowedInTool() const = 0;
+	
+	/** Gets the Display Name Text of the Item */
+	virtual FText GetDisplayName() const = 0;
+
+	/** Gets the Class/Type of this Item (e.g. for Items that represent UObjects, it will be the UObject class) */
+	virtual FText GetClassName() const = 0;
+
+	virtual FSlateIcon GetIcon() const = 0;
+
+	virtual const FSlateBrush* GetDefaultIconBrush() const { return nullptr; }
+
+	virtual const FSlateBrush* GetIconBrush() const = 0;
+
+	virtual FText GetIconTooltipText() const = 0;
+
+	/** Gets the View Modes that this Item Supports */
+	virtual ENavigationToolItemViewMode GetSupportedViewModes(const INavigationToolView& InToolView) const = 0;
+
+	/** Whether this Item should be visualized in the given View Mode, for the given Navigation Tool View */
+	virtual bool IsViewModeSupported(const ENavigationToolItemViewMode InViewMode, const INavigationToolView& InToolView) const
+	{
+		return EnumHasAnyFlags(InViewMode, GetSupportedViewModes(InToolView));
+	};
+
+	/**
+	 * Called when objects have been replaced on the Engine side. Used to replace any UObjects used by this item
+	 * @param InReplacementMap the map of old object that is garbage to the new object that replaced it
+	 * @param bInRecursive     whether to recurse this same function to children items
+	 */
+	virtual void OnObjectsReplaced(const TMap<UObject*, UObject*>& InReplacementMap, const bool bInRecursive)
+	{
+		if (bInRecursive)
+		{
+			for (const FNavigationToolItemPtr& ChildItem : GetChildren())
+			{
+				if (ChildItem.IsValid())
+				{
+					ChildItem->OnObjectsReplaced(InReplacementMap, bInRecursive);
+				}
+			}
+		}
+	}
+
+	/** Function responsible of Generating the Label Widget for this Item (i.e. the column containing the Icon and the Name) */
+	virtual TSharedRef<SWidget> GenerateLabelWidget(const TSharedRef<SNavigationToolTreeRow>& InRow) = 0;
+
+	/** Whether this Item supports Visibility for the Given Type */
+	virtual bool ShowVisibility() const = 0;
+
+	/** Whether a change in Parent Visibility should also affect this Item's Visibility */
+	virtual bool CanReceiveParentVisibilityPropagation() const = 0;
+
+	/** Whether this Item is currently visible or not for the Given Type */
+	virtual bool GetVisibility() const = 0;
+
+	/** Called when the Visibility on Item has been changed on the Navigation Tool side */
+	virtual void OnVisibilityChanged(const bool bInNewVisibility) {}
+
+	/**
+	 * Delegate signature for when the Item Expansion Changes
+	 * const TSharedPtr<FNavigationToolView>& - The Navigation Tool View where this Expansion happened
+	 * bool - Whether the Item is expanded
+	 */
+	using FOnExpansionChanged = TMulticastDelegate<void(const TSharedPtr<INavigationToolView>&, bool)>;
+
+	/** Called when Expansion state (Expanded/Collapsed) has been changed */
+	virtual FOnExpansionChanged& OnExpansionChanged() = 0;
+
+	/** Whether the Item is able to expand when AutoExpand functionality is enabled */
+	virtual bool CanAutoExpand() const = 0;
+
+	/** Whether this Item can be deleted or not */
+	virtual bool CanDelete() const = 0;
+
+	/** The implementation to delete the item */
+	virtual bool Delete() = 0;
+
+	virtual void AddFlags(const ENavigationToolItemFlags InFlag) = 0;
+	virtual void RemoveFlags(const ENavigationToolItemFlags InFlag) = 0;
+
+	virtual bool HasAnyFlags(const ENavigationToolItemFlags InFlag) const = 0;
+	virtual bool HasAllFlags(const ENavigationToolItemFlags InFlag) const = 0;
+
+	virtual void SetFlags(const ENavigationToolItemFlags InFlags) = 0;
+	virtual ENavigationToolItemFlags GetFlags() const = 0;
+
+	/** Gets the Tags found for this Item (e.g. for Actors, actor tags and for Components Component Tags) */
+	virtual TArray<FName> GetTags() const { return TArray<FName>(); }
+
+	/** Returns the item's height in tree, Root Item should return 0 as it has no Parent */
+	virtual int32 GetItemTreeHeight() const
+	{
+		int32 Height = 0;
+		FNavigationToolItemPtr TopParent = GetParent();
+		while (TopParent.IsValid())
+		{
+			TopParent = TopParent->GetParent();
+			++Height;
+		}
+		return Height;
+	}
+
+	/**
+	 * Delegate signature for relaying an Item Rename action
+	 * ENavigationToolRenameAction the type of action being relayed (e.g. request a rename, or notify rename complete, etc)
+	 * const TSharedPtr<INavigationToolView>& - The Navigation Tool View where the rename action is taking place
+	 */
+	using FOnRenameAction = TMulticastDelegate<void(ENavigationToolRenameAction, const TSharedPtr<INavigationToolView>&)>;
+
+	/** Broadcasts whenever a rename action takes place from a given view (e.g. when pressing "F2" to rename, or committing the rename text) */
+	virtual FOnRenameAction& OnRenameAction() = 0;
+
+	/** Determines if and where the incoming Drag Drop Event can be processed by this item */
+	virtual TOptional<EItemDropZone> CanAcceptDrop(const FDragDropEvent& InDragDropEvent, const EItemDropZone InDropZone) = 0;
+
+	/** Processes the Drag and Drop Event for this Item */
+	virtual FReply AcceptDrop(const FDragDropEvent& InDragDropEvent, const EItemDropZone InDropZone) = 0;
+
+	/** Gets the Text Color to use for the Item in the Navigation Tool (e.g. Components should be grayed out a bit) */
+	virtual FLinearColor GetItemBrushColor() const = 0;
+
+	/** Whether Ignoring Pending Kill. Useful to get underlying UObjects that are pending kill and get the pointer to it and not a null value */
+	bool IsIgnoringPendingKill() const { return HasAllFlags(ENavigationToolItemFlags::IgnorePendingKill); }
+
+	/** Gets whether this item is expanded */
+	virtual bool IsExpanded() const = 0;
+
+	/** Sets the expansion state of this item */
+	virtual void SetExpansion(const bool bInIsExpanded) = 0;
+};
+
+/** Adds Scoped Item Flags, removes them when out of scope. Useful for temp checks like IgnorePendingKill */
+struct FNavigationToolItemFlagGuard
+{
+	FNavigationToolItemFlagGuard(const FNavigationToolItemPtr InItem, const ENavigationToolItemFlags ItemFlags)
+	{
+		if (InItem.IsValid())
+		{
+			Item         = InItem;
+			OldItemFlags = InItem->GetFlags();
+			InItem->SetFlags(ItemFlags);
+		}
+	}
+
+	~FNavigationToolItemFlagGuard()
+	{
+		if (Item.IsValid())
+		{
+			Item->SetFlags(OldItemFlags);
+		}
+	}
+
+protected:
+	FNavigationToolItemPtr Item;
+	
+	ENavigationToolItemFlags OldItemFlags;
+};
+
+}
diff --git a/Engine/Source/Editor/Sequencer/Public/Tools/NavigationTool/Items/NavigationToolActor.h b/Engine/Source/Editor/Sequencer/Public/Tools/NavigationTool/Items/NavigationToolActor.h
new file mode 100644
--- /dev/null
+++ b/Engine/Source/Editor/Sequencer/Public/Tools/NavigationTool/Items/NavigationToolActor.h
@@ -0,0 +1,54 @@
+// Copyright Epic Games, Inc. All Rights Reserved.
+
+#pragma once
+
+#include "NavigationToolBinding.h"
+
+class AActor;
+
+namespace UE::Sequencer::NavigationTool
+{
+
+/**
+ * Navigation Tool Item representing an AActor binding
+ */
+class SEQUENCER_API FNavigationToolActor
+	: public FNavigationToolBinding
+{
+public:
+	UE_NAVIGATIONTOOL_INHERITS_WITH_SUPER(FNavigationToolActor
+		, FNavigationToolBinding);
+
+	FNavigationToolActor(INavigationTool& InTool
+		, const FNavigationToolItemPtr& InParentItem
+		, const FMovieSceneBinding& InBinding
+		, UObject* const InObject
+		, const FNavigationToolItemPtr& InReferencingItem);
+
+	//~ Begin INavigationToolItem
+
+	virtual void FindChildren(TArray<FNavigationToolItemPtr>& OutChildren, const bool bInRecursive) override;
+	virtual void GetItemProxies(TArray<TSharedPtr<FNavigationToolItemProxy>>& OutItemProxies) override;
+
+	virtual bool IsAllowedInTool() const override;
+	virtual ENavigationToolItemViewMode GetSupportedViewModes(const INavigationToolView& InToolView) const override;
+	virtual bool CanReceiveParentVisibilityPropagation() const override { return true; }
+
+	virtual FText GetDisplayName() const override;
+	virtual TArray<FName> GetTags() const override;
+
+	virtual bool ShowVisibility() const override { return true; }
+	virtual bool GetVisibility() const override;
+	virtual void OnVisibilityChanged(const bool bInNewVisibility) override;
+
+	//~ End INavigationToolItem
+
+	//~ Begin IRenameableExtension
+	virtual bool CanRename() const override;
+	virtual bool Rename(const FString& InName) override;
+	//~ End IRenameableExtension
+
+	AActor* GetActor() const;
+};
+
+}
diff --git a/Engine/Source/Editor/Sequencer/Public/Tools/NavigationTool/Items/NavigationToolBinding.h b/Engine/Source/Editor/Sequencer/Public/Tools/NavigationTool/Items/NavigationToolBinding.h
new file mode 100644
--- /dev/null
+++ b/Engine/Source/Editor/Sequencer/Public/Tools/NavigationTool/Items/NavigationToolBinding.h
@@ -0,0 +1,74 @@
+﻿// Copyright Epic Games, Inc. All Rights Reserved.
+
+#pragma once
+
+#include "MovieSceneBinding.h"
+#include "Textures/SlateIcon.h"
+#include "Tools/NavigationTool/Extensions/IRenameableExtension.h"
+#include "Tools/NavigationTool/Extensions/ISequenceLockableExtension.h"
+#include "Tools/NavigationTool/Items/NavigationToolObject.h"
+
+namespace UE::Sequencer::NavigationTool
+{
+
+/**
+ * Navigation Tool Item representing a Sequence binding
+ */
+class SEQUENCER_API FNavigationToolBinding
+	: public FNavigationToolObject
+	, public IRenameableExtension
+	//, public ISequenceLockableExtension
+{
+public:
+	UE_NAVIGATIONTOOL_INHERITS_WITH_SUPER(FNavigationToolBinding
+		, FNavigationToolObject
+		, IRenameableExtension
+		/*, ISequenceLockableExtension*/);
+
+	FNavigationToolBinding(INavigationTool& InTool
+		, const FNavigationToolItemPtr& InParentItem
+		, const FMovieSceneBinding& InBinding
+		, UObject* const InObject
+		, const FNavigationToolItemPtr& InReferencingItem);
+
+	//~ Begin INavigationToolItem
+
+	virtual void FindChildren(TArray<FNavigationToolItemPtr>& OutChildren, const bool bInRecursive) override;
+
+	virtual bool CanBeTopLevel() const override { return false; }
+	virtual bool ShouldSort() const override { return false; }
+
+	virtual FText GetDisplayName() const override;
+	virtual const FSlateBrush* GetDefaultIconBrush() const override;
+	virtual FSlateIcon GetIcon() const override;
+	virtual FText GetIconTooltipText() const override;
+
+	virtual bool IsSelected(const FNavigationToolScopedSelection& InSelection) const override;
+	virtual void Select(FNavigationToolScopedSelection& InSelection) const override;
+
+	virtual bool CanDelete() const override;
+	virtual bool Delete() override;
+
+	//~ End INavigationToolItem
+
+	//~ Begin IRenameableExtension
+	virtual bool CanRename() const override;
+	virtual bool Rename(const FString& InName) override;
+	//~ End IRenameableExtension
+
+	//~ Begin ISequenceLockableExtension
+	//virtual EItemSequenceLockState GetLockState() const override;
+	//virtual void SetIsLocked(const bool bInIsLocked) override;
+	//~ End ISequenceLockableExtension
+
+	const FMovieSceneBinding& GetBinding() const;
+
+	UObject* FindBoundObject() const;
+
+protected:
+	FSlateIcon Icon;
+
+	FMovieSceneBinding Binding;
+};
+
+}
diff --git a/Engine/Source/Editor/Sequencer/Public/Tools/NavigationTool/Items/NavigationToolComponent.h b/Engine/Source/Editor/Sequencer/Public/Tools/NavigationTool/Items/NavigationToolComponent.h
new file mode 100644
--- /dev/null
+++ b/Engine/Source/Editor/Sequencer/Public/Tools/NavigationTool/Items/NavigationToolComponent.h
@@ -0,0 +1,55 @@
+// Copyright Epic Games, Inc. All Rights Reserved.
+
+#pragma once
+
+#include "NavigationToolBinding.h"
+
+class UActorComponent;
+
+namespace UE::Sequencer::NavigationTool
+{
+
+/**
+ * Navigation Tool Item representing an Actor Component binding
+ */
+class SEQUENCER_API FNavigationToolComponent
+	: public FNavigationToolBinding
+{
+public:
+	UE_NAVIGATIONTOOL_INHERITS_WITH_SUPER(FNavigationToolComponent
+		, FNavigationToolBinding);
+
+	FNavigationToolComponent(INavigationTool& InTool
+		, const FNavigationToolItemPtr& InParentItem
+		, const FMovieSceneBinding& InBinding
+		, UObject* const InObject
+		, const FNavigationToolItemPtr& InReferencingItem);
+
+	//~ Begin INavigationToolItem
+
+	virtual void FindChildren(TArray<FNavigationToolItemPtr>& OutChildren, const bool bInRecursive) override;
+	virtual void GetItemProxies(TArray<TSharedPtr<FNavigationToolItemProxy>>& OutItemProxies) override;
+
+	virtual bool IsAllowedInTool() const override;
+	virtual ENavigationToolItemViewMode GetSupportedViewModes(const INavigationToolView& InToolView) const override;
+	virtual bool CanReceiveParentVisibilityPropagation() const override { return false; }
+	virtual TSharedRef<SWidget> GenerateLabelWidget(const TSharedRef<SNavigationToolTreeRow>& InRow) override;
+
+	virtual FLinearColor GetItemBrushColor() const override;
+	virtual TArray<FName> GetTags() const override;
+
+	virtual bool ShowVisibility() const override { return true; }
+	virtual bool GetVisibility() const override;
+	virtual void OnVisibilityChanged(const bool bInNewVisibility) override;
+
+	//~ End INavigationToolItem
+
+	//~ Begin IRenameableExtension
+	virtual bool CanRename() const override;
+	virtual bool Rename(const FString& InName) override;
+	//~ End IRenameableExtension
+
+	UActorComponent* GetComponent() const;
+};
+
+}
diff --git a/Engine/Source/Editor/Sequencer/Public/Tools/NavigationTool/Items/NavigationToolItem.h b/Engine/Source/Editor/Sequencer/Public/Tools/NavigationTool/Items/NavigationToolItem.h
new file mode 100644
--- /dev/null
+++ b/Engine/Source/Editor/Sequencer/Public/Tools/NavigationTool/Items/NavigationToolItem.h
@@ -0,0 +1,131 @@
+// Copyright Epic Games, Inc. All Rights Reserved.
+
+#pragma once
+
+#include "Misc/Optional.h"
+#include "Tools/NavigationTool/Extensions/IColorExtension.h"
+#include "Tools/NavigationTool/Items/INavigationToolItem.h"
+#include "Tools/NavigationTool/Items/NavigationToolItemId.h"
+#include "Tools/NavigationTool/NavigationToolItemType.h"
+#include "Tools/NavigationTool/NavigationToolSettings.h"
+
+namespace UE::Sequencer::NavigationTool
+{
+
+class INavigationToolProvider;
+class SNavigationToolTreeRow;
+
+/** Base Implementation of INavigationToolItem */
+class SEQUENCER_API FNavigationToolItem
+	: public INavigationToolItem
+	, public IColorExtension
+{
+	friend class FNavigationTool;
+
+public:
+	UE_NAVIGATIONTOOL_INHERITS_WITH_SUPER(FNavigationToolItem
+		, INavigationToolItem
+		, IColorExtension);
+
+	DECLARE_DELEGATE_RetVal_OneParam(FNavigationToolSaveState*, FNavigationToolGetToolState, const FNavigationToolItemPtr /*InItem*/);
+
+	FNavigationToolItem(INavigationTool& InTool, const FNavigationToolItemPtr& InParentItem);
+
+	//~ Begin INavigationToolItem
+	virtual INavigationTool& GetOwnerTool() const final override;
+	virtual TSharedPtr<FNavigationToolProvider> GetProvider() const final override;
+	virtual bool IsItemValid() const override;
+	virtual void RefreshChildren() override;
+	virtual void ResetChildren() override;
+	virtual const TArray<FNavigationToolItemPtr>& GetChildren() const override { return Children; }
+	virtual void FindChildren(TArray<FNavigationToolItemPtr>& OutChildren, const bool bInRecursive) override;
+	virtual void FindValidChildren(TArray<FNavigationToolItemPtr>& OutChildren, const bool bInRecursive) override;
+	virtual TArray<FNavigationToolItemPtr> FindPath(const TArray<FNavigationToolItemPtr>& InItems) const override;
+	virtual TArray<FNavigationToolItemPtr>& GetChildrenMutable() override { return Children; }
+	virtual IndexType GetChildIndex(const FNavigationToolItemPtr& ChildItem) const override;
+	virtual bool ShouldSort() const override { return false; }
+	virtual bool CanAddChild(const FNavigationToolItemPtr& InChild) const override;
+	virtual bool AddChild(const FNavigationToolAddItemParams& InAddItemParams) override;
+	virtual bool RemoveChild(const FNavigationToolRemoveItemParams& InRemoveItemParams) override;
+	virtual FNavigationToolItemPtr GetParent() const override { return WeakParent.Pin(); }
+	virtual void SetParent(FNavigationToolItemPtr InParent) override;
+	virtual bool CanBeTopLevel() const override { return false; }
+	virtual bool IsAllowedInTool() const override { return true; }
+	virtual FText GetDisplayName() const override { return FText::GetEmpty(); }
+	virtual FText GetClassName() const override { return FText::GetEmpty(); }
+	virtual ENavigationToolItemViewMode GetSupportedViewModes(const INavigationToolView& InToolView) const override;
+	virtual FNavigationToolItemId GetItemId() const override final;
+	virtual const FSlateBrush* GetIconBrush() const override;
+	virtual TSharedRef<SWidget> GenerateLabelWidget(const TSharedRef<SNavigationToolTreeRow>& InRow) override;
+	virtual bool ShowVisibility() const override { return false; }
+	virtual bool CanReceiveParentVisibilityPropagation() const override { return false; }
+	virtual bool GetVisibility() const override { return false; }
+	virtual bool CanAutoExpand() const override { return true; }
+	virtual bool CanDelete() const override { return false; }
+	virtual bool Delete() override;
+	virtual void AddFlags(ENavigationToolItemFlags Flags) override;
+	virtual void RemoveFlags(ENavigationToolItemFlags Flags) override;
+	virtual bool HasAnyFlags(ENavigationToolItemFlags Flags) const override;
+	virtual bool HasAllFlags(ENavigationToolItemFlags Flags) const override;
+	virtual void SetFlags(ENavigationToolItemFlags InFlags) override { ItemFlags = InFlags; }
+	virtual ENavigationToolItemFlags GetFlags() const override { return ItemFlags; }
+	virtual TArray<FName> GetTags() const override { return TArray<FName>(); };
+	virtual TOptional<EItemDropZone> CanAcceptDrop(const FDragDropEvent& InDragDropEvent, EItemDropZone InDropZone) override;
+	virtual FReply AcceptDrop(const FDragDropEvent& InDragDropEvent, EItemDropZone InDropZone) override;
+	virtual FLinearColor GetItemBrushColor() const override;
+	virtual FOnRenameAction& OnRenameAction() override { return OnRenameActionDelegate; }
+	virtual FOnExpansionChanged& OnExpansionChanged() override { return OnExpansionChangedDelegate; }
+	virtual bool IsExpanded() const override;
+	virtual void SetExpansion(const bool bInExpand) override;
+	//~ End INavigationToolItem
+
+	//~ Begin IColorExtension
+	virtual TOptional<FColor> GetColor() const override;
+	virtual void SetColor(const TOptional<FColor>& InColor) override;
+	//~ End IColorExtension
+
+protected:
+	friend class INavigationTool;
+
+	void SetProvider(const TWeakPtr<FNavigationToolProvider>& InWeakProvider);
+
+	/** Gets the Item Id with the latest information (e.g. parent, object, etc)*/
+	virtual FNavigationToolItemId CalculateItemId() const = 0;
+
+	/** Sets the ItemId member var to what CalculateItemId returns */
+	void RecalculateItemId();
+	
+	/** The actual implementation of putting the given Item under the Children array */
+	void AddChildChecked(const FNavigationToolAddItemParams& InAddItemParams);
+
+	/** The actual implementation of removing the given item from the Children array */
+	bool RemoveChildChecked(const FNavigationToolRemoveItemParams& InRemoveItemParams);
+	
+	/** Careful handling of multiple children being detected and added to this item children array */
+	void HandleNewSortableChildren(TArray<FNavigationToolItemPtr> InSortableChildren);
+
+	/** Reference to the Owning Navigation Tool */
+	INavigationTool& Tool;
+
+	/** Tool provider that is responsible for the creation of this item */
+	TWeakPtr<FNavigationToolProvider> WeakProvider;
+
+	/** Weak pointer to the Parent Item. Can be null, but if valid, the Parent should have this item in the Children Array */
+	FNavigationToolItemWeakPtr WeakParent;
+
+	/** Array of Shared pointers to the Child Items. These Items should have their ParentWeak pointing to this item */
+	TArray<FNavigationToolItemPtr> Children;
+
+	/** Delegate for when Expansion changes in the Item */
+	FOnExpansionChanged OnExpansionChangedDelegate;
+	
+	/** The delegate for renaming */
+	FOnRenameAction OnRenameActionDelegate;
+
+	/** The current flags set for this item */
+	ENavigationToolItemFlags ItemFlags = ENavigationToolItemFlags::None;
+
+	FNavigationToolItemId ItemId;
+};
+
+}
diff --git a/Engine/Source/Editor/Sequencer/Public/Tools/NavigationTool/Items/NavigationToolItemId.h b/Engine/Source/Editor/Sequencer/Public/Tools/NavigationTool/Items/NavigationToolItemId.h
new file mode 100644
--- /dev/null
+++ b/Engine/Source/Editor/Sequencer/Public/Tools/NavigationTool/Items/NavigationToolItemId.h
@@ -0,0 +1,110 @@
+// Copyright Epic Games, Inc. All Rights Reserved.
+
+#pragma once
+
+#include "Tools/NavigationTool/ItemProxies/INavigationToolItemProxyFactory.h"
+#include "UObject/SoftObjectPath.h"
+
+class UMovieSceneSequence;
+class UMovieSceneSubSection;
+class UObject;
+
+namespace UE::Sequencer::NavigationTool
+{
+
+class INavigationToolItem;
+
+/** Struct to identify an item in the Navigation Tool */
+struct SEQUENCER_API FNavigationToolItemId
+{
+	static FNavigationToolItemId RootId;
+
+	static constexpr const TCHAR* Separator = TEXT(",");
+
+	/** Default Ctor. Does not run any CalculateTypeHash so this instance remains invalid */
+	FNavigationToolItemId() = default;
+
+	/** More flexible option to just specify the string directly. Could be used for folders (e.g. for a nested folder C could be "A/B/C" */
+	FNavigationToolItemId(const FString& InUniqueId);
+
+	/**
+	 * Ctor used for Objects that appear ONCE in the Navigation Tool (e.g. an Actor or Component), so only the Object itself is needed for the Id
+	 * for objects that can appear multiple times because they're referenced multiple times, see the ctor below that takes in a Referencing Item and Referencing Id
+	 * @param InObject the object that the item represents
+	 */
+	FNavigationToolItemId(const UObject* InObject);
+
+	/**
+	 * Ctor used for Objects that are expected to appear multiple times in the Navigation Tool (e.g. a Material Ref)
+	 *
+	 * Example Id #1:
+	 * [Component Path], [Full Path of Material Asset], [Slot Index]
+	 * "/Game/World.World:PersistentLevel.StaticMeshActor_0.StaticMeshComponent,/Game/Materials/M_TestMaterial.M_TestMaterial,[Slot 0]"
+	 *
+	 * Example Id #2:
+	 * [Component Path], [Material Instance Dynamic Path], [Slot Index]
+	 * "/Game/World.World:PersistentLevel.StaticMeshActor_0.StaticMeshComponent,/Game/World.World:PersistentLevel.StaticMeshActor_0.StaticMeshComponent.MaterialInstanceDynamic_0,[Slot 0]"
+	 * 
+	 * @param InObject the Object being referenced multiple times (e.g. a Material)
+	 * @param InReferencingItem the Item holding a reference to the object (e.g. a Primitive Comp referencing a Material)
+	 * @param InReferencingId the Id to differentiate this reference from other references within the same object (e.g. can be slot or property name)
+	 */
+	FNavigationToolItemId(const UObject* InObject
+		, const FNavigationToolItemPtr& InReferencingItem
+		, const FString& InReferencingId);
+
+	/**
+	 * Ctor used for making the Item Id for a item proxy that will be under the given Parent Item.
+	 * Used when the actual Item Proxy is not created yet, but know it's factory and want to know whether the item proxy already exists
+	 * @param InParentItem the parent item holding the item proxy
+	 * @param InItemProxyFactory the factory responsible of creating the item proxy
+	 */
+	FNavigationToolItemId(const FNavigationToolItemPtr& InParentItem
+		, const INavigationToolItemProxyFactory& InItemProxyFactory);
+
+	/**
+	 * Ctor used for making the Item Id for a Item Proxy under the given Parent Item.
+	 * @param InParentItem the parent item holding the item proxy
+	 * @param InItemProxy the item proxy under the parent item
+	 */
+	FNavigationToolItemId(const FNavigationToolItemPtr& InParentItem
+		, const INavigationToolItem& InItemProxy);
+
+	/**  */
+	FNavigationToolItemId(const FNavigationToolItemPtr& InParentItem
+		, const UMovieSceneSequence* const InSequence
+		, const UMovieSceneSubSection* const InSubSection = nullptr
+		, const int32 InSectionIndex = 0);
+
+	FNavigationToolItemId(const FNavigationToolItemId& Other);
+	FNavigationToolItemId(FNavigationToolItemId&& Other) noexcept;
+
+	FNavigationToolItemId& operator=(const FNavigationToolItemId& Other);
+	FNavigationToolItemId& operator=(FNavigationToolItemId&& Other) noexcept;
+
+	bool operator==(const FNavigationToolItemId& Other) const;
+
+	friend uint32 GetTypeHash(const FNavigationToolItemId& InItemId)
+	{
+		return InItemId.CachedHash;
+	}
+
+	/** Returns whether this Id has a cached hash (i.e. ran any ctor except the default one) */
+	bool IsValid() const;
+
+	FString GetStringId() const;
+
+private:
+	static FString GetObjectPath(const UObject* const InObject);
+
+	static void AddSeparatedSegment(FString& OutString, const FString& InSegment);
+
+	void CalculateTypeHash();
+
+	FString Id;
+
+	uint32 CachedHash = 0;
+	bool bHasCachedHash = false;
+};
+
+}
diff --git a/Engine/Source/Editor/Sequencer/Public/Tools/NavigationTool/Items/NavigationToolItemProxy.h b/Engine/Source/Editor/Sequencer/Public/Tools/NavigationTool/Items/NavigationToolItemProxy.h
new file mode 100644
--- /dev/null
+++ b/Engine/Source/Editor/Sequencer/Public/Tools/NavigationTool/Items/NavigationToolItemProxy.h
@@ -0,0 +1,49 @@
+// Copyright Epic Games, Inc. All Rights Reserved.
+
+#pragma once
+
+#include "NavigationToolItem.h"
+
+namespace UE::Sequencer::NavigationTool
+{
+
+/**
+ * Item Proxies are Navigation Tool Items that with the sole purpose to group and hold common items together.
+ * The description or name of such commonality between these Items should be the name of the Proxy that holds them.
+ * 
+ * NOTE: Although Item Proxies by default require to have a Parent to Show in Navigation Tool,
+ * they can be created without parent as means to get overriden behavior (e.g. DisplayName, Icon, etc)
+ */
+class SEQUENCER_API FNavigationToolItemProxy : public FNavigationToolItem
+{
+public:
+	UE_NAVIGATIONTOOL_INHERITS_WITH_SUPER(FNavigationToolItemProxy, FNavigationToolItem);
+
+	FNavigationToolItemProxy(INavigationTool& InTool, const FNavigationToolItemPtr& InParentItem);
+
+	//~ Begin INavigationToolItem
+	virtual bool IsItemValid() const override;
+	virtual void FindChildren(TArray<FNavigationToolItemPtr>& OutChildren, bool bRecursive) override final;
+	virtual void SetParent(FNavigationToolItemPtr InParent) override;
+	virtual ENavigationToolItemViewMode GetSupportedViewModes(const INavigationToolView& InToolView) const override;
+	virtual bool CanAutoExpand() const override;
+	virtual FText GetClassName() const override { return FText::GetEmpty(); }
+	//~ End INavigationToolItem
+
+	uint32 GetPriority() const { return Priority; }
+	void SetPriority(uint32 InPriority) { Priority = InPriority; }
+
+	/** Gets the Items that this Item Proxy is representing / holding (i.e. children) */
+	virtual void GetProxiedItems(const TSharedRef<INavigationToolItem>& InParent, TArray<FNavigationToolItemPtr>& OutChildren, bool bInRecursive) = 0;
+
+protected:
+	//~ Begin FNavigationToolItem
+	virtual FNavigationToolItemId CalculateItemId() const override;
+	//~ End FNavigationToolItem
+
+private:
+	/** This Item Proxy's Order Priority (i.e. Highest priority is placed topmost or leftmost (depending on Orientation). Priority 0 is lowest priority */
+	uint32 Priority = 0;
+};
+
+}
diff --git a/Engine/Source/Editor/Sequencer/Public/Tools/NavigationTool/Items/NavigationToolItemUtils.h b/Engine/Source/Editor/Sequencer/Public/Tools/NavigationTool/Items/NavigationToolItemUtils.h
new file mode 100644
--- /dev/null
+++ b/Engine/Source/Editor/Sequencer/Public/Tools/NavigationTool/Items/NavigationToolItemUtils.h
@@ -0,0 +1,196 @@
+// Copyright Epic Games, Inc. All Rights Reserved.
+
+#pragma once
+
+#include "Containers/ContainersFwd.h"
+#include "NavigationToolSequence.h"
+#include "Algo/AllOf.h"
+#include "MVVM/ViewModelPtr.h"
+#include "Tools/NavigationTool/Items/INavigationToolItem.h"
+#include "Tools/NavigationTool/NavigationToolDefines.h"
+
+class UMovieSceneMetaData;
+class UMovieSceneSubSection;
+
+namespace UE::Sequencer::NavigationTool
+{
+
+class INavigationTool;
+class FNavigationToolItem;
+class FNavigationToolSequence;
+
+/**
+ * Compares the absolute order of the items in the Navigation Tool and returns true if A comes before B in the Navigation Tool.
+ * Useful to use when sorting Items.
+ */
+bool CompareToolItemOrder(const FNavigationToolItemPtr& A, const FNavigationToolItemPtr& B);
+
+/** Returns two subset arrays of Items: one is containing only Sortable Items and the other Non Sortable Items */
+void SplitSortableAndUnsortableItems(const TArray<FNavigationToolItemPtr>& InItems
+	, TArray<FNavigationToolItemPtr>& OutSortable
+	, TArray<FNavigationToolItemPtr>& OutUnsortable);
+
+/**  */
+UMovieSceneSubSection* GetSequenceItemSubSection(const FNavigationToolItemPtr& InItem);
+
+/**  */
+UMovieSceneMetaData* GetSequenceItemMetaData(const TSharedPtr<INavigationToolItem>& InItem);
+
+/**  */
+SEQUENCER_API FText SequenceDisplayNameWithModifyAndInactiveStatus(const TSharedPtr<FNavigationToolSequence>& InSequenceItem
+	, const bool bInRemoveParentName
+	, const bool bInIsActive);
+
+/**  */
+enum class ENavigationToolCompareState
+{
+	NoneTrue    = 0,
+	AllTrue     = 1 << 0,
+	PartialTrue = 1 << 1,
+};
+ENUM_CLASS_FLAGS(ENavigationToolCompareState)
+
+/**  */
+template<typename InItemType>
+ENavigationToolCompareState CompareChildrenItemState(const INavigationToolItem* const InItem
+	, const TFunctionRef<bool(const InItemType* const)>& InTrueFunction
+	, const TFunctionRef<bool(const InItemType* const)>& InFalseFunction
+	, const ENavigationToolCompareState InDefaultState = ENavigationToolCompareState::NoneTrue)
+{
+	if (!InItem)
+	{
+		return InDefaultState;
+	}
+
+	const TArray<const InItemType*> Children = InItem->GetChildrenOfType<const InItemType>();
+	if (Children.IsEmpty())
+	{
+		return InDefaultState;
+	}
+
+	const bool bAllMatch = Algo::AllOf(Children,
+		[&InTrueFunction](const InItemType* const InItem)
+		{
+			return InTrueFunction(InItem);
+		});
+	if (bAllMatch)
+	{
+		return ENavigationToolCompareState::AllTrue;
+	}
+
+	const bool bAllNoMatch = Algo::AllOf(Children,
+		[&InFalseFunction](const InItemType* const InItem)
+		{
+			return InFalseFunction(InItem);
+		});
+	if (bAllNoMatch)
+	{
+		return ENavigationToolCompareState::NoneTrue;
+	}
+
+	return ENavigationToolCompareState::PartialTrue;
+}
+
+/**  */
+template<typename InItemType>
+ENavigationToolCompareState CompareChildrenItemStateSimple(const INavigationToolItem* const InItem
+	, const TFunctionRef<bool(const InItemType* const)>& InTrueFunction)
+{
+	if (!InItem)
+	{
+		return ENavigationToolCompareState::NoneTrue;
+	}
+
+	const TArray<const InItemType*> Children = InItem->GetChildrenOfType<const InItemType>();
+
+	const bool bAllMatch = Algo::AllOf(Children,
+		[&InTrueFunction](const InItemType* const InItem)
+		{
+			return InTrueFunction(InItem);
+		});
+	if (bAllMatch)
+	{
+		return ENavigationToolCompareState::AllTrue;
+	}
+
+	const bool bAllNoMatch = Algo::AllOf(Children,
+		[&InTrueFunction](const InItemType* const InItem)
+		{
+			return !InTrueFunction(InItem);
+		});
+	if (bAllNoMatch)
+	{
+		return ENavigationToolCompareState::NoneTrue;
+	}
+
+	return ENavigationToolCompareState::PartialTrue;
+}
+
+/**  */
+template<typename InItemType>
+ENavigationToolCompareState CompareArrayState(const TArray<InItemType*>& InArray
+	, const TFunctionRef<bool(const InItemType* const)>& InTrueFunction
+	, const TFunctionRef<bool(const InItemType* const)>& InFalseFunction)
+{
+	if (InArray.IsEmpty())
+	{
+		return ENavigationToolCompareState::AllTrue;
+	}
+
+	const bool bAllMatch = Algo::AllOf(InArray,
+		[&InTrueFunction](const InItemType* const InItem)
+		{
+			return InTrueFunction(InItem);
+		});
+	if (bAllMatch)
+	{
+		return ENavigationToolCompareState::AllTrue;
+	}
+
+	const bool bAllNoMatch = Algo::AllOf(InArray,
+		[&InFalseFunction](const InItemType* const InItem)
+		{
+			return InFalseFunction(InItem);
+		});
+	if (bAllNoMatch)
+	{
+		return ENavigationToolCompareState::NoneTrue;
+	}
+
+	return ENavigationToolCompareState::PartialTrue;
+}
+
+/**  */
+template<typename InItemType>
+ENavigationToolCompareState CompareArrayStateSimple(const TArray<InItemType*>& InArray
+	, const TFunctionRef<bool(const InItemType* const)>& InTrueFunction)
+{
+	if (InArray.IsEmpty())
+	{
+		return ENavigationToolCompareState::AllTrue;
+	}
+
+	const bool bAllMatch = Algo::AllOf(InArray,
+		[&InTrueFunction](const InItemType* const InItem)
+		{
+			return InTrueFunction(InItem);
+		});
+	if (bAllMatch)
+	{
+		return ENavigationToolCompareState::AllTrue;
+	}
+
+	const bool bAllNoMatch = Algo::AllOf(InArray,
+		[&InTrueFunction](const InItemType* const InItem)
+		{
+			return !InTrueFunction(InItem);
+		});
+	if (bAllNoMatch)
+	{
+		return ENavigationToolCompareState::NoneTrue;
+	}
+
+	return ENavigationToolCompareState::PartialTrue;
+}
+
+}
diff --git a/Engine/Source/Editor/Sequencer/Public/Tools/NavigationTool/Items/NavigationToolObject.h b/Engine/Source/Editor/Sequencer/Public/Tools/NavigationTool/Items/NavigationToolObject.h
new file mode 100644
--- /dev/null
+++ b/Engine/Source/Editor/Sequencer/Public/Tools/NavigationTool/Items/NavigationToolObject.h
@@ -0,0 +1,55 @@
+// Copyright Epic Games, Inc. All Rights Reserved.
+
+#pragma once
+
+#include "NavigationToolItem.h"
+#include "UObject/WeakObjectPtr.h"
+
+namespace UE::Sequencer::NavigationTool
+{
+
+/** Item Class that represents a UObject */
+class SEQUENCER_API FNavigationToolObject
+	: public FNavigationToolItem
+{
+public:
+	UE_NAVIGATIONTOOL_INHERITS_WITH_SUPER(FNavigationToolObject
+		, FNavigationToolItem);
+
+	FNavigationToolObject(INavigationTool& InTool
+		, const FNavigationToolItemPtr& InParentItem
+		, UObject* const InObject
+		, const FNavigationToolItemWeakPtr& InWeakReferencingItem = nullptr
+		, const FString& InReferencingId = FString());
+
+	//~ Begin INavigationToolItem
+	virtual bool IsItemValid() const override;
+	virtual void Select(FNavigationToolScopedSelection& InSelection) const override;
+	virtual bool IsSelected(const FNavigationToolScopedSelection& InSelection) const override;
+	virtual void OnSelect() override;
+	virtual void OnDoubleClick() override;
+	virtual FText GetDisplayName() const override;
+	virtual FText GetClassName() const override;
+	virtual FSlateIcon GetIcon() const override;
+	virtual FText GetIconTooltipText() const override;
+	virtual bool IsAllowedInTool() const override;
+	virtual void OnObjectsReplaced(const TMap<UObject*, UObject*>& InReplacementMap, bool bRecursive) override;
+	//~ End INavigationToolItem
+
+	UObject* GetObject() const;
+	void SetObject(UObject* InObject);
+
+protected:
+	//~Begin FNavigationToolItem
+	virtual FNavigationToolItemId CalculateItemId() const override;
+	//~End FNavigationToolItem
+	
+	virtual void SetObject_Impl(UObject* InObject);
+
+	TWeakObjectPtr<> WeakObject;
+
+	FNavigationToolItemWeakPtr WeakReferencingItem;
+	FString ReferencingId;
+};
+
+}
diff --git a/Engine/Source/Editor/Sequencer/Public/Tools/NavigationTool/Items/NavigationToolObjectReference.h b/Engine/Source/Editor/Sequencer/Public/Tools/NavigationTool/Items/NavigationToolObjectReference.h
new file mode 100644
--- /dev/null
+++ b/Engine/Source/Editor/Sequencer/Public/Tools/NavigationTool/Items/NavigationToolObjectReference.h
@@ -0,0 +1,86 @@
+// Copyright Epic Games, Inc. All Rights Reserved.
+
+#pragma once
+
+#include "NavigationToolObject.h"
+
+namespace UE::Sequencer::NavigationTool
+{
+
+class FNavigationToolObjectReference;
+
+/**
+ * Object Item that will be shared across multiple objects (e.g. Materials).
+ * This is not an Item added directly to the Navigation Tool directly but rather shown through its Reference Items.
+ * 
+ * For Example, for a given material that is used across multiple Items,
+ * the Material Object itself would be considered a FNavigationToolSharedObject
+ * but the actual Item in the Navigation Tool shown will be a FNavigationToolObjectReference.
+ * 
+ * @see FNavigationToolObjectReference
+ */
+class SEQUENCER_API FNavigationToolSharedObject : public FNavigationToolObject
+{
+public:
+	UE_NAVIGATIONTOOL_INHERITS_WITH_SUPER(FNavigationToolSharedObject, FNavigationToolObject);
+
+	FNavigationToolSharedObject(INavigationTool& InTool
+		, const FNavigationToolItemPtr& InParentItem
+		, UObject* const InObject);
+
+	//~ Begin INavigationToolItem
+	virtual bool IsAllowedInTool() const override { return false; }
+	virtual bool CanBeTopLevel() const override { return false; }
+	//~ End INavigationToolItem
+	
+	void AddReference(const TSharedRef<FNavigationToolObjectReference>& InObjectReference);
+	
+	void RemoveReference(const TSharedRef<FNavigationToolObjectReference>& InObjectReference);
+
+	TArray<FNavigationToolItemPtr> GetObjectReferences() const;
+
+private:
+	TArray<TWeakPtr<FNavigationToolObjectReference>> ObjectReferences;
+};
+
+/**
+ * Reference Item of FNavigationToolSharedObject that holds an Object that is likely used across other items in the Navigation Tool.
+ * @see FNavigationToolSharedObject
+ */
+class SEQUENCER_API FNavigationToolObjectReference : public FNavigationToolObject
+{
+public:
+	UE_NAVIGATIONTOOL_INHERITS_WITH_SUPER(FNavigationToolObjectReference, FNavigationToolObject);
+
+	FNavigationToolObjectReference(INavigationTool& InTool
+		, const FNavigationToolItemPtr& InParentItem
+		, UObject* InObject
+		, const FNavigationToolItemPtr& InReferencingItem
+		, const FString& InReferenceId);
+	
+	//~ Begin INavigationToolItem
+	virtual void OnItemRegistered() override;
+	virtual void OnItemUnregistered() override;
+	virtual bool IsItemValid() const override;
+	virtual FText GetDisplayName() const override;
+	//~ End INavigationToolItem
+
+	FNavigationToolItemPtr GetReferencingItem() const { return ReferencingItemWeak.Pin(); }
+	
+protected:
+	//~Begin FNavigationToolItem
+	virtual FNavigationToolItemId CalculateItemId() const override;
+	//~End FNavigationToolItem
+	
+	/** The item referencing the object (e.g. a Component item holding reference of a Material) */
+	FNavigationToolItemWeakPtr ReferencingItemWeak;
+	
+	/**
+	 * Id to distinguish the reference from other references held by the same Referencing Item.
+	 * This reference Id can be the for example a slot index of the material in a primitive component
+	 * or can be the uproperty name of the variable holding the reference
+	 */
+	FString ReferenceId;
+};
+
+}
diff --git a/Engine/Source/Editor/Sequencer/Public/Tools/NavigationTool/Items/NavigationToolSequence.h b/Engine/Source/Editor/Sequencer/Public/Tools/NavigationTool/Items/NavigationToolSequence.h
new file mode 100644
--- /dev/null
+++ b/Engine/Source/Editor/Sequencer/Public/Tools/NavigationTool/Items/NavigationToolSequence.h
@@ -0,0 +1,114 @@
+// Copyright Epic Games, Inc. All Rights Reserved.
+
+#pragma once
+
+#include "NavigationToolItem.h"
+#include "NavigationToolObject.h"
+#include "Tools/NavigationTool/Extensions/IMarkerVisibilityExtension.h"
+#include "Tools/NavigationTool/Extensions/IRenameableExtension.h"
+#include "Tools/NavigationTool/Extensions/IRevisionControlExtension.h"
+#include "Tools/NavigationTool/Extensions/ISequenceLockableExtension.h"
+#include "Tools/NavigationTool/Extensions/ISequenceInactivableExtension.h"
+
+class UMovieScene;
+class UMovieSceneSequence;
+class UMovieSceneSubSection;
+
+namespace UE::Sequencer
+{
+	class FSequenceModel;
+}
+
+namespace UE::Sequencer::NavigationTool
+{
+
+/**
+ * Item in Navigation Tool representing a Sequence.
+ */
+class FNavigationToolSequence
+	: public FNavigationToolObject
+	, public IRenameableExtension
+	, public ISequenceInactivableExtension
+	, public ISequenceLockableExtension
+	, public IMarkerVisibilityExtension
+	, public IRevisionControlExtension
+{
+public:
+	UE_NAVIGATIONTOOL_INHERITS_WITH_SUPER(FNavigationToolSequence
+		, FNavigationToolObject
+		, IRenameableExtension
+		, ISequenceInactivableExtension
+		, ISequenceLockableExtension
+		, IMarkerVisibilityExtension
+		, IRevisionControlExtension);
+
+	SEQUENCER_API FNavigationToolSequence(INavigationTool& InTool
+		, const FNavigationToolItemPtr& InParentItem
+		, UMovieSceneSequence* const InSequence
+		, UMovieSceneSubSection* const InSubSection
+		, const int32 InSubSectionIndex);
+
+	//~ Begin INavigationToolItem
+	SEQUENCER_API virtual void Select(FNavigationToolScopedSelection& InSelection) const override;
+	SEQUENCER_API virtual bool IsSelected(const FNavigationToolScopedSelection& InSelection) const override;
+	SEQUENCER_API virtual void FindChildren(TArray<FNavigationToolItemPtr>& OutChildren, const bool bInRecursive) override;
+	SEQUENCER_API virtual void GetItemProxies(TArray<TSharedPtr<FNavigationToolItemProxy>>& OutItemProxies) override;
+	SEQUENCER_API virtual bool AddChild(const FNavigationToolAddItemParams& InAddItemParams) override;
+	SEQUENCER_API virtual bool RemoveChild(const FNavigationToolRemoveItemParams& InRemoveItemParams) override;
+	SEQUENCER_API virtual ENavigationToolItemViewMode GetSupportedViewModes(const INavigationToolView& InToolView) const override;
+	virtual bool CanBeTopLevel() const override { return true; }
+	SEQUENCER_API virtual bool IsAllowedInTool() const override;
+	virtual FText GetDisplayName() const override;
+	SEQUENCER_API virtual FText GetClassName() const override;
+	virtual FSlateIcon GetIcon() const override;
+	SEQUENCER_API virtual FText GetIconTooltipText() const override;
+	SEQUENCER_API virtual void OnSelect() override;
+	SEQUENCER_API virtual void OnDoubleClick() override;
+	//~ End INavigationToolItem
+
+	SEQUENCER_API UMovieSceneSequence* GetSequence() const;
+	SEQUENCER_API UMovieSceneSubSection* GetSubSection() const;
+	SEQUENCER_API UMovieScene* GetSequenceMovieScene() const;
+
+	//~ Begin IRenameableExtension
+	SEQUENCER_API virtual bool CanRename() const override;
+	SEQUENCER_API virtual bool Rename(const FString& InName) override;
+	//~ End IRenameableExtension
+
+	//~ Begin ISequenceDeactivatableExtension
+	SEQUENCER_API virtual EItemSequenceInactiveState GetInactiveState() const override;
+	SEQUENCER_API virtual void SetIsInactive(const bool bInIsInactive) override;
+	//~ End ISequenceDeactivatableExtension
+
+	//~ Begin IMarkerVisibilityExtension
+	SEQUENCER_API virtual EItemMarkerVisibility GetMarkerVisibility() const override;
+	SEQUENCER_API virtual void SetMarkerVisibility(const bool bInVisible) override;
+	//~ End IMarkerVisibilityExtension
+
+	//~ Begin ISequenceLockableExtension
+	SEQUENCER_API virtual EItemSequenceLockState GetLockState() const override;
+	SEQUENCER_API virtual void SetIsLocked(const bool bInIsLocked) override;
+	//~ End ISequenceLockableExtension
+
+	//~ Begin IColorExtension
+	SEQUENCER_API virtual TOptional<FColor> GetColor() const override;
+	SEQUENCER_API virtual void SetColor(const TOptional<FColor>& InColor) override;
+	//~ End IColorExtension
+
+	//~ Begin IRevisionControlExtension
+	SEQUENCER_API virtual EItemRevisionControlState GetRevisionControlState() const override;
+	SEQUENCER_API virtual void SetIsRevisionControlled(const bool bInIsRevisionControlled) override;
+	SEQUENCER_API virtual const FSlateBrush* GetRevisionControlStatusIcon() const override;
+	SEQUENCER_API virtual FText GetRevisionControlStatusText() const override;
+	//~ End IRevisionControlExtension
+
+protected:
+	//~Begin FNavigationToolItem
+	SEQUENCER_API virtual FNavigationToolItemId CalculateItemId() const override;
+	//~End FNavigationToolItem
+
+	TWeakObjectPtr<UMovieSceneSubSection> WeakSubSection;
+	int32 SubSectionIndex = 0;
+};
+
+}
diff --git a/Engine/Source/Editor/Sequencer/Public/Tools/NavigationTool/Items/NavigationToolSubTrack.h b/Engine/Source/Editor/Sequencer/Public/Tools/NavigationTool/Items/NavigationToolSubTrack.h
new file mode 100644
--- /dev/null
+++ b/Engine/Source/Editor/Sequencer/Public/Tools/NavigationTool/Items/NavigationToolSubTrack.h
@@ -0,0 +1,62 @@
+﻿// Copyright Epic Games, Inc. All Rights Reserved.
+
+#pragma once
+
+#include "NavigationToolTrack.h"
+#include "Textures/SlateIcon.h"
+#include "Tools/NavigationTool/Extensions/IMarkerVisibilityExtension.h"
+#include "Tools/NavigationTool/Extensions/ISequenceInactivableExtension.h"
+
+class UMovieSceneSequence;
+class UMovieSceneSection;
+class UMovieSceneSubTrack;
+
+namespace UE::Sequencer
+{
+	class FTrackModel;
+}
+
+namespace UE::Sequencer::NavigationTool
+{
+
+/**
+ * 
+ */
+class SEQUENCER_API FNavigationToolSubTrack
+	: public FNavigationToolTrack
+	, public ISequenceInactivableExtension
+	, public IMarkerVisibilityExtension
+{
+public:
+	UE_NAVIGATIONTOOL_INHERITS_WITH_SUPER(FNavigationToolSubTrack
+		, FNavigationToolTrack
+		, ISequenceInactivableExtension
+		, IMarkerVisibilityExtension);
+
+	FNavigationToolSubTrack(INavigationTool& InTool
+		, const FNavigationToolItemPtr& InParentItem
+		, UMovieSceneSubTrack* const InSubTrack
+		, const TWeakObjectPtr<UMovieSceneSequence>& InSequence
+		, const TWeakObjectPtr<UMovieSceneSection>& InSection
+        , const int32 InSubSectionIndex);
+
+	//~ Begin INavigationToolItem
+	virtual void FindChildren(TArray<FNavigationToolItemPtr>& OutChildren, const bool bInRecursive) override;
+	virtual TOptional<EItemDropZone> CanAcceptDrop(const FDragDropEvent& InDragDropEvent, EItemDropZone InDropZone) override;
+	virtual FReply AcceptDrop(const FDragDropEvent& InDragDropEvent, EItemDropZone InDropZone) override;
+	//~ End INavigationToolItem
+
+	UMovieSceneSubTrack* GetSubTrack() const;
+
+	//~ Begin ISequenceDeactivatableExtension
+	virtual EItemSequenceInactiveState GetInactiveState() const override;
+	virtual void SetIsInactive(const bool bInIsInactive) override;
+	//~ End ISequenceDeactivatableExtension
+
+	//~ Begin IMarkerVisibilityExtension
+	virtual EItemMarkerVisibility GetMarkerVisibility() const override;
+	virtual void SetMarkerVisibility(const bool bInVisible) override;
+	//~ End IMarkerVisibilityExtension
+};
+
+}
diff --git a/Engine/Source/Editor/Sequencer/Public/Tools/NavigationTool/Items/NavigationToolTrack.h b/Engine/Source/Editor/Sequencer/Public/Tools/NavigationTool/Items/NavigationToolTrack.h
new file mode 100644
--- /dev/null
+++ b/Engine/Source/Editor/Sequencer/Public/Tools/NavigationTool/Items/NavigationToolTrack.h
@@ -0,0 +1,73 @@
+﻿// Copyright Epic Games, Inc. All Rights Reserved.
+
+#pragma once
+
+#include "MVVM/ViewModels/TrackModel.h"
+#include "NavigationToolObjectReference.h"
+#include "Textures/SlateIcon.h"
+#include "Tools/NavigationTool/Extensions/ISequenceLockableExtension.h"
+
+class UMovieSceneSequence;
+class UMovieSceneSection;
+class UMovieSceneTrack;
+
+namespace UE::Sequencer::NavigationTool
+{
+
+/**
+ * 
+ */
+class SEQUENCER_API FNavigationToolTrack
+	: public FNavigationToolObject
+	, public ISequenceLockableExtension
+{
+public:
+	UE_NAVIGATIONTOOL_INHERITS_WITH_SUPER(FNavigationToolTrack
+		, FNavigationToolObject
+		, ISequenceLockableExtension);
+
+	FNavigationToolTrack(INavigationTool& InTool
+		, const FNavigationToolItemPtr& InParentItem
+		, UMovieSceneTrack* const InTrack
+		, const TWeakObjectPtr<UMovieSceneSequence>& InSequence
+		, const TWeakObjectPtr<UMovieSceneSection>& InSection
+		, const int32 InSubSectionIndex);
+
+	//~ Begin INavigationToolItem
+	virtual void FindChildren(TArray<FNavigationToolItemPtr>& OutChildren, const bool bInRecursive) override;
+	virtual bool CanBeTopLevel() const override { return false; }
+	virtual TOptional<EItemDropZone> CanAcceptDrop(const FDragDropEvent& InDragDropEvent, EItemDropZone InDropZone) override;
+	virtual FReply AcceptDrop(const FDragDropEvent& InDragDropEvent, EItemDropZone InDropZone) override;
+	virtual FText GetDisplayName() const override;
+	virtual const FSlateBrush* GetDefaultIconBrush() const override;
+	virtual FSlateIcon GetIcon() const override;
+	virtual bool IsSelected(const FNavigationToolScopedSelection& InSelection) const override;
+	virtual void Select(FNavigationToolScopedSelection& InSelection) const override;
+	//~ End INavigationToolItem
+
+	//~ Begin ISequenceLockableExtension
+	virtual EItemSequenceLockState GetLockState() const override;
+	virtual void SetIsLocked(const bool bInIsLocked) override;
+	//~ End ISequenceLockableExtension
+
+	UMovieSceneTrack* GetTrack() const;
+
+protected:
+	//~ Begin FNavigationToolItem
+	virtual FNavigationToolItemId CalculateItemId() const override;
+	//~ End FNavigationToolItem
+
+	//~ Begin FNavigationToolObjectItem
+	virtual void SetObject_Impl(UObject* InObject) override;
+	//~ End FNavigationToolObjectItem
+
+	void OnTrackObjectChanged();
+
+	TWeakObjectPtr<UMovieSceneSequence> WeakSequence;
+	TWeakObjectPtr<UMovieSceneSection> WeakSection;
+	int32 SubSectionIndex = 0;
+
+	FSlateIcon Icon;
+};
+
+}
diff --git a/Engine/Source/Editor/Sequencer/Public/Tools/NavigationTool/Menus/NavigationToolFilterBarContext.h b/Engine/Source/Editor/Sequencer/Public/Tools/NavigationTool/Menus/NavigationToolFilterBarContext.h
new file mode 100644
--- /dev/null
+++ b/Engine/Source/Editor/Sequencer/Public/Tools/NavigationTool/Menus/NavigationToolFilterBarContext.h
@@ -0,0 +1,32 @@
+﻿// Copyright Epic Games, Inc. All Rights Reserved.
+
+#pragma once
+
+#include "Filters/Menus/SequencerFilterBarContext.h"
+#include "Templates/SharedPointer.h"
+#include "Tools/NavigationTool/Filters/INavigationToolFilterBar.h"
+#include "NavigationToolFilterBarContext.generated.h"
+
+class UToolMenu;
+
+UCLASS()
+class UNavigationToolFilterBarContext : public UObject
+{
+	GENERATED_BODY()
+
+public:
+	void Init(const TWeakPtr<UE::Sequencer::NavigationTool::INavigationToolFilterBar>& InFilterBarWeak)
+	{
+		WeakFilterBar = InFilterBarWeak;
+	}
+
+	TSharedPtr<UE::Sequencer::NavigationTool::INavigationToolFilterBar> GetFilterBar() const
+	{
+		return WeakFilterBar.IsValid() ? WeakFilterBar.Pin() : nullptr;
+	}
+
+	FOnPopulateFilterBarMenu OnPopulateMenu;
+
+protected:
+	TWeakPtr<UE::Sequencer::NavigationTool::INavigationToolFilterBar> WeakFilterBar;
+};
diff --git a/Engine/Source/Editor/Sequencer/Public/Tools/NavigationTool/Menus/NavigationToolFilterMenuContext.h b/Engine/Source/Editor/Sequencer/Public/Tools/NavigationTool/Menus/NavigationToolFilterMenuContext.h
new file mode 100644
--- /dev/null
+++ b/Engine/Source/Editor/Sequencer/Public/Tools/NavigationTool/Menus/NavigationToolFilterMenuContext.h
@@ -0,0 +1,30 @@
+﻿// Copyright Epic Games, Inc. All Rights Reserved.
+
+#pragma once
+
+#include "Filters/Menus/SequencerFilterBarContext.h"
+#include "Templates/SharedPointer.h"
+#include "Tools/NavigationTool/Filters/Filters/NavigationToolFilterBase.h"
+#include "NavigationToolFilterMenuContext.generated.h"
+
+UCLASS()
+class UNavigationToolFilterMenuContext : public UObject
+{
+	GENERATED_BODY()
+
+public:
+	void Init(const TWeakPtr<UE::Sequencer::NavigationTool::FNavigationToolFilter>& InFilter)
+	{
+		WeakFilter = InFilter;
+	}
+
+	TSharedPtr<UE::Sequencer::NavigationTool::FNavigationToolFilter> GetFilter() const
+	{
+		return WeakFilter.Pin();
+	}
+
+	FOnPopulateFilterBarMenu OnPopulateMenu;
+
+protected:
+	TWeakPtr<UE::Sequencer::NavigationTool::FNavigationToolFilter> WeakFilter;
+};
diff --git a/Engine/Source/Editor/Sequencer/Public/Tools/NavigationTool/Menus/NavigationToolItemMenuContext.h b/Engine/Source/Editor/Sequencer/Public/Tools/NavigationTool/Menus/NavigationToolItemMenuContext.h
new file mode 100644
--- /dev/null
+++ b/Engine/Source/Editor/Sequencer/Public/Tools/NavigationTool/Menus/NavigationToolItemMenuContext.h
@@ -0,0 +1,45 @@
+// Copyright Epic Games, Inc. All Rights Reserved.
+
+#pragma once
+
+#include "Containers/Array.h"
+#include "Containers/ContainersFwd.h"
+#include "Templates/SharedPointer.h"
+#include "Tools/NavigationTool/NavigationToolDefines.h"
+#include "UObject/Object.h"
+#include "NavigationToolItemMenuContext.generated.h"
+
+namespace UE::Sequencer::NavigationTool
+{
+class INavigationTool;
+}
+
+UCLASS(MinimalAPI)
+class UNavigationToolItemMenuContext : public UObject
+{
+	GENERATED_BODY()
+
+	friend class FNavigationToolView;
+
+public:
+	void Init(const TWeakPtr<UE::Sequencer::NavigationTool::INavigationTool>& InTool, const TArray<UE::Sequencer::NavigationTool::FNavigationToolItemPtr>& InItemList)
+	{
+		WeakTool = InTool;
+		WeakItemList.Append(InItemList);
+	}
+
+	TSharedPtr<UE::Sequencer::NavigationTool::INavigationTool> GetTool() const
+	{
+		return WeakTool.Pin();
+	}
+
+	TConstArrayView<UE::Sequencer::NavigationTool::FNavigationToolItemWeakPtr> GetItems() const
+	{
+		return WeakItemList;
+	}
+
+private:
+	TWeakPtr<UE::Sequencer::NavigationTool::INavigationTool> WeakTool;
+
+	TArray<UE::Sequencer::NavigationTool::FNavigationToolItemWeakPtr> WeakItemList;
+};
diff --git a/Engine/Source/Editor/Sequencer/Public/Tools/NavigationTool/NavigationToolCastableTypeTable.h b/Engine/Source/Editor/Sequencer/Public/Tools/NavigationTool/NavigationToolCastableTypeTable.h
new file mode 100644
--- /dev/null
+++ b/Engine/Source/Editor/Sequencer/Public/Tools/NavigationTool/NavigationToolCastableTypeTable.h
@@ -0,0 +1,249 @@
+﻿// Copyright Epic Games, Inc. All Rights Reserved.
+
+#pragma once
+
+#include "CoreTypes.h"
+#include "HAL/Platform.h"
+#include "Templates/UnrealTemplate.h"
+#include "Containers/SparseBitSet.h"
+#include "UObject/NameTypes.h"
+
+#include <type_traits>
+
+template<typename T>
+struct TViewModelTypeID;
+
+template<typename ...T>
+struct TImplements;
+
+class ICastable;
+class IDynamicExtension;
+
+namespace UE::Sequencer::NavigationTool
+{
+
+struct FNavigationToolCastableTypeTable
+{
+	explicit FNavigationToolCastableTypeTable(uint8* InTypeMaskStorage);
+
+	FNavigationToolCastableTypeTable(const FNavigationToolCastableTypeTable&) = delete;
+	void operator=(const FNavigationToolCastableTypeTable&) = delete;
+	FNavigationToolCastableTypeTable(FNavigationToolCastableTypeTable&&) = delete;
+	void operator=(FNavigationToolCastableTypeTable&&) = delete;
+
+	/**
+	 * Make a new type table for the templated type, ID and typename.
+	 * The new table is allocated using malloc and is intentionally leaked
+	 */
+	template<typename T>
+	static FNavigationToolCastableTypeTable* MakeTypeTable(const void* Unused, uint32 TypeID, FName InName)
+	{
+		return nullptr;
+	}
+	template<typename T>
+	static FNavigationToolCastableTypeTable* MakeTypeTable(const ICastable* Unused, uint32 TypeID, FName InName)
+	{
+		FNavigationToolCastableTypeTableGenerator Generator;
+		Generator.Generate<T>();
+		return Generator.Commit(TypeID, InName);
+	}
+	template<typename T>
+	static FNavigationToolCastableTypeTable* MakeTypeTable(const IDynamicExtension* Unused, uint32 TypeID, FName InName)
+	{
+		FNavigationToolCastableTypeTableGenerator Generator;
+		Generator.Generate<T>();
+		return Generator.Commit(TypeID, InName);
+	}
+
+	/**
+	 * Attempt to locate a type table by its C++ name
+	 */
+	SEQUENCERCORE_API static const FNavigationToolCastableTypeTable* FindTypeByName(FName InName);
+
+
+	/**
+	 * Use this type table to cast the specified base ptr to another type
+	 * @note BasePtr must be an instance of the type specified by this type table, or the result is undefined.
+	 *
+	 * @return BasePtr casted to the type ID specified by ToType, or nullptr if the cast failed
+	 */
+	const void* Cast(const ICastable* BasePtr, uint32 ToType) const
+	{
+		return CastImpl(BasePtr, ToType);
+	}
+
+
+	/**
+	 * Use this type table to cast the specified base ptr to another type
+	 * @note BasePtr must be an instance of the type specified by this type table, or the result is undefined.
+	 *
+	 * @return BasePtr casted to the type ID specified by ToType, or nullptr if the cast failed
+	 */
+	const void* Cast(const IDynamicExtension* BasePtr, uint32 ToType) const
+	{
+		return CastImpl(BasePtr, ToType);
+	}
+
+	/**
+	 * Retrieve the C++ type name that this type table relates to
+	 */
+	FName GetTypeName() const
+	{
+		return ThisTypeName;
+	}
+
+
+	/**
+	 * Retrieve the type ID for this type table.
+	 * Warning: This type ID is not deterministic or stable across different processes and should not be saved persistently
+	 */
+	uint32 GetTypeID() const
+	{
+		return ThisTypeID;
+	}
+
+private:
+
+	/**
+	 * Cast implementation function
+	 */
+	const void* CastImpl(const void* BasePtr, uint32 ToType) const
+	{
+		const int32 SparseTypeIndex = TypeMask.GetSparseBucketIndex(ToType);
+		if (SparseTypeIndex != INDEX_NONE)
+		{
+			return static_cast<const uint8*>(BasePtr) + StaticTypeOffsets[SparseTypeIndex];
+		}
+		return nullptr;
+	}
+
+	/** Utility class that generates a type table used for dynamic casting and instance of an ICastable object */ 
+	struct FNavigationToolCastableTypeTableGenerator
+	{
+		/**
+		 * Generate the type table from a template type
+		 */
+		template<typename T>
+		void Generate()
+		{
+			// Compute the start offset from a base ICastable to T
+			// This ensures we compute the offsets from the correct point
+			const int16 StartOffset = ComputeStartOffset<T>((T*)0);
+			AddStaticType<T>((T*)nullptr, StartOffset);
+			ConsumeImplements<T>((T*)nullptr, StartOffset);
+		}
+
+		/** Commit this generator to its final form */
+		SEQUENCERCORE_API FNavigationToolCastableTypeTable* Commit(uint32 TypeID, FName InName);
+
+		/**
+		 * Compute the ptr offset from one class to another in an inheritance hierarchy.
+		 * @note: This can be negative for situations where multiple inheritance would put
+		 *        the FromType class after the ToType in memory.
+		 */
+		template<typename FromType, typename ToType>
+		static int16 ClassOffsetFrom()
+		{
+			alignas(FromType) uint8 Buffer[sizeof(FromType)];
+
+			FromType* From = reinterpret_cast<FromType*>(Buffer);
+			ToType*   To   = static_cast<ToType*>(From);
+
+			const int32 Offset = static_cast<int32>((uint8*)To - (uint8*)From);
+			checkf(Offset >= std::numeric_limits<int16>::lowest() && Offset <= std::numeric_limits<int16>::max(), TEXT("Class offset is too large to be supported by ICastable"));
+			return static_cast<int16>(Offset);
+		}
+
+		/**
+		 * Compute the start offset from an ICastable type to the root of T
+		 * This offset is added to all other class offsets that are computed from T
+		 * so that we can always cast from an ICastable* to the concrete types
+		 */
+		template<typename T>
+		static int16 ComputeStartOffset(const ICastable*)
+		{
+			return ClassOffsetFrom<ICastable, T>();
+		}
+
+		/**
+		 * Compute the start offset from an IDynamicExtension type to the root of T
+		 * This offset is added to all other class offsets that are computed from T
+		 * so that we can always cast from an IDynamicExtension* to the concrete types
+		 */
+		template<typename T>
+		static int16 ComputeStartOffset(const IDynamicExtension*)
+		{
+			return ClassOffsetFrom<IDynamicExtension, T>();
+		}
+
+		uint32 ByteSize() const
+		{
+			return sizeof(FNavigationToolCastableTypeTable) + sizeof(int16)*StaticTypeOffsets.Num() + sizeof(uint8)*StaticTypeOffsets.Num();
+		}
+
+		template<typename StartClass, typename T, typename U = typename T::Implements>
+		void ConsumeImplements(T*, int16 StartOffset)
+		{
+			ConsumeImplementsImpl<StartClass>((typename T::Implements*)nullptr, StartOffset);
+		}
+
+		template<typename StartClass>
+		void ConsumeImplements(...)
+		{
+		}
+
+		template<typename StartClass, typename ...Types>
+		void ConsumeImplementsImpl(TImplements<Types...>*, int16 StartOffset)
+		{
+			(AddStaticType<StartClass>((Types*)nullptr, StartOffset), ...);
+			(ConsumeImplements<StartClass>((Types*)nullptr, StartOffset), ...);
+		}
+
+		template<typename StartClass, typename T>
+		void ConsumeImplements(void*, int16 StartOffset)
+		{
+		}
+
+		template<typename StartClass, typename T, typename U = decltype(T::ID)>
+		void AddStaticType(T*, int16 StartOffset)
+		{
+			using namespace UE::MovieScene;
+
+			TViewModelTypeID<T>& ID = T::ID;
+
+			int16 OffsetFromBase = StartOffset + ClassOffsetFrom<StartClass, T>();
+			if (TypeMask.SetBit(ID.GetTypeID()) == ESparseBitSetBitResult::NewlySet)
+			{
+				int32 SparseIndex = TypeMask.GetSparseBucketIndex(ID.GetTypeID());
+				StaticTypeOffsets.Insert(OffsetFromBase, SparseIndex);
+			}
+		}
+
+		template<typename StartClass>
+		constexpr void AddStaticType(void*, int16 StartOffset)
+		{
+		}
+
+	private:
+
+		/* Accumulator sparse bitset of types that are implemented on the type that Generate was called with */
+		MovieScene::TFixedSparseBitSet<uint32> TypeMask;
+		/* Type offsets from T to the ptr for each type within TypeMask */
+		TArray<int16> StaticTypeOffsets;
+	};
+
+private:
+
+	/** Ptr to the static type offsets allocated at the end of this struct. Size = NumStaticTypes */
+	int16* StaticTypeOffsets;
+	/** Fixed sparse bitset of types that this table contains. Storage is allocated at the end of this struct. */
+	MovieScene::TFixedSparseBitSet<uint32, MovieScene::TFixedSparseBitSetBucketStorage<uint8>> TypeMask;
+	/** The number of static types implemented within this type table. */
+	uint32 NumStaticTypes;
+	/** Unique static type ID for this type. */
+	uint32 ThisTypeID;
+	/** The name of the C++ type that this table relates to. */
+	FName ThisTypeName;
+};
+
+}
diff --git a/Engine/Source/Editor/Sequencer/Public/Tools/NavigationTool/NavigationToolDefines.h b/Engine/Source/Editor/Sequencer/Public/Tools/NavigationTool/NavigationToolDefines.h
new file mode 100644
--- /dev/null
+++ b/Engine/Source/Editor/Sequencer/Public/Tools/NavigationTool/NavigationToolDefines.h
@@ -0,0 +1,69 @@
+// Copyright Epic Games, Inc. All Rights Reserved.
+
+#pragma once
+
+#include "Containers/ContainersFwd.h"
+#include "Templates/SharedPointer.h"
+#include "UObject/ObjectMacros.h"
+
+class AActor;
+class USceneComponent;
+
+namespace UE::Sequencer::NavigationTool
+{
+
+class INavigationToolItem;
+
+using FNavigationToolItemPtr = TSharedPtr<INavigationToolItem>;
+using FNavigationToolItemWeakPtr = TWeakPtr<INavigationToolItem>;
+
+SEQUENCER_API FName GetToolBarMenuName();
+SEQUENCER_API FName GetItemContextMenuName();
+
+}
+
+/** Flags specifying how an Item should be Selected */
+enum class ENavigationToolItemSelectionFlags : uint8
+{
+	None = 0,
+	/** Append to the Current Selection*/
+	AppendToCurrentSelection = 1 << 0,
+	/** Signal Selection Change (e.g. to trigger on Selection Change Delegate) */
+	SignalSelectionChange = 1 << 1,
+	/** Auto-include the Items' Children in the Selection */
+	IncludeChildren = 1 << 2,
+	/** Whether to Scroll first Item into View */
+	ScrollIntoView = 1 << 3,
+};
+ENUM_CLASS_FLAGS(ENavigationToolItemSelectionFlags);
+
+enum class ENavigationToolIgnoreNotifyFlags : uint8
+{
+	None = 0,
+	/** Ignores automatically handling actor spawns, usually so that it is manually handled */
+	Spawn = 1 << 0,
+	/** Ignores automatically handling actor duplications, usually so that it is manually handled */
+	Duplication = 1 << 1,
+	All = 0xFF,
+};
+ENUM_CLASS_FLAGS(ENavigationToolIgnoreNotifyFlags);
+
+enum class ENavigationToolRenameAction : uint8
+{
+	None,
+	Requested,
+	Completed,
+	Cancelled,
+};
+
+enum class ENavigationToolExtensionPosition
+{
+	Before,
+	After,
+};
+
+enum class ENavigationToolDragDropActionType : uint8
+{
+	Move,
+	Copy,
+};
diff --git a/Engine/Source/Editor/Sequencer/Public/Tools/NavigationTool/NavigationToolItemType.h b/Engine/Source/Editor/Sequencer/Public/Tools/NavigationTool/NavigationToolItemType.h
new file mode 100644
--- /dev/null
+++ b/Engine/Source/Editor/Sequencer/Public/Tools/NavigationTool/NavigationToolItemType.h
@@ -0,0 +1,205 @@
+// Copyright Epic Games, Inc. All Rights Reserved.
+
+#pragma once
+
+#include "NavigationToolItemTypeConcepts.h"
+#include "NavigationToolItemTypeId.h"
+#include "NavigationToolItemTypeTraits.h"
+#include "Templates/AndOrNot.h"
+#include "Templates/Models.h"
+
+namespace UE::Sequencer::NavigationTool
+{
+
+/**
+ * Macro to provide type information (type id and super types).
+ * This should be used within the class/struct scope only.
+ * @see UE_NAVIGATIONTOOL_TYPE_EXTERNAL for providing type information outside the class/struct scope
+ */
+#define UE_NAVIGATIONTOOL_TYPE(InThisType, ...)                                                 \
+	using FNavigationToolItemInherits = TNavigationToolItemInherits<InThisType, ##__VA_ARGS__>; \
+	static FNavigationToolItemTypeId GetStaticTypeId()                                          \
+	{                                                                                           \
+		static const FNavigationToolItemTypeId TypeId(TEXT(#InThisType));                       \
+		return TypeId;                                                                          \
+	}                                                                                           \
+
+/**
+ * Macro to implement a specialization of TNavigationToolItemExternalType to specify the UE_NAVIGATIONTOOL_TYPE for a given type.
+ * This should only be done for types where a UE_NAVIGATIONTOOL_TYPE internal implementation is not possible or desired.
+ * @see UE_NAVIGATIONTOOL_TYPE
+ */
+#define UE_NAVIGATIONTOOL_TYPE_EXTERNAL(InExternalType, ...)   \
+	template<>                                                 \
+	struct TNavigationToolItemExternalType<InExternalType>     \
+	{                                                          \
+		UE_NAVIGATIONTOOL_TYPE(InExternalType, ##__VA_ARGS__); \
+	};                                                         \
+
+/**
+ * Macro to implement UE_NAVIGATIONTOOL_TYPE for a given type and the overrides for the INavigationToolItemTypeCastable interface.
+ * This should be used within a class/struct that ultimately inherits from INavigationToolItemTypeCastable
+ * @see INavigationToolItemTypeCastable
+ * @see UE_NAVIGATIONTOOL_TYPE
+ */
+#define UE_NAVIGATIONTOOL_INHERITS(InThisType, ...)                                        \
+	UE_NAVIGATIONTOOL_TYPE(InThisType, ##__VA_ARGS__)                                      \
+	virtual FNavigationToolItemTypeId GetTypeId() const override                           \
+	{                                                                                      \
+		return TNavigationToolItemType<InThisType>::GetTypeId();                           \
+	}                                                                                      \
+	virtual const void* CastTo_Impl(FNavigationToolItemTypeId InCastToType) const override \
+	{                                                                                      \
+		return FNavigationToolItemInherits::Cast(this, InCastToType);                      \
+	}                                                                                      \
+
+#define UE_NAVIGATIONTOOL_INHERITS_WITH_SUPER(InThisType, InSuperType, ...) \
+	UE_NAVIGATIONTOOL_INHERITS(InThisType, InSuperType, ##__VA_ARGS__)      \
+	using Super = InSuperType;                                              \
+
+template<typename T>
+struct TNavigationToolItemExternalType;
+
+/**
+ * Determines that a typename T is a valid "Navigation Tool Item Type" if a FNavigationToolItemTypeId can be retrieved by calling either
+ * T::GetStaticTypeId() (via UE_NAVIGATIONTOOL_TYPE) or TNavigationToolItemExternalType<T>::GetStaticTypeId() (via UE_NAVIGATIONTOOL_TYPE_EXTERNAL)
+ */
+template<typename T>
+using TIsValidNavigationToolItemType = TOr<TModels<CNavigationToolItemStaticTypeable, T>, TModels<CNavigationToolItemStaticTypeable, TNavigationToolItemExternalType<T>>>;
+
+/**
+ * Helps retrieve the FNavigationToolItemTypeId for a given type T that satisfies the TIsValidNavigationToolItemType requirement
+ * If this is requirement is not satisfied, TNavigationToolItemType<T>::GetTypeId will be undefined for T.
+ * @see TIsValidNavigationToolItemType
+ */
+template<typename T, bool = TIsValidNavigationToolItemType<T>::Value>
+struct TNavigationToolItemType
+{
+	static FNavigationToolItemTypeId GetTypeId()
+	{
+		if constexpr (TModels_V<CNavigationToolItemStaticTypeable, T>)
+		{
+			return T::GetStaticTypeId();
+		}
+		else if constexpr (TModels_V<CNavigationToolItemStaticTypeable, TNavigationToolItemExternalType<T>>)
+		{
+			return TNavigationToolItemExternalType<T>::GetStaticTypeId();
+		}
+		else
+		{
+			checkNoEntry();
+			return FNavigationToolItemTypeId::Invalid();
+		}
+	}
+};
+
+/** No support for types that do not have a way to get their FNavigationToolItemTypeId */
+template<typename T>
+struct TNavigationToolItemType<T, false>
+{
+};
+
+/**
+ * Holds the direct super types of a given type T.
+ * Also serves to Cast a provided pointer to a desired type id.
+ */
+template<typename T, typename... InSuperTypes>
+class TNavigationToolItemInherits
+{
+	static_assert(TIsValidNavigationToolItemType<T>::Value, "Typename does not satisfy the requirements to be an Navigation Tool Item Type. See: TIsValidNavigationToolItemType");
+	static_assert(TIsDerivedFromAll<T, InSuperTypes...>::Value, "Some provided super types are not base types of T");
+
+	template<typename, typename InSuperType, typename... InOtherSuperTypes>
+	static const void* Cast(const T* InPtr, const FNavigationToolItemTypeId& InCastToType)
+	{
+		// Check first super type
+		if (TNavigationToolItemType<InSuperType>::GetTypeId() == InCastToType)
+		{
+			return static_cast<const InSuperType*>(InPtr);
+		}
+		// Check rest of super types
+		else if (const void* CastedPtr = Cast<void, InOtherSuperTypes...>(InPtr, InCastToType))
+		{
+			return CastedPtr;
+		}
+		// Recurse up the super type
+		else if constexpr (TModels_V<CNavigationToolItemInheritable, InSuperType>)
+		{
+			return InSuperType::FNavigationToolItemInherits::Cast(InPtr, InCastToType);
+		}
+		// Recurse up the super type for External Types
+		else if constexpr (TModels_V<CNavigationToolItemInheritable, TNavigationToolItemExternalType<InSuperType>>)
+		{
+			return TNavigationToolItemExternalType<InSuperType>::FNavigationToolItemInherits::Cast(InPtr, InCastToType);
+		}
+		else
+		{
+			return nullptr;
+		}
+	}
+
+	template<typename>
+	static const void* Cast(const T* InPtr, const FNavigationToolItemTypeId& InCastToType)
+	{
+		return nullptr;
+	}
+
+public:
+	static const void* Cast(const T* InPtr, const FNavigationToolItemTypeId& InCastToType)
+	{
+		// Check that Cast To Type is valid
+		if (!InCastToType.IsValid())
+		{
+			return nullptr;
+		}
+		// Check whether Cast To Type is this Type
+		if (TNavigationToolItemType<T>::GetTypeId() == InCastToType)
+		{
+			return InPtr;
+		}
+		return Cast<void, InSuperTypes...>(InPtr, InCastToType);
+	}
+};
+
+/**
+ * Interface Class to perform actions like IsA, or CastTo
+ * @see UE_NAVIGATIONTOOL_INHERITS
+ */
+class INavigationToolItemTypeCastable
+{
+public:
+	UE_NAVIGATIONTOOL_TYPE(INavigationToolItemTypeCastable)
+
+	virtual ~INavigationToolItemTypeCastable() = default;
+
+	virtual FNavigationToolItemTypeId GetTypeId() const = 0;
+
+	virtual const void* CastTo_Impl(FNavigationToolItemTypeId InId) const = 0;
+
+	template<typename InType>
+	bool IsA() const
+	{
+		return !!this->CastTo<InType>();
+	}
+
+	template<typename InType>
+	bool IsExactlyA() const
+	{
+		return GetTypeId() == TNavigationToolItemType<InType>::GetTypeId();
+	}
+
+	template<typename InType>
+	InType* CastTo()
+	{
+		const INavigationToolItemTypeCastable* ConstThis = this;
+		return const_cast<InType*>(ConstThis->CastTo<InType>());
+	}
+
+	template<typename InType>
+	const InType* CastTo() const
+	{
+		return static_cast<const InType*>(this->CastTo_Impl(TNavigationToolItemType<InType>::GetTypeId()));
+	}
+};
+
+}
diff --git a/Engine/Source/Editor/Sequencer/Public/Tools/NavigationTool/NavigationToolItemTypeConcepts.h b/Engine/Source/Editor/Sequencer/Public/Tools/NavigationTool/NavigationToolItemTypeConcepts.h
new file mode 100644
--- /dev/null
+++ b/Engine/Source/Editor/Sequencer/Public/Tools/NavigationTool/NavigationToolItemTypeConcepts.h
@@ -0,0 +1,27 @@
+﻿// Copyright Epic Games, Inc. All Rights Reserved.
+
+#pragma once
+
+namespace UE::Sequencer::NavigationTool
+{
+
+/** Describes a type that has a GetStaticTypeId() static function that returns FNavigationToolItemTypeId */
+struct CNavigationToolItemStaticTypeable
+{
+	template <typename T>
+	auto Requires(class FNavigationToolItemTypeId& Value) -> decltype(
+		Value = T::GetStaticTypeId()
+	);
+};
+
+/**
+ * Describes a type with a FNavigationToolItemInherits typedef or using declaration
+ * @see UE_NAVIGATIONTOOL_TYPE macro
+ */
+struct CNavigationToolItemInheritable
+{
+	template<typename T>
+	auto Requires()->typename T::FNavigationToolItemInherits&;
+};
+
+}
diff --git a/Engine/Source/Editor/Sequencer/Public/Tools/NavigationTool/NavigationToolItemTypeId.h b/Engine/Source/Editor/Sequencer/Public/Tools/NavigationTool/NavigationToolItemTypeId.h
new file mode 100644
--- /dev/null
+++ b/Engine/Source/Editor/Sequencer/Public/Tools/NavigationTool/NavigationToolItemTypeId.h
@@ -0,0 +1,52 @@
+// Copyright Epic Games, Inc. All Rights Reserved.
+
+#pragma once
+
+#include "Containers/StringFwd.h"
+#include "UObject/NameTypes.h"
+
+namespace UE::Sequencer::NavigationTool
+{
+
+class FNavigationToolItemTypeId
+{
+public:
+	static FNavigationToolItemTypeId Invalid()
+	{
+		return FNavigationToolItemTypeId(NAME_None);
+	}
+
+	explicit FNavigationToolItemTypeId(FName InTypeName)
+		: Name(InTypeName)
+	{}
+
+	friend uint32 GetTypeHash(FNavigationToolItemTypeId InTypeId)
+	{
+		return GetTypeHash(InTypeId.Name);
+	}
+
+	bool operator==(FNavigationToolItemTypeId InOther) const
+	{
+		return this->Name == InOther.Name;
+	}
+
+	bool IsValid() const
+	{
+		return !Name.IsNone();
+	}
+
+	FName ToName() const
+	{
+		return Name;
+	}
+
+	FString ToString() const
+	{
+		return Name.ToString();
+	}
+
+private:
+	FName Name;
+};
+
+}
diff --git a/Engine/Source/Editor/Sequencer/Public/Tools/NavigationTool/NavigationToolItemTypeTraits.h b/Engine/Source/Editor/Sequencer/Public/Tools/NavigationTool/NavigationToolItemTypeTraits.h
new file mode 100644
--- /dev/null
+++ b/Engine/Source/Editor/Sequencer/Public/Tools/NavigationTool/NavigationToolItemTypeTraits.h
@@ -0,0 +1,33 @@
+// Copyright Epic Games, Inc. All Rights Reserved.
+
+#pragma once
+
+#include <type_traits>
+
+namespace UE::Sequencer::NavigationTool
+{
+
+/** Returns true if InDerivedType is a derived type of ALL (And) the given InBaseTypes */
+template<typename InDerivedType, typename... InBaseTypes>
+struct TIsDerivedFromAll;
+
+template<typename InDerivedType>
+struct TIsDerivedFromAll<InDerivedType>
+{
+	static constexpr bool Value = true;
+};
+
+template<typename InDerivedType, typename InBaseType>
+struct TIsDerivedFromAll<InDerivedType, InBaseType>
+{
+	static constexpr bool Value = std::is_base_of_v<InBaseType, InDerivedType>;
+};
+
+template<typename InDerivedType, typename InBaseType, typename... InOtherBaseTypes>
+struct TIsDerivedFromAll<InDerivedType, InBaseType, InOtherBaseTypes...>
+{
+	static constexpr bool Value = TIsDerivedFromAll<InDerivedType, InBaseType>::Value
+		&& TIsDerivedFromAll<InDerivedType, InOtherBaseTypes...>::Value;
+};
+
+}
