[Sequencer] Navigation Tool fix for Mac/Linux build of initial submit

diff --git a/Engine/Source/Editor/Sequencer/Private/Tools/NavigationTool/Customization/NavigationToolObjectIconCustomization.cpp b/Engine/Source/Editor/Sequencer/Private/Tools/NavigationTool/Customization/NavigationToolObjectIconCustomization.cpp
--- a/Engine/Source/Editor/Sequencer/Private/Tools/NavigationTool/Customization/NavigationToolObjectIconCustomization.cpp
+++ b/Engine/Source/Editor/Sequencer/Private/Tools/NavigationTool/Customization/NavigationToolObjectIconCustomization.cpp
@@ -30,7 +30,7 @@ FName FNavigationToolObjectIconCustomization::GetToolItemIdentifier() const
 
 bool FNavigationToolObjectIconCustomization::HasOverrideIcon(TSharedPtr<const FNavigationToolItem> InToolItem) const
 {
-	if (const FNavigationToolObject* ObjectItem = InToolItem->CastTo<FNavigationToolObject>())
+	/*if (const FNavigationToolObject* ObjectItem = InToolItem->CastTo<FNavigationToolObject>())
 	{
 		if (const UObject* Object = ObjectItem->GetObject())
 		{
@@ -39,7 +39,7 @@ bool FNavigationToolObjectIconCustomization::HasOverrideIcon(TSharedPtr<const FN
 				return ObjectClass->GetFName() == SupportedClassName;
 			}
 		}
-	}
+	}*/
 
 	return false;
 }
diff --git a/Engine/Source/Editor/Sequencer/Private/Tools/NavigationTool/Items/NavigationToolBinding.cpp b/Engine/Source/Editor/Sequencer/Private/Tools/NavigationTool/Items/NavigationToolBinding.cpp
--- a/Engine/Source/Editor/Sequencer/Private/Tools/NavigationTool/Items/NavigationToolBinding.cpp
+++ b/Engine/Source/Editor/Sequencer/Private/Tools/NavigationTool/Items/NavigationToolBinding.cpp
@@ -22,7 +22,7 @@ FNavigationToolBinding::FNavigationToolBinding(INavigationTool& InTool
 	, const FMovieSceneBinding& InBinding
 	, UObject* const InObject
 	, const FNavigationToolItemPtr& InReferencingItem)
-	: Super(InTool, InParentItem, InObject, InReferencingItem, Binding.GetObjectGuid().ToString())
+	: Super(InTool, InParentItem, InObject, InReferencingItem, InBinding.GetObjectGuid().ToString())
 	, Binding(InBinding)
 {
 }
diff --git a/Engine/Source/Editor/Sequencer/Private/Tools/NavigationTool/Utils/NavigationToolMiscUtils.cpp b/Engine/Source/Editor/Sequencer/Private/Tools/NavigationTool/Utils/NavigationToolMiscUtils.cpp
--- a/Engine/Source/Editor/Sequencer/Private/Tools/NavigationTool/Utils/NavigationToolMiscUtils.cpp
+++ b/Engine/Source/Editor/Sequencer/Private/Tools/NavigationTool/Utils/NavigationToolMiscUtils.cpp
@@ -30,7 +30,7 @@ FSourceControlStatePtr GetSourceControlState(UPackage* const InPackage)
 		return nullptr;
 	}
 
-	return SourceControlProvider.GetState({ InPackage }, EStateCacheUsage::Use);
+	return SourceControlProvider.GetState(InPackage, EStateCacheUsage::Use);
 }
 
 const FSlateBrush* GetSourceControlStatusBrush(UPackage* const InPackage)
@@ -69,7 +69,7 @@ void FocusSequence(const INavigationTool& InTool, UMovieSceneSequence& InSequenc
 		return;
 	}
 
-	const FMovieSceneSequenceIDRef SequenceID = EvaluationState->FindSequenceId(&InSequence);
+	FMovieSceneSequenceIDRef SequenceID = EvaluationState->FindSequenceId(&InSequence);
 
 	UMovieSceneSubSection* const SubSection = Sequencer->FindSubSection(SequenceID);
 	if (IsValid(SubSection))
diff --git a/Engine/Source/Editor/Sequencer/Private/Tools/NavigationTool/Utils/NavigationToolMovieSceneUtils.cpp b/Engine/Source/Editor/Sequencer/Private/Tools/NavigationTool/Utils/NavigationToolMovieSceneUtils.cpp
--- a/Engine/Source/Editor/Sequencer/Private/Tools/NavigationTool/Utils/NavigationToolMovieSceneUtils.cpp
+++ b/Engine/Source/Editor/Sequencer/Private/Tools/NavigationTool/Utils/NavigationToolMovieSceneUtils.cpp
@@ -17,7 +17,7 @@ UMovieSceneSubSection* FindSequenceSubSection(ISequencer& InSequencer, UMovieSce
 		return nullptr;
 	}
 
-	const FMovieSceneSequenceIDRef SequenceID = EvaluationState->FindSequenceId(InSequence);
+	FMovieSceneSequenceIDRef SequenceID = EvaluationState->FindSequenceId(InSequence);
 
 	return InSequencer.FindSubSection(SequenceID);
 }
