[Sequencer] Navigation Tool inital submit Linux build fix for && within || error

diff --git a/Engine/Source/Editor/Sequencer/Private/Tools/NavigationTool/Filters/Filters/NavigationToolBuiltInFilter.cpp b/Engine/Source/Editor/Sequencer/Private/Tools/NavigationTool/Filters/Filters/NavigationToolBuiltInFilter.cpp
--- a/Engine/Source/Editor/Sequencer/Private/Tools/NavigationTool/Filters/Filters/NavigationToolBuiltInFilter.cpp
+++ b/Engine/Source/Editor/Sequencer/Private/Tools/NavigationTool/Filters/Filters/NavigationToolBuiltInFilter.cpp
@@ -66,10 +66,12 @@ bool FNavigationToolBuiltInFilter::PassesFilter(const FNavigationToolFilterType
 
 	const ENavigationToolFilterMode FilterMode = FilterParams.GetFilterMode();
 
-	if (EnumHasAnyFlags(FilterMode, ENavigationToolFilterMode::MatchesType)
-		&& ItemObjectPassesFilter(InItem) || ItemTypePassesFilter(InItem))
+	if (EnumHasAnyFlags(FilterMode, ENavigationToolFilterMode::MatchesType))
 	{
-		return true;
+		if (ItemObjectPassesFilter(InItem) || ItemTypePassesFilter(InItem))
+		{
+			return true;
+		}
 	}
 
 	if (EnumHasAnyFlags(FilterMode, ENavigationToolFilterMode::ContainerOfType))
@@ -87,7 +89,7 @@ bool FNavigationToolBuiltInFilter::PassesFilter(const FNavigationToolFilterType
 				continue;
 			}
 
-			if (ItemObjectPassesFilter(Item) || ItemTypePassesFilter(InItem))
+			if (ItemObjectPassesFilter(Item) || ItemTypePassesFilter(Item))
 			{
 				return true;
 			}
