Sequencer: Follow up to 40944960. Users prefer that Link Filters with Curve Editor is false by default

#rnx

diff --git a/Engine/Source/Editor/Sequencer/Private/SequencerSettings.cpp b/Engine/Source/Editor/Sequencer/Private/SequencerSettings.cpp
--- a/Engine/Source/Editor/Sequencer/Private/SequencerSettings.cpp
+++ b/Engine/Source/Editor/Sequencer/Private/SequencerSettings.cpp
@@ -34,7 +34,7 @@ USequencerSettings::USequencerSettings( const FObjectInitializer& ObjectInitiali
 	ZoomPosition = ESequencerZoomPosition::SZP_CurrentTime;
 	bAutoScrollEnabled = false;
 	bLinkCurveEditorTimeRange = false;
-	bLinkFiltersWithCurveEditor = true;
+	bLinkFiltersWithCurveEditor = false;
 	bSynchronizeCurveEditorSelection = true;
 	bIsolateCurveEditorToSelection = true;
 	bCurveEditorVisible = false;
