Add Menu Hooks to the snapping menu in the sequencer tool bar to add custom snapping option UI

#rnx

diff --git a/Engine/Source/Editor/Sequencer/Private/SSequencer.cpp b/Engine/Source/Editor/Sequencer/Private/SSequencer.cpp
--- a/Engine/Source/Editor/Sequencer/Private/SSequencer.cpp
+++ b/Engine/Source/Editor/Sequencer/Private/SSequencer.cpp
@@ -2506,38 +2506,44 @@ TSharedRef<SWidget> SSequencer::MakeRenderMovieMenu()
 
 TSharedRef<SWidget> SSequencer::MakeSnapMenu()
 {
-	FMenuBuilder MenuBuilder( false, SequencerPtr.Pin()->GetCommandBindings() );
-
-	MenuBuilder.BeginSection( "KeyAndSectionsSnapping", LOCTEXT( "SnappingMenuKeyAndSectionsHeader", "Key and Sections Snapping" ) );
+	
+	static const FName MenuName("Sequencer.SnapOptions");
+	if (!UToolMenus::Get()->IsMenuRegistered(MenuName))
 	{
-		MenuBuilder.AddMenuEntry( FSequencerCommands::Get().ToggleSnapKeyTimesToKeys );
-		MenuBuilder.AddMenuEntry( FSequencerCommands::Get().ToggleSnapSectionTimesToSections );
-		MenuBuilder.AddMenuEntry( FSequencerCommands::Get().ToggleSnapKeysAndSectionsToPlayRange );
-	}
-	MenuBuilder.EndSection();
+		UToolMenu* Toolbar = UToolMenus::Get()->RegisterMenu(MenuName, NAME_None, EMultiBoxType::Menu);
+		Toolbar->bSearchable = false;
+		{
+			FToolMenuSection& Section = Toolbar->AddSection( "KeyAndSectionsSnapping", LOCTEXT( "SnappingMenuKeyAndSectionsHeader", "Key and Sections Snapping" ) );
+			Section.AddMenuEntry( FSequencerCommands::Get().ToggleSnapKeyTimesToKeys );
+			Section.AddMenuEntry( FSequencerCommands::Get().ToggleSnapSectionTimesToSections );
+			Section.AddMenuEntry( FSequencerCommands::Get().ToggleSnapKeysAndSectionsToPlayRange );
+		}
 
-	MenuBuilder.BeginSection( "PlayHeadSnapping", LOCTEXT( "SnappingMenuPlayHeadHeader", "Play Head Snapping" ) );
-	{
-		MenuBuilder.AddMenuEntry( FSequencerCommands::Get().ToggleSnapPlayTimeToKeys );
-		MenuBuilder.AddMenuEntry( FSequencerCommands::Get().ToggleSnapPlayTimeToSections );
-		MenuBuilder.AddMenuEntry( FSequencerCommands::Get().ToggleSnapPlayTimeToMarkers );
-		MenuBuilder.AddMenuEntry( FSequencerCommands::Get().ToggleSnapPlayTimeToPressedKey );
-		MenuBuilder.AddMenuEntry( FSequencerCommands::Get().ToggleSnapPlayTimeToDraggedKey );
-	}
-	MenuBuilder.EndSection();
+		{
+			FToolMenuSection& Section = Toolbar->AddSection( "PlayHeadSnapping", LOCTEXT( "SnappingMenuPlayHeadHeader", "Play Head Snapping" ) );
+			Section.AddMenuEntry( FSequencerCommands::Get().ToggleSnapPlayTimeToKeys );
+			Section.AddMenuEntry( FSequencerCommands::Get().ToggleSnapPlayTimeToSections );
+			Section.AddMenuEntry( FSequencerCommands::Get().ToggleSnapPlayTimeToMarkers );
+			Section.AddMenuEntry( FSequencerCommands::Get().ToggleSnapPlayTimeToPressedKey );
+			Section.AddMenuEntry( FSequencerCommands::Get().ToggleSnapPlayTimeToDraggedKey );
+		}
 
-	MenuBuilder.BeginSection( "CurveSnapping", LOCTEXT( "SnappingMenuCurveHeader", "Curve Snapping" ) );
-	{
-		MenuBuilder.AddMenuEntry( FSequencerCommands::Get().ToggleSnapCurveValueToInterval );
+		{
+			FToolMenuSection& Section = Toolbar->AddSection( "CurveSnapping", LOCTEXT( "SnappingMenuCurveHeader", "Curve Snapping" ) );
+			Section.AddMenuEntry( FSequencerCommands::Get().ToggleSnapCurveValueToInterval );
+		}
+		
+		{
+			FToolMenuSection& Section = Toolbar->AddSection( "IntervalSnapping", LOCTEXT("SnappingMenuIntervalHeader", "Force Whole Frames") );
+			Section.AddMenuEntry(FSequencerCommands::Get().ToggleForceWholeFrames);
+		}
 	}
-	MenuBuilder.EndSection();
 
-	MenuBuilder.BeginSection( "IntervalSnapping", LOCTEXT("SnappingMenuIntervalHeader", "Force Whole Frames"));
-	{
-		MenuBuilder.AddMenuEntry(FSequencerCommands::Get().ToggleForceWholeFrames);
-	}
-	MenuBuilder.EndSection();
-	return MenuBuilder.MakeWidget();
+	USequencerToolMenuContext* ContextObject = NewObject<USequencerToolMenuContext>();
+	ContextObject->WeakSequencer = SequencerPtr;
+	FToolMenuContext Context(SequencerPtr.Pin()->GetCommandBindings(), nullptr, ContextObject);
+
+	return UToolMenus::Get()->GenerateWidget(MenuName, Context);
 }
 
 TSharedRef<SWidget> SSequencer::MakeTimeRange(const TSharedRef<SWidget>& InnerContent, bool bShowWorkingRange, bool bShowViewRange, bool bShowPlaybackRange)
diff --git a/Engine/Source/Editor/Sequencer/Private/SequencerTimeSliderController.cpp b/Engine/Source/Editor/Sequencer/Private/SequencerTimeSliderController.cpp
--- a/Engine/Source/Editor/Sequencer/Private/SequencerTimeSliderController.cpp
+++ b/Engine/Source/Editor/Sequencer/Private/SequencerTimeSliderController.cpp
@@ -111,6 +111,8 @@ FSequencerTimeSliderController::~FSequencerTimeSliderController()
 
 FFrameTime FSequencerTimeSliderController::SnapSequencerTime(FFrameTime InTime) const
 {
+	using namespace UE::Sequencer;
+	
 	TSharedPtr<FSequencer> Sequencer = WeakSequencer.Pin();
 	if (!Sequencer.IsValid())
 	{
@@ -130,6 +132,14 @@ FFrameTime FSequencerTimeSliderController::SnapSequencerTime(FFrameTime InTime)
 			InTime = FFrameRate::Snap(InTime, GetTickResolution(), GetDisplayRate());
 		}
 	}
+	if (Sequencer->GetSequencerSettings()->GetIsSnapEnabled())
+	{
+		TViewModelPtr<IClockExtension> Clock = Sequencer->GetViewModel()->GetRootSequenceModel().ImplicitCast();
+		if (Clock.IsValid() && Clock->SupportsSnapping() && Clock->ShouldSnapFrameTime())
+		{
+			InTime = Clock->SnapFrameTime(InTime);
+		}
+	}
 
 	return InTime;
 }
diff --git a/Engine/Source/Editor/Sequencer/Private/Tools/SequencerSnapField.cpp b/Engine/Source/Editor/Sequencer/Private/Tools/SequencerSnapField.cpp
--- a/Engine/Source/Editor/Sequencer/Private/Tools/SequencerSnapField.cpp
+++ b/Engine/Source/Editor/Sequencer/Private/Tools/SequencerSnapField.cpp
@@ -12,6 +12,7 @@
 #include "MVVM/ViewModels/ViewModel.h"
 #include "MVVM/Views/SOutlinerView.h"
 #include "MVVM/Extensions/ISnappableExtension.h"
+#include "MVVM/Extensions/IClockExtension.h"
 #include "IKeyArea.h"
 
 struct FSnapGridVisitor : ISequencerEntityVisitor, UE::Sequencer::ISnapField
@@ -113,6 +114,15 @@ void FSequencerSnapField::Initialize(const FSequencer& InSequencer, UE::Sequence
 {
 	using namespace UE::Sequencer;
 
+	if (InSequencer.GetSequencerSettings()->GetIsSnapEnabled())
+	{
+		TViewModelPtr<IClockExtension> Extension = InSequencer.GetViewModel()->GetRootSequenceModel().ImplicitCast();
+		if (Extension.IsValid() && Extension->SupportsSnapping() && Extension->ShouldSnapFrameTime())
+		{
+			ClockExtension = Extension;
+		}
+	}
+
 	TSharedPtr<SOutlinerView> TreeView = StaticCastSharedRef<SSequencer>(InSequencer.GetSequencerWidget())->GetTreeView();
 
 	TArray<TViewModelPtr<IOutlinerExtension>> VisibleItems;
@@ -290,6 +300,22 @@ TOptional<FSequencerSnapField::FSnapResult> FSequencerSnapField::Snap(const TArr
 			ThisSnap = Snap(Time, Threshold);
 		}
 
+		if (ClockExtension.IsValid())
+		{
+			FFrameNumber ClockSnapTime = ClockExtension->SnapFrameTime(Time).FrameNumber;
+			if (ClockSnapTime != Time)
+			{
+				if (ThisSnap.IsSet())
+				{
+					ThisSnap->SnappedTime = ClockSnapTime;
+				}
+				else
+				{
+					ThisSnap = { Time, ClockSnapTime, 1.f };
+				}
+			}
+		}
+
 		if (bSnapToInterval)
 		{
 			FFrameTime ThisTime = ThisSnap.IsSet() ? ThisSnap->SnappedTime : Time;
diff --git a/Engine/Source/Editor/Sequencer/Private/Tools/SequencerSnapField.h b/Engine/Source/Editor/Sequencer/Private/Tools/SequencerSnapField.h
--- a/Engine/Source/Editor/Sequencer/Private/Tools/SequencerSnapField.h
+++ b/Engine/Source/Editor/Sequencer/Private/Tools/SequencerSnapField.h
@@ -10,6 +10,11 @@
 class IKeyArea;
 class FSequencer;
 
+namespace UE::Sequencer
+{
+	class IClockExtension;
+}
+
 enum class ESequencerScrubberStyle : uint8;
 
 /** A snapping field that provides efficient snapping calculations on a range of values */
@@ -57,6 +62,8 @@ private:
 	FFrameRate DisplayRate;
 	ESequencerScrubberStyle ScrubStyle;
 
+	UE::Sequencer::TViewModelPtr<UE::Sequencer::IClockExtension> ClockExtension;
+
 	bool bSnapToInterval = false;
 	bool bSnapToLikeTypes = false;
 };
diff --git a/Engine/Source/Editor/Sequencer/Public/MVVM/Extensions/IClockExtension.h b/Engine/Source/Editor/Sequencer/Public/MVVM/Extensions/IClockExtension.h
--- a/Engine/Source/Editor/Sequencer/Public/MVVM/Extensions/IClockExtension.h
+++ b/Engine/Source/Editor/Sequencer/Public/MVVM/Extensions/IClockExtension.h
@@ -15,6 +15,7 @@ namespace UE::Sequencer
 class IClockExtension
 {
 public:
+	
 
 	UE_SEQUENCER_DECLARE_VIEW_MODEL_TYPE_ID_API(SEQUENCER_API, IClockExtension)
 
@@ -45,6 +46,18 @@ public:
 	{
 		return true;
 	}
+	virtual bool SupportsSnapping() const
+	{
+		return false;
+	}
+	virtual bool ShouldSnapFrameTime() const
+	{
+		return SupportsSnapping();
+	}
+	virtual FFrameTime SnapFrameTime(FFrameTime InFrameTime) const
+	{
+		return InFrameTime;
+	}
 };
 
 } // namespace UE::Sequencer
