Sequencer: Fix empty LOCTEXT

#rnx

diff --git a/Engine/Source/Editor/Sequencer/Private/Tools/NavigationTool/Widgets/ModalTextInputDialog.cpp b/Engine/Source/Editor/Sequencer/Private/Tools/NavigationTool/Widgets/ModalTextInputDialog.cpp
--- a/Engine/Source/Editor/Sequencer/Private/Tools/NavigationTool/Widgets/ModalTextInputDialog.cpp
+++ b/Engine/Source/Editor/Sequencer/Private/Tools/NavigationTool/Widgets/ModalTextInputDialog.cpp
@@ -1,4 +1,4 @@
-﻿// Copyright Epic Games, Inc. All Rights Reserved.
+// Copyright Epic Games, Inc. All Rights Reserved.
 
 #include "ModalTextInputDialog.h"
 #include "Framework/Application/SlateApplication.h"
@@ -184,7 +184,7 @@ FReply SModalTextInputDialog::OnCancelClick()
 FModalTextInputDialog::FModalTextInputDialog()
 {
 	DialogTitle = LOCTEXT("DefaultDialogTitle", "Text Input");
-	InputLabel = LOCTEXT("DefaultInputLabel", "");
+	InputLabel = FText::GetEmpty();
 }
 
 bool FModalTextInputDialog::Open(const FString& InDefaultText, FString& OutText, const TSharedPtr<SWindow>& InParentWindow)
