[Sequencer] Add "BindingName" text expression filter

diff --git a/Engine/Source/Editor/Sequencer/Private/Filters/Filters/SequencerTrackFilter_Text.cpp b/Engine/Source/Editor/Sequencer/Private/Filters/Filters/SequencerTrackFilter_Text.cpp
--- a/Engine/Source/Editor/Sequencer/Private/Filters/Filters/SequencerTrackFilter_Text.cpp
+++ b/Engine/Source/Editor/Sequencer/Private/Filters/Filters/SequencerTrackFilter_Text.cpp
@@ -6,12 +6,16 @@
 #include "Filters/SequencerTrackFilterTextExpressionExtension.h"
 #include "Misc/TextFilterExpressionEvaluator.h"
 #include "MVVM/ViewModelPtr.h"
-#include "Filters/TextExpressions/SequencerTextFilterExpression_Keyed.h"
+#include "Filters/TextExpressions/SequencerTextFilterExpression_BindingName.h"
+#include "Filters/TextExpressions/SequencerTextFilterExpression_BindingType.h"
 #include "Filters/TextExpressions/SequencerTextFilterExpression_Condition.h"
 #include "Filters/TextExpressions/SequencerTextFilterExpression_ConditionClass.h"
 #include "Filters/TextExpressions/SequencerTextFilterExpression_ConditionFunc.h"
 #include "Filters/TextExpressions/SequencerTextFilterExpression_ConditionPasses.h"
+#include "Filters/TextExpressions/SequencerTextFilterExpression_CustomBinding.h"
+#include "Filters/TextExpressions/SequencerTextFilterExpression_EmptyBinding.h"
 #include "Filters/TextExpressions/SequencerTextFilterExpression_Group.h"
+#include "Filters/TextExpressions/SequencerTextFilterExpression_Keyed.h"
 #include "Filters/TextExpressions/SequencerTextFilterExpression_Level.h"
 #include "Filters/TextExpressions/SequencerTextFilterExpression_Locked.h"
 #include "Filters/TextExpressions/SequencerTextFilterExpression_Modified.h"
@@ -24,9 +28,6 @@
 #include "Filters/TextExpressions/SequencerTextFilterExpression_Time.h"
 #include "Filters/TextExpressions/SequencerTextFilterExpression_TrackClass.h"
 #include "Filters/TextExpressions/SequencerTextFilterExpression_Unbound.h"
-#include "Filters/TextExpressions/SequencerTextFilterExpression_BindingType.h"
-#include "Filters/TextExpressions/SequencerTextFilterExpression_CustomBinding.h"
-#include "Filters/TextExpressions/SequencerTextFilterExpression_EmptyBinding.h"
 #include "UObject/UObjectIterator.h"
 
 #define LOCTEXT_NAMESPACE "SequencerTrackFilter_Text"
@@ -61,6 +62,7 @@ FSequencerTrackFilter_Text::FSequencerTrackFilter_Text(ISequencerTrackFilters& I
 	TextFilterExpressionContexts.Add(MakeShared<FSequencerTextFilterExpression_Muted>(InFilterInterface));
 	TextFilterExpressionContexts.Add(MakeShared<FSequencerTextFilterExpression_Soloed>(InFilterInterface));
 	TextFilterExpressionContexts.Add(MakeShared<FSequencerTextFilterExpression_Tag>(InFilterInterface));
+	TextFilterExpressionContexts.Add(MakeShared<FSequencerTextFilterExpression_BindingName>(InFilterInterface));
 	TextFilterExpressionContexts.Add(MakeShared<FSequencerTextFilterExpression_BindingType>(InFilterInterface));
 	TextFilterExpressionContexts.Add(MakeShared<FSequencerTextFilterExpression_CustomBinding>(InFilterInterface));
 	TextFilterExpressionContexts.Add(MakeShared<FSequencerTextFilterExpression_EmptyBinding>(InFilterInterface));
diff --git a/Engine/Source/Editor/Sequencer/Private/Filters/TextExpressions/SequencerTextFilterExpression_BindingName.cpp b/Engine/Source/Editor/Sequencer/Private/Filters/TextExpressions/SequencerTextFilterExpression_BindingName.cpp
new file mode 100644
--- /dev/null
+++ b/Engine/Source/Editor/Sequencer/Private/Filters/TextExpressions/SequencerTextFilterExpression_BindingName.cpp
@@ -0,0 +1,50 @@
+﻿// Copyright Epic Games, Inc. All Rights Reserved.
+
+#include "SequencerTextFilterExpression_BindingName.h"
+
+using namespace UE::Sequencer;
+
+#define LOCTEXT_NAMESPACE "SequencerTextFilterExpression_BindingName"
+
+FSequencerTextFilterExpression_BindingName::FSequencerTextFilterExpression_BindingName(ISequencerTrackFilters& InFilterInterface)
+	: FSequencerTextFilterExpressionContext(InFilterInterface)
+{
+}
+
+TSet<FName> FSequencerTextFilterExpression_BindingName::GetKeys() const
+{
+	return { TEXT("BindingName") };
+}
+
+ESequencerTextFilterValueType FSequencerTextFilterExpression_BindingName::GetValueType() const
+{
+	return ESequencerTextFilterValueType::String;
+}
+
+FText FSequencerTextFilterExpression_BindingName::GetDescription() const
+{
+	return LOCTEXT("ExpressionDescription_BindingName", "Filter by track object binding name");
+}
+
+bool FSequencerTextFilterExpression_BindingName::TestComplexExpression(const FName& InKey
+	, const FTextFilterString& InValue
+	, const ETextFilterComparisonOperation InComparisonOperation
+	, const ETextFilterTextComparisonMode InTextComparisonMode) const
+{
+	if (!FSequencerTextFilterExpressionContext::TestComplexExpression(InKey, InValue, InComparisonOperation, InTextComparisonMode))
+	{
+		return true;
+	}
+
+	UObject* const BoundObject = GetBoundObject();
+	if (IsValid(BoundObject))
+	{
+		const FString ObjectName = BoundObject->GetName();
+		const bool bPassed = TextFilterUtils::TestComplexExpression(ObjectName, InValue, InComparisonOperation, InTextComparisonMode);
+		return bPassed;
+	}
+
+	return false;
+}
+
+#undef LOCTEXT_NAMESPACE
diff --git a/Engine/Source/Editor/Sequencer/Private/Filters/TextExpressions/SequencerTextFilterExpression_BindingName.h b/Engine/Source/Editor/Sequencer/Private/Filters/TextExpressions/SequencerTextFilterExpression_BindingName.h
new file mode 100644
--- /dev/null
+++ b/Engine/Source/Editor/Sequencer/Private/Filters/TextExpressions/SequencerTextFilterExpression_BindingName.h
@@ -0,0 +1,24 @@
+﻿// Copyright Epic Games, Inc. All Rights Reserved.
+
+#pragma once
+
+#include "Filters/SequencerTextFilterExpressionContext.h"
+
+class FSequencerTextFilterExpression_BindingName : public FSequencerTextFilterExpressionContext
+{
+public:
+	FSequencerTextFilterExpression_BindingName(ISequencerTrackFilters& InFilterInterface);
+
+	//~ Begin FSequencerTextFilterExpressionContext
+	virtual TSet<FName> GetKeys() const override;
+	virtual ESequencerTextFilterValueType GetValueType() const override;
+	virtual FText GetDescription() const override;
+	//~ End FSequencerTextFilterExpressionContext
+
+	//~ Begin ITextFilterExpressionContext
+	virtual bool TestComplexExpression(const FName& InKey
+		, const FTextFilterString& InValue
+		, const ETextFilterComparisonOperation InComparisonOperation
+		, const ETextFilterTextComparisonMode InTextComparisonMode) const override;
+	//~ End ITextFilterExpressionContext
+};
