[Sequencer] Add Keyboard Shortcuts for Track States based on Track Selection

diff --git a/Engine/Source/Editor/Sequencer/Private/Sequencer.cpp b/Engine/Source/Editor/Sequencer/Private/Sequencer.cpp
--- a/Engine/Source/Editor/Sequencer/Private/Sequencer.cpp
+++ b/Engine/Source/Editor/Sequencer/Private/Sequencer.cpp
@@ -11970,6 +11970,22 @@ void FSequencer::BindCommands()
 		FExecuteAction::CreateSP(this, &FSequencer::AlignSelectionToPlayhead),
 		FCanExecuteAction::CreateSP(this, &FSequencer::CanAlignSelectionToPlayhead));
 
+	SequencerCommandBindings->MapAction(
+		Commands.ToggleTrackSelectionPin,
+		FExecuteAction::CreateSP(this, &FSequencer::ToggleTrackSelectionPin));
+
+	SequencerCommandBindings->MapAction(
+		Commands.ToggleTrackSelectionLock,
+		FExecuteAction::CreateSP(this, &FSequencer::ToggleTrackSelectionLock));
+
+	SequencerCommandBindings->MapAction(
+		Commands.ToggleTrackSelectionMute,
+		FExecuteAction::CreateSP(this, &FSequencer::ToggleTrackSelectionMute));
+
+	SequencerCommandBindings->MapAction(
+		Commands.ToggleTrackSelectionSolo,
+		FExecuteAction::CreateSP(this, &FSequencer::ToggleTrackSelectionSolo));
+
 	if (NavigationTool.IsValid())
 	{
 		const TSharedRef<UE::Sequencer::NavigationTool::FNavigationTool> NavigationToolRef = NavigationTool.ToSharedRef();
@@ -12721,16 +12737,96 @@ void FSequencer::SortSelectedOutlinerItems(const bool bInSortBySelectOrder, cons
 	SequencerHelpers::SortOutlinerItems(*this, OutlinerItems, bInSortBySelectOrder, bInDescending);
 }
 
-bool FSequencer::CanAlignSelectionToPlayhead()
+bool FSequencer::CanAlignSelectionToPlayhead() const
 {
 	return FSequencerSelectionAlignmentUtils::CanAlignSelection(*this);
 }
 
-void FSequencer::AlignSelectionToPlayhead()
+void FSequencer::AlignSelectionToPlayhead() const
 {
 	FSequencerSelectionAlignmentUtils::AlignSelectionToPlayhead(*this);
 }
 
+void FSequencer::ToggleTrackSelectionPin()
+{
+	using namespace UE::Sequencer;
+
+	FOutlinerSelection& OutlinerSelection = GetSelection().Outliner;
+
+	const bool bAllPinned = Algo::AllOf(OutlinerSelection, [](const FViewModelPtr& Selection)
+		{
+			const TViewModelPtr<IPinnableExtension> Pinnable = Selection.ImplicitCast();
+			return Pinnable.IsValid() && Pinnable->IsPinned();
+		});
+
+	const bool bNewPinState = !bAllPinned;
+
+	for (const TViewModelPtr Pinnable : OutlinerSelection.Filter<IPinnableExtension>())
+	{
+		Pinnable->SetPinned(bNewPinState);
+	}
+}
+
+void FSequencer::ToggleTrackSelectionLock()
+{
+	using namespace UE::Sequencer;
+
+	FOutlinerSelection& OutlinerSelection = GetSelection().Outliner;
+
+	const bool bAllLocked = Algo::AllOf(OutlinerSelection, [](const FViewModelPtr& Selection)
+		{
+			const TViewModelPtr<ILockableExtension> Lockable = Selection.ImplicitCast();
+			return Lockable.IsValid() && Lockable->GetLockState() != ELockableLockState::None;
+		});
+
+	const bool bNewLockState = !bAllLocked;
+
+	for (const TViewModelPtr Lockable : OutlinerSelection.Filter<ILockableExtension>())
+	{
+		Lockable->SetIsLocked(bNewLockState);
+	}
+}
+
+void FSequencer::ToggleTrackSelectionMute()
+{
+	using namespace UE::Sequencer;
+
+	FOutlinerSelection& OutlinerSelection = GetSelection().Outliner;
+
+	const bool bAllMuted = Algo::AllOf(OutlinerSelection, [](const FViewModelPtr& Selection)
+		{
+			const TViewModelPtr<IMutableExtension> Mutable = Selection.ImplicitCast();
+			return Mutable.IsValid() && Mutable->IsMuted();
+		});
+
+	const bool bNewMuteState = !bAllMuted;
+
+	for (const TViewModelPtr Mutable : OutlinerSelection.Filter<IMutableExtension>())
+	{
+		Mutable->SetIsMuted(bNewMuteState);
+	}
+}
+
+void FSequencer::ToggleTrackSelectionSolo()
+{
+	using namespace UE::Sequencer;
+
+	FOutlinerSelection& OutlinerSelection = GetSelection().Outliner;
+
+	const bool bAllSoloed = Algo::AllOf(OutlinerSelection, [](const FViewModelPtr& Selection)
+		{
+			const TViewModelPtr<ISoloableExtension> Soloable = Selection.ImplicitCast();
+			return Soloable.IsValid() && Soloable->IsSolo();
+		});
+
+	const bool bNewSoloState = !bAllSoloed;
+
+	for (const TViewModelPtr Soloable : OutlinerSelection.Filter<ISoloableExtension>())
+	{
+		Soloable->SetIsSoloed(bNewSoloState);
+	}
+}
+
 TSharedPtr<UE::Sequencer::NavigationTool::INavigationTool> FSequencer::GetNavigationTool() const
 {
 	return NavigationTool;
diff --git a/Engine/Source/Editor/Sequencer/Private/Sequencer.h b/Engine/Source/Editor/Sequencer/Private/Sequencer.h
--- a/Engine/Source/Editor/Sequencer/Private/Sequencer.h
+++ b/Engine/Source/Editor/Sequencer/Private/Sequencer.h
@@ -1295,8 +1295,13 @@ private:
 
 	void SortSelectedOutlinerItems(const bool bInSortBySelectOrder, const bool bInDescending);
 
-	bool CanAlignSelectionToPlayhead();
-	void AlignSelectionToPlayhead();
+	bool CanAlignSelectionToPlayhead() const;
+	void AlignSelectionToPlayhead() const;
+
+	void ToggleTrackSelectionPin();
+	void ToggleTrackSelectionLock();
+	void ToggleTrackSelectionMute();
+	void ToggleTrackSelectionSolo();
 
 	/** User-supplied settings object for this sequencer */
 	TObjectPtr<USequencerSettings> Settings;
diff --git a/Engine/Source/Editor/Sequencer/Private/SequencerCommands.cpp b/Engine/Source/Editor/Sequencer/Private/SequencerCommands.cpp
--- a/Engine/Source/Editor/Sequencer/Private/SequencerCommands.cpp
+++ b/Engine/Source/Editor/Sequencer/Private/SequencerCommands.cpp
@@ -201,6 +201,11 @@ void FSequencerCommands::RegisterCommands()
 
 	UI_COMMAND(AlignSelectionToPlayhead, "Align Selection To Playhead", "Aligns layer bars, sections, or key selections to the current playhead location", EUserInterfaceActionType::Button, FInputChord());
 
+	UI_COMMAND(ToggleTrackSelectionPin, "Toggle Track Pin", "Toggle the pinned state of Sequencer track selection", EUserInterfaceActionType::Button, FInputChord());
+	UI_COMMAND(ToggleTrackSelectionLock, "Toggle Track Lock", "Toggle the locked state of Sequencer track selection", EUserInterfaceActionType::Button, FInputChord());
+	UI_COMMAND(ToggleTrackSelectionMute, "Toggle Track Mute", "Toggle the muted state of Sequencer track selection", EUserInterfaceActionType::Button, FInputChord());
+	UI_COMMAND(ToggleTrackSelectionSolo, "Toggle Track Solo", "Toggle the soloed state of Sequencer track selection", EUserInterfaceActionType::Button, FInputChord());
+
 	UI_COMMAND(ToggleNavigationToolVisible, "Sequence Navigation", "Show the Sequence Navigation Tool", EUserInterfaceActionType::ToggleButton, FInputChord());
 }
 
diff --git a/Engine/Source/Editor/Sequencer/Public/SequencerCommands.h b/Engine/Source/Editor/Sequencer/Public/SequencerCommands.h
--- a/Engine/Source/Editor/Sequencer/Public/SequencerCommands.h
+++ b/Engine/Source/Editor/Sequencer/Public/SequencerCommands.h
@@ -481,6 +481,18 @@ public:
 	/** Aligns layer bars, sections, or key selections to the current playhead location */
 	TSharedPtr< FUICommandInfo > AlignSelectionToPlayhead;
 
+	/** Toggle the pinned state of Sequencer track selection */
+	TSharedPtr< FUICommandInfo > ToggleTrackSelectionPin;
+
+	/** Toggle the locked state of Sequencer track selection */
+	TSharedPtr< FUICommandInfo > ToggleTrackSelectionLock;
+
+	/** Toggle the muted state of Sequencer track selection */
+	TSharedPtr< FUICommandInfo > ToggleTrackSelectionMute;
+
+	/** Toggle the soloed state of Sequencer track selection */
+	TSharedPtr< FUICommandInfo > ToggleTrackSelectionSolo;
+
 	/** Toggles whether or not the Navigation Tool should be visible. */
 	TSharedPtr< FUICommandInfo > ToggleNavigationToolVisible;
 
