[Sequencer] Navigation Tool:
- Enable Navigation Tool Level Sequence provider
- Fix missing binding crash
- Cleanup FNavigationToolTrack item Id creation code
- Show missing bindings and stylize red
- Show bindings when in children of root sequence
- Save built in filter enabled states to local config
- Fix unbound filter not filtering correctly
- Tab window visibility state is now correctly saved/restored

diff --git a/Engine/Plugins/MovieScene/LevelSequenceEditor/Source/LevelSequenceEditor/Private/LevelSequenceEditorToolkit.cpp b/Engine/Plugins/MovieScene/LevelSequenceEditor/Source/LevelSequenceEditor/Private/LevelSequenceEditorToolkit.cpp
--- a/Engine/Plugins/MovieScene/LevelSequenceEditor/Source/LevelSequenceEditor/Private/LevelSequenceEditorToolkit.cpp
+++ b/Engine/Plugins/MovieScene/LevelSequenceEditor/Source/LevelSequenceEditor/Private/LevelSequenceEditorToolkit.cpp
@@ -319,12 +319,12 @@ void FLevelSequenceEditorToolkit::Initialize(const EToolkitMode::Type Mode, cons
 		}
 	}
 
-	// using namespace UE::Sequencer::NavigationTool;
-	// if (const TSharedPtr<INavigationTool> Tool = Sequencer->GetNavigationTool())
-	// {
-	// 	NavigationToolProvider = MakeShared<FLevelSequenceNavigationToolProvider>();
-	// 	FNavigationToolExtender::RegisterOnlyToolProvider(Tool.ToSharedRef(), NavigationToolProvider.ToSharedRef());
-	// }
+	using namespace UE::Sequencer::NavigationTool;
+	if (const TSharedPtr<INavigationTool> Tool = Sequencer->GetNavigationTool())
+	{
+		NavigationToolProvider = MakeShared<FLevelSequenceNavigationToolProvider>();
+		FNavigationToolExtender::RegisterOnlyToolProvider(Tool.ToSharedRef(), NavigationToolProvider.ToSharedRef());
+	}
 }
 
 
diff --git a/Engine/Plugins/MovieScene/LevelSequenceEditor/Source/LevelSequenceEditor/Private/LevelSequenceNavigationToolProvider.cpp b/Engine/Plugins/MovieScene/LevelSequenceEditor/Source/LevelSequenceEditor/Private/LevelSequenceNavigationToolProvider.cpp
--- a/Engine/Plugins/MovieScene/LevelSequenceEditor/Source/LevelSequenceEditor/Private/LevelSequenceNavigationToolProvider.cpp
+++ b/Engine/Plugins/MovieScene/LevelSequenceEditor/Source/LevelSequenceEditor/Private/LevelSequenceNavigationToolProvider.cpp
@@ -57,8 +57,6 @@ FName FLevelSequenceNavigationToolProvider::GetDefaultColumnView() const
 
 void FLevelSequenceNavigationToolProvider::OnExtendColumns(FNavigationToolColumnExtender& OutExtender)
 {
-	FNavigationToolProvider::OnExtendColumns(OutExtender);
-
 	OutExtender.AddColumn<FNavigationToolDeactiveStateColumn>();
 	OutExtender.AddColumn<FNavigationToolMarkerVisibilityColumn>();
 	OutExtender.AddColumn<FNavigationToolLockColumn>();
@@ -78,6 +76,8 @@ void FLevelSequenceNavigationToolProvider::OnExtendColumns(FNavigationToolColumn
 	{
 		OutExtender.AddColumn<FNavigationToolRevisionControlColumn>();
 	}
+
+	FNavigationToolProvider::OnExtendColumns(OutExtender);
 }
 
 void FLevelSequenceNavigationToolProvider::OnExtendColumnViews(TSet<FNavigationToolColumnView>& OutColumnViews)
@@ -94,6 +94,8 @@ void FLevelSequenceNavigationToolProvider::OnExtendColumnViews(TSet<FNavigationT
 			FNavigationToolOutTimeColumn::GetStaticTypeId().ToName(),
 			FNavigationToolHBiasColumn::GetStaticTypeId().ToName()
 		};
+
+	FNavigationToolProvider::OnExtendColumnViews(OutColumnViews);
 }
 
 void FLevelSequenceNavigationToolProvider::OnExtendItemChildren(INavigationTool& InTool
@@ -135,7 +137,10 @@ void FLevelSequenceNavigationToolProvider::OnExtendItemChildren(INavigationTool&
 void FLevelSequenceNavigationToolProvider::OnExtendBuiltInFilters(TArray<FNavigationToolBuiltInFilterParams>& OutFilterParams)
 {
 	OutFilterParams.Add(FNavigationToolBuiltInFilterParams::CreateSequenceFilter());
+	OutFilterParams.Add(FNavigationToolBuiltInFilterParams::CreateTrackFilter());
 	OutFilterParams.Add(FNavigationToolBuiltInFilterParams::CreateBindingFilter());
+
+	FNavigationToolProvider::OnExtendBuiltInFilters(OutFilterParams);
 }
 
 FMovieSceneEditorData* FLevelSequenceNavigationToolProvider::GetRootMovieSceneEditorData(const INavigationTool& InTool) const
diff --git a/Engine/Source/Editor/Sequencer/Private/Tools/NavigationTool/Filters/Filters/NavigationToolBuiltInFilterParams.cpp b/Engine/Source/Editor/Sequencer/Private/Tools/NavigationTool/Filters/Filters/NavigationToolBuiltInFilterParams.cpp
--- a/Engine/Source/Editor/Sequencer/Private/Tools/NavigationTool/Filters/Filters/NavigationToolBuiltInFilterParams.cpp
+++ b/Engine/Source/Editor/Sequencer/Private/Tools/NavigationTool/Filters/Filters/NavigationToolBuiltInFilterParams.cpp
@@ -172,19 +172,27 @@ FNavigationToolBuiltInFilterParams FNavigationToolBuiltInFilterParams::CreateSeq
 		, { ULevelSequence::StaticClass(), UMovieSceneSequence::StaticClass() }
 		, ENavigationToolFilterMode::MatchesType
 		, FSlateIconFinder::FindIconForClass(ALevelSequenceActor::StaticClass()).GetIcon()
-		, LOCTEXT("SequenceFilterDisplayName", "Level Sequence")
+		, LOCTEXT("LevelSequenceFilterDisplayName", "Level Sequence")
 		, LOCTEXT("LevelSequenceFilterTooltip", "Level Sequence"));
 }
 
+FNavigationToolBuiltInFilterParams FNavigationToolBuiltInFilterParams::CreateTrackFilter()
+{
+	return FNavigationToolBuiltInFilterParams(TEXT("Track")
+		, { UE::Sequencer::NavigationTool::FNavigationToolTrack::GetStaticTypeId() }
+		, { UMovieSceneTrack::StaticClass() }
+		, ENavigationToolFilterMode::MatchesType
+		, FSlateIconFinder::FindIconForClass(UMovieSceneTrack::StaticClass()).GetIcon()
+		, LOCTEXT("TrackFilterDisplayName", "Track")
+		, LOCTEXT("TrackFilterTooltip", "Track"));
+}
+
 FNavigationToolBuiltInFilterParams FNavigationToolBuiltInFilterParams::CreateBindingFilter()
 {
 	return FNavigationToolBuiltInFilterParams(TEXT("Binding")
-		, {
-			UE::Sequencer::NavigationTool::FNavigationToolBinding::GetStaticTypeId()
-			,
-			UE::Sequencer::NavigationTool::FNavigationToolActor::GetStaticTypeId()
-			,
-			UE::Sequencer::NavigationTool::FNavigationToolComponent::GetStaticTypeId() }
+		, { UE::Sequencer::NavigationTool::FNavigationToolBinding::GetStaticTypeId()
+			, UE::Sequencer::NavigationTool::FNavigationToolActor::GetStaticTypeId()
+			, UE::Sequencer::NavigationTool::FNavigationToolComponent::GetStaticTypeId() }
 		, { }
 		, ENavigationToolFilterMode::MatchesType
 		, nullptr
diff --git a/Engine/Source/Editor/Sequencer/Private/Tools/NavigationTool/Filters/Filters/NavigationToolFilter_Unbound.cpp b/Engine/Source/Editor/Sequencer/Private/Tools/NavigationTool/Filters/Filters/NavigationToolFilter_Unbound.cpp
--- a/Engine/Source/Editor/Sequencer/Private/Tools/NavigationTool/Filters/Filters/NavigationToolFilter_Unbound.cpp
+++ b/Engine/Source/Editor/Sequencer/Private/Tools/NavigationTool/Filters/Filters/NavigationToolFilter_Unbound.cpp
@@ -44,15 +44,14 @@ FString FNavigationToolFilter_Unbound::GetName() const
 
 bool FNavigationToolFilter_Unbound::PassesFilter(FNavigationToolFilterType InItem) const
 {
-	const FNavigationToolBinding* const BindingItem = InItem->CastTo<FNavigationToolBinding>();
+	FNavigationToolBinding* const BindingItem = InItem->CastTo<FNavigationToolBinding>();
 	if (!BindingItem)
 	{
 		return false;
 	}
 
-	UObject* const BoundObject = BindingItem->FindBoundObject();
-
-	return IsValid(BoundObject);
+	const UObject* const BoundObject = BindingItem->GetCachedBoundObject();
+	return !IsValid(BoundObject);
 }
 
 }
diff --git a/Engine/Source/Editor/Sequencer/Private/Tools/NavigationTool/Items/NavigationToolActor.cpp b/Engine/Source/Editor/Sequencer/Private/Tools/NavigationTool/Items/NavigationToolActor.cpp
--- a/Engine/Source/Editor/Sequencer/Private/Tools/NavigationTool/Items/NavigationToolActor.cpp
+++ b/Engine/Source/Editor/Sequencer/Private/Tools/NavigationTool/Items/NavigationToolActor.cpp
@@ -12,6 +12,8 @@
 #include "Tools/NavigationTool/ItemProxies/NavigationToolComponentProxy.h"
 #include "Tools/NavigationTool/NavigationTool.h"
 #include "Tools/NavigationTool/NavigationToolSettings.h"
+#include "Tools/NavigationTool/Utils/NavigationToolMiscUtils.h"
+#include "Tools/NavigationTool/Utils/NavigationToolMovieSceneUtils.h"
 
 #define LOCTEXT_NAMESPACE "NavigationToolActor"
 
@@ -20,14 +22,9 @@ namespace UE::Sequencer::NavigationTool
 
 FNavigationToolActor::FNavigationToolActor(INavigationTool& InTool
 	, const FNavigationToolItemPtr& InParentItem
-	, const FMovieSceneBinding& InBinding
-	, UObject* const InObject
-	, const FNavigationToolItemPtr& InReferencingItem)
-	: Super(InTool
-	, InParentItem
-	, InBinding
-	, InObject
-	, InReferencingItem)
+	, const TSharedPtr<FNavigationToolSequence>& InParentSequenceItem
+	, const FMovieSceneBinding& InBinding)
+	: Super(InTool, InParentItem, InParentSequenceItem, InBinding)
 {
 }
 
@@ -35,19 +32,19 @@ void FNavigationToolActor::FindChildren(TArray<FNavigationToolItemPtr>& OutChild
 {
 	Super::FindChildren(OutChildren, bInRecursive);
 
-	AActor* const ThisBoundActor = GetActor();
-	if (!IsValid(ThisBoundActor))
+	const AActor* const ThisActor = GetActor();
+	if (!IsValid(ThisActor))
 	{
 		return;
 	}
 
-	const TSharedPtr<FNavigationToolSequence> SequenceParentItem = StaticCastSharedPtr<FNavigationToolSequence>(WeakReferencingItem.Pin());
-	if (!SequenceParentItem.IsValid())
+	const TSharedPtr<FNavigationToolSequence> ParentSequenceItem = WeakParentSequenceItem.Pin();
+	if (!ParentSequenceItem.IsValid())
 	{
 		return;
 	}
 
-	UMovieSceneSequence* const Sequence = SequenceParentItem->GetSequence();
+	UMovieSceneSequence* const Sequence = ParentSequenceItem->GetSequence();
 	if (!IsValid(Sequence))
 	{
 		return;
@@ -65,23 +62,23 @@ void FNavigationToolActor::FindChildren(TArray<FNavigationToolItemPtr>& OutChild
 		return;
 	}
 
-	const TSharedRef<FNavigationToolBinding> SharedThisRef = SharedThis(this);
-
-	TArray<FMovieSceneBinding>& Bindings = MovieScene->GetBindings();
-	Bindings.Sort([](const FMovieSceneBinding& InA, const FMovieSceneBinding& InB)
-		{
-			return InA.GetSortingOrder() < InB.GetSortingOrder();
-		});
+	const TSharedRef<FNavigationToolActor> SharedThisRef = SharedThis(this);
 
+	const TArray<FMovieSceneBinding> Bindings = ParentSequenceItem->GetSortedBindings();
 	for (const FMovieSceneBinding& CurrentBinding : Bindings)
 	{
-		UObject* const BoundObject = Sequencer->FindSpawnedObjectOrTemplate(CurrentBinding.GetObjectGuid());
-		if (IsValid(BoundObject))
+		const FGuid& CurrentBindingGuid = CurrentBinding.GetObjectGuid();
+		const UClass* const CurrentObjectClass = GetBindingObjectClass(*MovieScene, CurrentBindingGuid);
+		if (CurrentObjectClass && CurrentObjectClass->IsChildOf<UActorComponent>())
 		{
-			UActorComponent* const Component = Cast<UActorComponent>(BoundObject);
-			if (IsValid(Component) && Component->GetOwner() == ThisBoundActor)
+			const TArrayView<TWeakObjectPtr<>> WeakBoundObjects = ResolveBoundObjects(*Sequencer, Sequence, CurrentBindingGuid);
+			const TWeakObjectPtr<> CurrentBoundObject = WeakBoundObjects.IsEmpty() ? nullptr : WeakBoundObjects[0];
+
+			const UActorComponent* const CurrentBoundActorComponent = Cast<UActorComponent>(CurrentBoundObject);
+			if (IsValid(CurrentBoundActorComponent) && CurrentBoundActorComponent->GetOwner() == ThisActor)
 			{
-				const FNavigationToolItemPtr NewItem = Tool.FindOrAdd<FNavigationToolComponent>(GetProvider(), SharedThisRef, CurrentBinding, BoundObject, SharedThisRef);
+				const FNavigationToolItemPtr NewItem = Tool.FindOrAdd<FNavigationToolComponent>(GetProvider()
+					, SharedThisRef, WeakParentSequenceItem.Pin(), CurrentBinding);
 				OutChildren.Add(NewItem);
 				if (bInRecursive)
 				{
@@ -105,8 +102,17 @@ void FNavigationToolActor::GetItemProxies(TArray<TSharedPtr<FNavigationToolItemP
 bool FNavigationToolActor::IsAllowedInTool() const
 {
 	FNavigationTool& ToolPrivate = static_cast<FNavigationTool&>(Tool);
-	const bool bIsAllowed = ToolPrivate.IsObjectAllowedInTool(GetActor());
-	return bIsAllowed;
+
+	const AActor* const UnderlyingActor = GetActor();
+	if (!IsValid(UnderlyingActor))
+	{
+		// Always allow unbound binding items
+		return true;
+	}
+
+	const bool bActorAllowed = ToolPrivate.IsObjectAllowedInTool(UnderlyingActor);
+
+	return bActorAllowed;
 }
 
 ENavigationToolItemViewMode FNavigationToolActor::GetSupportedViewModes(const INavigationToolView& InToolView) const
@@ -114,18 +120,10 @@ ENavigationToolItemViewMode FNavigationToolActor::GetSupportedViewModes(const IN
 	return ENavigationToolItemViewMode::ItemTree | ENavigationToolItemViewMode::HorizontalItemList;
 }
 
-FText FNavigationToolActor::GetDisplayName() const
-{
-	if (const AActor* const UnderlyingActor = GetActor())
-	{
-		return FText::FromString(UnderlyingActor->GetActorLabel());
-	}
-	return FText::GetEmpty();
-}
-
 TArray<FName> FNavigationToolActor::GetTags() const
 {
-	if (AActor* const UnderlyingActor = GetActor())
+	const AActor* const UnderlyingActor = GetActor();
+	if (IsValid(UnderlyingActor))
 	{
 		return UnderlyingActor->Tags;
 	}
@@ -134,7 +132,8 @@ TArray<FName> FNavigationToolActor::GetTags() const
 
 bool FNavigationToolActor::GetVisibility() const
 {
-	if (const AActor* const UnderlyingActor = GetActor())
+	const AActor* const UnderlyingActor = GetActor();
+	if (IsValid(UnderlyingActor))
 	{
 		return !UnderlyingActor->IsTemporarilyHiddenInEditor(true);
 	}
@@ -143,7 +142,8 @@ bool FNavigationToolActor::GetVisibility() const
 
 void FNavigationToolActor::OnVisibilityChanged(const bool bInNewVisibility)
 {
-	if (AActor* const UnderlyingActor = GetActor())
+	AActor* const UnderlyingActor = GetActor();
+	if (IsValid(UnderlyingActor))
 	{
 		UnderlyingActor->SetIsTemporarilyHiddenInEditor(!bInNewVisibility);
 	}
@@ -177,7 +177,7 @@ bool FNavigationToolActor::Rename(const FString& InName)
 
 AActor* FNavigationToolActor::GetActor() const
 {
-	return WeakObject.IsValid() ? Cast<AActor>(WeakObject.Get(IsIgnoringPendingKill())) : nullptr;
+	return WeakBoundObject.IsValid() ? Cast<AActor>(GetCachedBoundObject()) : nullptr;
 }
 
 }
diff --git a/Engine/Source/Editor/Sequencer/Private/Tools/NavigationTool/Items/NavigationToolBinding.cpp b/Engine/Source/Editor/Sequencer/Private/Tools/NavigationTool/Items/NavigationToolBinding.cpp
--- a/Engine/Source/Editor/Sequencer/Private/Tools/NavigationTool/Items/NavigationToolBinding.cpp
+++ b/Engine/Source/Editor/Sequencer/Private/Tools/NavigationTool/Items/NavigationToolBinding.cpp
@@ -2,12 +2,15 @@
 
 #include "Tools/NavigationTool/Items/NavigationToolBinding.h"
 #include "ISequencer.h"
+#include "MovieScene.h"
 #include "ScopedTransaction.h"
+#include "Sections/MovieSceneSubSection.h"
 #include "Styling/SlateIconFinder.h"
 #include "Tools/NavigationTool/INavigationTool.h"
-#include "Tools/NavigationTool/Items/NavigationToolItemUtils.h"
 #include "Tools/NavigationTool/Items/NavigationToolSequence.h"
 #include "Tools/NavigationTool/NavigationToolScopedSelection.h"
+#include "Tools/NavigationTool/Utils/NavigationToolMiscUtils.h"
+#include "Tools/NavigationTool/Utils/NavigationToolMovieSceneUtils.h"
 #include "Widgets/Views/STableRow.h"
 
 #define LOCTEXT_NAMESPACE "NavigationToolBinding"
@@ -19,16 +22,24 @@ using namespace UE::Sequencer;
 
 FNavigationToolBinding::FNavigationToolBinding(INavigationTool& InTool
 	, const FNavigationToolItemPtr& InParentItem
-	, const FMovieSceneBinding& InBinding
-	, UObject* const InObject
-	, const FNavigationToolItemPtr& InReferencingItem)
-	: Super(InTool, InParentItem, InObject, InReferencingItem, InBinding.GetObjectGuid().ToString())
+	, const TSharedPtr<FNavigationToolSequence>& InParentSequenceItem
+	, const FMovieSceneBinding& InBinding)
+	: Super(InTool, InParentItem)
+	, WeakParentSequenceItem(InParentSequenceItem)
 	, Binding(InBinding)
 {
+	CacheBoundObject();
+}
+
+bool FNavigationToolBinding::IsItemValid() const
+{
+	return FNavigationToolItem::IsItemValid() && Binding.GetObjectGuid().IsValid();
 }
 
 void FNavigationToolBinding::FindChildren(TArray<FNavigationToolItemPtr>& OutChildren, const bool bInRecursive)
 {
+	CacheBoundObject();
+
 	Super::FindChildren(OutChildren, bInRecursive);
 }
 
@@ -37,13 +48,14 @@ FText FNavigationToolBinding::GetDisplayName() const
 	return FText::FromString(Binding.GetName());
 }
 
-const FSlateBrush* FNavigationToolBinding::GetDefaultIconBrush() const
+FText FNavigationToolBinding::GetClassName() const
 {
-	if (WeakObject.IsValid())
+	const UClass* const BoundObjectClass = GetBoundObjectClass();
+	if (IsValid(BoundObjectClass))
 	{
-		return FSlateIconFinder::FindIconForClass(WeakObject->GetClass()).GetIcon();
+		return BoundObjectClass->GetDisplayNameText();
 	}
-	return FAppStyle::GetBrush(TEXT("Sequencer.InvalidSpawnableIcon"));
+	return FText::GetEmpty();
 }
 
 FSlateIcon FNavigationToolBinding::GetIcon() const
@@ -53,7 +65,24 @@ FSlateIcon FNavigationToolBinding::GetIcon() const
 
 FText FNavigationToolBinding::GetIconTooltipText() const
 {
-	return GetDisplayName();
+	const UObject* const UnderlyingObject = GetCachedBoundObject();
+	if (IsValid(UnderlyingObject))
+	{
+		return FText::Format(LOCTEXT("BoundObjectToolTip", " Class: %s (BindingID: %s)")
+			, UnderlyingObject->GetClass()->GetDisplayNameText(), FText::FromString(LexToString(Binding.GetObjectGuid())));
+	}
+	return FText::Format(LOCTEXT("InvalidBoundObjectToolTip", "The object bound to this track is missing (BindingID: {0}).")
+		, FText::FromString(LexToString(Binding.GetObjectGuid())));
+}
+
+FLinearColor FNavigationToolBinding::GetItemBrushColor() const
+{
+	const UObject* const UnderlyingObject = GetCachedBoundObject();
+	if (IsValid(UnderlyingObject))
+	{
+		return FNavigationToolItem::GetItemBrushColor();
+	}
+	return FStyleColors::Error.GetSpecifiedColor();
 }
 
 bool FNavigationToolBinding::IsSelected(const FNavigationToolScopedSelection& InSelection) const
@@ -65,20 +94,57 @@ void FNavigationToolBinding::Select(FNavigationToolScopedSelection& InSelection)
 {
 	InSelection.Select(Binding.GetObjectGuid());
 }
+	
+void FNavigationToolBinding::OnSelect()
+{
+	Super::OnSelect();
+}
+
+void FNavigationToolBinding::OnDoubleClick()
+{
+	Super::OnDoubleClick();
+
+	const TSharedPtr<ISequencer> Sequencer = Tool.GetSequencer();
+	if (!Sequencer.IsValid())
+	{
+		return;
+	}
+
+	UObject* const BoundObject = GetCachedBoundObject();
+	if (!IsValid(BoundObject))
+	{
+		return;
+	}
+
+	const FGuid ObjectId = Sequencer->FindObjectId(*BoundObject, Sequencer->GetFocusedTemplateID());
+	Sequencer->SelectObject(ObjectId);
+}
 
 bool FNavigationToolBinding::CanDelete() const
 {
-	return WeakObject.IsValid();
+	return WeakBoundObject.IsValid();
 }
 
 bool FNavigationToolBinding::Delete()
 {
-	return !WeakObject.IsValid();
+	return !WeakBoundObject.IsValid();
+}
+
+FNavigationToolItemId FNavigationToolBinding::CalculateItemId() const
+{
+	const TSharedPtr<FNavigationToolSequence> ParentSequenceItem = WeakParentSequenceItem.Pin();
+	check(ParentSequenceItem.IsValid());
+
+	return FNavigationToolItemId(GetParent()
+		, ParentSequenceItem->GetSequence()
+		, ParentSequenceItem->GetSubSection()
+		, ParentSequenceItem->GetSubSectionIndex()
+		, Binding.GetName());
 }
 
 bool FNavigationToolBinding::CanRename() const
 {
-	return WeakObject.IsValid();
+	return WeakBoundObject.IsValid();
 }
 
 bool FNavigationToolBinding::Rename(const FString& InName)
@@ -133,11 +199,52 @@ void FNavigationToolBinding::SetIsLocked(const bool bInIsLocked)
 	}
 }*/
 
-UObject* FNavigationToolBinding::FindBoundObject() const
+const UClass* FNavigationToolBinding::GetBoundObjectClass() const
+{
+	const TSharedPtr<FNavigationToolSequence> ParentSequenceItem = WeakParentSequenceItem.Pin();
+	if (!ParentSequenceItem.IsValid())
+	{
+		return nullptr;
+	}
+
+	UMovieSceneSequence* const Sequence = ParentSequenceItem->GetSequence();
+	if (!IsValid(Sequence))
+	{
+		return nullptr;
+	}
+
+	UMovieScene* const MovieScene = Sequence->GetMovieScene();
+	if (!IsValid(MovieScene))
+	{
+		return nullptr;
+	}
+
+	return GetBindingObjectClass(*MovieScene, Binding.GetObjectGuid());
+}
+
+UObject* FNavigationToolBinding::GetCachedBoundObject() const
+{
+	return WeakBoundObject.Get();
+}
+
+UObject* FNavigationToolBinding::CacheBoundObject()
 {
 	if (const TSharedPtr<ISequencer> Sequencer = Tool.GetSequencer())
 	{
-		return Sequencer->FindSpawnedObjectOrTemplate(Binding.GetObjectGuid());
+		UMovieSceneSequence* const Sequence = WeakParentSequenceItem.Pin()->GetSequence();
+		const TArrayView<TWeakObjectPtr<>> WeakBoundObjects = ResolveBoundObjects(*Sequencer, Sequence, Binding.GetObjectGuid());
+		WeakBoundObject = WeakBoundObjects.IsEmpty() ? nullptr : WeakBoundObjects[0];
+
+		if (WeakBoundObject.IsValid())
+		{
+			Icon = FSlateIconFinder::FindIconForClass(WeakBoundObject->GetClass());
+		}
+		else
+		{
+			Icon = FSlateIcon(FAppStyle::GetAppStyleSetName(), TEXT("Sequencer.InvalidSpawnableIcon"));
+		}
+
+		return WeakBoundObject.Get();
 	}
 	return nullptr;
 }
diff --git a/Engine/Source/Editor/Sequencer/Private/Tools/NavigationTool/Items/NavigationToolComponent.cpp b/Engine/Source/Editor/Sequencer/Private/Tools/NavigationTool/Items/NavigationToolComponent.cpp
--- a/Engine/Source/Editor/Sequencer/Private/Tools/NavigationTool/Items/NavigationToolComponent.cpp
+++ b/Engine/Source/Editor/Sequencer/Private/Tools/NavigationTool/Items/NavigationToolComponent.cpp
@@ -16,36 +16,15 @@ namespace UE::Sequencer::NavigationTool
 
 FNavigationToolComponent::FNavigationToolComponent(INavigationTool& InTool
 	, const FNavigationToolItemPtr& InParentItem
-	, const FMovieSceneBinding& InBinding
-	, UObject* const InObject
-	, const FNavigationToolItemPtr& InReferencingItem)
-	: Super(InTool
-	, InParentItem
-	, InBinding
-	, InObject
-	, InReferencingItem)
+	, const TSharedPtr<FNavigationToolSequence>& InParentSequenceItem
+	, const FMovieSceneBinding& InBinding)
+	: Super(InTool, InParentItem, InParentSequenceItem, InBinding)
 {
 }
 
 void FNavigationToolComponent::FindChildren(TArray<FNavigationToolItemPtr>& OutChildren, const bool bInRecursive)
 {
 	Super::FindChildren(OutChildren, bInRecursive);
-
-	/*const TSharedRef<FNavigationToolComponent> SharedThisRef = SharedThis(this);
-
-	for (UMovieSceneTrack* const Track : Binding.GetTracks())
-	{
-		const TArray<UMovieSceneSection*>& AllSections = Track->GetAllSections();
-		for (int32 SectionIndex = 0; SectionIndex < AllSections.Num(); ++SectionIndex)
-		{
-			const FNavigationToolItemPtr NewItem = Tool.FindOrAdd<FNavigationToolTrack>(GetProvider(), SharedThisRef, Track, Sequence, AllSections[SectionIndex], SectionIndex);
-			OutChildren.Add(NewItem);
-			if (bInRecursive)
-			{
-				NewItem->FindChildren(OutChildren, bInRecursive);
-			}
-		}
-	}*/
 }
 
 void FNavigationToolComponent::GetItemProxies(TArray<TSharedPtr<FNavigationToolItemProxy>>& OutItemProxies)
@@ -65,14 +44,17 @@ bool FNavigationToolComponent::IsAllowedInTool() const
 {
 	FNavigationTool& ToolPrivate = static_cast<FNavigationTool&>(Tool);
 
-	// Make sure Owner is Allowed (this also returns false if Owner is invalid)
 	const UActorComponent* const UnderlyingComponent = GetComponent();
-	if (!IsValid(UnderlyingComponent) || !ToolPrivate.IsObjectAllowedInTool(UnderlyingComponent->GetOwner()))
+	if (!IsValid(UnderlyingComponent))
 	{
-		return false;
+		// Always allow unbound binding items
+		return true;
 	}
 
-	return ToolPrivate.IsObjectAllowedInTool(UnderlyingComponent);
+	const bool bOwnerAllowed = ToolPrivate.IsObjectAllowedInTool(UnderlyingComponent->GetOwner());
+	const bool bComponentAllowed = ToolPrivate.IsObjectAllowedInTool(UnderlyingComponent);
+
+	return bOwnerAllowed && bComponentAllowed;
 }
 
 ENavigationToolItemViewMode FNavigationToolComponent::GetSupportedViewModes(const INavigationToolView& InToolView) const
@@ -151,7 +133,7 @@ bool FNavigationToolComponent::Rename(const FString& InName)
 
 UActorComponent* FNavigationToolComponent::GetComponent() const
 {
-	return WeakObject.IsValid() ? Cast<UActorComponent>(WeakObject.Get(IsIgnoringPendingKill())) : nullptr;
+	return WeakBoundObject.IsValid() ? Cast<UActorComponent>(GetCachedBoundObject()) : nullptr;
 }
 
 }
diff --git a/Engine/Source/Editor/Sequencer/Private/Tools/NavigationTool/Items/NavigationToolItemId.cpp b/Engine/Source/Editor/Sequencer/Private/Tools/NavigationTool/Items/NavigationToolItemId.cpp
--- a/Engine/Source/Editor/Sequencer/Private/Tools/NavigationTool/Items/NavigationToolItemId.cpp
+++ b/Engine/Source/Editor/Sequencer/Private/Tools/NavigationTool/Items/NavigationToolItemId.cpp
@@ -2,7 +2,7 @@
 
 #include "Tools/NavigationTool/Items/NavigationToolItemId.h"
 #include "MovieSceneSequence.h"
-#include "Sections/MovieSceneSubSection.h"
+#include "MovieSceneSection.h"
 #include "Tools/NavigationTool/Items/INavigationToolItem.h"
 #include "Tools/NavigationTool/Items/NavigationToolItem.h"
 
@@ -51,8 +51,9 @@ FNavigationToolItemId::FNavigationToolItemId(const FNavigationToolItemPtr& InPar
 
 FNavigationToolItemId::FNavigationToolItemId(const FNavigationToolItemPtr& InParentItem
 	, const UMovieSceneSequence* const InSequence
-	, const UMovieSceneSubSection* const InSubSection
-	, const int32 InSectionIndex)
+	, const UMovieSceneSection* const InSection
+	, const int32 InSectionIndex
+	, const FString& InReferenceId)
 {
 	check(InSequence);
 
@@ -63,11 +64,16 @@ FNavigationToolItemId::FNavigationToolItemId(const FNavigationToolItemPtr& InPar
 
 	AddSeparatedSegment(Id, GetObjectPath(InSequence));
 
-	if (InSubSection)
+	if (InSection)
 	{
-		AddSeparatedSegment(Id, GetObjectPath(InSubSection));
+		AddSeparatedSegment(Id, GetObjectPath(InSection));
 		AddSeparatedSegment(Id, FString::FromInt(InSectionIndex));
 	}
+
+	if (!InReferenceId.IsEmpty())
+	{
+		AddSeparatedSegment(Id, InReferenceId);
+	}
 }
 
 FNavigationToolItemId::FNavigationToolItemId(const FNavigationToolItemId& Other)
diff --git a/Engine/Source/Editor/Sequencer/Private/Tools/NavigationTool/Items/NavigationToolSequence.cpp b/Engine/Source/Editor/Sequencer/Private/Tools/NavigationTool/Items/NavigationToolSequence.cpp
--- a/Engine/Source/Editor/Sequencer/Private/Tools/NavigationTool/Items/NavigationToolSequence.cpp
+++ b/Engine/Source/Editor/Sequencer/Private/Tools/NavigationTool/Items/NavigationToolSequence.cpp
@@ -66,26 +66,23 @@ void FNavigationToolSequence::FindChildren(TArray<FNavigationToolItemPtr>& OutCh
 		return;
 	}
 
-	const TSharedPtr<ISequencer> Sequencer = Tool.GetSequencer();
-	if (!Sequencer.IsValid())
-	{
-		return;
-	}
-
 	const TSharedRef<FNavigationToolSequence> SharedThisRef = SharedThis(this);
 
-	for (UMovieSceneTrack* const Track : MovieScene->GetTracks())
+	const TArray<UMovieSceneTrack*> MovieSceneTracks = MovieScene->GetTracks();
+	for (UMovieSceneTrack* const Track : MovieSceneTracks)
 	{
 		FNavigationToolItemPtr NewItem;
 
 		UMovieSceneSubTrack* const SubTrack = Cast<UMovieSceneSubTrack>(Track);
 		if (IsValid(SubTrack))
 		{
-			NewItem = Tool.FindOrAdd<FNavigationToolSubTrack>(GetProvider(), SharedThisRef, SubTrack, Sequence, WeakSubSection, SubSectionIndex);
+			NewItem = Tool.FindOrAdd<FNavigationToolSubTrack>(GetProvider()
+				, SharedThisRef, SubTrack, Sequence, WeakSubSection, SubSectionIndex);
 		}
 		else if (IsValid(Track))
 		{
-			NewItem = Tool.FindOrAdd<FNavigationToolTrack>(GetProvider(), SharedThisRef, Track, Sequence, WeakSubSection, SubSectionIndex);
+			NewItem = Tool.FindOrAdd<FNavigationToolTrack>(GetProvider()
+				, SharedThisRef, Track, Sequence, WeakSubSection, SubSectionIndex);
 		}
 
 		if (NewItem.IsValid())
@@ -98,26 +95,20 @@ void FNavigationToolSequence::FindChildren(TArray<FNavigationToolItemPtr>& OutCh
 		}
 	}
 
-	TArray<FMovieSceneBinding>& Bindings = MovieScene->GetBindings();
-	Bindings.Sort([](const FMovieSceneBinding& InA, const FMovieSceneBinding& InB)
-		{
-			return InA.GetSortingOrder() < InB.GetSortingOrder();
-		});
-
 	// Only show actor bindings for a Sequence
+	const TArray<FMovieSceneBinding> Bindings = GetSortedBindings();
 	for (const FMovieSceneBinding& Binding : Bindings)
 	{
-		UObject* const BoundObject = Sequencer->FindSpawnedObjectOrTemplate(Binding.GetObjectGuid());
-		if (IsValid(BoundObject))
+		const FGuid& BindingObjectGuid = Binding.GetObjectGuid();
+		const UClass* const BoundObjectClass = GetBindingObjectClass(*MovieScene, BindingObjectGuid);
+		if (BoundObjectClass && BoundObjectClass->IsChildOf<AActor>())
 		{
-			if (AActor* const Actor = Cast<AActor>(BoundObject))
+			const FNavigationToolItemPtr NewItem = Tool.FindOrAdd<FNavigationToolActor>(GetProvider()
+				, SharedThisRef, SharedThisRef, Binding);
+			OutChildren.Add(NewItem);
+			if (bInRecursive)
 			{
-				const FNavigationToolItemPtr NewItem = Tool.FindOrAdd<FNavigationToolActor>(GetProvider(), SharedThisRef, Binding, BoundObject, SharedThisRef);
-				OutChildren.Add(NewItem);
-				if (bInRecursive)
-				{
-					NewItem->FindChildren(OutChildren, bInRecursive);
-				}
+				NewItem->FindChildren(OutChildren, bInRecursive);
 			}
 		}
 	}
@@ -189,26 +180,6 @@ FText FNavigationToolSequence::GetIconTooltipText() const
 	return FText::GetEmpty();
 }
 
-UMovieSceneSequence* FNavigationToolSequence::GetSequence() const
-{
-	return Cast<UMovieSceneSequence>(GetObject());
-}
-
-UMovieSceneSubSection* FNavigationToolSequence::GetSubSection() const
-{
-	return WeakSubSection.IsValid() ? WeakSubSection.Get() : nullptr;
-}
-
-UMovieScene* FNavigationToolSequence::GetSequenceMovieScene() const
-{
-	UMovieSceneSequence* const Sequence = GetSequence();
-	if (IsValid(Sequence))
-	{
-		return Sequence->GetMovieScene();
-	}
-	return nullptr;
-}
-
 void FNavigationToolSequence::OnSelect()
 {
 	Super::OnSelect();
@@ -446,7 +417,7 @@ EItemRevisionControlState FNavigationToolSequence::GetRevisionControlState() con
 		return EItemRevisionControlState::None;
 	}
 
-	const FSourceControlStatePtr RevisionControlState = GetSourceControlState(WeakSubSection->GetPackage());
+	const FSourceControlStatePtr RevisionControlState = FindSourceControlState(WeakSubSection->GetPackage());
 	if (!RevisionControlState.IsValid())
 	{
 		return EItemRevisionControlState::None;
@@ -473,7 +444,7 @@ const FSlateBrush* FNavigationToolSequence::GetRevisionControlStatusIcon() const
 		return nullptr;
 	}
 
-	return GetSourceControlStatusBrush(WeakSubSection->GetPackage());
+	return FindSourceControlStatusBrush(WeakSubSection->GetPackage());
 }
 
 FText FNavigationToolSequence::GetRevisionControlStatusText() const
@@ -489,7 +460,56 @@ FText FNavigationToolSequence::GetRevisionControlStatusText() const
 		return FText::GetEmpty();
 	}
 
-	return GetSourceControlStatusText(WeakSubSection->GetPackage());
+	return FindSourceControlStatusText(WeakSubSection->GetPackage());
+}
+
+UMovieSceneSequence* FNavigationToolSequence::GetSequence() const
+{
+	return Cast<UMovieSceneSequence>(GetObject());
+}
+
+UMovieSceneSubSection* FNavigationToolSequence::GetSubSection() const
+{
+	return WeakSubSection.IsValid() ? WeakSubSection.Get() : nullptr;
+}
+
+int32 FNavigationToolSequence::GetSubSectionIndex() const
+{
+	return SubSectionIndex;
+}
+
+UMovieScene* FNavigationToolSequence::GetSequenceMovieScene() const
+{
+	UMovieSceneSequence* const Sequence = GetSequence();
+	if (IsValid(Sequence))
+	{
+		return Sequence->GetMovieScene();
+	}
+	return nullptr;
+}
+
+TArray<FMovieSceneBinding> FNavigationToolSequence::GetSortedBindings() const
+{
+	UMovieScene* const MovieScene = GetSequenceMovieScene();
+	if (!IsValid(MovieScene))
+	{
+		return {};
+	}
+
+	TArray<FMovieSceneBinding> Bindings = MovieScene->GetBindings();
+
+	Bindings.Sort([](const FMovieSceneBinding& InA, const FMovieSceneBinding& InB)
+		{
+			const int32 SortingOrderA = InA.GetSortingOrder();
+			const int32 SortingOrderB = InB.GetSortingOrder();
+			if (SortingOrderA == SortingOrderB)
+			{
+				return InA.GetName() < InB.GetName();
+			}
+			return SortingOrderA < SortingOrderB;
+		});
+
+	return MoveTemp(Bindings);
 }
 
 }
diff --git a/Engine/Source/Editor/Sequencer/Private/Tools/NavigationTool/Items/NavigationToolTrack.cpp b/Engine/Source/Editor/Sequencer/Private/Tools/NavigationTool/Items/NavigationToolTrack.cpp
--- a/Engine/Source/Editor/Sequencer/Private/Tools/NavigationTool/Items/NavigationToolTrack.cpp
+++ b/Engine/Source/Editor/Sequencer/Private/Tools/NavigationTool/Items/NavigationToolTrack.cpp
@@ -34,29 +34,11 @@ FNavigationToolTrack::FNavigationToolTrack(INavigationTool& InTool
 
 FNavigationToolItemId FNavigationToolTrack::CalculateItemId() const
 {
-	UMovieSceneTrack* const Track = GetTrack();
-	check(IsValid(Track));
-
-	UMovieSceneSequence* const Sequence = WeakSequence.Get();
-	check(IsValid(Sequence));
-
-	const FNavigationToolItemPtr ParentItem = GetParent();
-
-	FString NewId = ParentItem.IsValid() ? ParentItem->GetFullPath() : FString();
-
-	if (WeakSection.IsValid())
-	{
-		NewId += FSoftObjectPath(WeakSection.Get()).ToString()
-			+ FNavigationToolItemId::Separator
-			+ FString::FromInt(SubSectionIndex)
-			+ FNavigationToolItemId::Separator;
-	}
-
-	NewId += /*FSoftObjectPath(Sequence).ToString()
-		+ FNavigationToolItemId::Separator
-		+ */FSoftObjectPath(Track).ToString();
-
-	return FNavigationToolItemId(NewId);
+	return FNavigationToolItemId(GetParent()
+		, WeakSequence.Get()
+		, WeakSection.Get()
+		, SubSectionIndex
+		, FNavigationToolItemId::GetObjectPath(GetTrack()));
 }
 
 void FNavigationToolTrack::FindChildren(TArray<FNavigationToolItemPtr>& OutChildren, const bool bInRecursive)
diff --git a/Engine/Source/Editor/Sequencer/Private/Tools/NavigationTool/Menus/NavigationToolToolbarMenu.cpp b/Engine/Source/Editor/Sequencer/Private/Tools/NavigationTool/Menus/NavigationToolToolbarMenu.cpp
--- a/Engine/Source/Editor/Sequencer/Private/Tools/NavigationTool/Menus/NavigationToolToolbarMenu.cpp
+++ b/Engine/Source/Editor/Sequencer/Private/Tools/NavigationTool/Menus/NavigationToolToolbarMenu.cpp
@@ -382,6 +382,12 @@ void FNavigationToolToolbarMenu::OnToggleGlobalFilter(const TSharedRef<FNavigati
 
 	InFilter->SetActive(bNewActiveState);
 
+	UNavigationToolSettings* const ToolSettings = GetMutableDefault<UNavigationToolSettings>();
+	if (ensure(IsValid(ToolSettings)))
+	{
+		ToolSettings->SetBuiltInFilterEnabled(InFilter->GetFilterParams().GetFilterId(), bNewActiveState);
+	}
+
 	InTool->Refresh();
 }
 
diff --git a/Engine/Source/Editor/Sequencer/Private/Tools/NavigationTool/NavigationTool.cpp b/Engine/Source/Editor/Sequencer/Private/Tools/NavigationTool/NavigationTool.cpp
--- a/Engine/Source/Editor/Sequencer/Private/Tools/NavigationTool/NavigationTool.cpp
+++ b/Engine/Source/Editor/Sequencer/Private/Tools/NavigationTool/NavigationTool.cpp
@@ -54,19 +54,6 @@ void FNavigationTool::Init()
 
 	LoadAllProviderStates();
 
-	//LoadAllProviderStates()
-	{
-		GlobalFilters.Reset();
-
-		ForEachProvider([this](const TSharedRef<FNavigationToolProvider>& InProvider)
-			{
-				LoadProviderState(InProvider);
-				return true;
-			});
-
-		Refresh();
-	}
-
 	if (const TSharedPtr<FSequencer> Sequencer = StaticCastSharedPtr<FSequencer>(WeakSequencer.Pin()))
 	{
 		Sequencer->OnChannelChanged().AddSPLambda(this, [this](const FMovieSceneChannelMetaData*, UMovieSceneSection*)
@@ -142,7 +129,7 @@ void FNavigationTool::ForEachProvider(const TFunction<bool(const TSharedRef<FNav
 	FNavigationToolExtender::ForEachToolProvider(*this, InPredicate);
 }
 
-void FNavigationTool::LoadAllProviderStates()
+void FNavigationTool::LoadAllProviderStates(const bool bInRefreshTool)
 {
 	GlobalFilters.Reset();
 
@@ -152,7 +139,10 @@ void FNavigationTool::LoadAllProviderStates()
 			return true;
 		});
 
-	Refresh();
+	if (bInRefreshTool)
+	{
+		Refresh();
+	}
 }
 
 void FNavigationTool::SaveAllProviderStates()
@@ -173,41 +163,48 @@ void FNavigationTool::OnProvidersChanged(const TSharedRef<INavigationTool>& InTo
 		return;
 	}
 
+	LoadAllProviderStates(false);
+
 	if (InChangeType == ENavigationToolProvidersChangeType::Add)
 	{
 		const UNavigationToolSettings* const ToolSettings = GetDefault<UNavigationToolSettings>();
-
-		LoadProviderState(InProvider);
-
-		ForEachToolView([&InProvider](const TSharedPtr<FNavigationToolView>& InToolView)
-			{
-				InToolView->LoadViewState(InProvider);
-
-				InToolView->CreateColumns(InProvider);
-				InToolView->CreateDefaultColumnViews(InProvider);
-			});
-
-		// We apply default views *after* all columns have been created for all providers
-		if (ToolSettings->ShouldApplyDefaultColumnView())
+		if (ensure(IsValid(ToolSettings)))
 		{
+			LoadProviderState(InProvider);
+
 			ForEachToolView([&InProvider](const TSharedPtr<FNavigationToolView>& InToolView)
 				{
-					InToolView->ApplyCustomColumnView(InProvider->GetDefaultColumnView());
+					InToolView->LoadViewState(InProvider);
+
+					InToolView->CreateColumns(InProvider);
+					InToolView->CreateDefaultColumnViews(InProvider);
 				});
+
+			// We apply default views *after* all columns have been created for all providers
+			if (ToolSettings->ShouldApplyDefaultColumnView())
+			{
+				ForEachToolView([&InProvider](const TSharedPtr<FNavigationToolView>& InToolView)
+					{
+						InToolView->ApplyCustomColumnView(InProvider->GetDefaultColumnView());
+					});
+			}
 		}
 	}
 	else if (InChangeType == ENavigationToolProvidersChangeType::Remove)
 	{
 		// Remove extended column views
 		UNavigationToolSettings* const ToolSettings = GetMutableDefault<UNavigationToolSettings>();
-		TSet<FNavigationToolColumnView>& CustomColumnViews = ToolSettings->GetCustomColumnViews();
-
-		for (const FName ColumnViewName : InProvider->GetExtendedColumnViewNames())
+		if (ensure(IsValid(ToolSettings)))
 		{
-			CustomColumnViews.Remove(ColumnViewName);
-		}
+			TSet<FNavigationToolColumnView>& CustomColumnViews = ToolSettings->GetCustomColumnViews();
 
-		ToolSettings->SaveConfig();
+			for (const FName ColumnViewName : InProvider->GetExtendedColumnViewNames())
+			{
+				CustomColumnViews.Remove(ColumnViewName);
+			}
+
+			ToolSettings->SaveConfig();
+		}
 	}
 
 	Refresh();
@@ -355,6 +352,12 @@ void FNavigationTool::SaveProviderState(const TSharedRef<FNavigationToolProvider
 
 void FNavigationTool::LoadProviderState(const TSharedRef<FNavigationToolProvider>& InProvider)
 {
+	const UNavigationToolSettings* const ToolSettings = GetDefault<UNavigationToolSettings>();
+	if (!ensure(IsValid(ToolSettings)))
+	{
+		return;
+	}
+
 	InProvider->UpdateItemIdContexts(*this);
 	
 	if (FNavigationToolSaveState* const SaveState = InProvider->GetSaveState(*this))
@@ -389,7 +392,9 @@ void FNavigationTool::LoadProviderState(const TSharedRef<FNavigationToolProvider
 		{
 			const TSharedPtr<FNavigationToolBuiltInFilter> NewFilter = MakeShared<FNavigationToolBuiltInFilter>(NewFilterParams);
 
-			if (!NewFilterParams.IsEnabledByDefault())
+			const TSet<FName>& EnabledBuiltInFilters = ToolSettings->GetEnabledBuiltInFilters();
+			if (!NewFilterParams.IsEnabledByDefault()
+				&& !EnabledBuiltInFilters.Contains(NewFilterParams.GetFilterId()))
 			{
 				NewFilter->SetActive(false);
 			}
diff --git a/Engine/Source/Editor/Sequencer/Private/Tools/NavigationTool/NavigationTool.h b/Engine/Source/Editor/Sequencer/Private/Tools/NavigationTool/NavigationTool.h
--- a/Engine/Source/Editor/Sequencer/Private/Tools/NavigationTool/NavigationTool.h
+++ b/Engine/Source/Editor/Sequencer/Private/Tools/NavigationTool/NavigationTool.h
@@ -48,7 +48,7 @@ public:
 
 	void ForEachProvider(const TFunction<bool(const TSharedRef<FNavigationToolProvider>& InToolProvider)>& InPredicate);
 
-	void LoadAllProviderStates();
+	void LoadAllProviderStates(const bool bInRefreshTool = true);
 	void SaveAllProviderStates();
 
 	bool CanProcessSequenceSpawn(UMovieSceneSequence* const InSequence) const;
diff --git a/Engine/Source/Editor/Sequencer/Private/Tools/NavigationTool/NavigationToolSettings.cpp b/Engine/Source/Editor/Sequencer/Private/Tools/NavigationTool/NavigationToolSettings.cpp
--- a/Engine/Source/Editor/Sequencer/Private/Tools/NavigationTool/NavigationToolSettings.cpp
+++ b/Engine/Source/Editor/Sequencer/Private/Tools/NavigationTool/NavigationToolSettings.cpp
@@ -120,6 +120,25 @@ FNavigationToolColumnView* UNavigationToolSettings::FindCustomColumnView(const F
 	return nullptr;
 }
 
+void UNavigationToolSettings::SetBuiltInFilterEnabled(const FName InFilterName, const bool bInEnabled)
+{
+	if (InFilterName.IsNone())
+	{
+		return;
+	}
+
+	if (bInEnabled)
+	{
+		EnabledBuiltInFilters.Add(InFilterName);
+	}
+	else
+	{
+		EnabledBuiltInFilters.Remove(InFilterName);
+	}
+
+	SaveConfig();
+}
+
 FSequencerFilterBarConfig& UNavigationToolSettings::FindOrAddFilterBar(const FName InIdentifier, const bool bInSaveConfig)
 {
 	FSequencerFilterBarConfig& FilterBarSettings = NavigationToolFilterBars.FindOrAdd(InIdentifier);
diff --git a/Engine/Source/Editor/Sequencer/Private/Tools/NavigationTool/NavigationToolTab.cpp b/Engine/Source/Editor/Sequencer/Private/Tools/NavigationTool/NavigationToolTab.cpp
--- a/Engine/Source/Editor/Sequencer/Private/Tools/NavigationTool/NavigationToolTab.cpp
+++ b/Engine/Source/Editor/Sequencer/Private/Tools/NavigationTool/NavigationToolTab.cpp
@@ -22,6 +22,8 @@ FNavigationToolTab::FNavigationToolTab(INavigationTool& InOwnerTool)
 
 void FNavigationToolTab::Init()
 {
+	bShuttingDown = false;
+
 	RegisterToolTab();
 
 	USequencerSettings* const SequencerSettings = GetSequencerSettings();
@@ -36,6 +38,8 @@ void FNavigationToolTab::Init()
 
 void FNavigationToolTab::Shutdown()
 {
+	bShuttingDown = true;
+
 	CloseToolTab();
 	UnregisterToolTab();
 }
@@ -62,12 +66,6 @@ bool FNavigationToolTab::IsToolTabVisible() const
 
 void FNavigationToolTab::ShowHideToolTab(const bool bInVisible)
 {
-	USequencerSettings* const SequencerSettings = GetSequencerSettings();
-	if (ensure(IsValid(SequencerSettings)))
-	{
-		SequencerSettings->SetNavigationToolVisible(bInVisible);
-	}
-
 	if (bInVisible)
 	{
 		if (const TSharedPtr<FTabManager> TabManager = GetTabManager())
@@ -178,28 +176,26 @@ TSharedRef<SDockTab> FNavigationToolTab::SpawnToolTab(const FSpawnTabArgs& InArg
 	check(ToolWidget.IsValid());
 
 	const TSharedRef<SDockTab> ToolTab = SNew(SDockTab)
-		.OnTabClosed_Lambda([](TSharedRef<SDockTab> InDockTab)
+		.OnTabClosed_Lambda([this](const TSharedRef<SDockTab> InDockTab)
 			{
-				USequencerSettings* const Settings = GetMutableDefault<USequencerSettings>();
-				if (IsValid(Settings))
+				USequencerSettings* const Settings = GetSequencerSettings();
+				if (IsValid(Settings) && !bShuttingDown)
 				{
 					Settings->SetNavigationToolVisible(false);
-					Settings->SaveConfig();
 				}
 			})
 		[
 			ToolWidget.ToSharedRef()
 		];
 
-	OwnerTool.Refresh();
-
-	// Set the Navigation Tool last visible state to restore when Sequencer is opened
-	USequencerSettings* const SequencerSettings = GetMutableDefault<USequencerSettings>();
+	USequencerSettings* const SequencerSettings = GetSequencerSettings();
 	if (IsValid(SequencerSettings))
 	{
 		SequencerSettings->SetNavigationToolVisible(true);
 	}
 
+	OwnerTool.Refresh();
+
 	return ToolTab;
 }
 
diff --git a/Engine/Source/Editor/Sequencer/Private/Tools/NavigationTool/NavigationToolTab.h b/Engine/Source/Editor/Sequencer/Private/Tools/NavigationTool/NavigationToolTab.h
--- a/Engine/Source/Editor/Sequencer/Private/Tools/NavigationTool/NavigationToolTab.h
+++ b/Engine/Source/Editor/Sequencer/Private/Tools/NavigationTool/NavigationToolTab.h
@@ -54,6 +54,9 @@ protected:
 	INavigationTool& OwnerTool;
 
 	FName NavigationToolTabId;
+
+	/** State variable to help the dock tab know when it should save its visibility */
+	bool bShuttingDown = false;
 };
 
 }
diff --git a/Engine/Source/Editor/Sequencer/Private/Tools/NavigationTool/Providers/NavigationToolProvider.cpp b/Engine/Source/Editor/Sequencer/Private/Tools/NavigationTool/Providers/NavigationToolProvider.cpp
--- a/Engine/Source/Editor/Sequencer/Private/Tools/NavigationTool/Providers/NavigationToolProvider.cpp
+++ b/Engine/Source/Editor/Sequencer/Private/Tools/NavigationTool/Providers/NavigationToolProvider.cpp
@@ -17,6 +17,17 @@ void FNavigationToolProvider::OnExtendColumnViews(TSet<FNavigationToolColumnView
 	}
 }
 
+void FNavigationToolProvider::OnExtendBuiltInFilters(TArray<FNavigationToolBuiltInFilterParams>& OutFilterParams)
+{
+	INavigationToolProvider::OnExtendBuiltInFilters(OutFilterParams);
+
+	ExtendedBuiltInFilterNames.Reset(OutFilterParams.Num());
+	for (const FNavigationToolBuiltInFilterParams& FilterParams : OutFilterParams)
+	{
+		ExtendedBuiltInFilterNames.Add(FilterParams.GetFilterId());
+	}
+}
+
 void FNavigationToolProvider::UpdateItemIdContexts(const INavigationTool& InTool)
 {
 	FNavigationToolSaveState* const SaveState = GetSaveState(InTool);
@@ -109,6 +120,11 @@ const TArray<FName>& FNavigationToolProvider::GetExtendedColumnViewNames() const
 	return ExtendedColumnViewNames;
 }
 
+const TArray<FName>& FNavigationToolProvider::GetExtendedBuiltInFilterNames() const
+{
+	return ExtendedBuiltInFilterNames;
+}
+
 TOptional<EItemDropZone> FNavigationToolProvider::OnToolItemCanAcceptDrop(const FDragDropEvent& InDragDropEvent
 	, EItemDropZone InDropZone
 	, const FNavigationToolItemPtr& InTargetItem) const
diff --git a/Engine/Source/Editor/Sequencer/Private/Tools/NavigationTool/Utils/NavigationToolMiscUtils.cpp b/Engine/Source/Editor/Sequencer/Private/Tools/NavigationTool/Utils/NavigationToolMiscUtils.cpp
--- a/Engine/Source/Editor/Sequencer/Private/Tools/NavigationTool/Utils/NavigationToolMiscUtils.cpp
+++ b/Engine/Source/Editor/Sequencer/Private/Tools/NavigationTool/Utils/NavigationToolMiscUtils.cpp
@@ -4,6 +4,7 @@
 #include "ISequencer.h"
 #include "ISourceControlModule.h"
 #include "ISourceControlProvider.h"
+#include "Evaluation/MovieSceneEvaluationTemplateInstance.h"
 #include "Sections/MovieSceneSubSection.h"
 #include "Tools/NavigationTool/INavigationTool.h"
 #include "UObject/Package.h"
@@ -11,7 +12,7 @@
 namespace UE::Sequencer::NavigationTool
 {
 
-FSourceControlStatePtr GetSourceControlState(UPackage* const InPackage)
+FSourceControlStatePtr FindSourceControlState(UPackage* const InPackage)
 {
 	if (!IsValid(InPackage))
 	{
@@ -33,9 +34,9 @@ FSourceControlStatePtr GetSourceControlState(UPackage* const InPackage)
 	return SourceControlProvider.GetState(InPackage, EStateCacheUsage::Use);
 }
 
-const FSlateBrush* GetSourceControlStatusBrush(UPackage* const InPackage)
+const FSlateBrush* FindSourceControlStatusBrush(UPackage* const InPackage)
 {
-	if (const FSourceControlStatePtr SourceControlState = GetSourceControlState(InPackage))
+	if (const FSourceControlStatePtr SourceControlState = FindSourceControlState(InPackage))
 	{
 		const FSlateIcon Icon = SourceControlState->GetIcon();
 		if (Icon.IsSet())
@@ -46,9 +47,9 @@ const FSlateBrush* GetSourceControlStatusBrush(UPackage* const InPackage)
 	return nullptr;
 }
 
-FText GetSourceControlStatusText(UPackage* const InPackage)
+FText FindSourceControlStatusText(UPackage* const InPackage)
 {
-	if (const FSourceControlStatePtr SourceControlState = GetSourceControlState(InPackage))
+	if (const FSourceControlStatePtr SourceControlState = FindSourceControlState(InPackage))
 	{
 		return SourceControlState->GetDisplayTooltip();
 	}
@@ -78,4 +79,24 @@ void FocusSequence(const INavigationTool& InTool, UMovieSceneSequence& InSequenc
 	}
 }
 
+FMovieSceneSequenceID ResolveSequenceID(ISequencer& InSequencer, UMovieSceneSequence* const InSequence)
+{
+	if (FMovieSceneEvaluationState* const EvaluationState = InSequencer.GetSharedPlaybackState()->FindCapability<FMovieSceneEvaluationState>())
+	{
+		return EvaluationState->FindSequenceId(InSequence);
+	}
+	return FMovieSceneSequenceID();
+}
+
+TArrayView<TWeakObjectPtr<>> ResolveBoundObjects(ISequencer& InSequencer
+	, UMovieSceneSequence* const InSequence, const FGuid& InBindingId)
+{
+	if (FMovieSceneEvaluationState* const EvaluationState = InSequencer.GetSharedPlaybackState()->FindCapability<FMovieSceneEvaluationState>())
+	{
+		const FMovieSceneSequenceIDRef SequenceID = EvaluationState->FindSequenceId(InSequence);
+		return InSequencer.FindBoundObjects(InBindingId, SequenceID);
+	}
+	return TArrayView<TWeakObjectPtr<>>();
+}
+
 }
diff --git a/Engine/Source/Editor/Sequencer/Private/Tools/NavigationTool/Utils/NavigationToolMiscUtils.h b/Engine/Source/Editor/Sequencer/Private/Tools/NavigationTool/Utils/NavigationToolMiscUtils.h
--- a/Engine/Source/Editor/Sequencer/Private/Tools/NavigationTool/Utils/NavigationToolMiscUtils.h
+++ b/Engine/Source/Editor/Sequencer/Private/Tools/NavigationTool/Utils/NavigationToolMiscUtils.h
@@ -4,22 +4,59 @@
 
 #include "ISourceControlState.h"
 
+class FText;
+class ISequencer;
 class UMovieSceneSequence;
 class UPackage;
+struct FGuid;
+struct FMovieSceneSequenceID;
+struct FSlateBrush;
 
 namespace UE::Sequencer::NavigationTool
 {
 
 class INavigationTool;
 
-/**  */
-FSourceControlStatePtr GetSourceControlState(UPackage* const InPackage);
-/**  */
-const FSlateBrush* GetSourceControlStatusBrush(UPackage* const InPackage);
-/**  */
-FText GetSourceControlStatusText(UPackage* const InPackage);
+/**
+ * Find the source control state for a given package.
+ * @param InPackage The package for which the source control state is being queried.
+ * @return The source control state if available, otherwise nullptr.
+ */
+FSourceControlStatePtr FindSourceControlState(UPackage* const InPackage);
+
+/**
+ * Find the source control status brush associated with a given package.
+ * @param InPackage The package for which the source control status brush is being queried.
+ * @return A pointer to the slate brush representing the source control status, or nullptr if unavailable.
+ */
+const FSlateBrush* FindSourceControlStatusBrush(UPackage* const InPackage);
+
+/**
+ * Find the source control status text for the specified package.
+ * @param InPackage The package for which the source control status text is being requested.
+ * @return The tooltip text representing the source control status if available, otherwise an empty text.
+ */
+FText FindSourceControlStatusText(UPackage* const InPackage);
 
 /** Focuses a sub-movie scene (MovieScene within a MovieScene) in the sequencer. */
 void FocusSequence(const INavigationTool& InTool, UMovieSceneSequence& InSequence);
 
+/**
+ * Resolve the sequence ID for a given sequence within the provided sequencer.
+ * @param InSequencer The sequencer instance to retrieve the sequence ID from.
+ * @param InSequence The sequence whose ID is to be resolved.
+ * @return The resolved sequence ID if found, otherwise an empty sequence ID.
+ */
+FMovieSceneSequenceID ResolveSequenceID(ISequencer& InSequencer, UMovieSceneSequence* const InSequence);
+
+/**
+ * Resolves the objects bound to a specified binding ID within a sequence.
+ * @param InSequencer The sequencer instance used for evaluating bindings and sequences.
+ * @param InSequence The movie scene sequence to search for bindings.
+ * @param InBindingId The binding ID whose bound objects are being resolved.
+ * @return A view of the objects bound to the specified binding ID in the specified sequence.
+ */
+TArrayView<TWeakObjectPtr<>> ResolveBoundObjects(ISequencer& InSequencer
+	, UMovieSceneSequence* const InSequence, const FGuid& InBindingId);
+
 }
diff --git a/Engine/Source/Editor/Sequencer/Private/Tools/NavigationTool/Utils/NavigationToolMovieSceneUtils.cpp b/Engine/Source/Editor/Sequencer/Private/Tools/NavigationTool/Utils/NavigationToolMovieSceneUtils.cpp
--- a/Engine/Source/Editor/Sequencer/Private/Tools/NavigationTool/Utils/NavigationToolMovieSceneUtils.cpp
+++ b/Engine/Source/Editor/Sequencer/Private/Tools/NavigationTool/Utils/NavigationToolMovieSceneUtils.cpp
@@ -108,4 +108,17 @@ void GetSequenceSubSections(UMovieSceneSequence* const InSequence, TArray<UMovie
 	}
 }
 
+const UClass* GetBindingObjectClass(UMovieScene& InMovieScene, const FGuid& InBindingGuid)
+{
+	if (const FMovieScenePossessable* const Possessable = InMovieScene.FindPossessable(InBindingGuid))
+	{
+		return Possessable->GetPossessedObjectClass();
+	}
+	if (const FMovieSceneSpawnable* const Spawnable = InMovieScene.FindSpawnable(InBindingGuid))
+	{
+		return Spawnable->GetObjectTemplate()->GetClass();
+	}
+	return nullptr;
+}
+
 }
diff --git a/Engine/Source/Editor/Sequencer/Private/Tools/NavigationTool/Utils/NavigationToolMovieSceneUtils.h b/Engine/Source/Editor/Sequencer/Private/Tools/NavigationTool/Utils/NavigationToolMovieSceneUtils.h
--- a/Engine/Source/Editor/Sequencer/Private/Tools/NavigationTool/Utils/NavigationToolMovieSceneUtils.h
+++ b/Engine/Source/Editor/Sequencer/Private/Tools/NavigationTool/Utils/NavigationToolMovieSceneUtils.h
@@ -5,27 +5,53 @@
 #include "Containers/Array.h"
 
 class ISequencer;
+class UClass;
+class UMovieScene;
 class UMovieSceneSequence;
 class UMovieSceneSubSection;
+struct FGuid;
 
 namespace UE::Sequencer::NavigationTool
 {
 
-/**  */
+/**
+ * Finds the subsection corresponding to a given sequence within the specified sequencer.
+ * @param InSequencer The sequencer to search within.
+ * @param InSequence The movie scene sequence for which the subsection is to be found.
+ * @return The found movie scene subsection, or nullptr if no matching subsection is found.
+ */
 UMovieSceneSubSection* FindSequenceSubSection(ISequencer& InSequencer, UMovieSceneSequence* const InSequence);
 
-/**  */
+/**
+ * Checks if globally marked frames are enabled for the given movie scene sequence.
+ * @param InSequence The movie scene sequence to check.
+ * @return True if globally marked frames are enabled for the specified sequence, false otherwise.
+ */
 bool IsGloballyMarkedFramesForSequence(UMovieSceneSequence* const InSequence);
 
-/**  */
+/**
+ * Toggles the visibility of globally marked frames for the specified sequence within the provided sequencer.
+ * @param InSequencer The sequencer controlling the movie scene.
+ * @param InSequence The movie scene sequence for which the globally marked frame visibility is to be changed.
+ * @param bInVisible Indicates whether the globally marked frames should be visible (true) or hidden (false).
+ */
 void ShowGloballyMarkedFramesForSequence(ISequencer& InSequencer
-	, UMovieSceneSequence* const InSequence
-	, const bool bInVisible);
+	, UMovieSceneSequence* const InSequence, const bool bInVisible);
 
-/**  */
+/**
+ * Modifies the given sequence and its associated movie scene to ensure they are marked as modified.
+ * @param InSequence The movie scene sequence to modify, including its associated movie scene.
+ */
 void ModifySequenceAndMovieScene(UMovieSceneSequence* const InSequence);
 
-/**  */
+/**
+ * Retrieves all sections as subsections from the given movie scene sequence.
+ * @param InSequence The movie scene sequence to retrieve subsections from.
+ * @param OutSubSections The array to be populated with found movie scene subsections.
+ */
 void GetSequenceSubSections(UMovieSceneSequence* const InSequence, TArray<UMovieSceneSubSection*>& OutSubSections);
 
+/** Attempts to get the bound object class from a binding */
+const UClass* GetBindingObjectClass(UMovieScene& InMovieScene, const FGuid& InBindingGuid);
+
 }
diff --git a/Engine/Source/Editor/Sequencer/Public/Tools/NavigationTool/Filters/Filters/NavigationToolBuiltInFilterParams.h b/Engine/Source/Editor/Sequencer/Public/Tools/NavigationTool/Filters/Filters/NavigationToolBuiltInFilterParams.h
--- a/Engine/Source/Editor/Sequencer/Public/Tools/NavigationTool/Filters/Filters/NavigationToolBuiltInFilterParams.h
+++ b/Engine/Source/Editor/Sequencer/Public/Tools/NavigationTool/Filters/Filters/NavigationToolBuiltInFilterParams.h
@@ -14,6 +14,7 @@ struct FNavigationToolBuiltInFilterParams
 	GENERATED_BODY()
 
 	SEQUENCER_API static FNavigationToolBuiltInFilterParams CreateSequenceFilter();
+	SEQUENCER_API static FNavigationToolBuiltInFilterParams CreateTrackFilter();
 	SEQUENCER_API static FNavigationToolBuiltInFilterParams CreateBindingFilter();
 
 	FNavigationToolBuiltInFilterParams() = default;
diff --git a/Engine/Source/Editor/Sequencer/Public/Tools/NavigationTool/Items/NavigationToolActor.h b/Engine/Source/Editor/Sequencer/Public/Tools/NavigationTool/Items/NavigationToolActor.h
--- a/Engine/Source/Editor/Sequencer/Public/Tools/NavigationTool/Items/NavigationToolActor.h
+++ b/Engine/Source/Editor/Sequencer/Public/Tools/NavigationTool/Items/NavigationToolActor.h
@@ -21,9 +21,8 @@ public:
 
 	FNavigationToolActor(INavigationTool& InTool
 		, const FNavigationToolItemPtr& InParentItem
-		, const FMovieSceneBinding& InBinding
-		, UObject* const InObject
-		, const FNavigationToolItemPtr& InReferencingItem);
+		, const TSharedPtr<FNavigationToolSequence>& InParentSequenceItem
+		, const FMovieSceneBinding& InBinding);
 
 	//~ Begin INavigationToolItem
 
@@ -34,7 +33,6 @@ public:
 	virtual ENavigationToolItemViewMode GetSupportedViewModes(const INavigationToolView& InToolView) const override;
 	virtual bool CanReceiveParentVisibilityPropagation() const override { return true; }
 
-	virtual FText GetDisplayName() const override;
 	virtual TArray<FName> GetTags() const override;
 
 	virtual bool ShowVisibility() const override { return true; }
diff --git a/Engine/Source/Editor/Sequencer/Public/Tools/NavigationTool/Items/NavigationToolBinding.h b/Engine/Source/Editor/Sequencer/Public/Tools/NavigationTool/Items/NavigationToolBinding.h
--- a/Engine/Source/Editor/Sequencer/Public/Tools/NavigationTool/Items/NavigationToolBinding.h
+++ b/Engine/Source/Editor/Sequencer/Public/Tools/NavigationTool/Items/NavigationToolBinding.h
@@ -11,44 +11,52 @@
 namespace UE::Sequencer::NavigationTool
 {
 
+class FNavigationToolSequence;
+
 /**
  * Navigation Tool Item representing a Sequence binding
  */
 class SEQUENCER_API FNavigationToolBinding
-	: public FNavigationToolObject
+	: public FNavigationToolItem
 	, public IRenameableExtension
 	//, public ISequenceLockableExtension
 {
 public:
 	UE_NAVIGATIONTOOL_INHERITS_WITH_SUPER(FNavigationToolBinding
-		, FNavigationToolObject
+		, FNavigationToolItem
 		, IRenameableExtension
 		/*, ISequenceLockableExtension*/);
 
 	FNavigationToolBinding(INavigationTool& InTool
 		, const FNavigationToolItemPtr& InParentItem
-		, const FMovieSceneBinding& InBinding
-		, UObject* const InObject
-		, const FNavigationToolItemPtr& InReferencingItem);
+		, const TSharedPtr<FNavigationToolSequence>& InParentSequenceItem
+		, const FMovieSceneBinding& InBinding);
 
 	//~ Begin INavigationToolItem
 
+	virtual bool IsItemValid() const override;
+
 	virtual void FindChildren(TArray<FNavigationToolItemPtr>& OutChildren, const bool bInRecursive) override;
 
 	virtual bool CanBeTopLevel() const override { return false; }
 	virtual bool ShouldSort() const override { return false; }
 
 	virtual FText GetDisplayName() const override;
-	virtual const FSlateBrush* GetDefaultIconBrush() const override;
+	virtual FText GetClassName() const override;
 	virtual FSlateIcon GetIcon() const override;
 	virtual FText GetIconTooltipText() const override;
+	virtual FLinearColor GetItemBrushColor() const override;
 
 	virtual bool IsSelected(const FNavigationToolScopedSelection& InSelection) const override;
 	virtual void Select(FNavigationToolScopedSelection& InSelection) const override;
+	virtual void OnSelect() override;
+	virtual void OnDoubleClick() override;
 
 	virtual bool CanDelete() const override;
 	virtual bool Delete() override;
 
+	virtual FNavigationToolItemId CalculateItemId() const override;
+
 	//~ End INavigationToolItem
 
 	//~ Begin IRenameableExtension
@@ -61,14 +69,24 @@ public:
 	//virtual void SetIsLocked(const bool bInIsLocked) override;
 	//~ End ISequenceLockableExtension
 
+	const UClass* GetBoundObjectClass() const;
+
 	const FMovieSceneBinding& GetBinding() const;
 
-	UObject* FindBoundObject() const;
+	/** @return The cached object that is bound in sequencer. */
+	UObject* GetCachedBoundObject() const;
+
+	/** Caches the object if it's not already been cached and returns the result. */
+	UObject* CacheBoundObject();
 
 protected:
-	FSlateIcon Icon;
+	TWeakPtr<FNavigationToolSequence> WeakParentSequenceItem;
 
 	FMovieSceneBinding Binding;
+
+	TWeakObjectPtr<> WeakBoundObject;
+
+	FSlateIcon Icon;
 };
 
 }
diff --git a/Engine/Source/Editor/Sequencer/Public/Tools/NavigationTool/Items/NavigationToolComponent.h b/Engine/Source/Editor/Sequencer/Public/Tools/NavigationTool/Items/NavigationToolComponent.h
--- a/Engine/Source/Editor/Sequencer/Public/Tools/NavigationTool/Items/NavigationToolComponent.h
+++ b/Engine/Source/Editor/Sequencer/Public/Tools/NavigationTool/Items/NavigationToolComponent.h
@@ -21,9 +21,8 @@ public:
 
 	FNavigationToolComponent(INavigationTool& InTool
 		, const FNavigationToolItemPtr& InParentItem
-		, const FMovieSceneBinding& InBinding
-		, UObject* const InObject
-		, const FNavigationToolItemPtr& InReferencingItem);
+		, const TSharedPtr<FNavigationToolSequence>& InParentSequenceItem
+		, const FMovieSceneBinding& InBinding);
 
 	//~ Begin INavigationToolItem
 
diff --git a/Engine/Source/Editor/Sequencer/Public/Tools/NavigationTool/Items/NavigationToolItemId.h b/Engine/Source/Editor/Sequencer/Public/Tools/NavigationTool/Items/NavigationToolItemId.h
--- a/Engine/Source/Editor/Sequencer/Public/Tools/NavigationTool/Items/NavigationToolItemId.h
+++ b/Engine/Source/Editor/Sequencer/Public/Tools/NavigationTool/Items/NavigationToolItemId.h
@@ -6,7 +6,7 @@
 #include "UObject/SoftObjectPath.h"
 
 class UMovieSceneSequence;
-class UMovieSceneSubSection;
+class UMovieSceneSection;
 class UObject;
 
 namespace UE::Sequencer::NavigationTool
@@ -17,6 +17,8 @@ class INavigationToolItem;
 /** Struct to identify an item in the Navigation Tool */
 struct SEQUENCER_API FNavigationToolItemId
 {
+	static FString GetObjectPath(const UObject* const InObject);
+
 	static FNavigationToolItemId RootId;
 
 	static constexpr const TCHAR* Separator = TEXT(",");
@@ -30,7 +32,7 @@ struct SEQUENCER_API FNavigationToolItemId
 	/**
 	 * Ctor used for Objects that appear ONCE in the Navigation Tool (e.g. an Actor or Component), so only the Object itself is needed for the Id
 	 * for objects that can appear multiple times because they're referenced multiple times, see the ctor below that takes in a Referencing Item and Referencing Id
-	 * @param InObject the object that the item represents
+	 * @param InObject The object that the item represents
 	 */
 	FNavigationToolItemId(const UObject* InObject);
 
@@ -45,36 +47,44 @@ struct SEQUENCER_API FNavigationToolItemId
 	 * [Component Path], [Material Instance Dynamic Path], [Slot Index]
 	 * "/Game/World.World:PersistentLevel.StaticMeshActor_0.StaticMeshComponent,/Game/World.World:PersistentLevel.StaticMeshActor_0.StaticMeshComponent.MaterialInstanceDynamic_0,[Slot 0]"
 	 * 
-	 * @param InObject the Object being referenced multiple times (e.g. a Material)
-	 * @param InReferencingItem the Item holding a reference to the object (e.g. a Primitive Comp referencing a Material)
-	 * @param InReferencingId the Id to differentiate this reference from other references within the same object (e.g. can be slot or property name)
+	 * @param InObject The Object being referenced multiple times (e.g. a Material)
+	 * @param InReferencingItem The Item holding a reference to the object (e.g. a Primitive Comp referencing a Material)
+	 * @param InReferencingId The Id to differentiate this reference from other references within the same object (e.g. can be slot or property name)
 	 */
 	FNavigationToolItemId(const UObject* InObject
 		, const FNavigationToolItemPtr& InReferencingItem
 		, const FString& InReferencingId);
 
 	/**
-	 * Ctor used for making the Item Id for a item proxy that will be under the given Parent Item.
+	 * Ctor used for making the Item Id for an item proxy that will be under the given Parent Item.
 	 * Used when the actual Item Proxy is not created yet, but know it's factory and want to know whether the item proxy already exists
-	 * @param InParentItem the parent item holding the item proxy
-	 * @param InItemProxyFactory the factory responsible of creating the item proxy
+	 * @param InParentItem The parent item holding the item proxy
+	 * @param InItemProxyFactory The factory responsible for creating the item proxy
 	 */
 	FNavigationToolItemId(const FNavigationToolItemPtr& InParentItem
 		, const INavigationToolItemProxyFactory& InItemProxyFactory);
 
 	/**
-	 * Ctor used for making the Item Id for a Item Proxy under the given Parent Item.
-	 * @param InParentItem the parent item holding the item proxy
-	 * @param InItemProxy the item proxy under the parent item
+	 * Ctor used for making the Item Id for an Item Proxy under the given Parent Item.
+	 * @param InParentItem The parent item holding the item proxy
+	 * @param InItemProxy The item proxy under the parent item
 	 */
 	FNavigationToolItemId(const FNavigationToolItemPtr& InParentItem
 		, const INavigationToolItem& InItemProxy);
 
-	/**  */
+	/**
+	 * Ctor used for making the Item Id for most sequence items
+	 * @param InParentItem The parent item holding the item proxy
+	 * @param InSequence The parent sequence object
+	 * @param InSection The parent section this sequence belongs to, if any
+	 * @param InSectionIndex The parent section index, if any. Should be specified if InSection is used.
+	 * @param InReferenceId The optional additional reference string
+	 */
 	FNavigationToolItemId(const FNavigationToolItemPtr& InParentItem
 		, const UMovieSceneSequence* const InSequence
-		, const UMovieSceneSubSection* const InSubSection = nullptr
-		, const int32 InSectionIndex = 0);
+		, const UMovieSceneSection* const InSection = nullptr
+		, const int32 InSectionIndex = 0
+		, const FString& InReferenceId = FString());
 
 	FNavigationToolItemId(const FNavigationToolItemId& Other);
 	FNavigationToolItemId(FNavigationToolItemId&& Other) noexcept;
@@ -95,8 +105,6 @@ struct SEQUENCER_API FNavigationToolItemId
 	FString GetStringId() const;
 
 private:
-	static FString GetObjectPath(const UObject* const InObject);
-
 	static void AddSeparatedSegment(FString& OutString, const FString& InSegment);
 
 	void CalculateTypeHash();
diff --git a/Engine/Source/Editor/Sequencer/Public/Tools/NavigationTool/Items/NavigationToolSequence.h b/Engine/Source/Editor/Sequencer/Public/Tools/NavigationTool/Items/NavigationToolSequence.h
--- a/Engine/Source/Editor/Sequencer/Public/Tools/NavigationTool/Items/NavigationToolSequence.h
+++ b/Engine/Source/Editor/Sequencer/Public/Tools/NavigationTool/Items/NavigationToolSequence.h
@@ -13,6 +13,7 @@
 class UMovieScene;
 class UMovieSceneSequence;
 class UMovieSceneSubSection;
+struct FMovieSceneBinding;
 
 namespace UE::Sequencer
 {
@@ -66,10 +67,6 @@ public:
 	SEQUENCER_API virtual void OnDoubleClick() override;
 	//~ End INavigationToolItem
 
-	SEQUENCER_API UMovieSceneSequence* GetSequence() const;
-	SEQUENCER_API UMovieSceneSubSection* GetSubSection() const;
-	SEQUENCER_API UMovieScene* GetSequenceMovieScene() const;
-
 	//~ Begin IRenameableExtension
 	SEQUENCER_API virtual bool CanRename() const override;
 	SEQUENCER_API virtual bool Rename(const FString& InName) override;
@@ -102,6 +99,15 @@ public:
 	SEQUENCER_API virtual FText GetRevisionControlStatusText() const override;
 	//~ End IRevisionControlExtension
 
+	SEQUENCER_API UMovieSceneSequence* GetSequence() const;
+
+	SEQUENCER_API UMovieSceneSubSection* GetSubSection() const;
+	SEQUENCER_API int32 GetSubSectionIndex() const;
+
+	SEQUENCER_API UMovieScene* GetSequenceMovieScene() const;
+
+	SEQUENCER_API TArray<FMovieSceneBinding> GetSortedBindings() const;
+
 protected:
 	//~Begin FNavigationToolItem
 	SEQUENCER_API virtual FNavigationToolItemId CalculateItemId() const override;
diff --git a/Engine/Source/Editor/Sequencer/Public/Tools/NavigationTool/NavigationToolSettings.h b/Engine/Source/Editor/Sequencer/Public/Tools/NavigationTool/NavigationToolSettings.h
--- a/Engine/Source/Editor/Sequencer/Public/Tools/NavigationTool/NavigationToolSettings.h
+++ b/Engine/Source/Editor/Sequencer/Public/Tools/NavigationTool/NavigationToolSettings.h
@@ -84,6 +84,9 @@ public:
 	TSet<FNavigationToolColumnView>& GetCustomColumnViews() { return CustomColumnViews; }
 	FNavigationToolColumnView* FindCustomColumnView(const FName InColumnId);
 
+	const TSet<FName>& GetEnabledBuiltInFilters() const { return EnabledBuiltInFilters; }
+	void SetBuiltInFilterEnabled(const FName InFilterName, const bool bInEnabled);
+
 	FSequencerFilterBarConfig& FindOrAddFilterBar(const FName InIdentifier, const bool bInSaveConfig);
 	FSequencerFilterBarConfig* FindFilterBar(const FName InIdentifier);
 	bool RemoveFilterBar(const FName InIdentifier);
@@ -135,6 +138,9 @@ private:
 	UPROPERTY(Config)
 	TSet<FNavigationToolColumnView> CustomColumnViews;
 
+	UPROPERTY(Config)
+	TSet<FName> EnabledBuiltInFilters;
+
 	/** Saved settings for each unique filter bar instance mapped by instance identifier */
 	UPROPERTY(config, EditAnywhere, Category = "Filtering")
 	TMap<FName, FSequencerFilterBarConfig> NavigationToolFilterBars;
diff --git a/Engine/Source/Editor/Sequencer/Public/Tools/NavigationTool/Providers/NavigationToolProvider.h b/Engine/Source/Editor/Sequencer/Public/Tools/NavigationTool/Providers/NavigationToolProvider.h
--- a/Engine/Source/Editor/Sequencer/Public/Tools/NavigationTool/Providers/NavigationToolProvider.h
+++ b/Engine/Source/Editor/Sequencer/Public/Tools/NavigationTool/Providers/NavigationToolProvider.h
@@ -18,6 +18,7 @@ public:
 	//~ Begin INavigationToolProvider
 
 	SEQUENCER_API virtual void OnExtendColumnViews(TSet<FNavigationToolColumnView>& OutColumnViews) override;
+	SEQUENCER_API virtual void OnExtendBuiltInFilters(TArray<FNavigationToolBuiltInFilterParams>& OutFilterParams) override;
 
 	SEQUENCER_API virtual bool ShouldLockTool() const override { return false; }
 	SEQUENCER_API virtual bool ShouldHideItem(const FNavigationToolItemPtr& InItem) const override { return false; }
@@ -36,9 +37,11 @@ public:
 	void EnsureToolViewCount(const INavigationTool& InTool, const int32 InToolViewId);
 
 	const TArray<FName>& GetExtendedColumnViewNames() const;
+	const TArray<FName>& GetExtendedBuiltInFilterNames() const;
 
 protected:
 	TArray<FName> ExtendedColumnViewNames;
+	TArray<FName> ExtendedBuiltInFilterNames;
 };
 
 }
