[Sequencer] Fix horde issue of const qualifier on reference type

diff --git a/Engine/Source/Editor/Sequencer/Private/Tools/NavigationTool/Utils/NavigationToolMiscUtils.cpp b/Engine/Source/Editor/Sequencer/Private/Tools/NavigationTool/Utils/NavigationToolMiscUtils.cpp
--- a/Engine/Source/Editor/Sequencer/Private/Tools/NavigationTool/Utils/NavigationToolMiscUtils.cpp
+++ b/Engine/Source/Editor/Sequencer/Private/Tools/NavigationTool/Utils/NavigationToolMiscUtils.cpp
@@ -93,7 +93,7 @@ TArrayView<TWeakObjectPtr<>> ResolveBoundObjects(ISequencer& InSequencer
 {
 	if (FMovieSceneEvaluationState* const EvaluationState = InSequencer.GetSharedPlaybackState()->FindCapability<FMovieSceneEvaluationState>())
 	{
-		const FMovieSceneSequenceIDRef SequenceID = EvaluationState->FindSequenceId(InSequence);
+		FMovieSceneSequenceIDRef SequenceID = EvaluationState->FindSequenceId(InSequence);
 		return InSequencer.FindBoundObjects(InBindingId, SequenceID);
 	}
 	return TArrayView<TWeakObjectPtr<>>();
