Sequencer: Fix duplicate LOCTEXT

#rnx

diff --git a/Engine/Source/Editor/Sequencer/Private/Tools/NavigationTool/Filters/Filters/NavigationToolFilter_Text.cpp b/Engine/Source/Editor/Sequencer/Private/Tools/NavigationTool/Filters/Filters/NavigationToolFilter_Text.cpp
--- a/Engine/Source/Editor/Sequencer/Private/Tools/NavigationTool/Filters/Filters/NavigationToolFilter_Text.cpp
+++ b/Engine/Source/Editor/Sequencer/Private/Tools/NavigationTool/Filters/Filters/NavigationToolFilter_Text.cpp
@@ -57,7 +57,7 @@ FText FNavigationToolFilter_Text::GetDisplayName() const
 
 FText FNavigationToolFilter_Text::GetToolTipText() const
 {
-	return LOCTEXT("NavigationToolFilter_Text", "Show only items that match the input text");
+	return LOCTEXT("NavigationToolFilter_TextTooltip", "Show only items that match the input text");
 }
 
 FString FNavigationToolFilter_Text::GetName() const
diff --git a/Engine/Source/Editor/Sequencer/Private/Tools/NavigationTool/Items/NavigationToolItemUtils.cpp b/Engine/Source/Editor/Sequencer/Private/Tools/NavigationTool/Items/NavigationToolItemUtils.cpp
--- a/Engine/Source/Editor/Sequencer/Private/Tools/NavigationTool/Items/NavigationToolItemUtils.cpp
+++ b/Engine/Source/Editor/Sequencer/Private/Tools/NavigationTool/Items/NavigationToolItemUtils.cpp
@@ -167,7 +167,7 @@ FText SequenceDisplayNameWithModifyAndInactiveStatus(const TSharedPtr<FNavigatio
 	{
 		if (bIsDirty)
 		{
-			OutText = FText::Format(NSLOCTEXT("Sequencer", "ActiveDirtySequenceLabel", "{0}*"), NewSequenceName);
+			OutText = FText::Format(NSLOCTEXT("Sequencer", "ActiveDirtySequenceFormat", "{0}*"), NewSequenceName);
 		}
 	}
 	else
@@ -176,12 +176,12 @@ FText SequenceDisplayNameWithModifyAndInactiveStatus(const TSharedPtr<FNavigatio
 
 		if (bIsDirty)
 		{
-			OutText = FText::Format(NSLOCTEXT("Sequencer", "InactiveDirtySequenceLabel", "{0}* [{1}]")
+			OutText = FText::Format(NSLOCTEXT("Sequencer", "InactiveDirtySequenceFormat", "{0}* [{1}]")
 				, NewSequenceName, InactiveSequenceLabel);
 		}
 		else
 		{
-			OutText = FText::Format(NSLOCTEXT("Sequencer", "InactiveDirtySequenceLabel", "{0} [{1}]")
+			OutText = FText::Format(NSLOCTEXT("Sequencer", "InactiveSequenceFormat", "{0} [{1}]")
 				, NewSequenceName, InactiveSequenceLabel);
 		}
 	}
