[Sequencer] Add snapping logic that prioritizes sections instead of intervals

diff --git a/Engine/Source/Editor/Sequencer/Private/Tools/SequencerSnapField.cpp b/Engine/Source/Editor/Sequencer/Private/Tools/SequencerSnapField.cpp
--- a/Engine/Source/Editor/Sequencer/Private/Tools/SequencerSnapField.cpp
+++ b/Engine/Source/Editor/Sequencer/Private/Tools/SequencerSnapField.cpp
@@ -289,6 +289,7 @@ TOptional<FSequencerSnapField::FSnapResult> FSequencerSnapField::Snap(const TArr
 {
 	TOptional<FSnapResult> ProspectiveSnap;
 
+	const FFrameTime IntervalSnapThreshold = FFrameTime::FromDecimal((TickResolution / DisplayRate).AsDecimal());
 	int32 NumSnaps = 0;
 	float MaxSnapWeight = 0.f;
 	for (FFrameTime Time : InTimes)
@@ -297,7 +298,7 @@ TOptional<FSequencerSnapField::FSnapResult> FSequencerSnapField::Snap(const TArr
 		
 		if (bSnapToLikeTypes)
 		{
-			ThisSnap = Snap(Time, Threshold);
+			ThisSnap = Snap(Time, FMath::Max(IntervalSnapThreshold, Threshold));
 		}
 
 		if (ClockExtension.IsValid())
@@ -316,26 +317,19 @@ TOptional<FSequencerSnapField::FSnapResult> FSequencerSnapField::Snap(const TArr
 			}
 		}
 
-		if (bSnapToInterval)
+		const bool bShouldTrySnappingToInterval = !bSnapToLikeTypes || IntervalSnapThreshold > Threshold;
+		if (bSnapToInterval && bShouldTrySnappingToInterval)
 		{
 			FFrameTime ThisTime = ThisSnap.IsSet() ? ThisSnap->SnappedTime : Time;
 			FFrameTime IntervalTime = FFrameRate::TransformTime(ThisTime, TickResolution, DisplayRate);
-			FFrameNumber PlayIntervalTime = ScrubStyle == ESequencerScrubberStyle::FrameBlock ? IntervalTime.FloorToFrame() : IntervalTime.RoundToFrame();
+			FFrameNumber PlayIntervalTime = IntervalTime.RoundToFrame();
 			FFrameNumber IntervalSnap = FFrameRate::TransformTime(PlayIntervalTime, DisplayRate, TickResolution).FloorToFrame();
 
-			const int32 IntervalSnapThreshold = FMath::RoundToInt((TickResolution / DisplayRate).AsDecimal());
-
 			FFrameTime ThisSnapAmount = IntervalSnap - ThisTime;
-			if (FMath::Abs(ThisSnapAmount) <= IntervalSnapThreshold)
+			bool bSnapToLikeTypesHasPriority = ThisSnap.IsSet() && FMath::Abs(ThisTime - Time) <= FMath::Abs(IntervalSnap - Time);
+			if (!bSnapToLikeTypesHasPriority && FMath::Abs(ThisSnapAmount) <= IntervalSnapThreshold)
 			{
-				if (ThisSnap.IsSet())
-				{
-					ThisSnap->SnappedTime = IntervalSnap;
-				}
-				else
-				{
-					ThisSnap = { Time, IntervalSnap, 1.f };
-				}
+				ThisSnap = { Time, IntervalSnap, 1.f };
 			}
 		}
 
