Sequencer: Fix spelling

#rnx

diff --git a/Engine/Source/Editor/Sequencer/Private/SequencerCommands.cpp b/Engine/Source/Editor/Sequencer/Private/SequencerCommands.cpp
--- a/Engine/Source/Editor/Sequencer/Private/SequencerCommands.cpp
+++ b/Engine/Source/Editor/Sequencer/Private/SequencerCommands.cpp
@@ -171,7 +171,7 @@ void FSequencerCommands::RegisterCommands()
 	UI_COMMAND( RestoreAnimatedState, "Restore Pre-Animated State", "Restore any objects that have been animated by sequencer back to their original state.", EUserInterfaceActionType::Button, FInputChord(EModifierKey::Control, EKeys::R) );
 
 	UI_COMMAND( FixPossessableObjectClass, "Fix Possessable Object Class", "Try to automatically fix up possessables whose object class don't match the object class of their currently bound objects.", EUserInterfaceActionType::Button, FInputChord() );
-	UI_COMMAND( RebindPossessableReferences, "Rebind Possesable References", "Rebinds all possessables in the current sequence to ensure they're using the most robust referencing mechanism.", EUserInterfaceActionType::Button, FInputChord() );
+	UI_COMMAND( RebindPossessableReferences, "Rebind Possessable References", "Rebinds all possessables in the current sequence to ensure they're using the most robust referencing mechanism.", EUserInterfaceActionType::Button, FInputChord() );
 
 	UI_COMMAND( ToggleEvaluateSubSequencesInIsolation, "Evaluate Sub Sequences In Isolation", "When enabled, will only evaluate the currently focused sequence; otherwise evaluate from the root sequence.", EUserInterfaceActionType::ToggleButton, FInputChord() );
 
