Sequencer: Fix rebind possessable mistakenly rebinding child components without the correct parent context.

diff --git a/Engine/Source/Editor/Sequencer/Private/Sequencer.cpp b/Engine/Source/Editor/Sequencer/Private/Sequencer.cpp
--- a/Engine/Source/Editor/Sequencer/Private/Sequencer.cpp
+++ b/Engine/Source/Editor/Sequencer/Private/Sequencer.cpp
@@ -10908,7 +10908,16 @@ void FSequencer::RebindPossessableReferences()
 			FocusedSequence->UnbindPossessableObjects(Pair.Key);
 			for (UObject* Object : Pair.Value)
 			{
-				FocusedSequence->BindPossessableObject(Pair.Key, *Object, PlaybackContext);
+				UObject* BindingContext = PlaybackContext;
+
+				// Find this object's parent object, if it has one.
+				UObject* ParentObject = FocusedSequence->GetParentObject(Object);
+				if (ParentObject)
+				{
+					BindingContext = ParentObject;
+				}
+
+				FocusedSequence->BindPossessableObject(Pair.Key, *Object, BindingContext);
 			}
 		}
 	}
