[Sequencer]
- Change track mute and solo to *not* be saved with the sequence asset
- Add "Deactivate" outliner column
- Add hotkeyable action for toggling deactivated track selection state
- Tooltips for outliner Pin, Lock, Deactivate, Mute, and Solo toggle buttons

diff --git a/Engine/Source/Editor/Sequencer/Private/MVVM/SequencerViewModelTypeIDs.cpp b/Engine/Source/Editor/Sequencer/Private/MVVM/SequencerViewModelTypeIDs.cpp
--- a/Engine/Source/Editor/Sequencer/Private/MVVM/SequencerViewModelTypeIDs.cpp
+++ b/Engine/Source/Editor/Sequencer/Private/MVVM/SequencerViewModelTypeIDs.cpp
@@ -51,7 +51,7 @@ UE_SEQUENCER_DEFINE_CASTABLE(FChannelModel);
 UE_SEQUENCER_DEFINE_CASTABLE(FFolderModel);
 UE_SEQUENCER_DEFINE_CASTABLE(FLayerBarModel);
 UE_SEQUENCER_DEFINE_CASTABLE(FOutlinerItemModel);
-UE_SEQUENCER_DEFINE_CASTABLE(FMuteSoloOutlinerItemModel);
+UE_SEQUENCER_DEFINE_CASTABLE(FEvaluableOutlinerItemModel);
 UE_SEQUENCER_DEFINE_CASTABLE(FObjectBindingModel);
 UE_SEQUENCER_DEFINE_CASTABLE(FPossessableModel);
 UE_SEQUENCER_DEFINE_CASTABLE(FSectionModel);
diff --git a/Engine/Source/Editor/Sequencer/Private/MVVM/ViewModels/FolderModel.cpp b/Engine/Source/Editor/Sequencer/Private/MVVM/ViewModels/FolderModel.cpp
--- a/Engine/Source/Editor/Sequencer/Private/MVVM/ViewModels/FolderModel.cpp
+++ b/Engine/Source/Editor/Sequencer/Private/MVVM/ViewModels/FolderModel.cpp
@@ -43,8 +43,7 @@ TMap<TWeakObjectPtr<UMovieSceneFolder>, FColor> FFolderModel::InitialFolderColor
 bool FFolderModel::bFolderPickerWasCancelled = false;
 
 FFolderModel::FFolderModel(UMovieSceneFolder* Folder)
-	: FMuteSoloOutlinerItemModel()
-	, TrackAreaList(EViewModelListType::TrackArea)
+	: TrackAreaList(EViewModelListType::TrackArea)
 	, WeakFolder(Folder)
 	, OwnerModel(nullptr)
 {
@@ -821,7 +820,6 @@ void FFolderModel::Delete()
 	}
 }
 
-
 } // namespace UE::Sequencer
 
 #undef LOCTEXT_NAMESPACE
diff --git a/Engine/Source/Editor/Sequencer/Private/MVVM/ViewModels/OutlinerColumns/OutlinerIndicatorColumn.cpp b/Engine/Source/Editor/Sequencer/Private/MVVM/ViewModels/OutlinerColumns/OutlinerIndicatorColumn.cpp
--- a/Engine/Source/Editor/Sequencer/Private/MVVM/ViewModels/OutlinerColumns/OutlinerIndicatorColumn.cpp
+++ b/Engine/Source/Editor/Sequencer/Private/MVVM/ViewModels/OutlinerColumns/OutlinerIndicatorColumn.cpp
@@ -20,7 +20,7 @@ FOutlinerIndicatorColumn::FOutlinerIndicatorColumn()
 {
 	Name = FCommonOutlinerNames::Indicator;
 	Label = LOCTEXT("IndicatorColumnLabel", "Indicators");
-	Position = FOutlinerColumnPosition{ 0, EOutlinerColumnGroup::FarLeftGutter };
+	Position = FOutlinerColumnPosition{ 0, EOutlinerColumnGroup::LeftGutter };
 	Layout = FOutlinerColumnLayout{ 14, FMargin(0.f, 0.f), HAlign_Fill, VAlign_Fill, EOutlinerColumnSizeMode::Fixed, EOutlinerColumnFlags::None };
 }
 
diff --git a/Engine/Source/Editor/Sequencer/Private/MVVM/ViewModels/OutlinerColumns/PinOutlinerColumn.cpp b/Engine/Source/Editor/Sequencer/Private/MVVM/ViewModels/OutlinerColumns/PinOutlinerColumn.cpp
--- a/Engine/Source/Editor/Sequencer/Private/MVVM/ViewModels/OutlinerColumns/PinOutlinerColumn.cpp
+++ b/Engine/Source/Editor/Sequencer/Private/MVVM/ViewModels/OutlinerColumns/PinOutlinerColumn.cpp
@@ -19,7 +19,7 @@ FPinOutlinerColumn::FPinOutlinerColumn()
 {
 	Name     = FCommonOutlinerNames::Pin;
 	Label    = LOCTEXT("PinColumnLabel", "Pin");
-	Position = FOutlinerColumnPosition{ 0, EOutlinerColumnGroup::LeftGutter };
+	Position = FOutlinerColumnPosition{ 0, EOutlinerColumnGroup::AssetToggles };
 	Layout   = FOutlinerColumnLayout{ 14, FMargin(4.f, 0.f), HAlign_Center, VAlign_Center, EOutlinerColumnSizeMode::Fixed, EOutlinerColumnFlags::None };
 }
 
diff --git a/Engine/Source/Editor/Sequencer/Private/MVVM/ViewModels/OutlinerItemModel.cpp b/Engine/Source/Editor/Sequencer/Private/MVVM/ViewModels/OutlinerItemModel.cpp
--- a/Engine/Source/Editor/Sequencer/Private/MVVM/ViewModels/OutlinerItemModel.cpp
+++ b/Engine/Source/Editor/Sequencer/Private/MVVM/ViewModels/OutlinerItemModel.cpp
@@ -309,21 +309,31 @@ bool FOutlinerItemModelMixin::IsPinned() const
 
 bool FOutlinerItemModelMixin::IsDimmed() const
 {
-	const FViewModel* ViewModel = AsViewModel();
+	const FViewModel* const ViewModel = AsViewModel();
+	const TSharedPtr<FSharedViewModelData> SharedData = ViewModel->GetSharedData();
+	if (!SharedData.IsValid())
+	{
+		return false;
+	}
 
-	FMuteStateCacheExtension* MuteState = ViewModel->GetSharedData()->CastThis<FMuteStateCacheExtension>();
-	FSoloStateCacheExtension* SoloState = ViewModel->GetSharedData()->CastThis<FSoloStateCacheExtension>();
+	const FDeactiveStateCacheExtension* const DeactiveState = SharedData->CastThis<FDeactiveStateCacheExtension>();
+	const FMuteStateCacheExtension* const MuteState = SharedData->CastThis<FMuteStateCacheExtension>();
+	const FSoloStateCacheExtension* const SoloState = SharedData->CastThis<FSoloStateCacheExtension>();
 
-	check(MuteState && SoloState);
+	check(DeactiveState && MuteState && SoloState);
 
-	ECachedMuteState MuteFlags = MuteState->GetCachedFlags(ViewModel->GetModelID());
-	ECachedSoloState SoloFlags = SoloState->GetCachedFlags(ViewModel->GetModelID());
+	const uint32 ModelID = ViewModel->GetModelID();
 
+	const ECachedDeactiveState DeactiveFlags = DeactiveState->GetCachedFlags(ModelID);
+	const ECachedMuteState MuteFlags = MuteState->GetCachedFlags(ModelID);
+	const ECachedSoloState SoloFlags = SoloState->GetCachedFlags(ModelID);
+
+	const bool bIsDeactive   = EnumHasAnyFlags(DeactiveFlags, ECachedDeactiveState::Deactivated | ECachedDeactiveState::ImplicitlyDeactivatedByParent);
 	const bool bAnySoloNodes = EnumHasAnyFlags(SoloState->GetRootFlags(), ECachedSoloState::Soloed | ECachedSoloState::PartiallySoloedChildren);
 	const bool bIsMuted      = EnumHasAnyFlags(MuteFlags, ECachedMuteState::Muted  | ECachedMuteState::ImplicitlyMutedByParent);
 	const bool bIsSoloed     = EnumHasAnyFlags(SoloFlags, ECachedSoloState::Soloed | ECachedSoloState::ImplicitlySoloedByParent);
 
-	const bool bDisableEval = bIsMuted || (bAnySoloNodes && !bIsSoloed);
+	const bool bDisableEval = bIsDeactive || bIsMuted || (bAnySoloNodes && !bIsSoloed);
 	return bDisableEval;
 }
 
@@ -1144,7 +1154,7 @@ bool FOutlinerItemModelMixin::PassesFilter(const FCurveEditorTreeFilter* InFilte
 
 
 
-bool FMuteSoloOutlinerItemModel::IsSolo() const
+bool FEvaluableOutlinerItemModel::IsSolo() const
 {
 	TSharedPtr<FSequencerEditorViewModel> EditorViewModel = GetEditor();
 	TSharedPtr<FSequencer> Sequencer = EditorViewModel ? EditorViewModel->GetSequencerImpl() : nullptr;
@@ -1166,7 +1176,7 @@ bool FMuteSoloOutlinerItemModel::IsSolo() const
 	return false;
 }
 
-void FMuteSoloOutlinerItemModel::SetIsSoloed(bool bIsSoloed)
+void FEvaluableOutlinerItemModel::SetIsSoloed(bool bIsSoloed)
 {
 	FViewModel* ViewModel = AsViewModel();
 	TSharedPtr<FSequenceModel> SequenceModel = ViewModel->FindAncestorOfType<FSequenceModel>();
@@ -1181,8 +1191,6 @@ void FMuteSoloOutlinerItemModel::SetIsSoloed(bool bIsSoloed)
 
 		TArray<FString>& SoloNodes = MovieScene->GetSoloNodes();
 
-		MovieScene->Modify();
-
 		// Should be always called on the main thread, but thread_local for safety
 		static thread_local TStringBuilder<256> StaticBuffer;
 		StaticBuffer.Reset();
@@ -1204,7 +1212,7 @@ void FMuteSoloOutlinerItemModel::SetIsSoloed(bool bIsSoloed)
 	}
 }
 
-bool FMuteSoloOutlinerItemModel::IsMuted() const
+bool FEvaluableOutlinerItemModel::IsMuted() const
 {
 	TSharedPtr<FSequencerEditorViewModel> EditorViewModel = GetEditor();
 	TSharedPtr<FSequencer> Sequencer = EditorViewModel ? EditorViewModel->GetSequencerImpl() : nullptr;
@@ -1226,7 +1234,7 @@ bool FMuteSoloOutlinerItemModel::IsMuted() const
 	return false;
 }
 
-void FMuteSoloOutlinerItemModel::SetIsMuted(bool bIsMuted)
+void FEvaluableOutlinerItemModel::SetIsMuted(bool bIsMuted)
 {
 	FViewModel* ViewModel = AsViewModel();
 	TSharedPtr<FSequenceModel> SequenceModel = ViewModel->FindAncestorOfType<FSequenceModel>();
@@ -1249,8 +1257,6 @@ void FMuteSoloOutlinerItemModel::SetIsMuted(bool bIsMuted)
 
 		const TCHAR* NodePath = StaticBuffer.ToString();
 
-		MovieScene->Modify();
-
 		if (bIsMuted)
 		{
 			// Mark Mute, being careful as we might be re-marking an already Mute node
@@ -1264,6 +1270,76 @@ void FMuteSoloOutlinerItemModel::SetIsMuted(bool bIsMuted)
 	}
 }
 
+bool FEvaluableOutlinerItemModel::IsDeactivated() const
+{
+	for (const TViewModelPtr<ITrackExtension> TrackNode : GetDescendantsOfType<ITrackExtension>(true))
+	{
+		UMovieSceneTrack* const Track = TrackNode->GetTrack();
+		if (!IsValid(Track))
+		{
+			continue;
+		}
+
+		if (const TViewModelPtr<FTrackRowModel> TrackRowModel = TrackNode.ImplicitCast())
+		{
+			if (!Track->IsRowEvalDisabled(TrackNode->GetRowIndex(), /*bInCheckLocal=*/false))
+			{
+				return false;
+			}
+		}
+		else
+		{
+			// Deactive state (dirtying, saved with asset, evaluation)
+			if (!Track->IsEvalDisabled(/*bInCheckLocal=*/false))
+			{
+				return false;
+			}
+		}
+	}
+
+	return true;
+}
+
+void FEvaluableOutlinerItemModel::SetIsDeactivated(const bool bInIsDeactivated)
+{
+	bool bAnyChanged = false;
+
+	for (const TViewModelPtr<ITrackExtension> TrackNode : GetDescendantsOfType<ITrackExtension>(true))
+	{
+		UMovieSceneTrack* const Track = TrackNode->GetTrack();
+		if (!IsValid(Track))
+		{
+			continue;
+		}
+
+		if (const TViewModelPtr<FTrackRowModel> TrackRowModel = TrackNode.ImplicitCast())
+		{
+			// Deactive state (dirtying, saved with asset, evaluation)
+			if (bInIsDeactivated != Track->IsRowEvalDisabled(TrackNode->GetRowIndex(), /*bInCheckLocal=*/false))
+			{
+				Track->Modify();
+				Track->SetRowEvalDisabled(bInIsDeactivated, TrackNode->GetRowIndex());
+				bAnyChanged = true;
+			}
+		}
+		else
+		{
+			// Deactive state (dirtying, saved with asset, evaluation)
+			if (bInIsDeactivated != Track->IsEvalDisabled(/*bInCheckLocal=*/false))
+			{
+				Track->Modify();
+				Track->SetEvalDisabled(bInIsDeactivated);
+				bAnyChanged = true;
+			}
+		}
+	}
+
+	if (bAnyChanged)
+	{
+		GetEditor()->GetSequencer()->NotifyMovieSceneDataChanged(EMovieSceneDataChangeType::TrackValueChanged);
+	}
+}
+
 } // namespace Sequencer
 } // namespace UE
 
diff --git a/Engine/Source/Editor/Sequencer/Private/MVVM/ViewModels/SequenceModel.cpp b/Engine/Source/Editor/Sequencer/Private/MVVM/ViewModels/SequenceModel.cpp
--- a/Engine/Source/Editor/Sequencer/Private/MVVM/ViewModels/SequenceModel.cpp
+++ b/Engine/Source/Editor/Sequencer/Private/MVVM/ViewModels/SequenceModel.cpp
@@ -60,6 +60,7 @@ void FSequenceModel::InitializeExtensions()
 
 	// Re-generate hierarchical caches when the sequence changes
 	ThisSharedData->AddDynamicExtension(FOutlinerCacheExtension::ID);
+	ThisSharedData->AddDynamicExtension(FDeactiveStateCacheExtension::ID);
 	ThisSharedData->AddDynamicExtension(FMuteStateCacheExtension::ID);
 	ThisSharedData->AddDynamicExtension(FSoloStateCacheExtension::ID);
 	ThisSharedData->AddDynamicExtension(FLockStateCacheExtension::ID);
diff --git a/Engine/Source/Editor/Sequencer/Private/SSequencer.cpp b/Engine/Source/Editor/Sequencer/Private/SSequencer.cpp
--- a/Engine/Source/Editor/Sequencer/Private/SSequencer.cpp
+++ b/Engine/Source/Editor/Sequencer/Private/SSequencer.cpp
@@ -2693,11 +2693,12 @@ void SSequencer::UpdateLayoutTree()
 			bool bAnyChanged = false;
 
 			TSharedPtr<FSharedViewModelData> SharedData = Sequencer->GetViewModel()->GetRootModel()->GetSharedData();
-			FOutlinerCacheExtension*     OutlinerCache  = SharedData->CastThis<FOutlinerCacheExtension>();
-			FMuteStateCacheExtension*    MuteState      = SharedData->CastThis<FMuteStateCacheExtension>();
-			FSoloStateCacheExtension*    SoloState      = SharedData->CastThis<FSoloStateCacheExtension>();
+			FOutlinerCacheExtension*           OutlinerCache = SharedData->CastThis<FOutlinerCacheExtension>();
+			FDeactiveStateCacheExtension* DeactiveState = SharedData->CastThis<FDeactiveStateCacheExtension>();
+			FMuteStateCacheExtension*          MuteState     = SharedData->CastThis<FMuteStateCacheExtension>();
+			FSoloStateCacheExtension*          SoloState     = SharedData->CastThis<FSoloStateCacheExtension>();
 
-			check(OutlinerCache && MuteState && SoloState);
+			check(OutlinerCache && DeactiveState && MuteState && SoloState);
 
 			// Hack - we shouldn't really not just forcibly update these here, but currently this function is getting forcibly called before
 			//        UpdateCachedFlags has a chance to naturally update itself in response to the signature change
@@ -2708,44 +2709,48 @@ void SSequencer::UpdateLayoutTree()
 			for (TViewModelPtr<ITrackExtension> TrackNode
 				: Sequencer->GetNodeTree()->GetRootNode()->GetDescendantsOfType<ITrackExtension>())
 			{
-				if (TrackNode->GetTrack())
+				UMovieSceneTrack* const Track = TrackNode->GetTrack();
+				if (!IsValid(Track))
 				{
-					ECachedMuteState MuteFlags = MuteState->GetCachedFlags(TrackNode.AsModel());
-					ECachedSoloState SoloFlags = SoloState->GetCachedFlags(TrackNode.AsModel());
+					continue;
+				}
+
+				const ECachedMuteState MuteFlags = MuteState->GetCachedFlags(TrackNode);
+				const ECachedSoloState SoloFlags = SoloState->GetCachedFlags(TrackNode);
 
-					UMovieSceneTrack* Track = TrackNode->GetTrack();
-					const bool bIsMuted  = EnumHasAnyFlags(MuteFlags, ECachedMuteState::Muted  | ECachedMuteState::ImplicitlyMutedByParent);
-					const bool bIsSoloed = EnumHasAnyFlags(SoloFlags, ECachedSoloState::Soloed | ECachedSoloState::ImplicitlySoloedByParent);
+				const bool bIsMuted = EnumHasAnyFlags(MuteFlags, ECachedMuteState::Muted | ECachedMuteState::ImplicitlyMutedByParent);
+				const bool bIsSoloed = EnumHasAnyFlags(SoloFlags, ECachedSoloState::Soloed | ECachedSoloState::ImplicitlySoloedByParent);
 
-					const bool bDisableEval = bIsMuted || (bAnySoloNodes && !bIsSoloed);
+				const bool bLocalEvalDisabled = bIsMuted || (bAnySoloNodes && !bIsSoloed);
 
-					TViewModelPtr<FTrackRowModel> TrackRowModel = TrackNode.ImplicitCast();
-					if (TrackRowModel)
+				if (const TViewModelPtr<FTrackRowModel> TrackRowModel = TrackNode.ImplicitCast())
+				{
+					// Mute/Solo state (local, non dirtying, mute/solo evaluation)
+					if (bLocalEvalDisabled != Track->IsLocalRowEvalDisabled(TrackNode->GetRowIndex()))
 					{
-						if (bDisableEval != Track->IsRowEvalDisabled(TrackNode->GetRowIndex()))
-						{
-							Track->Modify();
-							Track->SetRowEvalDisabled(bDisableEval, TrackNode->GetRowIndex());
-							bAnyChanged = true;
-						}
+						Track->MarkAsChanged();
+						Track->SetLocalRowEvalDisabled(bLocalEvalDisabled, TrackNode->GetRowIndex());
+						bAnyChanged = true;
 					}
-					else
+				}
+				else
+				{
+					// Mute/Solo state (local, non dirtying, mute/solo evaluation)
+					if (bLocalEvalDisabled != Track->IsLocalEvalDisabled())
 					{
-						if (bDisableEval != Track->IsEvalDisabled())
-						{
-							Track->Modify();
-							Track->SetEvalDisabled(bDisableEval);
-							bAnyChanged = true;
-						}	
+						Track->MarkAsChanged();
+						Track->SetLocalEvalDisabled(bLocalEvalDisabled);
+						bAnyChanged = true;
 					}
 				}
 			}
+
 			if (bAnyChanged)
 			{
 				Sequencer->NotifyMovieSceneDataChanged(EMovieSceneDataChangeType::TrackValueChanged);
 			}
 		}
-		
+
 		if (NodeGroupManager.IsValid())
 		{
 			NodeGroupManager->RefreshNodeGroups();
diff --git a/Engine/Source/Editor/Sequencer/Private/Sequencer.cpp b/Engine/Source/Editor/Sequencer/Private/Sequencer.cpp
--- a/Engine/Source/Editor/Sequencer/Private/Sequencer.cpp
+++ b/Engine/Source/Editor/Sequencer/Private/Sequencer.cpp
@@ -7103,51 +7103,34 @@ void FSequencer::OnNodePathChanged(const FString& OldPath, const FString& NewPat
 		// Find any solo/muted nodes with a path that is a child of the renamed node, and rename their paths as well
 		FString PathPrefix = OldPath + '.';
 
-		TArray<FString> PathsToRename;
-		for (const FString& NodePath : MovieScene->GetSoloNodes())
+		auto RenamePathsInArray = [this, &NewPath, &PathPrefix](TArray<FString>& InArray)
 		{
-			if (NodePath.StartsWith(PathPrefix) && NodePath != NewPath)
+			TArray<FString> PathsToRename;
+			for (const FString& NodePath : InArray)
 			{
-				PathsToRename.Add(NodePath);
-			}
-		}
-
-		for (const FString& NodePath : PathsToRename)
-		{
-			FString NewNodePath = NodePath;
-			if (NewNodePath.RemoveFromStart(PathPrefix))
-			{
-				NewNodePath = NewPath + '.' + NewNodePath;
-				if (NodeTree->GetNodeAtPath(NewNodePath))
+				if (NodePath.StartsWith(PathPrefix) && NodePath != NewPath)
 				{
-					MovieScene->GetSoloNodes().Remove(NodePath);
-					MovieScene->GetSoloNodes().Add(NewNodePath);
+					PathsToRename.Add(NodePath);
 				}
 			}
-		}
-
-		PathsToRename.Empty();
-		for (const FString& NodePath : MovieScene->GetMuteNodes())
-		{
-			if (NodePath.StartsWith(PathPrefix) && NodePath != NewPath)
-			{
-				PathsToRename.Add(NodePath);
-			}
-		}
 
-		for (const FString& NodePath : PathsToRename)
-		{
-			FString NewNodePath = NodePath;
-			if (NewNodePath.RemoveFromStart(PathPrefix))
+			for (const FString& NodePath : PathsToRename)
 			{
-				NewNodePath = NewPath + '.' + NewNodePath;
-				if (NodeTree->GetNodeAtPath(NewNodePath))
+				FString NewNodePath = NodePath;
+				if (NewNodePath.RemoveFromStart(PathPrefix))
 				{
-					MovieScene->GetMuteNodes().Remove(NodePath);
-					MovieScene->GetMuteNodes().Add(NewNodePath);
+					NewNodePath = NewPath + '.' + NewNodePath;
+					if (NodeTree->GetNodeAtPath(NewNodePath))
+					{
+						InArray.Remove(NodePath);
+						InArray.Add(NewNodePath);
+					}
 				}
 			}
-		}
+		};
+
+		RenamePathsInArray(MovieScene->GetSoloNodes());
+		RenamePathsInArray(MovieScene->GetMuteNodes());
 	}
 }
 
@@ -11987,6 +11970,10 @@ void FSequencer::BindCommands()
 		Commands.ToggleTrackSelectionLock,
 		FExecuteAction::CreateSP(this, &FSequencer::ToggleTrackSelectionLock));
 
+	SequencerCommandBindings->MapAction(
+		Commands.ToggleTrackSelectionDeactive,
+		FExecuteAction::CreateSP(this, &FSequencer::ToggleTrackSelectionDeactive));
+
 	SequencerCommandBindings->MapAction(
 		Commands.ToggleTrackSelectionMute,
 		FExecuteAction::CreateSP(this, &FSequencer::ToggleTrackSelectionMute));
@@ -12796,6 +12783,26 @@ void FSequencer::ToggleTrackSelectionLock()
 	}
 }
 
+void FSequencer::ToggleTrackSelectionDeactive()
+{
+	using namespace UE::Sequencer;
+
+	FOutlinerSelection& OutlinerSelection = GetSelection().Outliner;
+
+	const bool bAllDeactived = Algo::AllOf(OutlinerSelection, [](const FViewModelPtr& Selection)
+		{
+			const TViewModelPtr<IDeactivatableExtension> Deactivatable = Selection.ImplicitCast();
+			return Deactivatable.IsValid() && Deactivatable->IsDeactivated();
+		});
+
+	const bool bNewDeactiveState = !bAllDeactived;
+
+	for (const TViewModelPtr Deactivatable : OutlinerSelection.Filter<IDeactivatableExtension>())
+	{
+		Deactivatable->SetIsDeactivated(bNewDeactiveState);
+	}
+}
+
 void FSequencer::ToggleTrackSelectionMute()
 {
 	using namespace UE::Sequencer;
diff --git a/Engine/Source/Editor/Sequencer/Private/Sequencer.h b/Engine/Source/Editor/Sequencer/Private/Sequencer.h
--- a/Engine/Source/Editor/Sequencer/Private/Sequencer.h
+++ b/Engine/Source/Editor/Sequencer/Private/Sequencer.h
@@ -1300,6 +1300,7 @@ private:
 
 	void ToggleTrackSelectionPin();
 	void ToggleTrackSelectionLock();
+	void ToggleTrackSelectionDeactive();
 	void ToggleTrackSelectionMute();
 	void ToggleTrackSelectionSolo();
 
diff --git a/Engine/Source/Editor/Sequencer/Private/SequencerCommands.cpp b/Engine/Source/Editor/Sequencer/Private/SequencerCommands.cpp
--- a/Engine/Source/Editor/Sequencer/Private/SequencerCommands.cpp
+++ b/Engine/Source/Editor/Sequencer/Private/SequencerCommands.cpp
@@ -203,6 +203,7 @@ void FSequencerCommands::RegisterCommands()
 
 	UI_COMMAND(ToggleTrackSelectionPin, "Toggle Track Pin", "Toggle the pinned state of Sequencer track selection", EUserInterfaceActionType::Button, FInputChord());
 	UI_COMMAND(ToggleTrackSelectionLock, "Toggle Track Lock", "Toggle the locked state of Sequencer track selection", EUserInterfaceActionType::Button, FInputChord());
+	UI_COMMAND(ToggleTrackSelectionDeactive, "Toggle Track Deactive", "Toggle the deactive state of Sequencer track selection", EUserInterfaceActionType::Button, FInputChord());
 	UI_COMMAND(ToggleTrackSelectionMute, "Toggle Track Mute", "Toggle the muted state of Sequencer track selection", EUserInterfaceActionType::Button, FInputChord());
 	UI_COMMAND(ToggleTrackSelectionSolo, "Toggle Track Solo", "Toggle the soloed state of Sequencer track selection", EUserInterfaceActionType::Button, FInputChord());
 
diff --git a/Engine/Source/Editor/Sequencer/Private/SequencerModule.cpp b/Engine/Source/Editor/Sequencer/Private/SequencerModule.cpp
--- a/Engine/Source/Editor/Sequencer/Private/SequencerModule.cpp
+++ b/Engine/Source/Editor/Sequencer/Private/SequencerModule.cpp
@@ -35,6 +35,7 @@
 #include "MVVM/ViewModels/SequencerEditorViewModel.h"
 
 #include "MVVM/ViewModels/OutlinerColumns/OutlinerIndicatorColumn.h"
+#include "MVVM/ViewModels/OutlinerColumns/DeactivateOutlinerColumn.h"
 #include "MVVM/ViewModels/OutlinerColumns/LockOutlinerColumn.h"
 #include "MVVM/ViewModels/OutlinerColumns/MuteOutlinerColumn.h"
 #include "MVVM/ViewModels/OutlinerColumns/PinOutlinerColumn.h"
@@ -494,8 +495,11 @@ public:
 
 			// Register left gutter columns
 			PinOutlinerColumnHandle  = RegisterOutlinerColumn(FOnCreateOutlinerColumn::CreateStatic([]{ return TSharedRef<IOutlinerColumn>(MakeShared<FPinOutlinerColumn>()); }));
-			MuteOutlinerColumnHandle = RegisterOutlinerColumn(FOnCreateOutlinerColumn::CreateStatic([]{ return TSharedRef<IOutlinerColumn>(MakeShared<FMuteOutlinerColumn>()); }));
 			LockOutlinerColumnHandle = RegisterOutlinerColumn(FOnCreateOutlinerColumn::CreateStatic([]{ return TSharedRef<IOutlinerColumn>(MakeShared<FLockOutlinerColumn>()); }));
+			DeactivateOutlinerColumnHandle = RegisterOutlinerColumn(FOnCreateOutlinerColumn::CreateStatic([]{ return TSharedRef<IOutlinerColumn>(MakeShared<FDeactivateOutlinerColumn>()); }));
+
+			// Register left gutter columns
+			MuteOutlinerColumnHandle = RegisterOutlinerColumn(FOnCreateOutlinerColumn::CreateStatic([]{ return TSharedRef<IOutlinerColumn>(MakeShared<FMuteOutlinerColumn>()); }));
 			SoloOutlinerColumnHandle = RegisterOutlinerColumn(FOnCreateOutlinerColumn::CreateStatic([]{ return TSharedRef<IOutlinerColumn>(MakeShared<FSoloOutlinerColumn>()); }));
 
 			// Register center columns
@@ -560,6 +564,7 @@ public:
 
 			// unregister outliner columns
 			UnregisterOutlinerColumn(OutlinerIndicatorColumnHandle);
+			UnregisterOutlinerColumn(DeactivateOutlinerColumnHandle);
 			UnregisterOutlinerColumn(PinOutlinerColumnHandle);
 			UnregisterOutlinerColumn(MuteOutlinerColumnHandle);
 			UnregisterOutlinerColumn(LockOutlinerColumnHandle);
@@ -769,6 +774,7 @@ private:
 
 	// Outliner Column Delegate Handles
 	FDelegateHandle OutlinerIndicatorColumnHandle;
+	FDelegateHandle DeactivateOutlinerColumnHandle;
 	FDelegateHandle PinOutlinerColumnHandle;
 	FDelegateHandle MuteOutlinerColumnHandle;
 	FDelegateHandle LockOutlinerColumnHandle;
diff --git a/Engine/Source/Editor/Sequencer/Private/Widgets/OutlinerColumns/SPinColumnWidget.cpp b/Engine/Source/Editor/Sequencer/Private/Widgets/OutlinerColumns/SPinColumnWidget.cpp
--- a/Engine/Source/Editor/Sequencer/Private/Widgets/OutlinerColumns/SPinColumnWidget.cpp
+++ b/Engine/Source/Editor/Sequencer/Private/Widgets/OutlinerColumns/SPinColumnWidget.cpp
@@ -9,6 +9,8 @@
 namespace UE::Sequencer
 {
 
+#define LOCTEXT_NAMESPACE "SPinColumnWidget"
+
 void SPinColumnWidget::OnToggleOperationComplete()
 {
 	// refresh the sequencer tree after operation is complete
@@ -17,6 +19,9 @@ void SPinColumnWidget::OnToggleOperationComplete()
 
 void SPinColumnWidget::Construct(const FArguments& InArgs, const TWeakPtr<IOutlinerColumn> InWeakOutlinerColumn, const FCreateOutlinerColumnParams& InParams)
 {
+	SetToolTipText(LOCTEXT("PinTooltip", "Pin this track and force it to always stay at the top of the list.\n\n"
+		"Saved with the asset."));
+
 	SColumnToggleWidget::Construct(
 		SColumnToggleWidget::FArguments(),
 		InWeakOutlinerColumn,
@@ -64,4 +69,6 @@ const FSlateBrush* SPinColumnWidget::GetActiveBrush() const
 	return FAppStyle::Get().GetBrush(NAME_PinnedBrush);
 }
 
+#undef LOCTEXT_NAMESPACE
+
 } // namespace UE::Sequencer
diff --git a/Engine/Source/Editor/Sequencer/Public/MVVM/ViewModels/FolderModel.h b/Engine/Source/Editor/Sequencer/Public/MVVM/ViewModels/FolderModel.h
--- a/Engine/Source/Editor/Sequencer/Public/MVVM/ViewModels/FolderModel.h
+++ b/Engine/Source/Editor/Sequencer/Public/MVVM/ViewModels/FolderModel.h
@@ -12,6 +12,7 @@
 #include "MVVM/Extensions/IDraggableOutlinerExtension.h"
 #include "MVVM/Extensions/ISortableExtension.h"
 #include "MVVM/Extensions/IDeletableExtension.h"
+#include "MVVM/Extensions/IDeactivatableExtension.h"
 #include "MVVM/ViewModels/ViewModelHierarchy.h"
 
 #include "UObject/WeakObjectPtr.h"
@@ -30,7 +31,7 @@ class FSequenceModel;
 class FSequencerOutlinerDragDropOp;
 
 class FFolderModel
-	: public FMuteSoloOutlinerItemModel
+	: public FEvaluableOutlinerItemModel
 	, public IRenameableExtension
 	, public ITrackAreaExtension
 	, public IGroupableExtension
@@ -43,7 +44,7 @@ class FFolderModel
 public:
 
 	UE_SEQUENCER_DECLARE_CASTABLE(FFolderModel
-		, FMuteSoloOutlinerItemModel
+		, FEvaluableOutlinerItemModel
 		, IRenameableExtension
 		, ITrackAreaExtension
 		, IGroupableExtension
diff --git a/Engine/Source/Editor/Sequencer/Public/MVVM/ViewModels/ObjectBindingModel.h b/Engine/Source/Editor/Sequencer/Public/MVVM/ViewModels/ObjectBindingModel.h
--- a/Engine/Source/Editor/Sequencer/Public/MVVM/ViewModels/ObjectBindingModel.h
+++ b/Engine/Source/Editor/Sequencer/Public/MVVM/ViewModels/ObjectBindingModel.h
@@ -41,7 +41,7 @@ enum class EObjectBindingType
 };
 
 class SEQUENCER_API FObjectBindingModel
-	: public FMuteSoloOutlinerItemModel
+	: public FEvaluableOutlinerItemModel
 	, public IObjectBindingExtension
 	, public IDraggableOutlinerExtension
 	, public ITrackAreaExtension
@@ -53,7 +53,7 @@ class SEQUENCER_API FObjectBindingModel
 public:
 
 	UE_SEQUENCER_DECLARE_CASTABLE(FObjectBindingModel
-		, FMuteSoloOutlinerItemModel
+		, FEvaluableOutlinerItemModel
 		, IObjectBindingExtension
 		, IDraggableOutlinerExtension
 		, ITrackAreaExtension
diff --git a/Engine/Source/Editor/Sequencer/Public/MVVM/ViewModels/OutlinerItemModel.h b/Engine/Source/Editor/Sequencer/Public/MVVM/ViewModels/OutlinerItemModel.h
--- a/Engine/Source/Editor/Sequencer/Public/MVVM/ViewModels/OutlinerItemModel.h
+++ b/Engine/Source/Editor/Sequencer/Public/MVVM/ViewModels/OutlinerItemModel.h
@@ -5,6 +5,7 @@
 #include "MVVM/ICastable.h"
 #include "MVVM/ViewModels/ViewModel.h"
 #include "MVVM/Extensions/ICurveEditorTreeItemExtension.h"
+#include "MVVM/Extensions/IDeactivatableExtension.h"
 #include "MVVM/Extensions/IDimmableExtension.h"
 #include "MVVM/Extensions/IGeometryExtension.h"
 #include "MVVM/Extensions/IHoveredExtension.h"
@@ -159,14 +160,19 @@ public:
 	UE_SEQUENCER_DECLARE_CASTABLE(FOutlinerItemModel, FViewModel, FOutlinerItemModelMixin);
 };
 
-class SEQUENCER_API FMuteSoloOutlinerItemModel
+class SEQUENCER_API FEvaluableOutlinerItemModel
 	: public FOutlinerItemModel
 	, public IMutableExtension
 	, public ISoloableExtension
+	, public IDeactivatableExtension
 {
 public:
 
-	UE_SEQUENCER_DECLARE_CASTABLE(FMuteSoloOutlinerItemModel, FOutlinerItemModel, IMutableExtension, ISoloableExtension);
+	UE_SEQUENCER_DECLARE_CASTABLE(FEvaluableOutlinerItemModel
+		, FOutlinerItemModel
+		, IMutableExtension
+		, ISoloableExtension
+		, IDeactivatableExtension);
 
 	/*~ ISoloableExtension */
 	bool IsSolo() const override;
@@ -175,8 +181,11 @@ public:
 	/*~ IMutableExtension */
 	bool IsMuted() const override;
 	void SetIsMuted(bool bIsMuted) override;
-};
 
+	/*~ IDeactivatableExtension */
+	virtual bool IsDeactivated() const override;
+	virtual void SetIsDeactivated(const bool bInIsDeactivated) override;
+};
 
 class SEQUENCER_API FOutlinerCacheExtension
 	: public FHierarchicalCacheExtension
diff --git a/Engine/Source/Editor/Sequencer/Public/MVVM/ViewModels/TrackModel.h b/Engine/Source/Editor/Sequencer/Public/MVVM/ViewModels/TrackModel.h
--- a/Engine/Source/Editor/Sequencer/Public/MVVM/ViewModels/TrackModel.h
+++ b/Engine/Source/Editor/Sequencer/Public/MVVM/ViewModels/TrackModel.h
@@ -37,7 +37,7 @@ namespace Sequencer
 class FSectionModel;
 
 class SEQUENCER_API FTrackModel
-	: public FMuteSoloOutlinerItemModel
+	: public FEvaluableOutlinerItemModel
 	, public IRenameableExtension
 	, public IResizableExtension
 	, public ITrackExtension
@@ -54,7 +54,7 @@ class SEQUENCER_API FTrackModel
 public:
 
 	UE_SEQUENCER_DECLARE_CASTABLE(FTrackModel
-		, FMuteSoloOutlinerItemModel
+		, FEvaluableOutlinerItemModel
 		, IRenameableExtension
 		, IResizableExtension
 		, ITrackExtension
diff --git a/Engine/Source/Editor/Sequencer/Public/MVVM/ViewModels/TrackRowModel.h b/Engine/Source/Editor/Sequencer/Public/MVVM/ViewModels/TrackRowModel.h
--- a/Engine/Source/Editor/Sequencer/Public/MVVM/ViewModels/TrackRowModel.h
+++ b/Engine/Source/Editor/Sequencer/Public/MVVM/ViewModels/TrackRowModel.h
@@ -27,7 +27,7 @@ namespace Sequencer
 class FSectionModel;
 
 class SEQUENCER_API FTrackRowModel
-	: public FMuteSoloOutlinerItemModel
+	: public FEvaluableOutlinerItemModel
 	, public ITrackAreaExtension
 	, public ILockableExtension
 	, public ITrackExtension
@@ -38,7 +38,15 @@ class SEQUENCER_API FTrackRowModel
 {
 public:
 
-	UE_SEQUENCER_DECLARE_CASTABLE(FTrackRowModel, FMuteSoloOutlinerItemModel, ITrackAreaExtension, ILockableExtension, ITrackExtension, IDeletableExtension, IMutableExtension, ISoloableExtension, IRenameableExtension, IConditionableExtension);
+	UE_SEQUENCER_DECLARE_CASTABLE(FTrackRowModel
+		, FEvaluableOutlinerItemModel
+		, ITrackAreaExtension
+		, ILockableExtension
+		, ITrackExtension
+		, IRenameableExtension
+		, IResizableExtension
+		, IDeletableExtension
+		, IConditionableExtension);
 
 	explicit FTrackRowModel(UMovieSceneTrack* InTrack, int32 InRowIndex);
 	~FTrackRowModel();
diff --git a/Engine/Source/Editor/Sequencer/Public/SequencerCommands.h b/Engine/Source/Editor/Sequencer/Public/SequencerCommands.h
--- a/Engine/Source/Editor/Sequencer/Public/SequencerCommands.h
+++ b/Engine/Source/Editor/Sequencer/Public/SequencerCommands.h
@@ -487,6 +487,9 @@ public:
 	/** Toggle the locked state of Sequencer track selection */
 	TSharedPtr< FUICommandInfo > ToggleTrackSelectionLock;
 
+	/** Toggle the deactived state of Sequencer track selection */
+    TSharedPtr< FUICommandInfo > ToggleTrackSelectionDeactive;
+
 	/** Toggle the muted state of Sequencer track selection */
 	TSharedPtr< FUICommandInfo > ToggleTrackSelectionMute;
 
diff --git a/Engine/Source/Editor/SequencerCore/Private/MVVM/Extensions/IDeactivatableExtension.cpp b/Engine/Source/Editor/SequencerCore/Private/MVVM/Extensions/IDeactivatableExtension.cpp
new file mode 100644
--- /dev/null
+++ b/Engine/Source/Editor/SequencerCore/Private/MVVM/Extensions/IDeactivatableExtension.cpp
@@ -0,0 +1,131 @@
+﻿// Copyright Epic Games, Inc. All Rights Reserved.
+
+#include "MVVM/Extensions/IDeactivatableExtension.h"
+
+namespace UE::Sequencer
+{
+
+ECachedDeactiveState CombinePropagatedChildFlags(const ECachedDeactiveState ParentFlags, ECachedDeactiveState CombinedChildFlags)
+{
+	// If this item is deactivatable, but not deactivated, do not inherit deactivated state
+	if ((ParentFlags & (ECachedDeactiveState::Deactivatable | ECachedDeactiveState::Deactivated)) == ECachedDeactiveState::Deactivatable
+		&& EnumHasAnyFlags(CombinedChildFlags, ECachedDeactiveState::Deactivated))
+	{
+		CombinedChildFlags &= ~ECachedDeactiveState::Deactivated;
+		CombinedChildFlags |= ECachedDeactiveState::PartiallyDeactivatedChildren;
+	}
+
+	return ParentFlags | CombinedChildFlags;
+}
+
+ECachedDeactiveState FDeactiveStateCacheExtension::ComputeFlagsForModel(const FViewModelPtr& ViewModel)
+{
+	ECachedDeactiveState ParentFlags = IndividualItemFlags.Last();
+	ECachedDeactiveState ThisModelFlags = ECachedDeactiveState::None;
+
+	if (EnumHasAnyFlags(ParentFlags, ECachedDeactiveState::Deactivated | ECachedDeactiveState::ImplicitlyDeactivatedByParent))
+	{
+		ThisModelFlags |= ECachedDeactiveState::ImplicitlyDeactivatedByParent;
+	}
+
+	if (TViewModelPtr<IDeactivatableExtension> Deactivatable = ViewModel.ImplicitCast())
+	{
+		ThisModelFlags |= ECachedDeactiveState::Deactivatable;
+		if (Deactivatable->IsDeactivated())
+		{
+			ThisModelFlags |= ECachedDeactiveState::Deactivated;
+		}
+		if (Deactivatable->IsInheritable())
+		{
+			ThisModelFlags |= ECachedDeactiveState::Inheritable;
+		}
+	}
+
+	return ThisModelFlags;
+
+	/*ECachedDeactiveState ThisModelFlags = ECachedDeactiveState::None;
+	
+	ECachedDeactiveState ParentFlags = IndividualItemFlags.Last();
+
+	if (EnumHasAnyFlags(ParentFlags, ECachedDeactiveState::Deactivated | ECachedDeactiveState::ImplicitlyDeactivatedByParent))
+	{
+		ThisModelFlags |= ECachedDeactiveState::ImplicitlyDeactivatedByParent;
+	}
+
+	if (const TViewModelPtr<IDeactivatableExtension> Deactivatable = ViewModel.ImplicitCast())
+	{
+		ThisModelFlags |= ECachedDeactiveState::Deactivatable;
+
+		switch (Deactivatable->IsDeactivated())
+		{
+		case EDeactivatableState::Deactivated:
+			ThisModelFlags |= ECachedDeactiveState::Deactivated;
+			break;
+		case EDeactivatableState::PartiallyDeactivated:
+			// Deactivatable, but not deactivated - any parent must not be fully deactivated
+			if (EnumHasAnyFlags(ParentFlags, ECachedDeactiveState::Deactivated))
+			{
+				ParentFlags |= ECachedDeactiveState::PartiallyDeactivatedChildren;
+				ParentFlags &= ~ECachedDeactiveState::Deactivated;
+			}
+			break;
+		case EDeactivatableState::None:
+			break;
+		}
+
+		if (Deactivatable->IsInheritable())
+		{
+			ThisModelFlags |= ECachedDeactiveState::Inheritable;
+		}
+	}*/
+
+	return ThisModelFlags;
+}
+
+void FDeactiveStateCacheExtension::PostComputeChildrenFlags(const FViewModelPtr& ViewModel, ECachedDeactiveState& OutThisModelFlags, ECachedDeactiveState& OutPropagateToParentFlags)
+{
+	const bool bIsInheritable = EnumHasAnyFlags(OutThisModelFlags, ECachedDeactiveState::Inheritable);
+	if (!bIsInheritable)
+	{
+		return;
+	}
+
+	// --------------------------------------------------------------------
+	// Handle deactivated state propagation
+	const bool bIsDeactivatable              = EnumHasAnyFlags(OutThisModelFlags, ECachedDeactiveState::Deactivatable);
+	const bool bIsDeactivated                = EnumHasAnyFlags(OutThisModelFlags, ECachedDeactiveState::Deactivated);
+	const bool bHasDeactivatableChildren     = EnumHasAnyFlags(OutThisModelFlags, ECachedDeactiveState::DeactivatableChildren);
+	const bool bSiblingsPartiallyDeactivated = EnumHasAnyFlags(OutPropagateToParentFlags, ECachedDeactiveState::PartiallyDeactivatedChildren);
+	const bool bSiblingsFullyDeactivated     = EnumHasAnyFlags(OutPropagateToParentFlags, ECachedDeactiveState::Deactivated);
+	const bool bHasAnyDeactivatableSiblings  = EnumHasAnyFlags(OutPropagateToParentFlags, ECachedDeactiveState::DeactivatableChildren);
+
+	if (bIsDeactivated)
+	{
+		if (!bSiblingsPartiallyDeactivated)
+		{
+			if (bHasAnyDeactivatableSiblings && !bSiblingsFullyDeactivated)
+			{
+				// This is the first deactivated deactivatable within the parent, but we know there are already other deactivatables so this has to be partially deactivated
+				OutPropagateToParentFlags |= ECachedDeactiveState::PartiallyDeactivatedChildren;
+			}
+			else
+			{
+				// Parent is (currently) fully deactivated because it contains no other deactivatable children, and is not already partially deactivated
+				OutPropagateToParentFlags |= ECachedDeactiveState::Deactivated;
+			}
+		}
+	}
+	else if (bIsDeactivatable && bSiblingsFullyDeactivated)
+	{
+		// Parent is no longer fully deactivated because it contains an activated deactivatable
+		OutPropagateToParentFlags |= ECachedDeactiveState::PartiallyDeactivatedChildren;
+		OutPropagateToParentFlags &= ~ECachedDeactiveState::Deactivated;
+	}
+
+	if (bIsDeactivatable)
+	{
+		OutPropagateToParentFlags |= ECachedDeactiveState::DeactivatableChildren;
+	}
+}
+
+} // namespace UE::Sequencer
diff --git a/Engine/Source/Editor/SequencerCore/Private/MVVM/ViewModelTypeIDs.cpp b/Engine/Source/Editor/SequencerCore/Private/MVVM/ViewModelTypeIDs.cpp
--- a/Engine/Source/Editor/SequencerCore/Private/MVVM/ViewModelTypeIDs.cpp
+++ b/Engine/Source/Editor/SequencerCore/Private/MVVM/ViewModelTypeIDs.cpp
@@ -12,6 +12,7 @@
 #include "MVVM/Extensions/IHoveredExtension.h"
 #include "MVVM/Extensions/IKeyExtension.h"
 #include "MVVM/Extensions/ILayerBarExtension.h"
+#include "MVVM/Extensions/IDeactivatableExtension.h"
 #include "MVVM/Extensions/IMutableExtension.h"
 #include "MVVM/Extensions/IObjectModelExtension.h"
 #include "MVVM/Extensions/IOutlinerExtension.h"
@@ -56,6 +57,7 @@ UE_SEQUENCER_DEFINE_VIEW_MODEL_TYPE_ID(IGroupableExtension);
 UE_SEQUENCER_DEFINE_VIEW_MODEL_TYPE_ID(IHotspotExtension);
 UE_SEQUENCER_DEFINE_VIEW_MODEL_TYPE_ID(IHoveredExtension);
 UE_SEQUENCER_DEFINE_VIEW_MODEL_TYPE_ID(ILayerBarExtension);
+UE_SEQUENCER_DEFINE_VIEW_MODEL_TYPE_ID(IDeactivatableExtension);
 UE_SEQUENCER_DEFINE_VIEW_MODEL_TYPE_ID(IMutableExtension);
 UE_SEQUENCER_DEFINE_VIEW_MODEL_TYPE_ID(ILockableExtension);
 UE_SEQUENCER_DEFINE_VIEW_MODEL_TYPE_ID(IOutlinerDropTargetOutlinerExtension);
@@ -82,6 +84,7 @@ UE_SEQUENCER_DEFINE_VIEW_MODEL_TYPE_ID(IKeyExtension);
 UE_SEQUENCER_DEFINE_VIEW_MODEL_TYPE_ID(IDynamicExtension);
 UE_SEQUENCER_DEFINE_VIEW_MODEL_TYPE_ID(FDynamicExtensionContainer);
 UE_SEQUENCER_DEFINE_VIEW_MODEL_TYPE_ID(FLinkedOutlinerExtension);
+UE_SEQUENCER_DEFINE_VIEW_MODEL_TYPE_ID(FDeactiveStateCacheExtension);
 UE_SEQUENCER_DEFINE_VIEW_MODEL_TYPE_ID(FMuteStateCacheExtension);
 UE_SEQUENCER_DEFINE_VIEW_MODEL_TYPE_ID(FSoloStateCacheExtension);
 UE_SEQUENCER_DEFINE_VIEW_MODEL_TYPE_ID(FLockStateCacheExtension);
diff --git a/Engine/Source/Editor/SequencerCore/Private/MVVM/ViewModels/OutlinerColumns/DeactivateOutlinerColumn.cpp b/Engine/Source/Editor/SequencerCore/Private/MVVM/ViewModels/OutlinerColumns/DeactivateOutlinerColumn.cpp
new file mode 100644
--- /dev/null
+++ b/Engine/Source/Editor/SequencerCore/Private/MVVM/ViewModels/OutlinerColumns/DeactivateOutlinerColumn.cpp
@@ -0,0 +1,40 @@
+// Copyright Epic Games, Inc. All Rights Reserved.
+
+#include "MVVM/ViewModels/OutlinerColumns/DeactivateOutlinerColumn.h"
+#include "MVVM/Extensions/IDeactivatableExtension.h"
+#include "MVVM/Extensions/IOutlinerExtension.h"
+#include "MVVM/SharedViewModelData.h"
+#include "MVVM/Views/OutlinerColumns/SDeactivateColumnWidget.h"
+
+#define LOCTEXT_NAMESPACE "FDeactivateOutlinerColumn"
+
+namespace UE::Sequencer
+{
+
+FDeactivateOutlinerColumn::FDeactivateOutlinerColumn()
+{
+	Name     = FCommonOutlinerNames::Deactivate;
+	Label    = LOCTEXT("DeactivateColumnLabel", "Deactivate");
+	Position = FOutlinerColumnPosition{ 20, EOutlinerColumnGroup::AssetToggles };
+	Layout   = FOutlinerColumnLayout{ 14, FMargin(4.f, 0.f), HAlign_Center, VAlign_Center, EOutlinerColumnSizeMode::Fixed, EOutlinerColumnFlags::None };
+}
+
+bool FDeactivateOutlinerColumn::IsItemCompatibleWithColumn(const FCreateOutlinerColumnParams& InParams) const
+{
+	if (const FDeactiveStateCacheExtension* const StateCache = InParams.OutlinerExtension.AsModel()->GetSharedData()->CastThis<FDeactiveStateCacheExtension>())
+	{
+		return EnumHasAnyFlags(StateCache->GetCachedFlags(InParams.OutlinerExtension)
+			, ECachedDeactiveState::Deactivatable | ECachedDeactiveState::DeactivatableChildren);
+	}
+
+	return false;
+}
+
+TSharedPtr<SWidget> FDeactivateOutlinerColumn::CreateColumnWidget(const FCreateOutlinerColumnParams& InParams, const TSharedRef<ISequencerTreeViewRow>& TreeViewRow)
+{
+	return SNew(SDeactivateColumnWidget, SharedThis(this), InParams);
+}
+	
+} // namespace UE::Sequencer
+
+#undef LOCTEXT_NAMESPACE
diff --git a/Engine/Source/Editor/SequencerCore/Private/MVVM/ViewModels/OutlinerColumns/LabelOutlinerColumn.cpp b/Engine/Source/Editor/SequencerCore/Private/MVVM/ViewModels/OutlinerColumns/LabelOutlinerColumn.cpp
--- a/Engine/Source/Editor/SequencerCore/Private/MVVM/ViewModels/OutlinerColumns/LabelOutlinerColumn.cpp
+++ b/Engine/Source/Editor/SequencerCore/Private/MVVM/ViewModels/OutlinerColumns/LabelOutlinerColumn.cpp
@@ -3,6 +3,8 @@
 #include "MVVM/ViewModels/OutlinerColumns/LabelOutlinerColumn.h"
 #include "MVVM/Extensions/IOutlinerExtension.h"
 
+#define LOCTEXT_NAMESPACE "FLabelOutlinerColumn"
+
 namespace UE::Sequencer
 {
 
@@ -11,7 +13,7 @@ FLabelOutlinerColumn::FLabelOutlinerColumn()
 	EOutlinerColumnFlags Flags = EOutlinerColumnFlags::OverflowSubsequentEmptyCells | EOutlinerColumnFlags::OverflowOnHover | EOutlinerColumnFlags::Hidden;
 
 	Name     = FCommonOutlinerNames::Label;
-	Label    = NSLOCTEXT("FLabelOutlinerColumn", "LabelColumnLabel", "Label");
+	Label    = LOCTEXT("LabelColumnLabel", "Label");
 	Position = FOutlinerColumnPosition{ 0, EOutlinerColumnGroup::Center };
 	Layout   = FOutlinerColumnLayout{ 0.6f, FMargin(4.f, 0.f), HAlign_Fill, VAlign_Center, EOutlinerColumnSizeMode::Stretch, Flags };
 }
@@ -22,3 +24,5 @@ TSharedPtr<SWidget> FLabelOutlinerColumn::CreateColumnWidget(const FCreateOutlin
 }
 
 } // namespace UE::Sequencer
+
+#undef LOCTEXT_NAMESPACE
diff --git a/Engine/Source/Editor/SequencerCore/Private/MVVM/ViewModels/OutlinerColumns/LockOutlinerColumn.cpp b/Engine/Source/Editor/SequencerCore/Private/MVVM/ViewModels/OutlinerColumns/LockOutlinerColumn.cpp
--- a/Engine/Source/Editor/SequencerCore/Private/MVVM/ViewModels/OutlinerColumns/LockOutlinerColumn.cpp
+++ b/Engine/Source/Editor/SequencerCore/Private/MVVM/ViewModels/OutlinerColumns/LockOutlinerColumn.cpp
@@ -14,7 +14,7 @@ FLockOutlinerColumn::FLockOutlinerColumn()
 {
 	Name     = FCommonOutlinerNames::Lock;
 	Label    = LOCTEXT("LockColumnLabel", "Lock");
-	Position = FOutlinerColumnPosition{ 10, EOutlinerColumnGroup::LeftGutter };
+	Position = FOutlinerColumnPosition{ 10, EOutlinerColumnGroup::AssetToggles };
 	Layout   = FOutlinerColumnLayout{ 14, FMargin(4.f, 0.f), HAlign_Center, VAlign_Center, EOutlinerColumnSizeMode::Fixed, EOutlinerColumnFlags::None };
 }
 
diff --git a/Engine/Source/Editor/SequencerCore/Private/MVVM/ViewModels/OutlinerColumns/MuteOutlinerColumn.cpp b/Engine/Source/Editor/SequencerCore/Private/MVVM/ViewModels/OutlinerColumns/MuteOutlinerColumn.cpp
--- a/Engine/Source/Editor/SequencerCore/Private/MVVM/ViewModels/OutlinerColumns/MuteOutlinerColumn.cpp
+++ b/Engine/Source/Editor/SequencerCore/Private/MVVM/ViewModels/OutlinerColumns/MuteOutlinerColumn.cpp
@@ -16,7 +16,7 @@ FMuteOutlinerColumn::FMuteOutlinerColumn()
 {
 	Name     = FCommonOutlinerNames::Mute;
 	Label    = LOCTEXT("MuteColumnLabel", "Mute");
-	Position = FOutlinerColumnPosition{ 30, EOutlinerColumnGroup::LeftGutter };
+	Position = FOutlinerColumnPosition{ 0, EOutlinerColumnGroup::TransientToggles };
 	Layout   = FOutlinerColumnLayout{ 14, FMargin(4.f, 0.f), HAlign_Center, VAlign_Center, EOutlinerColumnSizeMode::Fixed, EOutlinerColumnFlags::None };
 }
 
@@ -35,8 +35,6 @@ TSharedPtr<SWidget> FMuteOutlinerColumn::CreateColumnWidget(const FCreateOutline
 	return SNew(SMuteColumnWidget, SharedThis(this), InParams);
 }
 
-
 } // namespace UE::Sequencer
 
-
 #undef LOCTEXT_NAMESPACE
diff --git a/Engine/Source/Editor/SequencerCore/Private/MVVM/ViewModels/OutlinerColumns/OutlinerColumnTypes.cpp b/Engine/Source/Editor/SequencerCore/Private/MVVM/ViewModels/OutlinerColumns/OutlinerColumnTypes.cpp
--- a/Engine/Source/Editor/SequencerCore/Private/MVVM/ViewModels/OutlinerColumns/OutlinerColumnTypes.cpp
+++ b/Engine/Source/Editor/SequencerCore/Private/MVVM/ViewModels/OutlinerColumns/OutlinerColumnTypes.cpp
@@ -8,6 +8,7 @@ namespace UE::Sequencer
 FName FCommonOutlinerNames::Indicator("Indicator");
 FName FCommonOutlinerNames::Pin("Pin");
 FName FCommonOutlinerNames::Lock("Lock");
+FName FCommonOutlinerNames::Deactivate("Deactivate");
 FName FCommonOutlinerNames::Mute("Mute");
 FName FCommonOutlinerNames::Solo("Solo");
 FName FCommonOutlinerNames::Label("Label");
diff --git a/Engine/Source/Editor/SequencerCore/Private/MVVM/ViewModels/OutlinerColumns/SoloOutlinerColumn.cpp b/Engine/Source/Editor/SequencerCore/Private/MVVM/ViewModels/OutlinerColumns/SoloOutlinerColumn.cpp
--- a/Engine/Source/Editor/SequencerCore/Private/MVVM/ViewModels/OutlinerColumns/SoloOutlinerColumn.cpp
+++ b/Engine/Source/Editor/SequencerCore/Private/MVVM/ViewModels/OutlinerColumns/SoloOutlinerColumn.cpp
@@ -14,7 +14,7 @@ FSoloOutlinerColumn::FSoloOutlinerColumn()
 {
 	Name     = FCommonOutlinerNames::Solo;
 	Label    = LOCTEXT("SoloColumnLabel", "Solo");
-	Position = FOutlinerColumnPosition{ 20, EOutlinerColumnGroup::LeftGutter };
+	Position = FOutlinerColumnPosition{ 10, EOutlinerColumnGroup::TransientToggles };
 	Layout   = FOutlinerColumnLayout{ 14, FMargin(4.f, 0.f), HAlign_Center, VAlign_Center, EOutlinerColumnSizeMode::Fixed, EOutlinerColumnFlags::None };
 }
 
diff --git a/Engine/Source/Editor/SequencerCore/Private/MVVM/Views/OutlinerColumns/SDeactivateColumnWidget.cpp b/Engine/Source/Editor/SequencerCore/Private/MVVM/Views/OutlinerColumns/SDeactivateColumnWidget.cpp
new file mode 100644
--- /dev/null
+++ b/Engine/Source/Editor/SequencerCore/Private/MVVM/Views/OutlinerColumns/SDeactivateColumnWidget.cpp
@@ -0,0 +1,117 @@
+// Copyright Epic Games, Inc. All Rights Reserved.
+
+#include "MVVM/Views/OutlinerColumns/SDeactivateColumnWidget.h"
+#include "MVVM/Selection/SequencerCoreSelection.h"
+#include "MVVM/Selection/SequencerOutlinerSelection.h"
+#include "MVVM/SharedViewModelData.h"
+#include "MVVM/Extensions/IDeactivatableExtension.h"
+#include "MVVM/ViewModels/EditorViewModel.h"
+#include "MVVM/ViewModels/OutlinerColumns/IOutlinerColumn.h"
+#include "MVVM/ViewModels/ViewModelIterators.h"
+
+#define LOCTEXT_NAMESPACE "SDeactivateColumnWidget"
+
+namespace UE::Sequencer
+{
+
+void SDeactivateColumnWidget::OnToggleOperationComplete()
+{
+	// refresh the sequencer tree after operation is complete
+	RefreshSequencerTree();
+}
+
+void SDeactivateColumnWidget::Construct(const FArguments& InArgs, const TWeakPtr<IOutlinerColumn> InWeakOutlinerColumn, const FCreateOutlinerColumnParams& InParams)
+{
+	SetToolTipText(LOCTEXT("DeactivateTooltip", "Deactivate this track and disable evaluation.\n\n"
+		"Saved with the asset."));
+
+	SColumnToggleWidget::Construct(
+		SColumnToggleWidget::FArguments(),
+		InWeakOutlinerColumn,
+		InParams
+	);
+
+	WeakStateCacheExtension = CastViewModel<FDeactiveStateCacheExtension>(InParams.OutlinerExtension.AsModel()->GetSharedData());
+}
+
+bool SDeactivateColumnWidget::IsActive() const
+{
+	if (const TViewModelPtr<FDeactiveStateCacheExtension> StateCache = WeakStateCacheExtension.Pin())
+	{
+		return EnumHasAnyFlags(StateCache->GetCachedFlags(ModelID), ECachedDeactiveState::Deactivated);
+	}
+
+	return false;
+}
+
+void SDeactivateColumnWidget::SetIsActive(const bool bInIsActive)
+{
+	const TViewModelPtr<IOutlinerExtension> OutlinerItem = WeakOutlinerExtension.Pin();
+	const TSharedPtr<FEditorViewModel>      Editor       = WeakEditor.Pin();
+	if (!OutlinerItem || !Editor)
+	{
+		return;
+	}
+
+	const FScopedTransaction Transaction(LOCTEXT("SetNodeDeactivated", "Set Node Deactivated"));
+
+	if (OutlinerItem->GetSelectionState() == EOutlinerSelectionState::SelectedDirectly)
+	{
+		// If selected, modify all selected items
+		for (const FViewModelPtr Selected : *Editor->GetSelection()->GetOutlinerSelection())
+		{
+			SetIsActive_Internal(Selected, bInIsActive);
+		}
+	}
+	else
+	{
+		SetIsActive_Internal(OutlinerItem.AsModel(), bInIsActive);
+	}
+}
+
+void SDeactivateColumnWidget::SetIsActive_Internal(const FViewModelPtr& ViewModel, const bool bInIsActive)
+{
+	// If this is deactivatable, deactivate only this
+	if (const TViewModelPtr<IDeactivatableExtension> Deactivatable = ViewModel.ImplicitCast())
+	{
+		Deactivatable->SetIsDeactivated(bInIsActive);
+	}
+	// Otherwise deactivate deactivatable children of this (if any)
+	else for (const TViewModelPtr<IDeactivatableExtension> Child : ViewModel->GetDescendantsOfType<IDeactivatableExtension>())
+	{
+		if (Child->IsInheritable())
+		{
+			Child->SetIsDeactivated(bInIsActive);
+		}
+	}
+}
+
+bool SDeactivateColumnWidget::IsChildActive() const
+{
+	if (const TViewModelPtr<FDeactiveStateCacheExtension> StateCache = WeakStateCacheExtension.Pin())
+	{
+		return EnumHasAnyFlags(StateCache->GetCachedFlags(ModelID), ECachedDeactiveState::PartiallyDeactivatedChildren);
+	}
+
+	return false;
+}
+
+bool SDeactivateColumnWidget::IsImplicitlyActive() const
+{
+	if (const TViewModelPtr<FDeactiveStateCacheExtension> StateCache = WeakStateCacheExtension.Pin())
+	{
+		return EnumHasAnyFlags(StateCache->GetCachedFlags(ModelID), ECachedDeactiveState::ImplicitlyDeactivatedByParent);
+	}
+
+	return false;
+}
+
+const FSlateBrush* SDeactivateColumnWidget::GetActiveBrush() const
+{
+	static const FName BrushName = TEXT("Sequencer.Column.Mute");
+	return FAppStyle::Get().GetBrush(BrushName);
+}
+
+} // namespace UE::Sequencer
+
+#undef LOCTEXT_NAMESPACE
diff --git a/Engine/Source/Editor/SequencerCore/Private/MVVM/Views/OutlinerColumns/SDeactivateColumnWidget.h b/Engine/Source/Editor/SequencerCore/Private/MVVM/Views/OutlinerColumns/SDeactivateColumnWidget.h
new file mode 100644
--- /dev/null
+++ b/Engine/Source/Editor/SequencerCore/Private/MVVM/Views/OutlinerColumns/SDeactivateColumnWidget.h
@@ -0,0 +1,55 @@
+// Copyright Epic Games, Inc. All Rights Reserved.
+
+#pragma once
+
+#include "MVVM/Views/OutlinerColumns/SColumnToggleWidget.h"
+#include "SequencerCoreFwd.h"
+
+namespace UE::Sequencer
+{
+
+class FDeactiveStateCacheExtension;
+
+/**
+ * A widget that shows and controls the deactivated state of outliner items.
+ */
+class SDeactivateColumnWidget
+	: public SColumnToggleWidget
+{
+public:
+	SLATE_BEGIN_ARGS(SDeactivateColumnWidget) {}
+	SLATE_END_ARGS()
+
+	void Construct(const FArguments& InArgs, const TWeakPtr<IOutlinerColumn> InWeakOutlinerColumn, const FCreateOutlinerColumnParams& InParams);
+
+	/** Refreshes the Sequencer Tree. */
+	virtual void OnToggleOperationComplete();
+
+protected:
+	//~ Begin SColumnToggleWidget
+
+	/** Returns whether or not this item is deactivated. */
+	virtual bool IsActive() const override;
+
+	/** Sets this item as activated or deactivated. If selected, applies to all selected items. */
+	virtual void SetIsActive(const bool bInIsActive) override;
+
+	/** Returns whether or not a child of this item is deactivated. */
+	virtual bool IsChildActive() const override;
+
+	/** Returns true if this item is implicitly deactivated, but not directly deactivated */
+	virtual bool IsImplicitlyActive() const override;
+
+	/** Returns the brush to display when this item is deactivated. */
+	virtual const FSlateBrush* GetActiveBrush() const override;
+
+	//~ End SColumnToggleWidget
+
+private:
+	void SetIsActive_Internal(const FViewModelPtr& ViewModel, const bool bInIsActive);
+
+	/** Weak cache extension ptr (can be null). */
+	TWeakViewModelPtr<FDeactiveStateCacheExtension> WeakStateCacheExtension;
+};
+
+} // namespace UE::Sequencer
diff --git a/Engine/Source/Editor/SequencerCore/Private/MVVM/Views/OutlinerColumns/SLockColumnWidget.cpp b/Engine/Source/Editor/SequencerCore/Private/MVVM/Views/OutlinerColumns/SLockColumnWidget.cpp
--- a/Engine/Source/Editor/SequencerCore/Private/MVVM/Views/OutlinerColumns/SLockColumnWidget.cpp
+++ b/Engine/Source/Editor/SequencerCore/Private/MVVM/Views/OutlinerColumns/SLockColumnWidget.cpp
@@ -13,8 +13,13 @@
 namespace UE::Sequencer
 {
 
+#define LOCTEXT_NAMESPACE "SLockColumnWidget"
+
 void SLockColumnWidget::Construct(const FArguments& InArgs, const TWeakPtr<IOutlinerColumn> InWeakOutlinerColumn, const FCreateOutlinerColumnParams& InParams)
 {
+	SetToolTipText(LOCTEXT("LockTooltip", "Lock this track and disable editing.\n\n"
+		"Saved with the asset."));
+
 	SColumnToggleWidget::Construct(
 		SColumnToggleWidget::FArguments(),
 		InWeakOutlinerColumn,
@@ -43,7 +48,7 @@ void SLockColumnWidget::SetIsActive(const bool bInIsActive)
 		return;
 	}
 
-	const FScopedTransaction Transaction(NSLOCTEXT("Sequencer", "SetNodeLocked", "Set Node Locked"));
+	const FScopedTransaction Transaction(LOCTEXT("SetNodeLocked", "Set Node Locked"));
 
 	if (OutlinerItem->GetSelectionState() == EOutlinerSelectionState::SelectedDirectly)
 	{
@@ -83,4 +88,6 @@ const FSlateBrush* SLockColumnWidget::GetActiveBrush() const
 	return FAppStyle::Get().GetBrush(NAME_LockBrush);
 }
 
+#undef LOCTEXT_NAMESPACE
+
 } // namespace UE::Sequencer
diff --git a/Engine/Source/Editor/SequencerCore/Private/MVVM/Views/OutlinerColumns/SMuteColumnWidget.cpp b/Engine/Source/Editor/SequencerCore/Private/MVVM/Views/OutlinerColumns/SMuteColumnWidget.cpp
--- a/Engine/Source/Editor/SequencerCore/Private/MVVM/Views/OutlinerColumns/SMuteColumnWidget.cpp
+++ b/Engine/Source/Editor/SequencerCore/Private/MVVM/Views/OutlinerColumns/SMuteColumnWidget.cpp
@@ -10,10 +10,11 @@
 #include "MVVM/ViewModels/EditorViewModel.h"
 #include "MVVM/ViewModels/OutlinerColumns/IOutlinerColumn.h"
 
-
 namespace UE::Sequencer
 {
 
+#define LOCTEXT_NAMESPACE "SMuteColumnWidget"
+
 void SMuteColumnWidget::OnToggleOperationComplete()
 {
 	// refresh the sequencer tree after operation is complete
@@ -22,6 +23,9 @@ void SMuteColumnWidget::OnToggleOperationComplete()
 
 void SMuteColumnWidget::Construct(const FArguments& InArgs, const TWeakPtr<IOutlinerColumn> InWeakOutlinerColumn, const FCreateOutlinerColumnParams& InParams)
 {
+	SetToolTipText(LOCTEXT("MuteTooltip", "Mute this track and disable evaluation locally.\n\n"
+		"*Not* saved with the asset, but saved in memory until editor restart"));
+
 	SColumnToggleWidget::Construct(
 		SColumnToggleWidget::FArguments(),
 		InWeakOutlinerColumn,
@@ -50,28 +54,26 @@ void SMuteColumnWidget::SetIsActive(const bool bInIsActive)
 		return;
 	}
 
-	const FScopedTransaction Transaction(NSLOCTEXT("Sequencer", "SetNodeMuted", "Set Node Muted"));
+	const FScopedTransaction Transaction(LOCTEXT("SetNodeMuted", "Set Node Muted"));
 
 	if (OutlinerItem->GetSelectionState() == EOutlinerSelectionState::SelectedDirectly)
 	{
 		// If selected, modify all selected items
 		for (FViewModelPtr Selected : *Editor->GetSelection()->GetOutlinerSelection())
 		{
-			SetIsActive(Selected, bInIsActive);
+			SetIsActive_Internal(Selected, bInIsActive);
 		}
 	}
 	else
 	{
-		SetIsActive(OutlinerItem.AsModel(), bInIsActive);
+		SetIsActive_Internal(OutlinerItem.AsModel(), bInIsActive);
 	}
 }
 
-void SMuteColumnWidget::SetIsActive(const FViewModelPtr& ViewModel, const bool bInIsActive)
+void SMuteColumnWidget::SetIsActive_Internal(const FViewModelPtr& ViewModel, const bool bInIsActive)
 {
-	TViewModelPtr<IMutableExtension> Mutable = ViewModel.ImplicitCast();
-
 	// If this is mutable, mute only this
-	if (Mutable)
+	if (const TViewModelPtr<IMutableExtension> Mutable = ViewModel.ImplicitCast())
 	{
 		Mutable->SetIsMuted(bInIsActive);
 	}
@@ -107,8 +109,10 @@ bool SMuteColumnWidget::IsImplicitlyActive() const
 
 const FSlateBrush* SMuteColumnWidget::GetActiveBrush() const
 {
-	static const FName NAME_MutedBrush = TEXT("Sequencer.Column.Mute");
+	static const FName NAME_MutedBrush = TEXT("Level.NotVisibleIcon16x");
 	return FAppStyle::Get().GetBrush(NAME_MutedBrush);
 }
 
+#undef LOCTEXT_NAMESPACE
+
 } // namespace UE::Sequencer
diff --git a/Engine/Source/Editor/SequencerCore/Private/MVVM/Views/OutlinerColumns/SMuteColumnWidget.h b/Engine/Source/Editor/SequencerCore/Private/MVVM/Views/OutlinerColumns/SMuteColumnWidget.h
--- a/Engine/Source/Editor/SequencerCore/Private/MVVM/Views/OutlinerColumns/SMuteColumnWidget.h
+++ b/Engine/Source/Editor/SequencerCore/Private/MVVM/Views/OutlinerColumns/SMuteColumnWidget.h
@@ -44,7 +44,7 @@ protected:
 
 private:
 
-	void SetIsActive(const FViewModelPtr& ViewModel, const bool bInIsActive);
+	void SetIsActive_Internal(const FViewModelPtr& ViewModel, const bool bInIsActive);
 
 	/** Weak cache extension ptr (can be null). */
 	TWeakViewModelPtr<FMuteStateCacheExtension> WeakMuteStateCacheExtension;
diff --git a/Engine/Source/Editor/SequencerCore/Private/MVVM/Views/OutlinerColumns/SSoloColumnWidget.cpp b/Engine/Source/Editor/SequencerCore/Private/MVVM/Views/OutlinerColumns/SSoloColumnWidget.cpp
--- a/Engine/Source/Editor/SequencerCore/Private/MVVM/Views/OutlinerColumns/SSoloColumnWidget.cpp
+++ b/Engine/Source/Editor/SequencerCore/Private/MVVM/Views/OutlinerColumns/SSoloColumnWidget.cpp
@@ -12,6 +12,8 @@
 namespace UE::Sequencer
 {
 
+#define LOCTEXT_NAMESPACE "SSoloColumnWidget"
+
 void SSoloColumnWidget::OnToggleOperationComplete()
 {
 	// refresh the sequencer tree after operation is complete
@@ -20,6 +22,9 @@ void SSoloColumnWidget::OnToggleOperationComplete()
 
 void SSoloColumnWidget::Construct(const FArguments& InArgs, const TWeakPtr<IOutlinerColumn> InWeakOutlinerColumn, const FCreateOutlinerColumnParams& InParams)
 {
+	SetToolTipText(LOCTEXT("SoloTooltip", "Solo this track and disable evaluation for all other tracks locally.\n\n"
+		"*Not* saved with the asset, but saved in memory until editor restart"));
+
 	SColumnToggleWidget::Construct(
 		SColumnToggleWidget::FArguments(),
 		InWeakOutlinerColumn,
@@ -48,7 +53,7 @@ void SSoloColumnWidget::SetIsActive(const bool bInIsActive)
 		return;
 	}
 
-	const FScopedTransaction Transaction(NSLOCTEXT("Sequencer", "SetNodeSoloed", "Set Node Soloed"));
+	const FScopedTransaction Transaction(LOCTEXT("SetNodeSoloed", "Set Node Soloed"));
 
 	if (OutlinerItem->GetSelectionState() == EOutlinerSelectionState::SelectedDirectly)
 	{
@@ -90,4 +95,6 @@ const FSlateBrush* SSoloColumnWidget::GetActiveBrush() const
 	return FAppStyle::Get().GetBrush(NAME_SoloedBrush);
 }
 
+#undef LOCTEXT_NAMESPACE
+
 } // namespace UE::Sequencer
diff --git a/Engine/Source/Editor/SequencerCore/Private/MVVM/Views/SOutlinerItemViewBase.cpp b/Engine/Source/Editor/SequencerCore/Private/MVVM/Views/SOutlinerItemViewBase.cpp
--- a/Engine/Source/Editor/SequencerCore/Private/MVVM/Views/SOutlinerItemViewBase.cpp
+++ b/Engine/Source/Editor/SequencerCore/Private/MVVM/Views/SOutlinerItemViewBase.cpp
@@ -12,6 +12,7 @@
 #include "Layout/Clipping.h"
 #include "Layout/Margin.h"
 #include "MVVM/Extensions/IDimmableExtension.h"
+#include "MVVM/Extensions/IDeactivatableExtension.h"
 #include "MVVM/Extensions/IMutableExtension.h"
 #include "MVVM/Extensions/IOutlinerExtension.h"
 #include "MVVM/Extensions/IRenameableExtension.h"
@@ -343,6 +344,12 @@ bool SOutlinerItemViewBase::IsDimmed() const
 		return Dimmable->IsDimmed();
 	}
 
+	const TViewModelPtr<IDeactivatableExtension> Deactivatable = Outliner.ImplicitCast();
+	if (Deactivatable && Deactivatable->IsDeactivated())
+	{
+		return true;
+	}
+
 	// Otherwise, provide some simple default logic: dim the item if it's muted.
 	TViewModelPtr<IMutableExtension> Mutable = Outliner.ImplicitCast();
 	if (Mutable && Mutable->IsMuted())
diff --git a/Engine/Source/Editor/SequencerCore/Private/MVVM/Views/SOutlinerView.cpp b/Engine/Source/Editor/SequencerCore/Private/MVVM/Views/SOutlinerView.cpp
--- a/Engine/Source/Editor/SequencerCore/Private/MVVM/Views/SOutlinerView.cpp
+++ b/Engine/Source/Editor/SequencerCore/Private/MVVM/Views/SOutlinerView.cpp
@@ -378,25 +378,54 @@ void SOutlinerView::UpdateOutlinerColumns()
 
 	// ----------------------------------------------------------------------------------------------------------
 	// Populate columns
-	const int32 NumFarLeftGutter	= CreateOutlinerColumnsForGroup(0,											  EOutlinerColumnGroup::FarLeftGutter);
-	const int32 NumLeftGutter		= CreateOutlinerColumnsForGroup(NumFarLeftGutter,							  EOutlinerColumnGroup::LeftGutter);
-	const int32 NumCenter			= CreateOutlinerColumnsForGroup(NumFarLeftGutter + NumLeftGutter,             EOutlinerColumnGroup::Center);
-	const int32 NumRightGutter		= CreateOutlinerColumnsForGroup(NumFarLeftGutter + NumLeftGutter + NumCenter, EOutlinerColumnGroup::RightGutter);
+	int32 CurrentIndex = 0;
+
+	const int32 NumLeftGutter = CreateOutlinerColumnsForGroup(CurrentIndex, EOutlinerColumnGroup::LeftGutter);
+	CurrentIndex += NumLeftGutter;
+	
+	const int32 NumToggles = CreateOutlinerColumnsForGroup(CurrentIndex, EOutlinerColumnGroup::AssetToggles);
+	CurrentIndex += NumToggles;
+
+	const int32 NumMuteSoloToggles = CreateOutlinerColumnsForGroup(CurrentIndex, EOutlinerColumnGroup::TransientToggles);
+	CurrentIndex += NumMuteSoloToggles;
+
+	const int32 NumCenter = CreateOutlinerColumnsForGroup(CurrentIndex, EOutlinerColumnGroup::Center);
+	CurrentIndex += NumCenter;
+
+	const int32 NumRightGutter = CreateOutlinerColumnsForGroup(CurrentIndex, EOutlinerColumnGroup::RightGutter);
+	CurrentIndex += NumRightGutter;
 
 	// ----------------------------------------------------------------------------------------------------------
 	// Add some padding to the leading and trailing edge of the first and last columns in each group respectively
 	//    This is implemented this way because columns can be turned on and off dynamically, but we must
 	//    always have a consistent padding within the group. Separators dynamically appear based on the presence of
 	//    each group so we can't put padding on those
-	if (NumLeftGutter > 0)
+	const float DefaultGroupPadding = 2.f;
+	if (NumToggles > 0)
 	{
-		ColumnMetaData->Columns[NumFarLeftGutter].CellPadding.Left  += 4.f;
-		ColumnMetaData->Columns[NumFarLeftGutter+NumLeftGutter-1].CellPadding.Right += 4.f;
+		const int32 FirstIndex = NumLeftGutter;
+		const int32 LastIndex = FirstIndex + NumToggles - 1;
+
+		ColumnMetaData->Columns[FirstIndex].CellPadding.Left += DefaultGroupPadding;
+		ColumnMetaData->Columns[LastIndex].CellPadding.Right += DefaultGroupPadding;
 	}
+
+	if (NumMuteSoloToggles > 0)
+	{
+		const int32 FirstIndex = NumLeftGutter + NumToggles;
+		const int32 LastIndex = FirstIndex + NumMuteSoloToggles - 1;
+
+		ColumnMetaData->Columns[FirstIndex].CellPadding.Left += DefaultGroupPadding;
+		ColumnMetaData->Columns[LastIndex].CellPadding.Right += DefaultGroupPadding;
+	}
+
 	if (NumCenter > 0)
 	{
-		ColumnMetaData->Columns[NumFarLeftGutter+NumLeftGutter            ].CellPadding.Left  += 4.f;
-		ColumnMetaData->Columns[NumFarLeftGutter+NumLeftGutter+NumCenter-1].CellPadding.Right += 4.f;
+		const int32 FirstIndex = NumLeftGutter + NumToggles + NumMuteSoloToggles;
+		const int32 LastIndex = FirstIndex + NumCenter - 1;
+
+		ColumnMetaData->Columns[FirstIndex].CellPadding.Left += DefaultGroupPadding;
+		ColumnMetaData->Columns[LastIndex].CellPadding.Right += DefaultGroupPadding;
 	}
 
 	// No additional padding on the right gutter intentionally
@@ -404,28 +433,35 @@ void SOutlinerView::UpdateOutlinerColumns()
 	// ----------------------------------------------------------------------------------------------------------
 	// Add separators between the groups.
 	//      Only add separators if they actually separate columns.
-	int32 NumSeparators = 0;
-
-	int32 InsertIndex = NumFarLeftGutter;
+	int32 InsertIndex = NumLeftGutter;
 	if (InsertIndex < ColumnMetaData->Columns.Num())
 	{
+		int32 NumSeparators = 0;
+
 		InsertSeparatorColumn(InsertIndex++, ++NumSeparators);
-		if (NumLeftGutter > 0)
+
+		auto InsertSeparatorAfterGroup = [this, &NumSeparators, &InsertIndex](const int32 InIndex)
 		{
-			InsertIndex += NumLeftGutter;
+			InsertIndex += InIndex;
 			if (InsertIndex < ColumnMetaData->Columns.Num())
 			{
 				InsertSeparatorColumn(InsertIndex++, ++NumSeparators);
 			}
+		};
+
+		if (NumToggles > 0)
+        {
+        	InsertSeparatorAfterGroup(NumToggles);
+        }
+
+		if (NumMuteSoloToggles > 0)
+		{
+			InsertSeparatorAfterGroup(NumMuteSoloToggles);
 		}
 
 		if (NumCenter > 0)
 		{
-			InsertIndex += NumCenter;
-			if (InsertIndex < ColumnMetaData->Columns.Num())
-			{
-				InsertSeparatorColumn(InsertIndex++, ++NumSeparators);
-			}
+			InsertSeparatorAfterGroup(NumCenter);
 		}
 	}
 }
diff --git a/Engine/Source/Editor/SequencerCore/Public/MVVM/Extensions/IDeactivatableExtension.h b/Engine/Source/Editor/SequencerCore/Public/MVVM/Extensions/IDeactivatableExtension.h
new file mode 100644
--- /dev/null
+++ b/Engine/Source/Editor/SequencerCore/Public/MVVM/Extensions/IDeactivatableExtension.h
@@ -0,0 +1,69 @@
+﻿// Copyright Epic Games, Inc. All Rights Reserved.
+
+#pragma once
+
+#include "MVVM/ViewModelTypeID.h"
+#include "MVVM/Extensions/HierarchicalCacheExtension.h"
+
+namespace UE::Sequencer
+{
+
+enum class EDeactivatableState
+{
+	None,
+	Deactivated,
+	PartiallyDeactivated,
+};
+
+/**
+ * An extension for outliner nodes that can be deactivated
+ */
+class SEQUENCERCORE_API IDeactivatableExtension
+{
+public:
+	UE_SEQUENCER_DECLARE_VIEW_MODEL_TYPE_ID(IDeactivatableExtension);
+
+	virtual ~IDeactivatableExtension() {}
+
+	/** Returns whether this item is deactivated */
+	virtual bool IsDeactivated() const = 0;
+
+	/** Set this item's deactivated state */
+	virtual void SetIsDeactivated(const bool bInIsDeactivated) = 0;
+
+	/** Returns whether this deactivatable can be deactivated by a parent, and should report its deactivated state to a parent */
+	virtual bool IsInheritable() const
+	{
+		return true;
+	}
+};
+
+enum class ECachedDeactiveState
+{
+	None                          = 0,
+
+	Deactivatable                 = 1 << 0,
+	DeactivatableChildren         = 1 << 1,
+	Deactivated                   = 1 << 2,
+	PartiallyDeactivatedChildren  = 1 << 3,
+	ImplicitlyDeactivatedByParent = 1 << 4,
+	Inheritable                   = 1 << 5,
+
+	InheritedFromChildren = DeactivatableChildren | PartiallyDeactivatedChildren,
+};
+ENUM_CLASS_FLAGS(ECachedDeactiveState)
+
+SEQUENCERCORE_API ECachedDeactiveState CombinePropagatedChildFlags(const ECachedDeactiveState ParentFlags, ECachedDeactiveState CombinedChildFlags);
+
+class SEQUENCERCORE_API FDeactiveStateCacheExtension
+	: public TFlagStateCacheExtension<ECachedDeactiveState>
+{
+public:
+	UE_SEQUENCER_DECLARE_VIEW_MODEL_TYPE_ID(FDeactiveStateCacheExtension);
+
+private:
+	ECachedDeactiveState ComputeFlagsForModel(const FViewModelPtr& ViewModel) override;
+	void PostComputeChildrenFlags(const FViewModelPtr& ViewModel, ECachedDeactiveState& OutThisModelFlags, ECachedDeactiveState& OutPropagateToParentFlags) override;
+};
+
+} // namespace UE::Sequencer
diff --git a/Engine/Source/Editor/SequencerCore/Public/MVVM/ViewModels/OutlinerColumns/DeactivateOutlinerColumn.h b/Engine/Source/Editor/SequencerCore/Public/MVVM/ViewModels/OutlinerColumns/DeactivateOutlinerColumn.h
new file mode 100644
--- /dev/null
+++ b/Engine/Source/Editor/SequencerCore/Public/MVVM/ViewModels/OutlinerColumns/DeactivateOutlinerColumn.h
@@ -0,0 +1,23 @@
+// Copyright Epic Games, Inc. All Rights Reserved.
+
+#pragma once
+
+#include "MVVM/ViewModels/OutlinerColumns/OutlinerColumnBase.h"
+
+namespace UE::Sequencer
+{
+
+/**
+ * A column for local (not serialized to the asset) muting/unmuting tracks.
+ */
+class FDeactivateOutlinerColumn
+	: public FOutlinerColumnBase
+{
+public:
+	SEQUENCERCORE_API FDeactivateOutlinerColumn();
+
+	bool IsItemCompatibleWithColumn(const FCreateOutlinerColumnParams& InParams) const override;
+	TSharedPtr<SWidget> CreateColumnWidget(const FCreateOutlinerColumnParams& InParams, const TSharedRef<ISequencerTreeViewRow>& TreeViewRow) override;
+};
+	
+} // namespace UE::Sequencer
diff --git a/Engine/Source/Editor/SequencerCore/Public/MVVM/ViewModels/OutlinerColumns/OutlinerColumnTypes.h b/Engine/Source/Editor/SequencerCore/Public/MVVM/ViewModels/OutlinerColumns/OutlinerColumnTypes.h
--- a/Engine/Source/Editor/SequencerCore/Public/MVVM/ViewModels/OutlinerColumns/OutlinerColumnTypes.h
+++ b/Engine/Source/Editor/SequencerCore/Public/MVVM/ViewModels/OutlinerColumns/OutlinerColumnTypes.h
@@ -23,12 +23,19 @@ ENUM_CLASS_FLAGS(EOutlinerColumnFlags)
 
 enum class EOutlinerColumnGroup : uint8
 {
-	FarLeftGutter,
-
+	/** Far left gutter area: Indicators */
 	LeftGutter,
 
+	/** Asset state toggles area: Pin, Lock, Deactivate (states saved to the asset) */
+	AssetToggles,
+
+	/** Editor transient state toggles area: Mute, Solo (states NOT saved to the asset) */
+	TransientToggles,
+
+	/** Main center area: Add, Edit, Label */
 	Center,
 
+	/** Far right gutter area: Key Frame, Key Nav, Color Picker, etc. */
 	RightGutter,
 };
 
@@ -68,6 +75,7 @@ struct FCommonOutlinerNames
 	static SEQUENCERCORE_API FName Indicator;
 	static SEQUENCERCORE_API FName Pin;
 	static SEQUENCERCORE_API FName Lock;
+	static SEQUENCERCORE_API FName Deactivate;
 	static SEQUENCERCORE_API FName Mute;
 	static SEQUENCERCORE_API FName Solo;
 	static SEQUENCERCORE_API FName Label;
diff --git a/Engine/Source/Runtime/MovieScene/Private/MovieScene.cpp b/Engine/Source/Runtime/MovieScene/Private/MovieScene.cpp
--- a/Engine/Source/Runtime/MovieScene/Private/MovieScene.cpp
+++ b/Engine/Source/Runtime/MovieScene/Private/MovieScene.cpp
@@ -261,6 +261,14 @@ void UMovieScene::Serialize( FArchive& Ar )
 			ObjectBindings.Sort();
 		}
 	}
+
+#if WITH_EDITORONLY_DATA
+	if (Ar.IsSaving())
+	{
+		MuteNodes_DEPRECATED.Empty();
+		SoloNodes_DEPRECATED.Empty();
+	}
+#endif
 }
 
 #if WITH_EDITOR
diff --git a/Engine/Source/Runtime/MovieScene/Private/MovieSceneTrack.cpp b/Engine/Source/Runtime/MovieScene/Private/MovieSceneTrack.cpp
--- a/Engine/Source/Runtime/MovieScene/Private/MovieSceneTrack.cpp
+++ b/Engine/Source/Runtime/MovieScene/Private/MovieSceneTrack.cpp
@@ -475,23 +475,57 @@ TArray<UMovieSceneCondition*> UMovieSceneTrack::GetAllConditions()
 
 #endif
 
-bool UMovieSceneTrack::IsRowEvalDisabled(int32 RowIndex) const
+bool UMovieSceneTrack::IsEvalDisabled(const bool bInCheckLocal) const
 {
-	return RowsDisabled.Contains(RowIndex);
+	return bIsEvalDisabled
+#if WITH_EDITORONLY_DATA
+		|| (bInCheckLocal && bIsLocalEvalDisabled)
+#endif
+	;
+}
+
+bool UMovieSceneTrack::IsRowEvalDisabled(const int32 InRowIndex, const bool bInCheckLocal) const
+{
+	return RowsDisabled.Contains(InRowIndex)
+#if WITH_EDITORONLY_DATA
+		|| (bInCheckLocal && LocalRowsDisabled.Contains(InRowIndex))
+#endif
+		;
+}
+
+void UMovieSceneTrack::SetRowEvalDisabled(const bool bInEvalDisabled, const int32 InRowIndex)
+{
+	if (bInEvalDisabled)
+	{
+		RowsDisabled.AddUnique(InRowIndex);
+	}
+	else
+	{
+		RowsDisabled.Remove(InRowIndex);
+	}
 }
 
-void UMovieSceneTrack::SetRowEvalDisabled(bool bEvalDisabled, int32 RowIndex)
+#if WITH_EDITOR
+
+bool UMovieSceneTrack::IsLocalRowEvalDisabled(const int32 InRowIndex) const
 {
-	if (bEvalDisabled)
+	return LocalRowsDisabled.Contains(InRowIndex);
+}
+
+void UMovieSceneTrack::SetLocalRowEvalDisabled(const bool bInEvalDisabled, const int32 InRowIndex)
+{
+	if (bInEvalDisabled)
 	{
-		RowsDisabled.AddUnique(RowIndex);
+		LocalRowsDisabled.AddUnique(InRowIndex);
 	}
 	else
 	{
-		RowsDisabled.Remove(RowIndex);
+		LocalRowsDisabled.Remove(InRowIndex);
 	}
 }
 
+#endif
+
 const FMovieSceneTrackRowMetadata* UMovieSceneTrack::FindTrackRowMetadata(int32 RowIndex) const
 {
 	return TrackRowMetadata.Find(RowIndex);
diff --git a/Engine/Source/Runtime/MovieScene/Public/MovieScene.h b/Engine/Source/Runtime/MovieScene/Public/MovieScene.h
--- a/Engine/Source/Runtime/MovieScene/Public/MovieScene.h
+++ b/Engine/Source/Runtime/MovieScene/Public/MovieScene.h
@@ -134,6 +134,14 @@ struct FMovieSceneEditorData
 	/** Saved Navigation Tool state to restore */
 	UPROPERTY()
 	FNavigationToolSaveState NavigationToolState;
+
+	/** Nodes currently marked Solo, stored as node tree paths */
+	UPROPERTY(Transient)
+	TArray<FString> SoloNodes;
+
+	/** Nodes currently marked Mute, stored as node tree paths */
+	UPROPERTY(Transient)
+	TArray<FString> MuteNodes;
 };
 
 
@@ -963,14 +971,14 @@ public:
 	MOVIESCENE_API void EmptyRootFolders();
 
 	/**
-	 * Gets the nodes marked as solo in the editor, as node tree paths
+	 * Gets the nodes marked as transient solo in the editor, as node tree paths
 	 */
-	TArray<FString>& GetSoloNodes() { return SoloNodes; }
+	TArray<FString>& GetSoloNodes() { return EditorData.SoloNodes; }
 
 	/**
-	 * Gets the nodes marked as muted in the editor, as node tree paths
+	 * Gets the nodes marked as transient muted in the editor, as node tree paths
 	 */
-	TArray<FString>& GetMuteNodes() { return MuteNodes; }
+	TArray<FString>& GetMuteNodes() { return EditorData.MuteNodes; }
 	
 	//WIP Set of Delegates
 	/** Gets a multicast delegate which is executed whenever a channel is changed, currently only set by Python/BP actions.
@@ -1361,12 +1369,12 @@ private:
 	TArray<TObjectPtr<UMovieSceneFolder>> RootFolders;
 
 	/** Nodes currently marked Solo, stored as node tree paths */
-	UPROPERTY()
-	TArray<FString> SoloNodes;
-	
+	UPROPERTY(meta = (DeprecatedProperty, DeprecationMessage = "Solo nodes are now saved in FMovieSceneEditorData as transient"))
+	TArray<FString> SoloNodes_DEPRECATED;
+
 	/** Nodes currently marked Mute, stored as node tree paths */
-	UPROPERTY()
-	TArray<FString> MuteNodes;
+	UPROPERTY(meta = (DeprecatedProperty, DeprecationMessage = "Mute nodes are now saved in FMovieSceneEditorData as transient"))
+	TArray<FString> MuteNodes_DEPRECATED;
 
 	/** Groups of sections which should maintain the same relative offset */
 	UPROPERTY()
diff --git a/Engine/Source/Runtime/MovieScene/Public/MovieSceneTrack.h b/Engine/Source/Runtime/MovieScene/Public/MovieSceneTrack.h
--- a/Engine/Source/Runtime/MovieScene/Public/MovieSceneTrack.h
+++ b/Engine/Source/Runtime/MovieScene/Public/MovieSceneTrack.h
@@ -277,6 +277,18 @@ protected:
 	UPROPERTY()
 	TArray<int32> RowsDisabled;
 
+#if WITH_EDITORONLY_DATA
+
+	/** Whether evaluation of this track has been disabled locally (not serialized with the asset) */
+	UPROPERTY(Transient)
+	bool bIsLocalEvalDisabled;
+
+	/** Which rows have been disabled locally (not serialized with the asset) */
+	UPROPERTY(Transient)
+	TArray<int32> LocalRowsDisabled;
+
+#endif
+
 public:
 
 	/**
@@ -381,8 +393,8 @@ public:
 	virtual void RemoveAllAnimationData() { }
 
 	/**
-	* @return Whether or not this track supports multiple row indices.
-	*/
+	 * @return Whether or not this track supports multiple row indices.
+	 */
 	virtual bool SupportsMultipleRows() const
 	{
 		return SupportedBlendTypes.Num() != 0;
@@ -394,13 +406,13 @@ public:
 	}
 
 	/** Set This Section as the one to key. If track doesn't support layered blends then don't implement
-	* @param InSection		Section that we want to key
-	*/
+	 * @param InSection		Section that we want to key
+	 */
 	 virtual void SetSectionToKey(UMovieSceneSection* InSection) {};
 
 	/** Get the section we want to key. If track doesn't support layered blends it will return nulltpr.
-	* @return The section to key if one exists
-	*/
+	 * @return The section to key if one exists
+	 */
 	 virtual UMovieSceneSection* GetSectionToKey() const { return nullptr; }
 
 	/** Gets the greatest row index of all the sections owned by this track. */
@@ -416,30 +428,46 @@ public:
 	MOVIESCENE_API virtual void OnRowIndicesChanged(const TMap<int32, int32>& NewToOldRowIndices);
 
 	/**
-	* @return Whether evaluation of this track should be disabled due to mute/solo settings
-	*/
-	bool IsEvalDisabled() const { return bIsEvalDisabled; }
-	MOVIESCENE_API bool IsRowEvalDisabled(int32 RowIndex) const;
+	 * @return Whether evaluation of this track should be disabled due to deactive setting
+	 */
+	MOVIESCENE_API bool IsEvalDisabled(const bool bInCheckLocal = true) const;
+	MOVIESCENE_API bool IsRowEvalDisabled(const int32 InRowIndex, const bool bInCheckLocal = true) const;
 
 	/**
-	* Called by Sequencer to set whether evaluation of this track should be disabled due to mute/solo settings
-	*/
-	void SetEvalDisabled(bool bEvalDisabled) { bIsEvalDisabled = bEvalDisabled; }
-	MOVIESCENE_API void SetRowEvalDisabled(bool bEvalDisabled, int32 RowIndex);
+	 * Called by Sequencer to set whether evaluation of this track should be disabled due to deactive setting
+	 */
+	void SetEvalDisabled(const bool bInEvalDisabled) { bIsEvalDisabled = bInEvalDisabled; }
+	MOVIESCENE_API void SetRowEvalDisabled(const bool bInEvalDisabled, const int32 InRowIndex);
+
+#if WITH_EDITOR
+
+	/**
+	 * @return Whether evaluation of this track should be disabled due to local muting/soloing (not serialized with the asset)
+	 */
+	bool IsLocalEvalDisabled() const { return bIsLocalEvalDisabled; }
+	MOVIESCENE_API bool IsLocalRowEvalDisabled(const int32 InRowIndex) const;
+
+	/**
+	 * Called by Sequencer to set whether evaluation of this track should be disabled due to local muting/soloing (not serialized with the asset)
+	 */
+	void SetLocalEvalDisabled(const bool bInEvalDisabled) { bIsLocalEvalDisabled = bInEvalDisabled; }
+	MOVIESCENE_API void SetLocalRowEvalDisabled(const bool bInEvalDisabled, int32 InRowIndex);
+
+#endif
 
 	/*
-	* Returns a pointer to optional track row metadata at the given RowIndex, or nullptr if none exists.
-	*/
+	 * Returns a pointer to optional track row metadata at the given RowIndex, or nullptr if none exists.
+	 */
 	MOVIESCENE_API const FMovieSceneTrackRowMetadata* FindTrackRowMetadata(int32 RowIndex) const;
 
 	/*
-	* Returns a pointer to optional track row metadata at the given RowIndex, or nullptr if none exists;
-	*/
+	 * Returns a pointer to optional track row metadata at the given RowIndex, or nullptr if none exists;
+	 */
 	MOVIESCENE_API FMovieSceneTrackRowMetadata* FindTrackRowMetadata(int32 RowIndex);
 
 	/*
-	* Returns a pointer to optional track row metadata at the given RowIndex, or nullptr if none exists;
-	*/
+	 * Returns a pointer to optional track row metadata at the given RowIndex, or nullptr if none exists;
+	 */
 	MOVIESCENE_API FMovieSceneTrackRowMetadata& FindOrAddTrackRowMetadata(int32 RowIndex);
 
 public:
