[Sequencer] Fix horde error

diff --git a/Engine/Source/Editor/Sequencer/Private/MVVM/ViewModels/OutlinerItemModel.cpp b/Engine/Source/Editor/Sequencer/Private/MVVM/ViewModels/OutlinerItemModel.cpp
--- a/Engine/Source/Editor/Sequencer/Private/MVVM/ViewModels/OutlinerItemModel.cpp
+++ b/Engine/Source/Editor/Sequencer/Private/MVVM/ViewModels/OutlinerItemModel.cpp
@@ -1272,9 +1272,9 @@ void FEvaluableOutlinerItemModel::SetIsMuted(bool bIsMuted)
 
 bool FEvaluableOutlinerItemModel::IsDeactivated() const
 {
-	for (const TViewModelPtr<ITrackExtension> TrackNode : GetDescendantsOfType<ITrackExtension>(true))
+	for (const TViewModelPtr<ITrackExtension>& TrackNode : GetDescendantsOfType<ITrackExtension>(true))
 	{
-		UMovieSceneTrack* const Track = TrackNode->GetTrack();
+		const UMovieSceneTrack* const Track = TrackNode->GetTrack();
 		if (!IsValid(Track))
 		{
 			continue;
@@ -1304,7 +1304,7 @@ void FEvaluableOutlinerItemModel::SetIsDeactivated(const bool bInIsDeactivated)
 {
 	bool bAnyChanged = false;
 
-	for (const TViewModelPtr<ITrackExtension> TrackNode : GetDescendantsOfType<ITrackExtension>(true))
+	for (const TViewModelPtr<ITrackExtension>& TrackNode : GetDescendantsOfType<ITrackExtension>(true))
 	{
 		UMovieSceneTrack* const Track = TrackNode->GetTrack();
 		if (!IsValid(Track))
diff --git a/Engine/Source/Editor/SequencerCore/Private/MVVM/Views/OutlinerColumns/SDeactivateColumnWidget.cpp b/Engine/Source/Editor/SequencerCore/Private/MVVM/Views/OutlinerColumns/SDeactivateColumnWidget.cpp
--- a/Engine/Source/Editor/SequencerCore/Private/MVVM/Views/OutlinerColumns/SDeactivateColumnWidget.cpp
+++ b/Engine/Source/Editor/SequencerCore/Private/MVVM/Views/OutlinerColumns/SDeactivateColumnWidget.cpp
@@ -77,7 +77,7 @@ void SDeactivateColumnWidget::SetIsActive_Internal(const FViewModelPtr& ViewMode
 		Deactivatable->SetIsDeactivated(bInIsActive);
 	}
 	// Otherwise deactivate deactivatable children of this (if any)
-	else for (const TViewModelPtr<IDeactivatableExtension> Child : ViewModel->GetDescendantsOfType<IDeactivatableExtension>())
+	else for (const TViewModelPtr<IDeactivatableExtension>& Child : ViewModel->GetDescendantsOfType<IDeactivatableExtension>())
 	{
 		if (Child->IsInheritable())
 		{
