Sequencer: Fix duplicate LOCTEXT

#rnx

diff --git a/Engine/Source/Editor/Sequencer/Private/Tools/NavigationTool/Widgets/SNavigationToolItemColumns.cpp b/Engine/Source/Editor/Sequencer/Private/Tools/NavigationTool/Widgets/SNavigationToolItemColumns.cpp
--- a/Engine/Source/Editor/Sequencer/Private/Tools/NavigationTool/Widgets/SNavigationToolItemColumns.cpp
+++ b/Engine/Source/Editor/Sequencer/Private/Tools/NavigationTool/Widgets/SNavigationToolItemColumns.cpp
@@ -350,7 +350,7 @@ FText SNavigationToolItemColumns::GetItemMenuButtonToolTip() const
 
 		return OutText.IsEmpty()
 			? LOCTEXT("NoColumnsDisplayedTooltip", "No columns displayed")
-			: FText::Format(LOCTEXT("ItemColumnText", "Columns Displayed: {0}\n\n"
+			: FText::Format(LOCTEXT("ColumnsDisplayedText", "Columns Displayed: {0}\n\n"
 				"Shift + Click to hide all columns"), OutText);
 	}
 
