[Sequencer] Remove unneeded ensure when there is no saved view state for a sequence, which will be the case for Motion Design embedded sequences until the Motion Design Nav Tool extension CL is submitted

diff --git a/Engine/Source/Editor/Sequencer/Private/Tools/NavigationTool/NavigationToolView.cpp b/Engine/Source/Editor/Sequencer/Private/Tools/NavigationTool/NavigationToolView.cpp
--- a/Engine/Source/Editor/Sequencer/Private/Tools/NavigationTool/NavigationToolView.cpp
+++ b/Engine/Source/Editor/Sequencer/Private/Tools/NavigationTool/NavigationToolView.cpp
@@ -517,7 +517,7 @@ ENavigationToolItemFlags FNavigationToolView::GetViewItemFlags(const FNavigation
 	}
 
 	FNavigationToolViewSaveState* const ViewSaveState = Provider->GetViewSaveState(*Tool, ToolViewId);
-	if (!ensure(ViewSaveState))
+	if (!ViewSaveState)
 	{
 		return ENavigationToolItemFlags::None;
 	}
