Sequencer: Use ParallelFor's for iterating over channels to improve bulk operations in particularly with Control Rig. Currently used for deleting keys.

diff --git a/Engine/Source/Editor/Sequencer/Private/Sequencer.cpp b/Engine/Source/Editor/Sequencer/Private/Sequencer.cpp
--- a/Engine/Source/Editor/Sequencer/Private/Sequencer.cpp
+++ b/Engine/Source/Editor/Sequencer/Private/Sequencer.cpp
@@ -5,6 +5,7 @@
 #include "GameFramework/Actor.h"
 #include "Engine/World.h"
 #include "Algo/RemoveIf.h"
+#include "Async/ParallelFor.h"
 #include "MVVM/SharedViewModelData.h"
 #include "MVVM/Selection/Selection.h"
 #include "MVVM/Extensions/IClockExtension.h"
@@ -1910,35 +1911,38 @@ void FSequencer::DeleteSections(const TSet<UMovieSceneSection*>& Sections)
 void FSequencer::DeleteSelectedKeys()
 {
 	using namespace UE::Sequencer;
+	const bool bUseParallel = true;
 
 	FScopedTransaction DeleteKeysTransaction( NSLOCTEXT("Sequencer", "DeleteSelectedKeys_Transaction", "Delete Selected Keys") );
 	bool bAnythingRemoved = false;
 
 	FSelectedKeysByChannel KeysByChannel(ViewModel->GetSelection()->KeySelection);
 	TSet<UObject*> ModifiedObjects;
-
-	for (const FSelectedChannelInfo& ChannelInfo : KeysByChannel.SelectedChannels)
+	FCriticalSection Mutex;
+	ParallelFor(KeysByChannel.SelectedChannels.Num(), [this, &KeysByChannel, &Mutex, &ModifiedObjects, &bAnythingRemoved](int32 Index)
 	{
+		const FSelectedChannelInfo& ChannelInfo = KeysByChannel.SelectedChannels[Index];
+
 		if (ChannelInfo.OwningSection->IsReadOnly())
 		{
-			continue;
+			return;
 		}
 
 		if (ChannelInfo.KeyHandles.Num() == 0)
 		{
-			continue;
+			return;
 		}
 
 		TSharedPtr<FChannelModel> Channel = ViewModel->GetSelection()->KeySelection.GetModelForKey(ChannelInfo.KeyHandles[0]);
 		if (!Channel)
 		{
-			continue;
+			return;
 		}
 
 		TSharedPtr<IKeyArea> KeyArea = Channel->GetKeyArea();
 		if (!KeyArea)
 		{
-			continue;
+			return;
 		}
 
 		UObject* Owner = nullptr;
@@ -1952,15 +1956,18 @@ void FSequencer::DeleteSelectedKeys()
 			Owner = ChannelInfo.OwningSection;
 		}
 
-		if (!ModifiedObjects.Contains(Owner))
 		{
-			Owner->Modify();
-			ModifiedObjects.Add(ChannelInfo.OwningSection);
+			FScopeLock Lock(&Mutex);
+			if (!ModifiedObjects.Contains(Owner))
+			{
+				Owner->Modify();
+				ModifiedObjects.Add(ChannelInfo.OwningSection);
+			}
+			bAnythingRemoved = true;
 		}
-
 		KeyArea->DeleteKeys(ChannelInfo.KeyHandles, GetLocalTime().Time.FloorToFrame());
-		bAnythingRemoved = true;
-	}
+
+	},bUseParallel ? EParallelForFlags::None : EParallelForFlags::ForceSingleThread);
 
 	if (bAnythingRemoved)
 	{
