[Sequencer] Fix Deactivate column from showing incorrect state when child tracks are deleted

diff --git a/Engine/Source/Editor/Sequencer/Private/MVVM/ViewModels/OutlinerItemModel.cpp b/Engine/Source/Editor/Sequencer/Private/MVVM/ViewModels/OutlinerItemModel.cpp
--- a/Engine/Source/Editor/Sequencer/Private/MVVM/ViewModels/OutlinerItemModel.cpp
+++ b/Engine/Source/Editor/Sequencer/Private/MVVM/ViewModels/OutlinerItemModel.cpp
@@ -1297,7 +1297,7 @@ bool FEvaluableOutlinerItemModel::IsDeactivated() const
 		}
 	}
 
-	return true;
+	return false;
 }
 
 void FEvaluableOutlinerItemModel::SetIsDeactivated(const bool bInIsDeactivated)
