[Sequencer] Fix deactivate toggle not deactivating tracks

diff --git a/Engine/Source/Editor/Sequencer/Private/MVVM/ViewModels/OutlinerItemModel.cpp b/Engine/Source/Editor/Sequencer/Private/MVVM/ViewModels/OutlinerItemModel.cpp
--- a/Engine/Source/Editor/Sequencer/Private/MVVM/ViewModels/OutlinerItemModel.cpp
+++ b/Engine/Source/Editor/Sequencer/Private/MVVM/ViewModels/OutlinerItemModel.cpp
@@ -1282,17 +1282,16 @@ bool FEvaluableOutlinerItemModel::IsDeactivated() const
 
 		if (const TViewModelPtr<FTrackRowModel> TrackRowModel = TrackNode.ImplicitCast())
 		{
-			if (!Track->IsRowEvalDisabled(TrackNode->GetRowIndex(), /*bInCheckLocal=*/false))
+			if (Track->IsRowEvalDisabled(TrackNode->GetRowIndex(), /*bInCheckLocal=*/false))
 			{
-				return false;
+				return true;
 			}
 		}
 		else
 		{
-			// Deactive state (dirtying, saved with asset, evaluation)
-			if (!Track->IsEvalDisabled(/*bInCheckLocal=*/false))
+			if (Track->IsEvalDisabled(/*bInCheckLocal=*/false))
 			{
-				return false;
+				return true;
 			}
 		}
 	}
@@ -1312,9 +1311,9 @@ void FEvaluableOutlinerItemModel::SetIsDeactivated(const bool bInIsDeactivated)
 			continue;
 		}
 
+		// Deactive state (dirtying, saved with asset, evaluation)
 		if (const TViewModelPtr<FTrackRowModel> TrackRowModel = TrackNode.ImplicitCast())
 		{
-			// Deactive state (dirtying, saved with asset, evaluation)
 			if (bInIsDeactivated != Track->IsRowEvalDisabled(TrackNode->GetRowIndex(), /*bInCheckLocal=*/false))
 			{
 				Track->Modify();
@@ -1324,7 +1323,6 @@ void FEvaluableOutlinerItemModel::SetIsDeactivated(const bool bInIsDeactivated)
 		}
 		else
 		{
-			// Deactive state (dirtying, saved with asset, evaluation)
 			if (bInIsDeactivated != Track->IsEvalDisabled(/*bInCheckLocal=*/false))
 			{
 				Track->Modify();
