[Sequencer] Change back to old way of handling IsDeactivated, but add check for when there are no child tracks

diff --git a/Engine/Source/Editor/Sequencer/Private/MVVM/ViewModels/OutlinerItemModel.cpp b/Engine/Source/Editor/Sequencer/Private/MVVM/ViewModels/OutlinerItemModel.cpp
--- a/Engine/Source/Editor/Sequencer/Private/MVVM/ViewModels/OutlinerItemModel.cpp
+++ b/Engine/Source/Editor/Sequencer/Private/MVVM/ViewModels/OutlinerItemModel.cpp
@@ -1272,7 +1272,13 @@ void FEvaluableOutlinerItemModel::SetIsMuted(bool bIsMuted)
 
 bool FEvaluableOutlinerItemModel::IsDeactivated() const
 {
-	for (const TViewModelPtr<ITrackExtension>& TrackNode : GetDescendantsOfType<ITrackExtension>(true))
+	const TParentFirstChildIterator<ITrackExtension> Descendants = GetDescendantsOfType<ITrackExtension>(true);
+	if (!Descendants)
+	{
+		return false;
+	}
+
+	for (const TViewModelPtr<ITrackExtension>& TrackNode : Descendants)
 	{
 		const UMovieSceneTrack* const Track = TrackNode->GetTrack();
 		if (!IsValid(Track))
@@ -1282,21 +1288,21 @@ bool FEvaluableOutlinerItemModel::IsDeactivated() const
 
 		if (const TViewModelPtr<FTrackRowModel> TrackRowModel = TrackNode.ImplicitCast())
 		{
-			if (Track->IsRowEvalDisabled(TrackNode->GetRowIndex(), /*bInCheckLocal=*/false))
+			if (!Track->IsRowEvalDisabled(TrackNode->GetRowIndex(), /*bInCheckLocal=*/false))
 			{
-				return true;
+				return false;
 			}
 		}
 		else
 		{
-			if (Track->IsEvalDisabled(/*bInCheckLocal=*/false))
+			if (!Track->IsEvalDisabled(/*bInCheckLocal=*/false))
 			{
-				return true;
+				return false;
 			}
 		}
 	}
 
-	return false;
+	return true;
 }
 
 void FEvaluableOutlinerItemModel::SetIsDeactivated(const bool bInIsDeactivated)
