Sequencer: Fix regression from 41073775 in not snapping to whole frames when the snapping button is enabled

#rnx

diff --git a/Engine/Source/Editor/Sequencer/Private/Tools/SequencerSnapField.cpp b/Engine/Source/Editor/Sequencer/Private/Tools/SequencerSnapField.cpp
--- a/Engine/Source/Editor/Sequencer/Private/Tools/SequencerSnapField.cpp
+++ b/Engine/Source/Editor/Sequencer/Private/Tools/SequencerSnapField.cpp
@@ -318,7 +318,7 @@ TOptional<FSequencerSnapField::FSnapResult> FSequencerSnapField::Snap(const TArr
 		}
 
 		const bool bShouldTrySnappingToInterval = !bSnapToLikeTypes || IntervalSnapThreshold > Threshold;
-		if (bSnapToInterval && bShouldTrySnappingToInterval)
+		if (bSnapToInterval || bShouldTrySnappingToInterval)
 		{
 			FFrameTime ThisTime = ThisSnap.IsSet() ? ThisSnap->SnappedTime : Time;
 			FFrameTime IntervalTime = FFrameRate::TransformTime(ThisTime, TickResolution, DisplayRate);
