[Sequencer] Show hotkeys for track filter menu

diff --git a/Engine/Source/Editor/Sequencer/Private/Filters/Menus/SequencerTrackFilterMenu.cpp b/Engine/Source/Editor/Sequencer/Private/Filters/Menus/SequencerTrackFilterMenu.cpp
--- a/Engine/Source/Editor/Sequencer/Private/Filters/Menus/SequencerTrackFilterMenu.cpp
+++ b/Engine/Source/Editor/Sequencer/Private/Filters/Menus/SequencerTrackFilterMenu.cpp
@@ -368,7 +368,7 @@ void FSequencerTrackFilterMenu::FillFiltersMenuCategory(FToolMenuSection& InOutS
 
 		const FText DisplayName = Filter->GetDisplayName();
 
-		InOutSection.AddMenuEntry(*DisplayName.ToString(),
+		FToolMenuEntry& NewMenuEntry = InOutSection.AddMenuEntry(*DisplayName.ToString(),
 			DisplayName,
 			Filter->GetToolTipText(),
 			Filter->GetIcon(),
@@ -384,6 +384,11 @@ void FSequencerTrackFilterMenu::FillFiltersMenuCategory(FToolMenuSection& InOutS
 					})
 			),
 			EUserInterfaceActionType::ToggleButton);
+
+		if (const TSharedPtr<FUICommandInfo> ToggleCommand = Filter->GetToggleCommand())
+		{
+			NewMenuEntry.InputBindingLabel = ToggleCommand->GetInputText();
+		}
 	}
 }
 
