[Sequencer] Navigation Tool improvements:
- Fix binding icons not showing
- Fix crash on expanding binding items caused by conflicting item Ids
- Double click to focus now works for all sequences/tracks/bindings
- Add playhead column that shows if the items range contains the playhead
- Better short name functionality
- Turn on Bindings global filter by default
- Submenus for Column Views and Filter Bar Options
- Item chips focus the item in sequencer similar to what double clicking an item does
- Item context menu action to focus sequence in content browser
- Source control column icons now show as expected (only sequences for now)
- Add "Playhead" and "Dirty" filters
- Show hotkeys for filter menu entries
- Sync selection to/from Sequencer <-> Navigation Tool
- Remove unneeded code
- Show disabled state when dragging (for now)
- Add new updated tool icon

diff --git a/Engine/Plugins/MovieScene/LevelSequenceEditor/Source/LevelSequenceEditor/Private/LevelSequenceNavigationToolProvider.cpp b/Engine/Plugins/MovieScene/LevelSequenceEditor/Source/LevelSequenceEditor/Private/LevelSequenceNavigationToolProvider.cpp
--- a/Engine/Plugins/MovieScene/LevelSequenceEditor/Source/LevelSequenceEditor/Private/LevelSequenceNavigationToolProvider.cpp
+++ b/Engine/Plugins/MovieScene/LevelSequenceEditor/Source/LevelSequenceEditor/Private/LevelSequenceNavigationToolProvider.cpp
@@ -16,6 +16,7 @@
 #include "Tools/NavigationTool/Columns/NavigationToolLockColumn.h"
 #include "Tools/NavigationTool/Columns/NavigationToolMarkerVisibilityColumn.h"
 #include "Tools/NavigationTool/Columns/NavigationToolOutTimeColumn.h"
+#include "Tools/NavigationTool/Columns/NavigationToolPlayheadColumn.h"
 #include "Tools/NavigationTool/Columns/NavigationToolRevisionControlColumn.h"
 #include "Tools/NavigationTool/Columns/NavigationToolStartFrameOffsetColumn.h"
 #include "Tools/NavigationTool/Columns/NavigationToolTakeColumn.h"
@@ -40,23 +41,44 @@ FLevelSequenceNavigationToolProvider::FLevelSequenceNavigationToolProvider()
 {
 }
 
+FName FLevelSequenceNavigationToolProvider::GetIdentifier() const
+{
+	return Identifier;
+}
+
+FName FLevelSequenceNavigationToolProvider::GetDefaultColumnView() const
+{
+	return AnimationColumnViewId;
+}
+
 void FLevelSequenceNavigationToolProvider::BindCommands(const TSharedRef<FUICommandList>& InCommandList)
 {
 	InCommandList->Append(ToolCommands);
 }
 
-FName FLevelSequenceNavigationToolProvider::GetIdentifier() const
+void FLevelSequenceNavigationToolProvider::OnRegistered()
 {
-	return Identifier;
+	// @TODO: Move from core Navigation Tool
+	/*ExtendToolToolBar();
+	ExtendToolItemContextMenu();
+
+	DragDropInitializedDelegate = FNavigationToolExtender::OnItemDragDropOpInitialized()
+		.AddSP(this, &FAvaNavigationToolProvider::ExtendToolItemDragDropOp);*/
 }
 
-FName FLevelSequenceNavigationToolProvider::GetDefaultColumnView() const
+void FLevelSequenceNavigationToolProvider::OnUnregistered()
 {
-	return AnimationColumnViewId;
+	// @TODO: Move from core Navigation Tool
+	/*RemoveToolToolBarExtension();
+	RemoveToolItemContextMenuExtension();
+
+	FNavigationToolExtender::OnItemDragDropOpInitialized()
+		.Remove(DragDropInitializedDelegate);*/
 }
 
 void FLevelSequenceNavigationToolProvider::OnExtendColumns(FNavigationToolColumnExtender& OutExtender)
 {
+	OutExtender.AddColumn<FNavigationToolPlayheadColumn>();
 	OutExtender.AddColumn<FNavigationToolDeactiveStateColumn>();
 	OutExtender.AddColumn<FNavigationToolMarkerVisibilityColumn>();
 	OutExtender.AddColumn<FNavigationToolLockColumn>();
@@ -87,6 +109,7 @@ void FLevelSequenceNavigationToolProvider::OnExtendColumnViews(TSet<FNavigationT
 
 	AnimationColumnView.VisibleColumns =
 		{
+			FNavigationToolPlayheadColumn::GetStaticTypeId().ToName(),
 			FNavigationToolColorColumn::GetStaticTypeId().ToName(),
 			FNavigationToolLabelColumn::GetStaticTypeId().ToName(),
 			FNavigationToolItemsColumn::GetStaticTypeId().ToName(),
diff --git a/Engine/Plugins/MovieScene/LevelSequenceEditor/Source/LevelSequenceEditor/Private/LevelSequenceNavigationToolProvider.h b/Engine/Plugins/MovieScene/LevelSequenceEditor/Source/LevelSequenceEditor/Private/LevelSequenceNavigationToolProvider.h
--- a/Engine/Plugins/MovieScene/LevelSequenceEditor/Source/LevelSequenceEditor/Private/LevelSequenceNavigationToolProvider.h
+++ b/Engine/Plugins/MovieScene/LevelSequenceEditor/Source/LevelSequenceEditor/Private/LevelSequenceNavigationToolProvider.h
@@ -8,7 +8,6 @@
 #include "UObject/WeakObjectPtr.h"
 
 class FUICommandList;
-class ISequencer;
 struct FMovieSceneEditorData;
 
 namespace UE::Sequencer::NavigationTool
@@ -27,14 +26,17 @@ public:
 	FLevelSequenceNavigationToolProvider();
 	virtual ~FLevelSequenceNavigationToolProvider() override = default;
 
-	void BindCommands(const TSharedRef<FUICommandList>& InCommandList);
-
 	//~ Begin INavigationToolProvider
 
 	virtual FName GetIdentifier() const override;
 
 	virtual FName GetDefaultColumnView() const override;
 
+	virtual void BindCommands(const TSharedRef<FUICommandList>& InCommandList) override;
+
+	virtual void OnRegistered() override;
+	virtual void OnUnregistered() override;
+
 	virtual void OnExtendColumns(FNavigationToolColumnExtender& OutExtender) override;
 	virtual void OnExtendColumnViews(TSet<FNavigationToolColumnView>& OutColumnViews) override;
 	virtual void OnExtendItemChildren(INavigationTool& InTool
diff --git a/Engine/Source/Editor/EditorStyle/Private/StarshipStyle.cpp b/Engine/Source/Editor/EditorStyle/Private/StarshipStyle.cpp
--- a/Engine/Source/Editor/EditorStyle/Private/StarshipStyle.cpp
+++ b/Engine/Source/Editor/EditorStyle/Private/StarshipStyle.cpp
@@ -2513,6 +2513,9 @@ void FStarshipEditorStyle::FStyle::SetupSequencerStyles()
 
 	// Navigation Tool Styling
 	{
+		Set("Sequencer.NavigationTool", new IMAGE_BRUSH_SVG("Sequencer/SequenceNavigator/SequenceNavigator_20", Icon20x20));
+		Set("Sequencer.NavigationTool.TabIcon", new IMAGE_BRUSH_SVG("Sequencer/SequenceNavigator/SequenceNavigator_16", Icon16x16));
+
 		const FTableRowStyle& AlternatingTableRowStyle = FAppStyle::Get().GetWidgetStyle<FTableRowStyle>(TEXT("TableView.AlternatingRow"));
 		Set("Sequencer.NavigationTool.TableViewRow", FTableRowStyle(AlternatingTableRowStyle));
 
diff --git a/Engine/Source/Editor/Sequencer/Private/SSequencer.cpp b/Engine/Source/Editor/Sequencer/Private/SSequencer.cpp
--- a/Engine/Source/Editor/Sequencer/Private/SSequencer.cpp
+++ b/Engine/Source/Editor/Sequencer/Private/SSequencer.cpp
@@ -1668,7 +1668,7 @@ void SSequencer::PopulateToolBar(UToolMenu* InMenu)
 		FToolMenuEntry ShowNavigationToolEntry = FToolMenuEntry::InitToolBarButton(
 			FSequencerCommands::Get().ToggleNavigationToolVisible,
 			TAttribute<FText>(),TAttribute<FText>(),
-			FSlateIcon(FAppStyle::GetAppStyleSetName(), TEXT("GraphEditor.Sequence_16x")));
+			FSlateIcon(FAppStyle::GetAppStyleSetName(), TEXT("Sequencer.NavigationTool")));
 		ShowNavigationToolEntry.StyleNameOverride = SequencerToolbarStyleName;
 		NavigationToolSection.AddEntry(ShowNavigationToolEntry);
 	}
diff --git a/Engine/Source/Editor/Sequencer/Private/Tools/NavigationTool/Columns/NavigationToolDeactiveStateColumn.cpp b/Engine/Source/Editor/Sequencer/Private/Tools/NavigationTool/Columns/NavigationToolDeactiveStateColumn.cpp
--- a/Engine/Source/Editor/Sequencer/Private/Tools/NavigationTool/Columns/NavigationToolDeactiveStateColumn.cpp
+++ b/Engine/Source/Editor/Sequencer/Private/Tools/NavigationTool/Columns/NavigationToolDeactiveStateColumn.cpp
@@ -17,7 +17,7 @@ namespace UE::Sequencer::NavigationTool
 
 FText FNavigationToolDeactiveStateColumn::GetColumnDisplayNameText() const
 {
-	return LOCTEXT("NavigationToolMuteStateColumn", "Sequence Evaluation Mute State");
+	return LOCTEXT("NavigationToolMuteStateColumn", "Sequence Evaluation");
 }
 
 const FSlateBrush* FNavigationToolDeactiveStateColumn::GetIconBrush() const
diff --git a/Engine/Source/Editor/Sequencer/Private/Tools/NavigationTool/Columns/NavigationToolPlayheadColumn.cpp b/Engine/Source/Editor/Sequencer/Private/Tools/NavigationTool/Columns/NavigationToolPlayheadColumn.cpp
new file mode 100644
--- /dev/null
+++ b/Engine/Source/Editor/Sequencer/Private/Tools/NavigationTool/Columns/NavigationToolPlayheadColumn.cpp
@@ -0,0 +1,58 @@
+﻿// Copyright Epic Games, Inc. All Rights Reserved.
+
+#include "Tools/NavigationTool/Columns/NavigationToolPlayheadColumn.h"
+#include "Tools/NavigationTool/INavigationToolView.h"
+#include "Tools/NavigationTool/Items/INavigationToolItem.h"
+#include "Tools/NavigationTool/Widgets/Columns/SNavigationToolPlayhead.h"
+#include "Widgets/Views/SHeaderRow.h"
+
+#define LOCTEXT_NAMESPACE "NavigationToolPlayheadColumn"
+
+namespace UE::Sequencer::NavigationTool
+{
+
+FText FNavigationToolPlayheadColumn::GetColumnDisplayNameText() const
+{
+	return LOCTEXT("NavigationToolPlayheadColumn", "Playhead");
+}
+
+const FSlateBrush* FNavigationToolPlayheadColumn::GetIconBrush() const
+{
+	return FAppStyle::GetBrush(TEXT("GenericPlay"));
+}
+
+SHeaderRow::FColumn::FArguments FNavigationToolPlayheadColumn::ConstructHeaderRowColumn(const TSharedRef<INavigationToolView>& InToolView
+	, const float InFillSize)
+{
+	const FName ColumnId = GetColumnId();
+
+	return SHeaderRow::Column(ColumnId)
+		.FixedWidth(24.f)
+		.HAlignHeader(HAlign_Center)
+		.VAlignHeader(VAlign_Center)
+		.HAlignCell(HAlign_Center)
+		.VAlignCell(VAlign_Center)
+		.DefaultLabel(GetColumnDisplayNameText())
+		.DefaultTooltip(GetColumnDisplayNameText())
+		[
+			SNew(SImage)
+			.ColorAndOpacity(FSlateColor::UseForeground())
+			.Image(FAppStyle::GetBrush(TEXT("GenericPlay")))
+		];
+}
+
+TSharedRef<SWidget> FNavigationToolPlayheadColumn::ConstructRowWidget(const FNavigationToolItemPtr& InItem
+	, const TSharedRef<INavigationToolView>& InToolView
+	, const TSharedRef<SNavigationToolTreeRow>& InRow)
+{
+	if (!InItem->IsA<IPlayheadExtension>())
+	{
+		return SNullWidget::NullWidget;
+	}
+
+	return SNew(SNavigationToolPlayhead, SharedThis(this), InItem, InToolView, InRow);
+}
+
+}
+
+#undef LOCTEXT_NAMESPACE
diff --git a/Engine/Source/Editor/Sequencer/Private/Tools/NavigationTool/Customization/NavigationToolExtender.cpp b/Engine/Source/Editor/Sequencer/Private/Tools/NavigationTool/Customization/NavigationToolExtender.cpp
--- a/Engine/Source/Editor/Sequencer/Private/Tools/NavigationTool/Customization/NavigationToolExtender.cpp
+++ b/Engine/Source/Editor/Sequencer/Private/Tools/NavigationTool/Customization/NavigationToolExtender.cpp
@@ -57,15 +57,20 @@ bool FNavigationToolExtender::RegisterToolProvider(const TSharedRef<INavigationT
 	const FName Identifier = InProvider->GetIdentifier();
 	if (const TSharedPtr<FNavigationToolProvider> Provider = FindToolProvider(RawToolRef, Identifier))
 	{
-		UnregisterToolProvider(InTool, Identifier);
-		//return false;
+		return false;
 	}
 
 	// Add the provider if an entry for this tool already exists
 	if (FNavigationToolProviderEntry* const Entry = FindToolProvidersEntry(RawToolRef))
 	{
 		Entry->WeakProviders.Add(InProvider);
-		UE_LOG(LogSequencer, Display, TEXT("Navigation Tool Provider Added: %s"), *InProvider->GetIdentifier().ToString());
+		UE_LOG(LogSequencer, Warning, TEXT("Found Registered Navigation Tool Provider: %s"), *InProvider->GetIdentifier().ToString());
+
+		if (const TSharedPtr<FUICommandList> BaseCommandList = InTool->GetBaseCommandList())
+		{
+			InProvider->BindCommands(BaseCommandList.ToSharedRef());
+		}
+		InProvider->OnRegistered();
 
 		ThisInstance.ProvidersChangedDelegate.Broadcast(InTool, InProvider, ENavigationToolProvidersChangeType::Add);
 
@@ -85,7 +90,13 @@ bool FNavigationToolExtender::RegisterToolProvider(const TSharedRef<INavigationT
 	FNavigationToolProviderEntry NewEntry(InTool, { InProvider });
 
 	ThisInstance.ToolProviders.Add(MoveTemp(NewEntry));
-	UE_LOG(LogSequencer, Display, TEXT("Navigation Tool Provider Entry Added: %s"), *ToolInstanceName.ToString());
+	UE_LOG(LogSequencer, Log, TEXT("Navigation Tool Provider Registered: %s"), *ToolInstanceName.ToString());
+
+	if (const TSharedPtr<FUICommandList> BaseCommandList = InTool->GetBaseCommandList())
+	{
+		InProvider->BindCommands(BaseCommandList.ToSharedRef());
+	}
+	InProvider->OnRegistered();
 
 	ThisInstance.ProvidersChangedDelegate.Broadcast(InTool, InProvider, ENavigationToolProvidersChangeType::Add);
 
@@ -118,10 +129,14 @@ bool FNavigationToolExtender::UnregisterToolProvider(const TSharedRef<INavigatio
 					{
 						if (Provider->GetIdentifier() == InProviderId)
 						{
-							ProviderIter.RemoveCurrent();
+							Provider->OnUnregistered();
 
 							ThisInstance.ProvidersChangedDelegate.Broadcast(InTool, Provider.ToSharedRef(), ENavigationToolProvidersChangeType::Remove);
 
+							UE_LOG(LogSequencer, Log, TEXT("Navigation Tool Provider Unregistered: %s"), *Provider->GetIdentifier().ToString());
+
+							ProviderIter.RemoveCurrent();
+
 							return true;
 						}
 					}
@@ -153,14 +168,14 @@ void FNavigationToolExtender::UnregisterAllToolProviders(const TSharedRef<INavig
 				{
 					if (const TSharedPtr<FNavigationToolProvider> Provider = ProviderIter->Pin())
 					{
-						ProviderIter.RemoveCurrent();
+						Provider->OnUnregistered();
 
 						ThisInstance.ProvidersChangedDelegate.Broadcast(InTool, Provider.ToSharedRef(), ENavigationToolProvidersChangeType::Remove);
+
+						UE_LOG(LogSequencer, Log, TEXT("Navigation Tool Provider Unregistered: %s"), *Provider->GetIdentifier().ToString());
 					}
-					else
-					{
-						ProviderIter.RemoveCurrent();
-					}
+
+					ProviderIter.RemoveCurrent();
 				}
 			}
 		}
diff --git a/Engine/Source/Editor/Sequencer/Private/Tools/NavigationTool/Customization/NavigationToolObjectIconCustomization.cpp b/Engine/Source/Editor/Sequencer/Private/Tools/NavigationTool/Customization/NavigationToolObjectIconCustomization.cpp
deleted file mode 100644
diff --git a/Engine/Source/Editor/Sequencer/Private/Tools/NavigationTool/DragDropOps/Handlers/NavigationToolActorDropHandler.cpp b/Engine/Source/Editor/Sequencer/Private/Tools/NavigationTool/DragDropOps/Handlers/NavigationToolActorDropHandler.cpp
deleted file mode 100644
diff --git a/Engine/Source/Editor/Sequencer/Private/Tools/NavigationTool/DragDropOps/Handlers/NavigationToolActorDropHandler.h b/Engine/Source/Editor/Sequencer/Private/Tools/NavigationTool/DragDropOps/Handlers/NavigationToolActorDropHandler.h
deleted file mode 100644
diff --git a/Engine/Source/Editor/Sequencer/Private/Tools/NavigationTool/DragDropOps/Handlers/NavigationToolSequenceDropHandler.cpp b/Engine/Source/Editor/Sequencer/Private/Tools/NavigationTool/DragDropOps/Handlers/NavigationToolSequenceDropHandler.cpp
new file mode 100644
--- /dev/null
+++ b/Engine/Source/Editor/Sequencer/Private/Tools/NavigationTool/DragDropOps/Handlers/NavigationToolSequenceDropHandler.cpp
@@ -0,0 +1,92 @@
+// Copyright Epic Games, Inc. All Rights Reserved.
+
+#include "NavigationToolSequenceDropHandler.h"
+#include "Framework/Application/SlateApplication.h"
+#include "Tools/NavigationTool/ItemActions/NavigationToolAddItem.h"
+#include "Tools/NavigationTool/Items/NavigationToolSequence.h"
+#include "Tools/NavigationTool/NavigationTool.h"
+#include "Widgets/Views/STableRow.h"
+
+namespace UE::Sequencer::NavigationTool
+{
+
+bool FNavigationToolSequenceDropHandler::IsDraggedItemSupported(const TSharedPtr<INavigationToolItem>& InDraggedItem) const
+{
+	return InDraggedItem->IsA<FNavigationToolSequence>();
+}
+
+TOptional<EItemDropZone> FNavigationToolSequenceDropHandler::CanDrop(EItemDropZone InDropZone, const FNavigationToolItemPtr& InTargetItem) const
+{
+	/*switch (ActionType)
+	{
+		case ENavigationToolDragDropActionType::Move:
+			// When moving, make sure the Destination is not one of the Items we're moving
+			if (!Items.Contains(InTargetItem))
+			{
+				return InDropZone;
+			}
+			break;
+
+		case ENavigationToolDragDropActionType::Copy:
+			return InDropZone;
+	}*/
+
+	return TOptional<EItemDropZone>();
+}
+
+bool FNavigationToolSequenceDropHandler::Drop(EItemDropZone InDropZone, const FNavigationToolItemPtr& InTargetItem)
+{
+	// Disabling drag and drop for now
+	switch (ActionType)
+	{
+	case ENavigationToolDragDropActionType::Move:
+		//MoveItems(InDropZone, InTargetItem);
+		break;
+
+	case ENavigationToolDragDropActionType::Copy:
+		//StaticCastSharedRef<FNavigationTool>(InTargetItem->GetOwnerTool())->DuplicateItems(Items, InTargetItem, InDropZone);
+		break;
+
+	default:
+		return false;
+	}
+
+	return false;
+}
+
+void FNavigationToolSequenceDropHandler::MoveItems(EItemDropZone InDropZone, const FNavigationToolItemPtr& InTargetItem)
+{
+	// Can't access this mouse event at this point, so use the Slate Application to check modifiers.
+	/*const FAttachmentTransformRules& TransformRules = InTargetItem->GetOwnerTool()->GetProvider().GetTransformRule(!FSlateApplication::Get().GetModifierKeys().IsShiftDown());
+	
+	FNavigationToolAddItemParams Params;
+	Params.RelativeItem             = InTargetItem;
+	Params.RelativeDropZone         = InDropZone;
+	Params.Flags                    = ENavigationToolAddItemFlags::Select | ENavigationToolAddItemFlags::Transact;
+	Params.AttachmentTransformRules = TransformRules;
+	Params.SelectionFlags           = ENavigationToolItemSelectionFlags::AppendToCurrentSelection;
+
+	const TSet<FNavigationToolItemPtr> DraggedItemSet(Items);
+
+	// Remove all Items whose Parent is in the Item Set 
+	Items.RemoveAll([&DraggedItemSet](const FNavigationToolItemPtr& InItem)
+		{
+			return !InItem.IsValid() || !InItem->GetParent().IsValid() || DraggedItemSet.Contains(InItem->GetParent());
+		});
+
+	// Reverse order for Onto since Item->AddChild(...) adds it to Index 0, so last item would be at the top, which is reversed
+	if (InDropZone == EItemDropZone::OntoItem)
+	{
+		Algo::Reverse(Items);
+	}
+
+	const TSharedRef<FNavigationTool> Tool = StaticCastSharedRef<FNavigationTool>(InTargetItem->GetOwnerTool());
+
+	for (const FNavigationToolItemPtr& Item : Items)
+	{
+		Params.Item = Item;
+		Tool->EnqueueItemAction<FNavigationToolAddItem>(Params);
+	}*/
+}
+
+}
diff --git a/Engine/Source/Editor/Sequencer/Private/Tools/NavigationTool/DragDropOps/Handlers/NavigationToolSequenceDropHandler.h b/Engine/Source/Editor/Sequencer/Private/Tools/NavigationTool/DragDropOps/Handlers/NavigationToolSequenceDropHandler.h
new file mode 100644
--- /dev/null
+++ b/Engine/Source/Editor/Sequencer/Private/Tools/NavigationTool/DragDropOps/Handlers/NavigationToolSequenceDropHandler.h
@@ -0,0 +1,26 @@
+﻿// Copyright Epic Games, Inc. All Rights Reserved.
+
+#pragma once
+
+#include "Tools/NavigationTool/DragDropOps/Handlers/NavigationToolItemDropHandler.h"
+
+namespace UE::Sequencer::NavigationTool
+{
+
+/** Class that handles Dropping Sequence Items into a Target Item */
+class FNavigationToolSequenceDropHandler : public FNavigationToolItemDropHandler
+{
+public:
+	UE_NAVIGATIONTOOL_INHERITS(FNavigationToolSequenceDropHandler, FNavigationToolItemDropHandler);
+
+private:
+	//~ Begin FNavigationToolItemDropHandler
+	virtual bool IsDraggedItemSupported(const TSharedPtr<INavigationToolItem>& InDraggedItem) const override;
+	virtual TOptional<EItemDropZone> CanDrop(EItemDropZone InDropZone, const FNavigationToolItemPtr& InTargetItem) const override;
+	virtual bool Drop(EItemDropZone InDropZone, const FNavigationToolItemPtr& InTargetItem) override;
+	//~ End FNavigationToolItemDropHandler
+
+	void MoveItems(EItemDropZone InDropZone, const FNavigationToolItemPtr& InTargetItem);
+};
+
+}
diff --git a/Engine/Source/Editor/Sequencer/Private/Tools/NavigationTool/DragDropOps/NavigationToolItemDragDropOp.cpp b/Engine/Source/Editor/Sequencer/Private/Tools/NavigationTool/DragDropOps/NavigationToolItemDragDropOp.cpp
--- a/Engine/Source/Editor/Sequencer/Private/Tools/NavigationTool/DragDropOps/NavigationToolItemDragDropOp.cpp
+++ b/Engine/Source/Editor/Sequencer/Private/Tools/NavigationTool/DragDropOps/NavigationToolItemDragDropOp.cpp
@@ -2,12 +2,12 @@
 
 #include "Tools/NavigationTool/DragDropOps/NavigationToolItemDragDropOp.h"
 #include "GameFramework/Actor.h"
+#include "LevelSequenceActor.h"
 #include "ScopedTransaction.h"
 #include "Styling/SlateIconFinder.h"
-#include "Tools/NavigationTool/DragDropOps/Handlers/NavigationToolActorDropHandler.h"
+#include "Tools/NavigationTool/DragDropOps/Handlers/NavigationToolSequenceDropHandler.h"
 #include "Tools/NavigationTool/DragDropOps/Handlers/NavigationToolItemDropHandler.h"
 #include "Tools/NavigationTool/Items/INavigationToolItem.h"
-#include "Tools/NavigationTool/Items/NavigationToolActor.h"
 #include "Tools/NavigationTool/NavigationTool.h"
 #include "Tools/NavigationTool/NavigationToolView.h"
 
@@ -53,24 +53,6 @@ TSharedRef<FNavigationToolItemDragDropOp> FNavigationToolItemDragDropOp::New(con
 	return DragDropOp;
 }
 
-void FNavigationToolItemDragDropOp::GetDragDropOpActors(TArray<TWeakObjectPtr<AActor>>& OutActors) const
-{
-	for (const FNavigationToolItemPtr& Item : Items)
-	{
-		if (!Item.IsValid())
-		{
-			continue;
-		}
-		if (FNavigationToolActor* const ActorItem = Item->CastTo<FNavigationToolActor>())
-		{
-			if (AActor* const UnderlyingActor = ActorItem->GetActor())
-			{
-				OutActors.Add(UnderlyingActor);
-			}
-		}
-	}
-}
-
 FReply FNavigationToolItemDragDropOp::Drop(EItemDropZone InDropZone, const FNavigationToolItemPtr& InTargetItem)
 {
 	FReply Reply = FReply::Unhandled();
@@ -130,14 +112,14 @@ void FNavigationToolItemDragDropOp::Init(const TArray<FNavigationToolItemPtr>& I
 	, const TSharedPtr<FNavigationToolView>& InToolView
 	, ENavigationToolDragDropActionType InActionType)
 {
-	Items            = InItems;
+	Items = InItems;
 	WeakToolView = InToolView;
-	ActionType       = InActionType;
-	MouseCursor      = EMouseCursor::GrabHandClosed;
+	ActionType = InActionType;
+	MouseCursor = EMouseCursor::GrabHandClosed;
 
 	CurrentIconBrush = Items.Num() == 1
 		? Items[0]->GetIconBrush()
-		: FSlateIconFinder::FindIconForClass(AActor::StaticClass()).GetIcon();
+		: FSlateIconFinder::FindIconForClass(ALevelSequenceActor::StaticClass()).GetIcon();
 
 	CurrentHoverText = FText::Format(LOCTEXT("HoverText", "{0} {1}")
 		, GetActionName(InActionType)
@@ -149,7 +131,7 @@ void FNavigationToolItemDragDropOp::Init(const TArray<FNavigationToolItemPtr>& I
 	Construct();
 
 	// Add Default Drop Handlers 
-	AddDropHandler<FNavigationToolActorDropHandler>();
+	AddDropHandler<FNavigationToolSequenceDropHandler>();
 
 	FNavigationToolExtender::OnItemDragDropOpInitialized().Broadcast(*this);
 }
diff --git a/Engine/Source/Editor/Sequencer/Private/Tools/NavigationTool/Filters/Filters/NavigationToolBuiltInFilter.cpp b/Engine/Source/Editor/Sequencer/Private/Tools/NavigationTool/Filters/Filters/NavigationToolBuiltInFilter.cpp
--- a/Engine/Source/Editor/Sequencer/Private/Tools/NavigationTool/Filters/Filters/NavigationToolBuiltInFilter.cpp
+++ b/Engine/Source/Editor/Sequencer/Private/Tools/NavigationTool/Filters/Filters/NavigationToolBuiltInFilter.cpp
@@ -47,16 +47,17 @@ FSlateIcon FNavigationToolBuiltInFilter::GetIcon() const
 	return FSlateIcon(FAppStyle::GetAppStyleSetName(), GetIconName());
 }
 
+TSharedPtr<FUICommandInfo> FNavigationToolBuiltInFilter::GetToggleCommand() const
+{
+	return FilterParams.GetToggleCommand();
+}
+
 bool FNavigationToolBuiltInFilter::PassesFilter(const FNavigationToolFilterType InItem) const
 {
 	auto ItemObjectPassesFilter = [this](const FNavigationToolItemPtr& InItem)
 		{
-			if (const FNavigationToolObject* ObjectItem = InItem->CastTo<FNavigationToolObject>())
-			{
-				const UObject* const Object = ObjectItem->GetObject();
-				return IsValid(Object) && FilterParams.IsValidObjectClass(Object->GetClass());
-			}
-			return false;
+			UObject* const ItemObject = InItem->GetItemObject();
+			return IsValid(ItemObject) && FilterParams.IsValidObjectClass(ItemObject->GetClass());
 		};
 
 	auto ItemTypePassesFilter = [this](const FNavigationToolItemPtr& InItem)
diff --git a/Engine/Source/Editor/Sequencer/Private/Tools/NavigationTool/Filters/Filters/NavigationToolBuiltInFilter.h b/Engine/Source/Editor/Sequencer/Private/Tools/NavigationTool/Filters/Filters/NavigationToolBuiltInFilter.h
--- a/Engine/Source/Editor/Sequencer/Private/Tools/NavigationTool/Filters/Filters/NavigationToolBuiltInFilter.h
+++ b/Engine/Source/Editor/Sequencer/Private/Tools/NavigationTool/Filters/Filters/NavigationToolBuiltInFilter.h
@@ -28,6 +28,8 @@ public:
 
 	virtual FSlateIcon GetIcon() const;
 
+	virtual TSharedPtr<FUICommandInfo> GetToggleCommand() const;
+
 	virtual bool IsActive() const;
 	virtual void SetActive(const bool bInActive);
 
diff --git a/Engine/Source/Editor/Sequencer/Private/Tools/NavigationTool/Filters/Filters/NavigationToolBuiltInFilterParams.cpp b/Engine/Source/Editor/Sequencer/Private/Tools/NavigationTool/Filters/Filters/NavigationToolBuiltInFilterParams.cpp
--- a/Engine/Source/Editor/Sequencer/Private/Tools/NavigationTool/Filters/Filters/NavigationToolBuiltInFilterParams.cpp
+++ b/Engine/Source/Editor/Sequencer/Private/Tools/NavigationTool/Filters/Filters/NavigationToolBuiltInFilterParams.cpp
@@ -4,6 +4,7 @@
 #include "LevelSequence.h"
 #include "LevelSequenceActor.h"
 #include "Styling/SlateIconFinder.h"
+#include "Tools/NavigationTool/Filters/NavigationToolFilterCommands.h"
 #include "Tools/NavigationTool/Filters/Filters/NavigationToolFilter_Text.h"
 #include "Tools/NavigationTool/Items/NavigationToolActor.h"
 #include "Tools/NavigationTool/Items/NavigationToolBinding.h"
@@ -19,6 +20,7 @@ FNavigationToolBuiltInFilterParams::FNavigationToolBuiltInFilterParams(const FNa
 	, const FSlateBrush* const InIconBrush
 	, const FText& InDisplayName
 	, const FText& InTooltipText
+	, const TSharedPtr<FUICommandInfo>& InToggleCommand
 	, const bool InEnabledByDefault
 	, const EClassFlags InRequiredClassFlags
 	, const EClassFlags InRestrictedClassFlags)
@@ -30,6 +32,7 @@ FNavigationToolBuiltInFilterParams::FNavigationToolBuiltInFilterParams(const FNa
 	, TooltipText(InTooltipText)
 	, bUseOverrideIcon(false)
 	, bEnabledByDefault(InEnabledByDefault)
+	, ToggleCommand(InToggleCommand)
 	, IconBrush(InIconBrush)
 	, RequiredClassFlags(InRequiredClassFlags)
 	, RestrictedClassFlags(InRestrictedClassFlags)
@@ -110,6 +113,11 @@ bool FNavigationToolBuiltInFilterParams::IsEnabledByDefault() const
 	return bEnabledByDefault;
 }
 
+TSharedPtr<FUICommandInfo> FNavigationToolBuiltInFilterParams::GetToggleCommand() const
+{
+	return ToggleCommand;
+}
+
 void FNavigationToolBuiltInFilterParams::SetOverrideIconColor(const FSlateColor& InNewIconColor)
 {
 	OverrideIcon.TintColor = InNewIconColor;
@@ -150,7 +158,7 @@ void FNavigationToolBuiltInFilterParams::SetFilterText(const FText& InText)
 	FilterText = InText;
 }
 
-bool FNavigationToolBuiltInFilterParams::PassesFilterText(const UE::Sequencer::NavigationTool::FNavigationToolObject& InItem) const
+bool FNavigationToolBuiltInFilterParams::PassesFilterText(const UE::Sequencer::NavigationTool::FNavigationToolItemPtr& InItem) const
 {
 	// Since it needs to be a const function we have to create the TextFilter here instead of being a variable
 	if (FilterText.IsEmptyOrWhitespace())
@@ -173,7 +181,8 @@ FNavigationToolBuiltInFilterParams FNavigationToolBuiltInFilterParams::CreateSeq
 		, ENavigationToolFilterMode::MatchesType
 		, FSlateIconFinder::FindIconForClass(ALevelSequenceActor::StaticClass()).GetIcon()
 		, LOCTEXT("LevelSequenceFilterDisplayName", "Level Sequence")
-		, LOCTEXT("LevelSequenceFilterTooltip", "Level Sequence"));
+		, LOCTEXT("LevelSequenceFilterTooltip", "Level Sequence")
+		, FNavigationToolFilterCommands::Get().ToggleFilter_Sequence);
 }
 
 FNavigationToolBuiltInFilterParams FNavigationToolBuiltInFilterParams::CreateTrackFilter()
@@ -184,7 +193,8 @@ FNavigationToolBuiltInFilterParams FNavigationToolBuiltInFilterParams::CreateTra
 		, ENavigationToolFilterMode::MatchesType
 		, FSlateIconFinder::FindIconForClass(UMovieSceneTrack::StaticClass()).GetIcon()
 		, LOCTEXT("TrackFilterDisplayName", "Track")
-		, LOCTEXT("TrackFilterTooltip", "Track"));
+		, LOCTEXT("TrackFilterTooltip", "Track")
+		, FNavigationToolFilterCommands::Get().ToggleFilter_Track);
 }
 
 FNavigationToolBuiltInFilterParams FNavigationToolBuiltInFilterParams::CreateBindingFilter()
@@ -198,7 +208,7 @@ FNavigationToolBuiltInFilterParams FNavigationToolBuiltInFilterParams::CreateBin
 		, nullptr
 		, LOCTEXT("BindingFilterDisplayName", "Binding")
 		, LOCTEXT("BindingFilterTooltip", "Binding")
-		, false);
+		, FNavigationToolFilterCommands::Get().ToggleFilter_Binding);
 }
 
 #undef LOCTEXT_NAMESPACE
diff --git a/Engine/Source/Editor/Sequencer/Private/Tools/NavigationTool/Filters/Filters/NavigationToolFilter_Dirty.cpp b/Engine/Source/Editor/Sequencer/Private/Tools/NavigationTool/Filters/Filters/NavigationToolFilter_Dirty.cpp
new file mode 100644
--- /dev/null
+++ b/Engine/Source/Editor/Sequencer/Private/Tools/NavigationTool/Filters/Filters/NavigationToolFilter_Dirty.cpp
@@ -0,0 +1,82 @@
+﻿// Copyright Epic Games, Inc. All Rights Reserved.
+
+#include "Tools/NavigationTool/Filters/Filters/NavigationToolFilter_Dirty.h"
+#include "Sequencer.h"
+#include "Tools/NavigationTool/Filters/NavigationToolFilterCommands.h"
+#include "Tools/NavigationTool/Items/NavigationToolBinding.h"
+#include "Tools/NavigationTool/Items/NavigationToolSequence.h"
+#include "Tools/NavigationTool/Items/NavigationToolTrack.h"
+
+#define LOCTEXT_NAMESPACE "NavigationToolFilter_Dirty"
+
+namespace UE::Sequencer::NavigationTool
+{
+
+using namespace UE::Sequencer;
+
+FNavigationToolFilter_Dirty::FNavigationToolFilter_Dirty(INavigationToolFilterBar& InFilterInterface, TSharedPtr<FFilterCategory> InCategory)
+	: FNavigationToolFilter(InFilterInterface, MoveTemp(InCategory))
+{
+}
+
+FText FNavigationToolFilter_Dirty::GetDefaultToolTipText() const
+{
+	return LOCTEXT("Tooltip", "Show only items that are Dirty");
+}
+
+TSharedPtr<FUICommandInfo> FNavigationToolFilter_Dirty::GetToggleCommand() const
+{
+	return FNavigationToolFilterCommands::Get().ToggleFilter_Dirty;
+}
+
+FText FNavigationToolFilter_Dirty::GetDisplayName() const
+{
+	return LOCTEXT("DisplayName", "Dirty");
+}
+
+FSlateIcon FNavigationToolFilter_Dirty::GetIcon() const
+{
+	return FSlateIcon(FAppStyle::GetAppStyleSetName(), TEXT("Icons.DirtyBadge"));
+}
+
+FString FNavigationToolFilter_Dirty::GetName() const
+{
+	return StaticName();
+}
+
+bool FNavigationToolFilter_Dirty::PassesFilter(FNavigationToolFilterType InItem) const
+{
+	if (const FNavigationToolSequence* const SequenceItem = InItem->CastTo<FNavigationToolSequence>())
+	{
+		const UMovieSceneSequence* const Sequence = SequenceItem->GetSequence();
+		return IsValid(Sequence) && IsObjectPackageDirty(Sequence);
+	}
+
+	if (const FNavigationToolTrack* const TrackItem = InItem->CastTo<FNavigationToolTrack>())
+	{
+		const UMovieSceneTrack* const Track = TrackItem->GetTrack();
+		return IsValid(Track) && IsObjectPackageDirty(Track);
+	}
+
+	if (const FNavigationToolBinding* const BindingItem = InItem->CastTo<FNavigationToolBinding>())
+	{
+		const UObject* const BoundObject = BindingItem->GetCachedBoundObject();
+		return IsValid(BoundObject) && IsObjectPackageDirty(BoundObject);
+	}
+
+	return false;
+}
+
+bool FNavigationToolFilter_Dirty::IsObjectPackageDirty(const UObject* const InObject)
+{
+	if (IsValid(InObject))
+	{
+		const UPackage* const Package = InObject->GetPackage();
+		return IsValid(Package) && Package->IsDirty();
+	}
+	return false;
+}
+
+}
+
+#undef LOCTEXT_NAMESPACE
diff --git a/Engine/Source/Editor/Sequencer/Private/Tools/NavigationTool/Filters/Filters/NavigationToolFilter_Marks.cpp b/Engine/Source/Editor/Sequencer/Private/Tools/NavigationTool/Filters/Filters/NavigationToolFilter_Marks.cpp
--- a/Engine/Source/Editor/Sequencer/Private/Tools/NavigationTool/Filters/Filters/NavigationToolFilter_Marks.cpp
+++ b/Engine/Source/Editor/Sequencer/Private/Tools/NavigationTool/Filters/Filters/NavigationToolFilter_Marks.cpp
@@ -51,13 +51,13 @@ bool FNavigationToolFilter_Marks::PassesFilter(FNavigationToolFilterType InItem)
 	}
 
 	UMovieSceneSequence* const Sequence = SequenceItem->GetSequence();
-	if (IsValid(Sequence))
+	if (!IsValid(Sequence))
 	{
 		return false;
 	}
 
 	UMovieScene* const MovieScene = Sequence->GetMovieScene();
-	if (IsValid(MovieScene))
+	if (!IsValid(MovieScene))
 	{
 		return false;
 	}
diff --git a/Engine/Source/Editor/Sequencer/Private/Tools/NavigationTool/Filters/Filters/NavigationToolFilter_Playhead.cpp b/Engine/Source/Editor/Sequencer/Private/Tools/NavigationTool/Filters/Filters/NavigationToolFilter_Playhead.cpp
new file mode 100644
--- /dev/null
+++ b/Engine/Source/Editor/Sequencer/Private/Tools/NavigationTool/Filters/Filters/NavigationToolFilter_Playhead.cpp
@@ -0,0 +1,165 @@
+﻿// Copyright Epic Games, Inc. All Rights Reserved.
+
+#include "Tools/NavigationTool/Filters/Filters/NavigationToolFilter_Playhead.h"
+#include "Engine/World.h"
+#include "ISequencer.h"
+#include "TimerManager.h"
+#include "Tools/NavigationTool/Filters/NavigationToolFilterCommands.h"
+#include "Tools/NavigationTool/Items/NavigationToolItemUtils.h"
+
+#define LOCTEXT_NAMESPACE "NavigationToolFilter_Playhead"
+
+namespace UE::Sequencer::NavigationTool
+{
+
+using namespace UE::Sequencer;
+
+FNavigationToolFilter_Playhead::FNavigationToolFilter_Playhead(INavigationToolFilterBar& InFilterInterface, TSharedPtr<FFilterCategory> InCategory)
+	: FNavigationToolFilter(InFilterInterface, MoveTemp(InCategory))
+{
+}
+
+FNavigationToolFilter_Playhead::~FNavigationToolFilter_Playhead()
+{
+	UnbindEvents();
+}
+
+void FNavigationToolFilter_Playhead::BindEvents()
+{
+	ISequencer& Sequencer = FilterInterface.GetSequencer();
+	WeakSequencer = Sequencer.AsWeak();
+
+	Sequencer.OnPlayEvent().AddSP(this, &FNavigationToolFilter_Playhead::OnPlayEvent);
+	Sequencer.OnStopEvent().AddSP(this, &FNavigationToolFilter_Playhead::OnStopEvent);
+	Sequencer.OnBeginScrubbingEvent().AddSP(this, &FNavigationToolFilter_Playhead::OnBeginScrubbingEvent);
+	Sequencer.OnEndScrubbingEvent().AddSP(this, &FNavigationToolFilter_Playhead::OnEndScrubbingEvent);
+}
+
+void FNavigationToolFilter_Playhead::UnbindEvents()
+{
+	const TSharedPtr<ISequencer> Sequencer = WeakSequencer.Pin();
+	if (!Sequencer.IsValid())
+	{
+		return;
+	}
+
+	Sequencer->OnPlayEvent().RemoveAll(this);
+	Sequencer->OnStopEvent().RemoveAll(this);
+	Sequencer->OnBeginScrubbingEvent().RemoveAll(this);
+	Sequencer->OnEndScrubbingEvent().RemoveAll(this);
+}
+
+FText FNavigationToolFilter_Playhead::GetDefaultToolTipText() const
+{
+	return LOCTEXT("NavigationToolFilter_PlayheadToolTip", "Show only items whose range contains the current playhead location");
+}
+
+TSharedPtr<FUICommandInfo> FNavigationToolFilter_Playhead::GetToggleCommand() const
+{
+	return FNavigationToolFilterCommands::Get().ToggleFilter_Playhead;
+}
+
+void FNavigationToolFilter_Playhead::ActiveStateChanged(const bool bInActive)
+{
+	FNavigationToolFilter::ActiveStateChanged(bInActive);
+
+	if (bInActive)
+	{
+		BindEvents();
+	}
+	else
+	{
+		UnbindEvents();
+	}
+}
+
+FText FNavigationToolFilter_Playhead::GetDisplayName() const
+{
+	return LOCTEXT("NavigationToolFilter_Playhead", "Playhead");
+}
+
+FSlateIcon FNavigationToolFilter_Playhead::GetIcon() const
+{
+	return FSlateIcon(FAppStyle::GetAppStyleSetName(), TEXT("PlayWorld.ContinueExecution.Small"));
+}
+
+FString FNavigationToolFilter_Playhead::GetName() const
+{
+	return StaticName();
+}
+
+bool FNavigationToolFilter_Playhead::PassesFilter(FNavigationToolFilterType InItem) const
+{
+	const IPlayheadExtension* const PlayheadItem = InItem->CastTo<IPlayheadExtension>();
+	if (!PlayheadItem)
+	{
+		return false;
+	}
+
+	return PlayheadItem->ContainsPlayhead() == EItemContainsPlayhead::ContainsPlayhead;
+}
+
+void FNavigationToolFilter_Playhead::OnPlayEvent()
+{
+	StartRefreshTimer();
+}
+
+void FNavigationToolFilter_Playhead::OnStopEvent()
+{
+	StopRefreshTimer();
+}
+
+void FNavigationToolFilter_Playhead::OnBeginScrubbingEvent()
+{
+	StartRefreshTimer();
+}
+
+void FNavigationToolFilter_Playhead::OnEndScrubbingEvent()
+{
+	StopRefreshTimer();
+}
+
+FTimerManager* FNavigationToolFilter_Playhead::GetTimerManager() const
+{
+	UObject* const ContextObject = FilterInterface.GetSequencer().GetPlaybackContext();
+	if (!IsValid(ContextObject))
+	{
+		return nullptr;
+	}
+
+	UWorld* const World = ContextObject->GetWorld();
+	if (!IsValid(World))
+	{
+		return nullptr;
+	}
+
+	return &World->GetTimerManager();
+}
+
+void FNavigationToolFilter_Playhead::StartRefreshTimer()
+{
+	if (TimerHandle.IsValid())
+	{
+		return;
+	}
+
+	if (FTimerManager* const TimerManager = GetTimerManager())
+	{
+		TimerManager->SetTimer(TimerHandle, [this]()
+			{
+				FilterInterface.RequestFilterUpdate();
+			}, 0.1f, true);
+	}
+}
+
+void FNavigationToolFilter_Playhead::StopRefreshTimer()
+{
+	if (FTimerManager* const TimerManager = GetTimerManager())
+	{
+		TimerManager->ClearTimer(TimerHandle);
+	}
+}
+
+}
+
+#undef LOCTEXT_NAMESPACE
diff --git a/Engine/Source/Editor/Sequencer/Private/Tools/NavigationTool/Filters/NavigationToolFilterBar.cpp b/Engine/Source/Editor/Sequencer/Private/Tools/NavigationTool/Filters/NavigationToolFilterBar.cpp
--- a/Engine/Source/Editor/Sequencer/Private/Tools/NavigationTool/Filters/NavigationToolFilterBar.cpp
+++ b/Engine/Source/Editor/Sequencer/Private/Tools/NavigationTool/Filters/NavigationToolFilterBar.cpp
@@ -8,7 +8,9 @@
 #include "Sequencer.h"
 #include "Stats/StatsMisc.h"
 #include "Tools/NavigationTool/Filters/Filters/NavigationToolFilter_CustomText.h"
+#include "Tools/NavigationTool/Filters/Filters/NavigationToolFilter_Dirty.h"
 #include "Tools/NavigationTool/Filters/Filters/NavigationToolFilter_Marks.h"
+#include "Tools/NavigationTool/Filters/Filters/NavigationToolFilter_Playhead.h"
 #include "Tools/NavigationTool/Filters/Filters/NavigationToolFilter_Unbound.h"
 #include "Tools/NavigationTool/Filters/Filters/NavigationToolFilterBase.h"
 #include "Tools/NavigationTool/Filters/Filters/NavigationToolFilterCollection.h"
@@ -103,6 +105,8 @@ void FNavigationToolFilterBar::CreateDefaultFilters()
 	// Add misc category filters
 	AddFilterIfSupported(CommonFilters, MakeShared<FNavigationToolFilter_Unbound>(*this, MiscCategory));
 	AddFilterIfSupported(CommonFilters, MakeShared<FNavigationToolFilter_Marks>(*this, MiscCategory));
+	AddFilterIfSupported(CommonFilters, MakeShared<FNavigationToolFilter_Playhead>(*this, MiscCategory));
+	AddFilterIfSupported(CommonFilters, MakeShared<FNavigationToolFilter_Dirty>(*this, MiscCategory));
 
 	// Add global user-defined filters
 	for (TObjectIterator<UNavigationToolFilterExtension> ExtensionIt(RF_NoFlags); ExtensionIt; ++ExtensionIt)
@@ -171,11 +175,11 @@ void FNavigationToolFilterBar::BindCommands(const TSharedPtr<FUICommandList>& In
 
 	CommandList->MapAction(FilterCommands.ActivateAllFilters,
 		FExecuteAction::CreateSP(this, &FNavigationToolFilterBar::ActivateAllEnabledFilters, true, TArray<FString>()),
-		FCanExecuteAction::CreateSP(this, &FNavigationToolFilterBar::AreAllEnabledFiltersActive, true, TArray<FString>()));
+		FCanExecuteAction::CreateSP(this, &FNavigationToolFilterBar::AreAllEnabledFiltersActive, false, TArray<FString>()));
 
 	CommandList->MapAction(FilterCommands.DeactivateAllFilters,
 		FExecuteAction::CreateSP(this, &FNavigationToolFilterBar::ActivateAllEnabledFilters, false, TArray<FString>()),
-		FCanExecuteAction::CreateSP(this, &FNavigationToolFilterBar::AreAllEnabledFiltersActive, false, TArray<FString>()));
+		FCanExecuteAction::CreateSP(this, &FNavigationToolFilterBar::AreAllEnabledFiltersActive, true, TArray<FString>()));
 
 	// Bind all filter actions
 	const TSharedPtr<ISequencer> Sequencer = Tool.GetSequencer();
@@ -535,7 +539,8 @@ bool FNavigationToolFilterBar::SetFilterActive(const TSharedRef<FNavigationToolF
 		InFilter->SetActive(bNewActive);
 		InFilter->ActiveStateChanged(bNewActive);
 
-		const ENavigationToolFilterChange FilterChangeType = bNewActive ? ENavigationToolFilterChange::Activate : ENavigationToolFilterChange::Deactivate;
+		const ENavigationToolFilterChange FilterChangeType = bNewActive
+			? ENavigationToolFilterChange::Activate : ENavigationToolFilterChange::Deactivate;
 		BroadcastFiltersChanged(InFilter, FilterChangeType);
 
 		if (bInRequestFilterUpdate)
@@ -818,7 +823,8 @@ bool FNavigationToolFilterBar::SetFilterEnabled(const TSharedRef<FNavigationTool
 	{
 		ToolSettings->SaveConfig();
 
-		const ENavigationToolFilterChange FilterChangeType = bInEnabled ? ENavigationToolFilterChange::Enable : ENavigationToolFilterChange::Disable;
+		const ENavigationToolFilterChange FilterChangeType = bInEnabled
+			? ENavigationToolFilterChange::Enable : ENavigationToolFilterChange::Disable;
 		BroadcastFiltersChanged(InFilter, FilterChangeType);
 
 		if (!bInEnabled && IsFilterActive(InFilter))
@@ -862,7 +868,8 @@ void FNavigationToolFilterBar::EnableFilters(const bool bInEnable
 				const FString FilterName = InFilter->GetDisplayName().ToString();
 				if (Config.SetFilterEnabled(FilterName, bInEnable))
 				{
-					const ENavigationToolFilterChange FilterChangeType = bInEnable ? ENavigationToolFilterChange::Enable : ENavigationToolFilterChange::Disable;
+					const ENavigationToolFilterChange FilterChangeType = bInEnable
+						? ENavigationToolFilterChange::Enable : ENavigationToolFilterChange::Disable;
 					BroadcastFiltersChanged(InFilter, FilterChangeType);
 
 					if (!bInEnable && IsFilterActive(InFilter))
@@ -1112,7 +1119,8 @@ void FNavigationToolFilterBar::EnableCustomTextFilters(const bool bInEnable, con
 					CustomTextFilter->ActiveStateChanged(false);
 				}
 
-				const ENavigationToolFilterChange FilterChangeType = bInEnable ? ENavigationToolFilterChange::Enable : ENavigationToolFilterChange::Disable;
+				const ENavigationToolFilterChange FilterChangeType = bInEnable
+					? ENavigationToolFilterChange::Enable : ENavigationToolFilterChange::Disable;
 				BroadcastFiltersChanged(CustomTextFilter, FilterChangeType);
 			}
 		}
@@ -1752,7 +1760,7 @@ void FNavigationToolFilterBar::BroadcastStateChanged()
 {
 	const UNavigationToolSettings* const Settings = GetDefault<UNavigationToolSettings>();
 	const EFilterBarLayout Layout = IsValid(Settings) ? Settings->GetFilterBarLayout() : EFilterBarLayout::Horizontal;
-	StateChangedEvent.Broadcast(ShouldShowFilterBarWidget(), Layout);
+	StateChangedEvent.Broadcast(IsFilterBarVisible(), Layout);
 }
 
 void FNavigationToolFilterBar::BroadcastFiltersChanged(const TSharedRef<FNavigationToolFilter>& InFilter, const ENavigationToolFilterChange InChangeType)
diff --git a/Engine/Source/Editor/Sequencer/Private/Tools/NavigationTool/Filters/NavigationToolFilterCommands.cpp b/Engine/Source/Editor/Sequencer/Private/Tools/NavigationTool/Filters/NavigationToolFilterCommands.cpp
--- a/Engine/Source/Editor/Sequencer/Private/Tools/NavigationTool/Filters/NavigationToolFilterCommands.cpp
+++ b/Engine/Source/Editor/Sequencer/Private/Tools/NavigationTool/Filters/NavigationToolFilterCommands.cpp
@@ -78,25 +78,43 @@ void FNavigationToolFilterCommands::RegisterCommands()
 	UI_COMMAND(ToggleFilter_Sequence
 		, "Toggle Sequence Filter"
 		, "Toggle the filter for Sequence items"
-		, EUserInterfaceActionType::Button
+		, EUserInterfaceActionType::ToggleButton
 		, FInputChord());
 
 	UI_COMMAND(ToggleFilter_Track
 		, "Toggle Track Filter"
 		, "Toggle the filter for Track items"
-		, EUserInterfaceActionType::Button
+		, EUserInterfaceActionType::ToggleButton
+		, FInputChord());
+
+	UI_COMMAND(ToggleFilter_Binding
+		, "Toggle Binding Filter"
+		, "Toggle the filter for Binding items"
+		, EUserInterfaceActionType::ToggleButton
 		, FInputChord());
 
 	UI_COMMAND(ToggleFilter_Unbound
 		, "Toggle Unbound Filter"
 		, "Toggle the filter for displaying items that contain unbound tracks"
-		, EUserInterfaceActionType::Button
+		, EUserInterfaceActionType::ToggleButton
 		, FInputChord());
 
 	UI_COMMAND(ToggleFilter_Marks
 		, "Toggle Marks Filter"
 		, "Toggle the filter for displaying items that contain marked frames"
-		, EUserInterfaceActionType::Button
+		, EUserInterfaceActionType::ToggleButton
+		, FInputChord());
+
+	UI_COMMAND(ToggleFilter_Playhead
+		, "Toggle Playhead Filter"
+		, "Toggle the filter for displaying items whose range contains the current playhead location"
+		, EUserInterfaceActionType::ToggleButton
+		, FInputChord());
+
+	UI_COMMAND(ToggleFilter_Dirty
+		, "Toggle Dirty Filter"
+		, "Toggle the filter for displaying items whose package is marked dirty"
+		, EUserInterfaceActionType::ToggleButton
 		, FInputChord());
 }
 
diff --git a/Engine/Source/Editor/Sequencer/Private/Tools/NavigationTool/Filters/NavigationToolFilterCommands.h b/Engine/Source/Editor/Sequencer/Private/Tools/NavigationTool/Filters/NavigationToolFilterCommands.h
--- a/Engine/Source/Editor/Sequencer/Private/Tools/NavigationTool/Filters/NavigationToolFilterCommands.h
+++ b/Engine/Source/Editor/Sequencer/Private/Tools/NavigationTool/Filters/NavigationToolFilterCommands.h
@@ -38,9 +38,12 @@ public:
 
 	TSharedPtr<FUICommandInfo> ToggleFilter_Sequence;
 	TSharedPtr<FUICommandInfo> ToggleFilter_Track;
+	TSharedPtr<FUICommandInfo> ToggleFilter_Binding;
 
 	TSharedPtr<FUICommandInfo> ToggleFilter_Unbound;
 	TSharedPtr<FUICommandInfo> ToggleFilter_Marks;
+	TSharedPtr<FUICommandInfo> ToggleFilter_Playhead;
+	TSharedPtr<FUICommandInfo> ToggleFilter_Dirty;
 };
 
 #undef LOCTEXT_NAMESPACE
diff --git a/Engine/Source/Editor/Sequencer/Private/Tools/NavigationTool/ItemProxies/NavigationToolComponentProxy.cpp b/Engine/Source/Editor/Sequencer/Private/Tools/NavigationTool/ItemProxies/NavigationToolComponentProxy.cpp
--- a/Engine/Source/Editor/Sequencer/Private/Tools/NavigationTool/ItemProxies/NavigationToolComponentProxy.cpp
+++ b/Engine/Source/Editor/Sequencer/Private/Tools/NavigationTool/ItemProxies/NavigationToolComponentProxy.cpp
@@ -1,11 +1,7 @@
 // Copyright Epic Games, Inc. All Rights Reserved.
 
 #include "Tools/NavigationTool/ItemProxies/NavigationToolComponentProxy.h"
-#include "Components/ActorComponent.h"
-#include "GameFramework/Actor.h"
 #include "Styling/SlateIconFinder.h"
-#include "Tools/NavigationTool/Items/NavigationToolActor.h"
-#include "Tools/NavigationTool/Items/NavigationToolObject.h"
 #include "Tools/NavigationTool/NavigationTool.h"
 
 #define LOCTEXT_NAMESPACE "NavigationToolComponentProxy"
@@ -42,26 +38,6 @@ ENavigationToolItemViewMode FNavigationToolComponentProxy::GetSupportedViewModes
 
 void FNavigationToolComponentProxy::GetProxiedItems(const TSharedRef<INavigationToolItem>& InParent, TArray<FNavigationToolItemPtr>& OutChildren, bool bInRecursive)
 {
-	if (const FNavigationToolActor* const ActorItem = InParent->CastTo<FNavigationToolActor>())
-	{
-		const AActor* const Actor = ActorItem->GetActor();
-		if (!IsValid(Actor))
-		{
-			return;
-		}
-
-		const TSharedRef<FNavigationToolComponentProxy> SharedThisRef = SharedThis(this);
-
-		for (UActorComponent* const ActorComponent : Actor->GetComponents())
-		{
-			const FNavigationToolItemPtr NewItem = Tool.FindOrAdd<FNavigationToolObject>(InParent->GetProvider(), SharedThisRef, ActorComponent, SharedThisRef);
-			OutChildren.Add(NewItem);
-			if (bInRecursive)
-			{
-				NewItem->FindChildren(OutChildren, bInRecursive);
-			}
-		}
-	}
 }
 
 }
diff --git a/Engine/Source/Editor/Sequencer/Private/Tools/NavigationTool/Items/NavigationToolActor.cpp b/Engine/Source/Editor/Sequencer/Private/Tools/NavigationTool/Items/NavigationToolActor.cpp
--- a/Engine/Source/Editor/Sequencer/Private/Tools/NavigationTool/Items/NavigationToolActor.cpp
+++ b/Engine/Source/Editor/Sequencer/Private/Tools/NavigationTool/Items/NavigationToolActor.cpp
@@ -5,15 +5,14 @@
 #include "GameFramework/Actor.h"
 #include "ISequencer.h"
 #include "MovieScene.h"
+#include "MovieSceneCommonHelpers.h"
 #include "MovieSceneSequence.h"
 #include "ScopedTransaction.h"
 #include "Tools/NavigationTool/Items/NavigationToolComponent.h"
 #include "Tools/NavigationTool/Items/NavigationToolSequence.h"
-#include "Tools/NavigationTool/ItemProxies/NavigationToolComponentProxy.h"
 #include "Tools/NavigationTool/NavigationTool.h"
 #include "Tools/NavigationTool/NavigationToolSettings.h"
 #include "Tools/NavigationTool/Utils/NavigationToolMiscUtils.h"
-#include "Tools/NavigationTool/Utils/NavigationToolMovieSceneUtils.h"
 
 #define LOCTEXT_NAMESPACE "NavigationToolActor"
 
@@ -68,7 +67,7 @@ void FNavigationToolActor::FindChildren(TArray<FNavigationToolItemPtr>& OutChild
 	for (const FMovieSceneBinding& CurrentBinding : Bindings)
 	{
 		const FGuid& CurrentBindingGuid = CurrentBinding.GetObjectGuid();
-		const UClass* const CurrentObjectClass = GetBindingObjectClass(*MovieScene, CurrentBindingGuid);
+		const UClass* const CurrentObjectClass = MovieSceneHelpers::GetBoundObjectClass(Sequence, CurrentBindingGuid);
 		if (CurrentObjectClass && CurrentObjectClass->IsChildOf<UActorComponent>())
 		{
 			const TArrayView<TWeakObjectPtr<>> WeakBoundObjects = ResolveBoundObjects(*Sequencer, Sequence, CurrentBindingGuid);
@@ -78,7 +77,7 @@ void FNavigationToolActor::FindChildren(TArray<FNavigationToolItemPtr>& OutChild
 			if (IsValid(CurrentBoundActorComponent) && CurrentBoundActorComponent->GetOwner() == ThisActor)
 			{
 				const FNavigationToolItemPtr NewItem = Tool.FindOrAdd<FNavigationToolComponent>(GetProvider()
-					, SharedThisRef, WeakParentSequenceItem.Pin(), CurrentBinding);
+					, SharedThisRef, ParentSequenceItem, CurrentBinding);
 				OutChildren.Add(NewItem);
 				if (bInRecursive)
 				{
@@ -89,16 +88,6 @@ void FNavigationToolActor::FindChildren(TArray<FNavigationToolItemPtr>& OutChild
 	}
 }
 
-void FNavigationToolActor::GetItemProxies(TArray<TSharedPtr<FNavigationToolItemProxy>>& OutItemProxies)
-{
-	Super::GetItemProxies(OutItemProxies);
-
-	if (const TSharedPtr<FNavigationToolItemProxy> ComponentItemProxy = Tool.GetOrCreateItemProxy<FNavigationToolComponentProxy>(SharedThis(this)))
-	{
-		OutItemProxies.Add(ComponentItemProxy);
-	}
-}
-
 bool FNavigationToolActor::IsAllowedInTool() const
 {
 	FNavigationTool& ToolPrivate = static_cast<FNavigationTool&>(Tool);
diff --git a/Engine/Source/Editor/Sequencer/Private/Tools/NavigationTool/Items/NavigationToolBinding.cpp b/Engine/Source/Editor/Sequencer/Private/Tools/NavigationTool/Items/NavigationToolBinding.cpp
--- a/Engine/Source/Editor/Sequencer/Private/Tools/NavigationTool/Items/NavigationToolBinding.cpp
+++ b/Engine/Source/Editor/Sequencer/Private/Tools/NavigationTool/Items/NavigationToolBinding.cpp
@@ -4,6 +4,9 @@
 #include "ISequencer.h"
 #include "MovieScene.h"
 #include "ScopedTransaction.h"
+#include "MVVM/ObjectBindingModelStorageExtension.h"
+#include "MVVM/TrackModelStorageExtension.h"
+#include "MVVM/ViewModels/SequencerEditorViewModel.h"
 #include "Sections/MovieSceneSubSection.h"
 #include "Styling/SlateIconFinder.h"
 #include "Tools/NavigationTool/INavigationTool.h"
@@ -28,12 +31,24 @@ FNavigationToolBinding::FNavigationToolBinding(INavigationTool& InTool
 	, WeakParentSequenceItem(InParentSequenceItem)
 	, Binding(InBinding)
 {
+	Icon = FSlateIcon(FAppStyle::GetAppStyleSetName(), TEXT("Sequencer.InvalidSpawnableIcon"));
+
 	CacheBoundObject();
 }
 
 bool FNavigationToolBinding::IsItemValid() const
 {
-	return FNavigationToolItem::IsItemValid() && Binding.GetObjectGuid().IsValid();
+	return Binding.GetObjectGuid().IsValid();
+}
+
+UObject* FNavigationToolBinding::GetItemObject() const
+{
+	return GetCachedBoundObject();
+}
+
+bool FNavigationToolBinding::IsAllowedInTool() const
+{
+	return IsItemValid();
 }
 
 void FNavigationToolBinding::FindChildren(TArray<FNavigationToolItemPtr>& OutChildren, const bool bInRecursive)
@@ -77,12 +92,42 @@ FText FNavigationToolBinding::GetIconTooltipText() const
 
 FLinearColor FNavigationToolBinding::GetItemBrushColor() const
 {
-	const UObject* const UnderlyingObject = GetCachedBoundObject();
-	if (IsValid(UnderlyingObject))
+	const TSharedPtr<ISequencer> Sequencer = Tool.GetSequencer();
+	if (!Sequencer.IsValid())
+	{
+		return FLinearColor::Red;
+	}
+
+	UMovieSceneSequence* const Sequence = GetSequence();
+	if (!IsValid(Sequence))
 	{
-		return FNavigationToolItem::GetItemBrushColor();
+		return FLinearColor::Red;
 	}
-	return FStyleColors::Error.GetSpecifiedColor();
+
+	const TArrayView<TWeakObjectPtr<>> BoundObjects = ResolveBoundObjects(*Sequencer, Sequence, Binding.GetObjectGuid());
+	if (BoundObjects.Num() > 0)
+	{
+		int32 NumValidObjects = 0;
+		for (const TWeakObjectPtr<>& BoundObject : BoundObjects)
+		{
+			if (BoundObject.IsValid())
+			{
+				++NumValidObjects;
+			}
+		}
+
+		if (NumValidObjects == BoundObjects.Num())
+		{
+			return FNavigationToolItem::GetItemBrushColor();
+		}
+
+		if (NumValidObjects > 0)
+		{
+			return FLinearColor::Yellow;
+		}
+	}
+
+	return FLinearColor::Red;
 }
 
 bool FNavigationToolBinding::IsSelected(const FNavigationToolScopedSelection& InSelection) const
@@ -92,7 +137,11 @@ bool FNavigationToolBinding::IsSelected(const FNavigationToolScopedSelection& In
 
 void FNavigationToolBinding::Select(FNavigationToolScopedSelection& InSelection) const
 {
-	InSelection.Select(Binding.GetObjectGuid());
+	const FGuid ObjectGuid = Binding.GetObjectGuid();
+	if (ObjectGuid.IsValid())
+	{
+		InSelection.Select(ObjectGuid);
+	}
 }
 	
 void FNavigationToolBinding::OnSelect()
@@ -102,27 +151,18 @@ void FNavigationToolBinding::OnSelect()
 
 void FNavigationToolBinding::OnDoubleClick()
 {
-	Super::OnDoubleClick();
-
-	const TSharedPtr<ISequencer> Sequencer = Tool.GetSequencer();
-	if (!Sequencer.IsValid())
-	{
-		return;
-	}
-
-	UObject* const BoundObject = GetCachedBoundObject();
-	if (!IsValid(BoundObject))
+	UMovieSceneSequence* const Sequence = GetSequence();
+	if (IsValid(Sequence))
 	{
-		return;
+		FocusSequence(Tool, *Sequence, *this);
 	}
-
-	const FGuid ObjectId = Sequencer->FindObjectId(*BoundObject, Sequencer->GetFocusedTemplateID());
-	Sequencer->SelectObject(ObjectId);
 }
 
 bool FNavigationToolBinding::CanDelete() const
 {
-	return WeakBoundObject.IsValid();
+	// Disable for now
+	return false;
+	//return WeakBoundObject.IsValid();
 }
 
 bool FNavigationToolBinding::Delete()
@@ -139,12 +179,14 @@ FNavigationToolItemId FNavigationToolBinding::CalculateItemId() const
 		, ParentSequenceItem->GetSequence()
 		, ParentSequenceItem->GetSubSection()
 		, ParentSequenceItem->GetSubSectionIndex()
-		, Binding.GetName());
+		, Binding.GetObjectGuid().ToString());
 }
 
 bool FNavigationToolBinding::CanRename() const
 {
-	return WeakBoundObject.IsValid();
+	// Disable for now
+	return false;
+	//return WeakBoundObject.IsValid();
 }
 
 bool FNavigationToolBinding::Rename(const FString& InName)
@@ -199,15 +241,54 @@ void FNavigationToolBinding::SetIsLocked(const bool bInIsLocked)
 	}
 }*/
 
+EItemContainsPlayhead FNavigationToolBinding::ContainsPlayhead() const
+{
+	const TSharedPtr<ISequencer> Sequencer = Tool.GetSequencer();
+	if (!Sequencer.IsValid())
+	{
+		return EItemContainsPlayhead::None;
+	}
+
+	FTrackModelStorageExtension* const TrackModelStorage = Sequencer->GetViewModel()->CastDynamic<FTrackModelStorageExtension>();
+	if (!TrackModelStorage)
+	{
+		return EItemContainsPlayhead::None;
+	}
+
+	const FQualifiedFrameTime PlayheadTime = Sequencer->GetLocalTime();
+
+	for (UMovieSceneTrack* const Track : Binding.GetTracks())
+	{
+		const TSharedPtr<FTrackModel> TrackModel = TrackModelStorage->FindModelForTrack(Track);
+		if (!TrackModel.IsValid())
+		{
+			continue;
+		}
+
+		for (const TViewModelPtr<ILayerBarExtension>& TrackAreaModel : TrackModel->GetTrackAreaModelListAs<ILayerBarExtension>())
+		{
+			if (const TViewModelPtr<ILayerBarExtension> BarModel = TrackAreaModel.ImplicitCast())
+			{
+				if (BarModel->GetLayerBarRange().Contains(PlayheadTime.Time.FrameNumber))
+				{
+					return EItemContainsPlayhead::ContainsPlayhead;
+				}
+			}
+		}
+	}
+
+	return EItemContainsPlayhead::None;
+}
+
 const UClass* FNavigationToolBinding::GetBoundObjectClass() const
 {
-	const TSharedPtr<FNavigationToolSequence> ParentSequenceItem = WeakParentSequenceItem.Pin();
-	if (!ParentSequenceItem.IsValid())
+	const TSharedPtr<ISequencer> Sequencer = Tool.GetSequencer();
+	if (!Sequencer.IsValid())
 	{
 		return nullptr;
 	}
 
-	UMovieSceneSequence* const Sequence = ParentSequenceItem->GetSequence();
+	UMovieSceneSequence* const Sequence = GetSequence();
 	if (!IsValid(Sequence))
 	{
 		return nullptr;
@@ -219,7 +300,7 @@ const UClass* FNavigationToolBinding::GetBoundObjectClass() const
 		return nullptr;
 	}
 
-	return GetBindingObjectClass(*MovieScene, Binding.GetObjectGuid());
+	return MovieSceneHelpers::GetBoundObjectClass(Sequence, Binding.GetObjectGuid());
 }
 
 UObject* FNavigationToolBinding::GetCachedBoundObject() const
@@ -229,24 +310,25 @@ UObject* FNavigationToolBinding::GetCachedBoundObject() const
 
 UObject* FNavigationToolBinding::CacheBoundObject()
 {
-	if (const TSharedPtr<ISequencer> Sequencer = Tool.GetSequencer())
+	const TSharedPtr<ISequencer> Sequencer = Tool.GetSequencer();
+	if (!Sequencer.IsValid())
 	{
-		UMovieSceneSequence* const Sequence = WeakParentSequenceItem.Pin()->GetSequence();
-		const TArrayView<TWeakObjectPtr<>> WeakBoundObjects = ResolveBoundObjects(*Sequencer, Sequence, Binding.GetObjectGuid());
-		WeakBoundObject = WeakBoundObjects.IsEmpty() ? nullptr : WeakBoundObjects[0];
-
-		if (WeakBoundObject.IsValid())
-		{
-			Icon = FSlateIconFinder::FindIconForClass(WeakBoundObject->GetClass());
-		}
-		else
-		{
-			Icon = FSlateIcon(FAppStyle::GetAppStyleSetName(), TEXT("Sequencer.InvalidSpawnableIcon"));
-		}
+		return nullptr;
+	}
 
-		return WeakBoundObject.Get();
+	UMovieSceneSequence* const Sequence = GetSequence();
+	if (!IsValid(Sequence))
+	{
+		return nullptr;
 	}
-	return nullptr;
+
+	const TArrayView<TWeakObjectPtr<>> WeakBoundObjects = ResolveBoundObjects(*Sequencer, Sequence, Binding.GetObjectGuid());
+	WeakBoundObject = WeakBoundObjects.IsEmpty() ? nullptr : WeakBoundObjects[0];
+
+	const UClass* const BoundObjectClass = MovieSceneHelpers::GetBoundObjectClass(Sequence, Binding.GetObjectGuid());
+	Icon = FSlateIconFinder::FindIconForClass(BoundObjectClass);
+
+	return WeakBoundObject.Get();
 }
 
 const FMovieSceneBinding& FNavigationToolBinding::GetBinding() const
@@ -254,6 +336,50 @@ const FMovieSceneBinding& FNavigationToolBinding::GetBinding() const
 	return Binding;
 }
 
+UMovieSceneSequence* FNavigationToolBinding::GetSequence() const
+{
+	if (const TSharedPtr<FNavigationToolSequence> ParentSequenceItem = WeakParentSequenceItem.Pin())
+	{
+		return ParentSequenceItem->GetSequence();
+	}
+	return nullptr;
+}
+
+TViewModelPtr<FObjectBindingModel> FNavigationToolBinding::GetViewModel() const
+{
+	const FGuid ObjectGuid = Binding.GetObjectGuid();
+	if (!ObjectGuid.IsValid())
+	{
+		return nullptr;
+	}
+
+	const TSharedPtr<ISequencer> Sequencer = Tool.GetSequencer();
+	if (!Sequencer.IsValid())
+	{
+		return nullptr;
+	}
+
+	const TSharedPtr<FSequencerEditorViewModel> ViewModel = Sequencer->GetViewModel();
+	if (!ViewModel.IsValid())
+	{
+		return nullptr;
+	}
+
+	const FViewModelPtr RootViewModel = ViewModel->GetRootModel();
+	if (!RootViewModel.IsValid())
+	{
+		return nullptr;
+	}
+
+	const FObjectBindingModelStorageExtension* const StorageExtension = RootViewModel->CastDynamic<FObjectBindingModelStorageExtension>();
+	if (!StorageExtension)
+	{
+		return nullptr;
+	}
+
+	return StorageExtension->FindModelForObjectBinding(ObjectGuid);
+}
+
 }
 
 #undef LOCTEXT_NAMESPACE
diff --git a/Engine/Source/Editor/Sequencer/Private/Tools/NavigationTool/Items/NavigationToolItem.cpp b/Engine/Source/Editor/Sequencer/Private/Tools/NavigationTool/Items/NavigationToolItem.cpp
--- a/Engine/Source/Editor/Sequencer/Private/Tools/NavigationTool/Items/NavigationToolItem.cpp
+++ b/Engine/Source/Editor/Sequencer/Private/Tools/NavigationTool/Items/NavigationToolItem.cpp
@@ -7,13 +7,11 @@
 #include "Styling/StyleColors.h"
 #include "Textures/SlateIcon.h"
 #include "Tools/NavigationTool/DragDropOps/NavigationToolItemDragDropOp.h"
-#include "Tools/NavigationTool/Items/NavigationToolActor.h"
 #include "Tools/NavigationTool/Items/NavigationToolItemUtils.h"
 #include "Tools/NavigationTool/ItemActions/NavigationToolAddItem.h"
 #include "Tools/NavigationTool/ItemActions/NavigationToolRemoveItem.h"
 #include "Tools/NavigationTool/NavigationTool.h"
 #include "Tools/NavigationTool/NavigationToolView.h"
-#include "Tools/NavigationTool/NavigationToolScopedSelection.h"
 #include "Tools/NavigationTool/Providers/NavigationToolProvider.h"
 #include "Tools/NavigationTool/Widgets/Columns/SNavigationToolLabelItem.h"
 #include "Widgets/Views/STableRow.h"
@@ -23,32 +21,6 @@
 namespace UE::Sequencer::NavigationTool
 {
 
-EDetachmentRule GetDetachmentRuleFromAttachmentRule(EAttachmentRule AttachmentRule)
-{
-	switch (AttachmentRule)
-	{
-		case EAttachmentRule::KeepRelative:
-			return EDetachmentRule::KeepRelative;
-
-		default:
-			return EDetachmentRule::KeepWorld;
-	}
-};
-
-FDetachmentTransformRules GetDetachmentRulesFromAttachmentRules(const TOptional<FAttachmentTransformRules>& AttachmentTransformRules)
-{
-	if (AttachmentTransformRules.IsSet())
-	{
-		return FDetachmentTransformRules(
-			GetDetachmentRuleFromAttachmentRule(AttachmentTransformRules->LocationRule),
-			GetDetachmentRuleFromAttachmentRule(AttachmentTransformRules->RotationRule),
-			GetDetachmentRuleFromAttachmentRule(AttachmentTransformRules->ScaleRule),
-			true
-		);
-	}
-	return FDetachmentTransformRules::KeepWorldTransform;
-};
-
 FNavigationToolItem::FNavigationToolItem(INavigationTool& InTool, const FNavigationToolItemPtr& InParentItem)
 	: Tool(InTool)
 	, WeakParent(InParentItem)
@@ -244,7 +216,7 @@ ENavigationToolItemViewMode FNavigationToolItem::GetSupportedViewModes(const INa
 
 FNavigationToolItemId FNavigationToolItem::GetItemId() const
 {
-	if (ItemId.IsValid())
+	if (ItemId.IsValidId())
 	{
 		return ItemId;
 	}
@@ -306,8 +278,14 @@ TOptional<EItemDropZone> FNavigationToolItem::CanAcceptDrop(const FDragDropEvent
 		const TOptional<EItemDropZone> DropZone = ItemDragDropOp->CanDrop(InDropZone, SharedThis(this));
 		if (DropZone.IsSet())
 		{
+			if (GetItemId() != FNavigationToolItemId::RootId)
+			{
+				ItemDragDropOp->CurrentIconBrush = GetIconBrush();
+			}
 			return *DropZone;
 		}
+
+		ItemDragDropOp->CurrentIconBrush = FAppStyle::GetBrush(TEXT("Graph.ConnectorFeedback.Error"));
 	}
 
 	const TSharedRef<FNavigationToolItem> SharedThisRef = SharedThis(this);
diff --git a/Engine/Source/Editor/Sequencer/Private/Tools/NavigationTool/Items/NavigationToolItemId.cpp b/Engine/Source/Editor/Sequencer/Private/Tools/NavigationTool/Items/NavigationToolItemId.cpp
--- a/Engine/Source/Editor/Sequencer/Private/Tools/NavigationTool/Items/NavigationToolItemId.cpp
+++ b/Engine/Source/Editor/Sequencer/Private/Tools/NavigationTool/Items/NavigationToolItemId.cpp
@@ -1,10 +1,12 @@
 ﻿// Copyright Epic Games, Inc. All Rights Reserved.
 
 #include "Tools/NavigationTool/Items/NavigationToolItemId.h"
-#include "MovieSceneSequence.h"
+#include "MovieScene.h"
 #include "MovieSceneSection.h"
+#include "MovieSceneSequence.h"
+#include "MVVM/Extensions/IObjectBindingExtension.h"
+#include "Sequencer.h"
 #include "Tools/NavigationTool/Items/INavigationToolItem.h"
-#include "Tools/NavigationTool/Items/NavigationToolItem.h"
 
 namespace UE::Sequencer::NavigationTool
 {
@@ -17,12 +19,6 @@ FNavigationToolItemId::FNavigationToolItemId(const FString& InUniqueId)
 	CalculateTypeHash();
 }
 
-FNavigationToolItemId::FNavigationToolItemId(const UObject* InObject)
-{
-	Id = GetObjectPath(InObject);
-	CalculateTypeHash();
-}
-
 FNavigationToolItemId::FNavigationToolItemId(const UObject* InObject
 	, const FNavigationToolItemPtr& InReferencingItem
 	, const FString& InReferencingId)
@@ -64,7 +60,7 @@ FNavigationToolItemId::FNavigationToolItemId(const FNavigationToolItemPtr& InPar
 
 	AddSeparatedSegment(Id, GetObjectPath(InSequence));
 
-	if (InSection)
+	if (IsValid(InSection))
 	{
 		AddSeparatedSegment(Id, GetObjectPath(InSection));
 		AddSeparatedSegment(Id, FString::FromInt(InSectionIndex));
@@ -74,6 +70,48 @@ FNavigationToolItemId::FNavigationToolItemId(const FNavigationToolItemPtr& InPar
 	{
 		AddSeparatedSegment(Id, InReferenceId);
 	}
+
+	CalculateTypeHash();
+}
+
+FNavigationToolItemId::FNavigationToolItemId(const FViewModelPtr& InViewModel)
+{
+	check(InViewModel.IsValid());
+
+	const TViewModelPtr<FSectionModel> SectionModel = InViewModel->FindAncestorOfType<FSectionModel>(true);
+	const TViewModelPtr<FSequenceModel> SequenceModel = InViewModel->FindAncestorOfType<FSequenceModel>();
+
+	UMovieSceneSequence* const Sequence = SequenceModel->GetSequence();
+	UMovieSceneSection* const Section = SectionModel.IsValid() ? SectionModel->GetSection() : nullptr;
+
+	int32 SectionIndex = 0;
+	FString ReferenceId;
+
+	const TArray<UMovieSceneSection*> AllSections = SequenceModel->GetMovieScene()->GetAllSections();
+	for (int32 Index = 0; Index < AllSections.Num(); ++Index)
+	{
+		if (AllSections[Index] == Section)
+		{
+			SectionIndex = Index;
+			break;
+		}
+	}
+
+	if (const TViewModelPtr<ITrackExtension> TrackExtension = InViewModel.ImplicitCast())
+	{
+		const UMovieSceneTrack* const Track = TrackExtension->GetTrack();
+		if (IsValid(Track))
+		{
+			ReferenceId = GetObjectPath(Track);
+		}
+	}
+	else if (const TViewModelPtr<IObjectBindingExtension> ObjectBindingExtension = InViewModel.ImplicitCast())
+	{
+		ReferenceId = ObjectBindingExtension->GetObjectGuid().ToString();
+	}
+
+	FNavigationToolItemId(nullptr, Sequence, Section, SectionIndex, ReferenceId);
+	CalculateTypeHash();
 }
 
 FNavigationToolItemId::FNavigationToolItemId(const FNavigationToolItemId& Other)
@@ -105,7 +143,7 @@ bool FNavigationToolItemId::operator==(const FNavigationToolItemId& Other) const
 	return Id == Other.Id && CachedHash == Other.CachedHash;
 }
 
-bool FNavigationToolItemId::IsValid() const
+bool FNavigationToolItemId::IsValidId() const
 {
 	return bHasCachedHash;
 }
diff --git a/Engine/Source/Editor/Sequencer/Private/Tools/NavigationTool/Items/NavigationToolItemUtils.cpp b/Engine/Source/Editor/Sequencer/Private/Tools/NavigationTool/Items/NavigationToolItemUtils.cpp
--- a/Engine/Source/Editor/Sequencer/Private/Tools/NavigationTool/Items/NavigationToolItemUtils.cpp
+++ b/Engine/Source/Editor/Sequencer/Private/Tools/NavigationTool/Items/NavigationToolItemUtils.cpp
@@ -2,40 +2,40 @@
 
 #include "Tools/NavigationTool/Items/NavigationToolItemUtils.h"
 #include "Containers/Array.h"
-#include "ISourceControlModule.h"
-#include "ISourceControlProvider.h"
 #include "MovieSceneSequence.h"
 #include "ScopedTransaction.h"
 #include "Sections/MovieSceneCinematicShotSection.h"
 #include "Sections/MovieSceneSubSection.h"
-#include "Sequencer.h"
-#include "SequencerUtilities.h"
+#include "ISequencer.h"
 #include "Tools/NavigationTool/Items/INavigationToolItem.h"
 #include "Tools/NavigationTool/Items/NavigationToolSequence.h"
 #include "Tools/NavigationTool/Items/NavigationToolTreeRoot.h"
 #include "Tools/NavigationTool/NavigationTool.h"
 #include "Tools/NavigationTool/NavigationToolView.h"
-#include "Tools/NavigationTool/Items/NavigationToolSubTrack.h"
+#include "Tools/NavigationTool/Utils/NavigationToolMovieSceneUtils.h"
 #include "TrackEditors/SubTrackEditorBase.h"
-#include "Tracks/MovieSceneSubTrack.h"
+
+#define LOCTEXT_NAMESPACE "NavigationToolItemUtils"
 
 namespace UE::Sequencer::NavigationTool
 {
 
-bool CompareToolItemOrder(const FNavigationToolItemPtr& A, const FNavigationToolItemPtr& B)
+bool CompareToolItemOrder(const FNavigationToolItemPtr& InA, const FNavigationToolItemPtr& InB)
 {
-	if (!A.IsValid())
+	if (!InA.IsValid())
 	{
 		return false;
 	}
-	if (!B.IsValid())
+
+	if (!InB.IsValid())
 	{
 		return false;
 	}
-	if (const FNavigationToolItemPtr LowestCommonAncestor = FNavigationTool::FindLowestCommonAncestor({A, B}))
+
+	if (const FNavigationToolItemPtr LowestCommonAncestor = FNavigationTool::FindLowestCommonAncestor({ InA, InB }))
 	{
-		const TArray<FNavigationToolItemPtr> PathToA = LowestCommonAncestor->FindPath({A});
-		const TArray<FNavigationToolItemPtr> PathToB = LowestCommonAncestor->FindPath({B});
+		const TArray<FNavigationToolItemPtr> PathToA = LowestCommonAncestor->FindPath({ InA });
+		const TArray<FNavigationToolItemPtr> PathToB = LowestCommonAncestor->FindPath({ InB });
 
 		int32 Index = 0;
 	
@@ -57,8 +57,10 @@ bool CompareToolItemOrder(const FNavigationToolItemPtr& A, const FNavigationTool
 			PathBIndex = LowestCommonAncestor->GetChildIndex(PathToB[Index]);
 			++Index;
 		}
+
 		return PathAIndex < PathBIndex;
 	}
+
 	return false;
 }
 
@@ -106,87 +108,76 @@ UMovieSceneMetaData* GetSequenceItemMetaData(const TSharedPtr<INavigationToolIte
 	return nullptr;
 }
 
-FText SequenceDisplayNameWithModifyAndInactiveStatus(const TSharedPtr<FNavigationToolSequence>& InSequenceItem
-	, const bool bInRemoveParentName
-	, const bool bInIsActive)
+void RemoveSequenceDisplayNameParentPrefix(FText& InOutDisplayName
+	, const TSharedRef<FNavigationToolSequence>& InSequenceItem)
 {
-	if (!InSequenceItem.IsValid())
-	{
-		return FText::GetEmpty();
-	}
-
 	UMovieSceneSequence* const Sequence = InSequenceItem->GetSequence();
 	if (!IsValid(Sequence))
 	{
-		return FText::GetEmpty();
+		return;
 	}
 
-	UMovieSceneCinematicShotSection* const ShotSection = Cast<UMovieSceneCinematicShotSection>(InSequenceItem->GetSubSection());
-
-	const FText OriginalSequenceName = IsValid(ShotSection)
-		? FText::FromString(ShotSection->GetShotDisplayName()) : Sequence->GetDisplayName();
-
-	FText ParentSequenceName = FText::GetEmpty();
-	FText NewSequenceName = OriginalSequenceName;
+	const FNavigationToolItemPtr Parent = InSequenceItem->GetParent();
+	if (!Parent.IsValid() || Parent->IsA<FNavigationToolTreeRoot>())
+	{
+		return;
+	}
 
-	// Remove parent name from child name if it contains it
-	if (bInRemoveParentName)
+	// Go up the hierarchy to find the parent sequence
+	const FNavigationToolSequence* const ParentSequenceItem = InSequenceItem->FindAncestorOfType<FNavigationToolSequence>();
+	if (!ParentSequenceItem)
 	{
-		if (const FNavigationToolItemPtr Parent = InSequenceItem->GetParent())
-		{
-			if (!Parent->IsA<FNavigationToolTreeRoot>())
-			{
-				// Go up the hierarchy to find the parent sequence
-				if (const FNavigationToolSequence* const ParentSequenceItem = InSequenceItem->FindAncestorOfType<FNavigationToolSequence>())
-				{
-					UMovieSceneSequence* const ParentSequence = ParentSequenceItem->GetSequence();
-					if (IsValid(ParentSequence))
-					{
-						ParentSequenceName = ParentSequence->GetDisplayName();
-					}
-				}
-
-				// Remove the parent name from children
-				FString NewSequenceNameStr = ParentSequenceName.IsEmpty() ? NewSequenceName.ToString() : ParentSequenceName.ToString();
-				const FString ParentSequenceNameStr = ParentSequenceName.ToString();
-
-				if (NewSequenceNameStr.StartsWith(ParentSequenceNameStr) && NewSequenceNameStr.Len() > ParentSequenceNameStr.Len())
-				{
-					NewSequenceNameStr.RemoveFromStart(ParentSequenceNameStr);
-					ParentSequenceName = FText::FromString(NewSequenceNameStr);
-				}
-			}
-		}
+		return;
 	}
 
-	FText OutText = NewSequenceName;
+	UMovieSceneSequence* const ParentSequence = ParentSequenceItem->GetSequence();
+	if (!IsValid(ParentSequence))
+	{
+		return;
+	}
 
-	const bool bIsDirty = Sequence->GetPackage()->IsDirty();
+	const UMovieSceneCinematicShotSection* const ShotSection = Cast<UMovieSceneCinematicShotSection>(InSequenceItem->GetSubSection());
 
-	if (bInIsActive)
+	FText NewSequenceName = IsValid(ShotSection)
+		? FText::FromString(ShotSection->GetShotDisplayName()) : Sequence->GetDisplayName();
+	if (NewSequenceName.IsEmpty())
 	{
-		if (bIsDirty)
-		{
-			OutText = FText::Format(NSLOCTEXT("Sequencer", "ActiveDirtySequenceFormat", "{0}*"), NewSequenceName);
-		}
+		return;
 	}
-	else
+
+	// Remove the parent prefix from child name
+	static FString SeparatorString = TEXT("_");
+	const FString ParentSequenceNameStr = ParentSequence->GetDisplayName().ToString();
+
+	FString Prefix;
+	FString BaseName;
+	if (ParentSequenceNameStr.Split(SeparatorString, &Prefix, &BaseName))
 	{
-		const FText InactiveSequenceLabel = NSLOCTEXT("Sequencer", "InactiveSequenceLabel", "Inactive");
+		Prefix += SeparatorString;
 
-		if (bIsDirty)
-		{
-			OutText = FText::Format(NSLOCTEXT("Sequencer", "InactiveDirtySequenceFormat", "{0}* [{1}]")
-				, NewSequenceName, InactiveSequenceLabel);
-		}
-		else
+		FString NewSequenceNameStr = NewSequenceName.ToString();
+
+		if (NewSequenceNameStr.StartsWith(Prefix)
+			&& NewSequenceNameStr.Len() > Prefix.Len())
 		{
-			OutText = FText::Format(NSLOCTEXT("Sequencer", "InactiveSequenceFormat", "{0} [{1}]")
-				, NewSequenceName, InactiveSequenceLabel);
+			NewSequenceNameStr.RemoveFromStart(Prefix);
+			NewSequenceName = FText::FromString(NewSequenceNameStr);
 		}
 	}
 
-	return OutText;
+	InOutDisplayName = NewSequenceName;
+}
+
+void AppendSequenceDisplayNameDirtyStatus(FText& InOutDisplayName
+	, const UMovieSceneSequence& InSequence)
+{
+	UPackage* const Package = InSequence.GetPackage();
+	if (IsValid(Package) && Package->IsDirty())
+	{
+		InOutDisplayName = FText::Format(LOCTEXT("DirtySymbol", "{0}*"), InOutDisplayName);
+	}
 }
 
 }
+
+#undef LOCTEXT_NAMESPACE
diff --git a/Engine/Source/Editor/Sequencer/Private/Tools/NavigationTool/Items/NavigationToolObject.cpp b/Engine/Source/Editor/Sequencer/Private/Tools/NavigationTool/Items/NavigationToolObject.cpp
deleted file mode 100644
diff --git a/Engine/Source/Editor/Sequencer/Private/Tools/NavigationTool/Items/NavigationToolObjectReference.cpp b/Engine/Source/Editor/Sequencer/Private/Tools/NavigationTool/Items/NavigationToolObjectReference.cpp
deleted file mode 100644
diff --git a/Engine/Source/Editor/Sequencer/Private/Tools/NavigationTool/Items/NavigationToolSequence.cpp b/Engine/Source/Editor/Sequencer/Private/Tools/NavigationTool/Items/NavigationToolSequence.cpp
--- a/Engine/Source/Editor/Sequencer/Private/Tools/NavigationTool/Items/NavigationToolSequence.cpp
+++ b/Engine/Source/Editor/Sequencer/Private/Tools/NavigationTool/Items/NavigationToolSequence.cpp
@@ -1,10 +1,12 @@
 // Copyright Epic Games, Inc. All Rights Reserved.
 
 #include "Tools/NavigationTool/Items/NavigationToolSequence.h"
-#include "ISourceControlState.h"
 #include "LevelSequenceActor.h"
 #include "GameFramework/Actor.h"
+#include "Misc/QualifiedFrameTime.h"
 #include "MovieScene.h"
+#include "MVVM/SectionModelStorageExtension.h"
+#include "MVVM/ViewModels/SectionModel.h"
 #include "ScopedTransaction.h"
 #include "Sections/MovieSceneCinematicShotSection.h"
 #include "Sections/MovieSceneSubSection.h"
@@ -30,24 +32,28 @@
 	, const FNavigationToolItemPtr& InParentItem
 	, UMovieSceneSequence* const InSequence
 	, UMovieSceneSubSection* const InSubSection
-    , const int32 InSubSectionIndex)
-	: Super(InTool, InParentItem, InSequence)
+	, const int32 InSubSectionIndex)
+	: Super(InTool, InParentItem)
+	, IRevisionControlExtension(InSequence)
 	, WeakSubSection(InSubSection)
 	, SubSectionIndex(InSubSectionIndex)
+	, WeakSequence(InSequence)
 {
 }
 
-bool FNavigationToolSequence::IsSelected(const FNavigationToolScopedSelection& InSelection) const
+bool FNavigationToolSequence::IsItemValid() const
 {
-	return WeakSubSection.IsValid() && InSelection.IsSelected(WeakSubSection.Get());
+	return WeakSequence.IsValid();
 }
 
-void FNavigationToolSequence::Select(FNavigationToolScopedSelection& InSelection) const
+UObject* FNavigationToolSequence::GetItemObject() const
 {
-	if (WeakSubSection.IsValid())
-	{
-		InSelection.Select(WeakSubSection.Get());
-	}
+	return GetSequence();
 }
 
+bool FNavigationToolSequence::IsAllowedInTool() const
+{
+	return WeakSequence.IsValid();
+}
+
 void FNavigationToolSequence::FindChildren(TArray<FNavigationToolItemPtr>& OutChildren, const bool bInRecursive)
@@ -100,7 +106,7 @@
 	for (const FMovieSceneBinding& Binding : Bindings)
 	{
 		const FGuid& BindingObjectGuid = Binding.GetObjectGuid();
-		const UClass* const BoundObjectClass = GetBindingObjectClass(*MovieScene, BindingObjectGuid);
+		const UClass* const BoundObjectClass = MovieSceneHelpers::GetBoundObjectClass(Sequence, BindingObjectGuid);
 		if (BoundObjectClass && BoundObjectClass->IsChildOf<AActor>())
 		{
 			const FNavigationToolItemPtr NewItem = Tool.FindOrAdd<FNavigationToolActor>(GetProvider()
@@ -138,26 +144,45 @@
 	return ENavigationToolItemViewMode::ItemTree | ~ENavigationToolItemViewMode::HorizontalItemList;
 }
 
-bool FNavigationToolSequence::IsAllowedInTool() const
-{
-	return GetSequence() != nullptr;
-}
-
 FText FNavigationToolSequence::GetDisplayName() const
 {
-	UMovieSceneCinematicShotSection* const ShotSection = Cast<UMovieSceneCinematicShotSection>(WeakSubSection.Get());
+	FText NewDisplayName = FText::GetEmpty();
+
+	const UMovieSceneCinematicShotSection* const ShotSection = Cast<UMovieSceneCinematicShotSection>(WeakSubSection.Get());
 	if (IsValid(ShotSection))
 	{
-		return FText::FromString(ShotSection->GetShotDisplayName());
+		NewDisplayName = FText::FromString(ShotSection->GetShotDisplayName());
+	}
+	else
+	{
+		const UMovieSceneSequence* const Sequence = GetSequence();
+		if (IsValid(Sequence))
+		{
+			NewDisplayName = Sequence->GetDisplayName();
+		}
 	}
 
-	UMovieSceneSequence* const Sequence = GetSequence();
+	if (NewDisplayName.IsEmpty())
+	{
+		return NewDisplayName;
+	}
+
+	// Apply additional options to display name
+	const TSharedRef<FNavigationToolSequence> ThisSequenceRef = SharedThis(const_cast<FNavigationToolSequence*>(this));
+	const UNavigationToolSettings* const ToolSettings = GetDefault<UNavigationToolSettings>();
+
+	if (ToolSettings->ShouldUseShortNames())
+	{
+		RemoveSequenceDisplayNameParentPrefix(NewDisplayName, ThisSequenceRef);
+	}
+
+	UMovieSceneSequence* const Sequence = ThisSequenceRef->GetSequence();
 	if (IsValid(Sequence))
 	{
-		return Sequence->GetDisplayName();
+		AppendSequenceDisplayNameDirtyStatus(NewDisplayName, *Sequence);
 	}
 
-	return FText::GetEmpty();
+	return NewDisplayName;
 }
 
 FText FNavigationToolSequence::GetClassName() const
@@ -177,7 +202,20 @@
 
 FText FNavigationToolSequence::GetIconTooltipText() const
 {
-	return FText::GetEmpty();
+	return ALevelSequenceActor::StaticClass()->GetDisplayNameText();
 }
 
+bool FNavigationToolSequence::IsSelected(const FNavigationToolScopedSelection& InSelection) const
+{
+	return WeakSubSection.IsValid() && InSelection.IsSelected(WeakSubSection.Get());
+}
+
+void FNavigationToolSequence::Select(FNavigationToolScopedSelection& InSelection) const
+{
+	if (WeakSubSection.IsValid())
+	{
+		InSelection.Select(WeakSubSection.Get());
+	}
+}
+
 void FNavigationToolSequence::OnSelect()
@@ -187,14 +225,41 @@
 
 void FNavigationToolSequence::OnDoubleClick()
 {
-	// Allow selection of this item
-	Super::OnDoubleClick();
-
 	UMovieSceneSequence* const Sequence = GetSequence();
-	if (IsValid(Sequence))
+	if (!IsValid(Sequence))
 	{
-		FocusSequence(Tool, *Sequence);
+		return;
 	}
+
+	const FModifierKeysState ModifierKeys = FSlateApplication::Get().GetModifierKeys();
+
+	if (ModifierKeys.IsAltDown())
+	{
+		// Instead of focusing the sequence in the Sequencer, focus the parent sequence and select sequence
+		if (const FNavigationToolSequence* const ParentSequenceItem = FindAncestorOfType<FNavigationToolSequence>())
+		{
+			UMovieSceneSequence* const ParentSequence = ParentSequenceItem->GetSequence();
+			if (IsValid(ParentSequence))
+			{
+				FocusSequence(Tool, *ParentSequence, *this);
+			}
+		}
+	}
+
+	FocusSequence(Tool, *Sequence);
 }
 
+void FNavigationToolSequence::OnObjectsReplaced(const TMap<UObject*, UObject*>& InReplacementMap, bool bInRecursive)
+{
+	// Get the Object even if it's Pending Kill (most likely it is)
+	const UMovieSceneSequence* const ObjectPendingKill = WeakSequence.Get(true);
+	if (ObjectPendingKill && InReplacementMap.Contains(ObjectPendingKill))
+	{
+		WeakSequence = Cast<UMovieSceneSequence>(InReplacementMap[ObjectPendingKill]);
+	}
+
+	// This handles calling OnObjectsReplaced for every child item
+	Super::OnObjectsReplaced(InReplacementMap, bInRecursive);
+}
+
 FNavigationToolItemId FNavigationToolSequence::CalculateItemId() const
@@ -204,7 +269,9 @@
 
 bool FNavigationToolSequence::CanRename() const
 {
-	const UMovieSceneSequence* const Sequence = GetSequence();
+	// Disable for now
+	return false;
+	/*const UMovieSceneSequence* const Sequence = GetSequence();
 	if (!IsValid(Sequence))
 	{
 		return false;
@@ -216,7 +283,7 @@
 		return false;
 	}
 
-	return !MovieScene->IsReadOnly();
+	return !MovieScene->IsReadOnly();*/
 }
 
 bool FNavigationToolSequence::Rename(const FString& InName)
@@ -406,66 +473,75 @@
 
 EItemRevisionControlState FNavigationToolSequence::GetRevisionControlState() const
 {
-	if (!WeakSubSection.IsValid())
+	UMovieSceneSequence* const Sequence = GetSequence();
+	if (!IsValid(Sequence))
 	{
 		return EItemRevisionControlState::None;
 	}
 
-	const TSharedPtr<ISequencer> Sequencer = Tool.GetSequencer();
-	if (!Sequencer.IsValid())
-	{
-		return EItemRevisionControlState::None;
-	}
-
-	const FSourceControlStatePtr RevisionControlState = FindSourceControlState(WeakSubSection->GetPackage());
+	const FSourceControlStatePtr RevisionControlState = FindSourceControlState(Sequence->GetPackage());
 	if (!RevisionControlState.IsValid())
 	{
 		return EItemRevisionControlState::None;
 	}
 
-	return RevisionControlState->IsCheckedOut() ? EItemRevisionControlState::CheckedOut : EItemRevisionControlState::None;
+	return EItemRevisionControlState::SourceControlled;
 }
 
-void FNavigationToolSequence::SetIsRevisionControlled(const bool bInIsRevisionControlled)
-{
-	// Revision Control is only used to view state.
-}
-
 const FSlateBrush* FNavigationToolSequence::GetRevisionControlStatusIcon() const
 {
-	if (!WeakSubSection.IsValid())
+	UMovieSceneSequence* const Sequence = GetSequence();
+	if (IsValid(Sequence))
 	{
-		return nullptr;
+		return FindSourceControlStatusBrush(Sequence->GetPackage());
 	}
-
-	const TSharedPtr<ISequencer> Sequencer = Tool.GetSequencer();
-	if (!Sequencer.IsValid())
-	{
-		return nullptr;
-	}
-
-	return FindSourceControlStatusBrush(WeakSubSection->GetPackage());
+	return nullptr;
 }
 
 FText FNavigationToolSequence::GetRevisionControlStatusText() const
 {
-	if (!WeakSubSection.IsValid())
+	UMovieSceneSequence* const Sequence = GetSequence();
+	if (IsValid(Sequence))
 	{
-		return FText::GetEmpty();
+		return FindSourceControlStatusText(Sequence->GetPackage());
 	}
+	return FText::GetEmpty();
+}
 
+EItemContainsPlayhead FNavigationToolSequence::ContainsPlayhead() const
+{
 	const TSharedPtr<ISequencer> Sequencer = Tool.GetSequencer();
 	if (!Sequencer.IsValid())
 	{
-		return FText::GetEmpty();
+		return EItemContainsPlayhead::None;
 	}
 
-	return FindSourceControlStatusText(WeakSubSection->GetPackage());
+	const FQualifiedFrameTime PlayheadTime = Sequencer->GetLocalTime();
+
+	UMovieSceneSubSection* const SubSection = GetSubSection();
+	if (IsValid(SubSection))
+	{
+		return SubSection->GetTrueRange().Contains(PlayheadTime.Time.FrameNumber)
+			? EItemContainsPlayhead::ContainsPlayhead : EItemContainsPlayhead::None;
+	}
+
+	UMovieSceneSequence* const RootSequence = Sequencer->GetRootMovieSceneSequence();
+	if (IsValid(RootSequence))
+	{
+		UMovieScene* const RootSequenceMovieScene = RootSequence->GetMovieScene();
+		if (IsValid(RootSequenceMovieScene))
+		{
+			return RootSequenceMovieScene->GetPlaybackRange().Contains(PlayheadTime.Time.FrameNumber)
+				? EItemContainsPlayhead::ContainsPlayhead : EItemContainsPlayhead::None;
+		}
+	}
+
+	return EItemContainsPlayhead::None;
 }
 
 UMovieSceneSequence* FNavigationToolSequence::GetSequence() const
 {
-	return Cast<UMovieSceneSequence>(GetObject());
+	return WeakSequence.Get();
 }
 
 UMovieSceneSubSection* FNavigationToolSequence::GetSubSection() const
@@ -512,6 +588,40 @@
 	return MoveTemp(Bindings);
 }
 
+TViewModelPtr<FSectionModel> FNavigationToolSequence::GetViewModel() const
+{
+	if (!WeakSubSection.IsValid())
+	{
+		return nullptr;
+	}
+
+	const TSharedPtr<ISequencer> Sequencer = Tool.GetSequencer();
+	if (!Sequencer.IsValid())
+	{
+		return nullptr;
+	}
+
+	const TSharedPtr<FSequencerEditorViewModel> ViewModel = Sequencer->GetViewModel();
+	if (!ViewModel.IsValid())
+	{
+		return nullptr;
+	}
+
+	const FViewModelPtr RootViewModel = ViewModel->GetRootModel();
+	if (!RootViewModel.IsValid())
+	{
+		return nullptr;
+	}
+
+	const FSectionModelStorageExtension* const StorageExtension = RootViewModel->CastDynamic<FSectionModelStorageExtension>();
+	if (!StorageExtension)
+	{
+		return nullptr;
+	}
+
+	return StorageExtension->FindModelForSection(WeakSubSection.Get());
 }
 
+}
+
 #undef LOCTEXT_NAMESPACE
diff --git a/Engine/Source/Editor/Sequencer/Private/Tools/NavigationTool/Items/NavigationToolSubTrack.cpp b/Engine/Source/Editor/Sequencer/Private/Tools/NavigationTool/Items/NavigationToolSubTrack.cpp
--- a/Engine/Source/Editor/Sequencer/Private/Tools/NavigationTool/Items/NavigationToolSubTrack.cpp
+++ b/Engine/Source/Editor/Sequencer/Private/Tools/NavigationTool/Items/NavigationToolSubTrack.cpp
@@ -44,15 +44,15 @@ void FNavigationToolSubTrack::FindChildren(TArray<FNavigationToolItemPtr>& OutCh
 		const TSharedRef<FNavigationToolSubTrack> SharedThisRef = SharedThis(this);
 
 		const TArray<UMovieSceneSection*>& AllSections = SubTrack->GetAllSections();
-		for (int32 SectionIndex = 0; SectionIndex < AllSections.Num(); ++SectionIndex)
+		for (int32 Index = 0; Index < AllSections.Num(); ++Index)
 		{
-			UMovieSceneSubSection* const SubSection = Cast<UMovieSceneSubSection>(AllSections[SectionIndex]);
+			UMovieSceneSubSection* const SubSection = Cast<UMovieSceneSubSection>(AllSections[Index]);
 			if (IsValid(SubSection))
 			{
 				UMovieSceneSequence* const Sequence = SubSection->GetSequence();
 				if (ensure(IsValid(Sequence)))
 				{
-					const FNavigationToolItemPtr NewItem = Tool.FindOrAdd<FNavigationToolSequence>(GetProvider(), SharedThisRef, Sequence, SubSection, SectionIndex);
+					const FNavigationToolItemPtr NewItem = Tool.FindOrAdd<FNavigationToolSequence>(GetProvider(), SharedThisRef, Sequence, SubSection, Index);
 					OutChildren.Add(NewItem);
 					if (bInRecursive)
 					{
diff --git a/Engine/Source/Editor/Sequencer/Private/Tools/NavigationTool/Items/NavigationToolTrack.cpp b/Engine/Source/Editor/Sequencer/Private/Tools/NavigationTool/Items/NavigationToolTrack.cpp
--- a/Engine/Source/Editor/Sequencer/Private/Tools/NavigationTool/Items/NavigationToolTrack.cpp
+++ b/Engine/Source/Editor/Sequencer/Private/Tools/NavigationTool/Items/NavigationToolTrack.cpp
@@ -5,12 +5,15 @@
 #include "ISequencerTrackEditor.h"
 #include "MovieSceneSequence.h"
 #include "MovieSceneTrack.h"
+#include "MVVM/TrackModelStorageExtension.h"
+#include "MVVM/ViewModels/SequencerEditorViewModel.h"
 #include "ScopedTransaction.h"
 #include "Styling/SlateIconFinder.h"
 #include "Tools/NavigationTool/INavigationTool.h"
 #include "Tools/NavigationTool/Items/NavigationToolItemUtils.h"
 #include "Tools/NavigationTool/Items/NavigationToolSequence.h"
 #include "Tools/NavigationTool/NavigationToolScopedSelection.h"
+#include "Tools/NavigationTool/Utils/NavigationToolMiscUtils.h"
 #include "Widgets/Views/STableRow.h"
 
 #define LOCTEXT_NAMESPACE "NavigationToolTrack"
@@ -24,47 +27,64 @@ FNavigationToolTrack::FNavigationToolTrack(INavigationTool& InTool
 	, const TWeakObjectPtr<UMovieSceneSequence>& InSequence
 	, const TWeakObjectPtr<UMovieSceneSection>& InSection
 	, const int32 InSubSectionIndex)
-	: Super(InTool, InParentItem, InTrack)
+	: Super(InTool, InParentItem)
 	, WeakSequence(InSequence)
 	, WeakSection(InSection)
-	, SubSectionIndex(InSubSectionIndex)
+	, SectionIndex(InSubSectionIndex)
+	, WeakTrack(InTrack)
 {
 	OnTrackObjectChanged();
 }
 
-FNavigationToolItemId FNavigationToolTrack::CalculateItemId() const
+bool FNavigationToolTrack::IsItemValid() const
 {
-	return FNavigationToolItemId(GetParent()
-		, WeakSequence.Get()
-		, WeakSection.Get()
-		, SubSectionIndex
-		, FNavigationToolItemId::GetObjectPath(GetTrack()));
+	return WeakTrack.IsValid();
 }
 
-void FNavigationToolTrack::FindChildren(TArray<FNavigationToolItemPtr>& OutChildren, const bool bInRecursive)
+UObject* FNavigationToolTrack::GetItemObject() const
 {
-	Super::FindChildren(OutChildren, bInRecursive);
+	return GetTrack();
 }
 
-TOptional<EItemDropZone> FNavigationToolTrack::CanAcceptDrop(const FDragDropEvent& InDragDropEvent, EItemDropZone InDropZone)
+bool FNavigationToolTrack::IsAllowedInTool() const
 {
-	return Super::CanAcceptDrop(InDragDropEvent, InDropZone);
+	return IsItemValid();
+}
+
+FNavigationToolItemId FNavigationToolTrack::CalculateItemId() const
+{
+	return FNavigationToolItemId(GetParent()
+		, WeakSequence.Get()
+		, WeakSection.Get()
+		, SectionIndex
+		, FNavigationToolItemId::GetObjectPath(WeakTrack.Get()));
 }
 
-FReply FNavigationToolTrack::AcceptDrop(const FDragDropEvent& InDragDropEvent, EItemDropZone InDropZone)
+void FNavigationToolTrack::FindChildren(TArray<FNavigationToolItemPtr>& OutChildren, const bool bInRecursive)
 {
-	return Super::AcceptDrop(InDragDropEvent, InDropZone);
+	Super::FindChildren(OutChildren, bInRecursive);
 }
 
 FText FNavigationToolTrack::GetDisplayName() const
 {
-	if (const UMovieSceneTrack* const Track = GetTrack())
+	UMovieSceneTrack* const Track = GetTrack();
+	if (IsValid(Track))
 	{
 		return Track->GetDisplayName();
 	}
 	return Super::GetDisplayName();
 }
 
+FText FNavigationToolTrack::GetClassName() const
+{
+	UMovieSceneTrack* const Track = GetTrack();
+	if (IsValid(Track))
+	{
+		return FText::FromString(Track->GetClass()->GetName());
+	}
+	return FText::FromString(UMovieSceneTrack::StaticClass()->GetName());
+}
+
 const FSlateBrush* FNavigationToolTrack::GetDefaultIconBrush() const
 {
 	UMovieSceneTrack* const Track = GetTrack();
@@ -92,6 +112,16 @@ FSlateIcon FNavigationToolTrack::GetIcon() const
 	return Icon;
 }
 
+FText FNavigationToolTrack::GetIconTooltipText() const
+{
+	const UMovieSceneTrack* const Track = GetTrack();
+	if (IsValid(Track))
+	{
+		return Track->GetClass()->GetDisplayNameText();
+	}
+	return FText::GetEmpty();
+}
+
 bool FNavigationToolTrack::IsSelected(const FNavigationToolScopedSelection& InSelection) const
 {
 	UMovieSceneTrack* const Track = GetTrack();
@@ -107,9 +137,38 @@ void FNavigationToolTrack::Select(FNavigationToolScopedSelection& InSelection) c
 	}
 }
 
-void FNavigationToolTrack::SetObject_Impl(UObject* InObject)
+void FNavigationToolTrack::OnSelect()
 {
-	Super::SetObject_Impl(InObject);
+	Super::OnSelect();
+}
+
+void FNavigationToolTrack::OnDoubleClick()
+{
+	if (!WeakSequence.IsValid())
+	{
+		return;
+	}
+
+	UMovieSceneTrack* const Track = GetTrack();
+	if (!IsValid(Track))
+	{
+		return;
+	}
+
+	FocusSequence(Tool, *WeakSequence.Get(), *this);
+}
+
+void FNavigationToolTrack::OnObjectsReplaced(const TMap<UObject*, UObject*>& InReplacementMap, bool bInRecursive)
+{
+	// Get the Object even if it's Pending Kill (most likely it is)
+	const UMovieSceneTrack* const ObjectPendingKill = WeakTrack.Get(true);
+	if (ObjectPendingKill && InReplacementMap.Contains(ObjectPendingKill))
+	{
+		WeakTrack = Cast<UMovieSceneTrack>(InReplacementMap[ObjectPendingKill]);
+	}
+
+	// This handles calling OnObjectsReplaced for every child item
+	Super::OnObjectsReplaced(InReplacementMap, bInRecursive);
 
 	OnTrackObjectChanged();
 }
@@ -151,6 +210,34 @@ void FNavigationToolTrack::SetIsLocked(const bool bInIsLocked)
 	}
 }
 
+EItemContainsPlayhead FNavigationToolTrack::ContainsPlayhead() const
+{
+	const TSharedPtr<ISequencer> Sequencer = Tool.GetSequencer();
+	if (!Sequencer.IsValid())
+	{
+		return EItemContainsPlayhead::None;
+	}
+
+	if (WeakSection.IsValid())
+	{
+		const FQualifiedFrameTime PlayheadTime = Sequencer->GetLocalTime();
+		return WeakSection.Pin()->GetTrueRange().Contains(PlayheadTime.Time.FrameNumber)
+			? EItemContainsPlayhead::ContainsPlayhead : EItemContainsPlayhead::None;
+	}
+
+	const ENavigationToolCompareState State = CompareChildrenItemState<IPlayheadExtension>(this,
+		[](const IPlayheadExtension* const InPlayheadItem)
+			{
+				return InPlayheadItem->ContainsPlayhead() == EItemContainsPlayhead::ContainsPlayhead;
+			},
+		[](const IPlayheadExtension* const InPlayheadItem)
+			{
+				return InPlayheadItem->ContainsPlayhead() != EItemContainsPlayhead::ContainsPlayhead;
+			});
+
+	return static_cast<EItemContainsPlayhead>(State);
+}
+
 void FNavigationToolTrack::OnTrackObjectChanged()
 {
 	UMovieSceneTrack* const Track = GetTrack();
@@ -178,7 +265,57 @@ void FNavigationToolTrack::OnTrackObjectChanged()
 
 UMovieSceneTrack* FNavigationToolTrack::GetTrack() const
 {
-	return Cast<UMovieSceneTrack>(GetObject());
+	return WeakTrack.Get();
+}
+
+UMovieSceneSequence* FNavigationToolTrack::GetSequence() const
+{
+	return WeakSequence.Get();
+}
+
+UMovieSceneSection* FNavigationToolTrack::GetSection() const
+{
+	return WeakSection.Get();
+}
+
+int32 FNavigationToolTrack::GetSectionIndex() const
+{
+	return SectionIndex;
+}
+
+TViewModelPtr<FTrackModel> FNavigationToolTrack::GetViewModel() const
+{
+	UMovieSceneTrack* const Track = GetTrack();
+	if (!IsValid(Track))
+	{
+		return nullptr;
+	}
+
+	const TSharedPtr<ISequencer> Sequencer = Tool.GetSequencer();
+	if (!Sequencer.IsValid())
+	{
+		return nullptr;
+	}
+
+	const TSharedPtr<FSequencerEditorViewModel> ViewModel = Sequencer->GetViewModel();
+	if (!ViewModel.IsValid())
+	{
+		return nullptr;
+	}
+
+	const FViewModelPtr RootViewModel = ViewModel->GetRootModel();
+	if (!RootViewModel.IsValid())
+	{
+		return nullptr;
+	}
+
+	const FTrackModelStorageExtension* const StorageExtension = RootViewModel->CastDynamic<FTrackModelStorageExtension>();
+	if (!StorageExtension)
+	{
+		return nullptr;
+	}
+
+	return StorageExtension->FindModelForTrack(Track);
 }
 
 }
diff --git a/Engine/Source/Editor/Sequencer/Private/Tools/NavigationTool/Menus/NavigationToolFilterMenu.cpp b/Engine/Source/Editor/Sequencer/Private/Tools/NavigationTool/Menus/NavigationToolFilterMenu.cpp
--- a/Engine/Source/Editor/Sequencer/Private/Tools/NavigationTool/Menus/NavigationToolFilterMenu.cpp
+++ b/Engine/Source/Editor/Sequencer/Private/Tools/NavigationTool/Menus/NavigationToolFilterMenu.cpp
@@ -2,14 +2,14 @@
 
 #include "Tools/NavigationTool/Menus/NavigationToolFilterMenu.h"
 #include "Sequencer.h"
+#include "ToolMenu.h"
+#include "ToolMenus.h"
+#include "ToolMenuSection.h"
 #include "Tools/NavigationTool/Filters/Filters/NavigationToolFilter_CustomText.h"
 #include "Tools/NavigationTool/Filters/Filters/NavigationToolFilterBase.h"
 #include "Tools/NavigationTool/Filters/NavigationToolFilterBar.h"
 #include "Tools/NavigationTool/Menus/NavigationToolFilterBarContext.h"
 #include "Tools/NavigationTool/NavigationToolSettings.h"
-#include "ToolMenu.h"
-#include "ToolMenus.h"
-#include "ToolMenuSection.h"
 
 #define LOCTEXT_NAMESPACE "NavigationToolFilterMenu"
 
@@ -237,7 +237,7 @@ void FNavigationToolFilterMenu::FillFiltersMenuCategory(FToolMenuSection& InOutS
 
 		const FText DisplayName = Filter->GetDisplayName();
 
-		InOutSection.AddMenuEntry(*DisplayName.ToString(),
+		FToolMenuEntry& NewMenuEntry = InOutSection.AddMenuEntry(*DisplayName.ToString(),
 			DisplayName,
 			Filter->GetToolTipText(),
 			Filter->GetIcon(),
@@ -253,6 +253,11 @@ void FNavigationToolFilterMenu::FillFiltersMenuCategory(FToolMenuSection& InOutS
 					})
 			),
 			EUserInterfaceActionType::ToggleButton);
+
+		if (const TSharedPtr<FUICommandInfo> ToggleCommand = Filter->GetToggleCommand())
+		{
+			NewMenuEntry.InputBindingLabel = ToggleCommand->GetInputText();
+		}
 	}
 }
 
diff --git a/Engine/Source/Editor/Sequencer/Private/Tools/NavigationTool/Menus/NavigationToolItemContextMenu.cpp b/Engine/Source/Editor/Sequencer/Private/Tools/NavigationTool/Menus/NavigationToolItemContextMenu.cpp
--- a/Engine/Source/Editor/Sequencer/Private/Tools/NavigationTool/Menus/NavigationToolItemContextMenu.cpp
+++ b/Engine/Source/Editor/Sequencer/Private/Tools/NavigationTool/Menus/NavigationToolItemContextMenu.cpp
@@ -35,7 +35,7 @@ TSharedRef<SWidget> FNavigationToolItemContextMenu::CreateMenu(const TSharedRef<
 	UNavigationToolItemMenuContext* const ContextObject = NewObject<UNavigationToolItemMenuContext>();
 	ContextObject->Init(InToolView->GetOwnerTool(), InItemList);
 
-	FToolMenuContext Context(InToolView->GetViewCommandList(), nullptr, ContextObject);
+	FToolMenuContext Context(InToolView->GetBaseCommandList(), nullptr, ContextObject);
 	return ToolMenus->GenerateWidget(MenuName, Context);
 }
 
@@ -82,7 +82,8 @@ void FNavigationToolItemContextMenu::PopulateMenu(UToolMenu* const InMenu)
 
 void FNavigationToolItemContextMenu::CreateGenericSection(UToolMenu& InMenu)
 {
-	const FGenericCommands& GenericCommands = FGenericCommands::Get();
+	// Disabling all for now
+	/*const FGenericCommands& GenericCommands = FGenericCommands::Get();
 
 	FToolMenuSection& GenericSection = InMenu.FindOrAddSection(TEXT("GenericActions"), LOCTEXT("GenericActionsHeader", "Generic Actions"));
 
@@ -91,7 +92,7 @@ void FNavigationToolItemContextMenu::CreateGenericSection(UToolMenu& InMenu)
 	//GenericSection.AddMenuEntry(GenericCommands.Paste);
 	//GenericSection.AddMenuEntry(GenericCommands.Duplicate);
 	GenericSection.AddMenuEntry(GenericCommands.Delete);
-	GenericSection.AddMenuEntry(GenericCommands.Rename);
+	GenericSection.AddMenuEntry(GenericCommands.Rename);*/
 }
 
 void FNavigationToolItemContextMenu::CreateToolSection(UToolMenu& InMenu)
@@ -100,11 +101,23 @@ void FNavigationToolItemContextMenu::CreateToolSection(UToolMenu& InMenu)
 
 	FToolMenuSection& ToolSection = InMenu.FindOrAddSection(TEXT("ToolActions"), LOCTEXT("ToolActionsHeader", "Navigation Tool Actions"));
 
-	ToolSection.AddMenuEntry(ToolCommands.ExpandAll);
-	ToolSection.AddMenuEntry(ToolCommands.CollapseAll);
-	ToolSection.AddMenuEntry(ToolCommands.SelectAllChildren);
-	ToolSection.AddMenuEntry(ToolCommands.SelectImmediateChildren);
-	ToolSection.AddMenuEntry(ToolCommands.FocusSingleSelection);
+	FToolMenuEntry& ExpandAllEntry = ToolSection.AddMenuEntry(ToolCommands.ExpandAll);
+	ExpandAllEntry.Icon = FSlateIcon(FAppStyle::GetAppStyleSetName(), TEXT("TreeArrow_Expanded"));
+
+	FToolMenuEntry& CollapseAllEntry = ToolSection.AddMenuEntry(ToolCommands.CollapseAll);
+	CollapseAllEntry.Icon = FSlateIcon(FAppStyle::GetAppStyleSetName(), TEXT("TreeArrow_Collapsed"));
+
+	FToolMenuEntry& SelectAllChildrenEntry = ToolSection.AddMenuEntry(ToolCommands.SelectAllChildren);
+	SelectAllChildrenEntry.Icon = FSlateIcon(FAppStyle::GetAppStyleSetName(), TEXT("LevelEditor.SelectAllDescendants"));
+
+	FToolMenuEntry& SelectImmediateChildrenEntry = ToolSection.AddMenuEntry(ToolCommands.SelectImmediateChildren);
+	SelectImmediateChildrenEntry.Icon = FSlateIcon(FAppStyle::GetAppStyleSetName(), TEXT("LevelEditor.SelectImmediateChildren"));
+
+	FToolMenuEntry& FocusSelectionEntry = ToolSection.AddMenuEntry(ToolCommands.FocusSingleSelection);
+	FocusSelectionEntry.Icon = FSlateIcon(FAppStyle::GetAppStyleSetName(), TEXT("WorldPartition.FocusSelection"));
+
+	FToolMenuEntry& FocusInContentBrowserEntry = ToolSection.AddMenuEntry(ToolCommands.FocusInContentBrowser);
+	FocusInContentBrowserEntry.Icon = FSlateIcon(FAppStyle::GetAppStyleSetName(), TEXT("ContentBrowser.TabIcon"));
 }
 
 }
diff --git a/Engine/Source/Editor/Sequencer/Private/Tools/NavigationTool/Menus/NavigationToolToolbarMenu.cpp b/Engine/Source/Editor/Sequencer/Private/Tools/NavigationTool/Menus/NavigationToolToolbarMenu.cpp
--- a/Engine/Source/Editor/Sequencer/Private/Tools/NavigationTool/Menus/NavigationToolToolbarMenu.cpp
+++ b/Engine/Source/Editor/Sequencer/Private/Tools/NavigationTool/Menus/NavigationToolToolbarMenu.cpp
@@ -38,6 +38,7 @@ TSharedRef<SWidget> FNavigationToolToolbarMenu::CreateToolbar(const TSharedRef<F
 		ToolBar->StyleName = TEXT("StatusBarToolBar");
 		ToolBar->bToolBarForceSmallIcons = true;
 		ToolBar->bToolBarIsFocusable = true;
+		ToolBar->bShouldCloseWindowAfterMenuSelection = false;
 		ToolBar->AddDynamicSection(TEXT("Main"), FNewToolMenuDelegate::CreateLambda([this](UToolMenu* const InMenu)
 			{
 				UNavigationToolViewMenuContext* const Context = InMenu->FindContext<UNavigationToolViewMenuContext>();
@@ -109,35 +110,32 @@ void FNavigationToolToolbarMenu::CreateSettingsMenu(UToolMenu* const InMenu)
 	}
 
 	const FNavigationToolCommands& ToolCommands = FNavigationToolCommands::Get();
-	const FNavigationToolFilterCommands& FilterCommands = FNavigationToolFilterCommands::Get();
 
 	FToolMenuSection& ViewOptionsSection = InMenu->FindOrAddSection(TEXT("ViewOptions")
 		, LOCTEXT("ViewOptionsHeading", "View Options"));
 
 	ViewOptionsSection.AddMenuEntry(ToolCommands.ExpandAll);
 	ViewOptionsSection.AddMenuEntry(ToolCommands.CollapseAll);
-	ViewOptionsSection.AddMenuEntry(ToolCommands.ToggleMutedHierarchy);
+	// Disable until this is applicable
+	//ViewOptionsSection.AddMenuEntry(ToolCommands.ToggleMutedHierarchy);
 	ViewOptionsSection.AddMenuEntry(ToolCommands.ToggleAutoExpandToSelection);
 	ViewOptionsSection.AddMenuEntry(ToolCommands.ToggleShortNames);
 
 	CreateItemViewOptionsMenu(InMenu);
 
-	CreateColumnViewOptionsMenu(InMenu);
-
-	FToolMenuSection& FilterVisibilitySection = InMenu->FindOrAddSection(TEXT("Visibility")
-		, LOCTEXT("VisibilityHeading", "Filter Bar"));
-
-	FilterVisibilitySection.AddMenuEntry(FilterCommands.ToggleFilterBarVisibility);
-
-	FToolMenuSection& FilterLayoutSection = InMenu->FindOrAddSection(TEXT("Layout")
-		, LOCTEXT("LayoutHeading", "Filter Bar Layout"));
-
-	FilterLayoutSection.AddMenuEntry(FilterCommands.SetToVerticalLayout);
-	FilterLayoutSection.AddMenuEntry(FilterCommands.SetToHorizontalLayout);
-
 	FToolMenuSection& MiscSection = InMenu->FindOrAddSection(TEXT("Misc")
 		, LOCTEXT("MiscHeading", "Misc"));
 
+	MiscSection.AddSubMenu(TEXT("ColumnViews")
+		, LOCTEXT("ColumnViewsSubMenu", "Column Views")
+		, LOCTEXT("ColumnViewsSubMenuTooltip", "")
+		, FNewToolMenuDelegate::CreateSP(this, &FNavigationToolToolbarMenu::CreateColumnViewOptionsMenu));
+
+	MiscSection.AddSubMenu(TEXT("FilterBarOptions")
+		, LOCTEXT("FilterBarOptionsSubMenu", "Filter Bar Options")
+		, LOCTEXT("FilterBarOptionsSubMenuTooltip", "")
+		, FNewToolMenuDelegate::CreateSP(this, &FNavigationToolToolbarMenu::CreateFilterBarOptionsMenu));
+
 	MiscSection.AddMenuEntry(ToolCommands.OpenToolSettings);
 }
 
@@ -260,7 +258,7 @@ void FNavigationToolToolbarMenu::CreateItemViewOptionsMenu(UToolMenu* const InMe
 	{
 		const TSharedRef<FNavigationToolBuiltInFilter> FilterRef = Filter.ToSharedRef();
 
-		GlobalFiltersSection.AddMenuEntry(NAME_None,
+		FToolMenuEntry& NewMenuEntry = GlobalFiltersSection.AddMenuEntry(NAME_None,
 			Filter->GetDisplayName(),
 			Filter->GetToolTipText(),
 			Filter->GetIcon(),
@@ -270,6 +268,11 @@ void FNavigationToolToolbarMenu::CreateItemViewOptionsMenu(UToolMenu* const InMe
 				FIsActionChecked::CreateSP(this, &FNavigationToolToolbarMenu::IsGlobalFilterActive, FilterRef)
 			),
 			EUserInterfaceActionType::ToggleButton);
+
+		if (const TSharedPtr<FUICommandInfo> ToggleCommand = FilterRef->GetToggleCommand())
+		{
+			NewMenuEntry.InputBindingLabel = ToggleCommand->GetInputText();
+		}
 	}
 }
 
@@ -326,6 +329,27 @@ void FNavigationToolToolbarMenu::CreateColumnViewOptionsMenu(UToolMenu* const In
 	}
 }
 
+void FNavigationToolToolbarMenu::CreateFilterBarOptionsMenu(UToolMenu* const InMenu)
+{
+	if (!IsValid(InMenu))
+	{
+		return;
+	}
+
+	const FNavigationToolFilterCommands& FilterCommands = FNavigationToolFilterCommands::Get();
+
+	FToolMenuSection& FilterBarVisibilitySection = InMenu->FindOrAddSection(TEXT("FilterBarVisibility")
+		, LOCTEXT("FilterBarVisibilityHeading", "Filter Bar"));
+
+	FilterBarVisibilitySection.AddMenuEntry(FilterCommands.ToggleFilterBarVisibility);
+
+	FToolMenuSection& FilterBarLayoutSection = InMenu->FindOrAddSection(TEXT("FilterBarLayout")
+		, LOCTEXT("FilterBarLayoutHeading", "Filter Bar Layout"));
+
+	FilterBarLayoutSection.AddMenuEntry(FilterCommands.SetToVerticalLayout);
+	FilterBarLayoutSection.AddMenuEntry(FilterCommands.SetToHorizontalLayout);
+}
+
 ECheckBoxState FNavigationToolToolbarMenu::GetCustomColumnViewMenuItemCheckState(const FName InColumnId, const TSharedRef<FNavigationToolView> InToolView) const
 {
 	UNavigationToolSettings* const ToolSettings = GetMutableDefault<UNavigationToolSettings>();
diff --git a/Engine/Source/Editor/Sequencer/Private/Tools/NavigationTool/Menus/NavigationToolToolbarMenu.h b/Engine/Source/Editor/Sequencer/Private/Tools/NavigationTool/Menus/NavigationToolToolbarMenu.h
--- a/Engine/Source/Editor/Sequencer/Private/Tools/NavigationTool/Menus/NavigationToolToolbarMenu.h
+++ b/Engine/Source/Editor/Sequencer/Private/Tools/NavigationTool/Menus/NavigationToolToolbarMenu.h
@@ -29,6 +29,7 @@ protected:
 	void CreateSettingsMenu(UToolMenu* const InMenu);
 	void CreateItemViewOptionsMenu(UToolMenu* const InMenu);
 	void CreateColumnViewOptionsMenu(UToolMenu* const InMenu);
+	void CreateFilterBarOptionsMenu(UToolMenu* const InMenu);
 
 	ECheckBoxState GetCustomColumnViewMenuItemCheckState(const FName InColumnId, const TSharedRef<FNavigationToolView> InToolView) const;
 	void OnDeleteCustomColumnViewMenuItemClick(const FName InColumnId);
diff --git a/Engine/Source/Editor/Sequencer/Private/Tools/NavigationTool/NavigationTool.cpp b/Engine/Source/Editor/Sequencer/Private/Tools/NavigationTool/NavigationTool.cpp
--- a/Engine/Source/Editor/Sequencer/Private/Tools/NavigationTool/NavigationTool.cpp
+++ b/Engine/Source/Editor/Sequencer/Private/Tools/NavigationTool/NavigationTool.cpp
@@ -17,14 +17,17 @@
 #include "ScopedTransaction.h"
 #include "Sections/MovieSceneSubSection.h"
 #include "Sequencer.h"
+#include "MVVM/Selection/Selection.h"
+#include "MVVM/ViewModels/TrackModel.h"
+#include "MVVM/ViewModels/TrackRowModel.h"
+#include "MVVM/ViewModels/ObjectBindingModel.h"
 #include "Textures/SlateIcon.h"
 #include "Tools/NavigationTool/Customization/NavigationToolExtender.h"
 #include "Tools/NavigationTool/Filters/Filters/NavigationToolBuiltInFilter.h"
 #include "Tools/NavigationTool/ItemActions/NavigationToolItemAction.h"
 #include "Tools/NavigationTool/ItemProxies/NavigationToolItemProxyRegistry.h"
+#include "Tools/NavigationTool/Items/NavigationToolBinding.h"
 #include "Tools/NavigationTool/Items/NavigationToolItemUtils.h"
-#include "Tools/NavigationTool/Items/NavigationToolObjectReference.h"
-#include "Tools/NavigationTool/Items/NavigationToolSequence.h"
 #include "Tools/NavigationTool/Items/NavigationToolTreeRoot.h"
 #include "Tools/NavigationTool/NavigationToolEditorSelection.h"
 #include "Tools/NavigationTool/NavigationToolScopedSelection.h"
@@ -73,7 +76,7 @@ void FNavigationTool::Init()
 				OnTreeViewChanged();
 			});
 		Sequencer->OnTreeViewChanged().AddSP(this, &FNavigationTool::OnTreeViewChanged);
-
+		Sequencer->OnEndScrubbingEvent().AddSP(this, &FNavigationTool::OnTreeViewChanged);
 		Sequencer->OnCloseEvent().AddSPLambda(this, [this](const TSharedRef<ISequencer> InSequencer)
 			{
 				if (InSequencer == WeakSequencer.Pin())
@@ -81,6 +84,8 @@ void FNavigationTool::Init()
 					Shutdown();
 				}
 			});
+		Sequencer->GetSelection().Outliner.OnChanged.AddSP(this, &FNavigationTool::OnSequencerSelectionChanged);
+		Sequencer->GetSelection().TrackArea.OnChanged.AddSP(this, &FNavigationTool::OnSequencerSelectionChanged);
 	}
 
 	// Listen to object replacement changes
@@ -102,13 +107,17 @@ void FNavigationTool::Shutdown()
 {
 	SaveAllProviderStates();
 
-	if (const TSharedPtr<ISequencer> Sequencer = WeakSequencer.Pin())
+	if (const TSharedPtr<FSequencer> Sequencer = StaticCastSharedPtr<FSequencer>(WeakSequencer.Pin()))
 	{
 		Sequencer->OnChannelChanged().RemoveAll(this);
 		Sequencer->OnMovieSceneBindingsChanged().RemoveAll(this);
 		Sequencer->OnMovieSceneBindingsPasted().RemoveAll(this);
 		Sequencer->OnMovieSceneDataChanged().RemoveAll(this);
 		Sequencer->OnTreeViewChanged().RemoveAll(this);
+		Sequencer->OnEndScrubbingEvent().RemoveAll(this);
+		Sequencer->OnCloseEvent().RemoveAll(this);
+		Sequencer->GetSelection().Outliner.OnChanged.RemoveAll(this);
+		Sequencer->GetSelection().TrackArea.OnChanged.RemoveAll(this);
 	}
 
 	FCoreUObjectDelegates::OnObjectsReplaced.RemoveAll(this);
@@ -170,11 +179,11 @@ void FNavigationTool::OnProvidersChanged(const TSharedRef<INavigationTool>& InTo
 		const UNavigationToolSettings* const ToolSettings = GetDefault<UNavigationToolSettings>();
 		if (ensure(IsValid(ToolSettings)))
 		{
-			LoadProviderState(InProvider);
+			//LoadProviderState(InProvider);
 
 			ForEachToolView([&InProvider](const TSharedPtr<FNavigationToolView>& InToolView)
 				{
-					InToolView->LoadViewState(InProvider);
+					//InToolView->LoadViewState(InProvider);
 
 					InToolView->CreateColumns(InProvider);
 					InToolView->CreateDefaultColumnViews(InProvider);
@@ -657,6 +666,35 @@ FNavigationToolItemPtr FNavigationTool::FindItem(const FNavigationToolItemId& In
 	return nullptr;
 }
 
+TArray<FNavigationToolItemPtr> FNavigationTool::TryFindItems(const FViewModelPtr& InViewModel) const
+{
+	if (const TViewModelPtr<FTrackModel> TrackModel = InViewModel.ImplicitCast())
+	{
+		return FindItemsFromMovieSceneObject(TrackModel->GetTrack());
+	}
+
+	if (const TViewModelPtr<FTrackRowModel> TrackRowModel = InViewModel.ImplicitCast())
+	{
+		return FindItemsFromMovieSceneObject(TrackRowModel->GetTrack());
+	}
+
+	if (const TViewModelPtr<FObjectBindingModel> ObjectBindingModel = InViewModel.ImplicitCast())
+	{
+		return FindItemsFromObjectGuid(ObjectBindingModel->GetObjectGuid());
+	}
+
+	if (const TViewModelPtr<FSectionModel> SectionModel = InViewModel.ImplicitCast())
+	{
+		UMovieSceneSubSection* const SubSection = Cast<UMovieSceneSubSection>(SectionModel->GetSection());
+		if (IsValid(SubSection))
+		{
+			return FindItemsFromMovieSceneObject(SubSection->GetSequence());
+		}
+	}
+
+	return {};
+}
+
 void FNavigationTool::SetIgnoreNotify(const ENavigationToolIgnoreNotifyFlags InFlag, const bool bInIgnore)
 {
 	if (bInIgnore)
@@ -718,19 +756,7 @@ void FNavigationTool::Tick(const float InDeltaTime)
 			{
 				if (Object.IsValid())
 				{
-					if (const FNavigationToolItemPtr Item = FindItem(FNavigationToolItemId(Object.Get())))
-					{
-						// If it's a Shared Object, select all its References as Shared Objects are not in the Navigation Tool directly
-						// and the outliner does not know which Object Reference should be selected
-						if (const FNavigationToolSharedObject* const SharedObjectItem = Item->CastTo<FNavigationToolSharedObject>())
-						{
-							ItemsToSelect.Append(SharedObjectItem->GetObjectReferences());
-						}
-						else
-						{
-							ItemsToSelect.Add(Item);
-						}
-					}
+					ItemsToSelect = FindItemsFromMovieSceneObject(Object.Get());
 				}
 			}
 		}
@@ -1329,6 +1355,82 @@ void FNavigationTool::OnTreeViewChanged()
 	RequestRefresh();
 }
 
+void FNavigationTool::OnSequencerSelectionChanged()
+{
+	UNavigationToolSettings* const ToolSettings = GetMutableDefault<UNavigationToolSettings>();
+	if (!ToolSettings->ShouldSyncSelectionToNavigationTool())
+	{
+		return;
+	}
+
+	const TSharedPtr<FSequencer> Sequencer = StaticCastSharedPtr<FSequencer>(WeakSequencer.Pin());
+	if (!Sequencer.IsValid())
+	{
+		return;
+	}
+
+	FSequencerSelection& SequencerSelection = Sequencer->GetSelection();
+
+	// Unbind to stop circular loop
+	SequencerSelection.Outliner.OnChanged.RemoveAll(this);
+
+	const bool bPreviouslySyncingSelectionToSequencer = ToolSettings->ShouldSyncSelectionToSequencer();
+	ToolSettings->SetSyncSelectionToSequencer(false);
+
+	TArray<FNavigationToolItemPtr> ItemsToSelect;
+
+	for (const FViewModelPtr OutlinerItem : SequencerSelection.Outliner)
+	{
+		ItemsToSelect.Append(TryFindItems(OutlinerItem));
+	}
+	for (const FViewModelPtr TrackAreaItem : SequencerSelection.TrackArea)
+	{
+		ItemsToSelect.Append(TryFindItems(TrackAreaItem));
+	}
+
+	SelectItems(ItemsToSelect);
+
+	// Bind from stopping circular loop
+	SequencerSelection.Outliner.OnChanged.AddSP(this, &FNavigationTool::OnSequencerSelectionChanged);
+
+	ToolSettings->SetSyncSelectionToSequencer(bPreviouslySyncingSelectionToSequencer);
+}
+
+TArray<FNavigationToolItemPtr> FNavigationTool::FindItemsFromMovieSceneObject(UObject* const InObject) const
+{
+	TArray<FNavigationToolItemPtr> OutItems;
+	if (IsValid(InObject))
+	{
+		for (const FNavigationToolItemPtr& Item : RootItem->GetDescendants())
+		{
+			if (Item->GetItemObject() == InObject)
+			{
+				OutItems.Add(Item);
+			}
+		}
+	}
+	return OutItems;
+}
+
+TArray<FNavigationToolItemPtr> FNavigationTool::FindItemsFromObjectGuid(const FGuid& InObjectGuid) const
+{
+	TArray<FNavigationToolItemPtr> OutItems;
+	if (InObjectGuid.IsValid())
+	{
+		for (const FNavigationToolItemPtr& Item : RootItem->GetDescendants())
+		{
+			if (const FNavigationToolBinding* const BindingItem = Item->CastTo<FNavigationToolBinding>())
+			{
+				if (BindingItem->GetBinding().GetObjectGuid() == InObjectGuid)
+				{
+					OutItems.Add(Item);
+				}
+			}
+		}
+	}
+	return OutItems;
+}
+
 }
 
 #undef LOCTEXT_NAMESPACE
diff --git a/Engine/Source/Editor/Sequencer/Private/Tools/NavigationTool/NavigationTool.h b/Engine/Source/Editor/Sequencer/Private/Tools/NavigationTool/NavigationTool.h
--- a/Engine/Source/Editor/Sequencer/Private/Tools/NavigationTool/NavigationTool.h
+++ b/Engine/Source/Editor/Sequencer/Private/Tools/NavigationTool/NavigationTool.h
@@ -5,6 +5,7 @@
 #include "EditorUndoClient.h"
 #include "EditorSequenceNavigationDefs.h"
 #include "Engine/EngineTypes.h"
+#include "SequencerCoreFwd.h"
 #include "Tools/NavigationTool/INavigationTool.h"
 #include "Tools/NavigationTool/Items/INavigationToolItem.h"
 #include "Tools/NavigationTool/Items/NavigationToolItemId.h"
@@ -53,8 +54,6 @@ public:
 
 	bool CanProcessSequenceSpawn(UMovieSceneSequence* const InSequence) const;
 
-	TSharedPtr<FUICommandList> GetBaseCommandList() const;
-
 	FOnToolLoaded OnToolLoaded;
 
 	/** Gathers all previously existing and new Item Proxies for a given Item */
@@ -70,6 +69,8 @@ public:
 
 	//~ Begin INavigationTool
 
+	virtual TSharedPtr<FUICommandList> GetBaseCommandList() const override;
+
 	virtual bool IsToolTabVisible() const override;
 	virtual void ShowHideToolTab(const bool bInVisible) override;
 	virtual void ToggleToolTabVisible() override;
@@ -95,6 +96,7 @@ public:
 	virtual TSharedRef<FNavigationToolItem> GetTreeRoot() const override;
 
 	virtual FNavigationToolItemPtr FindItem(const FNavigationToolItemId& InItemId) const override;
+	virtual TArray<FNavigationToolItemPtr> TryFindItems(const FViewModelPtr& InViewModel) const override;
 
 	virtual void SetIgnoreNotify(const ENavigationToolIgnoreNotifyFlags InFlag, const bool bInIgnore) override;
 
@@ -181,8 +183,6 @@ public:
 	/** Marks the Navigation Tool dirty. This triggers INavigationToolProvider::OnToolModified on next tick */
 	void SetToolModified();
 
-	void OnTreeViewChanged();
-
 private:
 	friend class FNavigationToolView;
 	friend class FNavigationToolToolbarMenu;
@@ -211,6 +211,13 @@ private:
 
 	USequencerSettings* GetSequencerSettings() const;
 
+	void OnTreeViewChanged();
+
+	void OnSequencerSelectionChanged();
+
+	TArray<FNavigationToolItemPtr> FindItemsFromMovieSceneObject(UObject* const InObject) const;
+	TArray<FNavigationToolItemPtr> FindItemsFromObjectGuid(const FGuid& InObjectGuid) const;
+
 	TWeakPtr<ISequencer> WeakSequencer;
 
 	TSharedRef<FNavigationToolTab> ToolTab;
diff --git a/Engine/Source/Editor/Sequencer/Private/Tools/NavigationTool/NavigationToolCommands.cpp b/Engine/Source/Editor/Sequencer/Private/Tools/NavigationTool/NavigationToolCommands.cpp
--- a/Engine/Source/Editor/Sequencer/Private/Tools/NavigationTool/NavigationToolCommands.cpp
+++ b/Engine/Source/Editor/Sequencer/Private/Tools/NavigationTool/NavigationToolCommands.cpp
@@ -121,6 +121,12 @@ void FNavigationToolCommands::RegisterCommands()
 		, "Sets this sequence as the root sequence to focus on in the Sequencer"
 		, EUserInterfaceActionType::Button
 		, FInputChord());
+
+	UI_COMMAND(FocusInContentBrowser
+		, "Focus in Content Browser"
+		, "Focus this sequence in the Content Browser"
+		, EUserInterfaceActionType::Button
+		, FInputChord());
 }
 
 }
diff --git a/Engine/Source/Editor/Sequencer/Private/Tools/NavigationTool/NavigationToolCommands.h b/Engine/Source/Editor/Sequencer/Private/Tools/NavigationTool/NavigationToolCommands.h
--- a/Engine/Source/Editor/Sequencer/Private/Tools/NavigationTool/NavigationToolCommands.h
+++ b/Engine/Source/Editor/Sequencer/Private/Tools/NavigationTool/NavigationToolCommands.h
@@ -53,6 +53,8 @@ public:
 	TSharedPtr<FUICommandInfo> SaveCurrentColumnView;
 
 	TSharedPtr<FUICommandInfo> FocusSingleSelection;
+
+	TSharedPtr<FUICommandInfo> FocusInContentBrowser;
 };
 
 }
diff --git a/Engine/Source/Editor/Sequencer/Private/Tools/NavigationTool/NavigationToolScopedSelection.cpp b/Engine/Source/Editor/Sequencer/Private/Tools/NavigationTool/NavigationToolScopedSelection.cpp
--- a/Engine/Source/Editor/Sequencer/Private/Tools/NavigationTool/NavigationToolScopedSelection.cpp
+++ b/Engine/Source/Editor/Sequencer/Private/Tools/NavigationTool/NavigationToolScopedSelection.cpp
@@ -7,9 +7,10 @@
 #include "MVVM/SectionModelStorageExtension.h"
 #include "MVVM/Selection/Selection.h"
 #include "MVVM/TrackModelStorageExtension.h"
+#include "MVVM/TrackRowModelStorageExtension.h"
 #include "MVVM/Views/SOutlinerView.h"
 #include "Sequencer.h"
-#include "SSequencer.h"
+#include "Utils/NavigationToolMiscUtils.h"
 
 namespace UE::Sequencer::NavigationTool
 {
@@ -141,8 +142,16 @@ void FNavigationToolScopedSelection::SyncSelections()
 		return;
 	}
 
-	const TSharedPtr<SSequencer> SequencerWidget = StaticCastSharedRef<SSequencer>(Sequencer.GetSequencerWidget());
-	const TSharedPtr<SOutlinerView> SequencerOutlinerView = SequencerWidget.IsValid() ? SequencerWidget->GetTreeView() : nullptr;
+	const FModifierKeysState ModifierKeys = FSlateApplication::Get().GetModifierKeys();
+	UNavigationToolSettings* const ToolSettings = GetMutableDefault<UNavigationToolSettings>();
+	const bool bPreviouslySyncingSelectionToNavigationTool = ToolSettings->ShouldSyncSelectionToNavigationTool();
+
+	if (!ModifierKeys.IsAltDown())
+	{
+		ToolSettings->SetSyncSelectionToNavigationTool(false);
+	}
+
+	const TSharedPtr<SOutlinerView> OutlinerView = GetOutlinerViewWidget(Sequencer);
 
 	Sequencer.EmptySelection();
 
@@ -154,51 +163,53 @@ void FNavigationToolScopedSelection::SyncSelections()
 			{
 				Selection->Outliner.Select(ObjectBindingModel);
 
-				if (SequencerOutlinerView.IsValid())
+				if (OutlinerView.IsValid())
 				{
-					SequencerOutlinerView->ExpandCollapseNode(ObjectBindingModel, true, ETreeRecursion::NonRecursive);
-					SequencerOutlinerView->RequestScrollIntoView(ObjectBindingModel);
+					OutlinerView->RequestScrollIntoView(ObjectBindingModel);
 				}
 			}
 		}
 	}
 
-	for (UMovieSceneSection* const Section : SelectedSections)
+	for (UMovieSceneTrack* const Track : SelectedTracks)
 	{
-		if (FSectionModelStorageExtension* const SectionModelStorage = RootViewModel->CastDynamic<FSectionModelStorageExtension>())
+		if (FTrackModelStorageExtension* const TrackModelStorage = RootViewModel->CastDynamic<FTrackModelStorageExtension>())
 		{
-			if (const TSharedPtr<FSectionModel> SectionModel = SectionModelStorage->FindModelForSection(Section))
+			if (const TSharedPtr<FTrackModel> TrackModel = TrackModelStorage->FindModelForTrack(Track))
 			{
-				Selection->TrackArea.Select(SectionModel);
+				Selection->Outliner.Select(TrackModel);
 
-				if (const TViewModelPtr<IOutlinerExtension> LinkedOutlinerItem = SectionModel->GetLinkedOutlinerItem())
+				if (OutlinerView.IsValid())
 				{
-					if (SequencerOutlinerView.IsValid())
-					{
-						SequencerOutlinerView->ExpandCollapseNode(LinkedOutlinerItem, true, ETreeRecursion::NonRecursive);
-						SequencerOutlinerView->RequestScrollIntoView(LinkedOutlinerItem);
-					}
+					OutlinerView->RequestScrollIntoView(TrackModel);
 				}
 			}
 		}
 	}
 
-	for (UMovieSceneTrack* const Track : SelectedTracks)
+	for (UMovieSceneSection* const Section : SelectedSections)
 	{
-		if (FTrackModelStorageExtension* const TrackModelStorage = RootViewModel->CastDynamic<FTrackModelStorageExtension>())
+		if (FSectionModelStorageExtension* const SectionModelStorage = RootViewModel->CastDynamic<FSectionModelStorageExtension>())
 		{
-			if (const TSharedPtr<FTrackModel> TrackModel = TrackModelStorage->FindModelForTrack(Track))
+			if (const TSharedPtr<FSectionModel> SectionModel = SectionModelStorage->FindModelForSection(Section))
 			{
-				Selection->TrackArea.Select(TrackModel);
+				Selection->TrackArea.Select(SectionModel);
 
-				if (SequencerOutlinerView.IsValid())
+				if (const TViewModelPtr<IOutlinerExtension> LinkedOutlinerItem = SectionModel->GetLinkedOutlinerItem())
 				{
-					SequencerOutlinerView->ExpandCollapseNode(TrackModel, true, ETreeRecursion::NonRecursive);
-					SequencerOutlinerView->RequestScrollIntoView(TrackModel);
+					if (OutlinerView.IsValid())
+					{
+						OutlinerView->RequestScrollIntoView(LinkedOutlinerItem);
+					}
 				}
 			}
 		}
 	}
+
+	if (!ModifierKeys.IsAltDown())
+	{
+		ToolSettings->SetSyncSelectionToNavigationTool(bPreviouslySyncingSelectionToNavigationTool);
+	}
 }
 
 ISequencer& FNavigationToolScopedSelection::GetSequencer() const
diff --git a/Engine/Source/Editor/Sequencer/Private/Tools/NavigationTool/NavigationToolSettings.cpp b/Engine/Source/Editor/Sequencer/Private/Tools/NavigationTool/NavigationToolSettings.cpp
--- a/Engine/Source/Editor/Sequencer/Private/Tools/NavigationTool/NavigationToolSettings.cpp
+++ b/Engine/Source/Editor/Sequencer/Private/Tools/NavigationTool/NavigationToolSettings.cpp
@@ -217,3 +217,25 @@ void UNavigationToolSettings::SetLastFilterBarSizeCoefficient(const float bInSiz
 	LastFilterBarSizeCoefficient = bInSizeCoefficient;
 	SaveConfig();
 }
+
+bool UNavigationToolSettings::ShouldSyncSelectionToNavigationTool() const
+{
+	return bSyncSelectionToNavigationTool;
+}
+
+void UNavigationToolSettings::SetSyncSelectionToNavigationTool(const bool bInSync)
+{
+	bSyncSelectionToNavigationTool = bInSync;
+    SaveConfig();
+}
+
+bool UNavigationToolSettings::ShouldSyncSelectionToSequencer() const
+{
+	return bSyncSelectionToSequencer;
+}
+
+void UNavigationToolSettings::SetSyncSelectionToSequencer(const bool bInSync)
+{
+	bSyncSelectionToSequencer = bInSync;
+	SaveConfig();
+}
diff --git a/Engine/Source/Editor/Sequencer/Private/Tools/NavigationTool/NavigationToolTab.cpp b/Engine/Source/Editor/Sequencer/Private/Tools/NavigationTool/NavigationToolTab.cpp
--- a/Engine/Source/Editor/Sequencer/Private/Tools/NavigationTool/NavigationToolTab.cpp
+++ b/Engine/Source/Editor/Sequencer/Private/Tools/NavigationTool/NavigationToolTab.cpp
@@ -72,8 +72,6 @@ void FNavigationToolTab::ShowHideToolTab(const bool bInVisible)
 		{
 			TabManager->TryInvokeTab(NavigationToolTabId);
 		}
-
-		OwnerTool.Refresh();
 	}
 	else
 	{
@@ -156,7 +154,7 @@ void FNavigationToolTab::RegisterToolTab()
 	TabManager->RegisterTabSpawner(NavigationToolTabId, FOnSpawnTab::CreateSP(this, &FNavigationToolTab::SpawnToolTab))
 		.SetGroup(WorkspaceMenu::GetMenuStructure().GetLevelEditorCinematicsCategory())
 		.SetDisplayName(LOCTEXT("NavigationToolTab", "Sequence Nav"))
-		.SetIcon(FSlateIcon(FAppStyle::GetAppStyleSetName(), TEXT("GraphEditor.Sequence_16x")));
+		.SetIcon(FSlateIcon(FAppStyle::GetAppStyleSetName(), TEXT("Sequencer.NavigationTool.TabIcon")));
 }
 
 void FNavigationToolTab::UnregisterToolTab()
diff --git a/Engine/Source/Editor/Sequencer/Private/Tools/NavigationTool/NavigationToolView.cpp b/Engine/Source/Editor/Sequencer/Private/Tools/NavigationTool/NavigationToolView.cpp
--- a/Engine/Source/Editor/Sequencer/Private/Tools/NavigationTool/NavigationToolView.cpp
+++ b/Engine/Source/Editor/Sequencer/Private/Tools/NavigationTool/NavigationToolView.cpp
@@ -1,9 +1,12 @@
 // Copyright Epic Games, Inc. All Rights Reserved.
 
 #include "NavigationToolView.h"
+#include "ContentBrowserModule.h"
 #include "Framework/Application/SlateApplication.h"
 #include "Framework/Commands/GenericCommands.h"
 #include "Framework/MultiBox/MultiBoxBuilder.h"
+#include "IContentBrowserSingleton.h"
+#include "LevelSequence.h"
 #include "Misc/MessageDialog.h"
 #include "MVVM/Selection/Selection.h"
 #include "Sequencer.h"
@@ -292,6 +295,10 @@ void FNavigationToolView::BindCommands(const TSharedPtr<FUICommandList>& InBaseC
 	ViewCommandList->MapAction(ToolCommands.FocusSingleSelection
 		, FExecuteAction::CreateSP(this, &FNavigationToolView::FocusSingleSelection)
 		, FCanExecuteAction::CreateSP(this, &FNavigationToolView::CanFocusSingleSelection));
+
+	ViewCommandList->MapAction(ToolCommands.FocusInContentBrowser
+		, FExecuteAction::CreateSP(this, &FNavigationToolView::FocusInContentBrowser)
+		, FCanExecuteAction::CreateSP(this, &FNavigationToolView::CanFocusInContentBrowser));
 }
 
 TSharedPtr<FUICommandList> FNavigationToolView::GetBaseCommandList() const
@@ -747,11 +754,11 @@ void FNavigationToolView::NotifyItemSelectionChanged(const TArray<FNavigationToo
 		return;
 	}
 	TGuardValue<bool> Guard(bSyncingItemSelection, true);
-	
+
 	SelectedItems            = InSelectedItems;
 	SortedSelectedItems      = InSelectedItems;
 	NextSelectedItemIntoView = -1;
-	
+
 	FNavigationTool::SortItems(SortedSelectedItems);
 
 	//If we have Pending Items Remaining but we switched Selection via Navigation, treat it as "I want to rename this too"
@@ -759,11 +766,15 @@ void FNavigationToolView::NotifyItemSelectionChanged(const TArray<FNavigationToo
 	{
 		bRequestedRename = true;
 	}
-	
+
 	const TSharedPtr<FNavigationTool> Tool = WeakTool.Pin();
 	if (bUpdateModeTools && Tool.IsValid())
 	{
-		Tool->SyncSequencerSelection(SelectedItems);
+		if (ToolSettings->ShouldSyncSelectionToSequencer())
+		{
+			Tool->SyncSequencerSelection(SelectedItems);
+		}
+
 		Tool->SelectItems(SelectedItems, ENavigationToolItemSelectionFlags::None);
 	}
 }
@@ -2151,21 +2162,65 @@ void FNavigationToolView::ApplyCustomColumnView(const FName InColumnViewName)
 
 bool FNavigationToolView::CanFocusSingleSelection() const
 {
-	if (SelectedItems.Num() != 1 || !GetOwnerTool().IsValid())
+	if (SelectedItems.Num() == 1)
+	{
+		const FNavigationToolSequence* const SequenceItem = SelectedItems[0]->CastTo<FNavigationToolSequence>();
+		return SequenceItem && IsValid(SequenceItem->GetSequence());
+	}
+	return false;
+}
+
+void FNavigationToolView::FocusSingleSelection()
+{
+	if (SelectedItems.Num() != 1)
+	{
+		return;
+	}
+
+	const FNavigationToolSequence* const SequenceItem = SelectedItems[0]->CastTo<FNavigationToolSequence>();
+	if (!SequenceItem)
+	{
+		return;
+	}
+
+	UMovieSceneSequence* const Sequence = SequenceItem->GetSequence();
+	if (!IsValid(Sequence))
+	{
+		return;
+	}
+
+	const TSharedPtr<INavigationTool> Tool = GetOwnerTool();
+	if (!Tool.IsValid())
+	{
+		return;
+	}
+
+	FocusSequence(*Tool, *Sequence);
+}
+
+bool FNavigationToolView::CanFocusInContentBrowser() const
+{
+	if (SelectedItems.Num() != 1)
 	{
 		return false;
 	}
 
 	const FNavigationToolSequence* const SequenceItem = SelectedItems[0]->CastTo<FNavigationToolSequence>();
-	if (!SequenceItem || !IsValid(SequenceItem->GetSequence()))
+	if (!SequenceItem)
 	{
 		return false;
 	}
 
-	return true;
+	UMovieSceneSequence* const Sequence = SequenceItem->GetSequence();
+	if (!IsValid(Sequence))
+	{
+		return false;
+	}
+
+	return Cast<ULevelSequence>(Sequence) != nullptr;
 }
 
-void FNavigationToolView::FocusSingleSelection()
+void FNavigationToolView::FocusInContentBrowser()
 {
 	if (SelectedItems.Num() != 1)
 	{
@@ -2184,13 +2239,16 @@ void FNavigationToolView::FocusSingleSelection()
 		return;
 	}
 
-	const TSharedPtr<INavigationTool> Tool = GetOwnerTool();
-	if (!Tool.IsValid())
+	ULevelSequence* const LevelSequence = Cast<ULevelSequence>(Sequence);
+	if (!IsValid(LevelSequence))
 	{
 		return;
 	}
 
-	FocusSequence(*Tool, *Sequence);
+	const TArray<UObject*> ObjectsToSync = { LevelSequence };
+
+	FContentBrowserModule& ContentBrowserModule = FModuleManager::Get().LoadModuleChecked<FContentBrowserModule>(TEXT("ContentBrowser"));
+	ContentBrowserModule.Get().SyncBrowserToAssets(ObjectsToSync);
 }
 
 bool FNavigationToolView::UpdateFilters()
diff --git a/Engine/Source/Editor/Sequencer/Private/Tools/NavigationTool/NavigationToolView.h b/Engine/Source/Editor/Sequencer/Private/Tools/NavigationTool/NavigationToolView.h
--- a/Engine/Source/Editor/Sequencer/Private/Tools/NavigationTool/NavigationToolView.h
+++ b/Engine/Source/Editor/Sequencer/Private/Tools/NavigationTool/NavigationToolView.h
@@ -359,6 +359,9 @@ protected:
 	bool CanFocusSingleSelection() const;
 	void FocusSingleSelection();
 
+	bool CanFocusInContentBrowser() const;
+	void FocusInContentBrowser();
+
 	bool UpdateFilters();
 
 	/** Local Identifier of this Instance */
diff --git a/Engine/Source/Editor/Sequencer/Private/Tools/NavigationTool/Utils/NavigationToolMiscUtils.cpp b/Engine/Source/Editor/Sequencer/Private/Tools/NavigationTool/Utils/NavigationToolMiscUtils.cpp
--- a/Engine/Source/Editor/Sequencer/Private/Tools/NavigationTool/Utils/NavigationToolMiscUtils.cpp
+++ b/Engine/Source/Editor/Sequencer/Private/Tools/NavigationTool/Utils/NavigationToolMiscUtils.cpp
@@ -1,12 +1,18 @@
 ﻿// Copyright Epic Games, Inc. All Rights Reserved.
 
 #include "NavigationToolMiscUtils.h"
-#include "ISequencer.h"
 #include "ISourceControlModule.h"
 #include "ISourceControlProvider.h"
-#include "Evaluation/MovieSceneEvaluationTemplateInstance.h"
+#include "MVVM/ObjectBindingModelStorageExtension.h"
+#include "MVVM/Selection/Selection.h"
+#include "MVVM/ViewModels/ObjectBindingModel.h"
+#include "MVVM/Views/SOutlinerView.h"
 #include "Sections/MovieSceneSubSection.h"
+#include "SSequencer.h"
 #include "Tools/NavigationTool/INavigationTool.h"
+#include "Tools/NavigationTool/Items/NavigationToolBinding.h"
+#include "Tools/NavigationTool/Items/NavigationToolSequence.h"
+#include "Tools/NavigationTool/Items/NavigationToolTrack.h"
 #include "UObject/Package.h"
 
 namespace UE::Sequencer::NavigationTool
@@ -56,6 +62,12 @@ FText FindSourceControlStatusText(UPackage* const InPackage)
 	return FText::GetEmpty();
 }
 
+TSharedPtr<SOutlinerView> GetOutlinerViewWidget(const ISequencer& InSequencer)
+{
+	const TSharedPtr<SSequencer> SequencerWidget = StaticCastSharedRef<SSequencer>(InSequencer.GetSequencerWidget());
+	return SequencerWidget.IsValid() ? SequencerWidget->GetTreeView() : nullptr;
+}
+
 void FocusSequence(const INavigationTool& InTool, UMovieSceneSequence& InSequence)
 {
 	const TSharedPtr<ISequencer> Sequencer = InTool.GetSequencer();
@@ -64,6 +76,17 @@ void FocusSequence(const INavigationTool& InTool, UMovieSceneSequence& InSequenc
 		return;
 	}
 
+	if (Sequencer->GetRootMovieSceneSequence() == &InSequence)
+	{
+		Sequencer->ResetToNewRootSequence(InSequence);
+		return;
+	}
+
+	if (Sequencer->GetFocusedMovieSceneSequence() == &InSequence)
+	{
+		return;
+	}
+
 	FMovieSceneEvaluationState* const EvaluationState = Sequencer->GetEvaluationState();
 	if (!EvaluationState)
 	{
@@ -71,15 +94,144 @@ void FocusSequence(const INavigationTool& InTool, UMovieSceneSequence& InSequenc
 	}
 
 	FMovieSceneSequenceIDRef SequenceID = EvaluationState->FindSequenceId(&InSequence);
-
 	UMovieSceneSubSection* const SubSection = Sequencer->FindSubSection(SequenceID);
-	if (IsValid(SubSection))
+	if (!IsValid(SubSection))
 	{
-		Sequencer->FocusSequenceInstance(*SubSection);
+		return;
 	}
+
+	Sequencer->FocusSequenceInstance(*SubSection);
 }
 
-FMovieSceneSequenceID ResolveSequenceID(ISequencer& InSequencer, UMovieSceneSequence* const InSequence)
+void FocusSequence(const INavigationTool& InTool
+	, UMovieSceneSequence& InSequence
+	, const FNavigationToolSequence& InSequenceItem)
+{
+	FocusSequence(InTool, InSequence);
+
+	if (const TSharedPtr<ISequencer> Sequencer = InTool.GetSequencer())
+	{
+		UMovieSceneSubSection* const SubSection = InSequenceItem.GetSubSection();
+		if (IsValid(SubSection))
+		{
+			Sequencer->EmptySelection();
+			Sequencer->SelectSection(SubSection);
+		}
+	}
+}
+
+void FocusSequence(const INavigationTool& InTool
+	, UMovieSceneSequence& InSequence
+	, const FNavigationToolTrack& InTrackItem)
+{
+	FocusSequence(InTool, InSequence);
+
+	if (const TSharedPtr<ISequencer> Sequencer = InTool.GetSequencer())
+	{
+		UMovieSceneTrack* const Track = InTrackItem.GetTrack();
+		if (IsValid(Track))
+		{
+			Sequencer->EmptySelection();
+			Sequencer->SelectTrack(Track);
+
+			if (const TSharedPtr<SOutlinerView> OutlinerView = GetOutlinerViewWidget(*Sequencer))
+			{
+				if (const TViewModelPtr<FTrackModel> TrackViewModel = InTrackItem.GetViewModel())
+				{
+					OutlinerView->ExpandCollapseNode(TrackViewModel, true, ETreeRecursion::NonRecursive);
+					OutlinerView->RequestScrollIntoView(TrackViewModel);
+				}
+			}
+		}
+	}
+}
+
+void FocusSequence(const INavigationTool& InTool
+	, UMovieSceneSequence& InSequence
+	, const FNavigationToolBinding& InBindingItem)
+{
+	FocusSequence(InTool, InSequence);
+
+	if (const TSharedPtr<ISequencer> Sequencer = InTool.GetSequencer())
+	{
+		SelectSequencerBindingTrack(*Sequencer, InBindingItem.GetBinding().GetObjectGuid());
+
+		if (const TSharedPtr<SOutlinerView> OutlinerView = GetOutlinerViewWidget(*Sequencer))
+		{
+			if (const TViewModelPtr<FObjectBindingModel> ObjectBindingViewModel = InBindingItem.GetViewModel())
+			{
+				OutlinerView->ExpandCollapseNode(ObjectBindingViewModel, true, ETreeRecursion::NonRecursive);
+				OutlinerView->RequestScrollIntoView(ObjectBindingViewModel);
+			}
+		}
+	}
+}
+
+void SelectSequencerBindingTrack(const ISequencer& InSequencer, const FGuid& InObjectId)
+{
+	const TSharedPtr<FSequencerEditorViewModel> ViewModel = InSequencer.GetViewModel();
+	if (!ViewModel.IsValid())
+	{
+		return;
+	}
+
+	const TSharedPtr<FSequencerSelection> Selection = ViewModel->GetSelection();
+	if (!Selection.IsValid())
+	{
+		return;
+	}
+
+	const FViewModelPtr RootViewModel = ViewModel->GetRootModel();
+	if (!Selection.IsValid())
+	{
+		return;
+	}
+
+	const FObjectBindingModelStorageExtension* const ObjectBindingStorage = RootViewModel->CastDynamic<FObjectBindingModelStorageExtension>();
+	if (!ObjectBindingStorage)
+	{
+		return;
+	}
+
+	const TSharedPtr<FObjectBindingModel> ObjectBindingModel = ObjectBindingStorage->FindModelForObjectBinding(InObjectId);
+	if (!ObjectBindingModel.IsValid())
+	{
+		return;
+	}
+
+	Selection->Empty();
+	Selection->Outliner.Select(ObjectBindingModel);
+}
+
+void FocusItemInSequencer(const INavigationTool& InTool, const FNavigationToolItemPtr& InItem)
+{
+	if (const FNavigationToolSequence* const SequenceItem = InItem->CastTo<FNavigationToolSequence>())
+	{
+		UMovieSceneSequence* const ParentSequence = SequenceItem->GetSequence();
+		if (IsValid(ParentSequence))
+		{
+			FocusSequence(InTool, *ParentSequence);
+		}
+	}
+	else if (const FNavigationToolTrack* const TrackItem = InItem->CastTo<FNavigationToolTrack>())
+	{
+		UMovieSceneSequence* const Sequence = TrackItem->GetSequence();
+		if (IsValid(Sequence))
+		{
+			FocusSequence(InTool, *Sequence, *TrackItem);
+		}
+	}
+	else if (const FNavigationToolBinding* const BindingItem = InItem->CastTo<FNavigationToolBinding>())
+	{
+		UMovieSceneSequence* const Sequence = BindingItem->GetSequence();
+		if (IsValid(Sequence))
+		{
+			FocusSequence(InTool, *Sequence, *BindingItem);
+		}
+	}
+}
+
+FMovieSceneSequenceID ResolveSequenceID(const ISequencer& InSequencer, UMovieSceneSequence* const InSequence)
 {
 	if (FMovieSceneEvaluationState* const EvaluationState = InSequencer.GetSharedPlaybackState()->FindCapability<FMovieSceneEvaluationState>())
 	{
@@ -88,14 +240,17 @@ FMovieSceneSequenceID ResolveSequenceID(ISequencer& InSequencer, UMovieSceneSequ
 	return FMovieSceneSequenceID();
 }
 
-TArrayView<TWeakObjectPtr<>> ResolveBoundObjects(ISequencer& InSequencer
+TArrayView<TWeakObjectPtr<>> ResolveBoundObjects(const ISequencer& InSequencer
 	, UMovieSceneSequence* const InSequence, const FGuid& InBindingId)
 {
-	if (FMovieSceneEvaluationState* const EvaluationState = InSequencer.GetSharedPlaybackState()->FindCapability<FMovieSceneEvaluationState>())
+	const TSharedRef<const MovieScene::FSharedPlaybackState> SharedPlaybackState = InSequencer.GetSharedPlaybackState();
+
+	if (FMovieSceneEvaluationState* const EvaluationState = SharedPlaybackState->FindCapability<FMovieSceneEvaluationState>())
 	{
 		FMovieSceneSequenceIDRef SequenceID = EvaluationState->FindSequenceId(InSequence);
-		return InSequencer.FindBoundObjects(InBindingId, SequenceID);
+		return EvaluationState->FindBoundObjects(InBindingId, SequenceID, SharedPlaybackState);
 	}
+
 	return TArrayView<TWeakObjectPtr<>>();
 }
 
diff --git a/Engine/Source/Editor/Sequencer/Private/Tools/NavigationTool/Utils/NavigationToolMiscUtils.h b/Engine/Source/Editor/Sequencer/Private/Tools/NavigationTool/Utils/NavigationToolMiscUtils.h
--- a/Engine/Source/Editor/Sequencer/Private/Tools/NavigationTool/Utils/NavigationToolMiscUtils.h
+++ b/Engine/Source/Editor/Sequencer/Private/Tools/NavigationTool/Utils/NavigationToolMiscUtils.h
@@ -2,8 +2,13 @@
 
 #pragma once
 
+#include "ISequencer.h"
 #include "ISourceControlState.h"
+#include "Tools/NavigationTool/NavigationToolDefines.h"
 
+class FNavigationToolBinding;
+class FNavigationToolSequence;
+class FNavigationToolTrack;
 class FText;
 class ISequencer;
 class UMovieSceneSequence;
@@ -12,6 +17,11 @@ struct FGuid;
 struct FMovieSceneSequenceID;
 struct FSlateBrush;
 
+namespace UE::Sequencer
+{
+class SOutlinerView;
+}
+
 namespace UE::Sequencer::NavigationTool
 {
 
@@ -38,16 +48,58 @@ const FSlateBrush* FindSourceControlStatusBrush(UPackage* const InPackage);
  */
 FText FindSourceControlStatusText(UPackage* const InPackage);
 
-/** Focuses a sub-movie scene (MovieScene within a MovieScene) in the sequencer. */
+TSharedPtr<SOutlinerView> GetOutlinerViewWidget(const ISequencer& InSequencer);
+
+/**
+ * Focuses a sub-movie scene (MovieScene within a MovieScene) in the sequencer.
+ * @param InTool The Navigation Tool that manages scene sequences.
+ * @param InSequence The movie scene sequence to focus on.
+ */
 void FocusSequence(const INavigationTool& InTool, UMovieSceneSequence& InSequence);
 
+/**
+ * Focuses a sub-movie scene (MovieScene within a MovieScene) in the sequencer.
+ * @param InTool The Navigation Tool that manages scene sequences.
+ * @param InSequence The movie scene sequence to focus on.
+ * @param InSequenceItem Optional sequence item to select in the sequence.
+ */
+void FocusSequence(const INavigationTool& InTool
+	, UMovieSceneSequence& InSequence
+	, const FNavigationToolSequence& InSequenceItem);
+
+/**
+ * Focuses a sub-movie scene (MovieScene within a MovieScene) in the sequencer.
+ * @param InTool The Navigation Tool that manages scene sequences.
+ * @param InSequence The movie scene sequence to focus on.
+ * @param InTrackItem Optional track item to select in the sequence.
+ */
+void FocusSequence(const INavigationTool& InTool
+	, UMovieSceneSequence& InSequence
+	, const FNavigationToolTrack& InTrackItem);
+
+/**
+ * Focuses a sub-movie scene (MovieScene within a MovieScene) in the sequencer.
+ * @param InTool The Navigation Tool that manages scene sequences.
+ * @param InSequence The movie scene sequence to focus on.
+ * @param InBindingItem Optional binding item to select in the sequence.
+ */
+void FocusSequence(const INavigationTool& InTool
+	, UMovieSceneSequence& InSequence
+	, const FNavigationToolBinding& InBindingItem);
+
+/** Selects an object binding track in a sequencer. */
+void SelectSequencerBindingTrack(const ISequencer& InSequencer, const FGuid& InObjectId);
+
+void FocusItemInSequencer(const INavigationTool& InTool
+	, const FNavigationToolItemPtr& InItem);
+
 /**
  * Resolve the sequence ID for a given sequence within the provided sequencer.
  * @param InSequencer The sequencer instance to retrieve the sequence ID from.
  * @param InSequence The sequence whose ID is to be resolved.
  * @return The resolved sequence ID if found, otherwise an empty sequence ID.
  */
-FMovieSceneSequenceID ResolveSequenceID(ISequencer& InSequencer, UMovieSceneSequence* const InSequence);
+FMovieSceneSequenceID ResolveSequenceID(const ISequencer& InSequencer, UMovieSceneSequence* const InSequence);
 
 /**
  * Resolves the objects bound to a specified binding ID within a sequence.
@@ -56,7 +108,7 @@ FMovieSceneSequenceID ResolveSequenceID(ISequencer& InSequencer, UMovieSceneSequ
  * @param InBindingId The binding ID whose bound objects are being resolved.
  * @return A view of the objects bound to the specified binding ID in the specified sequence.
  */
-TArrayView<TWeakObjectPtr<>> ResolveBoundObjects(ISequencer& InSequencer
+TArrayView<TWeakObjectPtr<>> ResolveBoundObjects(const ISequencer& InSequencer
 	, UMovieSceneSequence* const InSequence, const FGuid& InBindingId);
 
 }
diff --git a/Engine/Source/Editor/Sequencer/Private/Tools/NavigationTool/Utils/NavigationToolMovieSceneUtils.cpp b/Engine/Source/Editor/Sequencer/Private/Tools/NavigationTool/Utils/NavigationToolMovieSceneUtils.cpp
--- a/Engine/Source/Editor/Sequencer/Private/Tools/NavigationTool/Utils/NavigationToolMovieSceneUtils.cpp
+++ b/Engine/Source/Editor/Sequencer/Private/Tools/NavigationTool/Utils/NavigationToolMovieSceneUtils.cpp
@@ -108,17 +108,4 @@ void GetSequenceSubSections(UMovieSceneSequence* const InSequence, TArray<UMovie
 	}
 }
 
-const UClass* GetBindingObjectClass(UMovieScene& InMovieScene, const FGuid& InBindingGuid)
-{
-	if (const FMovieScenePossessable* const Possessable = InMovieScene.FindPossessable(InBindingGuid))
-	{
-		return Possessable->GetPossessedObjectClass();
-	}
-	if (const FMovieSceneSpawnable* const Spawnable = InMovieScene.FindSpawnable(InBindingGuid))
-	{
-		return Spawnable->GetObjectTemplate()->GetClass();
-	}
-	return nullptr;
-}
-
 }
diff --git a/Engine/Source/Editor/Sequencer/Private/Tools/NavigationTool/Utils/NavigationToolMovieSceneUtils.h b/Engine/Source/Editor/Sequencer/Private/Tools/NavigationTool/Utils/NavigationToolMovieSceneUtils.h
--- a/Engine/Source/Editor/Sequencer/Private/Tools/NavigationTool/Utils/NavigationToolMovieSceneUtils.h
+++ b/Engine/Source/Editor/Sequencer/Private/Tools/NavigationTool/Utils/NavigationToolMovieSceneUtils.h
@@ -5,11 +5,9 @@
 #include "Containers/Array.h"
 
 class ISequencer;
-class UClass;
 class UMovieScene;
 class UMovieSceneSequence;
 class UMovieSceneSubSection;
-struct FGuid;
 
 namespace UE::Sequencer::NavigationTool
 {
@@ -51,7 +49,4 @@ void ModifySequenceAndMovieScene(UMovieSceneSequence* const InSequence);
  */
 void GetSequenceSubSections(UMovieSceneSequence* const InSequence, TArray<UMovieSceneSubSection*>& OutSubSections);
 
-/** Attempts to get the bound object class from a binding */
-const UClass* GetBindingObjectClass(UMovieScene& InMovieScene, const FGuid& InBindingGuid);
-
 }
diff --git a/Engine/Source/Editor/Sequencer/Private/Tools/NavigationTool/Widgets/Columns/SNavigationToolItemChip.cpp b/Engine/Source/Editor/Sequencer/Private/Tools/NavigationTool/Widgets/Columns/SNavigationToolItemChip.cpp
--- a/Engine/Source/Editor/Sequencer/Private/Tools/NavigationTool/Widgets/Columns/SNavigationToolItemChip.cpp
+++ b/Engine/Source/Editor/Sequencer/Private/Tools/NavigationTool/Widgets/Columns/SNavigationToolItemChip.cpp
@@ -11,6 +11,8 @@
 #include "Widgets/Layout/SScaleBox.h"
 #include "Widgets/Views/STableRow.h"
 
+#define LOCTEXT_NAMESPACE "SNavigationToolItemChip"
+
 namespace UE::Sequencer::NavigationTool
 {
 
@@ -44,7 +46,8 @@ void SNavigationToolItemChip::Construct(const FArguments& InArgs
 				]
 			]
 		]);
-	SetToolTipText(TAttribute<FText>(InItem, &INavigationToolItem::GetDisplayName));
+
+	SetToolTipText(TAttribute<FText>::CreateSP(this, &SNavigationToolItemChip::GetTooltipText));
 	SetBorderImage(TAttribute<const FSlateBrush*>(this, &SNavigationToolItemChip::GetItemBackgroundBrush));
 }
 
@@ -138,7 +141,7 @@ void SNavigationToolItemChip::OnMouseCaptureLost(const FCaptureLostEvent& InCapt
 	Release();
 }
 
-FReply SNavigationToolItemChip::OnDragDetected(const FGeometry& InGeometry, const FPointerEvent& InMouseEvent)
+/*FReply SNavigationToolItemChip::OnDragDetected(const FGeometry& InGeometry, const FPointerEvent& InMouseEvent)
 {
 	const FNavigationToolItemPtr Item = WeakItem.Pin();
 	const TSharedPtr<INavigationToolView> ToolView = WeakView.Pin();
@@ -202,7 +205,7 @@ FReply SNavigationToolItemChip::OnDrop(const FGeometry& InGeometry, const FDragD
 void SNavigationToolItemChip::OnDragLeave(const FDragDropEvent& InDragDropEvent)
 {
 	ItemDropZone = TOptional<EItemDropZone>();
-}
+}*/
 
 bool SNavigationToolItemChip::IsSelected() const
 {
@@ -274,4 +277,22 @@ FSlateColor SNavigationToolItemChip::GetIconColor() const
 	return IsHovered() ? FStyleColors::ForegroundHover : FStyleColors::Foreground;
 }
 
+FText SNavigationToolItemChip::GetTooltipText() const
+{
+	const FNavigationToolItemPtr Item = WeakItem.Pin();
+	if (!Item.IsValid())
+	{
+		return FText::GetEmpty();
+	}
+
+	const FText ItemDisplayName = Item->GetDisplayName();
+
+	return FText::Format(LOCTEXT("TooltipText", "{0}\n\n"
+		"Click to select in Sequencer\n\n"
+		"Alt + Click to select item in Navigation Tool tree")
+		, ItemDisplayName);
+}
+
 }
+
+#undef LOCTEXT_NAMESPACE
diff --git a/Engine/Source/Editor/Sequencer/Private/Tools/NavigationTool/Widgets/Columns/SNavigationToolItemChip.h b/Engine/Source/Editor/Sequencer/Private/Tools/NavigationTool/Widgets/Columns/SNavigationToolItemChip.h
--- a/Engine/Source/Editor/Sequencer/Private/Tools/NavigationTool/Widgets/Columns/SNavigationToolItemChip.h
+++ b/Engine/Source/Editor/Sequencer/Private/Tools/NavigationTool/Widgets/Columns/SNavigationToolItemChip.h
@@ -44,10 +44,10 @@ public:
 	virtual FReply OnMouseButtonDown(const FGeometry& InGeometry, const FPointerEvent& InMouseEvent) override;
 	virtual FReply OnMouseButtonUp(const FGeometry& InGeometry, const FPointerEvent& InMouseEvent) override;
 	virtual void OnMouseCaptureLost(const FCaptureLostEvent& InCaptureLostEvent) override;
-	virtual FReply OnDragDetected(const FGeometry& InGeometry, const FPointerEvent& InMouseEvent) override;
+	/*virtual FReply OnDragDetected(const FGeometry& InGeometry, const FPointerEvent& InMouseEvent) override;
 	virtual FReply OnDragOver(const FGeometry& InGeometry, const FDragDropEvent& InDragDropEvent) override;
 	virtual FReply OnDrop(const FGeometry& InGeometry, const FDragDropEvent& InDragDropEvent) override;
-	virtual void OnDragLeave(const FDragDropEvent& InDragDropEvent) override;
+	virtual void OnDragLeave(const FDragDropEvent& InDragDropEvent) override;*/
 	//~ End SWidget
 
 	bool IsSelected() const;
@@ -63,6 +63,8 @@ private:
 
 	FSlateColor GetIconColor() const;
 
+	FText GetTooltipText() const;
+
 	FNavigationToolItemWeakPtr WeakItem;
 
 	TWeakPtr<INavigationToolView> WeakView;
diff --git a/Engine/Source/Editor/Sequencer/Private/Tools/NavigationTool/Widgets/Columns/SNavigationToolItemList.cpp b/Engine/Source/Editor/Sequencer/Private/Tools/NavigationTool/Widgets/Columns/SNavigationToolItemList.cpp
--- a/Engine/Source/Editor/Sequencer/Private/Tools/NavigationTool/Widgets/Columns/SNavigationToolItemList.cpp
+++ b/Engine/Source/Editor/Sequencer/Private/Tools/NavigationTool/Widgets/Columns/SNavigationToolItemList.cpp
@@ -3,6 +3,9 @@
 #include "SNavigationToolItemList.h"
 #include "Tools/NavigationTool/Items/NavigationToolItem.h"
 #include "Tools/NavigationTool/NavigationToolView.h"
+#include "Tools/NavigationTool/Items/NavigationToolBinding.h"
+#include "Tools/NavigationTool/Items/NavigationToolSequence.h"
+#include "Tools/NavigationTool/Utils/NavigationToolMiscUtils.h"
 #include "Tools/NavigationTool/Widgets/Columns/SNavigationToolItemChip.h"
 #include "Tools/NavigationTool/Widgets/SNavigationToolTreeRow.h"
 #include "Widgets/Layout/SScrollBox.h"
@@ -129,8 +132,8 @@ void SNavigationToolItemList::Refresh()
 			.Padding(0.f, 1.f)
 			[
 				SNew(SNavigationToolItemChip, Item.ToSharedRef(), ToolView)
-				.OnItemChipClicked(this, &SNavigationToolItemList::OnItemChipSelected)
 				.ChipStyle(TreeRow->GetStyle())
+				.OnItemChipClicked(this, &SNavigationToolItemList::OnItemChipSelected)
 				.OnValidDragOver(this, &SNavigationToolItemList::OnItemChipValidDragOver)
 			];
 	}
@@ -146,74 +149,19 @@ FReply SNavigationToolItemList::OnItemChipSelected(const FNavigationToolItemPtr&
 		return FReply::Unhandled();
 	}
 
-	TArray<FNavigationToolItemPtr> SelectedItems = ToolView->GetSelectedItems();
-
-	ENavigationToolItemSelectionFlags SelectionFlags = ENavigationToolItemSelectionFlags::SignalSelectionChange;
-
-	if (InMouseEvent.IsControlDown())
+	if (InMouseEvent.IsAltDown())
 	{
-		// Deselect if the Item in question is already selected (i.e. CTRL also behaves like a Toggle)
-		if (ToolView->IsItemSelected(InItem))
-		{
-			SelectedItems.Remove(InItem);
-			ToolView->SelectItems(MoveTemp(SelectedItems), SelectionFlags);
-			return FReply::Handled();
-		}
+		constexpr ENavigationToolItemSelectionFlags SelectionFlags = ENavigationToolItemSelectionFlags::SignalSelectionChange;
+		ToolView->SelectItems({ InItem }, SelectionFlags);
 
-		SelectionFlags |= ENavigationToolItemSelectionFlags::AppendToCurrentSelection;
+		return FReply::Handled();
 	}
 
-	// add the given item since it will be selected regardless if Shift (Range Selection) is pressed or not
-	TArray<FNavigationToolItemPtr> ItemsToSelect { InItem };
-
-	if (InMouseEvent.IsShiftDown() && !SelectedItems.IsEmpty())
+	if (const TSharedPtr<INavigationTool> Tool = ToolView->GetOwnerTool())
 	{
-		// Similar to Control, Shift should always append Selections, never remove
-		SelectionFlags |= ENavigationToolItemSelectionFlags::AppendToCurrentSelection;
-
-		int32 StartIndex = INDEX_NONE;
-
-		// Find whether there's a Child in the Selected Item List.
-		// Iterate in Reverse order since we want more recent Selections to take precedence
-		for (int32 ItemIndex = SelectedItems.Num() - 1; ItemIndex >= 0; --ItemIndex)
-		{
-			StartIndex = WeakChildItems.Find(SelectedItems[ItemIndex]);
-			if (StartIndex != INDEX_NONE)
-			{
-				break;
-			}
-		}
-
-		// If we manage to find an already selected child, start the range selection
-		if (StartIndex != INDEX_NONE)
-		{
-			int32 TargetIndex = WeakChildItems.Find(InItem);
-
-			StartIndex  = FMath::Clamp(StartIndex, 0, WeakChildItems.Num());
-			TargetIndex = FMath::Clamp(TargetIndex, 0, WeakChildItems.Num());
-
-			// If Target Index comes before, just iterate normally, and reverse later
-			if (TargetIndex < StartIndex)
-			{
-				Swap(StartIndex, TargetIndex);
-			}
-
-			// Empty List as we're going to add the InItem anyway through the List Iteration
-			ItemsToSelect.Empty(TargetIndex - StartIndex);
-
-			for (int32 ItemIndex = StartIndex; ItemIndex <= TargetIndex; ++ItemIndex)
-			{
-				ItemsToSelect.Add(WeakChildItems[ItemIndex].Pin());
-			}
-
-			if (TargetIndex < StartIndex)
-			{
-				Algo::Reverse(ItemsToSelect);
-			}
-		}
+		FocusItemInSequencer(*Tool, InItem);
 	}
 
-	ToolView->SelectItems(MoveTemp(ItemsToSelect), SelectionFlags);
 	return FReply::Handled();
 }
 
diff --git a/Engine/Source/Editor/Sequencer/Private/Tools/NavigationTool/Widgets/Columns/SNavigationToolLabelItem.cpp b/Engine/Source/Editor/Sequencer/Private/Tools/NavigationTool/Widgets/Columns/SNavigationToolLabelItem.cpp
--- a/Engine/Source/Editor/Sequencer/Private/Tools/NavigationTool/Widgets/Columns/SNavigationToolLabelItem.cpp
+++ b/Engine/Source/Editor/Sequencer/Private/Tools/NavigationTool/Widgets/Columns/SNavigationToolLabelItem.cpp
@@ -2,11 +2,10 @@
 
 #include "SNavigationToolLabelItem.h"
 #include "ActorEditorUtils.h"
-#include "Tools/NavigationTool/NavigationToolSettings.h"
+#include "Tools/NavigationTool/INavigationTool.h"
 #include "Tools/NavigationTool/Items/INavigationToolItem.h"
-#include "Tools/NavigationTool/NavigationToolView.h"
-#include "Tools/NavigationTool/Items/NavigationToolItemUtils.h"
 #include "Tools/NavigationTool/Items/NavigationToolSequence.h"
+#include "Tools/NavigationTool/NavigationToolView.h"
 #include "Tools/NavigationTool/Widgets/SNavigationToolExpanderArrow.h"
 #include "Tools/NavigationTool/Widgets/SNavigationToolTreeRow.h"
 #include "Widgets/Images/SImage.h"
@@ -135,19 +134,6 @@ FText SNavigationToolLabelItem::GetItemText() const
 		return Item->GetDisplayName();
 	}
 
-	const TSharedPtr<INavigationToolView> ToolView = WeakView.Pin();
-	if (!ToolView.IsValid())
-	{
-		return FText::GetEmpty();
-	}
-
-	if (FNavigationToolSequence* const SequenceItem = Item->CastTo<FNavigationToolSequence>())
-	{
-		const TSharedRef<FNavigationToolSequence> SequenceItemRef = StaticCastSharedRef<FNavigationToolSequence>(SequenceItem->AsShared());
-		const UNavigationToolSettings* const ToolSettings = GetDefault<UNavigationToolSettings>();
-		return SequenceDisplayNameWithModifyAndInactiveStatus(SequenceItemRef, ToolSettings->ShouldUseShortNames(), true);
-	}
-
 	return Item->GetDisplayName();
 }
 
diff --git a/Engine/Source/Editor/Sequencer/Private/Tools/NavigationTool/Widgets/Columns/SNavigationToolLabelItem.h b/Engine/Source/Editor/Sequencer/Private/Tools/NavigationTool/Widgets/Columns/SNavigationToolLabelItem.h
--- a/Engine/Source/Editor/Sequencer/Private/Tools/NavigationTool/Widgets/Columns/SNavigationToolLabelItem.h
+++ b/Engine/Source/Editor/Sequencer/Private/Tools/NavigationTool/Widgets/Columns/SNavigationToolLabelItem.h
@@ -14,7 +14,6 @@ namespace UE::Sequencer::NavigationTool
 class FNavigationToolItem;
 class FNavigationTool;
 class INavigationToolView;
-class SNavigationToolView;
 class SNavigationToolTreeRow;
 
 class SNavigationToolLabelItem : public SCompoundWidget
diff --git a/Engine/Source/Editor/Sequencer/Private/Tools/NavigationTool/Widgets/Columns/SNavigationToolPlayhead.cpp b/Engine/Source/Editor/Sequencer/Private/Tools/NavigationTool/Widgets/Columns/SNavigationToolPlayhead.cpp
new file mode 100644
--- /dev/null
+++ b/Engine/Source/Editor/Sequencer/Private/Tools/NavigationTool/Widgets/Columns/SNavigationToolPlayhead.cpp
@@ -0,0 +1,99 @@
+﻿// Copyright Epic Games, Inc. All Rights Reserved.
+
+#include "SNavigationToolPlayhead.h"
+#include "Editor.h"
+#include "Styling/StyleColors.h"
+#include "Tools/NavigationTool/INavigationTool.h"
+#include "Tools/NavigationTool/Columns/NavigationToolPlayheadColumn.h"
+#include "Tools/NavigationTool/Items/INavigationToolItem.h"
+#include "Tools/NavigationTool/Items/NavigationToolItemUtils.h"
+#include "Tools/NavigationTool/NavigationToolView.h"
+#include "Tools/NavigationTool/Widgets/SNavigationToolTreeRow.h"
+
+#define LOCTEXT_NAMESPACE "SNavigationToolPlayhead"
+
+namespace UE::Sequencer::NavigationTool
+{
+
+EItemContainsPlayhead GetItemContainsPlayhead(const FNavigationToolItemPtr& InItem)
+{
+	if (!InItem.IsValid())
+	{
+		return EItemContainsPlayhead::None;
+	}
+
+	IPlayheadExtension* const PlayheadItem = InItem->CastTo<IPlayheadExtension>();
+	if (!PlayheadItem)
+	{
+		return EItemContainsPlayhead::None;
+	}
+
+	return PlayheadItem->ContainsPlayhead();
+}
+
+void SNavigationToolPlayhead::Construct(const FArguments& InArgs
+	, const TSharedRef<FNavigationToolPlayheadColumn>& InColumn
+	, const FNavigationToolItemPtr& InItem
+	, const TSharedRef<INavigationToolView>& InView
+	, const TSharedRef<SNavigationToolTreeRow>& InRowWidget)
+{
+	WeakColumn = InColumn;
+	WeakItem = InItem;
+	WeakView = InView;
+	WeakRowWidget = InRowWidget;
+
+	SImage::Construct(SImage::FArguments()
+		.IsEnabled(this, &SNavigationToolPlayhead::IsVisibilityWidgetEnabled)
+		.ColorAndOpacity(this, &SNavigationToolPlayhead::GetForegroundColor)
+		.Image(this, &SNavigationToolPlayhead::GetBrush));
+}
+
+const FSlateBrush* SNavigationToolPlayhead::GetBrush() const
+{
+	return FAppStyle::GetBrush(TEXT("GenericPlay"));
+}
+
+FSlateColor SNavigationToolPlayhead::GetForegroundColor() const
+{
+	const bool bIsItemHovered = WeakRowWidget.IsValid() && WeakRowWidget.Pin()->IsHovered();
+
+	if (IsHovered() || bIsItemHovered)
+	{
+		switch (GetContainsPlayhead())
+		{
+		case EItemContainsPlayhead::None:
+			return FStyleColors::White25;
+
+		case EItemContainsPlayhead::PartiallyContainsPlayhead:
+			return FStyleColors::ForegroundHover.GetSpecifiedColor() * FStyleColors::AccentGreen.GetSpecifiedColor();
+
+		case EItemContainsPlayhead::ContainsPlayhead:
+			return FStyleColors::ForegroundHover.GetSpecifiedColor() * FStyleColors::AccentGreen.GetSpecifiedColor();
+		}
+	}
+	else
+	{
+		switch (GetContainsPlayhead())
+		{
+		case EItemContainsPlayhead::None:
+			return FStyleColors::Transparent;
+
+		case EItemContainsPlayhead::PartiallyContainsPlayhead:
+			return FStyleColors::White25.GetSpecifiedColor() * FStyleColors::AccentGreen.GetSpecifiedColor();
+
+		case EItemContainsPlayhead::ContainsPlayhead:
+			return FStyleColors::Foreground.GetSpecifiedColor() * FStyleColors::AccentGreen.GetSpecifiedColor();
+		}
+	}
+
+	return FStyleColors::Transparent;
+}
+
+EItemContainsPlayhead SNavigationToolPlayhead::GetContainsPlayhead() const
+{
+	return GetItemContainsPlayhead(WeakItem.Pin());
+}
+
+}
+
+#undef LOCTEXT_NAMESPACE
diff --git a/Engine/Source/Editor/Sequencer/Private/Tools/NavigationTool/Widgets/Columns/SNavigationToolPlayhead.h b/Engine/Source/Editor/Sequencer/Private/Tools/NavigationTool/Widgets/Columns/SNavigationToolPlayhead.h
new file mode 100644
--- /dev/null
+++ b/Engine/Source/Editor/Sequencer/Private/Tools/NavigationTool/Widgets/Columns/SNavigationToolPlayhead.h
@@ -0,0 +1,48 @@
+﻿// Copyright Epic Games, Inc. All Rights Reserved.
+
+#pragma once
+
+#include "Tools/NavigationTool/Extensions/IPlayheadExtension.h"
+#include "Tools/NavigationTool/NavigationToolDefines.h"
+#include "Widgets/Images/SImage.h"
+
+namespace UE::Sequencer::NavigationTool
+{
+
+class FNavigationToolPlayheadColumn;
+class INavigationToolView;
+class SNavigationToolTreeRow;
+
+/** Widget responsible for showing whether an items range contains the playhead */
+class SNavigationToolPlayhead : public SImage
+{
+public:
+	SLATE_BEGIN_ARGS(SNavigationToolPlayhead) {}
+	SLATE_END_ARGS()
+
+	void Construct(const FArguments& InArgs
+		, const TSharedRef<FNavigationToolPlayheadColumn>& InColumn
+		, const FNavigationToolItemPtr& InItem
+		, const TSharedRef<INavigationToolView>& InView
+		, const TSharedRef<SNavigationToolTreeRow>& InRowWidget);
+
+private:
+	/** Returns whether the widget is enabled or not */
+	virtual bool IsVisibilityWidgetEnabled() const { return true; }
+
+	virtual const FSlateBrush* GetBrush() const;
+
+	virtual FSlateColor GetForegroundColor() const override;
+	
+	EItemContainsPlayhead GetContainsPlayhead() const;
+
+	TWeakPtr<FNavigationToolPlayheadColumn> WeakColumn;
+
+	FNavigationToolItemWeakPtr WeakItem;
+
+	TWeakPtr<INavigationToolView> WeakView;
+
+	TWeakPtr<SNavigationToolTreeRow> WeakRowWidget;
+};
+
+}
diff --git a/Engine/Source/Editor/Sequencer/Private/Tools/NavigationTool/Widgets/Columns/SNavigationToolRevisionControl.cpp b/Engine/Source/Editor/Sequencer/Private/Tools/NavigationTool/Widgets/Columns/SNavigationToolRevisionControl.cpp
--- a/Engine/Source/Editor/Sequencer/Private/Tools/NavigationTool/Widgets/Columns/SNavigationToolRevisionControl.cpp
+++ b/Engine/Source/Editor/Sequencer/Private/Tools/NavigationTool/Widgets/Columns/SNavigationToolRevisionControl.cpp
@@ -15,23 +15,50 @@ namespace UE::Sequencer::NavigationTool
 
 EItemRevisionControlState GetRevisionControlState(const FNavigationToolItemPtr& InItem)
 {
-	if (const IRevisionControlExtension* const RevisionControlExt = InItem->CastTo<IRevisionControlExtension>())
+	if (!InItem.IsValid())
 	{
-		return RevisionControlExt->GetRevisionControlState();
+		return EItemRevisionControlState::None;
+	}
+
+	const IRevisionControlExtension* const RevisionControlItem = InItem->CastTo<IRevisionControlExtension>();
+	if (!RevisionControlItem)
+	{
+		return EItemRevisionControlState::None;
 	}
-	return EItemRevisionControlState::None;
+
+	return RevisionControlItem->GetRevisionControlState();
 }
 
-void SetIsRevisionControlled(const FNavigationToolItemPtr& InItem, const bool bInIsRevisionControlled)
+const FSlateBrush* GetRevisionControlStatusIcon(const FNavigationToolItemPtr& InItem)
 {
-	if (IRevisionControlExtension* const RevisionControlExt = InItem->CastTo<IRevisionControlExtension>())
+	if (!InItem.IsValid())
+	{
+		return nullptr;
+	}
+
+	const IRevisionControlExtension* const RevisionControlItem = InItem->CastTo<IRevisionControlExtension>();
+	if (!RevisionControlItem)
 	{
-		const bool bItemRevisionControlled = RevisionControlExt->GetRevisionControlState() == EItemRevisionControlState::CheckedOut;
-		if (bItemRevisionControlled != bInIsRevisionControlled)
-		{
-			RevisionControlExt->SetIsRevisionControlled(bInIsRevisionControlled);
-		}
+		return nullptr;
 	}
+
+	return RevisionControlItem->GetRevisionControlStatusIcon();
+}
+
+FText GetRevisionControlStatusText(const FNavigationToolItemPtr& InItem)
+{
+	if (!InItem.IsValid())
+	{
+		return FText::GetEmpty();
+	}
+
+	const IRevisionControlExtension* const RevisionControlItem = InItem->CastTo<IRevisionControlExtension>();
+	if (!RevisionControlItem)
+	{
+		return FText::GetEmpty();
+	}
+
+	return RevisionControlItem->GetRevisionControlStatusText();
 }
 
 void SNavigationToolRevisionControl::Construct(const FArguments& InArgs
@@ -54,119 +81,26 @@ void SNavigationToolRevisionControl::Construct(const FArguments& InArgs
 
 FSlateColor SNavigationToolRevisionControl::GetForegroundColor() const
 {
-	const IRevisionControlExtension* const RevisionControlExt = WeakItem.Pin()->CastTo<IRevisionControlExtension>();
-	if (!RevisionControlExt)
-	{
-		return FStyleColors::Transparent;
-	}
-
-	switch (GetRevisionControlState())
+	switch (GetRevisionControlState(WeakItem.Pin()))
 	{
 	case EItemRevisionControlState::None:
 		return FStyleColors::Transparent;
-	case EItemRevisionControlState::PartiallyCheckedOut:
+	case EItemRevisionControlState::PartiallySourceControlled:
 		return FStyleColors::White25;
-	case EItemRevisionControlState::CheckedOut:
+	case EItemRevisionControlState::SourceControlled:
 		return FStyleColors::Foreground;
 	}
-
 	return FStyleColors::Transparent;
 }
 
 const FSlateBrush* SNavigationToolRevisionControl::GetBrush() const
 {
-	const FNavigationToolItemPtr Item = WeakItem.Pin();
-	if (!Item.IsValid())
-	{
-		return nullptr;
-	}
-
-	const IRevisionControlExtension* const RevisionControlExt = Item->CastTo<IRevisionControlExtension>();
-	if (!RevisionControlExt)
-	{
-		return nullptr;
-	}
-
-	return RevisionControlExt->GetRevisionControlStatusIcon();
+	return GetRevisionControlStatusIcon(WeakItem.Pin());
 }
 
 FText SNavigationToolRevisionControl::GetToolTipText() const
 {
-	const FNavigationToolItemPtr Item = WeakItem.Pin();
-	if (!Item.IsValid())
-	{
-		return FText::GetEmpty();
-	}
-
-	const IRevisionControlExtension* const RevisionControlExt = Item->CastTo<IRevisionControlExtension>();
-	if (!RevisionControlExt)
-	{
-		return FText::GetEmpty();
-	}
-
-	return RevisionControlExt->GetRevisionControlStatusText();
-}
-
-EItemRevisionControlState SNavigationToolRevisionControl::GetRevisionControlState() const
-{
-	const FNavigationToolItemPtr Item = WeakItem.Pin();
-	if (!Item.IsValid())
-	{
-		return EItemRevisionControlState::None;
-	}
-
-	const IRevisionControlExtension* const RevisionControlExt = Item->CastTo<IRevisionControlExtension>();
-	if (!RevisionControlExt)
-	{
-		return EItemRevisionControlState::None;
-	}
-
-	return RevisionControlExt->GetRevisionControlState();
-}
-
-void SNavigationToolRevisionControl::SetIsRevisionControlled(const bool bInRevisionControlled)
-{
-	const FNavigationToolItemPtr Item = WeakItem.Pin();
-	if (!Item.IsValid())
-	{
-		return;
-	}
-
-	const IRevisionControlExtension* const RevisionControlExt = Item->CastTo<IRevisionControlExtension>();
-	if (!RevisionControlExt)
-	{
-		return;
-	}
-
-	const bool bIsRevisionControlled = RevisionControlExt->GetRevisionControlState() == EItemRevisionControlState::CheckedOut;
-	if (bIsRevisionControlled == bInRevisionControlled)
-	{
-		return;
-	}
-
-	OnSetIsRevisionControlled(Item, bInRevisionControlled);
-
-	if (const TSharedPtr<INavigationToolView> ToolView = WeakView.Pin())
-	{
-		ToolView->RequestRefresh();
-	}
-}
-
-void SNavigationToolRevisionControl::OnSetIsRevisionControlled(const FNavigationToolItemPtr& InItem, const bool bInIsRevisionControlled)
-{
-	if (IRevisionControlExtension* const RevisionControlExt = InItem->CastTo<IRevisionControlExtension>())
-	{
-		RevisionControlExt->SetIsRevisionControlled(bInIsRevisionControlled);
-	}
-
-	// Apply the same visibility to the children
-	for (const FNavigationToolItemPtr& Child : InItem->GetChildren())
-	{
-		if (IRevisionControlExtension* const RevisionControlExt = Child->CastTo<IRevisionControlExtension>())
-		{
-			RevisionControlExt->SetIsRevisionControlled(bInIsRevisionControlled);
-		}
-	}
+	return GetRevisionControlStatusText(WeakItem.Pin());
 }
 
 }
diff --git a/Engine/Source/Editor/Sequencer/Private/Tools/NavigationTool/Widgets/Columns/SNavigationToolRevisionControl.h b/Engine/Source/Editor/Sequencer/Private/Tools/NavigationTool/Widgets/Columns/SNavigationToolRevisionControl.h
--- a/Engine/Source/Editor/Sequencer/Private/Tools/NavigationTool/Widgets/Columns/SNavigationToolRevisionControl.h
+++ b/Engine/Source/Editor/Sequencer/Private/Tools/NavigationTool/Widgets/Columns/SNavigationToolRevisionControl.h
@@ -2,8 +2,6 @@
 
 #pragma once
 
-#include "ScopedTransaction.h"
-#include "Templates/UniquePtr.h"
 #include "Tools/NavigationTool/NavigationToolDefines.h"
 #include "Widgets/Images/SImage.h"
 
@@ -13,7 +11,6 @@ namespace UE::Sequencer::NavigationTool
 class FNavigationToolRevisionControlColumn;
 class INavigationToolView;
 class SNavigationToolTreeRow;
-enum class EItemRevisionControlState;
 
 class SNavigationToolRevisionControl : public SImage
 {
@@ -32,11 +29,6 @@ private:
 	
 	virtual const FSlateBrush* GetBrush() const;
 	virtual FText GetToolTipText() const;
-	
-	virtual void OnSetIsRevisionControlled(const FNavigationToolItemPtr& InItem, const bool bInIsRevisionControlled);
-
-	EItemRevisionControlState GetRevisionControlState() const;
-	void SetIsRevisionControlled(const bool bInRevisionControlled);
 
 	TWeakPtr<FNavigationToolRevisionControlColumn> WeakColumn;
 
diff --git a/Engine/Source/Editor/Sequencer/Public/Tools/NavigationTool/Columns/NavigationToolPlayheadColumn.h b/Engine/Source/Editor/Sequencer/Public/Tools/NavigationTool/Columns/NavigationToolPlayheadColumn.h
new file mode 100644
--- /dev/null
+++ b/Engine/Source/Editor/Sequencer/Public/Tools/NavigationTool/Columns/NavigationToolPlayheadColumn.h
@@ -0,0 +1,28 @@
+// Copyright Epic Games, Inc. All Rights Reserved.
+
+#pragma once
+
+#include "Tools/NavigationTool/Columns/NavigationToolColumn.h"
+
+namespace UE::Sequencer::NavigationTool
+{
+
+class INavigationToolView;
+
+class FNavigationToolPlayheadColumn : public FNavigationToolColumn
+{
+public:
+	UE_NAVIGATIONTOOL_INHERITS(FNavigationToolPlayheadColumn, FNavigationToolColumn);
+
+	//~ Begin INavigationToolColumn
+	SEQUENCER_API virtual FText GetColumnDisplayNameText() const override;
+	SEQUENCER_API virtual const FSlateBrush* GetIconBrush() const override;
+	SEQUENCER_API virtual SHeaderRow::FColumn::FArguments ConstructHeaderRowColumn(const TSharedRef<INavigationToolView>& InToolView, const float InFillSize) override;
+	virtual bool ShouldShowColumnByDefault() const override { return true; }
+	SEQUENCER_API virtual TSharedRef<SWidget> ConstructRowWidget(const FNavigationToolItemPtr& InItem
+		, const TSharedRef<INavigationToolView>& InToolView
+		, const TSharedRef<SNavigationToolTreeRow>& InRow) override;
+	//~ End INavigationToolColumn
+};
+
+}
diff --git a/Engine/Source/Editor/Sequencer/Public/Tools/NavigationTool/Customization/NavigationToolObjectIconCustomization.h b/Engine/Source/Editor/Sequencer/Public/Tools/NavigationTool/Customization/NavigationToolObjectIconCustomization.h
deleted file mode 100644
diff --git a/Engine/Source/Editor/Sequencer/Public/Tools/NavigationTool/DragDropOps/NavigationToolItemDragDropOp.h b/Engine/Source/Editor/Sequencer/Public/Tools/NavigationTool/DragDropOps/NavigationToolItemDragDropOp.h
--- a/Engine/Source/Editor/Sequencer/Public/Tools/NavigationTool/DragDropOps/NavigationToolItemDragDropOp.h
+++ b/Engine/Source/Editor/Sequencer/Public/Tools/NavigationTool/DragDropOps/NavigationToolItemDragDropOp.h
@@ -39,9 +39,6 @@ public:
 		return ActionType;
 	}
 
-	/** Retrieves the Actors found in the DragDropOp, assuming this is of type FNavigationToolItemDragDropOp */
-	SEQUENCER_API void GetDragDropOpActors(TArray<TWeakObjectPtr<AActor>>& OutActors) const;
-
 	template<typename InDropHandlerType
 		, typename = typename TEnableIf<TIsDerivedFrom<InDropHandlerType, FNavigationToolItemDropHandler>::Value>::Type
 		, typename... InArgTypes>
diff --git a/Engine/Source/Editor/Sequencer/Public/Tools/NavigationTool/Extensions/IPlayheadExtension.h b/Engine/Source/Editor/Sequencer/Public/Tools/NavigationTool/Extensions/IPlayheadExtension.h
new file mode 100644
--- /dev/null
+++ b/Engine/Source/Editor/Sequencer/Public/Tools/NavigationTool/Extensions/IPlayheadExtension.h
@@ -0,0 +1,26 @@
+﻿// Copyright Epic Games, Inc. All Rights Reserved.
+
+#pragma once
+
+#include "Tools/NavigationTool/NavigationToolItemType.h"
+
+namespace UE::Sequencer::NavigationTool
+{
+
+enum class EItemContainsPlayhead
+{
+	None                      = 0,
+	ContainsPlayhead          = 1 << 0,
+	PartiallyContainsPlayhead = 1 << 1,
+};
+ENUM_CLASS_FLAGS(EItemContainsPlayhead)
+
+class SEQUENCER_API IPlayheadExtension : public INavigationToolItemTypeCastable
+{
+public:
+	UE_NAVIGATIONTOOL_INHERITS(IPlayheadExtension, INavigationToolItemTypeCastable)
+
+	virtual EItemContainsPlayhead ContainsPlayhead() const = 0;
+};
+
+}
diff --git a/Engine/Source/Editor/Sequencer/Public/Tools/NavigationTool/Extensions/IRevisionControlExtension.h b/Engine/Source/Editor/Sequencer/Public/Tools/NavigationTool/Extensions/IRevisionControlExtension.h
--- a/Engine/Source/Editor/Sequencer/Public/Tools/NavigationTool/Extensions/IRevisionControlExtension.h
+++ b/Engine/Source/Editor/Sequencer/Public/Tools/NavigationTool/Extensions/IRevisionControlExtension.h
@@ -2,7 +2,10 @@
 
 #pragma once
 
+#include "ISourceControlModule.h"
+#include "Misc/PackageName.h"
 #include "Tools/NavigationTool/NavigationToolItemType.h"
+#include "UObject/Package.h"
 
 struct FSlateBrush;
 
@@ -11,9 +14,9 @@ namespace UE::Sequencer::NavigationTool
 
 enum class EItemRevisionControlState
 {
-	None                = 0,
-	CheckedOut          = 1 << 0,
-	PartiallyCheckedOut = 1 << 1,
+	None                      = 0,
+	SourceControlled          = 1 << 0,
+	PartiallySourceControlled = 1 << 1,
 };
 ENUM_CLASS_FLAGS(EItemRevisionControlState)
 
@@ -22,13 +25,34 @@ class SEQUENCER_API IRevisionControlExtension : public INavigationToolItemTypeCa
 public:
 	UE_NAVIGATIONTOOL_INHERITS(IRevisionControlExtension, INavigationToolItemTypeCastable)
 
-	virtual EItemRevisionControlState GetRevisionControlState() const = 0;
+	IRevisionControlExtension(const UObject* const InObject)
+	{
+		QueueRevisionControlStatusUpdate(InObject);
+	}
 
-	virtual void SetIsRevisionControlled(const bool bInRevisionControlled) = 0;
+	virtual EItemRevisionControlState GetRevisionControlState() const = 0;
 
 	virtual const FSlateBrush* GetRevisionControlStatusIcon() const = 0;
 
 	virtual FText GetRevisionControlStatusText() const = 0;
+
+	static void QueueRevisionControlStatusUpdate(const UObject* const InObject)
+	{
+		if (!IsValid(InObject))
+		{
+			return;
+		}
+
+		UPackage* const InPackage = InObject->GetPackage();
+		if (!IsValid(InPackage))
+		{
+			return;
+		}
+
+		const FString PackageFilename = FPackageName::LongPackageNameToFilename(InPackage->GetName()
+			, FPackageName::GetAssetPackageExtension());
+		ISourceControlModule::Get().QueueStatusUpdate(PackageFilename);
+	}
 };
 
 }
diff --git a/Engine/Source/Editor/Sequencer/Public/Tools/NavigationTool/Filters/Filters/NavigationToolBuiltInFilterParams.h b/Engine/Source/Editor/Sequencer/Public/Tools/NavigationTool/Filters/Filters/NavigationToolBuiltInFilterParams.h
--- a/Engine/Source/Editor/Sequencer/Public/Tools/NavigationTool/Filters/Filters/NavigationToolBuiltInFilterParams.h
+++ b/Engine/Source/Editor/Sequencer/Public/Tools/NavigationTool/Filters/Filters/NavigationToolBuiltInFilterParams.h
@@ -26,6 +26,7 @@ struct FNavigationToolBuiltInFilterParams
 		, const FSlateBrush* const InIconBrush = nullptr
 		, const FText& InDisplayName = FText::GetEmpty()
 		, const FText& InTooltipText = FText::GetEmpty()
+		, const TSharedPtr<FUICommandInfo>& InToggleCommand = nullptr
 		, const bool InEnabledByDefault = true
 		, const EClassFlags InRequiredClassFlags = CLASS_None
 		, const EClassFlags InRestrictedClassFlags = CLASS_None);
@@ -44,6 +45,8 @@ struct FNavigationToolBuiltInFilterParams
 
 	bool IsEnabledByDefault() const;
 
+	TSharedPtr<FUICommandInfo> GetToggleCommand() const;
+
 	void SetOverrideIconColor(const FSlateColor& InNewIconColor);
 
 	bool IsValidItemClass(const UE::Sequencer::NavigationTool::FNavigationToolItemTypeId InClassTypeId) const;
@@ -51,7 +54,7 @@ struct FNavigationToolBuiltInFilterParams
 
 	void SetFilterText(const FText& InText);
 
-	bool PassesFilterText(const UE::Sequencer::NavigationTool::FNavigationToolObject& InItem) const;
+	bool PassesFilterText(const UE::Sequencer::NavigationTool::FNavigationToolItemPtr& InItem) const;
 
 private:
 	UPROPERTY(EditAnywhere, Category = "Filter")
@@ -83,6 +86,8 @@ private:
 	UPROPERTY(EditAnywhere, Category = "Filter", meta=(InlineEditConditionToggle))
 	bool bEnabledByDefault = true;
 
+	TSharedPtr<FUICommandInfo> ToggleCommand;
+
 	const FSlateBrush* IconBrush = nullptr;
 
 	EClassFlags RequiredClassFlags = CLASS_None;
diff --git a/Engine/Source/Editor/Sequencer/Public/Tools/NavigationTool/Filters/Filters/NavigationToolFilter_Dirty.h b/Engine/Source/Editor/Sequencer/Public/Tools/NavigationTool/Filters/Filters/NavigationToolFilter_Dirty.h
new file mode 100644
--- /dev/null
+++ b/Engine/Source/Editor/Sequencer/Public/Tools/NavigationTool/Filters/Filters/NavigationToolFilter_Dirty.h
@@ -0,0 +1,37 @@
+﻿// Copyright Epic Games, Inc. All Rights Reserved.
+
+#pragma once
+
+#include "Tools/NavigationTool/Filters/Filters/NavigationToolFilterBase.h"
+#include "Tools/NavigationTool/Filters/NavigationToolFilterData.h"
+
+namespace UE::Sequencer::NavigationTool
+{
+
+class FNavigationToolFilter_Dirty : public FNavigationToolFilter
+{
+public:
+	static FString StaticName() { return TEXT("Dirty"); }
+
+	FNavigationToolFilter_Dirty(INavigationToolFilterBar& InFilterInterface, TSharedPtr<FFilterCategory> InCategory = nullptr);
+
+	//~ Begin FNavigationToolFilter
+	virtual FText GetDefaultToolTipText() const override;
+	virtual TSharedPtr<FUICommandInfo> GetToggleCommand() const override;
+	//~ End FNavigationToolFilter
+
+	//~ Begin FFilterBase
+	virtual FText GetDisplayName() const override;
+	virtual FSlateIcon GetIcon() const override;
+	//~ End FFilterBase
+
+	//~ Begin IFilter
+	virtual FString GetName() const override;
+	virtual bool PassesFilter(FNavigationToolFilterType InItem) const override;
+	//~ End IFilter
+
+protected:
+	static bool IsObjectPackageDirty(const UObject* const InObject);
+};
+
+}
diff --git a/Engine/Source/Editor/Sequencer/Public/Tools/NavigationTool/Filters/Filters/NavigationToolFilter_Playhead.h b/Engine/Source/Editor/Sequencer/Public/Tools/NavigationTool/Filters/Filters/NavigationToolFilter_Playhead.h
new file mode 100644
--- /dev/null
+++ b/Engine/Source/Editor/Sequencer/Public/Tools/NavigationTool/Filters/Filters/NavigationToolFilter_Playhead.h
@@ -0,0 +1,55 @@
+﻿// Copyright Epic Games, Inc. All Rights Reserved.
+
+#pragma once
+
+#include "Tools/NavigationTool/Filters/Filters/NavigationToolFilterBase.h"
+#include "Tools/NavigationTool/Filters/NavigationToolFilterData.h"
+
+namespace UE::Sequencer::NavigationTool
+{
+
+class FNavigationToolFilter_Playhead : public FNavigationToolFilter
+{
+public:
+	static FString StaticName() { return TEXT("Playhead"); }
+
+	FNavigationToolFilter_Playhead(INavigationToolFilterBar& InFilterInterface, TSharedPtr<FFilterCategory> InCategory = nullptr);
+	virtual ~FNavigationToolFilter_Playhead() override;
+
+	//~ Begin FNavigationToolFilter
+	virtual FText GetDefaultToolTipText() const override;
+	virtual TSharedPtr<FUICommandInfo> GetToggleCommand() const override;
+	virtual void ActiveStateChanged(const bool bInActive) override;
+	//~ End FNavigationToolFilter
+
+	//~ Begin FFilterBase
+	virtual FText GetDisplayName() const override;
+	virtual FSlateIcon GetIcon() const override;
+	//~ End FFilterBase
+
+	//~ Begin IFilter
+	virtual FString GetName() const override;
+	virtual bool PassesFilter(FNavigationToolFilterType InItem) const override;
+	//~ End IFilter
+
+protected:
+	void BindEvents();
+	void UnbindEvents();
+
+	void OnPlayEvent();
+	void OnStopEvent();
+
+	void OnBeginScrubbingEvent();
+	void OnEndScrubbingEvent();
+
+	FTimerManager* GetTimerManager() const;
+
+	void StartRefreshTimer();
+	void StopRefreshTimer();
+
+	FTimerHandle TimerHandle;
+
+	TWeakPtr<ISequencer> WeakSequencer;
+};
+
+}
diff --git a/Engine/Source/Editor/Sequencer/Public/Tools/NavigationTool/INavigationTool.h b/Engine/Source/Editor/Sequencer/Public/Tools/NavigationTool/INavigationTool.h
--- a/Engine/Source/Editor/Sequencer/Public/Tools/NavigationTool/INavigationTool.h
+++ b/Engine/Source/Editor/Sequencer/Public/Tools/NavigationTool/INavigationTool.h
@@ -55,6 +55,9 @@ public:
 	SEQUENCER_API virtual void ShowHideToolTab(const bool bInVisible) = 0;
 	SEQUENCER_API virtual void ToggleToolTabVisible() = 0;
 
+	/** Gets the Command List that the Navigation Tool Views will use to append their Command Lists to */
+	virtual TSharedPtr<FUICommandList> GetBaseCommandList() const = 0;
+
 	/** Sets the Command List that the Navigation Tool Views will use to append their Command Lists to */
 	virtual void SetBaseCommandList(const TSharedPtr<FUICommandList>& InBaseCommandList) = 0;
 
@@ -167,15 +170,22 @@ public:
 	virtual TSharedRef<FNavigationToolItem> GetTreeRoot() const = 0;
 
 	/**
-	 * Finds the Registered Item that has the given Id
-	 * @returns the Item with the given Id, or null if the item does not exist or was not registered to the Navigation Tool
+	 * Finds the registered item that has the given Id
+	 * @returns The item with the given Id, or null if the item does not exist or was not registered to the Navigation Tool
 	 */
 	virtual FNavigationToolItemPtr FindItem(const FNavigationToolItemId& InItemId) const = 0;
 
+	/**
+	 * Attempts to find all items associated with a Sequencer view model
+	 * @param InViewModel The Sequencer view model to find the items of
+	 * @return The list of Navigation Tool items found
+	 */
+	virtual TArray<FNavigationToolItemPtr> TryFindItems(const FViewModelPtr& InViewModel) const = 0;
+
 	/**
  	 * Adds or Removes the Ignore Notify Flags to prevent certain actions from automatically happening when they're triggered
  	 * @param InFlag the ignore flag to add or remove
- 	 * @param bIgnore whether to add (true) or remove (false) the flag
+ 	 * @param bInIgnore whether to add (true) or remove (false) the flag
  	 */
 	virtual void SetIgnoreNotify(const ENavigationToolIgnoreNotifyFlags InFlag, const bool bInIgnore) = 0;
 
diff --git a/Engine/Source/Editor/Sequencer/Public/Tools/NavigationTool/Items/INavigationToolItem.h b/Engine/Source/Editor/Sequencer/Public/Tools/NavigationTool/Items/INavigationToolItem.h
--- a/Engine/Source/Editor/Sequencer/Public/Tools/NavigationTool/Items/INavigationToolItem.h
+++ b/Engine/Source/Editor/Sequencer/Public/Tools/NavigationTool/Items/INavigationToolItem.h
@@ -43,6 +43,12 @@ protected:
 public:
 	UE_NAVIGATIONTOOL_INHERITS(INavigationToolItem, INavigationToolItemTypeCastable);
 
+	/** Determines whether the Item properties are in a valid state */
+	virtual bool IsItemValid() const = 0;
+
+	/** Gets the UObject that this item represents. May be null due if no item association or multiple. */
+	virtual UObject* GetItemObject() const { return nullptr; }
+
 	/** Gets the module provider responsible for the creation of this item */
 	virtual TSharedPtr<FNavigationToolProvider> GetProvider() const = 0;
 
@@ -64,9 +70,6 @@ public:
 	/** Gets the Navigation Tool that owns this item */
 	virtual INavigationTool& GetOwnerTool() const = 0;
 
-	/** Determines whether the Item properties are in a valid state */
-	virtual bool IsItemValid() const = 0;
-
 	/** Called when the Item has been registered into the Navigation Tool */
 	virtual void OnItemRegistered() {}
 
@@ -238,6 +241,82 @@ public:
 	/** Sets the Parent. Note that the Parent must've already have this instance as a child (check is done) */
 	virtual void SetParent(FNavigationToolItemPtr InParent) = 0;
 
+	template<typename InItemType>
+	InItemType* FindAncestorOfType(const bool bInIncludeThis = false) const
+	{
+		if (bInIncludeThis)
+		{
+			if (InItemType* const CastedParent = const_cast<InItemType*>(CastTo<const InItemType>()))
+			{
+				return CastedParent;
+			}
+		}
+
+		FNavigationToolItemPtr CurrentParent = GetParent();
+
+		while (CurrentParent.IsValid())
+		{
+			if (InItemType* const CastedParent = CurrentParent->CastTo<InItemType>())
+			{
+				return CastedParent;
+			}
+
+			CurrentParent = CurrentParent->GetParent();
+		}
+
+		return nullptr;
+	}
+
+	template<typename InItemType>
+	TArray<InItemType*> FindAncestorsOfType(const bool bInIncludeThis = false) const
+	{
+		TArray<InItemType*> OutAncestors;
+
+		if (bInIncludeThis)
+		{
+			if (InItemType* const CastedParent = const_cast<InItemType*>(CastTo<const InItemType>()))
+			{
+				OutAncestors.Add(CastedParent);
+			}
+		}
+
+		FNavigationToolItemPtr CurrentParent = GetParent();
+
+		while (CurrentParent.IsValid())
+		{
+			if (InItemType* const CastedParent = CurrentParent->CastTo<InItemType>())
+			{
+				OutAncestors.Add(CastedParent);
+			}
+
+			CurrentParent = CurrentParent->GetParent();
+		}
+
+		return OutAncestors;
+	}
+
+	TArray<FNavigationToolItemPtr> GetDescendants(const bool bInIncludeThis = false)
+	{
+		TArray<FNavigationToolItemPtr> OutDescendants;
+
+		if (bInIncludeThis)
+		{
+			OutDescendants.Add(AsShared());
+		}
+
+		for (const FNavigationToolItemPtr& Child : GetChildren())
+		{
+			OutDescendants.Add(Child);
+
+			if (Child.IsValid())
+			{
+				OutDescendants.Append(Child->GetDescendants());
+			}
+		}
+
+		return OutDescendants;
+	}
+
 	template<typename InItemType>
 	TArray<InItemType*> GetDescendantsOfType(const bool bInIncludeThis = false) const
 	{
diff --git a/Engine/Source/Editor/Sequencer/Public/Tools/NavigationTool/Items/NavigationToolActor.h b/Engine/Source/Editor/Sequencer/Public/Tools/NavigationTool/Items/NavigationToolActor.h
--- a/Engine/Source/Editor/Sequencer/Public/Tools/NavigationTool/Items/NavigationToolActor.h
+++ b/Engine/Source/Editor/Sequencer/Public/Tools/NavigationTool/Items/NavigationToolActor.h
@@ -27,7 +27,6 @@ public:
 	//~ Begin INavigationToolItem
 
 	virtual void FindChildren(TArray<FNavigationToolItemPtr>& OutChildren, const bool bInRecursive) override;
-	virtual void GetItemProxies(TArray<TSharedPtr<FNavigationToolItemProxy>>& OutItemProxies) override;
 
 	virtual bool IsAllowedInTool() const override;
 	virtual ENavigationToolItemViewMode GetSupportedViewModes(const INavigationToolView& InToolView) const override;
diff --git a/Engine/Source/Editor/Sequencer/Public/Tools/NavigationTool/Items/NavigationToolBinding.h b/Engine/Source/Editor/Sequencer/Public/Tools/NavigationTool/Items/NavigationToolBinding.h
--- a/Engine/Source/Editor/Sequencer/Public/Tools/NavigationTool/Items/NavigationToolBinding.h
+++ b/Engine/Source/Editor/Sequencer/Public/Tools/NavigationTool/Items/NavigationToolBinding.h
@@ -3,10 +3,12 @@
 #pragma once
 
 #include "MovieSceneBinding.h"
+#include "MVVM/ViewModels/ObjectBindingModel.h"
+#include "NavigationToolItem.h"
 #include "Textures/SlateIcon.h"
+#include "Tools/NavigationTool/Extensions/IPlayheadExtension.h"
 #include "Tools/NavigationTool/Extensions/IRenameableExtension.h"
 #include "Tools/NavigationTool/Extensions/ISequenceLockableExtension.h"
-#include "Tools/NavigationTool/Items/NavigationToolObject.h"
 
 namespace UE::Sequencer::NavigationTool
 {
@@ -20,12 +22,14 @@ class SEQUENCER_API FNavigationToolBinding
 	: public FNavigationToolItem
 	, public IRenameableExtension
 	//, public ISequenceLockableExtension
+	, public IPlayheadExtension
 {
 public:
 	UE_NAVIGATIONTOOL_INHERITS_WITH_SUPER(FNavigationToolBinding
 		, FNavigationToolItem
 		, IRenameableExtension
-		/*, ISequenceLockableExtension*/);
+		/*, ISequenceLockableExtension*/
+		, IPlayheadExtension);
 
 	FNavigationToolBinding(INavigationTool& InTool
 		, const FNavigationToolItemPtr& InParentItem
@@ -35,6 +39,8 @@ public:
 	//~ Begin INavigationToolItem
 
 	virtual bool IsItemValid() const override;
+	virtual UObject* GetItemObject() const override;
+	virtual bool IsAllowedInTool() const override;
 
 	virtual void FindChildren(TArray<FNavigationToolItemPtr>& OutChildren, const bool bInRecursive) override;
 
@@ -69,16 +75,25 @@ public:
 	//virtual void SetIsLocked(const bool bInIsLocked) override;
 	//~ End ISequenceLockableExtension
 
+	//~ Begin IPlayheadExtension
+	virtual EItemContainsPlayhead ContainsPlayhead() const override;
+	//~ End IPlayheadExtension
+
 	const UClass* GetBoundObjectClass() const;
 
 	const FMovieSceneBinding& GetBinding() const;
 
+	/** @return The Sequence this binding belongs to */
+	UMovieSceneSequence* GetSequence() const;
+
 	/** @return The cached object that is bound in sequencer. */
 	UObject* GetCachedBoundObject() const;
 
 	/** Caches the object if it's not already been cached and returns the result. */
 	UObject* CacheBoundObject();
 
+	TViewModelPtr<FObjectBindingModel> GetViewModel() const;
+
 protected:
 	TWeakPtr<FNavigationToolSequence> WeakParentSequenceItem;
 
diff --git a/Engine/Source/Editor/Sequencer/Public/Tools/NavigationTool/Items/NavigationToolItemId.h b/Engine/Source/Editor/Sequencer/Public/Tools/NavigationTool/Items/NavigationToolItemId.h
--- a/Engine/Source/Editor/Sequencer/Public/Tools/NavigationTool/Items/NavigationToolItemId.h
+++ b/Engine/Source/Editor/Sequencer/Public/Tools/NavigationTool/Items/NavigationToolItemId.h
@@ -2,6 +2,7 @@
 
 #pragma once
 
+#include "SequencerCoreFwd.h"
 #include "Tools/NavigationTool/ItemProxies/INavigationToolItemProxyFactory.h"
 #include "UObject/SoftObjectPath.h"
 
@@ -29,13 +30,6 @@ struct SEQUENCER_API FNavigationToolItemId
 	/** More flexible option to just specify the string directly. Could be used for folders (e.g. for a nested folder C could be "A/B/C" */
 	FNavigationToolItemId(const FString& InUniqueId);
 
-	/**
-	 * Ctor used for Objects that appear ONCE in the Navigation Tool (e.g. an Actor or Component), so only the Object itself is needed for the Id
-	 * for objects that can appear multiple times because they're referenced multiple times, see the ctor below that takes in a Referencing Item and Referencing Id
-	 * @param InObject The object that the item represents
-	 */
-	FNavigationToolItemId(const UObject* InObject);
-
 	/**
 	 * Ctor used for Objects that are expected to appear multiple times in the Navigation Tool (e.g. a Material Ref)
 	 *
@@ -86,6 +80,12 @@ struct SEQUENCER_API FNavigationToolItemId
 		, const int32 InSectionIndex = 0
 		, const FString& InReferenceId = FString());
 
+	/**
+	 * Ctor used for making the Item Id from a Sequencer view model
+	 * @param InViewModel The view model to construct the Id from
+	 */
+	FNavigationToolItemId(const FViewModelPtr& InViewModel);
+
 	FNavigationToolItemId(const FNavigationToolItemId& Other);
 	FNavigationToolItemId(FNavigationToolItemId&& Other) noexcept;
 
@@ -100,7 +100,7 @@ struct SEQUENCER_API FNavigationToolItemId
 	}
 
 	/** Returns whether this Id has a cached hash (i.e. ran any ctor except the default one) */
-	bool IsValid() const;
+	bool IsValidId() const;
 
 	FString GetStringId() const;
 
diff --git a/Engine/Source/Editor/Sequencer/Public/Tools/NavigationTool/Items/NavigationToolItemUtils.h b/Engine/Source/Editor/Sequencer/Public/Tools/NavigationTool/Items/NavigationToolItemUtils.h
--- a/Engine/Source/Editor/Sequencer/Public/Tools/NavigationTool/Items/NavigationToolItemUtils.h
+++ b/Engine/Source/Editor/Sequencer/Public/Tools/NavigationTool/Items/NavigationToolItemUtils.h
@@ -5,7 +5,6 @@
 #include "Containers/ContainersFwd.h"
 #include "NavigationToolSequence.h"
 #include "Algo/AllOf.h"
-#include "MVVM/ViewModelPtr.h"
 #include "Tools/NavigationTool/Items/INavigationToolItem.h"
 #include "Tools/NavigationTool/NavigationToolDefines.h"
 
@@ -23,7 +22,7 @@ class FNavigationToolSequence;
  * Compares the absolute order of the items in the Navigation Tool and returns true if A comes before B in the Navigation Tool.
  * Useful to use when sorting Items.
  */
-bool CompareToolItemOrder(const FNavigationToolItemPtr& A, const FNavigationToolItemPtr& B);
+bool CompareToolItemOrder(const FNavigationToolItemPtr& InA, const FNavigationToolItemPtr& InB);
 
 /** Returns two subset arrays of Items: one is containing only Sortable Items and the other Non Sortable Items */
 void SplitSortableAndUnsortableItems(const TArray<FNavigationToolItemPtr>& InItems
@@ -36,12 +35,22 @@ UMovieSceneSubSection* GetSequenceItemSubSection(const FNavigationToolItemPtr& I
 /**  */
 UMovieSceneMetaData* GetSequenceItemMetaData(const TSharedPtr<INavigationToolItem>& InItem);
 
-/**  */
-SEQUENCER_API FText SequenceDisplayNameWithModifyAndInactiveStatus(const TSharedPtr<FNavigationToolSequence>& InSequenceItem
-	, const bool bInRemoveParentName
-	, const bool bInIsActive);
+/**
+ * Removes the parent sequence prefix from the display name of a child sequence item.
+ * @param InOutDisplayName The display name to modify, which will have the parent prefix removed if found.
+ * @param InSequenceItem The child sequence item whose display name is being adjusted.
+ */
+void RemoveSequenceDisplayNameParentPrefix(FText& InOutDisplayName, const TSharedRef<FNavigationToolSequence>& InSequenceItem);
 
-/**  */
+/**
+ * Appends the dirty symbol to the given display name.
+ * If the associated package is dirty, a "*" tag is appended to the display name.
+ * @param InOutDisplayName The display name to be appended with the status.
+ * @param InSequence The sequence whose package is checked for a dirty status.
+ */
+void AppendSequenceDisplayNameDirtyStatus(FText& InOutDisplayName, const UMovieSceneSequence& InSequence);
+
+/** The result of an item comparison */
 enum class ENavigationToolCompareState
 {
 	NoneTrue    = 0,
@@ -50,7 +59,16 @@ enum class ENavigationToolCompareState
 };
 ENUM_CLASS_FLAGS(ENavigationToolCompareState)
 
-/**  */
+/**
+ * Evaluates an items children based on a predicate function and determines
+ * whether all, none, or some of the children match the condition.
+ * @param InItem The item which to evaluate the children of.
+ * @param InTrueFunction A function that defines the true condition.
+ * @param InFalseFunction A function that defines the false condition.
+ * @return ENavigationToolCompareState::AllTrue if all items satisfy the condition
+ *         ENavigationToolCompareState::NoneTrue if no items satisfy the condition
+ *         ENavigationToolCompareState::PartialTrue if some items satisfy the condition
+ */
 template<typename InItemType>
 ENavigationToolCompareState CompareChildrenItemState(const INavigationToolItem* const InItem
 	, const TFunctionRef<bool(const InItemType* const)>& InTrueFunction
@@ -91,7 +109,15 @@ ENavigationToolCompareState CompareChildrenItemState(const INavigationToolItem*
 	return ENavigationToolCompareState::PartialTrue;
 }
 
-/**  */
+/**
+ * Evaluates an items children based on a predicate function and determines
+ * whether all, none, or some of the children match the condition.
+ * @param InItem The item which to evaluate the children of.
+ * @param InTrueFunction A function that defines the true condition.
+ * @return ENavigationToolCompareState::AllTrue if all items satisfy the condition
+ *         ENavigationToolCompareState::NoneTrue if no items satisfy the condition
+ *         ENavigationToolCompareState::PartialTrue if some items satisfy the condition
+ */
 template<typename InItemType>
 ENavigationToolCompareState CompareChildrenItemStateSimple(const INavigationToolItem* const InItem
 	, const TFunctionRef<bool(const InItemType* const)>& InTrueFunction)
@@ -126,7 +152,16 @@ ENavigationToolCompareState CompareChildrenItemStateSimple(const INavigationTool
 	return ENavigationToolCompareState::PartialTrue;
 }
 
-/**  */
+/**
+ * Evaluates the state of items in an array based on a predicate function
+ * and determines whether all, none, or some of the items match the condition.
+ * @param InArray The array of items to evaluate.
+ * @param InTrueFunction A function that defines the true condition.
+ * @param InFalseFunction A function that defines the false condition.
+ * @return ENavigationToolCompareState::AllTrue if all items satisfy the condition
+ *         ENavigationToolCompareState::NoneTrue if no items satisfy the condition
+ *         ENavigationToolCompareState::PartialTrue if some items satisfy the condition
+ */
 template<typename InItemType>
 ENavigationToolCompareState CompareArrayState(const TArray<InItemType*>& InArray
 	, const TFunctionRef<bool(const InItemType* const)>& InTrueFunction
@@ -160,7 +195,15 @@ ENavigationToolCompareState CompareArrayState(const TArray<InItemType*>& InArray
 	return ENavigationToolCompareState::PartialTrue;
 }
 
-/**  */
+/**
+ * Evaluates the state of items in an array based on a predicate function
+ * and determines whether all, none, or some of the items match the condition.
+ * @param InArray The array of items to evaluate.
+ * @param InTrueFunction A function that defines the true condition.
+ * @return ENavigationToolCompareState::AllTrue if all items satisfy the condition
+ *         ENavigationToolCompareState::NoneTrue if no items satisfy the condition
+ *         ENavigationToolCompareState::PartialTrue if some items satisfy the condition
+ */
 template<typename InItemType>
 ENavigationToolCompareState CompareArrayStateSimple(const TArray<InItemType*>& InArray
 	, const TFunctionRef<bool(const InItemType* const)>& InTrueFunction)
diff --git a/Engine/Source/Editor/Sequencer/Public/Tools/NavigationTool/Items/NavigationToolObject.h b/Engine/Source/Editor/Sequencer/Public/Tools/NavigationTool/Items/NavigationToolObject.h
deleted file mode 100644
diff --git a/Engine/Source/Editor/Sequencer/Public/Tools/NavigationTool/Items/NavigationToolObjectReference.h b/Engine/Source/Editor/Sequencer/Public/Tools/NavigationTool/Items/NavigationToolObjectReference.h
deleted file mode 100644
diff --git a/Engine/Source/Editor/Sequencer/Public/Tools/NavigationTool/Items/NavigationToolSequence.h b/Engine/Source/Editor/Sequencer/Public/Tools/NavigationTool/Items/NavigationToolSequence.h
--- a/Engine/Source/Editor/Sequencer/Public/Tools/NavigationTool/Items/NavigationToolSequence.h
+++ b/Engine/Source/Editor/Sequencer/Public/Tools/NavigationTool/Items/NavigationToolSequence.h
@@ -3,8 +3,8 @@
 #pragma once
 
 #include "NavigationToolItem.h"
-#include "NavigationToolObject.h"
 #include "Tools/NavigationTool/Extensions/IMarkerVisibilityExtension.h"
+#include "Tools/NavigationTool/Extensions/IPlayheadExtension.h"
 #include "Tools/NavigationTool/Extensions/IRenameableExtension.h"
 #include "Tools/NavigationTool/Extensions/IRevisionControlExtension.h"
 #include "Tools/NavigationTool/Extensions/ISequenceLockableExtension.h"
@@ -17,7 +17,7 @@ struct FMovieSceneBinding;
 
 namespace UE::Sequencer
 {
-	class FSequenceModel;
+class FSectionModel;
 }
 
 namespace UE::Sequencer::NavigationTool
@@ -27,21 +27,23 @@ namespace UE::Sequencer::NavigationTool
  * Item in Navigation Tool representing a Sequence.
  */
 class FNavigationToolSequence
-	: public FNavigationToolObject
+	: public FNavigationToolItem
 	, public IRenameableExtension
 	, public ISequenceInactivableExtension
 	, public ISequenceLockableExtension
 	, public IMarkerVisibilityExtension
 	, public IRevisionControlExtension
+	, public IPlayheadExtension
 {
 public:
 	UE_NAVIGATIONTOOL_INHERITS_WITH_SUPER(FNavigationToolSequence
-		, FNavigationToolObject
+		, FNavigationToolItem
 		, IRenameableExtension
 		, ISequenceInactivableExtension
 		, ISequenceLockableExtension
 		, IMarkerVisibilityExtension
-		, IRevisionControlExtension);
+		, IRevisionControlExtension
+		, IPlayheadExtension);
 
 	SEQUENCER_API FNavigationToolSequence(INavigationTool& InTool
 		, const FNavigationToolItemPtr& InParentItem
@@ -50,21 +52,33 @@ public:
 		, const int32 InSubSectionIndex);
 
 	//~ Begin INavigationToolItem
-	SEQUENCER_API virtual void Select(FNavigationToolScopedSelection& InSelection) const override;
-	SEQUENCER_API virtual bool IsSelected(const FNavigationToolScopedSelection& InSelection) const override;
+
+	SEQUENCER_API virtual bool IsItemValid() const override;
+	SEQUENCER_API virtual UObject* GetItemObject() const override;
+	SEQUENCER_API virtual bool IsAllowedInTool() const override;
+
 	SEQUENCER_API virtual void FindChildren(TArray<FNavigationToolItemPtr>& OutChildren, const bool bInRecursive) override;
 	SEQUENCER_API virtual void GetItemProxies(TArray<TSharedPtr<FNavigationToolItemProxy>>& OutItemProxies) override;
+
 	SEQUENCER_API virtual bool AddChild(const FNavigationToolAddItemParams& InAddItemParams) override;
 	SEQUENCER_API virtual bool RemoveChild(const FNavigationToolRemoveItemParams& InRemoveItemParams) override;
+
 	SEQUENCER_API virtual ENavigationToolItemViewMode GetSupportedViewModes(const INavigationToolView& InToolView) const override;
 	virtual bool CanBeTopLevel() const override { return true; }
-	SEQUENCER_API virtual bool IsAllowedInTool() const override;
+
 	virtual FText GetDisplayName() const override;
 	SEQUENCER_API virtual FText GetClassName() const override;
 	virtual FSlateIcon GetIcon() const override;
 	SEQUENCER_API virtual FText GetIconTooltipText() const override;
+
+	SEQUENCER_API virtual bool IsSelected(const FNavigationToolScopedSelection& InSelection) const override;
+	SEQUENCER_API virtual void Select(FNavigationToolScopedSelection& InSelection) const override;
+
 	SEQUENCER_API virtual void OnSelect() override;
 	SEQUENCER_API virtual void OnDoubleClick() override;
+
+	SEQUENCER_API virtual void OnObjectsReplaced(const TMap<UObject*, UObject*>& InReplacementMap, bool bInRecursive) override;
+
 	//~ End INavigationToolItem
 
 	//~ Begin IRenameableExtension
@@ -94,11 +108,14 @@ public:
 
 	//~ Begin IRevisionControlExtension
 	SEQUENCER_API virtual EItemRevisionControlState GetRevisionControlState() const override;
-	SEQUENCER_API virtual void SetIsRevisionControlled(const bool bInIsRevisionControlled) override;
 	SEQUENCER_API virtual const FSlateBrush* GetRevisionControlStatusIcon() const override;
 	SEQUENCER_API virtual FText GetRevisionControlStatusText() const override;
 	//~ End IRevisionControlExtension
 
+	//~ Begin IPlayheadExtension
+	SEQUENCER_API virtual EItemContainsPlayhead ContainsPlayhead() const override;
+	//~ End IPlayheadExtension
+
 	SEQUENCER_API UMovieSceneSequence* GetSequence() const;
 
 	SEQUENCER_API UMovieSceneSubSection* GetSubSection() const;
@@ -108,6 +125,8 @@ public:
 
 	SEQUENCER_API TArray<FMovieSceneBinding> GetSortedBindings() const;
 
+	TViewModelPtr<FSectionModel> GetViewModel() const;
+
 protected:
 	//~Begin FNavigationToolItem
 	SEQUENCER_API virtual FNavigationToolItemId CalculateItemId() const override;
@@ -115,6 +134,7 @@ protected:
 
 	TWeakObjectPtr<UMovieSceneSubSection> WeakSubSection;
 	int32 SubSectionIndex = 0;
+	TWeakObjectPtr<UMovieSceneSequence> WeakSequence;
 };
 
 }
diff --git a/Engine/Source/Editor/Sequencer/Public/Tools/NavigationTool/Items/NavigationToolTrack.h b/Engine/Source/Editor/Sequencer/Public/Tools/NavigationTool/Items/NavigationToolTrack.h
--- a/Engine/Source/Editor/Sequencer/Public/Tools/NavigationTool/Items/NavigationToolTrack.h
+++ b/Engine/Source/Editor/Sequencer/Public/Tools/NavigationTool/Items/NavigationToolTrack.h
@@ -3,14 +3,20 @@
 #pragma once
 
 #include "MVVM/ViewModels/TrackModel.h"
-#include "NavigationToolObjectReference.h"
+#include "NavigationToolItem.h"
 #include "Textures/SlateIcon.h"
+#include "Tools/NavigationTool/Extensions/IPlayheadExtension.h"
 #include "Tools/NavigationTool/Extensions/ISequenceLockableExtension.h"
 
 class UMovieSceneSequence;
 class UMovieSceneSection;
 class UMovieSceneTrack;
 
+namespace UE::Sequencer
+{
+class FTrackModel;
+}
+
 namespace UE::Sequencer::NavigationTool
 {
 
@@ -18,31 +24,47 @@ namespace UE::Sequencer::NavigationTool
  * 
  */
 class SEQUENCER_API FNavigationToolTrack
-	: public FNavigationToolObject
+	: public FNavigationToolItem
 	, public ISequenceLockableExtension
+	, public IPlayheadExtension
 {
 public:
 	UE_NAVIGATIONTOOL_INHERITS_WITH_SUPER(FNavigationToolTrack
-		, FNavigationToolObject
-		, ISequenceLockableExtension);
+		, FNavigationToolItem
+		, ISequenceLockableExtension
+		, IPlayheadExtension);
 
 	FNavigationToolTrack(INavigationTool& InTool
 		, const FNavigationToolItemPtr& InParentItem
 		, UMovieSceneTrack* const InTrack
 		, const TWeakObjectPtr<UMovieSceneSequence>& InSequence
 		, const TWeakObjectPtr<UMovieSceneSection>& InSection
-		, const int32 InSubSectionIndex);
+		, const int32 InSectionIndex);
 
 	//~ Begin INavigationToolItem
+
+	virtual bool IsItemValid() const override;
+	virtual UObject* GetItemObject() const override;
+	virtual bool IsAllowedInTool() const override;
+
 	virtual void FindChildren(TArray<FNavigationToolItemPtr>& OutChildren, const bool bInRecursive) override;
+
 	virtual bool CanBeTopLevel() const override { return false; }
-	virtual TOptional<EItemDropZone> CanAcceptDrop(const FDragDropEvent& InDragDropEvent, EItemDropZone InDropZone) override;
-	virtual FReply AcceptDrop(const FDragDropEvent& InDragDropEvent, EItemDropZone InDropZone) override;
+
 	virtual FText GetDisplayName() const override;
+	virtual FText GetClassName() const override;
 	virtual const FSlateBrush* GetDefaultIconBrush() const override;
 	virtual FSlateIcon GetIcon() const override;
+	virtual FText GetIconTooltipText() const override;
+
 	virtual bool IsSelected(const FNavigationToolScopedSelection& InSelection) const override;
 	virtual void Select(FNavigationToolScopedSelection& InSelection) const override;
+
+	virtual void OnSelect() override;
+	virtual void OnDoubleClick() override;
+
+	void OnObjectsReplaced(const TMap<UObject*, UObject*>& InReplacementMap, bool bInRecursive);
+
 	//~ End INavigationToolItem
 
 	//~ Begin ISequenceLockableExtension
@@ -50,22 +72,29 @@ public:
 	virtual void SetIsLocked(const bool bInIsLocked) override;
 	//~ End ISequenceLockableExtension
 
+	//~ Begin IPlayheadExtension
+	virtual EItemContainsPlayhead ContainsPlayhead() const override;
+	//~ End IPlayheadExtension
+
 	UMovieSceneTrack* GetTrack() const;
 
+	UMovieSceneSequence* GetSequence() const;
+	UMovieSceneSection* GetSection() const;
+	int32 GetSectionIndex() const;
+
+	TViewModelPtr<FTrackModel> GetViewModel() const;
+
 protected:
 	//~ Begin FNavigationToolItem
 	virtual FNavigationToolItemId CalculateItemId() const override;
 	//~ End FNavigationToolItem
 
-	//~ Begin FNavigationToolObjectItem
-	virtual void SetObject_Impl(UObject* InObject) override;
-	//~ End FNavigationToolObjectItem
-
 	void OnTrackObjectChanged();
 
 	TWeakObjectPtr<UMovieSceneSequence> WeakSequence;
 	TWeakObjectPtr<UMovieSceneSection> WeakSection;
-	int32 SubSectionIndex = 0;
+	int32 SectionIndex = 0;
+	TWeakObjectPtr<UMovieSceneTrack> WeakTrack;
 
 	FSlateIcon Icon;
 };
diff --git a/Engine/Source/Editor/Sequencer/Public/Tools/NavigationTool/NavigationToolSettings.h b/Engine/Source/Editor/Sequencer/Public/Tools/NavigationTool/NavigationToolSettings.h
--- a/Engine/Source/Editor/Sequencer/Public/Tools/NavigationTool/NavigationToolSettings.h
+++ b/Engine/Source/Editor/Sequencer/Public/Tools/NavigationTool/NavigationToolSettings.h
@@ -64,7 +64,7 @@ public:
 
 	bool ShouldUseShortNames() const;
 	void SetUseShortNames(const bool bInUseShortNames);
-	
+
 	bool ShouldUseMutedHierarchy() const;
 	void SetUseMutedHierarchy(const bool bInUseMutedHierarchy);
 
@@ -106,6 +106,12 @@ public:
 	float GetLastFilterBarSizeCoefficient() const;
 	void SetLastFilterBarSizeCoefficient(const float bInSizeCoefficient);
 
+	bool ShouldSyncSelectionToNavigationTool() const;
+	void SetSyncSelectionToNavigationTool(const bool bInSync);
+
+	bool ShouldSyncSelectionToSequencer() const;
+	void SetSyncSelectionToSequencer(const bool bInSync);
+
 private:
 	/** Applies the default column view set by the provider when a view is loaded */
 	UPROPERTY(Config, EditAnywhere, Category = "Editor UX")
@@ -116,11 +122,11 @@ private:
 	bool bUseShortNames = true;
 
 	/** Whether to show the parent of the shown items, even if the parents are filtered out */
-	UPROPERTY(Config, EditAnywhere, Category = "Initial Values")
+	UPROPERTY(Config, EditAnywhere, Category = "Editor UX")
 	bool bUseMutedHierarchy = true;
 
 	/** Whether to auto expand the hierarchy to show the item when selected */
-	UPROPERTY(Config, EditAnywhere, Category = "Initial Values")
+	UPROPERTY(Config, EditAnywhere, Category = "Editor UX")
 	bool bAutoExpandToSelection = true;
 
 	/** Whether to show the lock state always, rather than only showing when the item is locked or hovered */
@@ -128,11 +134,11 @@ private:
 	bool bAlwaysShowLockState = false;
 
 	/** The View Mode a Non-Actor / Non-Component Item supports by default */
-	UPROPERTY(Config, EditAnywhere, Category = "Initial Values", meta=(Bitmask, BitmaskEnum="/Script/Sequencer.ENavigationToolItemViewMode"))
+	UPROPERTY(Config, EditAnywhere, Category = "Editor UX", meta=(Bitmask, BitmaskEnum="/Script/Sequencer.ENavigationToolItemViewMode"))
 	ENavigationToolItemViewMode ItemDefaultViewMode = ENavigationToolItemViewMode::None;
 
 	/** The View Mode a Proxy Item supports by default */
-	UPROPERTY(Config, EditAnywhere, Category = "Initial Values", meta=(Bitmask, BitmaskEnum="/Script/Sequencer.ENavigationToolItemViewMode"))
+	UPROPERTY(Config, EditAnywhere, Category = "Editor UX", meta=(Bitmask, BitmaskEnum="/Script/Sequencer.ENavigationToolItemViewMode"))
 	ENavigationToolItemViewMode ItemProxyViewMode = ENavigationToolItemViewMode::None;
 
 	UPROPERTY(Config)
@@ -164,4 +170,12 @@ private:
 	/** Last saved size of the filter bar to restore after closed */
 	UPROPERTY(config)
 	float LastFilterBarSizeCoefficient = 0.05f;
+
+	/** If true, syncs selections in Sequencer to Navigation Tool */
+	UPROPERTY(config, EditAnywhere, Category = "Editor UX")
+	bool bSyncSelectionToNavigationTool = true;
+
+	/** If true, syncs selections in Navigation Tool to Sequencer */
+	UPROPERTY(config, EditAnywhere, Category = "Editor UX")
+	bool bSyncSelectionToSequencer = true;
 };
diff --git a/Engine/Source/Editor/Sequencer/Public/Tools/NavigationTool/Providers/INavigationToolProvider.h b/Engine/Source/Editor/Sequencer/Public/Tools/NavigationTool/Providers/INavigationToolProvider.h
--- a/Engine/Source/Editor/Sequencer/Public/Tools/NavigationTool/Providers/INavigationToolProvider.h
+++ b/Engine/Source/Editor/Sequencer/Public/Tools/NavigationTool/Providers/INavigationToolProvider.h
@@ -37,6 +37,11 @@ public:
 	/** Saves the editor only state data to be restored */
 	virtual void SetSaveState(const INavigationTool& InTool, const FNavigationToolSaveState& InSaveState) const = 0;
 
+	virtual void OnRegistered() {}
+	virtual void OnUnregistered() {}
+
+	virtual void BindCommands(const TSharedRef<FUICommandList>& InCommandList) {}
+
 	/**
 	 * Extends the available list of columns that can be displayed in the Navigation Tool instance.
 	 * Users can override this method to add or modify columns in the Navigation Tool interface.
@@ -64,7 +69,6 @@ public:
 		, TArray<FNavigationToolItemPtr>& OutChildren
 		, const bool bInRecursive) {}
 
-
 	/**
 	 * Extends the list of item proxies for a specific item in a Navigation Tool instance.
 	 * Users can override this method to add or manipulate item proxies associated with a given item.
