[Sequencer] Fix Horde PVS warning

diff --git a/Engine/Source/Editor/Sequencer/Private/Tools/NavigationTool/Items/NavigationToolItemId.cpp b/Engine/Source/Editor/Sequencer/Private/Tools/NavigationTool/Items/NavigationToolItemId.cpp
--- a/Engine/Source/Editor/Sequencer/Private/Tools/NavigationTool/Items/NavigationToolItemId.cpp
+++ b/Engine/Source/Editor/Sequencer/Private/Tools/NavigationTool/Items/NavigationToolItemId.cpp
@@ -110,7 +110,7 @@ FNavigationToolItemId::FNavigationToolItemId(const FViewModelPtr& InViewModel)
 		ReferenceId = ObjectBindingExtension->GetObjectGuid().ToString();
 	}
 
-	FNavigationToolItemId(nullptr, Sequence, Section, SectionIndex, ReferenceId);
+	this->FNavigationToolItemId::FNavigationToolItemId(nullptr, Sequence, Section, SectionIndex, ReferenceId);
 	CalculateTypeHash();
 }
 
