[Sequencer]
- Move filter bar options to submenu
- Fix filter bar layout toggle regression

diff --git a/Engine/Source/Editor/Sequencer/Private/Filters/Menus/SequencerViewOptionsMenu.cpp b/Engine/Source/Editor/Sequencer/Private/Filters/Menus/SequencerViewOptionsMenu.cpp
--- a/Engine/Source/Editor/Sequencer/Private/Filters/Menus/SequencerViewOptionsMenu.cpp
+++ b/Engine/Source/Editor/Sequencer/Private/Filters/Menus/SequencerViewOptionsMenu.cpp
@@ -23,6 +23,7 @@ TSharedRef<SWidget> FSequencerViewOptionsMenu::CreateMenu(const TWeakPtr<FSequen
 	if (!UToolMenus::Get()->IsMenuRegistered(FilterMenuName))
 	{
 		UToolMenu* const Menu = UToolMenus::Get()->RegisterMenu(FilterMenuName);
+		Menu->bShouldCloseWindowAfterMenuSelection = false;
 		Menu->AddDynamicSection(NAME_None, FNewToolMenuDelegate::CreateLambda([this](UToolMenu* const InMenu)
 			{
 				if (USequencerMenuContext* const Context = InMenu->FindContext<USequencerMenuContext>())
@@ -60,7 +61,6 @@ void FSequencerViewOptionsMenu::PopulateMenu(UToolMenu* const InMenu)
 	PopulateFiltersSection(Menu);
 	PopulateSortAndOrganizeSection(Menu);
 	PopulateFilterOptionsSection(Menu);
-	PopulateLayoutSection(Menu);
 }
 
 void FSequencerViewOptionsMenu::PopulateFiltersSection(UToolMenu& InMenu)
@@ -143,41 +143,32 @@ void FSequencerViewOptionsMenu::PopulateFilterOptionsSection(UToolMenu& InMenu)
 			FIsActionChecked::CreateRaw(this, &FSequencerViewOptionsMenu::IsLinkFiltersWithCurveEditor)
 		),
 		EUserInterfaceActionType::ToggleButton);
+
+	OptionsSection.AddSubMenu(TEXT("FilterBarOptions")
+		, LOCTEXT("FilterBarOptionsSubMenu", "Filter Bar Options")
+		, LOCTEXT("FilterBarOptionsSubMenuTooltip", "Filter Bar Options")
+		, FNewToolMenuDelegate::CreateSP(this, &FSequencerViewOptionsMenu::CreateFilterBarOptionsSection));
 }
 
-void FSequencerViewOptionsMenu::PopulateLayoutSection(UToolMenu& InMenu)
+void FSequencerViewOptionsMenu::CreateFilterBarOptionsSection(UToolMenu* const InMenu)
 {
+	if (!IsValid(InMenu))
+	{
+		return;
+	}
+
 	const FSequencerTrackFilterCommands& TrackFilterCommands = FSequencerTrackFilterCommands::Get();
 
-	FToolMenuSection& VisibilitySection = InMenu.FindOrAddSection(TEXT("Visibility"), LOCTEXT("VisibilityHeading", "Filter Bar"));
+	FToolMenuSection& FilterBarVisibilitySection = InMenu->FindOrAddSection(TEXT("FilterBarVisibility")
+		, LOCTEXT("FilterBarVisibilityHeading", "Filter Bar"));
 
-	VisibilitySection.AddMenuEntry(TrackFilterCommands.ToggleFilterBarVisibility);
+	FilterBarVisibilitySection.AddMenuEntry(TrackFilterCommands.ToggleFilterBarVisibility);
 
-	FToolMenuSection& LayoutSection = InMenu.FindOrAddSection(TEXT("Layout"), LOCTEXT("LayoutHeading", "Filter Bar Layout"));
+	FToolMenuSection& FilterBarLayoutSection = InMenu->FindOrAddSection(TEXT("FilterBarLayout")
+		, LOCTEXT("FilterBarLayoutHeading", "Filter Bar Layout"));
 
-	LayoutSection.AddMenuEntry(
-		TEXT("VerticalLayout"),
-		LOCTEXT("FilterListVerticalLayout", "Vertical"),
-		LOCTEXT("FilterListVerticalLayoutToolTip", "Swap to a vertical layout for the filter bar"),
-		FSlateIcon(),
-		FUIAction(
-			FExecuteAction::CreateRaw(this, &FSequencerViewOptionsMenu::SetFilterLayout, EFilterBarLayout::Vertical),
-			FCanExecuteAction(),
-			FIsActionChecked::CreateRaw(this, &FSequencerViewOptionsMenu::IsFilterLayout, EFilterBarLayout::Vertical)
-		),
-		EUserInterfaceActionType::RadioButton);
-
-	LayoutSection.AddMenuEntry(
-		TEXT("HorizontalLayout"),
-		LOCTEXT("FilterListHorizontalLayout", "Horizontal"),
-		LOCTEXT("FilterListHorizontalLayoutToolTip", "Swap to a Horizontal layout for the filter bar"),
-		FSlateIcon(),
-		FUIAction(
-			FExecuteAction::CreateRaw(this, &FSequencerViewOptionsMenu::SetFilterLayout, EFilterBarLayout::Horizontal),
-			FCanExecuteAction(),
-			FIsActionChecked::CreateRaw(this, &FSequencerViewOptionsMenu::IsFilterLayout, EFilterBarLayout::Horizontal)
-		),
-		EUserInterfaceActionType::RadioButton);
+	FilterBarLayoutSection.AddMenuEntry(TrackFilterCommands.SetToVerticalLayout);
+	FilterBarLayoutSection.AddMenuEntry(TrackFilterCommands.SetToHorizontalLayout);
 }
 
 bool FSequencerViewOptionsMenu::IsFilterLayout(const EFilterBarLayout InLayout) const
diff --git a/Engine/Source/Editor/Sequencer/Private/Filters/Menus/SequencerViewOptionsMenu.h b/Engine/Source/Editor/Sequencer/Private/Filters/Menus/SequencerViewOptionsMenu.h
--- a/Engine/Source/Editor/Sequencer/Private/Filters/Menus/SequencerViewOptionsMenu.h
+++ b/Engine/Source/Editor/Sequencer/Private/Filters/Menus/SequencerViewOptionsMenu.h
@@ -20,7 +20,8 @@ protected:
 	void PopulateFiltersSection(UToolMenu& InMenu);
 	void PopulateSortAndOrganizeSection(UToolMenu& InMenu);
 	void PopulateFilterOptionsSection(UToolMenu& InMenu);
-	void PopulateLayoutSection(UToolMenu& InMenu);
+
+	void CreateFilterBarOptionsSection(UToolMenu* const InMenu);
 
 	bool IsFilterLayout(const EFilterBarLayout InLayout) const;
 	void SetFilterLayout(const EFilterBarLayout InLayout);
diff --git a/Engine/Source/Editor/Sequencer/Private/Filters/SequencerFilterBar.cpp b/Engine/Source/Editor/Sequencer/Private/Filters/SequencerFilterBar.cpp
--- a/Engine/Source/Editor/Sequencer/Private/Filters/SequencerFilterBar.cpp
+++ b/Engine/Source/Editor/Sequencer/Private/Filters/SequencerFilterBar.cpp
@@ -187,6 +187,16 @@ void FSequencerFilterBar::BindCommands()
 		FCanExecuteAction(),
 		FIsActionChecked::CreateSP(SequencerWidget, &SSequencer::IsFilterBarVisible));
 
+	CommandList->MapAction(TrackFilterCommands.SetToVerticalLayout,
+		FExecuteAction::CreateSP(this, &FSequencerFilterBar::SetToVerticalLayout),
+		FCanExecuteAction(),
+		FIsActionChecked::CreateSP(this, &FSequencerFilterBar::IsFilterBarLayout, EFilterBarLayout::Vertical));
+
+	CommandList->MapAction(TrackFilterCommands.SetToHorizontalLayout,
+		FExecuteAction::CreateSP(this, &FSequencerFilterBar::SetToHorizontalLayout),
+		FCanExecuteAction(),
+		FIsActionChecked::CreateSP(this, &FSequencerFilterBar::IsFilterBarLayout, EFilterBarLayout::Horizontal));
+
 	CommandList->MapAction(TrackFilterCommands.ResetFilters,
 		FExecuteAction::CreateSP(this, &FSequencerFilterBar::ResetFilters),
 		FCanExecuteAction::CreateSP(this, &FSequencerFilterBar::CanResetFilters));
diff --git a/Engine/Source/Editor/Sequencer/Private/Filters/SequencerFilterBar.h b/Engine/Source/Editor/Sequencer/Private/Filters/SequencerFilterBar.h
--- a/Engine/Source/Editor/Sequencer/Private/Filters/SequencerFilterBar.h
+++ b/Engine/Source/Editor/Sequencer/Private/Filters/SequencerFilterBar.h
@@ -64,7 +64,7 @@ public:
 	virtual bool CanResetFilters() const override;
 	virtual void ResetFilters() override;
 
-	FSequencerFiltersChanged& GetOnFiltersChanged() { return FiltersChangedEvent; }
+	FSequencerFiltersChanged& OnFiltersChanged() { return FiltersChangedEvent; }
 
 	TSharedRef<FSequencerTrackFilter_Text> GetTextFilter() const;
 	FText GetFilterErrorText() const;
diff --git a/Engine/Source/Editor/Sequencer/Private/Filters/SequencerTrackFilterCommands.cpp b/Engine/Source/Editor/Sequencer/Private/Filters/SequencerTrackFilterCommands.cpp
--- a/Engine/Source/Editor/Sequencer/Private/Filters/SequencerTrackFilterCommands.cpp
+++ b/Engine/Source/Editor/Sequencer/Private/Filters/SequencerTrackFilterCommands.cpp
@@ -17,6 +17,9 @@ void FSequencerTrackFilterCommands::RegisterCommands()
 {
 	UI_COMMAND(ToggleFilterBarVisibility, "Filter Bar Visibility", "Toggle the visibility of the filter bar. The filter bar will only be displayed when there are enabled filters", EUserInterfaceActionType::ToggleButton, FInputChord(EModifierKey::Control | EModifierKey::Shift, EKeys::F));
 
+	UI_COMMAND(SetToVerticalLayout, "Vertical", "Swap to a vertical layout for the filter bar", EUserInterfaceActionType::RadioButton, FInputChord());
+	UI_COMMAND(SetToHorizontalLayout, "Horizontal", "Swap to a horizontal layout for the filter bar", EUserInterfaceActionType::RadioButton, FInputChord());
+
 	UI_COMMAND(ResetFilters, "Reset Filters", "Reset all enabled filters", EUserInterfaceActionType::Button, FInputChord());
 
 	UI_COMMAND(ToggleMuteFilters, "Mute Filters", "Mute or Unmute all active filters", EUserInterfaceActionType::ToggleButton, FInputChord());
diff --git a/Engine/Source/Editor/Sequencer/Private/Filters/SequencerTrackFilterCommands.h b/Engine/Source/Editor/Sequencer/Private/Filters/SequencerTrackFilterCommands.h
--- a/Engine/Source/Editor/Sequencer/Private/Filters/SequencerTrackFilterCommands.h
+++ b/Engine/Source/Editor/Sequencer/Private/Filters/SequencerTrackFilterCommands.h
@@ -21,6 +21,9 @@ public:
 
 	TSharedPtr<FUICommandInfo> ToggleFilterBarVisibility;
 
+	TSharedPtr<FUICommandInfo> SetToVerticalLayout;
+	TSharedPtr<FUICommandInfo> SetToHorizontalLayout;
+
 	TSharedPtr<FUICommandInfo> ResetFilters;
 
 	TSharedPtr<FUICommandInfo> ToggleMuteFilters;
diff --git a/Engine/Source/Editor/Sequencer/Private/Filters/Widgets/SSequencerFilterBar.cpp b/Engine/Source/Editor/Sequencer/Private/Filters/Widgets/SSequencerFilterBar.cpp
--- a/Engine/Source/Editor/Sequencer/Private/Filters/Widgets/SSequencerFilterBar.cpp
+++ b/Engine/Source/Editor/Sequencer/Private/Filters/Widgets/SSequencerFilterBar.cpp
@@ -20,7 +20,7 @@ SSequencerFilterBar::~SSequencerFilterBar()
 {
 	if (const TSharedPtr<FSequencerFilterBar> FilterBar = WeakFilterBar.Pin())
 	{
-		FilterBar->GetOnFiltersChanged().RemoveAll(this);
+		FilterBar->OnFiltersChanged().RemoveAll(this);
 	}
 
 	ContextMenu.Reset();
@@ -87,7 +87,7 @@ void SSequencerFilterBar::Construct(const FArguments& InArgs, const TSharedRef<F
 
 	CreateFilterWidgetsFromConfig();
 
-	InFilterBar->GetOnFiltersChanged().AddSP(this, &SSequencerFilterBar::OnFiltersChanged);
+	InFilterBar->OnFiltersChanged().AddSP(this, &SSequencerFilterBar::OnFiltersChanged);
 }
 
 FReply SSequencerFilterBar::OnMouseButtonUp(const FGeometry& MyGeometry, const FPointerEvent& MouseEvent)
diff --git a/Engine/Source/Editor/Sequencer/Private/SSequencer.cpp b/Engine/Source/Editor/Sequencer/Private/SSequencer.cpp
--- a/Engine/Source/Editor/Sequencer/Private/SSequencer.cpp
+++ b/Engine/Source/Editor/Sequencer/Private/SSequencer.cpp
@@ -274,7 +274,11 @@ void SSequencer::Construct(const FArguments& InArgs, TSharedRef<FSequencer> InSe
 
 	ViewOptionsMenu = MakeShared<FSequencerViewOptionsMenu>();
 
-	GetFilterBar()->GetOnFiltersChanged().AddSP(this, &SSequencer::OnTrackFiltersChanged);
+	if (const TSharedPtr<FSequencerFilterBar> FilterBar = GetFilterBar())
+	{
+		FilterBar->OnStateChanged().AddSP(this, &SSequencer::OnFilterBarStateChanged);
+		FilterBar->OnFiltersChanged().AddSP(this, &SSequencer::OnTrackFiltersChanged);
+	}
 
 	ChildSlot
 	[
@@ -4120,6 +4124,11 @@ void SSequencer::SetFilterBarLayout(const EFilterBarLayout InLayout)
 	RebuildFilterBarContent();
 }
 
+void SSequencer::OnFilterBarStateChanged(const bool bInIsVisible, const EFilterBarLayout InNewLayout)
+{
+	RebuildFilterBarContent();
+}
+
 void SSequencer::OnTrackFiltersChanged(const ESequencerFilterChange InChangeType, const TSharedRef<FSequencerTrackFilter>& InFilter)
 {
 	RebuildFilterBarContent();
diff --git a/Engine/Source/Editor/Sequencer/Private/SSequencer.h b/Engine/Source/Editor/Sequencer/Private/SSequencer.h
--- a/Engine/Source/Editor/Sequencer/Private/SSequencer.h
+++ b/Engine/Source/Editor/Sequencer/Private/SSequencer.h
@@ -704,6 +704,8 @@ private:
 
 	void OnSidebarStateChanged(const FSidebarState& InNewState);
 
+	void OnFilterBarStateChanged(const bool bInIsVisible, const EFilterBarLayout InNewLayout);
+
 	void OnTrackFiltersChanged(const ESequencerFilterChange InChangeType, const TSharedRef<FSequencerTrackFilter>& InFilter);
 
 	/** This is the main container widget for the Sequencer (minus the toolbar). This is the parent of
diff --git a/Engine/Source/Editor/Sequencer/Private/Tools/NavigationTool/Menus/NavigationToolToolbarMenu.cpp b/Engine/Source/Editor/Sequencer/Private/Tools/NavigationTool/Menus/NavigationToolToolbarMenu.cpp
--- a/Engine/Source/Editor/Sequencer/Private/Tools/NavigationTool/Menus/NavigationToolToolbarMenu.cpp
+++ b/Engine/Source/Editor/Sequencer/Private/Tools/NavigationTool/Menus/NavigationToolToolbarMenu.cpp
@@ -133,7 +133,7 @@ void FNavigationToolToolbarMenu::CreateSettingsMenu(UToolMenu* const InMenu)
 
 	MiscSection.AddSubMenu(TEXT("FilterBarOptions")
 		, LOCTEXT("FilterBarOptionsSubMenu", "Filter Bar Options")
-		, LOCTEXT("FilterBarOptionsSubMenuTooltip", "")
+		, LOCTEXT("FilterBarOptionsSubMenuTooltip", "Filter Bar Options")
 		, FNewToolMenuDelegate::CreateSP(this, &FNavigationToolToolbarMenu::CreateFilterBarOptionsMenu));
 
 	MiscSection.AddMenuEntry(ToolCommands.OpenToolSettings);
