Sequencer: Prioritize keys over sections for Select All in Selection Range

When there are keys to be selected, don't also select sections. This mimics what happens when you do a marquee selection so that you don't move both keys and their sections at the same time (however you can still move keys and other sections).

diff --git a/Engine/Source/Editor/Sequencer/Private/Sequencer.cpp b/Engine/Source/Editor/Sequencer/Private/Sequencer.cpp
--- a/Engine/Source/Editor/Sequencer/Private/Sequencer.cpp
+++ b/Engine/Source/Editor/Sequencer/Private/Sequencer.cpp
@@ -2954,6 +2954,8 @@ void FSequencer::SelectInSelectionRange(const TSharedPtr<UE::Sequencer::FViewMod
 
 	TSharedPtr<FSequencerSelection> Selection = ViewModel->GetSelection();
 
+	TSet<UMovieSceneSection*> SectionsWithKeys;
+
 	if (bSelectKeys)
 	{
 		TArray<FKeyHandle> HandlesScratch;
@@ -2968,6 +2970,11 @@ void FSequencer::SelectInSelectionRange(const TSharedPtr<UE::Sequencer::FViewMod
 				HandlesScratch.Reset();
 				KeyArea->GetKeyHandles(HandlesScratch, SelectionRange);
 
+				if (HandlesScratch.Num() > 0)
+				{
+					SectionsWithKeys.Add(Section);
+				}
+
 				for (int32 Index = 0; Index < HandlesScratch.Num(); ++Index)
 				{
 					Selection->KeySelection.Select(Channel, HandlesScratch[Index]);
@@ -2983,7 +2990,11 @@ void FSequencer::SelectInSelectionRange(const TSharedPtr<UE::Sequencer::FViewMod
 			TRange<FFrameNumber> SectionRange = SectionModel->GetRange();
 			if (SectionRange.Overlaps(SelectionRange) && SectionRange.GetLowerBound().IsClosed() && SectionRange.GetUpperBound().IsClosed())
 			{
-				Selection->TrackArea.Select(SectionModel);
+				UMovieSceneSection* Section = SectionModel->GetSection();
+				if (!SectionsWithKeys.Contains(Section))
+				{
+					Selection->TrackArea.Select(SectionModel);
+				}
 			}
 		}
 	}
