[Sequencer] Fix horde PVS warning

diff --git a/Engine/Source/Editor/Sequencer/Private/Tools/NavigationTool/Items/NavigationToolItemId.cpp b/Engine/Source/Editor/Sequencer/Private/Tools/NavigationTool/Items/NavigationToolItemId.cpp
--- a/Engine/Source/Editor/Sequencer/Private/Tools/NavigationTool/Items/NavigationToolItemId.cpp
+++ b/Engine/Source/Editor/Sequencer/Private/Tools/NavigationTool/Items/NavigationToolItemId.cpp
@@ -51,26 +51,7 @@ FNavigationToolItemId::FNavigationToolItemId(const FNavigationToolItemPtr& InPar
 	, const int32 InSectionIndex
 	, const FString& InReferenceId)
 {
-	check(InSequence);
-
-	if (InParentItem.IsValid())
-	{
-		AddSeparatedSegment(Id, InParentItem->GetFullPath());
-	}
-
-	AddSeparatedSegment(Id, GetObjectPath(InSequence));
-
-	if (IsValid(InSection))
-	{
-		AddSeparatedSegment(Id, GetObjectPath(InSection));
-		AddSeparatedSegment(Id, FString::FromInt(InSectionIndex));
-	}
-
-	if (!InReferenceId.IsEmpty())
-	{
-		AddSeparatedSegment(Id, InReferenceId);
-	}
-
+	ConstructId(InParentItem, InSequence, InSection, InSectionIndex, InReferenceId);
 	CalculateTypeHash();
 }
 
@@ -110,10 +91,37 @@ FNavigationToolItemId::FNavigationToolItemId(const FViewModelPtr& InViewModel)
 		ReferenceId = ObjectBindingExtension->GetObjectGuid().ToString();
 	}
 
-	FNavigationToolItemId(nullptr, Sequence, Section, SectionIndex, ReferenceId);
+	ConstructId(nullptr, Sequence, Section, SectionIndex, ReferenceId);
 	CalculateTypeHash();
 }
 
+void FNavigationToolItemId::ConstructId(const FNavigationToolItemPtr& InParentItem
+	, const UMovieSceneSequence* const InSequence
+	, const UMovieSceneSection* const InSection
+	, const int32 InSectionIndex
+	, const FString& InReferenceId)
+{
+	check(InSequence);
+
+	if (InParentItem.IsValid())
+	{
+		AddSeparatedSegment(Id, InParentItem->GetFullPath());
+	}
+
+	AddSeparatedSegment(Id, GetObjectPath(InSequence));
+
+	if (IsValid(InSection))
+	{
+		AddSeparatedSegment(Id, GetObjectPath(InSection));
+		AddSeparatedSegment(Id, FString::FromInt(InSectionIndex));
+	}
+
+	if (!InReferenceId.IsEmpty())
+	{
+		AddSeparatedSegment(Id, InReferenceId);
+	}
+}
+
 FNavigationToolItemId::FNavigationToolItemId(const FNavigationToolItemId& Other)
 {
 	*this = Other;
diff --git a/Engine/Source/Editor/Sequencer/Private/Tools/NavigationTool/NavigationToolView.cpp b/Engine/Source/Editor/Sequencer/Private/Tools/NavigationTool/NavigationToolView.cpp
--- a/Engine/Source/Editor/Sequencer/Private/Tools/NavigationTool/NavigationToolView.cpp
+++ b/Engine/Source/Editor/Sequencer/Private/Tools/NavigationTool/NavigationToolView.cpp
@@ -52,17 +52,6 @@ FNavigationToolView::~FNavigationToolView()
 	{
 		ToolSettings->OnSettingChanged().RemoveAll(this);
 	}
-
-	if (const TSharedPtr<ISequencer> Sequencer = GetSequencer())
-	{
-		if (const TSharedPtr<FSequencerEditorViewModel> ViewModel = Sequencer->GetViewModel())
-		{
-			if (const TSharedPtr<FSequencerSelection> SequencerSelection = ViewModel->GetSelection())
-			{
-				SequencerSelection->OnChanged.AddSP(this, &FNavigationToolView::OnSequencerOutlinerSelectionChanged);
-			}
-		}
-	}
 }
 
 void FNavigationToolView::Init(const TSharedRef<FNavigationTool>& InTool, const TSharedPtr<FUICommandList>& InBaseCommandList)
@@ -89,14 +78,6 @@ void FNavigationToolView::Init(const TSharedRef<FNavigationTool>& InTool, const
 			}
 		});
 
-	if (const TSharedPtr<FSequencerEditorViewModel> ViewModel = Sequencer->GetViewModel())
-	{
-		if (const TSharedPtr<FSequencerSelection> SequencerSelection = ViewModel->GetSelection())
-		{
-			SequencerSelection->OnChanged.AddSP(this, &FNavigationToolView::OnSequencerOutlinerSelectionChanged);
-		}
-	}
-
 	ToolSettings->OnSettingChanged().AddSP(this, &FNavigationToolView::OnToolSettingsChanged);
 
 	ToolViewWidget = SNew(SNavigationToolView, SharedThis(this));
@@ -1679,50 +1660,6 @@ void FNavigationToolView::RefreshTool(bool bInImmediateRefresh)
 	}
 }
 
-void FNavigationToolView::OnSequencerOutlinerSelectionChanged()
-{
-	using namespace UE::Sequencer;
-
-	const TSharedPtr<INavigationTool> Tool = GetOwnerTool();
-	if (!Tool.IsValid())
-	{
-		return;
-	}
-
-	const TSharedPtr<ISequencer> Sequencer = Tool->GetSequencer();
-	if (!Sequencer.IsValid())
-	{
-		return;
-	}
-
-	// @TODO: need to stop cyclical selection triggering
-	/*ClearItemSelection(false);
-
-	const TSharedPtr<FSequencerSelection> Selection = Sequencer->GetViewModel()->GetSelection();
-
-	TSet<FNavigationToolItemPtr> ItemsToSelect;
-
-	for (const TViewModelPtr<IOutlinerExtension>& OutlinerItem : Selection->Outliner)
-	{
-		const FNavigationToolItemPtr FoundItem = Tool->FindItem(FNavigationToolItemId(OutlinerItem));
-		if (FoundItem.IsValid())
-		{
-			ItemsToSelect.Add(FoundItem);
-		}
-	}
-
-	for (const TViewModelPtr<FViewModel>& TrackArea : Selection->TrackArea)
-	{
-		const FNavigationToolItemPtr FoundItem = Tool->FindItem(FNavigationToolItemId(TrackArea));
-		if (FoundItem.IsValid())
-		{
-			ItemsToSelect.Add(FoundItem);
-		}
-	}
-
-	SelectItems(ItemsToSelect.Array(), ENavigationToolItemSelectionFlags::ScrollIntoView);*/
-}
-
 void FNavigationToolView::EnsureToolViewCount(const int32 InToolViewId) const
 {
 	const TSharedPtr<FNavigationTool> Tool = WeakTool.Pin();
diff --git a/Engine/Source/Editor/Sequencer/Private/Tools/NavigationTool/NavigationToolView.h b/Engine/Source/Editor/Sequencer/Private/Tools/NavigationTool/NavigationToolView.h
--- a/Engine/Source/Editor/Sequencer/Private/Tools/NavigationTool/NavigationToolView.h
+++ b/Engine/Source/Editor/Sequencer/Private/Tools/NavigationTool/NavigationToolView.h
@@ -27,7 +27,7 @@ class FNavigationToolFilterBar;
 class FNavigationToolItemContextMenu;
 class FNavigationToolProvider;
 class INavigationToolColumn;
-class INavigationToolProvider;
+
 class SNavigationToolView;
 
 /**
@@ -340,8 +340,6 @@ protected:
 	/** Triggers a Refresh on the FNavigationTool */
 	void RefreshTool(bool bInImmediateRefresh);
 
-	void OnSequencerOutlinerSelectionChanged();
-
 	void EnsureToolViewCount(const int32 InToolViewId) const;
 
 	/**  */
diff --git a/Engine/Source/Editor/Sequencer/Public/Tools/NavigationTool/Items/NavigationToolItemId.h b/Engine/Source/Editor/Sequencer/Public/Tools/NavigationTool/Items/NavigationToolItemId.h
--- a/Engine/Source/Editor/Sequencer/Public/Tools/NavigationTool/Items/NavigationToolItemId.h
+++ b/Engine/Source/Editor/Sequencer/Public/Tools/NavigationTool/Items/NavigationToolItemId.h
@@ -107,6 +107,20 @@ struct SEQUENCER_API FNavigationToolItemId
 private:
 	static void AddSeparatedSegment(FString& OutString, const FString& InSegment);
 
+	/**
+	 * Constructs an Item Id for most sequence items
+	 * @param InParentItem The parent item holding the item proxy
+	 * @param InSequence The parent sequence object
+	 * @param InSection The parent section this sequence belongs to, if any
+	 * @param InSectionIndex The parent section index, if any. Should be specified if InSection is used.
+	 * @param InReferenceId The optional additional reference string
+	 */
+	void ConstructId(const FNavigationToolItemPtr& InParentItem
+		, const UMovieSceneSequence* const InSequence
+		, const UMovieSceneSection* const InSection = nullptr
+		, const int32 InSectionIndex = 0
+		, const FString& InReferenceId = FString());
+
 	void CalculateTypeHash();
 
 	FString Id;
