[Sequencer]
- Sort Sequencer and Nav Tool filter widgets alphabetically
- Change to use AddSP for Sequencer and Nav Tool filter bars to fix crash
- Remove "Disable All Filters" menu options as they are the same as "Reset Filters"

diff --git a/Engine/Source/Editor/Sequencer/Private/Filters/Menus/SequencerFilterBarContextMenu.cpp b/Engine/Source/Editor/Sequencer/Private/Filters/Menus/SequencerFilterBarContextMenu.cpp
--- a/Engine/Source/Editor/Sequencer/Private/Filters/Menus/SequencerFilterBarContextMenu.cpp
+++ b/Engine/Source/Editor/Sequencer/Private/Filters/Menus/SequencerFilterBarContextMenu.cpp
@@ -116,11 +116,6 @@ void FSequencerFilterBarContextMenu::PopulateFilterBulkOptionsSection(UToolMenu&
 		, TrackFilterCommands.ToggleMuteFilters->GetLabel()
 		, TrackFilterCommands.ToggleMuteFilters->GetDescription()
 		, FSlateIcon(FAppStyle::Get().GetStyleSetName(), TEXT("Denied")));
-
-	Section.AddMenuEntry(TrackFilterCommands.DisableAllFilters
-		, TrackFilterCommands.DisableAllFilters->GetLabel()
-		, TrackFilterCommands.DisableAllFilters->GetDescription()
-		, FSlateIcon(FAppStyle::Get().GetStyleSetName(), TEXT("Icons.Minus")));
 }
 
 void FSequencerFilterBarContextMenu::OnActivateAllFilters(const bool bInActivate)
diff --git a/Engine/Source/Editor/Sequencer/Private/Filters/Menus/SequencerTrackFilterContextMenu.cpp b/Engine/Source/Editor/Sequencer/Private/Filters/Menus/SequencerTrackFilterContextMenu.cpp
--- a/Engine/Source/Editor/Sequencer/Private/Filters/Menus/SequencerTrackFilterContextMenu.cpp
+++ b/Engine/Source/Editor/Sequencer/Private/Filters/Menus/SequencerTrackFilterContextMenu.cpp
@@ -35,7 +35,7 @@ TSharedRef<SWidget> FSequencerTrackFilterContextMenu::CreateMenuWidget(const TSh
 
 	USequencerFilterMenuContext* const ContextObject = NewObject<USequencerFilterMenuContext>();
 	ContextObject->Init(InFilter);
-	ContextObject->OnPopulateFilterBarMenu = FOnPopulateFilterBarMenu::CreateRaw(this, &FSequencerTrackFilterContextMenu::PopulateMenu);
+	ContextObject->OnPopulateFilterBarMenu = FOnPopulateFilterBarMenu::CreateSP(this, &FSequencerTrackFilterContextMenu::PopulateMenu);
 
 	const FToolMenuContext MenuContext(CommandList, nullptr, ContextObject);
 	return UToolMenus::Get()->GenerateWidget(FilterMenuName, MenuContext);
@@ -80,13 +80,13 @@ void FSequencerTrackFilterContextMenu::PopulateFilterOptionsSection(UToolMenu& I
 		FText::Format(LOCTEXT("ActivateOnlyThisFilter", "Activate Only: {0}"), FilterName),
 		LOCTEXT("ActivateOnlyThisFilterTooltip", "Activate only this filter from the list."),
 		FSlateIcon(FAppStyle::GetAppStyleSetName(), TEXT("Icons.Plus")),
-		FUIAction(FExecuteAction::CreateRaw(this, &FSequencerTrackFilterContextMenu::OnActivateWithFilterException)));
+		FUIAction(FExecuteAction::CreateSP(this, &FSequencerTrackFilterContextMenu::OnActivateWithFilterException)));
 
 	Section.AddMenuEntry(TEXT("DisableFilter"),
 		FText::Format(LOCTEXT("DisableFilter", "Remove: {0}"), FilterName),
 		LOCTEXT("DisableFilterTooltip", "Disable this filter and remove it from the list. It can be added again in the filters menu."),
 		FSlateIcon(FAppStyle::GetAppStyleSetName(), TEXT("Icons.Minus")),
-		FUIAction(FExecuteAction::CreateRaw(this, &FSequencerTrackFilterContextMenu::OnDisableFilter)));
+		FUIAction(FExecuteAction::CreateSP(this, &FSequencerTrackFilterContextMenu::OnDisableFilter)));
 }
 
 void FSequencerTrackFilterContextMenu::PopulateCustomFilterOptionsSection(UToolMenu& InMenu)
@@ -105,13 +105,13 @@ void FSequencerTrackFilterContextMenu::PopulateCustomFilterOptionsSection(UToolM
 		FText::Format(LOCTEXT("EditCustomTextFilter", "Edit: {0}"), CustomTextFilter->GetDisplayName()),
 		LOCTEXT("EditCustomTextFilterTooltip", "Edit this custom text filter saved to config."),
 		FSlateIcon(FAppStyle::GetAppStyleSetName(), TEXT("Icons.Edit")),
-		FUIAction(FExecuteAction::CreateRaw(this, &FSequencerTrackFilterContextMenu::OnEditFilter)));
+		FUIAction(FExecuteAction::CreateSP(this, &FSequencerTrackFilterContextMenu::OnEditFilter)));
 
 	Section.AddMenuEntry(TEXT("DeleteCustomTextFilter"),
 		FText::Format(LOCTEXT("DeleteCustomTextFilter", "Delete: {0}"), CustomTextFilter->GetDisplayName()),
 		LOCTEXT("DeleteCustomTextFilterTooltip", "Delete this custom text filter from config.\n\nCAUTION: This cannot be undone!"),
 		FSlateIcon(FAppStyle::GetAppStyleSetName(), TEXT("Icons.Delete")),
-		FUIAction(FExecuteAction::CreateRaw(this, &FSequencerTrackFilterContextMenu::OnDeleteFilter)));
+		FUIAction(FExecuteAction::CreateSP(this, &FSequencerTrackFilterContextMenu::OnDeleteFilter)));
 }
 
 void FSequencerTrackFilterContextMenu::PopulateBulkOptionsSection(UToolMenu& InMenu)
@@ -146,11 +146,6 @@ void FSequencerTrackFilterContextMenu::PopulateBulkOptionsSection(UToolMenu& InM
 		, TrackFilterCommands.ToggleMuteFilters->GetLabel()
 		, TrackFilterCommands.ToggleMuteFilters->GetDescription()
 		, FSlateIcon(FAppStyle::Get().GetStyleSetName(), TEXT("Denied")));
-
-	Section.AddMenuEntry(TrackFilterCommands.DisableAllFilters
-		, TrackFilterCommands.DisableAllFilters->GetLabel()
-		, TrackFilterCommands.DisableAllFilters->GetDescription()
-		, FSlateIcon(FAppStyle::Get().GetStyleSetName(), TEXT("Icons.Minus")));
 }
 
 FText FSequencerTrackFilterContextMenu::GetFilterDisplayName() const
diff --git a/Engine/Source/Editor/Sequencer/Private/Filters/SequencerFilterBar.cpp b/Engine/Source/Editor/Sequencer/Private/Filters/SequencerFilterBar.cpp
--- a/Engine/Source/Editor/Sequencer/Private/Filters/SequencerFilterBar.cpp
+++ b/Engine/Source/Editor/Sequencer/Private/Filters/SequencerFilterBar.cpp
@@ -218,12 +218,10 @@ void FSequencerFilterBar::BindCommands()
 		FCanExecuteAction::CreateSP(this, &FSequencerFilterBar::HasAnyFilterEnabled));
 
 	CommandList->MapAction(TrackFilterCommands.ActivateAllFilters,
-		FExecuteAction::CreateSP(this, &FSequencerFilterBar::ActivateAllEnabledFilters, true, TArray<FString>()),
-		FCanExecuteAction::CreateSP(this, &FSequencerFilterBar::AreAllEnabledFiltersActive, true, TArray<FString>()));
+		FExecuteAction::CreateSP(this, &FSequencerFilterBar::ActivateAllEnabledFilters, true, TArray<FString>()));
 
 	CommandList->MapAction(TrackFilterCommands.DeactivateAllFilters,
-		FExecuteAction::CreateSP(this, &FSequencerFilterBar::ActivateAllEnabledFilters, false, TArray<FString>()),
-		FCanExecuteAction::CreateSP(this, &FSequencerFilterBar::AreAllEnabledFiltersActive, false, TArray<FString>()));
+		FExecuteAction::CreateSP(this, &FSequencerFilterBar::ActivateAllEnabledFilters, false, TArray<FString>()));
 
 	// Bind all filter actions
 	UMovieSceneSequence* const FocusedSequence = Sequencer.GetFocusedMovieSceneSequence();
@@ -633,7 +631,8 @@ bool FSequencerFilterBar::SetFilterActive(const TSharedRef<FSequencerTrackFilter
 		InFilter->SetActive(bNewActive);
 		InFilter->ActiveStateChanged(bNewActive);
 
-		const ESequencerFilterChange FilterChangeType = bNewActive ? ESequencerFilterChange::Activate : ESequencerFilterChange::Deactivate;
+		const ESequencerFilterChange FilterChangeType = bNewActive
+			? ESequencerFilterChange::Activate : ESequencerFilterChange::Deactivate;
 		FiltersChangedEvent.Broadcast(FilterChangeType, InFilter);
 
 		if (bInRequestFilterUpdate)
@@ -900,7 +899,8 @@ bool FSequencerFilterBar::SetFilterEnabled(const TSharedRef<FSequencerTrackFilte
 	{
 		SequencerSettings->SaveConfig();
 
-		const ESequencerFilterChange FilterChangeType = bInEnabled ? ESequencerFilterChange::Enable : ESequencerFilterChange::Disable;
+		const ESequencerFilterChange FilterChangeType = bInEnabled
+			? ESequencerFilterChange::Enable : ESequencerFilterChange::Disable;
 		FiltersChangedEvent.Broadcast(FilterChangeType, InFilter);
 
 		if (!bInEnabled && IsFilterActive(InFilter))
@@ -935,7 +935,8 @@ void FSequencerFilterBar::EnableFilters(const bool bInEnable
 				const FString FilterName = InFilter->GetDisplayName().ToString();
 				if (Config.SetFilterEnabled(FilterName, bInEnable))
 				{
-					const ESequencerFilterChange FilterChangeType = bInEnable ? ESequencerFilterChange::Enable : ESequencerFilterChange::Disable;
+					const ESequencerFilterChange FilterChangeType = bInEnable
+						? ESequencerFilterChange::Enable : ESequencerFilterChange::Disable;
 					FiltersChangedEvent.Broadcast(FilterChangeType, InFilter);
 
 					if (!bInEnable && IsFilterActive(InFilter))
@@ -1136,7 +1137,8 @@ void FSequencerFilterBar::ActivateCustomTextFilters(const bool bInActivate, cons
 					CustomTextFilter->ActiveStateChanged(false);
 				}
 
-				const ESequencerFilterChange FilterChangeType = bInActivate ? ESequencerFilterChange::Activate : ESequencerFilterChange::Deactivate;
+				const ESequencerFilterChange FilterChangeType = bInActivate
+					? ESequencerFilterChange::Activate : ESequencerFilterChange::Deactivate;
 				FiltersChangedEvent.Broadcast(FilterChangeType, CustomTextFilter);
 			}
 		}
@@ -1167,7 +1169,8 @@ void FSequencerFilterBar::EnableCustomTextFilters(const bool bInEnable, const TA
 					CustomTextFilter->ActiveStateChanged(false);
 				}
 
-				const ESequencerFilterChange FilterChangeType = bInEnable ? ESequencerFilterChange::Enable : ESequencerFilterChange::Disable;
+				const ESequencerFilterChange FilterChangeType = bInEnable
+					? ESequencerFilterChange::Enable : ESequencerFilterChange::Disable;
 				FiltersChangedEvent.Broadcast(FilterChangeType, CustomTextFilter);
 			}
 		}
diff --git a/Engine/Source/Editor/Sequencer/Private/Filters/Widgets/SSequencerFilterBar.cpp b/Engine/Source/Editor/Sequencer/Private/Filters/Widgets/SSequencerFilterBar.cpp
--- a/Engine/Source/Editor/Sequencer/Private/Filters/Widgets/SSequencerFilterBar.cpp
+++ b/Engine/Source/Editor/Sequencer/Private/Filters/Widgets/SSequencerFilterBar.cpp
@@ -12,6 +12,7 @@
 #include "SequencerLog.h"
 #include "SSequencerFilter.h"
 #include "SSequencerCustomTextFilterDialog.h"
+#include "Tools/NavigationTool/Menus/NavigationToolFilterContextMenu.h"
 #include "Widgets/SSequencerSearchBox.h"
 
 #define LOCTEXT_NAMESPACE "SSequencerFilterBar"
@@ -23,8 +24,6 @@ SSequencerFilterBar::~SSequencerFilterBar()
 		FilterBar->OnFiltersChanged().RemoveAll(this);
 	}
 
-	ContextMenu.Reset();
-
 	if (SSequencerCustomTextFilterDialog::IsOpen())
 	{
 		SSequencerCustomTextFilterDialog::CloseWindow();
@@ -46,6 +45,8 @@ void SSequencerFilterBar::Construct(const FArguments& InArgs, const TSharedRef<F
 	bCanChangeOrientation = InArgs._CanChangeOrientation;
 	FilterPillStyle = InArgs._FilterPillStyle;
 
+	ContextMenu = MakeShared<FSequencerFilterBarContextMenu>();
+
 	HorizontalContainerWidget = SNew(SFilterBarClippingHorizontalBox)
 		.OnWrapButtonClicked(FOnGetContent::CreateSP(this, &SSequencerFilterBar::OnWrapButtonClicked))
 		.IsFocusable(false);
@@ -83,8 +84,6 @@ void SSequencerFilterBar::Construct(const FArguments& InArgs, const TSharedRef<F
 
 	AttachFilterSearchBox(InArgs._FilterSearchBox);
 
-	ContextMenu = MakeShared<FSequencerFilterBarContextMenu>();
-
 	CreateFilterWidgetsFromConfig();
 
 	InFilterBar->OnFiltersChanged().AddSP(this, &SSequencerFilterBar::OnFiltersChanged);
@@ -314,7 +313,7 @@ void SSequencerFilterBar::OnFilterDoubleClick(const TSharedRef<FSequencerTrackFi
 
 TSharedRef<SWidget> SSequencerFilterBar::OnGetMenuContent(const TSharedRef<FSequencerTrackFilter> InFilter)
 {
-	const TSharedRef<FSequencerTrackFilterContextMenu> FilterContextMenu = MakeShared<FSequencerTrackFilterContextMenu>();
+	FilterContextMenu = MakeShared<FSequencerTrackFilterContextMenu>();
 	return FilterContextMenu->CreateMenuWidget(InFilter);
 }
 
@@ -333,6 +332,11 @@ void SSequencerFilterBar::ActivateAllButThis(const bool bInActive, const TShared
 void SSequencerFilterBar::AddFilterWidget(const TSharedRef<FSequencerTrackFilter>& InFilter, const TSharedRef<SSequencerFilter>& InFilterWidget)
 {
 	FilterWidgets.Add(InFilter, InFilterWidget);
+	FilterWidgets.KeySort([](const TSharedRef<FSequencerTrackFilter>& InFilterA
+		, const TSharedRef<FSequencerTrackFilter>& InFilterB)
+		{
+			return InFilterA->GetDisplayName().CompareTo(InFilterB->GetDisplayName()) < 0;
+		});
 
 	AddWidgetToLayout(InFilterWidget);
 }
diff --git a/Engine/Source/Editor/Sequencer/Private/Filters/Widgets/SSequencerFilterBar.h b/Engine/Source/Editor/Sequencer/Private/Filters/Widgets/SSequencerFilterBar.h
--- a/Engine/Source/Editor/Sequencer/Private/Filters/Widgets/SSequencerFilterBar.h
+++ b/Engine/Source/Editor/Sequencer/Private/Filters/Widgets/SSequencerFilterBar.h
@@ -8,7 +8,7 @@
 class FSequencer;
 class FSequencerFilterBarContextMenu;
 class FSequencerFilterBar;
-class FSequencerTrackFilter_CustomText;
+class FSequencerTrackFilterContextMenu;
 class SFilterBarClippingHorizontalBox;
 class SFilterExpressionHelpDialog;
 class SSequencerFilter;
@@ -128,4 +128,5 @@ protected:
 	TSharedPtr<SFilterExpressionHelpDialog> TextExpressionHelpDialog;
 
 	TSharedPtr<FSequencerFilterBarContextMenu> ContextMenu;
+	TSharedPtr<FSequencerTrackFilterContextMenu> FilterContextMenu;
 };
diff --git a/Engine/Source/Editor/Sequencer/Private/Tools/NavigationTool/Filters/NavigationToolFilterBar.cpp b/Engine/Source/Editor/Sequencer/Private/Tools/NavigationTool/Filters/NavigationToolFilterBar.cpp
--- a/Engine/Source/Editor/Sequencer/Private/Tools/NavigationTool/Filters/NavigationToolFilterBar.cpp
+++ b/Engine/Source/Editor/Sequencer/Private/Tools/NavigationTool/Filters/NavigationToolFilterBar.cpp
@@ -174,12 +174,10 @@ void FNavigationToolFilterBar::BindCommands(const TSharedPtr<FUICommandList>& In
 		FCanExecuteAction::CreateSP(this, &FNavigationToolFilterBar::HasAnyFilterEnabled));
 
 	CommandList->MapAction(FilterCommands.ActivateAllFilters,
-		FExecuteAction::CreateSP(this, &FNavigationToolFilterBar::ActivateAllEnabledFilters, true, TArray<FString>()),
-		FCanExecuteAction::CreateSP(this, &FNavigationToolFilterBar::AreAllEnabledFiltersActive, false, TArray<FString>()));
+		FExecuteAction::CreateSP(this, &FNavigationToolFilterBar::ActivateAllEnabledFilters, true, TArray<FString>()));
 
 	CommandList->MapAction(FilterCommands.DeactivateAllFilters,
-		FExecuteAction::CreateSP(this, &FNavigationToolFilterBar::ActivateAllEnabledFilters, false, TArray<FString>()),
-		FCanExecuteAction::CreateSP(this, &FNavigationToolFilterBar::AreAllEnabledFiltersActive, true, TArray<FString>()));
+		FExecuteAction::CreateSP(this, &FNavigationToolFilterBar::ActivateAllEnabledFilters, false, TArray<FString>()));
 
 	// Bind all filter actions
 	const TSharedPtr<ISequencer> Sequencer = Tool.GetSequencer();
diff --git a/Engine/Source/Editor/Sequencer/Private/Tools/NavigationTool/Filters/NavigationToolFilterCommands.cpp b/Engine/Source/Editor/Sequencer/Private/Tools/NavigationTool/Filters/NavigationToolFilterCommands.cpp
--- a/Engine/Source/Editor/Sequencer/Private/Tools/NavigationTool/Filters/NavigationToolFilterCommands.cpp
+++ b/Engine/Source/Editor/Sequencer/Private/Tools/NavigationTool/Filters/NavigationToolFilterCommands.cpp
@@ -39,7 +39,7 @@ void FNavigationToolFilterCommands::RegisterCommands()
 		, LOCTEXT("ToggleMuteFilters", "Mute Filters")
 		, LOCTEXT("ToggleMuteFiltersTooltip", "Mute or Unmute all active filters")
 		, FSlateIcon(FAppStyle::Get().GetStyleSetName(), TEXT("Denied"))
-		, EUserInterfaceActionType::Button
+		, EUserInterfaceActionType::ToggleButton
 		, FInputChord());
 
 	FUICommandInfo::MakeCommandInfo(AsShared()
diff --git a/Engine/Source/Editor/Sequencer/Private/Tools/NavigationTool/Menus/NavigationToolFilterBarContextMenu.cpp b/Engine/Source/Editor/Sequencer/Private/Tools/NavigationTool/Menus/NavigationToolFilterBarContextMenu.cpp
--- a/Engine/Source/Editor/Sequencer/Private/Tools/NavigationTool/Menus/NavigationToolFilterBarContextMenu.cpp
+++ b/Engine/Source/Editor/Sequencer/Private/Tools/NavigationTool/Menus/NavigationToolFilterBarContextMenu.cpp
@@ -116,11 +116,6 @@ void FNavigationToolFilterBarContextMenu::PopulateFilterBulkOptionsSection(UTool
 		, FilterCommands.ToggleMuteFilters->GetLabel()
 		, FilterCommands.ToggleMuteFilters->GetDescription()
 		, FSlateIcon(FAppStyle::Get().GetStyleSetName(), TEXT("Denied")));
-
-	Section.AddMenuEntry(FilterCommands.DisableAllFilters
-		, FilterCommands.DisableAllFilters->GetLabel()
-		, FilterCommands.DisableAllFilters->GetDescription()
-		, FSlateIcon(FAppStyle::Get().GetStyleSetName(), TEXT("Icons.Minus")));
 }
 
 void FNavigationToolFilterBarContextMenu::OnActivateAllFilters(const bool bInActivate)
diff --git a/Engine/Source/Editor/Sequencer/Private/Tools/NavigationTool/Menus/NavigationToolFilterContextMenu.cpp b/Engine/Source/Editor/Sequencer/Private/Tools/NavigationTool/Menus/NavigationToolFilterContextMenu.cpp
--- a/Engine/Source/Editor/Sequencer/Private/Tools/NavigationTool/Menus/NavigationToolFilterContextMenu.cpp
+++ b/Engine/Source/Editor/Sequencer/Private/Tools/NavigationTool/Menus/NavigationToolFilterContextMenu.cpp
@@ -40,7 +40,7 @@ TSharedRef<SWidget> FNavigationToolFilterContextMenu::CreateMenuWidget(const TSh
 
 	UNavigationToolFilterMenuContext* const ContextObject = NewObject<UNavigationToolFilterMenuContext>();
 	ContextObject->Init(InFilter);
-	ContextObject->OnPopulateMenu = FOnPopulateFilterBarMenu::CreateRaw(this, &FNavigationToolFilterContextMenu::PopulateMenu);
+	ContextObject->OnPopulateMenu = FOnPopulateFilterBarMenu::CreateSP(this, &FNavigationToolFilterContextMenu::PopulateMenu);
 
 	const FToolMenuContext MenuContext(CommandList, nullptr, ContextObject);
 	return UToolMenus::Get()->GenerateWidget(FilterMenuName, MenuContext);
@@ -110,13 +110,13 @@ void FNavigationToolFilterContextMenu::PopulateCustomFilterOptionsSection(UToolM
 		FText::Format(LOCTEXT("EditCustomTextFilter", "Edit: {0}"), FilterDisplayName),
 		LOCTEXT("EditCustomTextFilterTooltip", "Edit this custom text filter saved to config."),
 		FSlateIcon(FAppStyle::GetAppStyleSetName(), TEXT("Icons.Edit")),
-		FExecuteAction::CreateRaw(this, &FNavigationToolFilterContextMenu::OnEditFilter));
+		FExecuteAction::CreateSP(this, &FNavigationToolFilterContextMenu::OnEditFilter));
 
 	Section.AddMenuEntry(TEXT("DeleteCustomTextFilter"),
 		FText::Format(LOCTEXT("DeleteCustomTextFilter", "Delete: {0}"), FilterDisplayName),
 		LOCTEXT("DeleteCustomTextFilterTooltip", "Delete this custom text filter from config.\n\nCAUTION: This cannot be undone!"),
 		FSlateIcon(FAppStyle::GetAppStyleSetName(), TEXT("Icons.Delete")),
-		FExecuteAction::CreateRaw(this, &FNavigationToolFilterContextMenu::OnDeleteFilter));
+		FExecuteAction::CreateSP(this, &FNavigationToolFilterContextMenu::OnDeleteFilter));
 }
 
 void FNavigationToolFilterContextMenu::PopulateBulkOptionsSection(UToolMenu& InMenu)
@@ -132,7 +132,6 @@ void FNavigationToolFilterContextMenu::PopulateBulkOptionsSection(UToolMenu& InM
 	Section.AddMenuEntry(FilterCommands.ResetFilters);
 	Section.AddSeparator(NAME_None);
 	Section.AddMenuEntry(FilterCommands.ToggleMuteFilters);
-	Section.AddMenuEntry(FilterCommands.DisableAllFilters);
 }
 
 void FNavigationToolFilterContextMenu::OnDisableFilter()
diff --git a/Engine/Source/Editor/Sequencer/Private/Tools/NavigationTool/Widgets/SNavigationToolFilterBar.cpp b/Engine/Source/Editor/Sequencer/Private/Tools/NavigationTool/Widgets/SNavigationToolFilterBar.cpp
--- a/Engine/Source/Editor/Sequencer/Private/Tools/NavigationTool/Widgets/SNavigationToolFilterBar.cpp
+++ b/Engine/Source/Editor/Sequencer/Private/Tools/NavigationTool/Widgets/SNavigationToolFilterBar.cpp
@@ -31,8 +31,6 @@ SNavigationToolFilterBar::~SNavigationToolFilterBar()
 		FilterBar->OnFiltersChanged().RemoveAll(this);
 	}
 
-	ContextMenu.Reset();
-
 	if (SNavigationToolCustomTextFilterDialog::IsOpen())
 	{
 		SNavigationToolCustomTextFilterDialog::CloseWindow();
@@ -54,6 +52,8 @@ void SNavigationToolFilterBar::Construct(const FArguments& InArgs, const TWeakPt
 	bCanChangeOrientation = InArgs._CanChangeOrientation;
 	FilterPillStyle = InArgs._FilterPillStyle;
 
+	ContextMenu = MakeShared<FNavigationToolFilterBarContextMenu>();
+
 	HorizontalContainerWidget = SNew(SFilterBarClippingHorizontalBox)
 		.OnWrapButtonClicked(FOnGetContent::CreateSP(this, &SNavigationToolFilterBar::OnWrapButtonClicked))
 		.IsFocusable(false);
@@ -91,8 +91,6 @@ void SNavigationToolFilterBar::Construct(const FArguments& InArgs, const TWeakPt
 
 	AttachFilterSearchBox(InArgs._FilterSearchBox);
 
-	ContextMenu = MakeShared<FNavigationToolFilterBarContextMenu>();
-
 	CreateFilterWidgetsFromConfig();
 
 	if (const TSharedPtr<FNavigationToolFilterBar> FilterBar = WeakFilterBar.Pin())
@@ -327,7 +325,7 @@ FSlateColor SNavigationToolFilterBar::GetFilterBlockColor(const TSharedRef<FNavi
 
 TSharedRef<SWidget> SNavigationToolFilterBar::OnGetMenuContent(const TSharedRef<FNavigationToolFilter> InFilter)
 {
-	const TSharedRef<FNavigationToolFilterContextMenu> FilterContextMenu = MakeShared<FNavigationToolFilterContextMenu>();
+	FilterContextMenu = MakeShared<FNavigationToolFilterContextMenu>();
 	return FilterContextMenu->CreateMenuWidget(InFilter);
 }
 
@@ -346,6 +344,11 @@ void SNavigationToolFilterBar::ActivateAllButThis(const bool bInActive, const TS
 void SNavigationToolFilterBar::AddFilterWidget(const TSharedRef<FNavigationToolFilter>& InFilter, const TSharedRef<SSequencerFilter>& InFilterWidget)
 {
 	FilterWidgets.Add(InFilter, InFilterWidget);
+	FilterWidgets.KeySort([](const TSharedRef<FNavigationToolFilter>& InFilterA
+		, const TSharedRef<FNavigationToolFilter>& InFilterB)
+		{
+			return InFilterA->GetDisplayName().CompareTo(InFilterB->GetDisplayName()) < 0;
+		});
 
 	AddWidgetToLayout(InFilterWidget);
 }
diff --git a/Engine/Source/Editor/Sequencer/Private/Tools/NavigationTool/Widgets/SNavigationToolFilterBar.h b/Engine/Source/Editor/Sequencer/Private/Tools/NavigationTool/Widgets/SNavigationToolFilterBar.h
--- a/Engine/Source/Editor/Sequencer/Private/Tools/NavigationTool/Widgets/SNavigationToolFilterBar.h
+++ b/Engine/Source/Editor/Sequencer/Private/Tools/NavigationTool/Widgets/SNavigationToolFilterBar.h
@@ -6,7 +6,6 @@
 #include "Tools/NavigationTool/Filters/Filters/NavigationToolFilterBase.h"
 #include "Tools/NavigationTool/Filters/NavigationToolFilterBar.h"
 
-class FSequencer;
 class SFilterBarClippingHorizontalBox;
 class SFilterExpressionHelpDialog;
 class SSequencerFilter;
@@ -19,6 +18,7 @@ namespace UE::Sequencer::NavigationTool
 class FNavigationToolFilter;
 class FNavigationToolFilterBar;
 class FNavigationToolFilterBarContextMenu;
+class FNavigationToolFilterContextMenu;
 enum class ENavigationToolFilterChange : uint8;
 
 class SNavigationToolFilterBar : public SCompoundWidget
@@ -81,8 +81,6 @@ public:
 	// Set the state of the filter bar. Muted means that the filters are muted, but the context menu is still enabled and accessible.
 	void SetMuted(bool bInMuted);
 
-	void OnFiltersChanged(const ENavigationToolFilterChange InChangeType, const TSharedRef<FNavigationToolFilter>& InFilter);
-
 protected:
 	void AddWidgetToLayout(const TSharedRef<SWidget>& InWidget);
 	void RemoveWidgetFromLayout(const TSharedRef<SWidget>& InWidget);
@@ -115,6 +113,8 @@ protected:
 
 	void ActivateAllButThis(const bool bInActive, const TSharedRef<FNavigationToolFilter> InFilter);
 
+	void OnFiltersChanged(const ENavigationToolFilterChange InChangeType, const TSharedRef<FNavigationToolFilter>& InFilter);
+
 	TWeakPtr<FNavigationToolFilterBar> WeakFilterBar;
 
 	TWeakPtr<SSequencerSearchBox> WeakSearchBox;
@@ -132,6 +132,7 @@ protected:
 	TSharedPtr<SFilterExpressionHelpDialog> TextExpressionHelpDialog;
 
 	TSharedPtr<FNavigationToolFilterBarContextMenu> ContextMenu;
+	TSharedPtr<FNavigationToolFilterContextMenu> FilterContextMenu;
 };
 
 }
diff --git a/Engine/Source/Editor/Sequencer/Private/Tools/NavigationTool/Widgets/SNavigationToolView.cpp b/Engine/Source/Editor/Sequencer/Private/Tools/NavigationTool/Widgets/SNavigationToolView.cpp
--- a/Engine/Source/Editor/Sequencer/Private/Tools/NavigationTool/Widgets/SNavigationToolView.cpp
+++ b/Engine/Source/Editor/Sequencer/Private/Tools/NavigationTool/Widgets/SNavigationToolView.cpp
@@ -116,9 +116,10 @@ void SNavigationToolView::Construct(const FArguments& InArgs, const TSharedRef<F
 
 	ToolbarMenu = MakeShared<FNavigationToolToolbarMenu>();
 
-	if (const TSharedPtr<INavigationToolFilterBar> FilterBar = InToolView->GetFilterBar())
+	if (const TSharedPtr<FNavigationToolFilterBar> FilterBar = InToolView->GetFilterBar())
 	{
 		FilterBar->OnStateChanged().AddSP(this, &SNavigationToolView::OnFilterBarStateChanged);
+		FilterBar->OnFiltersChanged().AddSP(this, &SNavigationToolView::OnFiltersChanged);
 	}
 
 	RebuildWidget();
@@ -131,6 +132,7 @@ SNavigationToolView::~SNavigationToolView()
 		if (const TSharedPtr<FNavigationToolFilterBar> FilterBar = ToolView->GetFilterBar())
 		{
 			FilterBar->OnStateChanged().RemoveAll(this);
+			FilterBar->OnFiltersChanged().RemoveAll(this);
 		}
 	}
 }
@@ -627,6 +629,11 @@ void SNavigationToolView::OnFilterBarStateChanged(const bool InIsVisible, const
 	RebuildWidget();
 }
 
+void SNavigationToolView::OnFiltersChanged(const ENavigationToolFilterChange InChangeType, const TSharedRef<FNavigationToolFilter>& InFilter)
+{
+	RebuildWidget();
+}
+
 }
 
 #undef LOCTEXT_NAMESPACE
diff --git a/Engine/Source/Editor/Sequencer/Private/Tools/NavigationTool/Widgets/SNavigationToolView.h b/Engine/Source/Editor/Sequencer/Private/Tools/NavigationTool/Widgets/SNavigationToolView.h
--- a/Engine/Source/Editor/Sequencer/Private/Tools/NavigationTool/Widgets/SNavigationToolView.h
+++ b/Engine/Source/Editor/Sequencer/Private/Tools/NavigationTool/Widgets/SNavigationToolView.h
@@ -24,6 +24,7 @@ struct FCustomTextFilterData;
 namespace UE::Sequencer::NavigationTool
 {
 
+class FNavigationToolFilter;
 class FNavigationToolFilterBar;
 class FNavigationToolItemDragDropOp;
 class FNavigationToolToolbarMenu;
@@ -31,6 +32,7 @@ class FNavigationToolView;
 class SNavigationToolFilterBar;
 class SNavigationToolItemFilters;
 class SNavigationToolTreeView;
+enum class ENavigationToolFilterChange : uint8;
 
 class SNavigationToolView : public SCompoundWidget
 {
@@ -92,6 +94,7 @@ private:
 	void RebuildSearchAndFilterRow();
 
 	void OnFilterBarStateChanged(const bool InIsVisible, const EFilterBarLayout InNewLayout);
+	void OnFiltersChanged(const ENavigationToolFilterChange InChangeType, const TSharedRef<FNavigationToolFilter>& InFilter);
 
 	TSharedPtr<FNavigationToolFilterBar> GetFilterBar() const;
 
diff --git a/Engine/Source/Editor/Sequencer/Public/Tools/NavigationTool/Filters/NavigationToolFilterBar.h b/Engine/Source/Editor/Sequencer/Public/Tools/NavigationTool/Filters/NavigationToolFilterBar.h
--- a/Engine/Source/Editor/Sequencer/Public/Tools/NavigationTool/Filters/NavigationToolFilterBar.h
+++ b/Engine/Source/Editor/Sequencer/Public/Tools/NavigationTool/Filters/NavigationToolFilterBar.h
@@ -71,6 +71,8 @@ public:
 	virtual void ResetFilters() override;
 
 	FNavigationToolFiltersChanged& OnFiltersChanged() { return FiltersChangedEvent; }
+	virtual FOnFilterBarStateChanged& OnStateChanged() override { return StateChangedEvent; }
+	virtual FSimpleMulticastDelegate& OnRequestUpdate() override { return RequestUpdateEvent; }
 
 	TSharedRef<FNavigationToolFilter_Text> GetTextFilter() const;
 	FText GetFilterErrorText() const;
@@ -115,8 +117,6 @@ public:
 	virtual int32 GetTotalDisplayNodeCount() const override;
 	virtual int32 GetFilteredDisplayNodeCount() const override;
 
-	//virtual bool HasSelectedTracks() const override;
-
 	virtual const FTextFilterExpressionEvaluator& GetTextFilterExpressionEvaluator() const override;
 	virtual TArray<TSharedRef<ISequencerTextFilterExpressionContext>> GetTextFilterExpressionContexts() const override;
 
@@ -207,9 +207,6 @@ public:
 	bool IsFilterSupported(const TSharedRef<FNavigationToolFilter>& InFilter) const;
     bool IsFilterSupported(const FString& InFilterName) const;
 
-	virtual FOnFilterBarStateChanged& OnStateChanged() override { return StateChangedEvent; }
-	virtual FSimpleMulticastDelegate& OnRequestUpdate() override { return RequestUpdateEvent; }
-
 	virtual void OpenTextExpressionHelp() override;
 	virtual void SaveCurrentFilterSetAsCustomTextFilter() override;
 	virtual void CreateNewTextFilter() override;
