Sequencer: Turn off Snap Keys to Keys and Sections and Snap Sections to Keys and Sections by default

Now that Force Whole Frames is a separate toggle, we can turn off these two settings so that it's not as easy to snap to keys and sections when you're just dragging keys around.

diff --git a/Engine/Source/Editor/Sequencer/Private/SequencerSettings.cpp b/Engine/Source/Editor/Sequencer/Private/SequencerSettings.cpp
--- a/Engine/Source/Editor/Sequencer/Private/SequencerSettings.cpp
+++ b/Engine/Source/Editor/Sequencer/Private/SequencerSettings.cpp
@@ -19,8 +19,8 @@ USequencerSettings::USequencerSettings( const FObjectInitializer& ObjectInitiali
 	bCreateSpawnableCameras = true;
 	bShowRangeSlider = false;
 	bIsSnapEnabled = true;
-	bSnapKeyTimesToKeys = true;
-	bSnapSectionTimesToSections = true;
+	bSnapKeyTimesToKeys = false;
+	bSnapSectionTimesToSections = false;
 	bSnapPlayTimeToKeys = false;
 	bSnapPlayTimeToSections = false;
 	bSnapPlayTimeToMarkers = false;
