Sequencer: Fix empty LOCTEXT

#rnx

diff --git a/Engine/Source/Editor/Sequencer/Private/Tools/NavigationTool/Menus/NavigationToolToolbarMenu.cpp b/Engine/Source/Editor/Sequencer/Private/Tools/NavigationTool/Menus/NavigationToolToolbarMenu.cpp
--- a/Engine/Source/Editor/Sequencer/Private/Tools/NavigationTool/Menus/NavigationToolToolbarMenu.cpp
+++ b/Engine/Source/Editor/Sequencer/Private/Tools/NavigationTool/Menus/NavigationToolToolbarMenu.cpp
@@ -1,4 +1,4 @@
-﻿// Copyright Epic Games, Inc. All Rights Reserved.
+// Copyright Epic Games, Inc. All Rights Reserved.
 
 #include "NavigationToolToolbarMenu.h"
 #include "ToolMenu.h"
@@ -128,7 +128,7 @@ void FNavigationToolToolbarMenu::CreateSettingsMenu(UToolMenu* const InMenu)
 
 	MiscSection.AddSubMenu(TEXT("ColumnViews")
 		, LOCTEXT("ColumnViewsSubMenu", "Column Views")
-		, LOCTEXT("ColumnViewsSubMenuTooltip", "")
+		, LOCTEXT("ColumnViewsSubMenuTooltip", "Column View Options")
 		, FNewToolMenuDelegate::CreateSP(this, &FNavigationToolToolbarMenu::CreateColumnViewOptionsMenu));
 
 	MiscSection.AddSubMenu(TEXT("FilterBarOptions")
