[Sequencer] Fix "Sync Selection to Navigation Tool" from lagging the engine from SaveConfig()

diff --git a/Engine/Source/Editor/Sequencer/Private/Tools/NavigationTool/NavigationTool.cpp b/Engine/Source/Editor/Sequencer/Private/Tools/NavigationTool/NavigationTool.cpp
--- a/Engine/Source/Editor/Sequencer/Private/Tools/NavigationTool/NavigationTool.cpp
+++ b/Engine/Source/Editor/Sequencer/Private/Tools/NavigationTool/NavigationTool.cpp
@@ -84,8 +84,8 @@ void FNavigationTool::Init()
 					Shutdown();
 				}
 			});
-		Sequencer->GetSelection().Outliner.OnChanged.AddSP(this, &FNavigationTool::OnSequencerSelectionChanged);
-		Sequencer->GetSelection().TrackArea.OnChanged.AddSP(this, &FNavigationTool::OnSequencerSelectionChanged);
+		Sequencer->GetSelection().Outliner.OnChanged.AddSP(this, &FNavigationTool::OnSequencerOutlinerSelectionChanged);
+		Sequencer->GetSelection().TrackArea.OnChanged.AddSP(this, &FNavigationTool::OnSequencerTrackAreaSelectionChanged);
 	}
 
 	// Listen to object replacement changes
@@ -1355,7 +1355,7 @@ void FNavigationTool::OnTreeViewChanged()
 	RequestRefresh();
 }
 
-void FNavigationTool::OnSequencerSelectionChanged()
+void FNavigationTool::OnSequencerOutlinerSelectionChanged()
 {
 	UNavigationToolSettings* const ToolSettings = GetMutableDefault<UNavigationToolSettings>();
 	if (!ToolSettings->ShouldSyncSelectionToNavigationTool())
@@ -1371,11 +1371,14 @@ void FNavigationTool::OnSequencerSelectionChanged()
 
 	FSequencerSelection& SequencerSelection = Sequencer->GetSelection();
 
-	// Unbind to stop circular loop
+	// Unbind to avoid circular loop
 	SequencerSelection.Outliner.OnChanged.RemoveAll(this);
 
+	// Temporarily stop syncing selection to sequencer to avoid circular loop
+	// Calling SaveConfig() here will cause a massive slowdown and throw a log about "DeleteFile"
+	// in some cases. Not able to reproduce this locally.
 	const bool bPreviouslySyncingSelectionToSequencer = ToolSettings->ShouldSyncSelectionToSequencer();
-	ToolSettings->SetSyncSelectionToSequencer(false);
+	ToolSettings->SetSyncSelectionToSequencer(/*bInSync=*/false, /*bInSaveConfig=*/false);
 
 	TArray<FNavigationToolItemPtr> ItemsToSelect;
 
@@ -1383,6 +1386,41 @@ void FNavigationTool::OnSequencerSelectionChanged()
 	{
 		ItemsToSelect.Append(TryFindItems(OutlinerItem));
 	}
+
+	SelectItems(ItemsToSelect);
+
+	// Bind from avoiding circular loop
+	SequencerSelection.Outliner.OnChanged.AddSP(this, &FNavigationTool::OnSequencerOutlinerSelectionChanged);
+
+	ToolSettings->SetSyncSelectionToSequencer(bPreviouslySyncingSelectionToSequencer, /*bInSaveConfig=*/false);
+}
+
+void FNavigationTool::OnSequencerTrackAreaSelectionChanged()
+{
+	UNavigationToolSettings* const ToolSettings = GetMutableDefault<UNavigationToolSettings>();
+	if (!ToolSettings->ShouldSyncSelectionToNavigationTool())
+	{
+		return;
+	}
+
+	const TSharedPtr<FSequencer> Sequencer = StaticCastSharedPtr<FSequencer>(WeakSequencer.Pin());
+	if (!Sequencer.IsValid())
+	{
+		return;
+	}
+
+	FSequencerSelection& SequencerSelection = Sequencer->GetSelection();
+
+	// Unbind to avoid circular loop
+	SequencerSelection.TrackArea.OnChanged.RemoveAll(this);
+
+	// Calling SaveConfig() here will cause a massive slowdown and throw a log about "DeleteFile"
+	// in some cases. Not able to reproduce this locally.
+	const bool bPreviouslySyncingSelectionToSequencer = ToolSettings->ShouldSyncSelectionToSequencer();
+	ToolSettings->SetSyncSelectionToSequencer(/*bInSync=*/false, /*bInSaveConfig=*/false);
+
+	TArray<FNavigationToolItemPtr> ItemsToSelect;
+
 	for (const FViewModelPtr TrackAreaItem : SequencerSelection.TrackArea)
 	{
 		ItemsToSelect.Append(TryFindItems(TrackAreaItem));
@@ -1390,10 +1428,10 @@ void FNavigationTool::OnSequencerSelectionChanged()
 
 	SelectItems(ItemsToSelect);
 
-	// Bind from stopping circular loop
-	SequencerSelection.Outliner.OnChanged.AddSP(this, &FNavigationTool::OnSequencerSelectionChanged);
+	// Bind from avoiding circular loop
+	SequencerSelection.TrackArea.OnChanged.AddSP(this, &FNavigationTool::OnSequencerTrackAreaSelectionChanged);
 
-	ToolSettings->SetSyncSelectionToSequencer(bPreviouslySyncingSelectionToSequencer);
+	ToolSettings->SetSyncSelectionToSequencer(bPreviouslySyncingSelectionToSequencer, /*bInSaveConfig=*/false);
 }
 
 TArray<FNavigationToolItemPtr> FNavigationTool::FindItemsFromMovieSceneObject(UObject* const InObject) const
diff --git a/Engine/Source/Editor/Sequencer/Private/Tools/NavigationTool/NavigationTool.h b/Engine/Source/Editor/Sequencer/Private/Tools/NavigationTool/NavigationTool.h
--- a/Engine/Source/Editor/Sequencer/Private/Tools/NavigationTool/NavigationTool.h
+++ b/Engine/Source/Editor/Sequencer/Private/Tools/NavigationTool/NavigationTool.h
@@ -213,7 +213,8 @@ private:
 
 	void OnTreeViewChanged();
 
-	void OnSequencerSelectionChanged();
+	void OnSequencerOutlinerSelectionChanged();
+	void OnSequencerTrackAreaSelectionChanged();
 
 	TArray<FNavigationToolItemPtr> FindItemsFromMovieSceneObject(UObject* const InObject) const;
 	TArray<FNavigationToolItemPtr> FindItemsFromObjectGuid(const FGuid& InObjectGuid) const;
diff --git a/Engine/Source/Editor/Sequencer/Private/Tools/NavigationTool/NavigationToolSettings.cpp b/Engine/Source/Editor/Sequencer/Private/Tools/NavigationTool/NavigationToolSettings.cpp
--- a/Engine/Source/Editor/Sequencer/Private/Tools/NavigationTool/NavigationToolSettings.cpp
+++ b/Engine/Source/Editor/Sequencer/Private/Tools/NavigationTool/NavigationToolSettings.cpp
@@ -234,8 +234,11 @@ bool UNavigationToolSettings::ShouldSyncSelectionToSequencer() const
 	return bSyncSelectionToSequencer;
 }
 
-void UNavigationToolSettings::SetSyncSelectionToSequencer(const bool bInSync)
+void UNavigationToolSettings::SetSyncSelectionToSequencer(const bool bInSync, const bool bInSaveConfig)
 {
 	bSyncSelectionToSequencer = bInSync;
-	SaveConfig();
+	if (bInSaveConfig)
+	{
+		SaveConfig();
+	}
 }
diff --git a/Engine/Source/Editor/Sequencer/Public/Tools/NavigationTool/NavigationToolSettings.h b/Engine/Source/Editor/Sequencer/Public/Tools/NavigationTool/NavigationToolSettings.h
--- a/Engine/Source/Editor/Sequencer/Public/Tools/NavigationTool/NavigationToolSettings.h
+++ b/Engine/Source/Editor/Sequencer/Public/Tools/NavigationTool/NavigationToolSettings.h
@@ -110,7 +110,7 @@ public:
 	void SetSyncSelectionToNavigationTool(const bool bInSync);
 
 	bool ShouldSyncSelectionToSequencer() const;
-	void SetSyncSelectionToSequencer(const bool bInSync);
+	void SetSyncSelectionToSequencer(const bool bInSync, const bool bInSaveConfig = true);
 
 private:
 	/** Applies the default column view set by the provider when a view is loaded */
