Sequencer: Use the focused movie scene sequence instead of the root when determining filter support

#rnx

diff --git a/Engine/Source/Editor/Sequencer/Private/Filters/SequencerFilterBar.cpp b/Engine/Source/Editor/Sequencer/Private/Filters/SequencerFilterBar.cpp
--- a/Engine/Source/Editor/Sequencer/Private/Filters/SequencerFilterBar.cpp
+++ b/Engine/Source/Editor/Sequencer/Private/Filters/SequencerFilterBar.cpp
@@ -1811,15 +1811,15 @@ TSharedPtr<SSequencerFilterBar> FSequencerFilterBar::GetWidget() const
 
 bool FSequencerFilterBar::IsFilterSupported(const TSharedRef<FSequencerTrackFilter>& InFilter) const
 {
-	UMovieSceneSequence* const RootMovieSceneSequence = Sequencer.GetRootMovieSceneSequence();
-	if (!IsValid(RootMovieSceneSequence))
+	UMovieSceneSequence* const MovieSceneSequence = Sequencer.GetFocusedMovieSceneSequence();
+	if (!IsValid(MovieSceneSequence))
 	{
 		return false;
 	}
 
 	const FString FilterName = InFilter->GetName();
-	const bool bFilterSupportsSequence = InFilter->SupportsSequence(RootMovieSceneSequence);
-	const bool bSequenceSupportsFilter = RootMovieSceneSequence->IsFilterSupported(FilterName);
+	const bool bFilterSupportsSequence = InFilter->SupportsSequence(MovieSceneSequence);
+	const bool bSequenceSupportsFilter = MovieSceneSequence->IsFilterSupported(FilterName);
 	return bFilterSupportsSequence || bSequenceSupportsFilter;
 }
 
