[Sequencer] Navigation Tool:
- Only SaveConfig() when the property actually changes
- Add optional param to save config for sync functions

diff --git a/Engine/Source/Editor/Sequencer/Private/Tools/NavigationTool/NavigationToolSettings.cpp b/Engine/Source/Editor/Sequencer/Private/Tools/NavigationTool/NavigationToolSettings.cpp
--- a/Engine/Source/Editor/Sequencer/Private/Tools/NavigationTool/NavigationToolSettings.cpp
+++ b/Engine/Source/Editor/Sequencer/Private/Tools/NavigationTool/NavigationToolSettings.cpp
@@ -28,8 +28,11 @@ bool UNavigationToolSettings::ShouldUseMutedHierarchy() const
 
 void UNavigationToolSettings::SetUseMutedHierarchy(const bool bInUseMutedHierarchy)
 {
-	bUseMutedHierarchy = bInUseMutedHierarchy;
-	SaveConfig();
+	if (bUseMutedHierarchy != bInUseMutedHierarchy)
+	{
+		bUseMutedHierarchy = bInUseMutedHierarchy;
+		SaveConfig();
+	}
 }
 
 bool UNavigationToolSettings::ShouldAutoExpandToSelection() const
@@ -39,8 +42,11 @@ bool UNavigationToolSettings::ShouldAutoExpandToSelection() const
 
 void UNavigationToolSettings::SetAutoExpandToSelection(const bool bInAutoExpandToSelection)
 {
-	bAutoExpandToSelection = bInAutoExpandToSelection;
-	SaveConfig();
+	if (bAutoExpandToSelection != bInAutoExpandToSelection)
+	{
+		bAutoExpandToSelection = bInAutoExpandToSelection;
+		SaveConfig();
+	}
 }
 
 bool UNavigationToolSettings::ShouldAlwaysShowLockState() const
@@ -50,8 +56,11 @@ bool UNavigationToolSettings::ShouldAlwaysShowLockState() const
 
 void UNavigationToolSettings::SetAlwaysShowLockState(const bool bInAlwaysShowLockState)
 {
-	bAlwaysShowLockState = bInAlwaysShowLockState;
-	SaveConfig();
+	if (bAlwaysShowLockState != bInAlwaysShowLockState)
+	{
+		bAlwaysShowLockState = bInAlwaysShowLockState;
+		SaveConfig();
+	}
 }
 
 void UNavigationToolSettings::ToggleViewModeSupport(ENavigationToolItemViewMode& InOutViewMode, const ENavigationToolItemViewMode InFlags)
@@ -83,8 +92,11 @@ bool UNavigationToolSettings::ShouldUseShortNames() const
 
 void UNavigationToolSettings::SetUseShortNames(const bool bInUseShortNames)
 {
-	bUseShortNames = bInUseShortNames;
-	SaveConfig();
+	if (bUseShortNames != bInUseShortNames)
+	{
+		bUseShortNames = bInUseShortNames;
+		SaveConfig();
+	}
 }
 
 bool UNavigationToolSettings::ShouldApplyDefaultColumnView() const
@@ -94,8 +106,11 @@ bool UNavigationToolSettings::ShouldApplyDefaultColumnView() const
 
 void UNavigationToolSettings::SetApplyDefaultColumnView(const bool bInApplyDefaultColumnView)
 {
-	bApplyDefaultColumnView = bInApplyDefaultColumnView;
-	SaveConfig();
+	if (bApplyDefaultColumnView != bInApplyDefaultColumnView)
+	{
+		bApplyDefaultColumnView = bInApplyDefaultColumnView;
+		SaveConfig();
+	}
 }
 
 ENavigationToolItemViewMode UNavigationToolSettings::GetItemDefaultViewMode() const
@@ -129,14 +144,21 @@ void UNavigationToolSettings::SetBuiltInFilterEnabled(const FName InFilterName,
 
 	if (bInEnabled)
 	{
-		EnabledBuiltInFilters.Add(InFilterName);
+		bool bAlreadyExists = false;
+		EnabledBuiltInFilters.Add(InFilterName, &bAlreadyExists);
+		if (!bAlreadyExists)
+		{
+			SaveConfig();
+		}
 	}
 	else
 	{
-		EnabledBuiltInFilters.Remove(InFilterName);
+		if (EnabledBuiltInFilters.Contains(InFilterName))
+		{
+			EnabledBuiltInFilters.Remove(InFilterName);
+			SaveConfig();
+		}
 	}
-
-	SaveConfig();
 }
 
 FSequencerFilterBarConfig& UNavigationToolSettings::FindOrAddFilterBar(const FName InIdentifier, const bool bInSaveConfig)
@@ -158,9 +180,15 @@ FSequencerFilterBarConfig* UNavigationToolSettings::FindFilterBar(const FName In
 
 bool UNavigationToolSettings::RemoveFilterBar(const FName InIdentifier)
 {
-	const int32 RemovedCount = NavigationToolFilterBars.Remove(InIdentifier) > 0;
-	SaveConfig();
-	return RemovedCount > 0;
+	const int32 RemovedCount = NavigationToolFilterBars.Remove(InIdentifier);
+	const bool bRemoveSuccess = RemovedCount > 0;
+
+	if (bRemoveSuccess)
+	{
+		SaveConfig();
+	}
+
+	return bRemoveSuccess;
 }
 
 bool UNavigationToolSettings::ShouldAutoExpandNodesOnFilterPass() const
@@ -168,10 +196,13 @@ bool UNavigationToolSettings::ShouldAutoExpandNodesOnFilterPass() const
 	return bAutoExpandNodesOnFilterPass;
 }
 
-void UNavigationToolSettings::SetAutoExpandNodesOnFilterPass(const bool bInIncludeParents)
+void UNavigationToolSettings::SetAutoExpandNodesOnFilterPass(const bool bInAutoExpand)
 {
-	bAutoExpandNodesOnFilterPass = bInIncludeParents;
-	SaveConfig();
+	if (bAutoExpandNodesOnFilterPass != bInAutoExpand)
+	{
+		bAutoExpandNodesOnFilterPass = bInAutoExpand;
+		SaveConfig();
+	}
 }
 
 bool UNavigationToolSettings::ShouldUseFilterSubmenusForCategories() const
@@ -181,8 +212,11 @@ bool UNavigationToolSettings::ShouldUseFilterSubmenusForCategories() const
 
 void UNavigationToolSettings::SetUseFilterSubmenusForCategories(const bool bInUseFilterSubmenusForCategories)
 {
-	bUseFilterSubmenusForCategories = bInUseFilterSubmenusForCategories;
-	SaveConfig();
+	if (bUseFilterSubmenusForCategories != bInUseFilterSubmenusForCategories)
+	{
+		bUseFilterSubmenusForCategories = bInUseFilterSubmenusForCategories;
+		SaveConfig();
+	}
 }
 
 bool UNavigationToolSettings::IsFilterBarVisible() const
@@ -192,8 +226,11 @@ bool UNavigationToolSettings::IsFilterBarVisible() const
 
 void UNavigationToolSettings::SetFilterBarVisible(const bool bInVisible)
 {
-	bFilterBarVisible = bInVisible;
-	SaveConfig();
+	if (bFilterBarVisible != bInVisible)
+	{
+		bFilterBarVisible = bInVisible;
+		SaveConfig();
+	}
 }
 
 EFilterBarLayout UNavigationToolSettings::GetFilterBarLayout() const
@@ -203,8 +240,11 @@ EFilterBarLayout UNavigationToolSettings::GetFilterBarLayout() const
 
 void UNavigationToolSettings::SetFilterBarLayout(const EFilterBarLayout InLayout)
 {
-	LastFilterBarLayout = InLayout;
-	SaveConfig();
+	if (LastFilterBarLayout != InLayout)
+	{
+		LastFilterBarLayout = InLayout;
+		SaveConfig();
+	}
 }
 
 float UNavigationToolSettings::GetLastFilterBarSizeCoefficient() const
@@ -214,8 +254,11 @@ float UNavigationToolSettings::GetLastFilterBarSizeCoefficient() const
 
 void UNavigationToolSettings::SetLastFilterBarSizeCoefficient(const float bInSizeCoefficient)
 {
-	LastFilterBarSizeCoefficient = bInSizeCoefficient;
-	SaveConfig();
+	if (LastFilterBarSizeCoefficient != bInSizeCoefficient)
+	{
+		LastFilterBarSizeCoefficient = bInSizeCoefficient;
+		SaveConfig();
+	}
 }
 
 bool UNavigationToolSettings::ShouldSyncSelectionToNavigationTool() const
@@ -223,10 +266,16 @@ bool UNavigationToolSettings::ShouldSyncSelectionToNavigationTool() const
 	return bSyncSelectionToNavigationTool;
 }
 
-void UNavigationToolSettings::SetSyncSelectionToNavigationTool(const bool bInSync)
+void UNavigationToolSettings::SetSyncSelectionToNavigationTool(const bool bInSync, const bool bInSaveConfig)
 {
-	bSyncSelectionToNavigationTool = bInSync;
-    SaveConfig();
+	if (bSyncSelectionToNavigationTool != bInSync)
+	{
+		bSyncSelectionToNavigationTool = bInSync;
+		if (bInSaveConfig)
+		{
+			SaveConfig();
+		}
+	}
 }
 
 bool UNavigationToolSettings::ShouldSyncSelectionToSequencer() const
@@ -236,9 +285,12 @@ bool UNavigationToolSettings::ShouldSyncSelectionToSequencer() const
 
 void UNavigationToolSettings::SetSyncSelectionToSequencer(const bool bInSync, const bool bInSaveConfig)
 {
-	bSyncSelectionToSequencer = bInSync;
-	if (bInSaveConfig)
+	if (bSyncSelectionToSequencer != bInSync)
 	{
-		SaveConfig();
+		bSyncSelectionToSequencer = bInSync;
+		if (bInSaveConfig)
+		{
+			SaveConfig();
+		}
 	}
 }
diff --git a/Engine/Source/Editor/Sequencer/Public/Tools/NavigationTool/NavigationToolSettings.h b/Engine/Source/Editor/Sequencer/Public/Tools/NavigationTool/NavigationToolSettings.h
--- a/Engine/Source/Editor/Sequencer/Public/Tools/NavigationTool/NavigationToolSettings.h
+++ b/Engine/Source/Editor/Sequencer/Public/Tools/NavigationTool/NavigationToolSettings.h
@@ -92,7 +92,7 @@ public:
 	bool RemoveFilterBar(const FName InIdentifier);
 
 	bool ShouldAutoExpandNodesOnFilterPass() const;
-	void SetAutoExpandNodesOnFilterPass(const bool bInIncludeParents);
+	void SetAutoExpandNodesOnFilterPass(const bool bInAutoExpand);
 
 	bool ShouldUseFilterSubmenusForCategories() const;
 	void SetUseFilterSubmenusForCategories(const bool bInUseFilterSubmenusForCategories);
@@ -107,7 +107,7 @@ public:
 	void SetLastFilterBarSizeCoefficient(const float bInSizeCoefficient);
 
 	bool ShouldSyncSelectionToNavigationTool() const;
-	void SetSyncSelectionToNavigationTool(const bool bInSync);
+	void SetSyncSelectionToNavigationTool(const bool bInSync, const bool bInSaveConfig = true);
 
 	bool ShouldSyncSelectionToSequencer() const;
 	void SetSyncSelectionToSequencer(const bool bInSync, const bool bInSaveConfig = true);
