[Sequencer] Change BindingName text expression filter to use actor label instead of object name

diff --git a/Engine/Source/Editor/Sequencer/Private/Filters/TextExpressions/SequencerTextFilterExpression_BindingName.cpp b/Engine/Source/Editor/Sequencer/Private/Filters/TextExpressions/SequencerTextFilterExpression_BindingName.cpp
--- a/Engine/Source/Editor/Sequencer/Private/Filters/TextExpressions/SequencerTextFilterExpression_BindingName.cpp
+++ b/Engine/Source/Editor/Sequencer/Private/Filters/TextExpressions/SequencerTextFilterExpression_BindingName.cpp
@@ -1,6 +1,7 @@
 ﻿// Copyright Epic Games, Inc. All Rights Reserved.
 
 #include "SequencerTextFilterExpression_BindingName.h"
+#include "GameFramework/Actor.h"
 
 using namespace UE::Sequencer;
 
@@ -39,6 +40,14 @@ bool FSequencerTextFilterExpression_BindingName::TestComplexExpression(const FNa
 	UObject* const BoundObject = GetBoundObject();
 	if (IsValid(BoundObject))
 	{
+		const AActor* const Actor = Cast<AActor>(BoundObject);
+		if (IsValid(Actor))
+		{
+			const FString ActorLabel = Actor->GetActorLabel();
+			const bool bPassed = TextFilterUtils::TestComplexExpression(ActorLabel, InValue, InComparisonOperation, InTextComparisonMode);
+			return bPassed;
+		}
+
 		const FString ObjectName = BoundObject->GetName();
 		const bool bPassed = TextFilterUtils::TestComplexExpression(ObjectName, InValue, InComparisonOperation, InTextComparisonMode);
 		return bPassed;
