[Sequencer] RefreshTree on Pin/Mute/Solo/Deactivate hotkey press

diff --git a/Engine/Source/Editor/Sequencer/Private/Sequencer.cpp b/Engine/Source/Editor/Sequencer/Private/Sequencer.cpp
--- a/Engine/Source/Editor/Sequencer/Private/Sequencer.cpp
+++ b/Engine/Source/Editor/Sequencer/Private/Sequencer.cpp
@@ -12778,7 +12778,10 @@ void FSequencer::ToggleTrackSelectionPin()
 	for (const TViewModelPtr Pinnable : OutlinerSelection.Filter<IPinnableExtension>())
 	{
 		Pinnable->SetPinned(bNewPinState);
+		NodeTree->SavePinnedState(*Pinnable.AsModel(), bNewPinState);
 	}
+
+	RefreshTree();
 }
 
 void FSequencer::ToggleTrackSelectionLock()
@@ -12799,6 +12802,8 @@ void FSequencer::ToggleTrackSelectionLock()
 	{
 		Lockable->SetIsLocked(bNewLockState);
 	}
+
+	RefreshTree();
 }
 
 void FSequencer::ToggleTrackSelectionDeactive()
@@ -12819,6 +12824,8 @@ void FSequencer::ToggleTrackSelectionDeactive()
 	{
 		Deactivatable->SetIsDeactivated(bNewDeactiveState);
 	}
+
+	RefreshTree();
 }
 
 void FSequencer::ToggleTrackSelectionMute()
@@ -12839,6 +12846,8 @@ void FSequencer::ToggleTrackSelectionMute()
 	{
 		Mutable->SetIsMuted(bNewMuteState);
 	}
+
+	RefreshTree();
 }
 
 void FSequencer::ToggleTrackSelectionSolo()
@@ -12859,6 +12868,8 @@ void FSequencer::ToggleTrackSelectionSolo()
 	{
 		Soloable->SetIsSoloed(bNewSoloState);
 	}
+
+	RefreshTree();
 }
 
 TSharedPtr<UE::Sequencer::NavigationTool::INavigationTool> FSequencer::GetNavigationTool() const
