[Sequencer] Navigation Tool:
- Add Marker item type
- Add IInTimeExtension and IOutTimeExtension item extensions
- Add Marker global filter

diff --git a/Engine/Plugins/MovieScene/LevelSequenceEditor/Source/LevelSequenceEditor/Private/LevelSequenceNavigationToolProvider.cpp b/Engine/Plugins/MovieScene/LevelSequenceEditor/Source/LevelSequenceEditor/Private/LevelSequenceNavigationToolProvider.cpp
--- a/Engine/Plugins/MovieScene/LevelSequenceEditor/Source/LevelSequenceEditor/Private/LevelSequenceNavigationToolProvider.cpp
+++ b/Engine/Plugins/MovieScene/LevelSequenceEditor/Source/LevelSequenceEditor/Private/LevelSequenceNavigationToolProvider.cpp
@@ -162,6 +162,7 @@ void FLevelSequenceNavigationToolProvider::OnExtendBuiltInFilters(TArray<FNaviga
 	OutFilterParams.Add(FNavigationToolBuiltInFilterParams::CreateSequenceFilter());
 	OutFilterParams.Add(FNavigationToolBuiltInFilterParams::CreateTrackFilter());
 	OutFilterParams.Add(FNavigationToolBuiltInFilterParams::CreateBindingFilter());
+	OutFilterParams.Add(FNavigationToolBuiltInFilterParams::CreateMarkerFilter());
 
 	FNavigationToolProvider::OnExtendBuiltInFilters(OutFilterParams);
 }
diff --git a/Engine/Source/Editor/Sequencer/Private/Tools/NavigationTool/Columns/NavigationToolInTimeColumn.cpp b/Engine/Source/Editor/Sequencer/Private/Tools/NavigationTool/Columns/NavigationToolInTimeColumn.cpp
--- a/Engine/Source/Editor/Sequencer/Private/Tools/NavigationTool/Columns/NavigationToolInTimeColumn.cpp
+++ b/Engine/Source/Editor/Sequencer/Private/Tools/NavigationTool/Columns/NavigationToolInTimeColumn.cpp
@@ -39,11 +39,6 @@ TSharedRef<SWidget> FNavigationToolInTimeColumn::ConstructRowWidget(const FNavig
 	, const TSharedRef<INavigationToolView>& InToolView
 	, const TSharedRef<SNavigationToolTreeRow>& InRow)
 {
-	if (!InItem->IsA<FNavigationToolSequence>())
-	{
-		return SNullWidget::NullWidget;
-	}
-
 	return SNew(SNavigationToolInTime, InItem, InToolView, InRow);
 }
 
diff --git a/Engine/Source/Editor/Sequencer/Private/Tools/NavigationTool/Columns/NavigationToolOutTimeColumn.cpp b/Engine/Source/Editor/Sequencer/Private/Tools/NavigationTool/Columns/NavigationToolOutTimeColumn.cpp
--- a/Engine/Source/Editor/Sequencer/Private/Tools/NavigationTool/Columns/NavigationToolOutTimeColumn.cpp
+++ b/Engine/Source/Editor/Sequencer/Private/Tools/NavigationTool/Columns/NavigationToolOutTimeColumn.cpp
@@ -41,11 +41,6 @@ TSharedRef<SWidget> FNavigationToolOutTimeColumn::ConstructRowWidget(const FNavi
 	, const TSharedRef<INavigationToolView>& InToolView
 	, const TSharedRef<SNavigationToolTreeRow>& InRow)
 {
-	if (!InItem->IsA<FNavigationToolSequence>())
-	{
-		return SNullWidget::NullWidget;
-	}
-
 	return SNew(SNavigationToolOutTime, InItem, InToolView, InRow);
 }
 
diff --git a/Engine/Source/Editor/Sequencer/Private/Tools/NavigationTool/Filters/Filters/NavigationToolBuiltInFilterParams.cpp b/Engine/Source/Editor/Sequencer/Private/Tools/NavigationTool/Filters/Filters/NavigationToolBuiltInFilterParams.cpp
--- a/Engine/Source/Editor/Sequencer/Private/Tools/NavigationTool/Filters/Filters/NavigationToolBuiltInFilterParams.cpp
+++ b/Engine/Source/Editor/Sequencer/Private/Tools/NavigationTool/Filters/Filters/NavigationToolBuiltInFilterParams.cpp
@@ -9,6 +9,7 @@
 #include "Tools/NavigationTool/Items/NavigationToolActor.h"
 #include "Tools/NavigationTool/Items/NavigationToolBinding.h"
 #include "Tools/NavigationTool/Items/NavigationToolComponent.h"
+#include "Tools/NavigationTool/Items/NavigationToolMarker.h"
 #include "Tools/NavigationTool/Items/NavigationToolSequence.h"
 
 #define LOCTEXT_NAMESPACE "NavigationToolBuiltInFilterParams"
@@ -211,4 +212,16 @@ FNavigationToolBuiltInFilterParams FNavigationToolBuiltInFilterParams::CreateBin
 		, FNavigationToolFilterCommands::Get().ToggleFilter_Binding);
 }
 
+FNavigationToolBuiltInFilterParams FNavigationToolBuiltInFilterParams::CreateMarkerFilter()
+{
+	return FNavigationToolBuiltInFilterParams(TEXT("Marker")
+		, { UE::Sequencer::NavigationTool::FNavigationToolMarker::GetStaticTypeId() }
+		, { }
+		, ENavigationToolFilterMode::MatchesType
+		, nullptr
+		, LOCTEXT("MarkerFilterDisplayName", "Marker")
+		, LOCTEXT("MarkerFilterTooltip", "Marker")
+		, FNavigationToolFilterCommands::Get().ToggleFilter_Marker);
+}
+
 #undef LOCTEXT_NAMESPACE
diff --git a/Engine/Source/Editor/Sequencer/Private/Tools/NavigationTool/Filters/NavigationToolFilterCommands.cpp b/Engine/Source/Editor/Sequencer/Private/Tools/NavigationTool/Filters/NavigationToolFilterCommands.cpp
--- a/Engine/Source/Editor/Sequencer/Private/Tools/NavigationTool/Filters/NavigationToolFilterCommands.cpp
+++ b/Engine/Source/Editor/Sequencer/Private/Tools/NavigationTool/Filters/NavigationToolFilterCommands.cpp
@@ -93,6 +93,12 @@ void FNavigationToolFilterCommands::RegisterCommands()
 		, EUserInterfaceActionType::ToggleButton
 		, FInputChord());
 
+	UI_COMMAND(ToggleFilter_Marker
+		, "Toggle Marker Filter"
+		, "Toggle the filter for Marker items"
+		, EUserInterfaceActionType::ToggleButton
+		, FInputChord());
+
 	UI_COMMAND(ToggleFilter_Unbound
 		, "Toggle Unbound Filter"
 		, "Toggle the filter for displaying items that contain unbound tracks"
diff --git a/Engine/Source/Editor/Sequencer/Private/Tools/NavigationTool/Filters/NavigationToolFilterCommands.h b/Engine/Source/Editor/Sequencer/Private/Tools/NavigationTool/Filters/NavigationToolFilterCommands.h
--- a/Engine/Source/Editor/Sequencer/Private/Tools/NavigationTool/Filters/NavigationToolFilterCommands.h
+++ b/Engine/Source/Editor/Sequencer/Private/Tools/NavigationTool/Filters/NavigationToolFilterCommands.h
@@ -36,10 +36,13 @@ public:
 	TSharedPtr<FUICommandInfo> ActivateAllFilters;
 	TSharedPtr<FUICommandInfo> DeactivateAllFilters;
 
+	// Global Filters
 	TSharedPtr<FUICommandInfo> ToggleFilter_Sequence;
 	TSharedPtr<FUICommandInfo> ToggleFilter_Track;
 	TSharedPtr<FUICommandInfo> ToggleFilter_Binding;
+	TSharedPtr<FUICommandInfo> ToggleFilter_Marker;
 
+	// Normal Filters
 	TSharedPtr<FUICommandInfo> ToggleFilter_Unbound;
 	TSharedPtr<FUICommandInfo> ToggleFilter_Marks;
 	TSharedPtr<FUICommandInfo> ToggleFilter_Playhead;
diff --git a/Engine/Source/Editor/Sequencer/Private/Tools/NavigationTool/Items/NavigationToolItem.cpp b/Engine/Source/Editor/Sequencer/Private/Tools/NavigationTool/Items/NavigationToolItem.cpp
--- a/Engine/Source/Editor/Sequencer/Private/Tools/NavigationTool/Items/NavigationToolItem.cpp
+++ b/Engine/Source/Editor/Sequencer/Private/Tools/NavigationTool/Items/NavigationToolItem.cpp
@@ -240,6 +240,11 @@ const FSlateBrush* FNavigationToolItem::GetIconBrush() const
 	return GetIcon().GetIcon();
 }
 
+FSlateColor FNavigationToolItem::GetIconColor() const
+{
+	return FStyleColors::Foreground;
+}
+
 TSharedRef<SWidget> FNavigationToolItem::GenerateLabelWidget(const TSharedRef<SNavigationToolTreeRow>& InRow)
 {
 	return SNew(SNavigationToolLabelItem, SharedThis(this), InRow);
diff --git a/Engine/Source/Editor/Sequencer/Private/Tools/NavigationTool/Items/NavigationToolMarker.cpp b/Engine/Source/Editor/Sequencer/Private/Tools/NavigationTool/Items/NavigationToolMarker.cpp
new file mode 100644
--- /dev/null
+++ b/Engine/Source/Editor/Sequencer/Private/Tools/NavigationTool/Items/NavigationToolMarker.cpp
@@ -0,0 +1,321 @@
+// Copyright Epic Games, Inc. All Rights Reserved.
+
+#include "Tools/NavigationTool/Items/NavigationToolMarker.h"
+#include "MovieScene.h"
+#include "ScopedTransaction.h"
+#include "Sections/MovieSceneSubSection.h"
+#include "Sequencer.h"
+#include "Styling/StyleColors.h"
+#include "Tools/NavigationTool/Items/NavigationToolSequence.h"
+#include "Tools/NavigationTool/NavigationTool.h"
+#include "Tools/NavigationTool/NavigationToolScopedSelection.h"
+#include "Tools/NavigationTool/Utils/NavigationToolMiscUtils.h"
+#include "Tools/NavigationTool/Utils/NavigationToolMovieSceneUtils.h"
+
+#define LOCTEXT_NAMESPACE "NavigationToolMarker"
+
+namespace UE::Sequencer::NavigationTool
+{
+
+FNavigationToolMarker::FNavigationToolMarker(INavigationTool& InTool
+	, const FNavigationToolItemPtr& InParentItem
+	, const TSharedPtr<FNavigationToolSequence>& InParentSequenceItem
+	, const int32 InMarkedFrameIndex)
+	: Super(InTool, InParentItem)
+	, WeakParentSequenceItem(InParentSequenceItem)
+	, MarkedFrameIndex(InMarkedFrameIndex)
+{
+}
+
+bool FNavigationToolMarker::IsItemValid() const
+{
+	return WeakParentSequenceItem.IsValid() && MarkedFrameIndex != INDEX_NONE;
+}
+
+bool FNavigationToolMarker::IsAllowedInTool() const
+{
+	return IsItemValid();
+}
+
+ENavigationToolItemViewMode FNavigationToolMarker::GetSupportedViewModes(const INavigationToolView& InToolView) const
+{
+	return ENavigationToolItemViewMode::ItemTree | ENavigationToolItemViewMode::HorizontalItemList;
+}
+
+FText FNavigationToolMarker::GetDisplayName() const
+{
+	if (const FMovieSceneMarkedFrame* const MarkedFrame = GetMarkedFrame())
+	{
+		return FText::FromString(MarkedFrame->Label);
+	}
+	return FText::GetEmpty();
+}
+
+FText FNavigationToolMarker::GetClassName() const
+{
+	return FMovieSceneMarkedFrame::StaticStruct()->GetDisplayNameText();
+}
+
+FSlateIcon FNavigationToolMarker::GetIcon() const
+{
+	return FSlateIcon(FAppStyle::GetAppStyleSetName(), TEXT("AnimTimeline.SectionMarker"));
+}
+
+const FSlateBrush* FNavigationToolMarker::GetIconBrush() const
+{
+	return FAppStyle::GetBrush(TEXT("AnimTimeline.SectionMarker"));
+}
+
+FSlateColor FNavigationToolMarker::GetIconColor() const
+{
+	const TSharedPtr<ISequencer> Sequencer = Tool.GetSequencer();
+	if (!Sequencer.IsValid())
+	{
+		return FStyleColors::Foreground;
+	}
+
+	USequencerSettings* const SequencerSettings = Sequencer->GetSequencerSettings();
+	if (!IsValid(SequencerSettings))
+	{
+		return FStyleColors::Foreground;
+	}
+
+	return SequencerSettings->GetMarkedFrameColor();
+}
+
+FText FNavigationToolMarker::GetIconTooltipText() const
+{
+	if (const FMovieSceneMarkedFrame* const MarkedFrame = GetMarkedFrame())
+	{
+		return FText::FromString(MarkedFrame->Label);
+	}
+	return FText::GetEmpty();
+}
+
+bool FNavigationToolMarker::IsSelected(const FNavigationToolScopedSelection& InSelection) const
+{
+	return InSelection.IsSelected(GetParentSequence(), MarkedFrameIndex);
+}
+
+void FNavigationToolMarker::Select(FNavigationToolScopedSelection& InSelection) const
+{
+	InSelection.Select(GetParentSequence(), MarkedFrameIndex);
+}
+
+void FNavigationToolMarker::OnSelect()
+{
+	Super::OnSelect();
+}
+
+void FNavigationToolMarker::OnDoubleClick()
+{
+	UMovieSceneSequence* const Sequence = GetParentSequence();
+	if (!IsValid(Sequence))
+	{
+		return;
+	}
+
+	const FMovieSceneMarkedFrame* const MarkedFrame = GetMarkedFrame();
+	if (!MarkedFrame)
+	{
+		return;
+	}
+
+	FocusSequence(Tool, *Sequence, *MarkedFrame);
+}
+
+FNavigationToolItemId FNavigationToolMarker::CalculateItemId() const
+{
+	const TSharedPtr<FNavigationToolSequence> ParentSequenceItem = WeakParentSequenceItem.Pin();
+	check(ParentSequenceItem.IsValid());
+
+	FString MarkId = TEXT("Mark");
+	FNavigationToolItemId::AddSeparatedSegment(MarkId, FString::FromInt(MarkedFrameIndex));
+
+	return FNavigationToolItemId(GetParent()
+		, ParentSequenceItem->GetSequence()
+		, ParentSequenceItem->GetSubSection()
+		, ParentSequenceItem->GetSubSectionIndex()
+		, MarkId);
+}
+
+bool FNavigationToolMarker::CanDelete() const
+{
+	return IsItemValid();
+}
+
+bool FNavigationToolMarker::Delete()
+{
+	if (!IsItemValid())
+	{
+		return false;
+	}
+
+	UMovieScene* const MovieScene = GetParentMovieScene();
+	if (!IsValid(MovieScene))
+	{
+		return false;
+	}
+
+	const FScopedTransaction Transaction(LOCTEXT("NavigationToolDeleteMarker", "Delete Marker"));
+
+	MovieScene->Modify();
+	MovieScene->DeleteMarkedFrame(MarkedFrameIndex);
+
+	return true;
+}
+	
+bool FNavigationToolMarker::CanRename() const
+{
+	if (!IsItemValid())
+	{
+		return false;
+	}
+
+	const TSharedPtr<FNavigationToolSequence> ParentSequenceItem = WeakParentSequenceItem.Pin();
+	if (!ParentSequenceItem.IsValid())
+	{
+		return false;
+	}
+
+	const UMovieScene* const MovieScene = ParentSequenceItem->GetSequenceMovieScene();
+	if (!IsValid(MovieScene))
+	{
+		return false;
+	}
+
+	return MovieScene->CanModify();
+}
+
+bool FNavigationToolMarker::Rename(const FString& InName)
+{
+	if (!WeakParentSequenceItem.IsValid() || MarkedFrameIndex == INDEX_NONE)
+	{
+		return false;
+	}
+
+	UMovieScene* const MovieScene = WeakParentSequenceItem.Pin()->GetSequenceMovieScene();
+	if (!IsValid(MovieScene))
+	{
+		return false;
+	}
+
+	const TArray<FMovieSceneMarkedFrame>& MarkedFrames = MovieScene->GetMarkedFrames();
+	FMovieSceneMarkedFrame& MarkedFrame = const_cast<FMovieSceneMarkedFrame&>(MarkedFrames[MarkedFrameIndex]);
+
+	bool bRenamed = false;
+
+	if (!InName.Equals(MarkedFrame.Label, ESearchCase::CaseSensitive))
+	{
+		const FScopedTransaction Transaction(LOCTEXT("RenameTransaction", "Rename Marker"));
+
+		MovieScene->Modify();
+		MarkedFrame.Label = InName;
+
+		bRenamed = true;
+	}
+
+	return bRenamed;
+}
+
+FFrameNumber FNavigationToolMarker::GetInTime() const
+{
+	if (!WeakParentSequenceItem.IsValid() || MarkedFrameIndex == INDEX_NONE)
+	{
+		return 0;
+	}
+
+	UMovieSceneSequence* const Sequence = WeakParentSequenceItem.Pin()->GetSequence();
+	if (!IsValid(Sequence))
+	{
+		return 0;
+	}
+
+	UMovieScene* const MovieScene = Sequence->GetMovieScene();
+	if (!IsValid(MovieScene))
+	{
+		return 0;
+	}
+
+	const TArray<FMovieSceneMarkedFrame>& MarkedFrames = MovieScene->GetMarkedFrames();
+	FMovieSceneMarkedFrame& MarkedFrame = const_cast<FMovieSceneMarkedFrame&>(MarkedFrames[MarkedFrameIndex]);
+
+	return ConvertToDisplayRateTime(*Sequence, MarkedFrame.FrameNumber).FrameNumber;
+}
+
+void FNavigationToolMarker::SetInTime(const FFrameNumber& InTime)
+{
+	if (!WeakParentSequenceItem.IsValid() || MarkedFrameIndex == INDEX_NONE)
+	{
+		return;
+	}
+
+	UMovieSceneSequence* const Sequence = WeakParentSequenceItem.Pin()->GetSequence();
+	if (!IsValid(Sequence))
+	{
+		return;
+	}
+
+	UMovieScene* const MovieScene = Sequence->GetMovieScene();
+	if (!IsValid(MovieScene))
+	{
+		return;
+	}
+
+	const TArray<FMovieSceneMarkedFrame>& MarkedFrames = MovieScene->GetMarkedFrames();
+	FMovieSceneMarkedFrame& MarkedFrame = const_cast<FMovieSceneMarkedFrame&>(MarkedFrames[MarkedFrameIndex]);
+
+	const FScopedTransaction Transaction(LOCTEXT("SetTimeTransaction", "Set Marked Frame Time"));
+
+	MovieScene->Modify();
+	MarkedFrame.FrameNumber = ConvertToTickResolutionTime(*Sequence, InTime).FrameNumber;
+}
+
+int32 FNavigationToolMarker::GetMarkedFrameIndex() const
+{
+	return MarkedFrameIndex;
+}
+
+FMovieSceneMarkedFrame* FNavigationToolMarker::GetMarkedFrame() const
+{
+	if (!IsItemValid())
+	{
+		return nullptr;
+	}
+
+	UMovieScene* const MovieScene = GetParentMovieScene();
+	if (!IsValid(MovieScene))
+	{
+		return nullptr;
+	}
+
+	const TArray<FMovieSceneMarkedFrame>& MarkedFrames = MovieScene->GetMarkedFrames();
+	if (!MarkedFrames.IsValidIndex(MarkedFrameIndex))
+	{
+		return nullptr;
+	}
+
+	return const_cast<FMovieSceneMarkedFrame*>(&MarkedFrames[MarkedFrameIndex]);
+}
+
+UMovieSceneSequence* FNavigationToolMarker::GetParentSequence() const
+{
+	if (const TSharedPtr<FNavigationToolSequence> ParentSequenceItem = WeakParentSequenceItem.Pin())
+	{
+		return ParentSequenceItem->GetSequence();
+	}
+	return nullptr;
+}
+
+UMovieScene* FNavigationToolMarker::GetParentMovieScene() const
+{
+	if (const TSharedPtr<FNavigationToolSequence> ParentSequenceItem = WeakParentSequenceItem.Pin())
+	{
+		return ParentSequenceItem->GetSequenceMovieScene();
+	}
+	return nullptr;
+}
+
+}
+
+#undef LOCTEXT_NAMESPACE
diff --git a/Engine/Source/Editor/Sequencer/Private/Tools/NavigationTool/Items/NavigationToolSequence.cpp b/Engine/Source/Editor/Sequencer/Private/Tools/NavigationTool/Items/NavigationToolSequence.cpp
--- a/Engine/Source/Editor/Sequencer/Private/Tools/NavigationTool/Items/NavigationToolSequence.cpp
+++ b/Engine/Source/Editor/Sequencer/Private/Tools/NavigationTool/Items/NavigationToolSequence.cpp
@@ -16,11 +16,12 @@
 #include "Tools/NavigationTool/INavigationTool.h"
 #include "Tools/NavigationTool/Items/NavigationToolActor.h"
 #include "Tools/NavigationTool/Items/NavigationToolItemUtils.h"
-#include "Tools/NavigationTool/Utils/NavigationToolMovieSceneUtils.h"
+#include "Tools/NavigationTool/Items/NavigationToolMarker.h"
 #include "Tools/NavigationTool/Items/NavigationToolSubTrack.h"
 #include "Tools/NavigationTool/Items/NavigationToolTrack.h"
 #include "Tools/NavigationTool/NavigationToolScopedSelection.h"
 #include "Tools/NavigationTool/Utils/NavigationToolMiscUtils.h"
+#include "Tools/NavigationTool/Utils/NavigationToolMovieSceneUtils.h"
 #include "Tracks/MovieSceneCinematicShotTrack.h"
 
 #define LOCTEXT_NAMESPACE "NavigationToolSequence"
@@ -118,6 +119,19 @@ void FNavigationToolSequence::FindChildren(TArray<FNavigationToolItemPtr>& OutCh
 			}
 		}
 	}
+
+	// Add sequence marker items
+	const TArray<FMovieSceneMarkedFrame>& MarkedFrames = MovieScene->GetMarkedFrames();
+	for (int32 Index = 0; Index < MarkedFrames.Num(); ++Index)
+	{
+		const FNavigationToolItemPtr NewItem = Tool.FindOrAdd<FNavigationToolMarker>(GetProvider()
+			, SharedThisRef, SharedThisRef, Index);
+		OutChildren.Add(NewItem);
+		if (bInRecursive)
+		{
+			NewItem->FindChildren(OutChildren, bInRecursive);
+		}
+	}
 }
 
 void FNavigationToolSequence::GetItemProxies(TArray<TSharedPtr<FNavigationToolItemProxy>>& OutItemProxies)
@@ -539,6 +553,94 @@ EItemContainsPlayhead FNavigationToolSequence::ContainsPlayhead() const
 	return EItemContainsPlayhead::None;
 }
 
+FFrameNumber FNavigationToolSequence::GetInTime() const
+{
+	const UMovieSceneSequence* const Sequence = WeakSequence.Get();
+	if (!IsValid(Sequence))
+	{
+		return 0;
+	}
+
+	const UMovieSceneSubSection* const SubSection = WeakSubSection.Get();
+	if (!IsValid(SubSection))
+	{
+		return 0;
+	}
+
+	const FFrameTime LowerBoundValue = SubSection->SectionRange.Value.GetLowerBoundValue();
+
+	return ConvertToDisplayRateTime(*Sequence, LowerBoundValue).FrameNumber;
+}
+
+void FNavigationToolSequence::SetInTime(const FFrameNumber& InTime)
+{
+	const UMovieSceneSequence* const Sequence = WeakSequence.Get();
+	if (!IsValid(Sequence))
+	{
+		return;
+	}
+
+	UMovieSceneSubSection* const SubSection = WeakSubSection.Get();
+	if (!IsValid(SubSection))
+	{
+		return;
+	}
+
+	const FFrameNumber NewFrameTime = ConvertToTickResolutionTime(*Sequence, InTime).FrameNumber;
+
+	if (SubSection->SectionRange.Value.GetLowerBoundValue() == NewFrameTime)
+	{
+		return;
+	}
+
+	SubSection->Modify();
+	SubSection->SectionRange.Value.SetLowerBoundValue(NewFrameTime);
+}
+
+FFrameNumber FNavigationToolSequence::GetOutTime() const
+{
+	const UMovieSceneSequence* const Sequence = WeakSequence.Get();
+	if (!IsValid(Sequence))
+	{
+		return 0;
+	}
+
+	const UMovieSceneSubSection* const SubSection = WeakSubSection.Get();
+	if (!IsValid(SubSection))
+	{
+		return 0;
+	}
+
+	const FFrameTime UpperBoundValue = SubSection->SectionRange.Value.GetUpperBoundValue();
+
+	return ConvertToDisplayRateTime(*Sequence, UpperBoundValue).FrameNumber;
+}
+
+void FNavigationToolSequence::SetOutTime(const FFrameNumber& InTime)
+{
+	const UMovieSceneSequence* const Sequence = WeakSequence.Get();
+	if (!IsValid(Sequence))
+	{
+		return;
+	}
+
+	UMovieSceneSubSection* const SubSection = WeakSubSection.Get();
+	if (!IsValid(SubSection))
+	{
+		return;
+	}
+
+	const FFrameNumber NewFrameTime = ConvertToTickResolutionTime(*Sequence, InTime).FrameNumber;
+
+	if (SubSection->SectionRange.Value.GetUpperBoundValue() == NewFrameTime)
+	{
+		return;
+	}
+
+	SubSection->Modify();
+	SubSection->SectionRange.Value.SetUpperBoundValue(NewFrameTime);
+}
+
 UMovieSceneSequence* FNavigationToolSequence::GetSequence() const
 {
 	return WeakSequence.Get();
diff --git a/Engine/Source/Editor/Sequencer/Private/Tools/NavigationTool/NavigationToolScopedSelection.cpp b/Engine/Source/Editor/Sequencer/Private/Tools/NavigationTool/NavigationToolScopedSelection.cpp
--- a/Engine/Source/Editor/Sequencer/Private/Tools/NavigationTool/NavigationToolScopedSelection.cpp
+++ b/Engine/Source/Editor/Sequencer/Private/Tools/NavigationTool/NavigationToolScopedSelection.cpp
@@ -1,13 +1,10 @@
 // Copyright Epic Games, Inc. All Rights Reserved.
 
 #include "Tools/NavigationTool/NavigationToolScopedSelection.h"
-#include "Components/ActorComponent.h"
-#include "GameFramework/Actor.h"
 #include "MVVM/ObjectBindingModelStorageExtension.h"
 #include "MVVM/SectionModelStorageExtension.h"
 #include "MVVM/Selection/Selection.h"
 #include "MVVM/TrackModelStorageExtension.h"
-#include "MVVM/TrackRowModelStorageExtension.h"
 #include "MVVM/Views/SOutlinerView.h"
 #include "Sequencer.h"
 #include "Utils/NavigationToolMiscUtils.h"
@@ -102,6 +99,23 @@ void FNavigationToolScopedSelection::Select(UMovieSceneTrack* const InTrack)
 	SelectedTracks.Add(InTrack);
 }
 
+void FNavigationToolScopedSelection::Select(UMovieSceneSequence* const InSequence
+	, const int32 InMarkedFrameIndex)
+{
+	if (!ensureMsgf(Purpose == ENavigationToolScopedSelectionPurpose::Sync
+		, TEXT("Scope is trying to Select, but it's not a Sync Scope.")))
+	{
+		return;
+	}
+
+	if (!IsValid(InSequence) || InMarkedFrameIndex == INDEX_NONE)
+	{
+		return;
+	}
+
+	SelectedMarkedFrames.FindOrAdd(InSequence).Add(InMarkedFrameIndex);
+}
+
 bool FNavigationToolScopedSelection::IsSelected(const UObject* const InObject) const
 {
 	return ObjectsSet.Contains(InObject);
@@ -122,6 +136,16 @@ bool FNavigationToolScopedSelection::IsSelected(UMovieSceneTrack* const InTrack)
 	return SelectedTracks.Contains(InTrack);
 }
 
+bool FNavigationToolScopedSelection::IsSelected(UMovieSceneSequence* const InSequence
+	, const int32 InMarkedFrameIndex) const
+{
+	if (IsValid(InSequence) && SelectedMarkedFrames.Contains(InSequence))
+	{
+		return SelectedMarkedFrames[InSequence].Contains(InMarkedFrameIndex);
+	}
+	return false;
+}
+
 void FNavigationToolScopedSelection::SyncSelections()
 {
 	const TSharedPtr<FSequencerEditorViewModel> ViewModel = Sequencer.GetViewModel();
@@ -144,11 +168,11 @@ void FNavigationToolScopedSelection::SyncSelections()
 
 	const FModifierKeysState ModifierKeys = FSlateApplication::Get().GetModifierKeys();
 	UNavigationToolSettings* const ToolSettings = GetMutableDefault<UNavigationToolSettings>();
-	const bool bPreviouslySyncingSelectionToNavigationTool = ToolSettings->ShouldSyncSelectionToNavigationTool();
+	const bool bPreviouslySyncingSelection = ToolSettings->ShouldSyncSelectionToNavigationTool();
 
 	if (!ModifierKeys.IsAltDown())
 	{
-		ToolSettings->SetSyncSelectionToNavigationTool(false);
+		ToolSettings->SetSyncSelectionToNavigationTool(/*bInSync=*/false, /*bInSaveConfig=*/false);
 	}
 
 	const TSharedPtr<SOutlinerView> OutlinerView = GetOutlinerViewWidget(Sequencer);
@@ -208,7 +232,7 @@ void FNavigationToolScopedSelection::SyncSelections()
 
 	if (!ModifierKeys.IsAltDown())
 	{
-		ToolSettings->SetSyncSelectionToNavigationTool(bPreviouslySyncingSelectionToNavigationTool);
+		ToolSettings->SetSyncSelectionToNavigationTool(bPreviouslySyncingSelection, /*bInSaveConfig=*/false);
 	}
 }
 
diff --git a/Engine/Source/Editor/Sequencer/Private/Tools/NavigationTool/Utils/NavigationToolMiscUtils.cpp b/Engine/Source/Editor/Sequencer/Private/Tools/NavigationTool/Utils/NavigationToolMiscUtils.cpp
--- a/Engine/Source/Editor/Sequencer/Private/Tools/NavigationTool/Utils/NavigationToolMiscUtils.cpp
+++ b/Engine/Source/Editor/Sequencer/Private/Tools/NavigationTool/Utils/NavigationToolMiscUtils.cpp
@@ -3,6 +3,7 @@
 #include "NavigationToolMiscUtils.h"
 #include "ISourceControlModule.h"
 #include "ISourceControlProvider.h"
+#include "MovieScene.h"
 #include "MVVM/ObjectBindingModelStorageExtension.h"
 #include "MVVM/Selection/Selection.h"
 #include "MVVM/ViewModels/ObjectBindingModel.h"
@@ -203,6 +204,18 @@ void SelectSequencerBindingTrack(const ISequencer& InSequencer, const FGuid& InO
 	Selection->Outliner.Select(ObjectBindingModel);
 }
 
+void FocusSequence(const INavigationTool& InTool
+	, UMovieSceneSequence& InSequence
+	, const FMovieSceneMarkedFrame& InMarkedFrame)
+{
+	FocusSequence(InTool, InSequence);
+
+	if (const TSharedPtr<ISequencer> Sequencer = InTool.GetSequencer())
+	{
+		Sequencer->SetGlobalTime(InMarkedFrame.FrameNumber);
+	}
+}
+
 void FocusItemInSequencer(const INavigationTool& InTool, const FNavigationToolItemPtr& InItem)
 {
 	if (const FNavigationToolSequence* const SequenceItem = InItem->CastTo<FNavigationToolSequence>())
diff --git a/Engine/Source/Editor/Sequencer/Private/Tools/NavigationTool/Utils/NavigationToolMiscUtils.h b/Engine/Source/Editor/Sequencer/Private/Tools/NavigationTool/Utils/NavigationToolMiscUtils.h
--- a/Engine/Source/Editor/Sequencer/Private/Tools/NavigationTool/Utils/NavigationToolMiscUtils.h
+++ b/Engine/Source/Editor/Sequencer/Private/Tools/NavigationTool/Utils/NavigationToolMiscUtils.h
@@ -90,6 +90,16 @@ void FocusSequence(const INavigationTool& InTool
 /** Selects an object binding track in a sequencer. */
 void SelectSequencerBindingTrack(const ISequencer& InSequencer, const FGuid& InObjectId);
 
+/**
+ * Focuses a sub-movie scene (MovieScene within a MovieScene) in the sequencer.
+ * @param InTool The Navigation Tool that manages scene sequences.
+ * @param InSequence The movie scene sequence to focus on.
+ * @param InMarkedFrame Optional marked frame to select in the sequence.
+ */
+void FocusSequence(const INavigationTool& InTool
+	, UMovieSceneSequence& InSequence
+	, const FMovieSceneMarkedFrame& InMarkedFrame);
+
 void FocusItemInSequencer(const INavigationTool& InTool
 	, const FNavigationToolItemPtr& InItem);
 
diff --git a/Engine/Source/Editor/Sequencer/Private/Tools/NavigationTool/Utils/NavigationToolMovieSceneUtils.cpp b/Engine/Source/Editor/Sequencer/Private/Tools/NavigationTool/Utils/NavigationToolMovieSceneUtils.cpp
--- a/Engine/Source/Editor/Sequencer/Private/Tools/NavigationTool/Utils/NavigationToolMovieSceneUtils.cpp
+++ b/Engine/Source/Editor/Sequencer/Private/Tools/NavigationTool/Utils/NavigationToolMovieSceneUtils.cpp
@@ -1,6 +1,7 @@
 ﻿// Copyright Epic Games, Inc. All Rights Reserved.
 
 #include "NavigationToolMovieSceneUtils.h"
+#include "Misc/FrameTime.h"
 #include "MovieScene.h"
 #include "Sections/MovieSceneSubSection.h"
 #include "Sequencer.h"
@@ -9,6 +10,30 @@
 namespace UE::Sequencer::NavigationTool
 {
 
+FFrameTime ConvertToDisplayRateTime(const UMovieSceneSequence& InSequence, const FFrameTime& InTime)
+{
+	UMovieScene* const MovieScene = InSequence.GetMovieScene();
+	if (IsValid(MovieScene))
+	{
+		return ConvertFrameTime(InTime
+			, MovieScene->GetTickResolution()
+			, MovieScene->GetDisplayRate());
+	}
+	return 0;
+}
+
+FFrameTime ConvertToTickResolutionTime(const UMovieSceneSequence& InSequence, const FFrameTime& InTime)
+{
+	UMovieScene* const MovieScene = InSequence.GetMovieScene();
+	if (IsValid(MovieScene))
+	{
+		return ConvertFrameTime(InTime
+			, MovieScene->GetDisplayRate()
+			, MovieScene->GetTickResolution());
+	}
+	return 0;
+}
+
 UMovieSceneSubSection* FindSequenceSubSection(ISequencer& InSequencer, UMovieSceneSequence* const InSequence)
 {
 	FMovieSceneEvaluationState* const EvaluationState = InSequencer.GetEvaluationState();
diff --git a/Engine/Source/Editor/Sequencer/Private/Tools/NavigationTool/Utils/NavigationToolMovieSceneUtils.h b/Engine/Source/Editor/Sequencer/Private/Tools/NavigationTool/Utils/NavigationToolMovieSceneUtils.h
--- a/Engine/Source/Editor/Sequencer/Private/Tools/NavigationTool/Utils/NavigationToolMovieSceneUtils.h
+++ b/Engine/Source/Editor/Sequencer/Private/Tools/NavigationTool/Utils/NavigationToolMovieSceneUtils.h
@@ -8,10 +8,14 @@ class ISequencer;
 class UMovieScene;
 class UMovieSceneSequence;
 class UMovieSceneSubSection;
+struct FFrameTime;
 
 namespace UE::Sequencer::NavigationTool
 {
 
+FFrameTime ConvertToDisplayRateTime(const UMovieSceneSequence& InSequence, const FFrameTime& InTime);
+FFrameTime ConvertToTickResolutionTime(const UMovieSceneSequence& InSequence, const FFrameTime& InTime);
+
 /**
  * Finds the subsection corresponding to a given sequence within the specified sequencer.
  * @param InSequencer The sequencer to search within.
diff --git a/Engine/Source/Editor/Sequencer/Private/Tools/NavigationTool/Widgets/Columns/SNavigationToolInTime.cpp b/Engine/Source/Editor/Sequencer/Private/Tools/NavigationTool/Widgets/Columns/SNavigationToolInTime.cpp
--- a/Engine/Source/Editor/Sequencer/Private/Tools/NavigationTool/Widgets/Columns/SNavigationToolInTime.cpp
+++ b/Engine/Source/Editor/Sequencer/Private/Tools/NavigationTool/Widgets/Columns/SNavigationToolInTime.cpp
@@ -1,14 +1,11 @@
 // Copyright Epic Games, Inc. All Rights Reserved.
 
 #include "SNavigationToolInTime.h"
-#include "Sequencer.h"
+#include "ScopedTransaction.h"
+#include "Tools/NavigationTool/Extensions/IInTimeExtension.h"
 #include "Tools/NavigationTool/Items/NavigationToolItemUtils.h"
-#include "Tools/NavigationTool/Items/NavigationToolSequence.h"
-#include "Tools/NavigationTool/NavigationTool.h"
 #include "Tools/NavigationTool/NavigationToolView.h"
 #include "Tools/NavigationTool/Widgets/SNavigationToolTreeRow.h"
-#include "TrackEditors/SubTrackEditorBase.h"
-#include "Widgets/Input/SSpinBox.h"
 
 #define LOCTEXT_NAMESPACE "SNavigationToolInTime"
 
@@ -24,97 +21,39 @@ void SNavigationToolInTime::Construct(const FArguments& InArgs
 	WeakView = InView;
 	WeakRowWidget = InRowWidget;
 
-	auto EmptyContent = [this]()
+	if (!WeakItem.IsValid()
+		|| !WeakView.IsValid()
+		|| !WeakItem.Pin()->IsA<IInTimeExtension>())
 	{
-		ChildSlot
-		[
-			SNullWidget::NullWidget
-		];
-	};
-
-	if (!WeakItem.IsValid() || !WeakView.IsValid())
-	{
-		EmptyContent();
-		return;
-	}
-
-	const FNavigationToolSequence* const SequenceItem = WeakItem.Pin()->CastTo<FNavigationToolSequence>();
-	if (!SequenceItem || !IsValid(SequenceItem->GetSubSection()))
-	{
-		EmptyContent();
 		return;
 	}
 
-	ChildSlot
-		[
-			SNew(SSpinBox<int32>)
-			.Justification(ETextJustify::Center)
-			.Style(&FAppStyle::GetWidgetStyle<FSpinBoxStyle>(TEXT("Sequencer.NavigationTool.InTime.SpinBox")))
-			.Value(this, &SNavigationToolInTime::GetFrameTimeValue)
-			.OnValueChanged(this, &SNavigationToolInTime::OnFrameTimeValueChanged)
-			.OnValueCommitted(this, &SNavigationToolInTime::OnFrameTimeValueCommitted)
-			.OnBeginSliderMovement(this, &SNavigationToolInTime::OnBeginSliderMovement)
-			.OnEndSliderMovement(this, &SNavigationToolInTime::OnEndSliderMovement)
-		];
+	SNavigationToolTime::Construct(SNavigationToolTime::FArguments(), InItem, InView, InRowWidget);
 }
 
-int32 SNavigationToolInTime::GetFrameTimeValue() const
+IInTimeExtension* SNavigationToolInTime::GetInTimeExtension() const
 {
-	const TSharedPtr<INavigationToolView> ToolView = WeakView.Pin();
-	if (!ToolView.IsValid())
-	{
-		return 0;
-	}
-
-	const TSharedPtr<ISequencer> Sequencer = ToolView->GetSequencer();
-	if (!Sequencer.IsValid())
-	{
-		return 0;
-	}
-
-	UMovieSceneSubSection* const SubSection = Cast<UMovieSceneSubSection>(GetSequenceItemSubSection(WeakItem.Pin()));
-	if (!IsValid(SubSection))
-	{
-		return 0;
-	}
-
-	const FFrameNumber FrameTime = ConvertFrameTime(SubSection->SectionRange.Value.GetLowerBoundValue()
-		, Sequencer->GetFocusedTickResolution()
-		, Sequencer->GetFocusedDisplayRate()).FrameNumber;
-
-	return FrameTime.Value;
+	return WeakItem.IsValid() ? WeakItem.Pin()->CastTo<IInTimeExtension>() : nullptr;
 }
 
-void SNavigationToolInTime::OnFrameTimeValueChanged(const int32 InNewValue)
+FName SNavigationToolInTime::GetStyleName() const
 {
-	OnFrameTimeValueCommitted(InNewValue, ETextCommit::Default);
+	return TEXT("Sequencer.NavigationTool.InTime.SpinBox");
 }
 
-void SNavigationToolInTime::OnFrameTimeValueCommitted(const int32 InNewValue, const ETextCommit::Type InCommitType)
+int32 SNavigationToolInTime::GetFrameTimeValue() const
 {
-	const TSharedPtr<INavigationToolView> ToolView = WeakView.Pin();
-	if (!ToolView.IsValid())
-	{
-		return;
-	}
-
-	const TSharedPtr<ISequencer> Sequencer = ToolView->GetSequencer();
-	if (!Sequencer.IsValid())
-	{
-		return;
-	}
-
-	UMovieSceneSubSection* const SubSection = Cast<UMovieSceneSubSection>(GetSequenceItemSubSection(WeakItem.Pin()));
-	if (!IsValid(SubSection))
+	if (const IInTimeExtension* const InTimeExtension = GetInTimeExtension())
 	{
-		return;
+		return InTimeExtension->GetInTime().Value;
 	}
+	return 0;
+}
 
-	const FFrameNumber NewFrameTime = ConvertFrameTime(InNewValue
-		, Sequencer->GetFocusedDisplayRate()
-		, Sequencer->GetFocusedTickResolution()).FrameNumber;
-
-	if (SubSection->SectionRange.Value.GetLowerBoundValue() == NewFrameTime)
+void SNavigationToolInTime::OnFrameTimeValueCommitted(const int32 InNewValue, const ETextCommit::Type InCommitType)
+{
+	IInTimeExtension* const InTimeExtension = GetInTimeExtension();
+	if (!InTimeExtension)
 	{
 		return;
 	}
@@ -122,26 +61,12 @@ void SNavigationToolInTime::OnFrameTimeValueCommitted(const int32 InNewValue, co
 	const bool bShouldTransact = !UndoTransaction.IsValid() && (InCommitType == ETextCommit::OnEnter);
 	const FScopedTransaction Transaction(GetTransactionText(), bShouldTransact);
 
-	SubSection->Modify();
-	SubSection->SectionRange.Value.SetLowerBoundValue(NewFrameTime);
-}
-
-void SNavigationToolInTime::OnBeginSliderMovement()
-{
-	if (!UndoTransaction.IsValid())
-	{
-		UndoTransaction = MakeUnique<FScopedTransaction>(GetTransactionText());
-	}
-}
-
-void SNavigationToolInTime::OnEndSliderMovement(const int32 InNewValue)
-{
-	UndoTransaction.Reset();
+	return InTimeExtension->SetInTime(InNewValue);
 }
 
 FText SNavigationToolInTime::GetTransactionText() const
 {
-	return LOCTEXT("SetSequenceInTimeTransaction", "Set Sequence In Time");
+	return LOCTEXT("SetInTimeTransaction", "Set In Time");
 }
 
 }
diff --git a/Engine/Source/Editor/Sequencer/Private/Tools/NavigationTool/Widgets/Columns/SNavigationToolInTime.h b/Engine/Source/Editor/Sequencer/Private/Tools/NavigationTool/Widgets/Columns/SNavigationToolInTime.h
--- a/Engine/Source/Editor/Sequencer/Private/Tools/NavigationTool/Widgets/Columns/SNavigationToolInTime.h
+++ b/Engine/Source/Editor/Sequencer/Private/Tools/NavigationTool/Widgets/Columns/SNavigationToolInTime.h
@@ -2,18 +2,18 @@
 
 #pragma once
 
-#include "ScopedTransaction.h"
+#include "SNavigationToolTime.h"
 #include "Templates/UniquePtr.h"
 #include "Tools/NavigationTool/NavigationToolDefines.h"
-#include "Widgets/SCompoundWidget.h"
 
 namespace UE::Sequencer::NavigationTool
 {
 
+class IInTimeExtension;
 class INavigationToolView;
 class SNavigationToolTreeRow;
 
-class SNavigationToolInTime : public SCompoundWidget
+class SNavigationToolInTime : public SNavigationToolTime
 {
 public:
 	SLATE_BEGIN_ARGS(SNavigationToolInTime) {}
@@ -27,21 +27,15 @@ public:
 		, const TSharedRef<SNavigationToolTreeRow>& InRowWidget);
 
 private:
-	int32 GetFrameTimeValue() const;
-	void OnFrameTimeValueChanged(const int32 InNewValue);
-	void OnFrameTimeValueCommitted(const int32 InNewValue, const ETextCommit::Type InCommitType);
-	void OnBeginSliderMovement();
-	void OnEndSliderMovement(const int32 InNewValue);
+	IInTimeExtension* GetInTimeExtension() const;
 
-	FText GetTransactionText() const;
+	virtual FName GetStyleName() const override;
 
-	FNavigationToolItemWeakPtr WeakItem;
+	virtual int32 GetFrameTimeValue() const override;
 
-	TWeakPtr<INavigationToolView> WeakView;
+	virtual void OnFrameTimeValueCommitted(const int32 InNewValue, const ETextCommit::Type InCommitType) override;
 
-	TWeakPtr<SNavigationToolTreeRow> WeakRowWidget;
-
-	TUniquePtr<FScopedTransaction> UndoTransaction;
+	virtual FText GetTransactionText() const override;
 };
 
 }
diff --git a/Engine/Source/Editor/Sequencer/Private/Tools/NavigationTool/Widgets/Columns/SNavigationToolLabelItem.cpp b/Engine/Source/Editor/Sequencer/Private/Tools/NavigationTool/Widgets/Columns/SNavigationToolLabelItem.cpp
--- a/Engine/Source/Editor/Sequencer/Private/Tools/NavigationTool/Widgets/Columns/SNavigationToolLabelItem.cpp
+++ b/Engine/Source/Editor/Sequencer/Private/Tools/NavigationTool/Widgets/Columns/SNavigationToolLabelItem.cpp
@@ -60,6 +60,7 @@ void SNavigationToolLabelItem::Construct(const FArguments& InArgs
 					[
 						SNew(SImage)
 						.IsEnabled(this, &SNavigationToolLabelItem::IsItemEnabled)
+						.ColorAndOpacity(ItemRef, &INavigationToolItem::GetIconColor)
 						.Image(ItemRef, &INavigationToolItem::GetIconBrush)
 						.ToolTipText(ItemRef, &INavigationToolItem::GetIconTooltipText)
 					]
diff --git a/Engine/Source/Editor/Sequencer/Private/Tools/NavigationTool/Widgets/Columns/SNavigationToolOutTime.cpp b/Engine/Source/Editor/Sequencer/Private/Tools/NavigationTool/Widgets/Columns/SNavigationToolOutTime.cpp
--- a/Engine/Source/Editor/Sequencer/Private/Tools/NavigationTool/Widgets/Columns/SNavigationToolOutTime.cpp
+++ b/Engine/Source/Editor/Sequencer/Private/Tools/NavigationTool/Widgets/Columns/SNavigationToolOutTime.cpp
@@ -2,14 +2,10 @@
 
 #include "SNavigationToolOutTime.h"
 #include "ScopedTransaction.h"
-#include "Sequencer.h"
 #include "Tools/NavigationTool/Items/NavigationToolItemUtils.h"
-#include "Tools/NavigationTool/NavigationTool.h"
 #include "Tools/NavigationTool/NavigationToolView.h"
 #include "Tools/NavigationTool/Items/NavigationToolSequence.h"
 #include "Tools/NavigationTool/Widgets/SNavigationToolTreeRow.h"
-#include "TrackEditors/SubTrackEditorBase.h"
-#include "Widgets/Input/SSpinBox.h"
 
 #define LOCTEXT_NAMESPACE "SNavigationToolOutTime"
 
@@ -25,97 +21,39 @@ void SNavigationToolOutTime::Construct(const FArguments& InArgs
 	WeakView = InView;
 	WeakRowWidget = InRowWidget;
 
-	auto EmptyContent = [this]()
+	if (!WeakItem.IsValid()
+		|| !WeakView.IsValid()
+		|| !WeakItem.Pin()->IsA<IOutTimeExtension>())
 	{
-		ChildSlot
-		[
-			SNullWidget::NullWidget
-		];
-	};
-
-	if (!WeakItem.IsValid() || !WeakView.IsValid())
-	{
-		EmptyContent();
-		return;
-	}
-
-	const FNavigationToolSequence* const SequenceItem = WeakItem.Pin()->CastTo<FNavigationToolSequence>();
-	if (!SequenceItem || !IsValid(SequenceItem->GetSubSection()))
-	{
-		EmptyContent();
 		return;
 	}
 
-	ChildSlot
-		[
-			SNew(SSpinBox<int32>)
-			.Justification(ETextJustify::Center)
-			.Style(&FAppStyle::GetWidgetStyle<FSpinBoxStyle>(TEXT("Sequencer.NavigationTool.OutTime.SpinBox")))
-			.Value(this, &SNavigationToolOutTime::GetFrameTimeValue)
-			.OnValueChanged(this, &SNavigationToolOutTime::OnFrameTimeValueChanged)
-			.OnValueCommitted(this, &SNavigationToolOutTime::OnFrameTimeValueCommitted)
-			.OnBeginSliderMovement(this, &SNavigationToolOutTime::OnBeginSliderMovement)
-			.OnEndSliderMovement(this, &SNavigationToolOutTime::OnEndSliderMovement)
-		];
+	SNavigationToolTime::Construct(SNavigationToolTime::FArguments(), InItem, InView, InRowWidget);
 }
 
-int32 SNavigationToolOutTime::GetFrameTimeValue() const
+IOutTimeExtension* SNavigationToolOutTime::GetOutTimeExtension() const
 {
-	const TSharedPtr<INavigationToolView> ToolView = WeakView.Pin();
-	if (!ToolView.IsValid())
-	{
-		return 0;
-	}
-
-	const TSharedPtr<ISequencer> Sequencer = ToolView->GetSequencer();
-	if (!Sequencer.IsValid())
-	{
-		return 0;
-	}
-
-	UMovieSceneSubSection* const SubSection = Cast<UMovieSceneSubSection>(GetSequenceItemSubSection(WeakItem.Pin()));
-	if (!IsValid(SubSection))
-	{
-		return 0;
-	}
-
-	const FFrameNumber FrameTime = ConvertFrameTime(SubSection->SectionRange.Value.GetUpperBoundValue()
-		, Sequencer->GetFocusedTickResolution()
-		, Sequencer->GetFocusedDisplayRate()).FrameNumber;
-
-	return FrameTime.Value;
+	return WeakItem.IsValid() ? WeakItem.Pin()->CastTo<IOutTimeExtension>() : nullptr;
 }
 
-void SNavigationToolOutTime::OnFrameTimeValueChanged(const int32 InNewValue)
+FName SNavigationToolOutTime::GetStyleName() const
 {
-	OnFrameTimeValueCommitted(InNewValue, ETextCommit::Default);
+	return TEXT("Sequencer.NavigationTool.OutTime.SpinBox");
 }
 
-void SNavigationToolOutTime::OnFrameTimeValueCommitted(const int32 InNewValue, const ETextCommit::Type InCommitType)
+int32 SNavigationToolOutTime::GetFrameTimeValue() const
 {
-	const TSharedPtr<INavigationToolView> ToolView = WeakView.Pin();
-	if (!ToolView.IsValid())
-	{
-		return;
-	}
-
-	const TSharedPtr<ISequencer> Sequencer = ToolView->GetSequencer();
-	if (!Sequencer.IsValid())
-	{
-		return;
-	}
-
-	UMovieSceneSubSection* const SubSection = Cast<UMovieSceneSubSection>(GetSequenceItemSubSection(WeakItem.Pin()));
-	if (!IsValid(SubSection))
+	if (const IOutTimeExtension* const OutTimeExtension = GetOutTimeExtension())
 	{
-		return;
+		return OutTimeExtension->GetOutTime().Value;
 	}
+	return 0;
+}
 
-	const FFrameNumber NewFrameTime = ConvertFrameTime(InNewValue
-		, Sequencer->GetFocusedDisplayRate()
-		, Sequencer->GetFocusedTickResolution()).FrameNumber;
-
-	if (SubSection->SectionRange.Value.GetUpperBoundValue() == NewFrameTime)
+void SNavigationToolOutTime::OnFrameTimeValueCommitted(const int32 InNewValue, const ETextCommit::Type InCommitType)
+{
+	IOutTimeExtension* const OutTimeExtension = GetOutTimeExtension();
+	if (!OutTimeExtension)
 	{
 		return;
 	}
@@ -123,26 +61,12 @@ void SNavigationToolOutTime::OnFrameTimeValueCommitted(const int32 InNewValue, c
 	const bool bShouldTransact = !UndoTransaction.IsValid() && (InCommitType == ETextCommit::OnEnter);
 	const FScopedTransaction Transaction(GetTransactionText(), bShouldTransact);
 
-	SubSection->Modify();
-	SubSection->SectionRange.Value.SetUpperBoundValue(NewFrameTime);
-}
-
-void SNavigationToolOutTime::OnBeginSliderMovement()
-{
-	if (!UndoTransaction.IsValid())
-	{
-		UndoTransaction = MakeUnique<FScopedTransaction>(GetTransactionText());
-	}
-}
-
-void SNavigationToolOutTime::OnEndSliderMovement(const int32 InNewValue)
-{
-	UndoTransaction.Reset();
+	return OutTimeExtension->SetOutTime(InNewValue);
 }
 
 FText SNavigationToolOutTime::GetTransactionText() const
 {
-	return LOCTEXT("SetSequenceOutTimeTransaction", "Set Sequence Out Time");
+	return LOCTEXT("SetOutTimeTransaction", "Set Out Time");
 }
 
 }
diff --git a/Engine/Source/Editor/Sequencer/Private/Tools/NavigationTool/Widgets/Columns/SNavigationToolOutTime.h b/Engine/Source/Editor/Sequencer/Private/Tools/NavigationTool/Widgets/Columns/SNavigationToolOutTime.h
--- a/Engine/Source/Editor/Sequencer/Private/Tools/NavigationTool/Widgets/Columns/SNavigationToolOutTime.h
+++ b/Engine/Source/Editor/Sequencer/Private/Tools/NavigationTool/Widgets/Columns/SNavigationToolOutTime.h
@@ -2,18 +2,18 @@
 
 #pragma once
 
-#include "ScopedTransaction.h"
+#include "SNavigationToolTime.h"
 #include "Templates/UniquePtr.h"
 #include "Tools/NavigationTool/NavigationToolDefines.h"
-#include "Widgets/SCompoundWidget.h"
 
 namespace UE::Sequencer::NavigationTool
 {
 
+class IOutTimeExtension;
 class INavigationToolView;
 class SNavigationToolTreeRow;
 
-class SNavigationToolOutTime : public SCompoundWidget
+class SNavigationToolOutTime : public SNavigationToolTime
 {
 public:
 	SLATE_BEGIN_ARGS(SNavigationToolOutTime) {}
@@ -26,22 +26,16 @@ public:
 		, const TSharedRef<INavigationToolView>& InView
 		, const TSharedRef<SNavigationToolTreeRow>& InRowWidget);
 
-private:
-	int32 GetFrameTimeValue() const;
-	void OnFrameTimeValueChanged(const int32 InNewValue);
-	void OnFrameTimeValueCommitted(const int32 InNewValue, const ETextCommit::Type InCommitType);
-	void OnBeginSliderMovement();
-	void OnEndSliderMovement(const int32 InNewValue);
+protected:
+	IOutTimeExtension* GetOutTimeExtension() const;
 
-	FText GetTransactionText() const;
+	virtual FName GetStyleName() const override;
 
-	FNavigationToolItemWeakPtr WeakItem;
+	virtual int32 GetFrameTimeValue() const override;
 
-	TWeakPtr<INavigationToolView> WeakView;
+	virtual void OnFrameTimeValueCommitted(const int32 InNewValue, const ETextCommit::Type InCommitType) override;
 
-	TWeakPtr<SNavigationToolTreeRow> WeakRowWidget;
-
-	TUniquePtr<FScopedTransaction> UndoTransaction;
+	virtual FText GetTransactionText() const override;
 };
 
 }
diff --git a/Engine/Source/Editor/Sequencer/Private/Tools/NavigationTool/Widgets/Columns/SNavigationToolTime.cpp b/Engine/Source/Editor/Sequencer/Private/Tools/NavigationTool/Widgets/Columns/SNavigationToolTime.cpp
new file mode 100644
--- /dev/null
+++ b/Engine/Source/Editor/Sequencer/Private/Tools/NavigationTool/Widgets/Columns/SNavigationToolTime.cpp
@@ -0,0 +1,69 @@
+﻿// Copyright Epic Games, Inc. All Rights Reserved.
+
+#include "SNavigationToolTime.h"
+#include "Tools/NavigationTool/INavigationTool.h"
+#include "Tools/NavigationTool/Items/NavigationToolItemUtils.h"
+#include "Tools/NavigationTool/NavigationToolView.h"
+#include "Tools/NavigationTool/Widgets/SNavigationToolTreeRow.h"
+#include "Widgets/Input/SSpinBox.h"
+
+#define LOCTEXT_NAMESPACE "SNavigationToolTime"
+
+namespace UE::Sequencer::NavigationTool
+{
+
+void SNavigationToolTime::Construct(const FArguments& InArgs
+	, const FNavigationToolItemPtr& InItem
+	, const TSharedRef<INavigationToolView>& InView
+	, const TSharedRef<SNavigationToolTreeRow>& InRowWidget)
+{
+	WeakItem = InItem;
+	WeakView = InView;
+	WeakRowWidget = InRowWidget;
+
+	if (!WeakItem.IsValid()
+		|| !WeakView.IsValid()
+		|| !WeakItem.Pin()->IsA<IInTimeExtension>())
+	{
+		return;
+	}
+
+	ChildSlot
+	[
+		SNew(SSpinBox<int32>)
+		.Justification(ETextJustify::Center)
+		.Style(&FAppStyle::GetWidgetStyle<FSpinBoxStyle>(GetStyleName()))
+		.Value(this, &SNavigationToolTime::GetFrameTimeValue)
+		.OnValueChanged(this, &SNavigationToolTime::OnFrameTimeValueChanged)
+		.OnValueCommitted(this, &SNavigationToolTime::OnFrameTimeValueCommitted)
+		.OnBeginSliderMovement(this, &SNavigationToolTime::OnBeginSliderMovement)
+		.OnEndSliderMovement(this, &SNavigationToolTime::OnEndSliderMovement)
+	];
+}
+
+FName SNavigationToolTime::GetStyleName() const
+{
+	return TEXT("Sequencer.NavigationTool.SpinBox");
+}
+
+void SNavigationToolTime::OnFrameTimeValueChanged(const int32 InNewValue)
+{
+	OnFrameTimeValueCommitted(InNewValue, ETextCommit::Default);
+}
+
+void SNavigationToolTime::OnBeginSliderMovement()
+{
+	if (!UndoTransaction.IsValid())
+	{
+		UndoTransaction = MakeUnique<FScopedTransaction>(GetTransactionText());
+	}
+}
+
+void SNavigationToolTime::OnEndSliderMovement(const int32 InNewValue)
+{
+	UndoTransaction.Reset();
+}
+
+}
+
+#undef LOCTEXT_NAMESPACE
diff --git a/Engine/Source/Editor/Sequencer/Private/Tools/NavigationTool/Widgets/Columns/SNavigationToolTime.h b/Engine/Source/Editor/Sequencer/Private/Tools/NavigationTool/Widgets/Columns/SNavigationToolTime.h
new file mode 100644
--- /dev/null
+++ b/Engine/Source/Editor/Sequencer/Private/Tools/NavigationTool/Widgets/Columns/SNavigationToolTime.h
@@ -0,0 +1,51 @@
+﻿// Copyright Epic Games, Inc. All Rights Reserved.
+
+#pragma once
+
+#include "ScopedTransaction.h"
+#include "Templates/UniquePtr.h"
+#include "Tools/NavigationTool/NavigationToolDefines.h"
+#include "Widgets/SCompoundWidget.h"
+
+namespace UE::Sequencer::NavigationTool
+{
+
+class INavigationToolView;
+class SNavigationToolTreeRow;
+
+class SNavigationToolTime : public SCompoundWidget
+{
+public:
+	SLATE_BEGIN_ARGS(SNavigationToolTime) {}
+	SLATE_END_ARGS()
+
+	virtual ~SNavigationToolTime() override {}
+
+	void Construct(const FArguments& InArgs
+		, const FNavigationToolItemPtr& InItem
+		, const TSharedRef<INavigationToolView>& InView
+		, const TSharedRef<SNavigationToolTreeRow>& InRowWidget);
+
+protected:
+	virtual FName GetStyleName() const;
+
+	virtual int32 GetFrameTimeValue() const = 0;
+
+	virtual void OnFrameTimeValueChanged(const int32 InNewValue);
+	virtual void OnFrameTimeValueCommitted(const int32 InNewValue, const ETextCommit::Type InCommitType) = 0;
+
+	virtual void OnBeginSliderMovement();
+	virtual void OnEndSliderMovement(const int32 InNewValue);
+
+	virtual FText GetTransactionText() const = 0;
+
+	FNavigationToolItemWeakPtr WeakItem;
+
+	TWeakPtr<INavigationToolView> WeakView;
+
+	TWeakPtr<SNavigationToolTreeRow> WeakRowWidget;
+
+	TUniquePtr<FScopedTransaction> UndoTransaction;
+};
+
+}
diff --git a/Engine/Source/Editor/Sequencer/Public/Tools/NavigationTool/Extensions/IInTimeExtension.h b/Engine/Source/Editor/Sequencer/Public/Tools/NavigationTool/Extensions/IInTimeExtension.h
new file mode 100644
--- /dev/null
+++ b/Engine/Source/Editor/Sequencer/Public/Tools/NavigationTool/Extensions/IInTimeExtension.h
@@ -0,0 +1,22 @@
+﻿// Copyright Epic Games, Inc. All Rights Reserved.
+
+#pragma once
+
+#include "Misc/FrameNumber.h"
+#include "Misc/Optional.h"
+#include "Tools/NavigationTool/NavigationToolItemType.h"
+
+namespace UE::Sequencer::NavigationTool
+{
+
+class SEQUENCER_API IInTimeExtension : public INavigationToolItemTypeCastable
+{
+public:
+	UE_NAVIGATIONTOOL_INHERITS(IInTimeExtension, INavigationToolItemTypeCastable)
+
+	virtual FFrameNumber GetInTime() const = 0;
+
+	virtual void SetInTime(const FFrameNumber& InTime) = 0;
+};
+
+}
diff --git a/Engine/Source/Editor/Sequencer/Public/Tools/NavigationTool/Extensions/IOutTimeExtension.h b/Engine/Source/Editor/Sequencer/Public/Tools/NavigationTool/Extensions/IOutTimeExtension.h
new file mode 100644
--- /dev/null
+++ b/Engine/Source/Editor/Sequencer/Public/Tools/NavigationTool/Extensions/IOutTimeExtension.h
@@ -0,0 +1,22 @@
+﻿// Copyright Epic Games, Inc. All Rights Reserved.
+
+#pragma once
+
+#include "Misc/FrameNumber.h"
+#include "Misc/Optional.h"
+#include "Tools/NavigationTool/NavigationToolItemType.h"
+
+namespace UE::Sequencer::NavigationTool
+{
+
+class SEQUENCER_API IOutTimeExtension : public INavigationToolItemTypeCastable
+{
+public:
+	UE_NAVIGATIONTOOL_INHERITS(IOutTimeExtension, INavigationToolItemTypeCastable)
+
+	virtual FFrameNumber GetOutTime() const = 0;
+
+	virtual void SetOutTime(const FFrameNumber& InTime) = 0;
+};
+
+}
diff --git a/Engine/Source/Editor/Sequencer/Public/Tools/NavigationTool/Filters/Filters/NavigationToolBuiltInFilterParams.h b/Engine/Source/Editor/Sequencer/Public/Tools/NavigationTool/Filters/Filters/NavigationToolBuiltInFilterParams.h
--- a/Engine/Source/Editor/Sequencer/Public/Tools/NavigationTool/Filters/Filters/NavigationToolBuiltInFilterParams.h
+++ b/Engine/Source/Editor/Sequencer/Public/Tools/NavigationTool/Filters/Filters/NavigationToolBuiltInFilterParams.h
@@ -16,6 +16,7 @@ struct FNavigationToolBuiltInFilterParams
 	SEQUENCER_API static FNavigationToolBuiltInFilterParams CreateSequenceFilter();
 	SEQUENCER_API static FNavigationToolBuiltInFilterParams CreateTrackFilter();
 	SEQUENCER_API static FNavigationToolBuiltInFilterParams CreateBindingFilter();
+	SEQUENCER_API static FNavigationToolBuiltInFilterParams CreateMarkerFilter();
 
 	FNavigationToolBuiltInFilterParams() = default;
 
diff --git a/Engine/Source/Editor/Sequencer/Public/Tools/NavigationTool/Items/INavigationToolItem.h b/Engine/Source/Editor/Sequencer/Public/Tools/NavigationTool/Items/INavigationToolItem.h
--- a/Engine/Source/Editor/Sequencer/Public/Tools/NavigationTool/Items/INavigationToolItem.h
+++ b/Engine/Source/Editor/Sequencer/Public/Tools/NavigationTool/Items/INavigationToolItem.h
@@ -16,6 +16,7 @@ enum class EItemDropZone;
 enum class ENavigationToolItemViewMode : uint8;
 struct FLinearColor;
 struct FSlateBrush;
+struct FSlateColor;
 struct FSlateIcon;
 
 namespace UE::Sequencer::NavigationTool
@@ -373,22 +374,6 @@ public:
 		return OutChildren;
 	};
 
-	/*template<typename InItemType>
-	TArray<InItemType*> GetDescendantsOfType(const bool bInIncludeThis = false) const
-	{
-		TArray<InItemType*> OutDescendants;
-
-		for (InItemType* const Child : GetChildrenOfType<InItemType>(bInIncludeThis))
-		{
-			OutDescendants.Add(Child);
-
-			const TArray<InItemType*> ThisChildDescendants = Child->GetDescendantsOfType<InItemType>();
-			OutDescendants.Append(ThisChildDescendants);
-		}
-
-		return OutDescendants;
-	}*/
-	
 	/**
 	 * Finds the first ancestor view model that implements the given type
 	 */
@@ -432,12 +417,17 @@ public:
 	/** Gets the Class/Type of this Item (e.g. for Items that represent UObjects, it will be the UObject class) */
 	virtual FText GetClassName() const = 0;
 
+	/** Gets the slate icon for the item */
 	virtual FSlateIcon GetIcon() const = 0;
 
 	virtual const FSlateBrush* GetDefaultIconBrush() const { return nullptr; }
 
 	virtual const FSlateBrush* GetIconBrush() const = 0;
 
+	/** Gets the color for the item icon */
+	virtual FSlateColor GetIconColor() const = 0;
+
+	/** Gets the tooltip text for the item icon */
 	virtual FText GetIconTooltipText() const = 0;
 
 	/** Gets the View Modes that this Item Supports */
diff --git a/Engine/Source/Editor/Sequencer/Public/Tools/NavigationTool/Items/NavigationToolItem.h b/Engine/Source/Editor/Sequencer/Public/Tools/NavigationTool/Items/NavigationToolItem.h
--- a/Engine/Source/Editor/Sequencer/Public/Tools/NavigationTool/Items/NavigationToolItem.h
+++ b/Engine/Source/Editor/Sequencer/Public/Tools/NavigationTool/Items/NavigationToolItem.h
@@ -56,6 +56,7 @@ public:
 	virtual ENavigationToolItemViewMode GetSupportedViewModes(const INavigationToolView& InToolView) const override;
 	virtual FNavigationToolItemId GetItemId() const override final;
 	virtual const FSlateBrush* GetIconBrush() const override;
+	virtual FSlateColor GetIconColor() const override;
 	virtual TSharedRef<SWidget> GenerateLabelWidget(const TSharedRef<SNavigationToolTreeRow>& InRow) override;
 	virtual bool ShowVisibility() const override { return false; }
 	virtual bool CanReceiveParentVisibilityPropagation() const override { return false; }
diff --git a/Engine/Source/Editor/Sequencer/Public/Tools/NavigationTool/Items/NavigationToolMarker.h b/Engine/Source/Editor/Sequencer/Public/Tools/NavigationTool/Items/NavigationToolMarker.h
new file mode 100644
--- /dev/null
+++ b/Engine/Source/Editor/Sequencer/Public/Tools/NavigationTool/Items/NavigationToolMarker.h
@@ -0,0 +1,87 @@
+// Copyright Epic Games, Inc. All Rights Reserved.
+
+#pragma once
+
+#include "NavigationToolItem.h"
+#include "Tools/NavigationTool/Extensions/IInTimeExtension.h"
+#include "Tools/NavigationTool/Extensions/IRenameableExtension.h"
+
+class UMovieScene;
+struct FMovieSceneMarkedFrame;
+
+namespace UE::Sequencer::NavigationTool
+{
+
+class FNavigationToolSequence;
+
+/**
+ * Navigation Tool Item representing a sequence marker
+ */
+class SEQUENCER_API FNavigationToolMarker
+	: public FNavigationToolItem
+	, public IRenameableExtension
+	, public IInTimeExtension
+{
+public:
+	UE_NAVIGATIONTOOL_INHERITS_WITH_SUPER(FNavigationToolMarker
+		, FNavigationToolItem
+		, IRenameableExtension
+		, IInTimeExtension);
+
+	FNavigationToolMarker(INavigationTool& InTool
+		, const FNavigationToolItemPtr& InParentItem
+		, const TSharedPtr<FNavigationToolSequence>& InParentSequenceItem
+		, const int32 InMarkedFrameIndex);
+
+	//~ Begin INavigationToolItem
+
+	virtual bool IsItemValid() const override;
+	virtual bool IsAllowedInTool() const override;
+
+	virtual ENavigationToolItemViewMode GetSupportedViewModes(const INavigationToolView& InToolView) const override;
+
+	virtual FText GetDisplayName() const override;
+	virtual FText GetClassName() const override;
+	virtual FSlateIcon GetIcon() const override;
+	virtual const FSlateBrush* GetIconBrush() const override;
+	virtual FSlateColor GetIconColor() const override;
+	virtual FText GetIconTooltipText() const override;
+
+	virtual bool IsSelected(const FNavigationToolScopedSelection& InSelection) const override;
+	virtual void Select(FNavigationToolScopedSelection& InSelection) const override;
+	virtual void OnSelect() override;
+	virtual void OnDoubleClick() override;
+
+	virtual bool CanDelete() const override;
+	virtual bool Delete() override;
+
+	//~ End INavigationToolItem
+
+	//~ Begin IRenameableExtension
+	virtual bool CanRename() const override;
+	virtual bool Rename(const FString& InName) override;
+	//~ End IRenameableExtension
+
+	//~ Begin IInTimeExtension
+	virtual FFrameNumber GetInTime() const override;
+	virtual void SetInTime(const FFrameNumber& InTime) override;
+	//~ End IInTimeExtension
+
+	int32 GetMarkedFrameIndex() const;
+
+	FMovieSceneMarkedFrame* GetMarkedFrame() const;
+
+protected:
+	//~ Begin INavigationToolItem
+	virtual FNavigationToolItemId CalculateItemId() const override;
+	//~ End INavigationToolItem
+
+	UMovieSceneSequence* GetParentSequence() const;
+	UMovieScene* GetParentMovieScene() const;
+
+	TWeakPtr<FNavigationToolSequence> WeakParentSequenceItem;
+
+	int32 MarkedFrameIndex = INDEX_NONE;
+};
+
+}
diff --git a/Engine/Source/Editor/Sequencer/Public/Tools/NavigationTool/Items/NavigationToolSequence.h b/Engine/Source/Editor/Sequencer/Public/Tools/NavigationTool/Items/NavigationToolSequence.h
--- a/Engine/Source/Editor/Sequencer/Public/Tools/NavigationTool/Items/NavigationToolSequence.h
+++ b/Engine/Source/Editor/Sequencer/Public/Tools/NavigationTool/Items/NavigationToolSequence.h
@@ -3,7 +3,9 @@
 #pragma once
 
 #include "NavigationToolItem.h"
+#include "Tools/NavigationTool/Extensions/IInTimeExtension.h"
 #include "Tools/NavigationTool/Extensions/IMarkerVisibilityExtension.h"
+#include "Tools/NavigationTool/Extensions/IOutTimeExtension.h"
 #include "Tools/NavigationTool/Extensions/IPlayheadExtension.h"
 #include "Tools/NavigationTool/Extensions/IRenameableExtension.h"
 #include "Tools/NavigationTool/Extensions/IRevisionControlExtension.h"
@@ -34,6 +36,8 @@ class FNavigationToolSequence
 	, public IMarkerVisibilityExtension
 	, public IRevisionControlExtension
 	, public IPlayheadExtension
+	, public IInTimeExtension
+	, public IOutTimeExtension
 {
 public:
 	UE_NAVIGATIONTOOL_INHERITS_WITH_SUPER(FNavigationToolSequence
@@ -43,7 +47,9 @@ public:
 		, ISequenceLockableExtension
 		, IMarkerVisibilityExtension
 		, IRevisionControlExtension
-		, IPlayheadExtension);
+		, IPlayheadExtension
+		, IInTimeExtension
+		, IOutTimeExtension);
 
 	SEQUENCER_API FNavigationToolSequence(INavigationTool& InTool
 		, const FNavigationToolItemPtr& InParentItem
@@ -116,6 +122,16 @@ public:
 	SEQUENCER_API virtual EItemContainsPlayhead ContainsPlayhead() const override;
 	//~ End IPlayheadExtension
 
+	//~ Begin IInTimeExtension
+	SEQUENCER_API virtual FFrameNumber GetInTime() const override;
+	SEQUENCER_API virtual void SetInTime(const FFrameNumber& InTime) override;
+	//~ End IInTimeExtension
+
+	//~ Begin IOutTimeExtension
+	SEQUENCER_API virtual FFrameNumber GetOutTime() const override;
+	SEQUENCER_API virtual void SetOutTime(const FFrameNumber& InTime) override;
+	//~ End IOutTimeExtension
+
 	SEQUENCER_API UMovieSceneSequence* GetSequence() const;
 
 	SEQUENCER_API UMovieSceneSubSection* GetSubSection() const;
diff --git a/Engine/Source/Editor/Sequencer/Public/Tools/NavigationTool/NavigationToolScopedSelection.h b/Engine/Source/Editor/Sequencer/Public/Tools/NavigationTool/NavigationToolScopedSelection.h
--- a/Engine/Source/Editor/Sequencer/Public/Tools/NavigationTool/NavigationToolScopedSelection.h
+++ b/Engine/Source/Editor/Sequencer/Public/Tools/NavigationTool/NavigationToolScopedSelection.h
@@ -4,6 +4,7 @@
 
 #include "Containers/Array.h"
 #include "Containers/Set.h"
+#include "MovieSceneSequence.h"
 
 class ISequencer;
 class UMovieSceneSection;
@@ -33,11 +34,15 @@ public:
 	SEQUENCER_API void Select(const FGuid& InObjectGuid);
 	SEQUENCER_API void Select(UMovieSceneSection* const InSection);
 	SEQUENCER_API void Select(UMovieSceneTrack* const InTrack);
+	SEQUENCER_API void Select(UMovieSceneSequence* const InSequence
+		, const int32 InMarkedFrameIndex);
 
 	SEQUENCER_API bool IsSelected(const UObject* const InObject) const;
 	SEQUENCER_API bool IsSelected(const FGuid& InObjectGuid) const;
 	SEQUENCER_API bool IsSelected(UMovieSceneSection* const InSection) const;
 	SEQUENCER_API bool IsSelected(UMovieSceneTrack* const InTrack) const;
+	SEQUENCER_API bool IsSelected(UMovieSceneSequence* const InSequence
+		, const int32 InMarkedFrameIndex) const;
 
 	SEQUENCER_API ISequencer& GetSequencer() const;
 
@@ -52,6 +57,7 @@ private:
 	TArray<FGuid> SelectedObjectGuids;
 	TArray<UMovieSceneSection*> SelectedSections;
 	TArray<UMovieSceneTrack*> SelectedTracks;
+	TMap<UMovieSceneSequence*, TSet<int32>> SelectedMarkedFrames;
 
 	ENavigationToolScopedSelectionPurpose Purpose;
 };
