CurveEditor: Remove ensure that can trigger when using Actions -> Import on the USD Stage Editor.

diff --git a/Engine/Source/Editor/Sequencer/Private/MVVM/CurveEditorExtension.cpp b/Engine/Source/Editor/Sequencer/Private/MVVM/CurveEditorExtension.cpp
--- a/Engine/Source/Editor/Sequencer/Private/MVVM/CurveEditorExtension.cpp
+++ b/Engine/Source/Editor/Sequencer/Private/MVVM/CurveEditorExtension.cpp
@@ -646,6 +646,12 @@ void FCurveEditorExtension::RequestSyncSelection()
 			This->SyncSelection();
 		}
 
+		// Some external tools may reset GUndo in the meantime, e.g. USD Stage Editor during Actions->Import
+		if (!GUndo)
+		{
+			SelectionChange->Cancel();
+		}
+
 		SelectionChange.Reset();
 	});
 }
