[Sequencer] Navigation Tool
- Use binding reference name instead of resolving bound object
- Show all child tracks of passed BindingName text expression filter

diff --git a/Engine/Source/Editor/Sequencer/Private/Filters/TextExpressions/SequencerTextFilterExpression_BindingName.cpp b/Engine/Source/Editor/Sequencer/Private/Filters/TextExpressions/SequencerTextFilterExpression_BindingName.cpp
--- a/Engine/Source/Editor/Sequencer/Private/Filters/TextExpressions/SequencerTextFilterExpression_BindingName.cpp
+++ b/Engine/Source/Editor/Sequencer/Private/Filters/TextExpressions/SequencerTextFilterExpression_BindingName.cpp
@@ -1,7 +1,8 @@
 ﻿// Copyright Epic Games, Inc. All Rights Reserved.
 
 #include "SequencerTextFilterExpression_BindingName.h"
-#include "GameFramework/Actor.h"
+#include "MVVM/ViewModelPtr.h"
+#include "MVVM/ViewModels/ObjectBindingModel.h"
 
 using namespace UE::Sequencer;
 
@@ -37,23 +38,15 @@ bool FSequencerTextFilterExpression_BindingName::TestComplexExpression(const FNa
 		return true;
 	}
 
-	UObject* const BoundObject = GetBoundObject();
-	if (IsValid(BoundObject))
+	const TViewModelPtr<FObjectBindingModel> ObjectBindingModel = FilterItem->FindAncestorOfType<FObjectBindingModel>(true);
+	if (!ObjectBindingModel.IsValid())
 	{
-		const AActor* const Actor = Cast<AActor>(BoundObject);
-		if (IsValid(Actor))
-		{
-			const FString ActorLabel = Actor->GetActorLabel();
-			const bool bPassed = TextFilterUtils::TestComplexExpression(ActorLabel, InValue, InComparisonOperation, InTextComparisonMode);
-			return bPassed;
-		}
-
-		const FString ObjectName = BoundObject->GetName();
-		const bool bPassed = TextFilterUtils::TestComplexExpression(ObjectName, InValue, InComparisonOperation, InTextComparisonMode);
-		return bPassed;
+		return false;
 	}
 
-	return false;
+	const FString BindingLabel = ObjectBindingModel->GetLabel().ToString();
+	const bool bPassed = TextFilterUtils::TestComplexExpression(BindingLabel, InValue, InComparisonOperation, InTextComparisonMode);
+	return bPassed;
 }
 
 #undef LOCTEXT_NAMESPACE
