[Sequencer] Navigation Tool: Use FSequencerNumericTypeInterface for input boxes

diff --git a/Engine/Source/Editor/Sequencer/Private/Tools/NavigationTool/Items/NavigationToolMarker.cpp b/Engine/Source/Editor/Sequencer/Private/Tools/NavigationTool/Items/NavigationToolMarker.cpp
--- a/Engine/Source/Editor/Sequencer/Private/Tools/NavigationTool/Items/NavigationToolMarker.cpp
+++ b/Engine/Source/Editor/Sequencer/Private/Tools/NavigationTool/Items/NavigationToolMarker.cpp
@@ -240,7 +240,7 @@ FFrameNumber FNavigationToolMarker::GetInTime() const
 	const TArray<FMovieSceneMarkedFrame>& MarkedFrames = MovieScene->GetMarkedFrames();
 	FMovieSceneMarkedFrame& MarkedFrame = const_cast<FMovieSceneMarkedFrame&>(MarkedFrames[MarkedFrameIndex]);
 
-	return ConvertToDisplayRateTime(*Sequence, MarkedFrame.FrameNumber).FrameNumber;
+	return MarkedFrame.FrameNumber;
 }
 
 void FNavigationToolMarker::SetInTime(const FFrameNumber& InTime)
@@ -265,10 +265,15 @@ void FNavigationToolMarker::SetInTime(const FFrameNumber& InTime)
 	const TArray<FMovieSceneMarkedFrame>& MarkedFrames = MovieScene->GetMarkedFrames();
 	FMovieSceneMarkedFrame& MarkedFrame = const_cast<FMovieSceneMarkedFrame&>(MarkedFrames[MarkedFrameIndex]);
 
-	const FScopedTransaction Transaction(LOCTEXT("SetTimeTransaction", "Set Marked Frame Time"));
+	if (InTime == MarkedFrame.FrameNumber)
+	{
+		return;
+	}
+
+	const FScopedTransaction Transaction(LOCTEXT("SetMarkedFrameTimeTransaction", "Set Marked Frame Time"));
 
 	MovieScene->Modify();
-	MarkedFrame.FrameNumber = ConvertToTickResolutionTime(*Sequence, InTime).FrameNumber;
+	MarkedFrame.FrameNumber = InTime;
 }
 
 int32 FNavigationToolMarker::GetMarkedFrameIndex() const
diff --git a/Engine/Source/Editor/Sequencer/Private/Tools/NavigationTool/Items/NavigationToolSequence.cpp b/Engine/Source/Editor/Sequencer/Private/Tools/NavigationTool/Items/NavigationToolSequence.cpp
--- a/Engine/Source/Editor/Sequencer/Private/Tools/NavigationTool/Items/NavigationToolSequence.cpp
+++ b/Engine/Source/Editor/Sequencer/Private/Tools/NavigationTool/Items/NavigationToolSequence.cpp
@@ -568,8 +568,7 @@ FFrameNumber FNavigationToolSequence::GetInTime() const
 	}
 
 	const FFrameTime LowerBoundValue = SubSection->SectionRange.Value.GetLowerBoundValue();
-
-	return ConvertToDisplayRateTime(*Sequence, LowerBoundValue).FrameNumber;
+	return LowerBoundValue.FrameNumber;
 }
 
 void FNavigationToolSequence::SetInTime(const FFrameNumber& InTime)
@@ -586,15 +585,13 @@ void FNavigationToolSequence::SetInTime(const FFrameNumber& InTime)
 		return;
 	}
 
-	const FFrameNumber NewFrameTime = ConvertToTickResolutionTime(*Sequence, InTime).FrameNumber;
-
-	if (SubSection->SectionRange.Value.GetLowerBoundValue() == NewFrameTime)
+	if (SubSection->SectionRange.Value.GetLowerBoundValue() == InTime)
 	{
 		return;
 	}
 
 	SubSection->Modify();
-	SubSection->SectionRange.Value.SetLowerBoundValue(NewFrameTime);
+	SubSection->SectionRange.Value.SetLowerBoundValue(InTime);
 }
 
 FFrameNumber FNavigationToolSequence::GetOutTime() const
@@ -612,8 +609,7 @@ FFrameNumber FNavigationToolSequence::GetOutTime() const
 	}
 
 	const FFrameTime UpperBoundValue = SubSection->SectionRange.Value.GetUpperBoundValue();
-
-	return ConvertToDisplayRateTime(*Sequence, UpperBoundValue).FrameNumber;
+	return UpperBoundValue.FrameNumber;
 }
 
 void FNavigationToolSequence::SetOutTime(const FFrameNumber& InTime)
@@ -630,15 +626,13 @@ void FNavigationToolSequence::SetOutTime(const FFrameNumber& InTime)
 		return;
 	}
 
-	const FFrameNumber NewFrameTime = ConvertToTickResolutionTime(*Sequence, InTime).FrameNumber;
-
-	if (SubSection->SectionRange.Value.GetUpperBoundValue() == NewFrameTime)
+	if (SubSection->SectionRange.Value.GetUpperBoundValue() == InTime)
 	{
 		return;
 	}
 
 	SubSection->Modify();
-	SubSection->SectionRange.Value.SetUpperBoundValue(NewFrameTime);
+	SubSection->SectionRange.Value.SetUpperBoundValue(InTime);
 }
 
 UMovieSceneSequence* FNavigationToolSequence::GetSequence() const
diff --git a/Engine/Source/Editor/Sequencer/Private/Tools/NavigationTool/Widgets/Columns/SNavigationToolEndFrameOffset.cpp b/Engine/Source/Editor/Sequencer/Private/Tools/NavigationTool/Widgets/Columns/SNavigationToolEndFrameOffset.cpp
--- a/Engine/Source/Editor/Sequencer/Private/Tools/NavigationTool/Widgets/Columns/SNavigationToolEndFrameOffset.cpp
+++ b/Engine/Source/Editor/Sequencer/Private/Tools/NavigationTool/Widgets/Columns/SNavigationToolEndFrameOffset.cpp
@@ -22,117 +22,60 @@ void SNavigationToolEndFrameOffset::Construct(const FArguments& InArgs
 	WeakView = InView;
 	WeakRowWidget = InRowWidget;
 
-	auto EmptyContent = [this]()
+	if (!WeakItem.IsValid()
+		|| !WeakView.IsValid()
+		|| !WeakRowWidget.IsValid())
 	{
-		ChildSlot
-		[
-			SNullWidget::NullWidget
-		];
-	};
-
-	if (!WeakItem.IsValid() || !WeakView.IsValid())
-	{
-		EmptyContent();
 		return;
 	}
 
 	const FNavigationToolSequence* const SequenceItem = WeakItem.Pin()->CastTo<FNavigationToolSequence>();
 	if (!SequenceItem || !IsValid(SequenceItem->GetSubSection()))
 	{
-		EmptyContent();
 		return;
 	}
 
-	ChildSlot
-		[
-			SNew(SSpinBox<int32>)
-			.Justification(ETextJustify::Center)
-			.Style(&FAppStyle::Get(), TEXT("Sequencer.NavigationTool.SpinBox"))
-			.Value(this, &SNavigationToolEndFrameOffset::GetValue)
-			.OnValueChanged(this, &SNavigationToolEndFrameOffset::OnValueChanged)
-			.OnValueCommitted(this, &SNavigationToolEndFrameOffset::OnValueCommitted)
-			.OnBeginSliderMovement(this, &SNavigationToolEndFrameOffset::OnBeginSliderMovement)
-			.OnEndSliderMovement(this, &SNavigationToolEndFrameOffset::OnEndSliderMovement)
-		];
+	SNavigationToolTime::Construct(SNavigationToolTime::FArguments(), InItem, InView, InRowWidget);
 }
 
-int32 SNavigationToolEndFrameOffset::GetValue() const
+FName SNavigationToolEndFrameOffset::GetStyleName() const
 {
-	const TSharedPtr<INavigationToolView> ToolView = WeakView.Pin();
-	if (!ToolView.IsValid())
-	{
-		return 0;
-	}
-
-	const TSharedPtr<ISequencer> Sequencer = ToolView->GetSequencer();
-	if (!Sequencer.IsValid())
-	{
-		return 0;
-	}
+	return TEXT("Sequencer.NavigationTool.SpinBox");
+}
 
+double SNavigationToolEndFrameOffset::GetFrameTimeValue() const
+{
 	UMovieSceneSubSection* const SubSection = Cast<UMovieSceneSubSection>(GetSequenceItemSubSection(WeakItem.Pin()));
-	if (!IsValid(SubSection))
+	if (IsValid(SubSection))
 	{
-		return 0;
+		return SubSection->Parameters.EndFrameOffset.Value;
 	}
-
-	const FFrameNumber Offset = ConvertFrameTime(SubSection->Parameters.EndFrameOffset.Value
-		, Sequencer->GetFocusedTickResolution()
-		, Sequencer->GetFocusedDisplayRate()).FrameNumber;
-	return Offset.Value;
+	return 0.;
 }
 
-void SNavigationToolEndFrameOffset::OnValueChanged(const int32 InNewValue)
+void SNavigationToolEndFrameOffset::OnFrameTimeValueCommitted(const double InNewValue, const ETextCommit::Type InCommitType)
 {
-	const TSharedPtr<INavigationToolView> ToolView = WeakView.Pin();
-	if (!ToolView.IsValid())
-	{
-		return;
-	}
-
-	const TSharedPtr<ISequencer> Sequencer = ToolView->GetSequencer();
-	if (!Sequencer.IsValid())
+	UMovieSceneSubSection* const SubSection = Cast<UMovieSceneSubSection>(GetSequenceItemSubSection(WeakItem.Pin()));
+	if (!IsValid(SubSection))
 	{
 		return;
 	}
 
-	UMovieSceneSubSection* const SubSection = Cast<UMovieSceneSubSection>(GetSequenceItemSubSection(WeakItem.Pin()));
-	if (!IsValid(SubSection))
+	if (InNewValue == SubSection->Parameters.EndFrameOffset.Value)
 	{
 		return;
 	}
 
-	const FFrameNumber Offset = ConvertFrameTime(SubSection->Parameters.EndFrameOffset.Value
-		, Sequencer->GetFocusedDisplayRate()
-		, Sequencer->GetFocusedTickResolution()).FrameNumber;
-
-	const FScopedTransaction Transaction(GetTransactionText(), !UndoTransaction.IsValid());
+	const bool bShouldTransact = !UndoTransaction.IsValid() && (InCommitType == ETextCommit::OnEnter);
+	const FScopedTransaction Transaction(GetTransactionText(), bShouldTransact);
 
 	SubSection->Modify();
-	SubSection->Parameters.EndFrameOffset = Offset;
-}
-
-void SNavigationToolEndFrameOffset::OnValueCommitted(const int32 InNewValue, const ETextCommit::Type InCommitType)
-{
-	OnValueChanged(InNewValue);
-}
-
-void SNavigationToolEndFrameOffset::OnBeginSliderMovement()
-{
-	if (!UndoTransaction.IsValid())
-	{
-		UndoTransaction = MakeUnique<FScopedTransaction>(GetTransactionText());
-	}
-}
-
-void SNavigationToolEndFrameOffset::OnEndSliderMovement(const int32 InNewValue)
-{
-	UndoTransaction.Reset();
+	SubSection->Parameters.EndFrameOffset = static_cast<int32>(InNewValue);
 }
 
 FText SNavigationToolEndFrameOffset::GetTransactionText() const
 {
-	return LOCTEXT("SetSequenceEndFrameOffsetTransaction", "Set Sequence End Frame Offset");
+	return LOCTEXT("SetEndFrameOffsetTransaction", "Set End Frame Offset");
 }
 
 #undef LOCTEXT_NAMESPACE
diff --git a/Engine/Source/Editor/Sequencer/Private/Tools/NavigationTool/Widgets/Columns/SNavigationToolEndFrameOffset.h b/Engine/Source/Editor/Sequencer/Private/Tools/NavigationTool/Widgets/Columns/SNavigationToolEndFrameOffset.h
--- a/Engine/Source/Editor/Sequencer/Private/Tools/NavigationTool/Widgets/Columns/SNavigationToolEndFrameOffset.h
+++ b/Engine/Source/Editor/Sequencer/Private/Tools/NavigationTool/Widgets/Columns/SNavigationToolEndFrameOffset.h
@@ -2,10 +2,8 @@
 
 #pragma once
 
-#include "ScopedTransaction.h"
-#include "Templates/UniquePtr.h"
+#include "SNavigationToolTime.h"
 #include "Tools/NavigationTool/NavigationToolDefines.h"
-#include "Widgets/SCompoundWidget.h"
 
 namespace UE::Sequencer::NavigationTool
 {
@@ -13,7 +11,7 @@ namespace UE::Sequencer::NavigationTool
 class INavigationToolView;
 class SNavigationToolTreeRow;
 
-class SNavigationToolEndFrameOffset : public SCompoundWidget
+class SNavigationToolEndFrameOffset : public SNavigationToolTime
 {
 public:
 	SLATE_BEGIN_ARGS(SNavigationToolEndFrameOffset) {}
@@ -27,21 +25,13 @@ public:
 		, const TSharedRef<SNavigationToolTreeRow>& InRowWidget);
 
 private:
-	int32 GetValue() const;
-	void OnValueChanged(const int32 InNewValue);
-	void OnValueCommitted(const int32 InNewValue, const ETextCommit::Type InCommitType);
-	void OnBeginSliderMovement();
-	void OnEndSliderMovement(const int32 InNewValue);
+	virtual FName GetStyleName() const override;
 
-	FText GetTransactionText() const;
+	virtual double GetFrameTimeValue() const override;
 
-	FNavigationToolItemWeakPtr WeakItem;
+	virtual void OnFrameTimeValueCommitted(const double InNewValue, const ETextCommit::Type InCommitType) override;
 
-	TWeakPtr<INavigationToolView> WeakView;
-
-	TWeakPtr<SNavigationToolTreeRow> WeakRowWidget;
-
-	TUniquePtr<FScopedTransaction> UndoTransaction;
+	virtual FText GetTransactionText() const override;
 };
 
 }
diff --git a/Engine/Source/Editor/Sequencer/Private/Tools/NavigationTool/Widgets/Columns/SNavigationToolInTime.cpp b/Engine/Source/Editor/Sequencer/Private/Tools/NavigationTool/Widgets/Columns/SNavigationToolInTime.cpp
--- a/Engine/Source/Editor/Sequencer/Private/Tools/NavigationTool/Widgets/Columns/SNavigationToolInTime.cpp
+++ b/Engine/Source/Editor/Sequencer/Private/Tools/NavigationTool/Widgets/Columns/SNavigationToolInTime.cpp
@@ -23,6 +23,7 @@ void SNavigationToolInTime::Construct(const FArguments& InArgs
 
 	if (!WeakItem.IsValid()
 		|| !WeakView.IsValid()
+		|| !WeakRowWidget.IsValid()
 		|| !WeakItem.Pin()->IsA<IInTimeExtension>())
 	{
 		return;
@@ -41,7 +42,7 @@ FName SNavigationToolInTime::GetStyleName() const
 	return TEXT("Sequencer.NavigationTool.InTime.SpinBox");
 }
 
-int32 SNavigationToolInTime::GetFrameTimeValue() const
+double SNavigationToolInTime::GetFrameTimeValue() const
 {
 	if (const IInTimeExtension* const InTimeExtension = GetInTimeExtension())
 	{
@@ -50,7 +51,7 @@ int32 SNavigationToolInTime::GetFrameTimeValue() const
 	return 0;
 }
 
-void SNavigationToolInTime::OnFrameTimeValueCommitted(const int32 InNewValue, const ETextCommit::Type InCommitType)
+void SNavigationToolInTime::OnFrameTimeValueCommitted(const double InNewValue, const ETextCommit::Type InCommitType)
 {
 	IInTimeExtension* const InTimeExtension = GetInTimeExtension();
 	if (!InTimeExtension)
@@ -61,7 +62,7 @@ void SNavigationToolInTime::OnFrameTimeValueCommitted(const int32 InNewValue, co
 	const bool bShouldTransact = !UndoTransaction.IsValid() && (InCommitType == ETextCommit::OnEnter);
 	const FScopedTransaction Transaction(GetTransactionText(), bShouldTransact);
 
-	return InTimeExtension->SetInTime(InNewValue);
+	return InTimeExtension->SetInTime(static_cast<int32>(InNewValue));
 }
 
 FText SNavigationToolInTime::GetTransactionText() const
diff --git a/Engine/Source/Editor/Sequencer/Private/Tools/NavigationTool/Widgets/Columns/SNavigationToolInTime.h b/Engine/Source/Editor/Sequencer/Private/Tools/NavigationTool/Widgets/Columns/SNavigationToolInTime.h
--- a/Engine/Source/Editor/Sequencer/Private/Tools/NavigationTool/Widgets/Columns/SNavigationToolInTime.h
+++ b/Engine/Source/Editor/Sequencer/Private/Tools/NavigationTool/Widgets/Columns/SNavigationToolInTime.h
@@ -3,7 +3,6 @@
 #pragma once
 
 #include "SNavigationToolTime.h"
-#include "Templates/UniquePtr.h"
 #include "Tools/NavigationTool/NavigationToolDefines.h"
 
 namespace UE::Sequencer::NavigationTool
@@ -31,9 +30,9 @@ private:
 
 	virtual FName GetStyleName() const override;
 
-	virtual int32 GetFrameTimeValue() const override;
+	virtual double GetFrameTimeValue() const override;
 
-	virtual void OnFrameTimeValueCommitted(const int32 InNewValue, const ETextCommit::Type InCommitType) override;
+	virtual void OnFrameTimeValueCommitted(const double InNewValue, const ETextCommit::Type InCommitType) override;
 
 	virtual FText GetTransactionText() const override;
 };
diff --git a/Engine/Source/Editor/Sequencer/Private/Tools/NavigationTool/Widgets/Columns/SNavigationToolLength.cpp b/Engine/Source/Editor/Sequencer/Private/Tools/NavigationTool/Widgets/Columns/SNavigationToolLength.cpp
--- a/Engine/Source/Editor/Sequencer/Private/Tools/NavigationTool/Widgets/Columns/SNavigationToolLength.cpp
+++ b/Engine/Source/Editor/Sequencer/Private/Tools/NavigationTool/Widgets/Columns/SNavigationToolLength.cpp
@@ -7,7 +7,6 @@
 #include "Tools/NavigationTool/Items/NavigationToolSequence.h"
 #include "Tools/NavigationTool/NavigationToolView.h"
 #include "Tools/NavigationTool/Widgets/SNavigationToolTreeRow.h"
-#include "Widgets/Input/SEditableTextBox.h"
 
 #define LOCTEXT_NAMESPACE "SNavigationToolLength"
 
@@ -23,66 +22,49 @@ void SNavigationToolLength::Construct(const FArguments& InArgs
 	WeakView = InView;
 	WeakRowWidget = InRowWidget;
 
-	auto EmptyContent = [this]()
+	if (!WeakItem.IsValid()
+		|| !WeakView.IsValid()
+		|| !WeakRowWidget.IsValid())
 	{
-		ChildSlot
-		[
-			SNullWidget::NullWidget
-		];
-	};
-
-	if (!WeakItem.IsValid() || !WeakView.IsValid())
-	{
-		EmptyContent();
 		return;
 	}
 
 	const FNavigationToolSequence* const SequenceItem = WeakItem.Pin()->CastTo<FNavigationToolSequence>();
 	if (!SequenceItem || !IsValid(SequenceItem->GetSubSection()))
 	{
-		EmptyContent();
 		return;
 	}
 
-	ChildSlot
-		[
-			SNew(SEditableTextBox)
-			.Style(FAppStyle::Get(), TEXT("Sequencer.NavigationTool.NonEditableTextBox"))
-			.Justification(ETextJustify::Center)
-			.IsReadOnly(true)
-			.Text_Lambda([this]()
-				{
-					return FText::FromString(FString::FromInt(GetLengthValue()));
-				})
-		];
+	SNavigationToolTime::Construct(SNavigationToolTime::FArguments(), InItem, InView, InRowWidget);
 }
 
-int32 SNavigationToolLength::GetLengthValue() const
+FName SNavigationToolLength::GetStyleName() const
 {
-	const TSharedPtr<INavigationToolView> ToolView = WeakView.Pin();
-	if (!ToolView.IsValid())
-	{
-		return 0;
-	}
-
-	const TSharedPtr<ISequencer> Sequencer = ToolView->GetSequencer();
-	if (!Sequencer.IsValid())
-	{
-		return 0;
-	}
+	return TEXT("Sequencer.NavigationTool.SpinBox");
+}
 
+double SNavigationToolLength::GetFrameTimeValue() const
+{
 	UMovieSceneSubSection* const SubSection = Cast<UMovieSceneSubSection>(GetSequenceItemSubSection(WeakItem.Pin()));
-	if (!IsValid(SubSection))
+	if (IsValid(SubSection))
 	{
-		return 0;
+		return SubSection->GetRange().Size<FFrameNumber>().Value;
 	}
+	return 0.;
+}
 
-	const FFrameNumber FrameLength = SubSection->GetRange().Size<FFrameNumber>();
-	const FFrameNumber ConvertedFrameLength = ConvertFrameTime(FrameLength
-		, Sequencer->GetFocusedTickResolution()
-		, Sequencer->GetFocusedDisplayRate()).FrameNumber;
+FText SNavigationToolLength::GetTransactionText() const
+{
+	return LOCTEXT("SetLengthTransaction", "Set Length");
+}
 
-	return ConvertedFrameLength.Value;
+TSharedPtr<INumericTypeInterface<double>> SNavigationToolLength::GetNumericTypeInterface() const
+{
+	if (const TSharedPtr<ISequencer> Sequencer = GetSequencer())
+	{
+		return Sequencer->GetNumericTypeInterface(ENumericIntent::Duration);
+	}
+	return nullptr;
 }
 
 }
diff --git a/Engine/Source/Editor/Sequencer/Private/Tools/NavigationTool/Widgets/Columns/SNavigationToolLength.h b/Engine/Source/Editor/Sequencer/Private/Tools/NavigationTool/Widgets/Columns/SNavigationToolLength.h
--- a/Engine/Source/Editor/Sequencer/Private/Tools/NavigationTool/Widgets/Columns/SNavigationToolLength.h
+++ b/Engine/Source/Editor/Sequencer/Private/Tools/NavigationTool/Widgets/Columns/SNavigationToolLength.h
@@ -2,8 +2,8 @@
 
 #pragma once
 
+#include "SNavigationToolTime.h"
 #include "Tools/NavigationTool/NavigationToolDefines.h"
-#include "Widgets/SCompoundWidget.h"
 
 namespace UE::Sequencer::NavigationTool
 {
@@ -11,7 +11,7 @@ namespace UE::Sequencer::NavigationTool
 class INavigationToolView;
 class SNavigationToolTreeRow;
 
-class SNavigationToolLength : public SCompoundWidget
+class SNavigationToolLength : public SNavigationToolTime
 {
 public:
 	SLATE_BEGIN_ARGS(SNavigationToolLength) {}
@@ -25,13 +25,13 @@ public:
 		, const TSharedRef<SNavigationToolTreeRow>& InRowWidget);
 
 private:
-	int32 GetLengthValue() const;
+	virtual FName GetStyleName() const override;
 
-	FNavigationToolItemWeakPtr WeakItem;
+	virtual double GetFrameTimeValue() const override;
 
-	TWeakPtr<INavigationToolView> WeakView;
+	virtual FText GetTransactionText() const override;
 
-	TWeakPtr<SNavigationToolTreeRow> WeakRowWidget;
+	virtual TSharedPtr<INumericTypeInterface<double>> GetNumericTypeInterface() const;
 };
 
 }
diff --git a/Engine/Source/Editor/Sequencer/Private/Tools/NavigationTool/Widgets/Columns/SNavigationToolOutTime.cpp b/Engine/Source/Editor/Sequencer/Private/Tools/NavigationTool/Widgets/Columns/SNavigationToolOutTime.cpp
--- a/Engine/Source/Editor/Sequencer/Private/Tools/NavigationTool/Widgets/Columns/SNavigationToolOutTime.cpp
+++ b/Engine/Source/Editor/Sequencer/Private/Tools/NavigationTool/Widgets/Columns/SNavigationToolOutTime.cpp
@@ -23,6 +23,7 @@ void SNavigationToolOutTime::Construct(const FArguments& InArgs
 
 	if (!WeakItem.IsValid()
 		|| !WeakView.IsValid()
+		|| !WeakRowWidget.IsValid()
 		|| !WeakItem.Pin()->IsA<IOutTimeExtension>())
 	{
 		return;
@@ -41,7 +42,7 @@ FName SNavigationToolOutTime::GetStyleName() const
 	return TEXT("Sequencer.NavigationTool.OutTime.SpinBox");
 }
 
-int32 SNavigationToolOutTime::GetFrameTimeValue() const
+double SNavigationToolOutTime::GetFrameTimeValue() const
 {
 	if (const IOutTimeExtension* const OutTimeExtension = GetOutTimeExtension())
 	{
@@ -50,7 +51,7 @@ int32 SNavigationToolOutTime::GetFrameTimeValue() const
 	return 0;
 }
 
-void SNavigationToolOutTime::OnFrameTimeValueCommitted(const int32 InNewValue, const ETextCommit::Type InCommitType)
+void SNavigationToolOutTime::OnFrameTimeValueCommitted(const double InNewValue, const ETextCommit::Type InCommitType)
 {
 	IOutTimeExtension* const OutTimeExtension = GetOutTimeExtension();
 	if (!OutTimeExtension)
@@ -61,7 +62,7 @@ void SNavigationToolOutTime::OnFrameTimeValueCommitted(const int32 InNewValue, c
 	const bool bShouldTransact = !UndoTransaction.IsValid() && (InCommitType == ETextCommit::OnEnter);
 	const FScopedTransaction Transaction(GetTransactionText(), bShouldTransact);
 
-	return OutTimeExtension->SetOutTime(InNewValue);
+	return OutTimeExtension->SetOutTime(static_cast<int32>(InNewValue));
 }
 
 FText SNavigationToolOutTime::GetTransactionText() const
diff --git a/Engine/Source/Editor/Sequencer/Private/Tools/NavigationTool/Widgets/Columns/SNavigationToolOutTime.h b/Engine/Source/Editor/Sequencer/Private/Tools/NavigationTool/Widgets/Columns/SNavigationToolOutTime.h
--- a/Engine/Source/Editor/Sequencer/Private/Tools/NavigationTool/Widgets/Columns/SNavigationToolOutTime.h
+++ b/Engine/Source/Editor/Sequencer/Private/Tools/NavigationTool/Widgets/Columns/SNavigationToolOutTime.h
@@ -3,7 +3,6 @@
 #pragma once
 
 #include "SNavigationToolTime.h"
-#include "Templates/UniquePtr.h"
 #include "Tools/NavigationTool/NavigationToolDefines.h"
 
 namespace UE::Sequencer::NavigationTool
@@ -31,9 +30,9 @@ protected:
 
 	virtual FName GetStyleName() const override;
 
-	virtual int32 GetFrameTimeValue() const override;
+	virtual double GetFrameTimeValue() const override;
 
-	virtual void OnFrameTimeValueCommitted(const int32 InNewValue, const ETextCommit::Type InCommitType) override;
+	virtual void OnFrameTimeValueCommitted(const double InNewValue, const ETextCommit::Type InCommitType) override;
 
 	virtual FText GetTransactionText() const override;
 };
diff --git a/Engine/Source/Editor/Sequencer/Private/Tools/NavigationTool/Widgets/Columns/SNavigationToolStartFrameOffset.cpp b/Engine/Source/Editor/Sequencer/Private/Tools/NavigationTool/Widgets/Columns/SNavigationToolStartFrameOffset.cpp
--- a/Engine/Source/Editor/Sequencer/Private/Tools/NavigationTool/Widgets/Columns/SNavigationToolStartFrameOffset.cpp
+++ b/Engine/Source/Editor/Sequencer/Private/Tools/NavigationTool/Widgets/Columns/SNavigationToolStartFrameOffset.cpp
@@ -6,7 +6,6 @@
 #include "Tools/NavigationTool/NavigationToolView.h"
 #include "Tools/NavigationTool/Widgets/SNavigationToolTreeRow.h"
 #include "TrackEditors/SubTrackEditorBase.h"
-#include "Widgets/Input/SSpinBox.h"
 
 #define LOCTEXT_NAMESPACE "SNavigationToolStartFrameOffset"
 
@@ -22,88 +21,46 @@ void SNavigationToolStartFrameOffset::Construct(const FArguments& InArgs
 	WeakView = InView;
 	WeakRowWidget = InRowWidget;
 
-	UMovieSceneSubSection* const SubSection = Cast<UMovieSceneSubSection>(GetSequenceItemSubSection(WeakItem.Pin()));
-
-	if (!WeakView.IsValid() || !IsValid(SubSection))
+	if (!WeakItem.IsValid()
+		|| !WeakView.IsValid()
+		|| !WeakRowWidget.IsValid())
 	{
-		ChildSlot
-		[
-			SNullWidget::NullWidget
-		];
-
 		return;
 	}
 
-	ChildSlot
-		[
-			SNew(SSpinBox<int32>)
-			.Justification(ETextJustify::Center)
-			.Style(&FAppStyle::Get(), TEXT("Sequencer.NavigationTool.SpinBox"))
-			.Value(this, &SNavigationToolStartFrameOffset::GetValue)
-			.OnValueChanged(this, &SNavigationToolStartFrameOffset::OnValueChanged)
-			.OnValueCommitted(this, &SNavigationToolStartFrameOffset::OnValueCommitted)
-			.OnBeginSliderMovement(this, &SNavigationToolStartFrameOffset::OnBeginSliderMovement)
-			.OnEndSliderMovement(this, &SNavigationToolStartFrameOffset::OnEndSliderMovement)
-		];
-}
-
-int32 SNavigationToolStartFrameOffset::GetValue() const
-{
-	const TSharedPtr<INavigationToolView> ToolView = WeakView.Pin();
-	if (!ToolView.IsValid())
-	{
-		return 0;
-	}
-
-	const TSharedPtr<ISequencer> Sequencer = ToolView->GetSequencer();
-	if (!Sequencer.IsValid())
-	{
-		return 0;
-	}
-
-	UMovieSceneSubSection* const SubSection = Cast<UMovieSceneSubSection>(GetSequenceItemSubSection(WeakItem.Pin()));
-	if (!IsValid(SubSection))
+	const FNavigationToolSequence* const SequenceItem = WeakItem.Pin()->CastTo<FNavigationToolSequence>();
+	if (!SequenceItem || !IsValid(SequenceItem->GetSubSection()))
 	{
-		return 0;
+		return;
 	}
 
-	const FFrameNumber Offset = ConvertFrameTime(SubSection->Parameters.StartFrameOffset
-		, Sequencer->GetFocusedTickResolution()
-		, Sequencer->GetFocusedDisplayRate()).FrameNumber;
-
-	return Offset.Value;
+	SNavigationToolTime::Construct(SNavigationToolTime::FArguments(), InItem, InView, InRowWidget);
 }
 
-void SNavigationToolStartFrameOffset::OnValueChanged(const int32 InNewValue)
+FName SNavigationToolStartFrameOffset::GetStyleName() const
 {
-	OnValueCommitted(InNewValue, ETextCommit::Default);
+	return TEXT("Sequencer.NavigationTool.SpinBox");
 }
 
-void SNavigationToolStartFrameOffset::OnValueCommitted(const int32 InNewValue, const ETextCommit::Type InCommitType)
+double SNavigationToolStartFrameOffset::GetFrameTimeValue() const
 {
-	const TSharedPtr<INavigationToolView> ToolView = WeakView.Pin();
-	if (!ToolView.IsValid())
+	const UMovieSceneSubSection* const SubSection = Cast<UMovieSceneSubSection>(GetSequenceItemSubSection(WeakItem.Pin()));
+	if (IsValid(SubSection))
 	{
-		return;
-	}
-
-	const TSharedPtr<ISequencer> Sequencer = ToolView->GetSequencer();
-	if (!Sequencer.IsValid())
-	{
-		return;
+		return SubSection->Parameters.StartFrameOffset.Value;
 	}
+	return 0.;
+}
 
+void SNavigationToolStartFrameOffset::OnFrameTimeValueCommitted(const double InNewValue, const ETextCommit::Type InCommitType)
+{
 	UMovieSceneSubSection* const SubSection = Cast<UMovieSceneSubSection>(GetSequenceItemSubSection(WeakItem.Pin()));
 	if (!IsValid(SubSection))
 	{
 		return;
 	}
 
-	const FFrameNumber FrameTime = ConvertFrameTime(InNewValue
-		, Sequencer->GetFocusedDisplayRate()
-		, Sequencer->GetFocusedTickResolution()).FrameNumber;
-
-	if (InNewValue == FrameTime.Value)
+	if (InNewValue == SubSection->Parameters.StartFrameOffset.Value)
 	{
 		return;
 	}
@@ -112,25 +69,12 @@ void SNavigationToolStartFrameOffset::OnValueCommitted(const int32 InNewValue, c
 	const FScopedTransaction Transaction(GetTransactionText(), bShouldTransact);
 
 	SubSection->Modify();
-	SubSection->Parameters.StartFrameOffset = FrameTime;
-}
-
-void SNavigationToolStartFrameOffset::OnBeginSliderMovement()
-{
-	if (!UndoTransaction.IsValid())
-	{
-		UndoTransaction = MakeUnique<FScopedTransaction>(GetTransactionText());
-	}
-}
-
-void SNavigationToolStartFrameOffset::OnEndSliderMovement(const int32 InNewValue)
-{
-	UndoTransaction.Reset();
+	SubSection->Parameters.StartFrameOffset = static_cast<int32>(InNewValue);
 }
 
 FText SNavigationToolStartFrameOffset::GetTransactionText() const
 {
-	return LOCTEXT("SetSequenceStartFrameOffsetTransaction", "Set Sequence Start Frame Offset");
+	return LOCTEXT("SetStartFrameOffsetTransaction", "Set Start Frame Offset");
 }
 
 }
diff --git a/Engine/Source/Editor/Sequencer/Private/Tools/NavigationTool/Widgets/Columns/SNavigationToolStartFrameOffset.h b/Engine/Source/Editor/Sequencer/Private/Tools/NavigationTool/Widgets/Columns/SNavigationToolStartFrameOffset.h
--- a/Engine/Source/Editor/Sequencer/Private/Tools/NavigationTool/Widgets/Columns/SNavigationToolStartFrameOffset.h
+++ b/Engine/Source/Editor/Sequencer/Private/Tools/NavigationTool/Widgets/Columns/SNavigationToolStartFrameOffset.h
@@ -2,10 +2,8 @@
 
 #pragma once
 
-#include "ScopedTransaction.h"
-#include "Templates/UniquePtr.h"
+#include "SNavigationToolTime.h"
 #include "Tools/NavigationTool/NavigationToolDefines.h"
-#include "Widgets/SCompoundWidget.h"
 
 namespace UE::Sequencer::NavigationTool
 {
@@ -13,7 +11,7 @@ namespace UE::Sequencer::NavigationTool
 class INavigationToolView;
 class SNavigationToolTreeRow;
 
-class SNavigationToolStartFrameOffset : public SCompoundWidget
+class SNavigationToolStartFrameOffset : public SNavigationToolTime
 {
 public:
 	SLATE_BEGIN_ARGS(SNavigationToolStartFrameOffset) {}
@@ -27,21 +25,13 @@ public:
 		, const TSharedRef<SNavigationToolTreeRow>& InRowWidget);
 
 private:
-	int32 GetValue() const;
-	void OnValueChanged(const int32 InNewValue);
-	void OnValueCommitted(const int32 InNewValue, const ETextCommit::Type InCommitType);
-	void OnBeginSliderMovement();
-	void OnEndSliderMovement(const int32 InNewValue);
+	virtual FName GetStyleName() const override;
 
-	FText GetTransactionText() const;
+	virtual double GetFrameTimeValue() const override;
 
-	FNavigationToolItemWeakPtr WeakItem;
+	virtual void OnFrameTimeValueCommitted(const double InNewValue, const ETextCommit::Type InCommitType) override;
 
-	TWeakPtr<INavigationToolView> WeakView;
-
-	TWeakPtr<SNavigationToolTreeRow> WeakRowWidget;
-
-	TUniquePtr<FScopedTransaction> UndoTransaction;
+	virtual FText GetTransactionText() const override;
 };
 
 }
diff --git a/Engine/Source/Editor/Sequencer/Private/Tools/NavigationTool/Widgets/Columns/SNavigationToolTime.cpp b/Engine/Source/Editor/Sequencer/Private/Tools/NavigationTool/Widgets/Columns/SNavigationToolTime.cpp
--- a/Engine/Source/Editor/Sequencer/Private/Tools/NavigationTool/Widgets/Columns/SNavigationToolTime.cpp
+++ b/Engine/Source/Editor/Sequencer/Private/Tools/NavigationTool/Widgets/Columns/SNavigationToolTime.cpp
@@ -1,7 +1,7 @@
 ﻿// Copyright Epic Games, Inc. All Rights Reserved.
 
 #include "SNavigationToolTime.h"
-#include "Tools/NavigationTool/INavigationTool.h"
+#include "Sequencer.h"
 #include "Tools/NavigationTool/Items/NavigationToolItemUtils.h"
 #include "Tools/NavigationTool/NavigationToolView.h"
 #include "Tools/NavigationTool/Widgets/SNavigationToolTreeRow.h"
@@ -23,47 +23,93 @@ void SNavigationToolTime::Construct(const FArguments& InArgs
 
 	if (!WeakItem.IsValid()
 		|| !WeakView.IsValid()
-		|| !WeakItem.Pin()->IsA<IInTimeExtension>())
+		|| !WeakRowWidget.IsValid())
 	{
 		return;
 	}
 
 	ChildSlot
 	[
-		SNew(SSpinBox<int32>)
+		SNew(SSpinBox<double>)
 		.Justification(ETextJustify::Center)
 		.Style(&FAppStyle::GetWidgetStyle<FSpinBoxStyle>(GetStyleName()))
+		.TypeInterface(GetNumericTypeInterface())
 		.Value(this, &SNavigationToolTime::GetFrameTimeValue)
 		.OnValueChanged(this, &SNavigationToolTime::OnFrameTimeValueChanged)
 		.OnValueCommitted(this, &SNavigationToolTime::OnFrameTimeValueCommitted)
 		.OnBeginSliderMovement(this, &SNavigationToolTime::OnBeginSliderMovement)
 		.OnEndSliderMovement(this, &SNavigationToolTime::OnEndSliderMovement)
+		.Delta(this, &SNavigationToolTime::GetDisplayRateDeltaFrameCount)
 	];
 }
 
+bool SNavigationToolTime::IsReadOnly() const
+{
+	return bIsReadOnly;
+}
+
+void SNavigationToolTime::SetIsReadOnly(const bool bInIsReadOnly)
+{
+	bIsReadOnly = bInIsReadOnly;
+}
+
 FName SNavigationToolTime::GetStyleName() const
 {
 	return TEXT("Sequencer.NavigationTool.SpinBox");
 }
 
-void SNavigationToolTime::OnFrameTimeValueChanged(const int32 InNewValue)
+void SNavigationToolTime::OnFrameTimeValueChanged(const double InNewValue)
+{
+	if (!bIsReadOnly)
+	{
+		OnFrameTimeValueCommitted(InNewValue, ETextCommit::Default);
+	}
+}
+
+void SNavigationToolTime::OnFrameTimeValueCommitted(const double InNewValue, const ETextCommit::Type InCommitType)
 {
-	OnFrameTimeValueCommitted(InNewValue, ETextCommit::Default);
 }
 
 void SNavigationToolTime::OnBeginSliderMovement()
 {
-	if (!UndoTransaction.IsValid())
+	if (!bIsReadOnly && !UndoTransaction.IsValid())
 	{
 		UndoTransaction = MakeUnique<FScopedTransaction>(GetTransactionText());
 	}
 }
 
-void SNavigationToolTime::OnEndSliderMovement(const int32 InNewValue)
+void SNavigationToolTime::OnEndSliderMovement(const double InNewValue)
 {
 	UndoTransaction.Reset();
 }
 
+TSharedPtr<INumericTypeInterface<double>> SNavigationToolTime::GetNumericTypeInterface() const
+{
+	if (const TSharedPtr<ISequencer> Sequencer = GetSequencer())
+	{
+		return Sequencer->GetNumericTypeInterface(ENumericIntent::Position);
+	}
+	return nullptr;
+}
+
+double SNavigationToolTime::GetDisplayRateDeltaFrameCount() const
+{
+	if (const TSharedPtr<FSequencer> Sequencer = StaticCastSharedPtr<FSequencer>(GetSequencer()))
+	{
+		return Sequencer->GetDisplayRateDeltaFrameCount();
+	}
+	return 0.;
+}
+
+TSharedPtr<ISequencer> SNavigationToolTime::GetSequencer() const
+{
+	if (const TSharedPtr<INavigationToolView> View = WeakView.Pin())
+	{
+		return View->GetSequencer();
+	}
+	return nullptr;
+}
+
 }
 
 #undef LOCTEXT_NAMESPACE
diff --git a/Engine/Source/Editor/Sequencer/Private/Tools/NavigationTool/Widgets/Columns/SNavigationToolTime.h b/Engine/Source/Editor/Sequencer/Private/Tools/NavigationTool/Widgets/Columns/SNavigationToolTime.h
--- a/Engine/Source/Editor/Sequencer/Private/Tools/NavigationTool/Widgets/Columns/SNavigationToolTime.h
+++ b/Engine/Source/Editor/Sequencer/Private/Tools/NavigationTool/Widgets/Columns/SNavigationToolTime.h
@@ -5,8 +5,11 @@
 #include "ScopedTransaction.h"
 #include "Templates/UniquePtr.h"
 #include "Tools/NavigationTool/NavigationToolDefines.h"
+#include "Widgets/Input/NumericTypeInterface.h"
 #include "Widgets/SCompoundWidget.h"
 
+class ISequencer;
+
 namespace UE::Sequencer::NavigationTool
 {
 
@@ -26,19 +29,28 @@ public:
 		, const TSharedRef<INavigationToolView>& InView
 		, const TSharedRef<SNavigationToolTreeRow>& InRowWidget);
 
+	virtual bool IsReadOnly() const;
+	virtual void SetIsReadOnly(const bool bInIsReadOnly);
+
 protected:
 	virtual FName GetStyleName() const;
 
-	virtual int32 GetFrameTimeValue() const = 0;
+	virtual double GetFrameTimeValue() const = 0;
 
-	virtual void OnFrameTimeValueChanged(const int32 InNewValue);
-	virtual void OnFrameTimeValueCommitted(const int32 InNewValue, const ETextCommit::Type InCommitType) = 0;
+	virtual void OnFrameTimeValueChanged(const double InNewValue);
+	virtual void OnFrameTimeValueCommitted(const double InNewValue, const ETextCommit::Type InCommitType);
 
 	virtual void OnBeginSliderMovement();
-	virtual void OnEndSliderMovement(const int32 InNewValue);
+	virtual void OnEndSliderMovement(const double InNewValue);
+
+	virtual TSharedPtr<INumericTypeInterface<double>> GetNumericTypeInterface() const;
+
+	virtual double GetDisplayRateDeltaFrameCount() const;
 
 	virtual FText GetTransactionText() const = 0;
 
+	TSharedPtr<ISequencer> GetSequencer() const;
+
 	FNavigationToolItemWeakPtr WeakItem;
 
 	TWeakPtr<INavigationToolView> WeakView;
@@ -46,6 +58,8 @@ protected:
 	TWeakPtr<SNavigationToolTreeRow> WeakRowWidget;
 
 	TUniquePtr<FScopedTransaction> UndoTransaction;
+
+	bool bIsReadOnly = false;
 };
 
 }
