Sequencer build fix

#rnx

diff --git a/Engine/Source/Editor/Sequencer/Public/MVVM/Selection/Selection.h b/Engine/Source/Editor/Sequencer/Public/MVVM/Selection/Selection.h
--- a/Engine/Source/Editor/Sequencer/Public/MVVM/Selection/Selection.h
+++ b/Engine/Source/Editor/Sequencer/Public/MVVM/Selection/Selection.h
@@ -175,9 +175,9 @@ private:
 	virtual TSharedPtr<      FOutlinerSelection> GetOutlinerSelection()       override { return TSharedPtr<      FOutlinerSelection>(AsShared(), &Outliner); }
 	virtual TSharedPtr<const FOutlinerSelection> GetOutlinerSelection() const override { return TSharedPtr<const FOutlinerSelection>(AsShared(), &Outliner); }
 
-	virtual void PreSelectionSetChangeEvent(FSelectionBase* InSelectionSet) override;
-
-	void PreBroadcastChangeEvent() override;
+	SEQUENCER_API virtual void PreSelectionSetChangeEvent(FSelectionBase* InSelectionSet) override;
+	SEQUENCER_API virtual void PreBroadcastChangeEvent() override;
+	
 	TSet<TWeakViewModelPtr<IOutlinerExtension>> NodesWithKeysOrSections;
 };
 
