Sequencer: Add DoubleClick support to sequencer InputHandler

#rnx

diff --git a/Engine/Plugins/Runtime/GeometryCache/Source/GeometryCacheEd/Private/GeometryCacheTimeSliderController.cpp b/Engine/Plugins/Runtime/GeometryCache/Source/GeometryCacheEd/Private/GeometryCacheTimeSliderController.cpp
--- a/Engine/Plugins/Runtime/GeometryCache/Source/GeometryCacheEd/Private/GeometryCacheTimeSliderController.cpp
+++ b/Engine/Plugins/Runtime/GeometryCache/Source/GeometryCacheEd/Private/GeometryCacheTimeSliderController.cpp
@@ -337,6 +337,11 @@ FReply FGeometryCacheTimeSlideController::OnMouseButtonUp(SWidget& WidgetOwner,
 	return FReply::Unhandled();
 }
 
+FReply FGeometryCacheTimeSlideController::OnMouseButtonDoubleClick(SWidget& OwnerWidget, const FGeometry& MyGeometry, const FPointerEvent& MouseEvent)
+{
+	return FReply::Unhandled();
+}
+
 FReply FGeometryCacheTimeSlideController::OnMouseMove(SWidget& WidgetOwner, const FGeometry& MyGeometry, const FPointerEvent& MouseEvent)
 {
 	const bool bHandleLeftMouseButton = MouseEvent.IsMouseButtonDown(EKeys::LeftMouseButton);
diff --git a/Engine/Plugins/Runtime/GeometryCache/Source/GeometryCacheEd/Private/GeometryCacheTimeSliderController.h b/Engine/Plugins/Runtime/GeometryCache/Source/GeometryCacheEd/Private/GeometryCacheTimeSliderController.h
--- a/Engine/Plugins/Runtime/GeometryCache/Source/GeometryCacheEd/Private/GeometryCacheTimeSliderController.h
+++ b/Engine/Plugins/Runtime/GeometryCache/Source/GeometryCacheEd/Private/GeometryCacheTimeSliderController.h
@@ -26,6 +26,7 @@ public:
 	virtual int32 OnPaintViewArea(const FGeometry& AllottedGeometry, const FSlateRect& MyCullingRect, FSlateWindowElementList& OutDrawElements, int32 LayerId, bool bEnabled, const FPaintViewAreaArgs& Args) const override;
 	virtual FReply OnMouseButtonDown(SWidget& WidgetOwner, const FGeometry& MyGeometry, const FPointerEvent& MouseEvent) override;
 	virtual FReply OnMouseButtonUp(SWidget& WidgetOwner, const FGeometry& MyGeometry, const FPointerEvent& MouseEvent) override;
+	virtual FReply OnMouseButtonDoubleClick(SWidget& OwnerWidget, const FGeometry& MyGeometry, const FPointerEvent& MouseEvent) override;
 	virtual FReply OnMouseMove(SWidget& WidgetOwner, const FGeometry& MyGeometry, const FPointerEvent& MouseEvent) override;
 	virtual FReply OnMouseWheel(SWidget& WidgetOwner, const FGeometry& MyGeometry, const FPointerEvent& MouseEvent) override;
 	virtual FCursorReply OnCursorQuery(TSharedRef<const SWidget> WidgetOwner, const FGeometry& MyGeometry, const FPointerEvent& CursorEvent) const override;
diff --git a/Engine/Source/Editor/Persona/Private/AnimTimeline/AnimTimeSliderController.cpp b/Engine/Source/Editor/Persona/Private/AnimTimeline/AnimTimeSliderController.cpp
--- a/Engine/Source/Editor/Persona/Private/AnimTimeline/AnimTimeSliderController.cpp
+++ b/Engine/Source/Editor/Persona/Private/AnimTimeline/AnimTimeSliderController.cpp
@@ -676,6 +676,10 @@ FReply FAnimTimeSliderController::OnMouseButtonUp( SWidget& WidgetOwner, const F
 	return FReply::Unhandled();
 }
 
+FReply FAnimTimeSliderController::OnMouseButtonDoubleClick(SWidget& OwnerWidget, const FGeometry& MyGeometry, const FPointerEvent& MouseEvent)
+{
+	return FReply::Unhandled();
+}
 
 FReply FAnimTimeSliderController::OnMouseMove( SWidget& WidgetOwner, const FGeometry& MyGeometry, const FPointerEvent& MouseEvent )
 {
diff --git a/Engine/Source/Editor/Persona/Private/AnimTimeline/AnimTimeSliderController.h b/Engine/Source/Editor/Persona/Private/AnimTimeline/AnimTimeSliderController.h
--- a/Engine/Source/Editor/Persona/Private/AnimTimeline/AnimTimeSliderController.h
+++ b/Engine/Source/Editor/Persona/Private/AnimTimeline/AnimTimeSliderController.h
@@ -53,6 +53,7 @@ public:
 	virtual int32 OnPaintViewArea( const FGeometry& AllottedGeometry, const FSlateRect& MyCullingRect, FSlateWindowElementList& OutDrawElements, int32 LayerId, bool bEnabled, const FPaintViewAreaArgs& Args ) const override;
 	virtual FReply OnMouseButtonDown( SWidget& WidgetOwner, const FGeometry& MyGeometry, const FPointerEvent& MouseEvent ) override;
 	virtual FReply OnMouseButtonUp( SWidget& WidgetOwner, const FGeometry& MyGeometry, const FPointerEvent& MouseEvent ) override;
+	virtual FReply OnMouseButtonDoubleClick(SWidget& OwnerWidget, const FGeometry& MyGeometry, const FPointerEvent& MouseEvent) override;
 	virtual FReply OnMouseMove( SWidget& WidgetOwner, const FGeometry& MyGeometry, const FPointerEvent& MouseEvent ) override;
 	virtual FReply OnMouseWheel( SWidget& WidgetOwner, const FGeometry& MyGeometry, const FPointerEvent& MouseEvent ) override;
 	virtual FCursorReply OnCursorQuery( TSharedRef<const SWidget> WidgetOwner, const FGeometry& MyGeometry, const FPointerEvent& CursorEvent ) const override;
diff --git a/Engine/Source/Editor/Sequencer/Private/SequencerTimeSliderController.cpp b/Engine/Source/Editor/Sequencer/Private/SequencerTimeSliderController.cpp
--- a/Engine/Source/Editor/Sequencer/Private/SequencerTimeSliderController.cpp
+++ b/Engine/Source/Editor/Sequencer/Private/SequencerTimeSliderController.cpp
@@ -1113,12 +1113,12 @@ FReply FSequencerTimeSliderController::OnMouseButtonUp( SWidget& WidgetOwner, co
 	return FReply::Unhandled();
 }
 
-FReply FSequencerTimeSliderController::OnMouseButtonDoubleClick(TSharedRef<const SWidget> WidgetOwner, const FGeometry& MyGeometry, const FPointerEvent& MouseEvent)
+FReply FSequencerTimeSliderController::OnMouseButtonDoubleClick(SWidget& WidgetOwner, const FGeometry& MyGeometry, const FPointerEvent& MouseEvent)
 {
 	if (HoverMarkIndex != INDEX_NONE)
 	{
 		// Teleport the playhead to the double-clicked marker.
-		bool bHandleMiddleMouseButton = MouseEvent.GetEffectingButton() == EKeys::MiddleMouseButton && WidgetOwner->HasMouseCapture();
+		bool bHandleMiddleMouseButton = MouseEvent.GetEffectingButton() == EKeys::MiddleMouseButton && WidgetOwner.HasMouseCapture();
 		const TArray<FMovieSceneMarkedFrame>& MarkedFrames = TimeSliderArgs.MarkedFrames.Get();
 		const FFrameNumber MarkTime = MarkedFrames[HoverMarkIndex].FrameNumber;
 		CommitScrubPosition(MarkTime, /*bIsScrubbing=*/ false , /*bEvaluate*/ !bHandleMiddleMouseButton);
diff --git a/Engine/Source/Editor/Sequencer/Private/Tools/SequencerEditTool.h b/Engine/Source/Editor/Sequencer/Private/Tools/SequencerEditTool.h
--- a/Engine/Source/Editor/Sequencer/Private/Tools/SequencerEditTool.h
+++ b/Engine/Source/Editor/Sequencer/Private/Tools/SequencerEditTool.h
@@ -40,6 +40,11 @@ public:
 		return FReply::Unhandled();
 	}
 
+	virtual FReply OnMouseButtonDoubleClick(SWidget& OwnerWidget, const FGeometry& MyGeometry, const FPointerEvent& MouseEvent) override
+	{
+		return FReply::Unhandled();
+	}
+
 	virtual FReply OnMouseMove(SWidget& OwnerWidget, const FGeometry& MyGeometry, const FPointerEvent& MouseEvent) override
 	{
 		return FReply::Unhandled();
diff --git a/Engine/Source/Editor/Sequencer/Public/SequencerTimeSliderController.h b/Engine/Source/Editor/Sequencer/Public/SequencerTimeSliderController.h
--- a/Engine/Source/Editor/Sequencer/Public/SequencerTimeSliderController.h
+++ b/Engine/Source/Editor/Sequencer/Public/SequencerTimeSliderController.h
@@ -42,11 +42,11 @@ public:
 	virtual int32 OnPaintTimeSlider( bool bMirrorLabels, const FGeometry& AllottedGeometry, const FSlateRect& MyCullingRect, FSlateWindowElementList& OutDrawElements, int32 LayerId, const FWidgetStyle& InWidgetStyle, bool bParentEnabled ) const override;
 	virtual FReply OnMouseButtonDown( SWidget& WidgetOwner, const FGeometry& MyGeometry, const FPointerEvent& MouseEvent ) override;
 	virtual FReply OnMouseButtonUp( SWidget& WidgetOwner, const FGeometry& MyGeometry, const FPointerEvent& MouseEvent ) override;
+	virtual FReply OnMouseButtonDoubleClick(SWidget& WidgetOwner, const FGeometry& MyGeometry, const FPointerEvent& MouseEvent) override;
 	virtual FReply OnMouseMove( SWidget& WidgetOwner, const FGeometry& MyGeometry, const FPointerEvent& MouseEvent ) override;
 	virtual FReply OnMouseWheel( SWidget& WidgetOwner, const FGeometry& MyGeometry, const FPointerEvent& MouseEvent ) override;
 	virtual FReply OnTimeSliderMouseMove( SWidget& OwnerWidget, const FGeometry& MyGeometry, const FPointerEvent& MouseEvent ) override;
 	virtual FCursorReply OnCursorQuery( TSharedRef<const SWidget> WidgetOwner, const FGeometry& MyGeometry, const FPointerEvent& CursorEvent ) const override;
-	virtual FReply OnMouseButtonDoubleClick( TSharedRef<const SWidget> WidgetOwner, const FGeometry& MyGeometry, const FPointerEvent& MouseEvent ) override;
 	virtual double ComputeHeight() const override;
 	/** End ITimeSliderController Interface */
 
diff --git a/Engine/Source/Editor/SequencerCore/Private/MVVM/Views/STrackAreaView.cpp b/Engine/Source/Editor/SequencerCore/Private/MVVM/Views/STrackAreaView.cpp
--- a/Engine/Source/Editor/SequencerCore/Private/MVVM/Views/STrackAreaView.cpp
+++ b/Engine/Source/Editor/SequencerCore/Private/MVVM/Views/STrackAreaView.cpp
@@ -406,6 +406,19 @@ FReply STrackAreaView::OnMouseButtonUp( const FGeometry& MyGeometry, const FPoin
 	return InputStack.HandleMouseButtonUp(*this, MyGeometry, MouseEvent);
 }
 
+FReply STrackAreaView::OnMouseButtonDoubleClick(const FGeometry& MyGeometry, const FPointerEvent& MouseEvent)
+{
+	TSharedPtr<FTrackAreaViewModel> TrackArea = WeakViewModel.Pin();
+	ISequencerEditTool*             EditTool  = TrackArea ? TrackArea->GetEditTool() : nullptr;
+
+	if (EditTool)
+	{
+		// Always ensure the edit tool is set up
+		InputStack.SetHandlerAt(0, EditTool);
+	}
+	return InputStack.HandleMouseDoubleClick(*this, MyGeometry, MouseEvent);
+}
+
 void STrackAreaView::UpdateHoverStates( const FGeometry& MyGeometry, const FPointerEvent& MouseEvent )
 {
 	TSharedPtr<SOutlinerView>       PinnedOutliner = WeakOutliner.Pin();
diff --git a/Engine/Source/Editor/SequencerCore/Public/ISequencerInputHandler.h b/Engine/Source/Editor/SequencerCore/Public/ISequencerInputHandler.h
--- a/Engine/Source/Editor/SequencerCore/Public/ISequencerInputHandler.h
+++ b/Engine/Source/Editor/SequencerCore/Public/ISequencerInputHandler.h
@@ -14,6 +14,7 @@ struct ISequencerInputHandler
 {
 	virtual FReply OnMouseButtonDown(SWidget& OwnerWidget, const FGeometry& MyGeometry, const FPointerEvent& MouseEvent) = 0;
 	virtual FReply OnMouseButtonUp(SWidget& OwnerWidget, const FGeometry& MyGeometry, const FPointerEvent& MouseEvent) = 0;
+	virtual FReply OnMouseButtonDoubleClick(SWidget& OwnerWidget, const FGeometry& MyGeometry, const FPointerEvent& MouseEvent) = 0;
 	virtual FReply OnMouseMove(SWidget& OwnerWidget, const FGeometry& MyGeometry, const FPointerEvent& MouseEvent) = 0;
 	virtual FReply OnMouseWheel(SWidget& OwnerWidget, const FGeometry& MyGeometry, const FPointerEvent& MouseEvent) = 0;
 	virtual FReply OnKeyDown(SWidget& OwnerWidget, const FGeometry& MyGeometry, const FKeyEvent& InKeyEvent) { return FReply::Unhandled(); }
diff --git a/Engine/Source/Editor/SequencerCore/Public/MVVM/Views/STrackAreaView.h b/Engine/Source/Editor/SequencerCore/Public/MVVM/Views/STrackAreaView.h
--- a/Engine/Source/Editor/SequencerCore/Public/MVVM/Views/STrackAreaView.h
+++ b/Engine/Source/Editor/SequencerCore/Public/MVVM/Views/STrackAreaView.h
@@ -130,6 +130,7 @@ public:
 	/*~ SWidget interface */
 	FReply OnMouseButtonDown( const FGeometry& MyGeometry, const FPointerEvent& MouseEvent ) override;
 	FReply OnMouseButtonUp( const FGeometry& MyGeometry, const FPointerEvent& MouseEvent ) override;
+	FReply OnMouseButtonDoubleClick( const FGeometry& MyGeometry, const FPointerEvent& MouseEvent ) override;
 	FReply OnMouseMove( const FGeometry& MyGeometry, const FPointerEvent& MouseEvent ) override;
 	FReply OnMouseWheel( const FGeometry& MyGeometry, const FPointerEvent& MouseEvent ) override;
 	void OnMouseEnter(const FGeometry& MyGeometry, const FPointerEvent& MouseEvent) override;
diff --git a/Engine/Source/Editor/SequencerCore/Public/MVVM/Views/SequencerInputHandlerStack.h b/Engine/Source/Editor/SequencerCore/Public/MVVM/Views/SequencerInputHandlerStack.h
--- a/Engine/Source/Editor/SequencerCore/Public/MVVM/Views/SequencerInputHandlerStack.h
+++ b/Engine/Source/Editor/SequencerCore/Public/MVVM/Views/SequencerInputHandlerStack.h
@@ -73,6 +73,12 @@ public:
 		return ProcessEvent(&ISequencerInputHandler::OnMouseButtonUp, OwnerWidget, MyGeometry, MouseEvent);
 	}
 
+	/** Handle a mouse double click */
+	FReply HandleMouseDoubleClick(SWidget& OwnerWidget, const FGeometry& MyGeometry, const FPointerEvent& MouseEvent)
+	{
+		return ProcessEvent(&ISequencerInputHandler::OnMouseButtonDoubleClick, OwnerWidget, MyGeometry, MouseEvent);
+	}
+
 	/** Handle a mouse move */
 	FReply HandleMouseMove(SWidget& OwnerWidget, const FGeometry& MyGeometry, const FPointerEvent& MouseEvent)
 	{
@@ -89,6 +95,7 @@ public:
 	{
 		return ProcessKeyEvent(&ISequencerInputHandler::OnKeyDown, OwnerWidget, MyGeometry, InKeyEvent);
 	}
+	
 	FReply HandleKeyUp(SWidget& OwnerWidget, const FGeometry& MyGeometry, const FKeyEvent& InKeyEvent)
 	{
 		return ProcessKeyEvent(&ISequencerInputHandler::OnKeyUp, OwnerWidget, MyGeometry, InKeyEvent);
diff --git a/Engine/Source/Editor/SequencerWidgets/Private/SSequencerTimeSlider.cpp b/Engine/Source/Editor/SequencerWidgets/Private/SSequencerTimeSlider.cpp
--- a/Engine/Source/Editor/SequencerWidgets/Private/SSequencerTimeSlider.cpp
+++ b/Engine/Source/Editor/SequencerWidgets/Private/SSequencerTimeSlider.cpp
@@ -44,7 +44,7 @@ FReply SSequencerTimeSlider::OnMouseMove( const FGeometry& MyGeometry, const FPo
 
 FReply SSequencerTimeSlider::OnMouseButtonDoubleClick(const FGeometry& MyGeometry, const FPointerEvent& MouseEvent)
 {
-	return TimeSliderController->OnMouseButtonDoubleClick( SharedThis(this), MyGeometry, MouseEvent );
+	return TimeSliderController->OnMouseButtonDoubleClick( *this, MyGeometry, MouseEvent );
 }
 
 FVector2D SSequencerTimeSlider::ComputeDesiredSize( float ) const
diff --git a/Engine/Source/Editor/SequencerWidgets/Public/ITimeSlider.h b/Engine/Source/Editor/SequencerWidgets/Public/ITimeSlider.h
--- a/Engine/Source/Editor/SequencerWidgets/Public/ITimeSlider.h
+++ b/Engine/Source/Editor/SequencerWidgets/Public/ITimeSlider.h
@@ -86,7 +86,8 @@ public:
 
 	virtual int32 OnPaintViewArea( const FGeometry& AllottedGeometry, const FSlateRect& MyCullingRect, FSlateWindowElementList& OutDrawElements, int32 LayerId, bool bEnabled, const FPaintViewAreaArgs& Args ) const = 0;
 
-	virtual FReply OnMouseButtonDoubleClick( TSharedRef<const SWidget> WidgetOwner, const FGeometry& MyGeometry, const FPointerEvent& MouseEvent ) { return FReply::Unhandled(); }
+	virtual FReply OnMouseButtonDoubleClick( SWidget& WidgetOwner, const FGeometry& MyGeometry, const FPointerEvent& MouseEvent ) override { return FReply::Unhandled(); }
+
 	virtual FCursorReply OnCursorQuery( TSharedRef<const SWidget> WidgetOwner, const FGeometry& MyGeometry, const FPointerEvent& CursorEvent ) const { return FCursorReply::Unhandled(); }
 
 	virtual FReply OnTimeSliderMouseButtonDown(SWidget& OwnerWidget, const FGeometry& MyGeometry, const FPointerEvent& MouseEvent)
