[Sequencer] Move Navigation Tool styling from FStarshipStyle to FNavigationToolStyle

diff --git a/Engine/Source/Editor/EditorStyle/Private/StarshipStyle.cpp b/Engine/Source/Editor/EditorStyle/Private/StarshipStyle.cpp
--- a/Engine/Source/Editor/EditorStyle/Private/StarshipStyle.cpp
+++ b/Engine/Source/Editor/EditorStyle/Private/StarshipStyle.cpp
@@ -2511,36 +2511,6 @@ void FStarshipEditorStyle::FStyle::SetupSequencerStyles()
 		Set("Sequencer.Indicator.Condition", new IMAGE_BRUSH_SVG("Sequencer/Column_Widgets/SequencerCondition", Icon12x12));
 		Set("Sequencer.Indicator.TimeWarp", new IMAGE_BRUSH_SVG("Sequencer/Column_Widgets/SequencerTimeWarp", Icon12x12));
 	}
-
-	// Navigation Tool Styling
-	{
-		Set("Sequencer.NavigationTool", new IMAGE_BRUSH_SVG("Sequencer/SequenceNavigator/SequenceNavigator_20", Icon20x20));
-		Set("Sequencer.NavigationTool.TabIcon", new IMAGE_BRUSH_SVG("Sequencer/SequenceNavigator/SequenceNavigator_16", Icon16x16));
-
-		const FTableRowStyle& AlternatingTableRowStyle = FAppStyle::Get().GetWidgetStyle<FTableRowStyle>(TEXT("TableView.AlternatingRow"));
-		Set("Sequencer.NavigationTool.TableViewRow", FTableRowStyle(AlternatingTableRowStyle));
-
-		FSpinBoxStyle FrameTimeSpinStyle = GetWidgetStyle<FSpinBoxStyle>(TEXT("SpinBox"));
-		FrameTimeSpinStyle.ForegroundColor = FStyleColors::Foreground;
-		Set("Sequencer.NavigationTool.SpinBox", FrameTimeSpinStyle);
-
-		FSpinBoxStyle StartTimeSpinStyle = FrameTimeSpinStyle;
-		StartTimeSpinStyle.SetForegroundColor(FStyleColors::AccentGreen);
-		Set("Sequencer.NavigationTool.InTime.SpinBox", StartTimeSpinStyle);
-
-		FSpinBoxStyle EndTimeSpinStyle = FrameTimeSpinStyle;
-		EndTimeSpinStyle.SetForegroundColor(FStyleColors::AccentRed);
-		Set("Sequencer.NavigationTool.OutTime.SpinBox", EndTimeSpinStyle);
-
-		FEditableTextBoxStyle NonEditableTextBoxStyle = FEditableTextBoxStyle()
-			.SetPadding(FMargin(4.f))
-			.SetForegroundColor(FSlateColor::UseSubduedForeground())
-			.SetBackgroundImageNormal(BOX_BRUSH("Graph/CommonWidgets/TextBox", FMargin(4.0f/16.0f)))
-			.SetBackgroundImageHovered(BOX_BRUSH("Graph/CommonWidgets/TextBox_Hovered", FMargin(4.0f/16.0f)))
-			.SetBackgroundImageFocused(BOX_BRUSH("Graph/CommonWidgets/TextBox_Hovered", FMargin(4.0f/16.0f)))
-			.SetBackgroundImageReadOnly(FSlateRoundedBoxBrush(FStyleColors::Background, 4.0f, FStyleColors::InputOutline, 1.0f));
-		Set("Sequencer.NavigationTool.NonEditableTextBox", NonEditableTextBoxStyle);
-	}
 }
 
 void FStarshipEditorStyle::FStyle::SetupViewportStyles()
diff --git a/Engine/Source/Editor/Sequencer/Private/SSequencer.cpp b/Engine/Source/Editor/Sequencer/Private/SSequencer.cpp
--- a/Engine/Source/Editor/Sequencer/Private/SSequencer.cpp
+++ b/Engine/Source/Editor/Sequencer/Private/SSequencer.cpp
@@ -143,6 +143,7 @@
 #include "ContentBrowserModule.h"
 #include "Menus/SequencerToolbarUtils.h"
 #include "Settings/EditorExperimentalSettings.h"
+#include "Tools/NavigationTool/NavigationToolStyle.h"
 
 #define LOCTEXT_NAMESPACE "Sequencer"
 
@@ -1672,7 +1673,7 @@ void SSequencer::PopulateToolBar(UToolMenu* InMenu)
 		FToolMenuEntry ShowNavigationToolEntry = FToolMenuEntry::InitToolBarButton(
 			FSequencerCommands::Get().ToggleNavigationToolVisible,
 			TAttribute<FText>(),TAttribute<FText>(),
-			FSlateIcon(FAppStyle::GetAppStyleSetName(), TEXT("Sequencer.NavigationTool")));
+			FSlateIcon(FNavigationToolStyle::Get().GetStyleSetName(), TEXT("Icon.ToolBar")));
 		ShowNavigationToolEntry.StyleNameOverride = SequencerToolbarStyleName;
 		NavigationToolSection.AddEntry(ShowNavigationToolEntry);
 	}
diff --git a/Engine/Source/Editor/Sequencer/Private/Tools/NavigationTool/Menus/NavigationToolToolbarMenu.cpp b/Engine/Source/Editor/Sequencer/Private/Tools/NavigationTool/Menus/NavigationToolToolbarMenu.cpp
--- a/Engine/Source/Editor/Sequencer/Private/Tools/NavigationTool/Menus/NavigationToolToolbarMenu.cpp
+++ b/Engine/Source/Editor/Sequencer/Private/Tools/NavigationTool/Menus/NavigationToolToolbarMenu.cpp
@@ -257,11 +257,12 @@ void FNavigationToolToolbarMenu::CreateItemViewOptionsMenu(UToolMenu* const InMe
 	for (const TSharedPtr<FNavigationToolBuiltInFilter>& Filter : Tool->GlobalFilters)
 	{
 		const TSharedRef<FNavigationToolBuiltInFilter> FilterRef = Filter.ToSharedRef();
+		const FText DisplayName = FilterRef->GetDisplayName();
 
-		FToolMenuEntry& NewMenuEntry = GlobalFiltersSection.AddMenuEntry(NAME_None,
-			Filter->GetDisplayName(),
-			Filter->GetToolTipText(),
-			Filter->GetIcon(),
+		FToolMenuEntry& NewMenuEntry = GlobalFiltersSection.AddMenuEntry(*DisplayName.ToString(),
+			DisplayName,
+			FilterRef->GetToolTipText(),
+			FilterRef->GetIcon(),
 			FUIAction(
 				FExecuteAction::CreateSP(this, &FNavigationToolToolbarMenu::OnToggleGlobalFilter, FilterRef, ToolRef),
 				FCanExecuteAction(),
diff --git a/Engine/Source/Editor/Sequencer/Private/Tools/NavigationTool/NavigationToolCommands.cpp b/Engine/Source/Editor/Sequencer/Private/Tools/NavigationTool/NavigationToolCommands.cpp
--- a/Engine/Source/Editor/Sequencer/Private/Tools/NavigationTool/NavigationToolCommands.cpp
+++ b/Engine/Source/Editor/Sequencer/Private/Tools/NavigationTool/NavigationToolCommands.cpp
@@ -1,6 +1,7 @@
 // Copyright Epic Games, Inc. All Rights Reserved.
 
 #include "NavigationToolCommands.h"
+#include "Styling/AppStyle.h"
 
 #define LOCTEXT_NAMESPACE "NavigationToolCommands"
 
diff --git a/Engine/Source/Editor/Sequencer/Private/Tools/NavigationTool/NavigationToolCommands.h b/Engine/Source/Editor/Sequencer/Private/Tools/NavigationTool/NavigationToolCommands.h
--- a/Engine/Source/Editor/Sequencer/Private/Tools/NavigationTool/NavigationToolCommands.h
+++ b/Engine/Source/Editor/Sequencer/Private/Tools/NavigationTool/NavigationToolCommands.h
@@ -3,7 +3,6 @@
 #pragma once
 
 #include "Framework/Commands/Commands.h"
-#include "Styling/AppStyle.h"
 
 namespace UE::Sequencer::NavigationTool
 {
diff --git a/Engine/Source/Editor/Sequencer/Private/Tools/NavigationTool/NavigationToolStyle.cpp b/Engine/Source/Editor/Sequencer/Private/Tools/NavigationTool/NavigationToolStyle.cpp
new file mode 100644
--- /dev/null
+++ b/Engine/Source/Editor/Sequencer/Private/Tools/NavigationTool/NavigationToolStyle.cpp
@@ -0,0 +1,59 @@
+﻿// Copyright Epic Games, Inc. All Rights Reserved.
+
+#include "NavigationToolStyle.h"
+#include "Brushes/SlateImageBrush.h"
+#include "Brushes/SlateRoundedBoxBrush.h"
+#include "Fonts/SlateFontInfo.h"
+#include "Layout/Margin.h"
+#include "Misc/Paths.h"
+#include "Styling/AppStyle.h"
+#include "Styling/CoreStyle.h"
+#include "Styling/StyleColors.h"
+#include "Styling/SlateStyleMacros.h"
+#include "Styling/SlateStyleRegistry.h"
+#include "Styling/SlateTypes.h"
+#include "Styling/ToolBarStyle.h"
+
+FNavigationToolStyle::FNavigationToolStyle()
+	: FSlateStyleSet(TEXT("NavigationToolStyle"))
+{
+	FSlateStyleSet::SetContentRoot(FPaths::EngineContentDir() / TEXT("Editor/Slate"));
+	FSlateStyleSet::SetCoreContentRoot(FPaths::EngineContentDir() / TEXT("Slate"));
+
+	static const FVector2D Icon16x16(16.0f, 16.0f);
+	static const FVector2D Icon20x20(20.0f, 20.0f);
+
+	Set("Icon.ToolBar", new IMAGE_BRUSH_SVG("Sequencer/SequenceNavigator/SequenceNavigator_20", Icon20x20));
+	Set("Icon.Tab", new IMAGE_BRUSH_SVG("Sequencer/SequenceNavigator/SequenceNavigator_16", Icon16x16));
+
+	const FTableRowStyle& AlternatingTableRowStyle = FAppStyle::GetWidgetStyle<FTableRowStyle>(TEXT("TableView.AlternatingRow"));
+	Set("TableViewRow", FTableRowStyle(AlternatingTableRowStyle));
+
+	FSpinBoxStyle FrameTimeSpinStyle = FAppStyle::GetWidgetStyle<FSpinBoxStyle>(TEXT("SpinBox"));
+	FrameTimeSpinStyle.ForegroundColor = FStyleColors::Foreground;
+	Set("SpinBox", FrameTimeSpinStyle);
+
+	FSpinBoxStyle StartTimeSpinStyle = FrameTimeSpinStyle
+		.SetForegroundColor(FStyleColors::AccentGreen);
+	Set("SpinBox.InTime", StartTimeSpinStyle);
+
+	FSpinBoxStyle EndTimeSpinStyle = FrameTimeSpinStyle
+		.SetForegroundColor(FStyleColors::AccentRed);
+	Set("SpinBox.OutTime", EndTimeSpinStyle);
+
+	FEditableTextBoxStyle NonEditableTextBoxStyle = FEditableTextBoxStyle()
+		.SetPadding(FMargin(4.f))
+		.SetForegroundColor(FSlateColor::UseSubduedForeground())
+		.SetBackgroundImageNormal(BOX_BRUSH("Graph/CommonWidgets/TextBox", FMargin(4.0f/16.0f)))
+		.SetBackgroundImageHovered(BOX_BRUSH("Graph/CommonWidgets/TextBox_Hovered", FMargin(4.0f/16.0f)))
+		.SetBackgroundImageFocused(BOX_BRUSH("Graph/CommonWidgets/TextBox_Hovered", FMargin(4.0f/16.0f)))
+		.SetBackgroundImageReadOnly(FSlateRoundedBoxBrush(FStyleColors::Background, 4.0f, FStyleColors::InputOutline, 1.0f));
+	Set("NonEditableTextBox", NonEditableTextBoxStyle);
+
+	FSlateStyleRegistry::RegisterSlateStyle(*this);
+}
+
+FNavigationToolStyle::~FNavigationToolStyle()
+{
+	FSlateStyleRegistry::UnRegisterSlateStyle(*this);
+}
diff --git a/Engine/Source/Editor/Sequencer/Private/Tools/NavigationTool/NavigationToolStyle.h b/Engine/Source/Editor/Sequencer/Private/Tools/NavigationTool/NavigationToolStyle.h
new file mode 100644
--- /dev/null
+++ b/Engine/Source/Editor/Sequencer/Private/Tools/NavigationTool/NavigationToolStyle.h
@@ -0,0 +1,18 @@
+﻿// Copyright Epic Games, Inc. All Rights Reserved.
+
+#pragma once
+
+#include "Styling/SlateStyle.h"
+
+class FNavigationToolStyle final : public FSlateStyleSet
+{
+public:
+	static FNavigationToolStyle& Get()
+	{
+		static FNavigationToolStyle Instance;
+		return Instance;
+	}
+
+	FNavigationToolStyle();
+	virtual ~FNavigationToolStyle() override;
+};
diff --git a/Engine/Source/Editor/Sequencer/Private/Tools/NavigationTool/NavigationToolTab.cpp b/Engine/Source/Editor/Sequencer/Private/Tools/NavigationTool/NavigationToolTab.cpp
--- a/Engine/Source/Editor/Sequencer/Private/Tools/NavigationTool/NavigationToolTab.cpp
+++ b/Engine/Source/Editor/Sequencer/Private/Tools/NavigationTool/NavigationToolTab.cpp
@@ -7,6 +7,7 @@
 #include "Toolkits/IToolkitHost.h"
 #include "Tools/NavigationTool/INavigationTool.h"
 #include "Tools/NavigationTool/INavigationToolView.h"
+#include "Tools/NavigationTool/NavigationToolStyle.h"
 #include "WorkspaceMenuStructure.h"
 #include "WorkspaceMenuStructureModule.h"
 
@@ -154,7 +155,7 @@ void FNavigationToolTab::RegisterToolTab()
 	TabManager->RegisterTabSpawner(NavigationToolTabId, FOnSpawnTab::CreateSP(this, &FNavigationToolTab::SpawnToolTab))
 		.SetGroup(WorkspaceMenu::GetMenuStructure().GetLevelEditorCinematicsCategory())
 		.SetDisplayName(LOCTEXT("NavigationToolTab", "Sequence Nav"))
-		.SetIcon(FSlateIcon(FAppStyle::GetAppStyleSetName(), TEXT("Sequencer.NavigationTool.TabIcon")));
+		.SetIcon(FSlateIcon(FNavigationToolStyle::Get().GetStyleSetName(), TEXT("Icon.Tab")));
 }
 
 void FNavigationToolTab::UnregisterToolTab()
diff --git a/Engine/Source/Editor/Sequencer/Private/Tools/NavigationTool/Widgets/Columns/SNavigationToolEndFrameOffset.cpp b/Engine/Source/Editor/Sequencer/Private/Tools/NavigationTool/Widgets/Columns/SNavigationToolEndFrameOffset.cpp
--- a/Engine/Source/Editor/Sequencer/Private/Tools/NavigationTool/Widgets/Columns/SNavigationToolEndFrameOffset.cpp
+++ b/Engine/Source/Editor/Sequencer/Private/Tools/NavigationTool/Widgets/Columns/SNavigationToolEndFrameOffset.cpp
@@ -38,11 +38,6 @@ void SNavigationToolEndFrameOffset::Construct(const FArguments& InArgs
 	SNavigationToolTime::Construct(SNavigationToolTime::FArguments(), InItem, InView, InRowWidget);
 }
 
-FName SNavigationToolEndFrameOffset::GetStyleName() const
-{
-	return TEXT("Sequencer.NavigationTool.SpinBox");
-}
-
 double SNavigationToolEndFrameOffset::GetFrameTimeValue() const
 {
 	UMovieSceneSubSection* const SubSection = Cast<UMovieSceneSubSection>(GetSequenceItemSubSection(WeakItem.Pin()));
diff --git a/Engine/Source/Editor/Sequencer/Private/Tools/NavigationTool/Widgets/Columns/SNavigationToolEndFrameOffset.h b/Engine/Source/Editor/Sequencer/Private/Tools/NavigationTool/Widgets/Columns/SNavigationToolEndFrameOffset.h
--- a/Engine/Source/Editor/Sequencer/Private/Tools/NavigationTool/Widgets/Columns/SNavigationToolEndFrameOffset.h
+++ b/Engine/Source/Editor/Sequencer/Private/Tools/NavigationTool/Widgets/Columns/SNavigationToolEndFrameOffset.h
@@ -25,8 +25,6 @@ public:
 		, const TSharedRef<SNavigationToolTreeRow>& InRowWidget);
 
 private:
-	virtual FName GetStyleName() const override;
-
 	virtual double GetFrameTimeValue() const override;
 
 	virtual void OnFrameTimeValueCommitted(const double InNewValue, const ETextCommit::Type InCommitType) override;
diff --git a/Engine/Source/Editor/Sequencer/Private/Tools/NavigationTool/Widgets/Columns/SNavigationToolHBias.cpp b/Engine/Source/Editor/Sequencer/Private/Tools/NavigationTool/Widgets/Columns/SNavigationToolHBias.cpp
--- a/Engine/Source/Editor/Sequencer/Private/Tools/NavigationTool/Widgets/Columns/SNavigationToolHBias.cpp
+++ b/Engine/Source/Editor/Sequencer/Private/Tools/NavigationTool/Widgets/Columns/SNavigationToolHBias.cpp
@@ -4,7 +4,7 @@
 #include "ScopedTransaction.h"
 #include "Tools/NavigationTool/Items/NavigationToolItemUtils.h"
 #include "Tools/NavigationTool/Items/NavigationToolSequence.h"
-#include "Tools/NavigationTool/NavigationTool.h"
+#include "Tools/NavigationTool/NavigationToolStyle.h"
 #include "Tools/NavigationTool/NavigationToolView.h"
 #include "Tools/NavigationTool/Widgets/SNavigationToolTreeRow.h"
 #include "TrackEditors/SubTrackEditorBase.h"
@@ -56,7 +56,7 @@ void SNavigationToolHBias::Construct(const FArguments& InArgs
 		[
 			SNew(SSpinBox<int32>)
 			.Justification(ETextJustify::Center)
-			.Style(&FAppStyle::Get(), TEXT("Sequencer.NavigationTool.SpinBox"))
+			.Style(&FNavigationToolStyle::Get(), TEXT("SpinBox"))
 			.Value(this, &SNavigationToolHBias::GetValue)
 			.OnValueChanged(this, &SNavigationToolHBias::OnValueChanged)
 			.OnValueCommitted(this, &SNavigationToolHBias::OnValueCommitted)
diff --git a/Engine/Source/Editor/Sequencer/Private/Tools/NavigationTool/Widgets/Columns/SNavigationToolInTime.cpp b/Engine/Source/Editor/Sequencer/Private/Tools/NavigationTool/Widgets/Columns/SNavigationToolInTime.cpp
--- a/Engine/Source/Editor/Sequencer/Private/Tools/NavigationTool/Widgets/Columns/SNavigationToolInTime.cpp
+++ b/Engine/Source/Editor/Sequencer/Private/Tools/NavigationTool/Widgets/Columns/SNavigationToolInTime.cpp
@@ -39,7 +39,7 @@ IInTimeExtension* SNavigationToolInTime::GetInTimeExtension() const
 
 FName SNavigationToolInTime::GetStyleName() const
 {
-	return TEXT("Sequencer.NavigationTool.InTime.SpinBox");
+	return TEXT("SpinBox.InTime");
 }
 
 double SNavigationToolInTime::GetFrameTimeValue() const
diff --git a/Engine/Source/Editor/Sequencer/Private/Tools/NavigationTool/Widgets/Columns/SNavigationToolLength.cpp b/Engine/Source/Editor/Sequencer/Private/Tools/NavigationTool/Widgets/Columns/SNavigationToolLength.cpp
--- a/Engine/Source/Editor/Sequencer/Private/Tools/NavigationTool/Widgets/Columns/SNavigationToolLength.cpp
+++ b/Engine/Source/Editor/Sequencer/Private/Tools/NavigationTool/Widgets/Columns/SNavigationToolLength.cpp
@@ -38,11 +38,6 @@ void SNavigationToolLength::Construct(const FArguments& InArgs
 	SNavigationToolTime::Construct(SNavigationToolTime::FArguments(), InItem, InView, InRowWidget);
 }
 
-FName SNavigationToolLength::GetStyleName() const
-{
-	return TEXT("Sequencer.NavigationTool.SpinBox");
-}
-
 double SNavigationToolLength::GetFrameTimeValue() const
 {
 	UMovieSceneSubSection* const SubSection = Cast<UMovieSceneSubSection>(GetSequenceItemSubSection(WeakItem.Pin()));
diff --git a/Engine/Source/Editor/Sequencer/Private/Tools/NavigationTool/Widgets/Columns/SNavigationToolLength.h b/Engine/Source/Editor/Sequencer/Private/Tools/NavigationTool/Widgets/Columns/SNavigationToolLength.h
--- a/Engine/Source/Editor/Sequencer/Private/Tools/NavigationTool/Widgets/Columns/SNavigationToolLength.h
+++ b/Engine/Source/Editor/Sequencer/Private/Tools/NavigationTool/Widgets/Columns/SNavigationToolLength.h
@@ -25,13 +25,11 @@ public:
 		, const TSharedRef<SNavigationToolTreeRow>& InRowWidget);
 
 private:
-	virtual FName GetStyleName() const override;
-
 	virtual double GetFrameTimeValue() const override;
 
 	virtual FText GetTransactionText() const override;
 
-	virtual TSharedPtr<INumericTypeInterface<double>> GetNumericTypeInterface() const;
+	virtual TSharedPtr<INumericTypeInterface<double>> GetNumericTypeInterface() const override;
 };
 
 }
diff --git a/Engine/Source/Editor/Sequencer/Private/Tools/NavigationTool/Widgets/Columns/SNavigationToolOutTime.cpp b/Engine/Source/Editor/Sequencer/Private/Tools/NavigationTool/Widgets/Columns/SNavigationToolOutTime.cpp
--- a/Engine/Source/Editor/Sequencer/Private/Tools/NavigationTool/Widgets/Columns/SNavigationToolOutTime.cpp
+++ b/Engine/Source/Editor/Sequencer/Private/Tools/NavigationTool/Widgets/Columns/SNavigationToolOutTime.cpp
@@ -39,7 +39,7 @@ IOutTimeExtension* SNavigationToolOutTime::GetOutTimeExtension() const
 
 FName SNavigationToolOutTime::GetStyleName() const
 {
-	return TEXT("Sequencer.NavigationTool.OutTime.SpinBox");
+	return TEXT("SpinBox.OutTime");
 }
 
 double SNavigationToolOutTime::GetFrameTimeValue() const
diff --git a/Engine/Source/Editor/Sequencer/Private/Tools/NavigationTool/Widgets/Columns/SNavigationToolStartFrameOffset.cpp b/Engine/Source/Editor/Sequencer/Private/Tools/NavigationTool/Widgets/Columns/SNavigationToolStartFrameOffset.cpp
--- a/Engine/Source/Editor/Sequencer/Private/Tools/NavigationTool/Widgets/Columns/SNavigationToolStartFrameOffset.cpp
+++ b/Engine/Source/Editor/Sequencer/Private/Tools/NavigationTool/Widgets/Columns/SNavigationToolStartFrameOffset.cpp
@@ -37,11 +37,6 @@ void SNavigationToolStartFrameOffset::Construct(const FArguments& InArgs
 	SNavigationToolTime::Construct(SNavigationToolTime::FArguments(), InItem, InView, InRowWidget);
 }
 
-FName SNavigationToolStartFrameOffset::GetStyleName() const
-{
-	return TEXT("Sequencer.NavigationTool.SpinBox");
-}
-
 double SNavigationToolStartFrameOffset::GetFrameTimeValue() const
 {
 	const UMovieSceneSubSection* const SubSection = Cast<UMovieSceneSubSection>(GetSequenceItemSubSection(WeakItem.Pin()));
diff --git a/Engine/Source/Editor/Sequencer/Private/Tools/NavigationTool/Widgets/Columns/SNavigationToolStartFrameOffset.h b/Engine/Source/Editor/Sequencer/Private/Tools/NavigationTool/Widgets/Columns/SNavigationToolStartFrameOffset.h
--- a/Engine/Source/Editor/Sequencer/Private/Tools/NavigationTool/Widgets/Columns/SNavigationToolStartFrameOffset.h
+++ b/Engine/Source/Editor/Sequencer/Private/Tools/NavigationTool/Widgets/Columns/SNavigationToolStartFrameOffset.h
@@ -25,8 +25,6 @@ public:
 		, const TSharedRef<SNavigationToolTreeRow>& InRowWidget);
 
 private:
-	virtual FName GetStyleName() const override;
-
 	virtual double GetFrameTimeValue() const override;
 
 	virtual void OnFrameTimeValueCommitted(const double InNewValue, const ETextCommit::Type InCommitType) override;
diff --git a/Engine/Source/Editor/Sequencer/Private/Tools/NavigationTool/Widgets/Columns/SNavigationToolTime.cpp b/Engine/Source/Editor/Sequencer/Private/Tools/NavigationTool/Widgets/Columns/SNavigationToolTime.cpp
--- a/Engine/Source/Editor/Sequencer/Private/Tools/NavigationTool/Widgets/Columns/SNavigationToolTime.cpp
+++ b/Engine/Source/Editor/Sequencer/Private/Tools/NavigationTool/Widgets/Columns/SNavigationToolTime.cpp
@@ -3,6 +3,7 @@
 #include "SNavigationToolTime.h"
 #include "Sequencer.h"
 #include "Tools/NavigationTool/Items/NavigationToolItemUtils.h"
+#include "Tools/NavigationTool/NavigationToolStyle.h"
 #include "Tools/NavigationTool/NavigationToolView.h"
 #include "Tools/NavigationTool/Widgets/SNavigationToolTreeRow.h"
 #include "Widgets/Input/SSpinBox.h"
@@ -32,7 +33,7 @@ void SNavigationToolTime::Construct(const FArguments& InArgs
 	[
 		SNew(SSpinBox<double>)
 		.Justification(ETextJustify::Center)
-		.Style(&FAppStyle::GetWidgetStyle<FSpinBoxStyle>(GetStyleName()))
+		.Style(&FNavigationToolStyle::Get().GetWidgetStyle<FSpinBoxStyle>(GetStyleName()))
 		.TypeInterface(GetNumericTypeInterface())
 		.Value(this, &SNavigationToolTime::GetFrameTimeValue)
 		.OnValueChanged(this, &SNavigationToolTime::OnFrameTimeValueChanged)
@@ -55,7 +56,7 @@ void SNavigationToolTime::SetIsReadOnly(const bool bInIsReadOnly)
 
 FName SNavigationToolTime::GetStyleName() const
 {
-	return TEXT("Sequencer.NavigationTool.SpinBox");
+	return TEXT("SpinBox");
 }
 
 void SNavigationToolTime::OnFrameTimeValueChanged(const double InNewValue)
diff --git a/Engine/Source/Editor/Sequencer/Private/Tools/NavigationTool/Widgets/SNavigationToolTreeRow.cpp b/Engine/Source/Editor/Sequencer/Private/Tools/NavigationTool/Widgets/SNavigationToolTreeRow.cpp
--- a/Engine/Source/Editor/Sequencer/Private/Tools/NavigationTool/Widgets/SNavigationToolTreeRow.cpp
+++ b/Engine/Source/Editor/Sequencer/Private/Tools/NavigationTool/Widgets/SNavigationToolTreeRow.cpp
@@ -1,9 +1,9 @@
 // Copyright Epic Games, Inc. All Rights Reserved.
 
 #include "Tools/NavigationTool/Widgets/SNavigationToolTreeRow.h"
-#include "Tools/NavigationTool/NavigationTool.h"
 #include "Tools/NavigationTool/Columns/INavigationToolColumn.h"
 #include "Tools/NavigationTool/Items/INavigationToolItem.h"
+#include "Tools/NavigationTool/NavigationToolStyle.h"
 #include "Tools/NavigationTool/NavigationToolView.h"
 #include "Tools/NavigationTool/Widgets/SNavigationToolTreeView.h"
 
@@ -25,7 +25,7 @@ void SNavigationToolTreeRow::Construct(const FArguments& InArgs
 	SetColorAndOpacity(TAttribute<FLinearColor>::CreateSP(&*InToolView, &FNavigationToolView::GetItemBrushColor, Item));
 
 	SMultiColumnTableRow::Construct(FSuperRowType::FArguments()
-			.Style(&FAppStyle::Get().GetWidgetStyle<FTableRowStyle>(TEXT("Sequencer.NavigationTool.TableViewRow")))
+			.Style(&FNavigationToolStyle::Get().GetWidgetStyle<FTableRowStyle>(TEXT("TableViewRow")))
 			.OnCanAcceptDrop(InToolView, &FNavigationToolView::OnCanDrop)
 			.OnDragDetected(InToolView, &FNavigationToolView::OnDragDetected, Item)
 			.OnDragEnter(InToolView, &FNavigationToolView::OnDragEnter, Item)
