Sequencer: Fix force whole frame case when resizing sections

#rnx

diff --git a/Engine/Source/Editor/Sequencer/Private/Tools/EditToolDragOperations.cpp b/Engine/Source/Editor/Sequencer/Private/Tools/EditToolDragOperations.cpp
--- a/Engine/Source/Editor/Sequencer/Private/Tools/EditToolDragOperations.cpp
+++ b/Engine/Source/Editor/Sequencer/Private/Tools/EditToolDragOperations.cpp
@@ -178,8 +178,6 @@ void FResizeSection::OnBeginDrag(const FPointerEvent& MouseEvent, FVector2D Loca
 	TSet<FSequencerSelectedKey> EmptyKeySet;
 	FInvalidKeyAndSectionSnappingCandidates SnapCandidates(EmptyKeySet, Sections);
 	SnapField = FSequencerSnapField(Sequencer, SnapCandidates, ESequencerEntity::Section | ESequencerEntity::Key);
-	SnapField.GetValue().SetSnapToInterval(Sequencer.GetSequencerSettings()->GetForceWholeFrames());
-	SnapField.GetValue().SetSnapToLikeTypes(Sequencer.GetSequencerSettings()->GetSnapSectionTimesToSections());
 
 	SectionInitTimes.Empty();
 
@@ -293,7 +291,13 @@ void FResizeSection::OnDrag(const FPointerEvent& MouseEvent, FVector2D LocalMous
 	FFrameNumber DeltaTime = (VirtualTrackArea.PixelToFrame(LocalMousePos.X) - MouseDownTime).RoundToFrame();
 
 	// Snapping
-	if ((Settings->GetIsSnapEnabled() || MouseEvent.IsShiftDown()) && Settings->GetSnapSectionTimesToSections())
+	const bool bSnapToInterval = Settings->GetForceWholeFrames();
+	const bool bSnapToLikeTypes = ((Settings->GetIsSnapEnabled() || MouseEvent.IsShiftDown()) && Settings->GetSnapSectionTimesToSections());
+	
+	SnapField.GetValue().SetSnapToInterval(bSnapToInterval);
+	SnapField.GetValue().SetSnapToLikeTypes(bSnapToLikeTypes);
+
+	if (bSnapToInterval || bSnapToLikeTypes)
 	{
 		TArray<FFrameTime> SectionTimes;
 		for (UMovieSceneSection* Section : Sections)
@@ -604,13 +608,9 @@ void FManipulateSectionEasing::OnBeginDrag(const FPointerEvent& MouseEvent, FVec
 
 	MouseDownTime = VirtualTrackArea.PixelToFrame(LocalMousePos.X);
 
-	if (Settings->GetSnapSectionTimesToSections())
-	{
-		// Construct a snap field of all section bounds
-		ISnapCandidate SnapCandidates;
-		SnapField = FSequencerSnapField(Sequencer, SnapCandidates, ESequencerEntity::Section);
-		SnapField.GetValue().SetSnapToInterval(Sequencer.GetSequencerSettings()->GetForceWholeFrames());
-	}
+	// Construct a snap field of all section bounds
+	ISnapCandidate SnapCandidates;
+	SnapField = FSequencerSnapField(Sequencer, SnapCandidates, ESequencerEntity::Section);
 
 	InitValue = bEaseIn ? Section->Easing.GetEaseInDuration() : Section->Easing.GetEaseOutDuration();
 }
@@ -634,7 +634,13 @@ void FManipulateSectionEasing::OnDrag(const FPointerEvent& MouseEvent, FVector2D
 	}
 
 	// Snapping
-	if ((Settings->GetIsSnapEnabled() || MouseEvent.IsShiftDown()) && Settings->GetSnapSectionTimesToSections())
+	const bool bSnapToInterval = Settings->GetForceWholeFrames();
+	const bool bSnapToLikeTypes = ((Settings->GetIsSnapEnabled() || MouseEvent.IsShiftDown()) && Settings->GetSnapSectionTimesToSections());
+
+	SnapField.GetValue().SetSnapToInterval(bSnapToInterval);
+	SnapField.GetValue().SetSnapToLikeTypes(bSnapToLikeTypes);
+
+	if (bSnapToInterval || bSnapToLikeTypes)
 	{
 		TArray<FFrameTime> SnapTimes;
 		if (bEaseIn)
