Sequencer: Fix crash pasting an object binding with a null spawnable template

#rnx

diff --git a/Engine/Source/Editor/Sequencer/Private/SequencerUtilities.cpp b/Engine/Source/Editor/Sequencer/Private/SequencerUtilities.cpp
--- a/Engine/Source/Editor/Sequencer/Private/SequencerUtilities.cpp
+++ b/Engine/Source/Editor/Sequencer/Private/SequencerUtilities.cpp
@@ -3164,8 +3164,12 @@ bool FSequencerUtilities::PasteBindings(const FString& TextToImport, TSharedRef<
 								{
 									if (CopyableBinding->SpawnableObjectTemplates.IsValidIndex(SpawnableBindingIndex))
 									{
-										UObject* NewObjectTemplate = StaticDuplicateObject(CopyableBinding->SpawnableObjectTemplates[SpawnableBindingIndex++], MovieScene);
-										SpawnableBinding->SetObjectTemplate(NewObjectTemplate);
+										UObject* SpawnableObjectTemplate = CopyableBinding->SpawnableObjectTemplates[SpawnableBindingIndex++];
+										if (SpawnableObjectTemplate)
+										{
+											UObject* NewObjectTemplate = StaticDuplicateObject(SpawnableObjectTemplate, MovieScene);
+											SpawnableBinding->SetObjectTemplate(NewObjectTemplate);
+										}
 									}
 
 									// If this is a preview spawnable, find the just added replaceable and link with that rather than creating a new binding
