Sequencer: Fix paste multiple sections from different rows onto the same multiple rows

diff --git a/Engine/Source/Editor/Sequencer/Private/SequencerUtilities.cpp b/Engine/Source/Editor/Sequencer/Private/SequencerUtilities.cpp
--- a/Engine/Source/Editor/Sequencer/Private/SequencerUtilities.cpp
+++ b/Engine/Source/Editor/Sequencer/Private/SequencerUtilities.cpp
@@ -2232,7 +2232,21 @@ bool FSequencerUtilities::PasteSections(const FString& TextToImport, FMovieScene
 				continue;
 			}
 
-			int32 RowIndex = SectionIndex < PasteSectionsParams.TrackRowIndices.Num() ? PasteSectionsParams.TrackRowIndices[SectionIndex] : Section->GetRowIndex();
+			int32 RowIndex = Section->GetRowIndex();
+
+			// If there is only 1 track to paste onto, paste the sections all onto the that track's row index
+			if (PasteSectionsParams.TrackRowIndices.Num() == 1)
+			{
+				RowIndex = PasteSectionsParams.TrackRowIndices[0];
+			}
+			// Otherwise if pasting onto multiple track rows, paste onto the same row index as the copied section
+			else if (Index < PasteSectionsParams.TrackRowIndices.Num())
+			{
+				if (PasteSectionsParams.TrackRowIndices[Index] != Section->GetRowIndex())
+				{
+					continue;
+				}
+			}
 
 			Track->Modify();
 
@@ -2300,23 +2314,10 @@ bool FSequencerUtilities::PasteSections(const FString& TextToImport, FMovieScene
 			}
 		}
 
-		// Regenerate for pasting onto the next track 
-		ImportedSections.Empty();
-		ImportedTracks.Empty();
-
-		ImportTracksFromText(TextToImport, ImportedTracks);
-
-		for (UMovieSceneCopyableTrack* CopyableTrack : ImportedTracks)
-		{
-			for (UMovieSceneSection* CopyableSection : CopyableTrack->Track->GetAllSections())
-			{
-				ImportedSections.Add(CopyableSection);
-			}
-		}
-
-		if (ImportedSections.Num() == 0)
+		// Remove sections that were pasted so that they aren't pasted again to another track
+		for (UMovieSceneSection* OutSection : OutSections)
 		{
-			ImportSectionsFromText(TextToImport, ImportedSections);
+			ImportedSections.Remove(OutSection);
 		}
 	}
 
