[Sequencer] Update global marker cache when marker frame is changed

diff --git a/Engine/Source/Editor/Sequencer/Private/Tools/NavigationTool/Items/NavigationToolMarker.cpp b/Engine/Source/Editor/Sequencer/Private/Tools/NavigationTool/Items/NavigationToolMarker.cpp
--- a/Engine/Source/Editor/Sequencer/Private/Tools/NavigationTool/Items/NavigationToolMarker.cpp
+++ b/Engine/Source/Editor/Sequencer/Private/Tools/NavigationTool/Items/NavigationToolMarker.cpp
@@ -270,10 +270,13 @@ void FNavigationToolMarker::SetInTime(const FFrameNumber& InTime)
 		return;
 	}
 
-	const FScopedTransaction Transaction(LOCTEXT("SetMarkedFrameTimeTransaction", "Set Marked Frame Time"));
-
 	MovieScene->Modify();
 	MarkedFrame.FrameNumber = InTime;
+
+	if (const TSharedPtr<FSequencer> Sequencer = StaticCastSharedPtr<FSequencer>(GetOwnerTool().GetSequencer()))
+	{
+		Sequencer->InvalidateGlobalMarkedFramesCache();
+	}
 }
 
 int32 FNavigationToolMarker::GetMarkedFrameIndex() const
