[Sequencer] Pin/Lock/Deactivate/Mute/Solo hotkey to work on all descendants

diff --git a/Engine/Source/Editor/Sequencer/Private/Sequencer.cpp b/Engine/Source/Editor/Sequencer/Private/Sequencer.cpp
--- a/Engine/Source/Editor/Sequencer/Private/Sequencer.cpp
+++ b/Engine/Source/Editor/Sequencer/Private/Sequencer.cpp
@@ -12767,21 +12767,39 @@ void FSequencer::ToggleTrackSelectionPin()
 
 	FOutlinerSelection& OutlinerSelection = GetSelection().Outliner;
 
-	const bool bAllPinned = Algo::AllOf(OutlinerSelection, [](const FViewModelPtr& Selection)
+	const bool bAllSame = Algo::AllOf(OutlinerSelection, [](const FViewModelPtr& InSelection)
 		{
-			const TViewModelPtr<IPinnableExtension> Pinnable = Selection.ImplicitCast();
-			return Pinnable.IsValid() && Pinnable->IsPinned();
+			for (const TViewModelPtr<IPinnableExtension>& Extension : InSelection->GetDescendantsOfType<IPinnableExtension>(true))
+			{
+				if (Extension.IsValid() && Extension->IsPinned())
+				{
+					return true;
+				}
+			}
+			return false;
 		});
 
-	const bool bNewPinState = !bAllPinned;
+	const bool bNewState = !bAllSame;
+
+	bool bAnyChanged = false;
 
-	for (const TViewModelPtr Pinnable : OutlinerSelection.Filter<IPinnableExtension>())
+	for (const TViewModelPtr<IOutlinerExtension>& OutlinerItem : OutlinerSelection)
 	{
-		Pinnable->SetPinned(bNewPinState);
-		NodeTree->SavePinnedState(*Pinnable.AsModel(), bNewPinState);
+		for (const TViewModelPtr<IPinnableExtension>& Extension : OutlinerItem.AsModel()->GetDescendantsOfType<IPinnableExtension>(true))
+		{
+			if (Extension->IsPinned() != bNewState)
+			{
+				Extension->SetPinned(bNewState);
+				NodeTree->SavePinnedState(*Extension.AsModel(), bNewState);
+				bAnyChanged = true;
+			}
+		}
 	}
 
-	RefreshTree();
+	if (bAnyChanged)
+	{
+		RefreshTree();
+	}
 }
 
 void FSequencer::ToggleTrackSelectionLock()
@@ -12790,20 +12808,39 @@ void FSequencer::ToggleTrackSelectionLock()
 
 	FOutlinerSelection& OutlinerSelection = GetSelection().Outliner;
 
-	const bool bAllLocked = Algo::AllOf(OutlinerSelection, [](const FViewModelPtr& Selection)
+	const bool bAllSame = Algo::AllOf(OutlinerSelection, [](const FViewModelPtr& InSelection)
 		{
-			const TViewModelPtr<ILockableExtension> Lockable = Selection.ImplicitCast();
-			return Lockable.IsValid() && Lockable->GetLockState() != ELockableLockState::None;
+			for (const TViewModelPtr<ILockableExtension>& Extension : InSelection->GetDescendantsOfType<ILockableExtension>(true))
+			{
+				if (Extension.IsValid() && Extension->GetLockState() != ELockableLockState::None)
+				{
+					return true;
+				}
+			}
+			return false;
 		});
 
-	const bool bNewLockState = !bAllLocked;
+	const bool bNewState = !bAllSame;
+
+	bool bAnyChanged = false;
 
-	for (const TViewModelPtr Lockable : OutlinerSelection.Filter<ILockableExtension>())
+	for (const TViewModelPtr<IOutlinerExtension>& OutlinerItem : OutlinerSelection)
 	{
-		Lockable->SetIsLocked(bNewLockState);
+		for (const TViewModelPtr<ILockableExtension>& Extension : OutlinerItem.AsModel()->GetDescendantsOfType<ILockableExtension>(true))
+        {
+			const bool bThisState = Extension->GetLockState() == ELockableLockState::Locked;
+			if (bThisState != bNewState)
+			{
+				Extension->SetIsLocked(bNewState);
+				bAnyChanged = true;
+			}
+        }
 	}
 
-	RefreshTree();
+	if (bAnyChanged)
+	{
+		RefreshTree();
+	}
 }
 
 void FSequencer::ToggleTrackSelectionDeactive()
@@ -12812,20 +12849,38 @@ void FSequencer::ToggleTrackSelectionDeactive()
 
 	FOutlinerSelection& OutlinerSelection = GetSelection().Outliner;
 
-	const bool bAllDeactived = Algo::AllOf(OutlinerSelection, [](const FViewModelPtr& Selection)
+	const bool bAllSame = Algo::AllOf(OutlinerSelection, [](const FViewModelPtr& InSelection)
 		{
-			const TViewModelPtr<IDeactivatableExtension> Deactivatable = Selection.ImplicitCast();
-			return Deactivatable.IsValid() && Deactivatable->IsDeactivated();
+			for (const TViewModelPtr<IDeactivatableExtension>& Extension : InSelection->GetDescendantsOfType<IDeactivatableExtension>(true))
+			{
+				if (Extension.IsValid() && Extension->IsDeactivated())
+				{
+					return true;
+				}
+			}
+			return false;
 		});
 
-	const bool bNewDeactiveState = !bAllDeactived;
+	const bool bNewState = !bAllSame;
+
+	bool bAnyChanged = false;
 
-	for (const TViewModelPtr Deactivatable : OutlinerSelection.Filter<IDeactivatableExtension>())
+	for (const TViewModelPtr<IOutlinerExtension>& OutlinerItem : OutlinerSelection)
 	{
-		Deactivatable->SetIsDeactivated(bNewDeactiveState);
+		for (const TViewModelPtr<IDeactivatableExtension>& Extension : OutlinerItem.AsModel()->GetDescendantsOfType<IDeactivatableExtension>(true))
+		{
+			if (Extension->IsDeactivated() != bNewState)
+			{
+				Extension->SetIsDeactivated(bNewState);
+				bAnyChanged = true;
+			}
+		}
 	}
 
-	RefreshTree();
+	if (bAnyChanged)
+	{
+		RefreshTree();
+	}
 }
 
 void FSequencer::ToggleTrackSelectionMute()
@@ -12834,20 +12889,38 @@ void FSequencer::ToggleTrackSelectionMute()
 
 	FOutlinerSelection& OutlinerSelection = GetSelection().Outliner;
 
-	const bool bAllMuted = Algo::AllOf(OutlinerSelection, [](const FViewModelPtr& Selection)
+	const bool bAllSame = Algo::AllOf(OutlinerSelection, [](const FViewModelPtr& InSelection)
 		{
-			const TViewModelPtr<IMutableExtension> Mutable = Selection.ImplicitCast();
-			return Mutable.IsValid() && Mutable->IsMuted();
+			for (const TViewModelPtr<IMutableExtension>& Extension : InSelection->GetDescendantsOfType<IMutableExtension>(true))
+			{
+				if (Extension.IsValid() && Extension->IsMuted())
+				{
+					return true;
+				}
+			}
+			return false;
 		});
 
-	const bool bNewMuteState = !bAllMuted;
+	const bool bNewState = !bAllSame;
 
-	for (const TViewModelPtr Mutable : OutlinerSelection.Filter<IMutableExtension>())
+	bool bAnyChanged = false;
+
+	for (const TViewModelPtr<IOutlinerExtension>& OutlinerItem : OutlinerSelection)
 	{
-		Mutable->SetIsMuted(bNewMuteState);
+		for (const TViewModelPtr<IMutableExtension>& Extension : OutlinerItem.AsModel()->GetDescendantsOfType<IMutableExtension>(true))
+		{
+			if (Extension->IsMuted() != bNewState)
+			{
+				Extension->SetIsMuted(bNewState);
+				bAnyChanged = true;
+			}
+		}
 	}
 
-	RefreshTree();
+	if (bAnyChanged)
+	{
+		RefreshTree();
+	}
 }
 
 void FSequencer::ToggleTrackSelectionSolo()
@@ -12856,20 +12929,38 @@ void FSequencer::ToggleTrackSelectionSolo()
 
 	FOutlinerSelection& OutlinerSelection = GetSelection().Outliner;
 
-	const bool bAllSoloed = Algo::AllOf(OutlinerSelection, [](const FViewModelPtr& Selection)
+	const bool bAllSame = Algo::AllOf(OutlinerSelection, [](const FViewModelPtr& InSelection)
 		{
-			const TViewModelPtr<ISoloableExtension> Soloable = Selection.ImplicitCast();
-			return Soloable.IsValid() && Soloable->IsSolo();
+			for (const TViewModelPtr<ISoloableExtension>& Extension : InSelection->GetDescendantsOfType<ISoloableExtension>(true))
+			{
+				if (Extension.IsValid() && Extension->IsSolo())
+				{
+					return true;
+				}
+			}
+			return false;
 		});
 
-	const bool bNewSoloState = !bAllSoloed;
+	const bool bNewState = !bAllSame;
 
-	for (const TViewModelPtr Soloable : OutlinerSelection.Filter<ISoloableExtension>())
+	bool bAnyChanged = false;
+
+	for (const TViewModelPtr<IOutlinerExtension>& OutlinerItem : OutlinerSelection)
 	{
-		Soloable->SetIsSoloed(bNewSoloState);
+		for (const TViewModelPtr<ISoloableExtension>& Extension : OutlinerItem.AsModel()->GetDescendantsOfType<ISoloableExtension>(true))
+		{
+			if (Extension->IsSolo() != bNewState)
+			{
+				Extension->SetIsSoloed(bNewState);
+				bAnyChanged = true;
+			}
+		}
 	}
 
-	RefreshTree();
+	if (bAnyChanged)
+	{
+		RefreshTree();
+	}
 }
 
 TSharedPtr<UE::Sequencer::NavigationTool::INavigationTool> FSequencer::GetNavigationTool() const
