[Sequencer] Navigation Tool improvements:
- Better handling of Sequencer to Nav Tool selection syncing
- Better provider handling
- Providers define what UMovieSceneSequence classes they support
- Fix item sorting bug causing items to not be displayed
- Change main toolbar menu name to be consistent with the rest of the Navigation Tool menu names

diff --git a/Engine/Plugins/MovieScene/LevelSequenceEditor/Source/LevelSequenceEditor/Private/LevelSequenceEditorToolkit.cpp b/Engine/Plugins/MovieScene/LevelSequenceEditor/Source/LevelSequenceEditor/Private/LevelSequenceEditorToolkit.cpp
--- a/Engine/Plugins/MovieScene/LevelSequenceEditor/Source/LevelSequenceEditor/Private/LevelSequenceEditorToolkit.cpp
+++ b/Engine/Plugins/MovieScene/LevelSequenceEditor/Source/LevelSequenceEditor/Private/LevelSequenceEditorToolkit.cpp
@@ -322,8 +322,14 @@ void FLevelSequenceEditorToolkit::Initialize(const EToolkitMode::Type Mode, cons
 	using namespace UE::Sequencer::NavigationTool;
 	if (const TSharedPtr<INavigationTool> Tool = Sequencer->GetNavigationTool())
 	{
-		NavigationToolProvider = MakeShared<FLevelSequenceNavigationToolProvider>();
-		FNavigationToolExtender::RegisterOnlyToolProvider(Tool.ToSharedRef(), NavigationToolProvider.ToSharedRef());
+		const FName ToolId = FNavigationToolExtender::GetToolInstanceId(*Tool);
+
+		const TSharedPtr<FNavigationToolProvider> FoundProvider = FNavigationToolExtender::FindToolProvider(ToolId, FLevelSequenceNavigationToolProvider::Identifier);
+		if (!FoundProvider.IsValid())
+		{
+			NavigationToolProvider = MakeShared<FLevelSequenceNavigationToolProvider>();
+			FNavigationToolExtender::RegisterToolProvider(ToolId, NavigationToolProvider.ToSharedRef());
+		}
 	}
 }
 
@@ -995,7 +1001,8 @@ void FLevelSequenceEditorToolkit::OnClose()
 		using namespace UE::Sequencer::NavigationTool;
 		if (const TSharedPtr<INavigationTool> Tool = Sequencer->GetNavigationTool())
 		{
-			FNavigationToolExtender::UnregisterToolProvider(Tool.ToSharedRef(), NavigationToolProvider->GetIdentifier());
+			const FName ToolId = FNavigationToolExtender::GetToolInstanceId(*Tool);
+			FNavigationToolExtender::UnregisterToolProvider(ToolId, NavigationToolProvider->GetIdentifier());
 		}
 		NavigationToolProvider.Reset();
 	}
diff --git a/Engine/Plugins/MovieScene/LevelSequenceEditor/Source/LevelSequenceEditor/Private/LevelSequenceNavigationToolProvider.cpp b/Engine/Plugins/MovieScene/LevelSequenceEditor/Source/LevelSequenceEditor/Private/LevelSequenceNavigationToolProvider.cpp
--- a/Engine/Plugins/MovieScene/LevelSequenceEditor/Source/LevelSequenceEditor/Private/LevelSequenceNavigationToolProvider.cpp
+++ b/Engine/Plugins/MovieScene/LevelSequenceEditor/Source/LevelSequenceEditor/Private/LevelSequenceNavigationToolProvider.cpp
@@ -46,6 +46,11 @@ FName FLevelSequenceNavigationToolProvider::GetIdentifier() const
 	return Identifier;
 }
 
+TSet<TSubclassOf<UMovieSceneSequence>> FLevelSequenceNavigationToolProvider::GetSupportedSequenceClasses() const
+{
+	return { ULevelSequence::StaticClass() };
+}
+
 FName FLevelSequenceNavigationToolProvider::GetDefaultColumnView() const
 {
 	return AnimationColumnViewId;
@@ -56,7 +61,7 @@ void FLevelSequenceNavigationToolProvider::BindCommands(const TSharedRef<FUIComm
 	InCommandList->Append(ToolCommands);
 }
 
-void FLevelSequenceNavigationToolProvider::OnRegistered()
+void FLevelSequenceNavigationToolProvider::OnActivate()
 {
 	// @TODO: Move from core Navigation Tool
 	/*ExtendToolToolBar();
@@ -66,7 +71,7 @@ void FLevelSequenceNavigationToolProvider::OnRegistered()
 		.AddSP(this, &FAvaNavigationToolProvider::ExtendToolItemDragDropOp);*/
 }
 
-void FLevelSequenceNavigationToolProvider::OnUnregistered()
+void FLevelSequenceNavigationToolProvider::OnDeactivate()
 {
 	// @TODO: Move from core Navigation Tool
 	/*RemoveToolToolBarExtension();
diff --git a/Engine/Plugins/MovieScene/LevelSequenceEditor/Source/LevelSequenceEditor/Private/LevelSequenceNavigationToolProvider.h b/Engine/Plugins/MovieScene/LevelSequenceEditor/Source/LevelSequenceEditor/Private/LevelSequenceNavigationToolProvider.h
--- a/Engine/Plugins/MovieScene/LevelSequenceEditor/Source/LevelSequenceEditor/Private/LevelSequenceNavigationToolProvider.h
+++ b/Engine/Plugins/MovieScene/LevelSequenceEditor/Source/LevelSequenceEditor/Private/LevelSequenceNavigationToolProvider.h
@@ -30,12 +30,14 @@ public:
 
 	virtual FName GetIdentifier() const override;
 
+	virtual TSet<TSubclassOf<UMovieSceneSequence>> GetSupportedSequenceClasses() const override;
+
 	virtual FName GetDefaultColumnView() const override;
 
 	virtual void BindCommands(const TSharedRef<FUICommandList>& InCommandList) override;
 
-	virtual void OnRegistered() override;
-	virtual void OnUnregistered() override;
+	virtual void OnActivate() override;
+	virtual void OnDeactivate() override;
 
 	virtual void OnExtendColumns(FNavigationToolColumnExtender& OutExtender) override;
 	virtual void OnExtendColumnViews(TSet<FNavigationToolColumnView>& OutColumnViews) override;
diff --git a/Engine/Source/Editor/Sequencer/Private/Tools/NavigationTool/Customization/NavigationToolExtender.cpp b/Engine/Source/Editor/Sequencer/Private/Tools/NavigationTool/Customization/NavigationToolExtender.cpp
--- a/Engine/Source/Editor/Sequencer/Private/Tools/NavigationTool/Customization/NavigationToolExtender.cpp
+++ b/Engine/Source/Editor/Sequencer/Private/Tools/NavigationTool/Customization/NavigationToolExtender.cpp
@@ -34,7 +34,7 @@ FNavigationToolExtender::~FNavigationToolExtender()
 	FNavigationToolCommands::Unregister();
 }
 
-FName FNavigationToolExtender::GetToolInstanceName(const INavigationTool& InTool)
+FName FNavigationToolExtender::GetToolInstanceId(const INavigationTool& InTool)
 {
 	if (const TSharedPtr<ISequencer> Sequencer = InTool.GetSequencer())
 	{
@@ -47,239 +47,148 @@ FName FNavigationToolExtender::GetToolInstanceName(const INavigationTool& InTool
 	return NAME_None;
 }
 
-bool FNavigationToolExtender::RegisterToolProvider(const TSharedRef<INavigationTool>& InTool
+bool FNavigationToolExtender::RegisterToolProvider(const FName InToolId
 	, const TSharedRef<FNavigationToolProvider>& InProvider)
 {
-	FNavigationToolExtender& ThisInstance = Get();
-	INavigationTool& RawToolRef = *InTool;
+	if (InToolId.IsNone())
+	{
+		UE_LOG(LogSequencer, Error, TEXT("Invalid Navigation Tool Instance Name"));
+		return false;
+	}
 
 	// Check if provider already exists for this tool
 	const FName Identifier = InProvider->GetIdentifier();
-	if (const TSharedPtr<FNavigationToolProvider> Provider = FindToolProvider(RawToolRef, Identifier))
+	if (const TSharedPtr<FNavigationToolProvider> Provider = FindToolProvider(InToolId, Identifier))
 	{
-		return false;
+		// Unregister old tool provider and continue to register new provider
+		UnregisterToolProvider(InToolId, Identifier);
 	}
 
+	FNavigationToolExtender& ThisInstance = Get();
+
 	// Add the provider if an entry for this tool already exists
-	if (FNavigationToolProviderEntry* const Entry = FindToolProvidersEntry(RawToolRef))
+	if (TSet<TSharedRef<FNavigationToolProvider>>* const Providers = FindToolProviders(InToolId))
 	{
-		Entry->WeakProviders.Add(InProvider);
-		UE_LOG(LogSequencer, Warning, TEXT("Found Registered Navigation Tool Provider: %s"), *InProvider->GetIdentifier().ToString());
+		Providers->Add(InProvider);
+		UE_LOG(LogSequencer, Log, TEXT("Navigation Tool Instance '%s' AND Provider '%s' Registered")
+			, *InToolId.ToString(), *InProvider->GetIdentifier().ToString());
 
-		if (const TSharedPtr<FUICommandList> BaseCommandList = InTool->GetBaseCommandList())
-		{
-			InProvider->BindCommands(BaseCommandList.ToSharedRef());
-		}
-		InProvider->OnRegistered();
-
-		ThisInstance.ProvidersChangedDelegate.Broadcast(InTool, InProvider, ENavigationToolProvidersChangeType::Add);
+		ThisInstance.ProvidersChangedDelegate.Broadcast(InToolId, InProvider, ENavigationToolProvidersChangeType::Add);
 
 		return true;
 	}
 
-	// Find and verify the tool instance name
-	const FName ToolInstanceName = GetToolInstanceName(RawToolRef);
-
-	if (ToolInstanceName.IsNone())
-	{
-		UE_LOG(LogSequencer, Error, TEXT("Invalid Navigation Tool Instance Name"));
-		return false;
-	}
-
 	// Add new tool entry
-	FNavigationToolProviderEntry NewEntry(InTool, { InProvider });
-
-	ThisInstance.ToolProviders.Add(MoveTemp(NewEntry));
-	UE_LOG(LogSequencer, Log, TEXT("Navigation Tool Provider Registered: %s"), *ToolInstanceName.ToString());
-
-	if (const TSharedPtr<FUICommandList> BaseCommandList = InTool->GetBaseCommandList())
-	{
-		InProvider->BindCommands(BaseCommandList.ToSharedRef());
-	}
-	InProvider->OnRegistered();
+	ThisInstance.ToolProviders.Add(InToolId, { InProvider });
+	UE_LOG(LogSequencer, Log, TEXT("Navigation Tool Instance '%s' Provider '%s' Registered")
+		, *InToolId.ToString(), *InProvider->GetIdentifier().ToString());
 
-	ThisInstance.ProvidersChangedDelegate.Broadcast(InTool, InProvider, ENavigationToolProvidersChangeType::Add);
+	ThisInstance.ProvidersChangedDelegate.Broadcast(InToolId, InProvider, ENavigationToolProvidersChangeType::Add);
 
 	return true;
 }
 
-bool FNavigationToolExtender::RegisterOnlyToolProvider(const TSharedRef<INavigationTool>& InTool,
-	const TSharedRef<FNavigationToolProvider>& InProvider)
+bool FNavigationToolExtender::UnregisterToolProvider(const FName InToolId, const FName InProviderId)
 {
-	UnregisterAllToolProviders(InTool);
-
-	return RegisterToolProvider(InTool, InProvider);
-}
+	TSet<TSharedRef<FNavigationToolProvider>>* const Providers = FindToolProviders(InToolId);
+	if (!Providers)
+	{
+		return false;
+	}
 
-bool FNavigationToolExtender::UnregisterToolProvider(const TSharedRef<INavigationTool>& InTool, const FName InProviderId)
-{
 	FNavigationToolExtender& ThisInstance = Get();
 
-	for (TSet<FNavigationToolProviderEntry>::TIterator EntryIter(ThisInstance.ToolProviders); EntryIter; ++EntryIter)
+	for (const TSharedRef<FNavigationToolProvider>& Provider : *Providers)
 	{
-		FNavigationToolProviderEntry& Entry = *EntryIter;
-
-		if (const TSharedPtr<INavigationTool> Tool = Entry.WeakTool.Pin())
+		if (Provider->GetIdentifier() == InProviderId)
 		{
-			if (Tool.ToSharedRef() == InTool)
-			{
-				for (TSet<TWeakPtr<FNavigationToolProvider>>::TIterator ProviderIter(Entry.WeakProviders); ProviderIter; ++ProviderIter)
-				{
-					if (const TSharedPtr<FNavigationToolProvider> Provider = ProviderIter->Pin())
-					{
-						if (Provider->GetIdentifier() == InProviderId)
-						{
-							Provider->OnUnregistered();
-
-							ThisInstance.ProvidersChangedDelegate.Broadcast(InTool, Provider.ToSharedRef(), ENavigationToolProvidersChangeType::Remove);
-
-							UE_LOG(LogSequencer, Log, TEXT("Navigation Tool Provider Unregistered: %s"), *Provider->GetIdentifier().ToString());
-
-							ProviderIter.RemoveCurrent();
-
-							return true;
-						}
-					}
-					else
-					{
-						ProviderIter.RemoveCurrent();
-					}
-				}
-			}
+			UE_LOG(LogSequencer, Log, TEXT("Navigation Tool Provider Unregistered: %s"), *Provider->GetIdentifier().ToString());
+
+			ThisInstance.ProvidersChangedDelegate.Broadcast(InToolId, Provider, ENavigationToolProvidersChangeType::Remove);
+
+			ThisInstance.ToolProviders[InToolId].Remove(Provider);
+
+			return true;
 		}
 	}
 
 	return false;
 }
 
-void FNavigationToolExtender::UnregisterAllToolProviders(const TSharedRef<INavigationTool>& InTool)
+TSet<TSharedRef<FNavigationToolProvider>>* FNavigationToolExtender::FindToolProviders(const FName InToolId)
 {
 	FNavigationToolExtender& ThisInstance = Get();
 
-	for (TSet<FNavigationToolProviderEntry>::TIterator EntryIter(ThisInstance.ToolProviders); EntryIter; ++EntryIter)
+	if (!ThisInstance.ToolProviders.Contains(InToolId))
 	{
-		FNavigationToolProviderEntry& Entry = *EntryIter;
-
-		if (const TSharedPtr<INavigationTool> Tool = Entry.WeakTool.Pin())
-		{
-			if (Tool.ToSharedRef() == InTool)
-			{
-				for (TSet<TWeakPtr<FNavigationToolProvider>>::TIterator ProviderIter(Entry.WeakProviders); ProviderIter; ++ProviderIter)
-				{
-					if (const TSharedPtr<FNavigationToolProvider> Provider = ProviderIter->Pin())
-					{
-						Provider->OnUnregistered();
-
-						ThisInstance.ProvidersChangedDelegate.Broadcast(InTool, Provider.ToSharedRef(), ENavigationToolProvidersChangeType::Remove);
-
-						UE_LOG(LogSequencer, Log, TEXT("Navigation Tool Provider Unregistered: %s"), *Provider->GetIdentifier().ToString());
-					}
-
-					ProviderIter.RemoveCurrent();
-				}
-			}
-		}
+		return nullptr;
 	}
+
+	return &ThisInstance.ToolProviders[InToolId];
 }
 
-FNavigationToolExtender::FNavigationToolProviderEntry* FNavigationToolExtender::FindToolProvidersEntry(const INavigationTool& InTool)
+TSharedPtr<FNavigationToolProvider> FNavigationToolExtender::FindToolProvider(const FName InToolId, const FName InProviderId)
 {
 	FNavigationToolExtender& ThisInstance = Get();
 
-	for (TSet<FNavigationToolProviderEntry>::TIterator EntryIter(ThisInstance.ToolProviders); EntryIter; ++EntryIter)
+	if (!ThisInstance.ToolProviders.Contains(InToolId))
 	{
-		FNavigationToolProviderEntry& Entry = *EntryIter;
-
-		if (const TSharedPtr<INavigationTool> Tool = Entry.WeakTool.Pin())
-		{
-			if (Tool.Get() == &InTool)
-			{
-				return &Entry;
-			}
-		}
+		return nullptr;
 	}
 
-	return nullptr;
-}
-
-TSharedPtr<FNavigationToolProvider> FNavigationToolExtender::FindToolProvider(const INavigationTool& InTool, const FName InProviderId)
-{
-	if (FNavigationToolProviderEntry* const Entry = FindToolProvidersEntry(InTool))
+	for (const TSharedRef<FNavigationToolProvider>& Provider : ThisInstance.ToolProviders[InToolId])
 	{
-		for (TSet<TWeakPtr<FNavigationToolProvider>>::TIterator ProviderIter(Entry->WeakProviders); ProviderIter; ++ProviderIter)
+		if (Provider->GetIdentifier() == InProviderId)
 		{
-			if (const TSharedPtr<FNavigationToolProvider> Provider = ProviderIter->Pin())
-			{
-				if (Provider->GetIdentifier() == InProviderId)
-				{
-					return Provider;
-				}
-			}
+			return Provider;
 		}
 	}
 
 	return nullptr;
 }
 
-bool FNavigationToolExtender::FindToolProviders(const INavigationTool& InTool
-	, TSet<TWeakPtr<FNavigationToolProvider>>& OutProviders)
+bool FNavigationToolExtender::FindToolProviders(const FName InToolId
+	, TSet<TSharedRef<FNavigationToolProvider>>& OutProviders)
 {
-	if (const FNavigationToolProviderEntry* const Entry = FindToolProvidersEntry(InTool))
+	if (TSet<TSharedRef<FNavigationToolProvider>>* const Providers = FindToolProviders(InToolId))
 	{
-		OutProviders.Append(Entry->WeakProviders);
+		OutProviders.Append(*Providers);
 		return true;
 	}
 	return false;
 }
 
-void FNavigationToolExtender::ForEachProvider(const TFunction<bool(INavigationTool& InTool
+void FNavigationToolExtender::ForEachProvider(const TFunction<bool(const FName InToolId
 	, const TSharedRef<FNavigationToolProvider>& InProvider)>& InPredicate)
 {
-	for (TSet<FNavigationToolProviderEntry>::TIterator EntryIter(Get().ToolProviders); EntryIter; ++EntryIter)
+	for (const TPair<FName, TSet<TSharedRef<FNavigationToolProvider>>>& ProviderPair : Get().ToolProviders)
 	{
-		FNavigationToolProviderEntry& Entry = *EntryIter;
-
-		const TSharedPtr<INavigationTool> Tool = Entry.WeakTool.Pin();
-		if (!Tool.IsValid())
-		{
-			EntryIter.RemoveCurrent();
-			continue;
-		}
-
-		const TSharedRef<INavigationTool> ToolRef = Tool.ToSharedRef();
-
-		for (TSet<TWeakPtr<FNavigationToolProvider>>::TIterator ProviderIter(Entry.WeakProviders); ProviderIter; ++ProviderIter)
+		for (const TSharedRef<FNavigationToolProvider>& Provider : ProviderPair.Value)
 		{
-			if (const TSharedPtr<FNavigationToolProvider> Provider = ProviderIter->Pin())
+			// Stop processing if the function returns false
+			if (!InPredicate(ProviderPair.Key, Provider))
 			{
-				// Stop processing if the function returns false
-				if (!InPredicate(*ToolRef, Provider.ToSharedRef()))
-				{
-					return;
-				}
+				return;
 			}
 		}
 	}
 }
 
-void FNavigationToolExtender::ForEachToolProvider(INavigationTool& InTool,
-	const TFunction<bool(const TSharedRef<FNavigationToolProvider>& InProvider)>& InPredicate)
+void FNavigationToolExtender::ForEachToolProvider(const FName InToolId
+	, const TFunction<bool(const TSharedRef<FNavigationToolProvider>& InProvider)>& InPredicate)
 {
-	FNavigationToolProviderEntry* const ToolEntry = Get().FindToolProvidersEntry(InTool);
-	if (!ToolEntry)
+	TSet<TSharedRef<FNavigationToolProvider>> const* Providers = Get().FindToolProviders(InToolId);
+	if (!Providers)
 	{
 		return;
 	}
 
-	for (TSet<TWeakPtr<FNavigationToolProvider>>::TIterator ProviderIter(ToolEntry->WeakProviders); ProviderIter; ++ProviderIter)
+	for (const TSharedRef<FNavigationToolProvider>& Provider : *Providers)
 	{
-		if (const TSharedPtr<FNavigationToolProvider> Provider = ProviderIter->Pin())
+		// Stop processing if the function returns false
+		if (!InPredicate(Provider))
 		{
-			// Stop processing if the function returns false
-			if (!InPredicate(Provider.ToSharedRef()))
-			{
-				return;
-			}
+			return;
 		}
 	}
 }
diff --git a/Engine/Source/Editor/Sequencer/Private/Tools/NavigationTool/Items/NavigationToolItem.cpp b/Engine/Source/Editor/Sequencer/Private/Tools/NavigationTool/Items/NavigationToolItem.cpp
--- a/Engine/Source/Editor/Sequencer/Private/Tools/NavigationTool/Items/NavigationToolItem.cpp
+++ b/Engine/Source/Editor/Sequencer/Private/Tools/NavigationTool/Items/NavigationToolItem.cpp
@@ -297,8 +297,7 @@ TOptional<EItemDropZone> FNavigationToolItem::CanAcceptDrop(const FDragDropEvent
 
 	TOptional<EItemDropZone> OutDropZone;
 
-	FNavigationToolExtender::ForEachToolProvider(GetOwnerTool(),
-		[this, &InDragDropEvent, InDropZone, &SharedThisRef, &OutDropZone]
+	GetOwnerTool().ForEachProvider([this, &InDragDropEvent, InDropZone, &SharedThisRef, &OutDropZone]
 		(const TSharedRef<FNavigationToolProvider>& InProvider)
 		{
 			OutDropZone = InProvider->OnToolItemCanAcceptDrop(InDragDropEvent, InDropZone, SharedThisRef);
@@ -327,8 +326,7 @@ FReply FNavigationToolItem::AcceptDrop(const FDragDropEvent& InDragDropEvent, EI
 
 	FReply OutReply = FReply::Unhandled();
 
-	FNavigationToolExtender::ForEachToolProvider(GetOwnerTool(),
-		[this, &InDragDropEvent, InDropZone, &SharedThisRef, &OutReply]
+	GetOwnerTool().ForEachProvider([this, &InDragDropEvent, InDropZone, &SharedThisRef, &OutReply]
 		(const TSharedRef<FNavigationToolProvider>& InProvider)
 		{
 			OutReply = InProvider->OnToolItemAcceptDrop(InDragDropEvent, InDropZone, SharedThisRef);
@@ -405,7 +403,9 @@ void FNavigationToolItem::AddChildChecked(const FNavigationToolAddItemParams& In
 		}
 	}
 
-	if (InAddItemParams.RelativeItem.IsValid() && InAddItemParams.RelativeDropZone.IsSet() && InAddItemParams.RelativeDropZone != EItemDropZone::OntoItem)
+	if (InAddItemParams.RelativeItem.IsValid()
+		&& InAddItemParams.RelativeDropZone.IsSet()
+		&& InAddItemParams.RelativeDropZone != EItemDropZone::OntoItem)
 	{
 		const int32 RelativeItemIndex = Children.Find(InAddItemParams.RelativeItem);
 		if (RelativeItemIndex != INDEX_NONE)
@@ -437,31 +437,38 @@ bool FNavigationToolItem::RemoveChildChecked(const FNavigationToolRemoveItemPara
 
 void FNavigationToolItem::HandleNewSortableChildren(TArray<FNavigationToolItemPtr> InSortableChildren)
 {
-	const TSharedPtr<FNavigationToolProvider> Provider = GetProvider();
-	if (!Provider.IsValid())
-	{
-		return;
-	}
+	InSortableChildren.Sort([this](const FNavigationToolItemPtr& InItemA, const FNavigationToolItemPtr& InItemB)
+		{
+			const FNavigationToolSaveState* const SaveStateA = InItemA->GetProviderSaveState();
+			if (!SaveStateA)
+			{
+				return false;
+			}
 
-	const FNavigationToolSaveState* const SaveState = Provider->GetSaveState(Tool);
-	if (!SaveState)
-	{
-		return;
-	}
+			const FNavigationToolSaveState* const SaveStateB = InItemB->GetProviderSaveState();
+			if (!SaveStateB)
+			{
+				return false;
+			}
 
-	InSortableChildren.Sort([SaveState](const FNavigationToolItemPtr& A, const FNavigationToolItemPtr& B)
-	{
-		const FNavigationToolSerializedTreeNode* const NodeA = SaveState->SerializedTree.FindTreeNode(FNavigationTool::MakeSerializedItemFromToolItem(A));
-		const FNavigationToolSerializedTreeNode* const NodeB = SaveState->SerializedTree.FindTreeNode(FNavigationTool::MakeSerializedItemFromToolItem(B));
-		return FNavigationToolSerializedTree::CompareTreeItemOrder(NodeA, NodeB);
-	});
+			const FNavigationToolSerializedTreeNode* const NodeA = SaveStateA->SerializedTree.FindTreeNode(InItemA->MakeSerializedItem());
+			const FNavigationToolSerializedTreeNode* const NodeB = SaveStateB->SerializedTree.FindTreeNode(InItemB->MakeSerializedItem());
+
+			return FNavigationToolSerializedTree::CompareTreeItemOrder(NodeA, NodeB);
+		});
 
 	FNavigationToolAddItemParams AddItemParams;
 	for (const FNavigationToolItemPtr& NewChild : InSortableChildren)
 	{
+		const FNavigationToolSaveState* const SaveState = NewChild->GetProviderSaveState();
+		if (!SaveState)
+		{
+			continue;
+		}
+
 		AddItemParams.Item = NewChild;
 
-		const FNavigationToolSerializedTreeNode* const TreeNode = SaveState->SerializedTree.FindTreeNode(FNavigationTool::MakeSerializedItemFromToolItem(NewChild));
+		const FNavigationToolSerializedTreeNode* const TreeNode = SaveState->SerializedTree.FindTreeNode(NewChild->MakeSerializedItem());
 		if (TreeNode && Children.IsValidIndex(TreeNode->GetLocalIndex()))
 		{
 			// Add Before the Child at Index, so this Item is at the specific Index
@@ -479,6 +486,15 @@ void FNavigationToolItem::HandleNewSortableChildren(TArray<FNavigationToolItemPt
 	}
 }
 
+FNavigationToolSaveState* FNavigationToolItem::GetProviderSaveState() const
+{
+	if (const TSharedPtr<FNavigationToolProvider> Provider = GetProvider())
+	{
+		return Provider->GetSaveState(Tool);
+	}
+	return nullptr;
+}
+
 }
 
 #undef LOCTEXT_NAMESPACE
diff --git a/Engine/Source/Editor/Sequencer/Private/Tools/NavigationTool/Items/NavigationToolTreeRoot.cpp b/Engine/Source/Editor/Sequencer/Private/Tools/NavigationTool/Items/NavigationToolTreeRoot.cpp
--- a/Engine/Source/Editor/Sequencer/Private/Tools/NavigationTool/Items/NavigationToolTreeRoot.cpp
+++ b/Engine/Source/Editor/Sequencer/Private/Tools/NavigationTool/Items/NavigationToolTreeRoot.cpp
@@ -22,7 +22,7 @@ void FNavigationToolTreeRoot::FindChildren(TArray<FNavigationToolItemPtr>& OutCh
 	const TSharedRef<FNavigationToolTreeRoot> SharedThisRef = SharedThis(this);
 
 	// Get items to add from providers
-	FNavigationToolExtender::ForEachToolProvider(GetOwnerTool(), [this, &OutChildren, bInRecursive, &SharedThisRef]
+	GetOwnerTool().ForEachProvider([this, &OutChildren, bInRecursive, &SharedThisRef]
 		(const TSharedRef<FNavigationToolProvider>& InToolProvider)
 		{
 			TArray<FNavigationToolItemPtr> ThisProviderItems;
diff --git a/Engine/Source/Editor/Sequencer/Private/Tools/NavigationTool/Menus/NavigationToolFilterMenu.cpp b/Engine/Source/Editor/Sequencer/Private/Tools/NavigationTool/Menus/NavigationToolFilterMenu.cpp
--- a/Engine/Source/Editor/Sequencer/Private/Tools/NavigationTool/Menus/NavigationToolFilterMenu.cpp
+++ b/Engine/Source/Editor/Sequencer/Private/Tools/NavigationTool/Menus/NavigationToolFilterMenu.cpp
@@ -18,7 +18,7 @@ namespace UE::Sequencer::NavigationTool
 
 TSharedRef<SWidget> FNavigationToolFilterMenu::CreateMenu(const TSharedRef<INavigationToolFilterBar>& InFilterBar)
 {
-	const FName FilterMenuName = TEXT("NavigationTool.FilterMenu");
+	const FName FilterMenuName = TEXT("Sequencer.NavigationTool.FilterMenu");
 	if (!UToolMenus::Get()->IsMenuRegistered(FilterMenuName))
 	{
 		UToolMenu* const Menu = UToolMenus::Get()->RegisterMenu(FilterMenuName);
diff --git a/Engine/Source/Editor/Sequencer/Private/Tools/NavigationTool/NavigationTool.cpp b/Engine/Source/Editor/Sequencer/Private/Tools/NavigationTool/NavigationTool.cpp
--- a/Engine/Source/Editor/Sequencer/Private/Tools/NavigationTool/NavigationTool.cpp
+++ b/Engine/Source/Editor/Sequencer/Private/Tools/NavigationTool/NavigationTool.cpp
@@ -29,7 +29,6 @@
 #include "Tools/NavigationTool/Items/NavigationToolBinding.h"
 #include "Tools/NavigationTool/Items/NavigationToolItemUtils.h"
 #include "Tools/NavigationTool/Items/NavigationToolTreeRoot.h"
-#include "Tools/NavigationTool/NavigationToolEditorSelection.h"
 #include "Tools/NavigationTool/NavigationToolScopedSelection.h"
 #include "Tools/NavigationTool/NavigationToolTab.h"
 #include "Tools/NavigationTool/NavigationToolView.h"
@@ -55,168 +54,68 @@ void FNavigationTool::Init()
 	const TSharedPtr<INavigationToolView> ToolView = RegisterToolView(DefaultViewId);
 	check(ToolView.IsValid());
 
-	LoadAllProviderStates();
-
-	if (const TSharedPtr<FSequencer> Sequencer = StaticCastSharedPtr<FSequencer>(WeakSequencer.Pin()))
-	{
-		Sequencer->OnChannelChanged().AddSPLambda(this, [this](const FMovieSceneChannelMetaData*, UMovieSceneSection*)
-			{
-				OnTreeViewChanged();
-			});
-		Sequencer->OnMovieSceneBindingsChanged().AddSPLambda(this, [this]()
-			{
-				OnTreeViewChanged();
-			});
-		Sequencer->OnMovieSceneBindingsPasted().AddSPLambda(this, [this](const TArray<FMovieSceneBinding>&)
-			{
-				OnTreeViewChanged();
-			});
-		Sequencer->OnMovieSceneDataChanged().AddSPLambda(this, [this](const EMovieSceneDataChangeType)
-			{
-				OnTreeViewChanged();
-			});
-		Sequencer->OnTreeViewChanged().AddSP(this, &FNavigationTool::OnTreeViewChanged);
-		Sequencer->OnEndScrubbingEvent().AddSP(this, &FNavigationTool::OnTreeViewChanged);
-		Sequencer->OnCloseEvent().AddSPLambda(this, [this](const TSharedRef<ISequencer> InSequencer)
-			{
-				if (InSequencer == WeakSequencer.Pin())
-				{
-					Shutdown();
-				}
-			});
-		Sequencer->GetSelection().Outliner.OnChanged.AddSP(this, &FNavigationTool::OnSequencerOutlinerSelectionChanged);
-		Sequencer->GetSelection().TrackArea.OnChanged.AddSP(this, &FNavigationTool::OnSequencerTrackAreaSelectionChanged);
-	}
-
-	// Listen to object replacement changes
-	FCoreUObjectDelegates::OnObjectsReplaced.AddSP(this, &FNavigationTool::OnObjectsReplaced);
-	FEditorDelegates::OnEditorActorReplaced.AddSP(this, &FNavigationTool::OnActorReplaced);
+	ToolTab->OnVisibilityChanged().AddSP(this, &FNavigationTool::OnTabVisibilityChanged);
 
 	FNavigationToolExtender::OnProvidersChanged().AddSP(this, &FNavigationTool::OnProvidersChanged);
 
-	// ToolTab uses the most recent tool view as content, so must call after view registration
-	ToolTab->Init();
-
-	if (GEditor)
+	if (const TSharedPtr<FSequencer> Sequencer = StaticCastSharedPtr<FSequencer>(WeakSequencer.Pin()))
 	{
-		GEditor->RegisterForUndo(this);
+		Sequencer->OnActivateSequence().AddSP(this, &FNavigationTool::OnSequencerActivate);
+		Sequencer->OnCloseEvent().AddSP(this, &FNavigationTool::OnSequencerClose);
 	}
 }
 
 void FNavigationTool::Shutdown()
 {
-	SaveAllProviderStates();
+	UnbindEvents();
 
-	if (const TSharedPtr<FSequencer> Sequencer = StaticCastSharedPtr<FSequencer>(WeakSequencer.Pin()))
-	{
-		Sequencer->OnChannelChanged().RemoveAll(this);
-		Sequencer->OnMovieSceneBindingsChanged().RemoveAll(this);
-		Sequencer->OnMovieSceneBindingsPasted().RemoveAll(this);
-		Sequencer->OnMovieSceneDataChanged().RemoveAll(this);
-		Sequencer->OnTreeViewChanged().RemoveAll(this);
-		Sequencer->OnEndScrubbingEvent().RemoveAll(this);
-		Sequencer->OnCloseEvent().RemoveAll(this);
-		Sequencer->GetSelection().Outliner.OnChanged.RemoveAll(this);
-		Sequencer->GetSelection().TrackArea.OnChanged.RemoveAll(this);
-	}
+	SaveAllProviderStates();
 
-	FCoreUObjectDelegates::OnObjectsReplaced.RemoveAll(this);
-	FEditorDelegates::OnEditorActorReplaced.RemoveAll(this);
+	ToolTab->OnVisibilityChanged().RemoveAll(this);
 
 	FNavigationToolExtender::OnProvidersChanged().RemoveAll(this);
 
-	ToolTab->Shutdown();
-
-	if (GEditor)
+	if (const TSharedPtr<FSequencer> Sequencer = StaticCastSharedPtr<FSequencer>(WeakSequencer.Pin()))
 	{
-		GEditor->UnregisterForUndo(this);
+		Sequencer->OnActivateSequence().RemoveAll(this);
+		Sequencer->OnCloseEvent().RemoveAll(this);
 	}
 }
 
-void FNavigationTool::ForEachProvider(const TFunction<bool(const TSharedRef<FNavigationToolProvider>& InToolProvider)>& InPredicate)
-{
-	FNavigationToolExtender::ForEachToolProvider(*this, InPredicate);
-}
-
-void FNavigationTool::LoadAllProviderStates(const bool bInRefreshTool)
+void FNavigationTool::ForEachProvider(const TFunction<bool(const TSharedRef<FNavigationToolProvider>& InToolProvider)>& InPredicate) const
 {
-	GlobalFilters.Reset();
-
-	ForEachProvider([this](const TSharedRef<FNavigationToolProvider>& InProvider)
-		{
-			LoadProviderState(InProvider);
-			return true;
-		});
-
-	if (bInRefreshTool)
+	const TSharedPtr<ISequencer> Sequencer = GetSequencer();
+	if (!Sequencer.IsValid())
 	{
-		Refresh();
+		return;
 	}
-}
-
-void FNavigationTool::SaveAllProviderStates()
-{
-	ForEachProvider([this](const TSharedRef<FNavigationToolProvider>& InProvider)
-		{
-			SaveProviderState(InProvider);
-			return true;
-		});
-}
 
-void FNavigationTool::OnProvidersChanged(const TSharedRef<INavigationTool>& InTool
-	, const TSharedRef<FNavigationToolProvider>& InProvider
-	, const ENavigationToolProvidersChangeType InChangeType)
-{
-	if (&InTool.Get() != this)
+	UMovieSceneSequence* const RootSequence = Sequencer->GetRootMovieSceneSequence();
+	if (!IsValid(RootSequence))
 	{
 		return;
 	}
 
-	LoadAllProviderStates(false);
-
-	if (InChangeType == ENavigationToolProvidersChangeType::Add)
-	{
-		const UNavigationToolSettings* const ToolSettings = GetDefault<UNavigationToolSettings>();
-		if (ensure(IsValid(ToolSettings)))
-		{
-			//LoadProviderState(InProvider);
-
-			ForEachToolView([&InProvider](const TSharedPtr<FNavigationToolView>& InToolView)
-				{
-					//InToolView->LoadViewState(InProvider);
-
-					InToolView->CreateColumns(InProvider);
-					InToolView->CreateDefaultColumnViews(InProvider);
-				});
+	const FName ToolId = FNavigationToolExtender::GetToolInstanceId(*this);
 
-			// We apply default views *after* all columns have been created for all providers
-			if (ToolSettings->ShouldApplyDefaultColumnView())
-			{
-				ForEachToolView([&InProvider](const TSharedPtr<FNavigationToolView>& InToolView)
-					{
-						InToolView->ApplyCustomColumnView(InProvider->GetDefaultColumnView());
-					});
-			}
-		}
-	}
-	else if (InChangeType == ENavigationToolProvidersChangeType::Remove)
-	{
-		// Remove extended column views
-		UNavigationToolSettings* const ToolSettings = GetMutableDefault<UNavigationToolSettings>();
-		if (ensure(IsValid(ToolSettings)))
+	FNavigationToolExtender::ForEachToolProvider(ToolId, [&InPredicate, RootSequence]
+		(const TSharedRef<FNavigationToolProvider>& InProvider)
 		{
-			TSet<FNavigationToolColumnView>& CustomColumnViews = ToolSettings->GetCustomColumnViews();
-
-			for (const FName ColumnViewName : InProvider->GetExtendedColumnViewNames())
+			if (InProvider->IsSequenceSupported(RootSequence))
 			{
-				CustomColumnViews.Remove(ColumnViewName);
+				return InPredicate(InProvider);
 			}
+			return true; // Continue processing providers
+		});
+}
 
-			ToolSettings->SaveConfig();
-		}
-	}
-
-	Refresh();
+void FNavigationTool::SaveAllProviderStates()
+{
+	ForEachProvider([this](const TSharedRef<FNavigationToolProvider>& InProvider)
+		{
+			InProvider->SaveState(*this);
+			return true;
+		});
 }
 
 USequencerSettings* FNavigationTool::GetSequencerSettings() const
@@ -302,9 +201,7 @@ bool FNavigationTool::IsToolLocked() const
 {
 	bool bOutShouldLock = false;
 
-	FNavigationTool& SharedThisRef = const_cast<FNavigationTool&>(*this);
-
-	FNavigationToolExtender::ForEachToolProvider(SharedThisRef, [&bOutShouldLock]
+	ForEachProvider([&bOutShouldLock]
 		(const TSharedRef<FNavigationToolProvider>& InToolProvider)
 		{
 			if (InToolProvider->ShouldLockTool())
@@ -333,33 +230,19 @@ void FNavigationTool::SetBaseCommandList(const TSharedPtr<FUICommandList>& InBas
 	WeakBaseCommandList = InBaseCommandList;
 }
 
-void FNavigationTool::SaveProviderState(const TSharedRef<FNavigationToolProvider>& InProvider)
+bool FNavigationTool::DoesGlobalFilterExist(const FName InFilterId)
 {
-	SaveSerializedTree(/*bInResetTree*/true);
-
-	/*ForEachToolView([&InProvider](const TSharedPtr<FNavigationToolView>& InToolView)
-		{
-			InToolView->SaveViewState(InProvider);
-		});*/
-
-	if (FNavigationToolSaveState* const SaveState = InProvider->GetSaveState(*this))
+	for (const TSharedPtr<FNavigationToolBuiltInFilter>& GlobalFilter : GlobalFilters)
 	{
-		// Remove any saved item colors that can no longer be found
-		TArray<FString> ItemIds;
-		SaveState->ItemColorMap.GetKeys(ItemIds);
-		for (const FString& ItemId : ItemIds)
+		if (GlobalFilter->GetName().Equals(InFilterId.ToString()))
 		{
-			if (!FindItem(ItemId).IsValid())
-			{
-				SaveState->ItemColorMap.Remove(ItemId);
-			}
+			return true;
 		}
 	}
-
-	InProvider->UpdateItemIdContexts(*this);
+	return false;
 }
 
-void FNavigationTool::LoadProviderState(const TSharedRef<FNavigationToolProvider>& InProvider)
+void FNavigationTool::RefreshGlobalFilters()
 {
 	const UNavigationToolSettings* const ToolSettings = GetDefault<UNavigationToolSettings>();
 	if (!ensure(IsValid(ToolSettings)))
@@ -367,33 +250,15 @@ void FNavigationTool::LoadProviderState(const TSharedRef<FNavigationToolProvider
 		return;
 	}
 
-	InProvider->UpdateItemIdContexts(*this);
-	
-	if (FNavigationToolSaveState* const SaveState = InProvider->GetSaveState(*this))
-	{
-		LoadSerializedTree(RootItem, &SaveState->SerializedTree);
-	}
+	TArray<FNavigationToolBuiltInFilterParams> NewGlobalFilterParams;
 
-	ForEachToolView([&InProvider](const TSharedPtr<FNavigationToolView>& InToolView)
+	ForEachProvider([&NewGlobalFilterParams](const TSharedRef<FNavigationToolProvider>& InProvider)
 		{
-			InToolView->LoadViewState(InProvider);
+			InProvider->OnExtendBuiltInFilters(NewGlobalFilterParams);
+			return true;
 		});
 
-	// Extend built in filters, make sure to not add any duplicates
-	TArray<FNavigationToolBuiltInFilterParams> NewGlobalFilterParams;
-	InProvider->OnExtendBuiltInFilters(NewGlobalFilterParams);
-
-	auto DoesGlobalFilterExist = [this](const FName InFilterId)
-	{
-		for (const TSharedPtr<FNavigationToolBuiltInFilter>& GlobalFilter : GlobalFilters)
-		{
-			if (GlobalFilter->GetName().Equals(InFilterId.ToString()))
-			{
-				return true;
-			}
-		}
-		return false;
-	};
+	GlobalFilters.Reset(NewGlobalFilterParams.Num());
 
 	for (const FNavigationToolBuiltInFilterParams& NewFilterParams : NewGlobalFilterParams)
 	{
@@ -413,80 +278,6 @@ void FNavigationTool::LoadProviderState(const TSharedRef<FNavigationToolProvider
 	}
 }
 
-void FNavigationTool::LoadSerializedTree(const FNavigationToolItemPtr& InParentItem
-	, FNavigationToolSerializedTree* const InSerializedTree)
-{
-	TArray<FNavigationToolItemPtr>& Children = InParentItem->GetChildrenMutable();
-	
-	TArray<FNavigationToolItemPtr> Sortable;
-	TArray<FNavigationToolItemPtr> Unsortable;
-	SplitSortableAndUnsortableItems(Children, Sortable, Unsortable);
-
-	// If Scene Tree is valid, Item Sorting should be empty as this function only takes a valid Scene Tree if
-	// loaded version supports Scene Trees (i.e. when Item Sorting stops being loaded in)
-	if (InSerializedTree)
-	{
-		Sortable.Sort([InSerializedTree](const FNavigationToolItemPtr& InA, const FNavigationToolItemPtr& InB)
-			{
-				const FNavigationToolSerializedTreeNode* const NodeA = InSerializedTree->FindTreeNode(MakeSerializedItemFromToolItem(InA));
-				const FNavigationToolSerializedTreeNode* const NodeB = InSerializedTree->FindTreeNode(MakeSerializedItemFromToolItem(InB));
-				return FNavigationToolSerializedTree::CompareTreeItemOrder(NodeA, NodeB);
-			});
-	}
-
-	Children = MoveTemp(Unsortable);
-	Children.Append(MoveTemp(Sortable));
-
-	for (FNavigationToolItemPtr& Child : Children)
-	{
-		LoadSerializedTree(Child, InSerializedTree);
-	}
-}
-
-void FNavigationTool::SaveSerializedTree(const bool bInResetTree)
-{
-	ForEachProvider([this, bInResetTree](const TSharedRef<FNavigationToolProvider>& InProvider)
-		{
-			FNavigationToolSaveState* const SaveState = InProvider->GetSaveState(*this);
-			if (!SaveState)
-			{
-				return true;
-			}
-
-			if (bInResetTree)
-			{
-				SaveState->SerializedTree.Reset();
-			}
-
-			SaveSerializedTreeRecursive(RootItem, SaveState->SerializedTree);
-			return true;
-		});
-}
-
-void FNavigationTool::SaveSerializedTreeRecursive(const FNavigationToolItemPtr& InParentItem
-	, FNavigationToolSerializedTree& InSerializedTree)
-{
-	const TArray<FNavigationToolItemPtr>& Children = InParentItem->GetChildren();
-
-	const FNavigationToolSerializedItem ParentSceneItem = MakeSerializedItemFromToolItem(InParentItem);
-
-	for (int32 Index = 0; Index < Children.Num(); ++Index)
-	{
-		if (const FNavigationToolItemPtr& ChildItem = Children[Index])
-		{
-			if (ChildItem->ShouldSort())
-			{
-				const FNavigationToolSerializedItem SceneItem = MakeSerializedItemFromToolItem(ChildItem);
-				if (SceneItem.IsValid())
-				{
-					InSerializedTree.GetOrAddTreeNode(SceneItem, ParentSceneItem);
-				}
-			}
-			SaveSerializedTreeRecursive(ChildItem, InSerializedTree);
-		}
-	}
-}
-
 TSharedPtr<INavigationToolView> FNavigationTool::RegisterToolView(const int32 InToolViewId)
 {
 	const TSharedPtr<FNavigationToolView> ToolView = FNavigationToolView::CreateInstance(InToolViewId
@@ -624,7 +415,11 @@ void FNavigationTool::Refresh()
 	// Save the Current Item ordering before Refreshing Children
 	// Do not reset tree just yet as there might be actors that still need to be discovered on the next pass
 	// This is done to save the items added from the Queued Actions and be considered when adding new items from Discovery
-	SaveSerializedTree(/*bInResetTree*/false);
+	ForEachProvider([this](const TSharedRef<FNavigationToolProvider>& InProvider)
+		{
+			InProvider->SaveSerializedTree(*this, /*bInResetTree*/false);
+			return true;
+		});
 
 	// Refresh each item's children. This also Updates each Child's parent var
 	ForEachItem([](const FNavigationToolItemPtr& InItem)
@@ -640,7 +435,11 @@ void FNavigationTool::Refresh()
 		});
 
 	// Save so that the Tree is updated to the latest Navigation Tool State
-	SaveSerializedTree(/*bInResetTree*/false);
+	ForEachProvider([this](const TSharedRef<FNavigationToolProvider>& InProvider)
+		{
+			InProvider->SaveSerializedTree(*this, /*bInResetTree*/false);
+			return true;
+		});
 }
 
 TSharedRef<FNavigationToolItem> FNavigationTool::GetTreeRoot() const
@@ -744,26 +543,23 @@ void FNavigationTool::Tick(const float InDeltaTime)
 		bToolDirty = false;
 	}
 
-	// Select Objects Pending Selection
-	if (ObjectsLastSelected.IsValid())
+	// Select Items Pending Selection
+	if (ItemsLastSelected.IsValid())
 	{
 		TArray<FNavigationToolItemPtr> ItemsToSelect;
 
-		if (ObjectsLastSelected->Num() > 0)
+		for (const TWeakPtr<INavigationToolItem>& WeakItem : *ItemsLastSelected)
 		{
-			ItemsToSelect.Reserve(ObjectsLastSelected->Num());
-			for (TWeakObjectPtr<UObject> Object : *ObjectsLastSelected)
+			if (const TSharedPtr<INavigationToolItem> Item = WeakItem.Pin())
 			{
-				if (Object.IsValid())
-				{
-					ItemsToSelect = FindItemsFromMovieSceneObject(Object.Get());
-				}
+				ItemsToSelect.Add(Item);
 			}
 		}
 
-		// Only Scroll Into View, don't signal selection since we just come from the selection notify itself
+		// Only Scroll Into View, don't signal selection since we just came from the selection notify itself
 		SelectItems(ItemsToSelect, ENavigationToolItemSelectionFlags::ScrollIntoView);
-		ObjectsLastSelected.Reset();
+
+		ItemsLastSelected.Reset();
 	}
 
 	ForEachToolView([InDeltaTime](const TSharedPtr<FNavigationToolView>& InToolView)
@@ -831,7 +627,7 @@ TSharedPtr<INavigationToolView> FNavigationTool::GetMostRecentToolView() const
 	return nullptr;
 }
 
-void FNavigationTool::ForEachToolView(const TFunction<void(const TSharedPtr<FNavigationToolView>& InToolView)>& InPredicate) const
+void FNavigationTool::ForEachToolView(const TFunction<void(const TSharedRef<FNavigationToolView>& InToolView)>& InPredicate) const
 {
 	FNavigationTool* const MutableThis = const_cast<FNavigationTool*>(this);
 	for (TMap<int32, TSharedPtr<FNavigationToolView>>::TIterator Iter(MutableThis->ToolViews); Iter; ++Iter)
@@ -839,7 +635,7 @@ void FNavigationTool::ForEachToolView(const TFunction<void(const TSharedPtr<FNav
 		const TSharedPtr<FNavigationToolView>& ToolView = Iter.Value();
 		if (ToolView.IsValid())
 		{
-			InPredicate(ToolView);
+			InPredicate(ToolView.ToSharedRef());
 		}
 		else
 		{
@@ -1113,22 +909,12 @@ FNavigationToolItemPtr FNavigationTool::FindLowestCommonAncestor(const TArray<FN
 	return LowestCommonAncestor;
 }
 
-FNavigationToolSerializedItem FNavigationTool::MakeSerializedItemFromToolItem(const FNavigationToolItemPtr& InItem)
-{
-	if (!InItem.IsValid())
-	{
-		return FNavigationToolSerializedItem();
-	}
-
-	return FNavigationToolSerializedItem(InItem->GetItemId().GetStringId());
-}
-
 void FNavigationTool::SortItems(TArray<FNavigationToolItemPtr>& OutToolItems, const bool bInReverseOrder)
 {
-	OutToolItems.Sort([bInReverseOrder](const FNavigationToolItemPtr& ItemA, const FNavigationToolItemPtr& ItemB)
-	{
-		return UE::Sequencer::NavigationTool::CompareToolItemOrder(ItemA, ItemB) != bInReverseOrder;
-	});
+	OutToolItems.Sort([bInReverseOrder](const FNavigationToolItemPtr& InItemA, const FNavigationToolItemPtr& InItemB)
+		{
+			return CompareToolItemOrder(InItemA, InItemB) != bInReverseOrder;
+		});
 }
 
 void FNavigationTool::NormalizeToTopLevelSelections(TArray<FNavigationToolItemPtr>& InOutItems)
@@ -1172,36 +958,6 @@ void FNavigationTool::SyncSequencerSelection(const TArray<FNavigationToolItemPtr
 	}
 }
 
-void FNavigationTool::OnObjectSelectionChanged(const FNavigationToolEditorSelection& InEditorSelection)
-{
-	bool bIsSyncingItemSelection = false;
-
-	// find if any outliner view is syncing item selection
-	// if we are, ignore object selection or it will cause another round of selections next tick
-	ForEachToolView([&bIsSyncingItemSelection](const TSharedPtr<FNavigationToolView>& InToolView)
-		{
-			bIsSyncingItemSelection |= InToolView->IsSyncingItemSelection();
-		});
-
-	if (bIsSyncingItemSelection)
-	{
-		return;
-	}
-	
-	if (!ObjectsLastSelected.IsValid())
-	{
-		ObjectsLastSelected = MakeShared<TArray<TWeakObjectPtr<UObject>>>();
-	}
-
-	TArray<UObject*> SelectedObjects = InEditorSelection.GetSelectedObjects<UObject, ENavigationToolSelectionSource::All>();
-	ObjectsLastSelected->Reserve(ObjectsLastSelected->Num() + SelectedObjects.Num());
-
-	for (UObject* const Object : SelectedObjects)
-	{
-		ObjectsLastSelected->AddUnique(Object);
-	}
-}
-
 const FNavigationToolItemProxyRegistry& FNavigationTool::GetItemProxyRegistry() const
 {
 	return ItemProxyRegistry;
@@ -1220,20 +976,9 @@ void FNavigationTool::OnObjectsReplaced(const TMap<UObject*, UObject*>& InReplac
 			InItem->OnObjectsReplaced(InReplacementMap, /*bRecursive*/false);
 		});
 
-	if (ObjectsLastSelected.IsValid())
-	{
-		for (TWeakObjectPtr<UObject>& Object : *ObjectsLastSelected)
-		{
-			if (Object.IsValid() && InReplacementMap.Contains(Object.Get()))
-			{
-				Object = InReplacementMap[Object.Get()];
-			}
-		}
-	}
-
 	for (const TSharedPtr<INavigationToolAction>& Action : PendingActions)
 	{
-		//Recursive needed since we only have direct reference to the underlying item in the Action, not its children
+		// Recursive needed since we only have direct reference to the underlying item in the Action, not its children
 		if (Action.IsValid())
 		{
 			Action->OnObjectsReplaced(InReplacementMap, true);
@@ -1272,20 +1017,93 @@ void FNavigationTool::ShowHideToolTab(const bool bInVisible)
 {
 	ToolTab->ShowHideToolTab(bInVisible);
 
-	if (!bInVisible)
+	if (bInVisible)
+	{
+		BindEvents();
+	}
+	else
 	{
+		UnbindEvents();
+
 		SaveAllProviderStates();
 	}
 }
 
 void FNavigationTool::ToggleToolTabVisible()
 {
-	ToolTab->ToggleToolTabVisible();
+	ShowHideToolTab(!IsToolTabVisible());
+}
 
-	if (!ToolTab->IsToolTabVisible())
+void FNavigationTool::OnTabVisibilityChanged(const bool bInVisible)
+{
+	if (bInVisible)
 	{
-		SaveAllProviderStates();
+		BindEvents();
+	
+		if (GEditor)
+		{
+			GEditor->RegisterForUndo(this);
+		}
+	}
+	else
+	{
+		UnbindEvents();
+
+		if (GEditor)
+		{
+			GEditor->UnregisterForUndo(this);
+		}
+	}
+}
+
+void FNavigationTool::BindEvents()
+{
+	if (const TSharedPtr<FSequencer> Sequencer = StaticCastSharedPtr<FSequencer>(WeakSequencer.Pin()))
+	{
+		Sequencer->OnChannelChanged().AddSPLambda(this, [this](const FMovieSceneChannelMetaData*, UMovieSceneSection*)
+			{
+				OnTreeViewChanged();
+			});
+		Sequencer->OnMovieSceneBindingsChanged().AddSP(this, &FNavigationTool::OnTreeViewChanged);
+		Sequencer->OnMovieSceneBindingsPasted().AddSPLambda(this, [this](const TArray<FMovieSceneBinding>&)
+			{
+				OnTreeViewChanged();
+			});
+		Sequencer->OnMovieSceneDataChanged().AddSPLambda(this, [this](const EMovieSceneDataChangeType)
+			{
+				OnTreeViewChanged();
+			});
+
+		Sequencer->OnTreeViewChanged().AddSP(this, &FNavigationTool::OnTreeViewChanged);
+		Sequencer->OnEndScrubbingEvent().AddSP(this, &FNavigationTool::OnTreeViewChanged);
+
+		Sequencer->GetSelection().Outliner.OnChanged.AddSP(this, &FNavigationTool::OnSequencerSelectionChanged);
+		Sequencer->GetSelection().TrackArea.OnChanged.AddSP(this, &FNavigationTool::OnSequencerSelectionChanged);
+	}
+
+	// Listen to object replacement changes
+	FCoreUObjectDelegates::OnObjectsReplaced.AddSP(this, &FNavigationTool::OnObjectsReplaced);
+	FEditorDelegates::OnEditorActorReplaced.AddSP(this, &FNavigationTool::OnActorReplaced);
+}
+
+void FNavigationTool::UnbindEvents()
+{
+	if (const TSharedPtr<FSequencer> Sequencer = StaticCastSharedPtr<FSequencer>(WeakSequencer.Pin()))
+	{
+		Sequencer->OnChannelChanged().RemoveAll(this);
+		Sequencer->OnMovieSceneBindingsChanged().RemoveAll(this);
+		Sequencer->OnMovieSceneBindingsPasted().RemoveAll(this);
+		Sequencer->OnMovieSceneDataChanged().RemoveAll(this);
+
+		Sequencer->OnTreeViewChanged().RemoveAll(this);
+		Sequencer->OnEndScrubbingEvent().RemoveAll(this);
+
+		Sequencer->GetSelection().Outliner.OnChanged.RemoveAll(this);
+		Sequencer->GetSelection().TrackArea.OnChanged.RemoveAll(this);
 	}
+
+	FCoreUObjectDelegates::OnObjectsReplaced.RemoveAll(this);
+	FEditorDelegates::OnEditorActorReplaced.RemoveAll(this);
 }
 
 void FNavigationTool::AddItem(const FNavigationToolItemPtr& InItem)
@@ -1355,7 +1173,7 @@ void FNavigationTool::OnTreeViewChanged()
 	RequestRefresh();
 }
 
-void FNavigationTool::OnSequencerOutlinerSelectionChanged()
+void FNavigationTool::OnSequencerSelectionChanged()
 {
 	UNavigationToolSettings* const ToolSettings = GetMutableDefault<UNavigationToolSettings>();
 	if (!ToolSettings->ShouldSyncSelectionToNavigationTool())
@@ -1363,42 +1181,8 @@ void FNavigationTool::OnSequencerOutlinerSelectionChanged()
 		return;
 	}
 
-	const TSharedPtr<FSequencer> Sequencer = StaticCastSharedPtr<FSequencer>(WeakSequencer.Pin());
-	if (!Sequencer.IsValid())
-	{
-		return;
-	}
-
-	FSequencerSelection& SequencerSelection = Sequencer->GetSelection();
-
-	// Unbind to avoid circular loop
-	SequencerSelection.Outliner.OnChanged.RemoveAll(this);
-
-	// Temporarily stop syncing selection to sequencer to avoid circular loop
-	// Calling SaveConfig() here will cause a massive slowdown and throw a log about "DeleteFile"
-	// in some cases. Not able to reproduce this locally.
-	const bool bPreviouslySyncingSelectionToSequencer = ToolSettings->ShouldSyncSelectionToSequencer();
-	ToolSettings->SetSyncSelectionToSequencer(/*bInSync=*/false, /*bInSaveConfig=*/false);
-
-	TArray<FNavigationToolItemPtr> ItemsToSelect;
-
-	for (const FViewModelPtr OutlinerItem : SequencerSelection.Outliner)
-	{
-		ItemsToSelect.Append(TryFindItems(OutlinerItem));
-	}
-
-	SelectItems(ItemsToSelect);
-
-	// Bind from avoiding circular loop
-	SequencerSelection.Outliner.OnChanged.AddSP(this, &FNavigationTool::OnSequencerOutlinerSelectionChanged);
-
-	ToolSettings->SetSyncSelectionToSequencer(bPreviouslySyncingSelectionToSequencer, /*bInSaveConfig=*/false);
-}
-
-void FNavigationTool::OnSequencerTrackAreaSelectionChanged()
-{
-	UNavigationToolSettings* const ToolSettings = GetMutableDefault<UNavigationToolSettings>();
-	if (!ToolSettings->ShouldSyncSelectionToNavigationTool())
+	// If any view is syncing item selection, ignore item selection or it will cause another round of selections next tick
+	if (AreAllViewsSyncingItemSelection())
 	{
 		return;
 	}
@@ -1411,27 +1195,17 @@ void FNavigationTool::OnSequencerTrackAreaSelectionChanged()
 
 	FSequencerSelection& SequencerSelection = Sequencer->GetSelection();
 
-	// Unbind to avoid circular loop
-	SequencerSelection.TrackArea.OnChanged.RemoveAll(this);
-
-	// Calling SaveConfig() here will cause a massive slowdown and throw a log about "DeleteFile"
-	// in some cases. Not able to reproduce this locally.
-	const bool bPreviouslySyncingSelectionToSequencer = ToolSettings->ShouldSyncSelectionToSequencer();
-	ToolSettings->SetSyncSelectionToSequencer(/*bInSync=*/false, /*bInSaveConfig=*/false);
-
-	TArray<FNavigationToolItemPtr> ItemsToSelect;
-
-	for (const FViewModelPtr TrackAreaItem : SequencerSelection.TrackArea)
+	if (!ItemsLastSelected.IsValid())
 	{
-		ItemsToSelect.Append(TryFindItems(TrackAreaItem));
+		ItemsLastSelected = MakeShared<TArray<FNavigationToolItemWeakPtr>>();
 	}
 
-	SelectItems(ItemsToSelect);
+	ItemsLastSelected->Reserve(ItemsLastSelected->Num() + SequencerSelection.Outliner.Num());
 
-	// Bind from avoiding circular loop
-	SequencerSelection.TrackArea.OnChanged.AddSP(this, &FNavigationTool::OnSequencerTrackAreaSelectionChanged);
-
-	ToolSettings->SetSyncSelectionToSequencer(bPreviouslySyncingSelectionToSequencer, /*bInSaveConfig=*/false);
+	for (const FViewModelPtr OutlinerItem : SequencerSelection.Outliner)
+	{
+		ItemsLastSelected->Append(TryFindItems(OutlinerItem));
+	}
 }
 
 TArray<FNavigationToolItemPtr> FNavigationTool::FindItemsFromMovieSceneObject(UObject* const InObject) const
@@ -1469,6 +1243,119 @@ TArray<FNavigationToolItemPtr> FNavigationTool::FindItemsFromObjectGuid(const FG
 	return OutItems;
 }
 
+bool FNavigationTool::AreAllViewsSyncingItemSelection() const
+{
+	bool bIsSyncingItemSelection = false;
+
+	ForEachToolView([&bIsSyncingItemSelection](const TSharedPtr<FNavigationToolView>& InView)
+		{
+			bIsSyncingItemSelection |= InView->IsSyncingItemSelection();
+		});
+
+	return bIsSyncingItemSelection;
+}
+
+void FNavigationTool::OnSequencerActivate(FMovieSceneSequenceIDRef InSequenceId)
+{
+	const TSharedPtr<ISequencer> Sequencer = WeakSequencer.Pin();
+	if (!Sequencer.IsValid() || InSequenceId != Sequencer->GetRootTemplateID())
+	{
+		return;
+	}
+
+	const UNavigationToolSettings* const ToolSettings = GetDefault<UNavigationToolSettings>();
+	if (!ensure(IsValid(ToolSettings)))
+	{
+		return;
+	}
+
+	ForEachProvider([this](const TSharedRef<FNavigationToolProvider>& InProvider)
+		{
+			InProvider->Activate(*this);
+			return true;
+		});
+
+	// We apply default views *after* all columns have been created for all providers
+	if (ToolSettings->ShouldApplyDefaultColumnView())
+	{
+		TArray<FName> DefaultColumnViews;
+
+		ForEachProvider([&DefaultColumnViews](const TSharedRef<FNavigationToolProvider>& InProvider)
+			{
+				const FName DefaultColumnView = InProvider->GetDefaultColumnView();
+				if (!DefaultColumnView.IsNone())
+				{
+					DefaultColumnViews.AddUnique(DefaultColumnView);
+				}
+				return true;
+			});
+
+		// @TODO: priority int to better find a default priority view instead of just using the first index?
+		if (!DefaultColumnViews.IsEmpty())
+		{
+			ForEachToolView([&DefaultColumnViews](const TSharedPtr<FNavigationToolView>& InToolView)
+				{
+					InToolView->ApplyCustomColumnView(DefaultColumnViews[0]);
+				});
+		}
+	}
+
+	RefreshGlobalFilters();
+
+	ToolTab->Init();
+
+	OnTreeViewChanged();
+}
+
+void FNavigationTool::OnSequencerClose(const TSharedRef<ISequencer> InSequencer)
+{
+	const TSharedPtr<ISequencer> Sequencer = WeakSequencer.Pin();
+	if (!Sequencer.IsValid() || InSequencer != Sequencer)
+	{
+		return;
+	}
+
+	ForEachProvider([this](const TSharedRef<FNavigationToolProvider>& InProvider)
+		{
+			InProvider->Deactivate(*this);
+			return true;
+		});
+
+	ToolTab->Shutdown();
+
+	Shutdown();
+}
+
+void FNavigationTool::OnProvidersChanged(const FName InToolId
+	, const TSharedRef<FNavigationToolProvider>& InProvider
+	, const ENavigationToolProvidersChangeType InChangeType)
+{
+	if (InToolId != FNavigationToolExtender::GetToolInstanceId(*this))
+	{
+		return;
+	}
+
+	const TSharedPtr<ISequencer> Sequencer = GetSequencer();
+	if (!Sequencer.IsValid()
+		|| !InProvider->IsSequenceSupported(Sequencer->GetRootMovieSceneSequence()))
+	{
+		return;
+	}
+
+	if (InChangeType == ENavigationToolProvidersChangeType::Add)
+	{
+		InProvider->Activate(*this);
+	}
+	else if (InChangeType == ENavigationToolProvidersChangeType::Remove)
+	{
+		InProvider->Deactivate(*this);
+	}
+
+	RefreshGlobalFilters();
+
+	OnTreeViewChanged();
+}
+
 }
 
 #undef LOCTEXT_NAMESPACE
diff --git a/Engine/Source/Editor/Sequencer/Private/Tools/NavigationTool/NavigationTool.h b/Engine/Source/Editor/Sequencer/Private/Tools/NavigationTool/NavigationTool.h
--- a/Engine/Source/Editor/Sequencer/Private/Tools/NavigationTool/NavigationTool.h
+++ b/Engine/Source/Editor/Sequencer/Private/Tools/NavigationTool/NavigationTool.h
@@ -22,8 +22,6 @@ class FUICommandList;
 class ISequencer;
 class UMovieSceneSequence;
 class USequencerSettings;
-class USceneComponent;
-enum class EItemDropZone;
 
 namespace UE::Sequencer::NavigationTool
 {
@@ -47,9 +45,6 @@ public:
 	void Init();
 	void Shutdown();
 
-	void ForEachProvider(const TFunction<bool(const TSharedRef<FNavigationToolProvider>& InToolProvider)>& InPredicate);
-
-	void LoadAllProviderStates(const bool bInRefreshTool = true);
 	void SaveAllProviderStates();
 
 	bool CanProcessSequenceSpawn(UMovieSceneSequence* const InSequence) const;
@@ -69,6 +64,8 @@ public:
 
 	//~ Begin INavigationTool
 
+	virtual void ForEachProvider(const TFunction<bool(const TSharedRef<FNavigationToolProvider>& InToolProvider)>& InPredicate) const override;
+
 	virtual TSharedPtr<FUICommandList> GetBaseCommandList() const override;
 
 	virtual bool IsToolTabVisible() const override;
@@ -100,7 +97,7 @@ public:
 
 	virtual void SetIgnoreNotify(const ENavigationToolIgnoreNotifyFlags InFlag, const bool bInIgnore) override;
 
-	virtual void OnObjectSelectionChanged(const FNavigationToolEditorSelection& InEditorSelection) override;
+	virtual void OnSequencerSelectionChanged() override;
 	virtual TArray<FNavigationToolItemPtr> GetSelectedItems(const bool bInNormalizeToTopLevelSelections = false) const override;
 	virtual void SelectItems(const TArray<FNavigationToolItemPtr>& InItems
 		, const ENavigationToolItemSelectionFlags InFlags = ENavigationToolItemSelectionFlags::SignalSelectionChange) const override;
@@ -142,7 +139,7 @@ public:
 	void UpdateRecentToolViews(const int32 InToolViewId);
 
 	/** Executes the given predicate for each Navigation Tool View registered */
-	void ForEachToolView(const TFunction<void(const TSharedPtr<FNavigationToolView>& InToolView)>& InPredicate) const;
+	void ForEachToolView(const TFunction<void(const TSharedRef<FNavigationToolView>& InToolView)>& InPredicate) const;
 
 	/** @return Number of actions that been added to the queue so far before triggering a refresh */
 	int32 GetPendingItemActionCount() const;
@@ -159,9 +156,6 @@ public:
 	/** Gets the closest item to all the given items while also being their common ancestor */
 	static FNavigationToolItemPtr FindLowestCommonAncestor(const TArray<FNavigationToolItemPtr>& Items);
 
-	/** Converts the given Navigation Tool Item to a Scene Item that can be serialized in the Scene Tree */
-	static FNavigationToolSerializedItem MakeSerializedItemFromToolItem(const FNavigationToolItemPtr& InItem);
-
 	/**
 	 * Sort the given array of items based on their ordering in the Navigation Tool
 	 * @see FNavigationTool::CompareToolItemOrder
@@ -183,42 +177,38 @@ public:
 	/** Marks the Navigation Tool dirty. This triggers INavigationToolProvider::OnToolModified on next tick */
 	void SetToolModified();
 
+	bool DoesGlobalFilterExist(const FName InFilterId);
+	void RefreshGlobalFilters();
+
 private:
 	friend class FNavigationToolView;
 	friend class FNavigationToolToolbarMenu;
 
-	void OnProvidersChanged(const TSharedRef<INavigationTool>& InTool
-		, const TSharedRef<FNavigationToolProvider>& InProvider
-		, const ENavigationToolProvidersChangeType InChangeType);
+	void OnTabVisibilityChanged(const bool bInVisible);
+
+	void BindEvents();
+	void UnbindEvents();
 
 	void AddItem(const FNavigationToolItemPtr& InItem);
 	void RemoveItem(const FNavigationToolItemId& InItemId);
 	void ForEachItem(TFunctionRef<void(const FNavigationToolItemPtr&)> InFunc);
 
-	void SaveProviderState(const TSharedRef<FNavigationToolProvider>& InProvider);
-	void LoadProviderState(const TSharedRef<FNavigationToolProvider>& InProvider);
-
-	/** Save all Navigation Tool items */
-	void SaveSerializedTree(const bool bInResetTree);
-
-	/** Save a Navigation Tool item */
-	void SaveSerializedTreeRecursive(const FNavigationToolItemPtr& InParentItem
-		, FNavigationToolSerializedTree& InSerializedTree);
-
-	/** Load a Navigation Tool item */
-	void LoadSerializedTree(const FNavigationToolItemPtr& InParentItem
-		, FNavigationToolSerializedTree* const InSerializedTree);
-
 	USequencerSettings* GetSequencerSettings() const;
 
 	void OnTreeViewChanged();
 
-	void OnSequencerOutlinerSelectionChanged();
-	void OnSequencerTrackAreaSelectionChanged();
-
 	TArray<FNavigationToolItemPtr> FindItemsFromMovieSceneObject(UObject* const InObject) const;
 	TArray<FNavigationToolItemPtr> FindItemsFromObjectGuid(const FGuid& InObjectGuid) const;
 
+	bool AreAllViewsSyncingItemSelection() const;
+
+	void OnSequencerActivate(FMovieSceneSequenceIDRef InSequenceId);
+	void OnSequencerClose(const TSharedRef<ISequencer> InSequencer);
+
+	void OnProvidersChanged(const FName InToolId
+		, const TSharedRef<FNavigationToolProvider>& InProvider
+		, const ENavigationToolProvidersChangeType InChangeType);
+
 	TWeakPtr<ISequencer> WeakSequencer;
 
 	TSharedRef<FNavigationToolTab> ToolTab;
@@ -242,8 +232,8 @@ private:
 	/** The current pending actions before refresh is called */
 	TArray<TSharedPtr<INavigationToolAction>> PendingActions;
 
-	/** The list of objects pending selection processing. filled in when getting notifies from the USelection instances */
-	TSharedPtr<TArray<TWeakObjectPtr<UObject>>> ObjectsLastSelected;
+	/** The list of items pending selection processing, filled in when the Sequencer selection changes */
+	TSharedPtr<TArray<FNavigationToolItemWeakPtr>> ItemsLastSelected;
 
 	/** The map of registered outliner views */
 	TMap<int32, TSharedPtr<FNavigationToolView>> ToolViews;
diff --git a/Engine/Source/Editor/Sequencer/Private/Tools/NavigationTool/NavigationToolTab.cpp b/Engine/Source/Editor/Sequencer/Private/Tools/NavigationTool/NavigationToolTab.cpp
--- a/Engine/Source/Editor/Sequencer/Private/Tools/NavigationTool/NavigationToolTab.cpp
+++ b/Engine/Source/Editor/Sequencer/Private/Tools/NavigationTool/NavigationToolTab.cpp
@@ -182,6 +182,8 @@ TSharedRef<SDockTab> FNavigationToolTab::SpawnToolTab(const FSpawnTabArgs& InArg
 				{
 					Settings->SetNavigationToolVisible(false);
 				}
+
+				VisibilityChangedDelegate.Broadcast(false);
 			})
 		[
 			ToolWidget.ToSharedRef()
@@ -193,6 +195,8 @@ TSharedRef<SDockTab> FNavigationToolTab::SpawnToolTab(const FSpawnTabArgs& InArg
 		SequencerSettings->SetNavigationToolVisible(true);
 	}
 
+	VisibilityChangedDelegate.Broadcast(true);
+
 	OwnerTool.Refresh();
 
 	return ToolTab;
diff --git a/Engine/Source/Editor/Sequencer/Private/Tools/NavigationTool/NavigationToolTab.h b/Engine/Source/Editor/Sequencer/Private/Tools/NavigationTool/NavigationToolTab.h
--- a/Engine/Source/Editor/Sequencer/Private/Tools/NavigationTool/NavigationToolTab.h
+++ b/Engine/Source/Editor/Sequencer/Private/Tools/NavigationTool/NavigationToolTab.h
@@ -20,6 +20,8 @@ class INavigationTool;
 class FNavigationToolTab : public TSharedFromThis<FNavigationToolTab>
 {
 public:
+	DECLARE_MULTICAST_DELEGATE_OneParam(FOnVisibilityChanged, const bool /*bInVisible*/);
+
 	/** @return Returns the auto generated tool tab id using the passed in ISequencer.
 	 * Tool tab ids are based on the sequencer instance settings name. */
 	static FName StaticToolTabId(const ISequencer& InSequencer);
@@ -40,6 +42,8 @@ public:
 
 	void ToggleToolTabVisible();
 
+	FOnVisibilityChanged& OnVisibilityChanged() { return VisibilityChangedDelegate; }
+
 protected:
 	TSharedPtr<IToolkitHost> GetToolkitHost() const;
 	TSharedPtr<FTabManager> GetTabManager() const;
@@ -57,6 +61,8 @@ protected:
 
 	/** State variable to help the dock tab know when it should save its visibility */
 	bool bShuttingDown = false;
+
+	FOnVisibilityChanged VisibilityChangedDelegate;
 };
 
 }
diff --git a/Engine/Source/Editor/Sequencer/Private/Tools/NavigationTool/NavigationToolView.cpp b/Engine/Source/Editor/Sequencer/Private/Tools/NavigationTool/NavigationToolView.cpp
--- a/Engine/Source/Editor/Sequencer/Private/Tools/NavigationTool/NavigationToolView.cpp
+++ b/Engine/Source/Editor/Sequencer/Private/Tools/NavigationTool/NavigationToolView.cpp
@@ -10,6 +10,7 @@
 #include "Misc/MessageDialog.h"
 #include "MVVM/Selection/Selection.h"
 #include "Sequencer.h"
+#include "SequencerLog.h"
 #include "Filters/Filters/NavigationToolBuiltInFilter.h"
 #include "Styling/SlateTypes.h"
 #include "Tools/NavigationTool/Columns/NavigationToolColumn.h"
@@ -470,8 +471,15 @@ void FNavigationToolView::SaveViewItemFlags(const FNavigationToolItemPtr& InItem
 	}
 
 	FNavigationToolSaveState* const SaveState = Provider->GetSaveState(*Tool);
-	if (!ensure(SaveState))
+	if (!SaveState)
 	{
+		UE_LOG(LogSequencer, Warning, TEXT("SaveViewItemFlags(): Navigation Tool Save State is NULL!"));
+		return;
+	}
+
+	if (!SaveState->ToolViewSaveStates.IsValidIndex(ToolViewId))
+	{
+		UE_LOG(LogSequencer, Warning, TEXT("SaveViewItemFlags(): '%d' is in invalid saved tool view state index"), ToolViewId);
 		return;
 	}
 
@@ -704,13 +712,13 @@ void FNavigationToolView::ClearItemSelection(bool bSignalSelectionChange)
 
 void FNavigationToolView::SetItemSelectionImpl(TArray<FNavigationToolItemPtr>&& InItems, bool bSignalSelectionChange)
 {
-	if (ShouldAutoExpandToSelection())
+	/*if (ShouldAutoExpandToSelection())
 	{
 		for (const FNavigationToolItemPtr& Item : InItems)
 		{
 			SetParentItemExpansions(Item, true);
 		}
-	}
+	}*/
 
 	SelectedItems = MoveTemp(InItems);
 
@@ -910,7 +918,7 @@ bool FNavigationToolView::ShouldShowItem(const FNavigationToolItemPtr& InItem, b
 
 	// Allow providers to determine whether the item should be hidden
 	bool bProviderShouldHideItem = false;
-	FNavigationToolExtender::ForEachToolProvider(*Tool, [&bProviderShouldHideItem, &InItem]
+	Tool->ForEachProvider([&bProviderShouldHideItem, &InItem]
 		(const TSharedRef<FNavigationToolProvider>& InToolProvider)
 		{
 			if (InToolProvider->ShouldHideItem(InItem))
@@ -1670,7 +1678,7 @@ void FNavigationToolView::EnsureToolViewCount(const int32 InToolViewId) const
 
 	FNavigationTool& ToolRef = *Tool;
 
-	FNavigationToolExtender::ForEachToolProvider(ToolRef, [&ToolRef, InToolViewId]
+	ToolRef.ForEachProvider([&ToolRef, InToolViewId]
 		(const TSharedRef<FNavigationToolProvider>& InToolProvider)
 		{
 			InToolProvider->EnsureToolViewCount(ToolRef, InToolViewId);
diff --git a/Engine/Source/Editor/Sequencer/Private/Tools/NavigationTool/Providers/NavigationToolProvider.cpp b/Engine/Source/Editor/Sequencer/Private/Tools/NavigationTool/Providers/NavigationToolProvider.cpp
--- a/Engine/Source/Editor/Sequencer/Private/Tools/NavigationTool/Providers/NavigationToolProvider.cpp
+++ b/Engine/Source/Editor/Sequencer/Private/Tools/NavigationTool/Providers/NavigationToolProvider.cpp
@@ -1,7 +1,11 @@
 ﻿// Copyright Epic Games, Inc. All Rights Reserved.
 
 #include "Tools/NavigationTool/Providers/NavigationToolProvider.h"
-#include "Tools/NavigationTool/INavigationTool.h"
+#include "SequencerLog.h"
+#include "Tools/NavigationTool/Items/NavigationToolItemUtils.h"
+#include "Tools/NavigationTool/NavigationTool.h"
+#include "Tools/NavigationTool/NavigationToolSettings.h"
+#include "Tools/NavigationTool/NavigationToolView.h"
 
 namespace UE::Sequencer::NavigationTool
 {
@@ -31,7 +35,7 @@ void FNavigationToolProvider::OnExtendBuiltInFilters(TArray<FNavigationToolBuilt
 void FNavigationToolProvider::UpdateItemIdContexts(const INavigationTool& InTool)
 {
 	FNavigationToolSaveState* const SaveState = GetSaveState(InTool);
-	if (!ensure(SaveState))
+	if (!SaveState)
 	{
 		return;
 	}
@@ -85,13 +89,15 @@ void FNavigationToolProvider::UpdateItemIdContexts(const INavigationTool& InTool
 FNavigationToolViewSaveState* FNavigationToolProvider::GetViewSaveState(const INavigationTool& InTool, const int32 InToolViewId) const
 {
 	FNavigationToolSaveState* const SaveState = GetSaveState(InTool);
-	if (!ensure(SaveState))
+	if (!SaveState)
 	{
+		UE_LOG(LogSequencer, Warning, TEXT("GetViewSaveState(): Navigation Tool Save State is NULL!"));
 		return nullptr;
 	}
 
 	if (!SaveState->ToolViewSaveStates.IsValidIndex(InToolViewId))
 	{
+		UE_LOG(LogSequencer, Warning, TEXT("GetViewSaveState(): '%d' is in invalid saved tool view state index"), InToolViewId);
 		return nullptr;
 	}
 
@@ -101,8 +107,9 @@ FNavigationToolViewSaveState* FNavigationToolProvider::GetViewSaveState(const IN
 void FNavigationToolProvider::EnsureToolViewCount(const INavigationTool& InTool, const int32 InToolViewId)
 {
 	FNavigationToolSaveState* const SaveState = GetSaveState(InTool);
-	if (!ensure(SaveState))
+	if (!SaveState)
 	{
+		UE_LOG(LogSequencer, Warning, TEXT("EnsureToolViewCount(): Navigation Tool Save State is NULL!"));
 		return;
 	}
 
@@ -115,16 +122,6 @@ void FNavigationToolProvider::EnsureToolViewCount(const INavigationTool& InTool,
 	}
 }
 
-const TArray<FName>& FNavigationToolProvider::GetExtendedColumnViewNames() const
-{
-	return ExtendedColumnViewNames;
-}
-
-const TArray<FName>& FNavigationToolProvider::GetExtendedBuiltInFilterNames() const
-{
-	return ExtendedBuiltInFilterNames;
-}
-
 TOptional<EItemDropZone> FNavigationToolProvider::OnToolItemCanAcceptDrop(const FDragDropEvent& InDragDropEvent
 	, EItemDropZone InDropZone
 	, const FNavigationToolItemPtr& InTargetItem) const
@@ -139,4 +136,181 @@ FReply FNavigationToolProvider::OnToolItemAcceptDrop(const FDragDropEvent& InDra
 	return FReply::Unhandled();
 }
 
+void FNavigationToolProvider::Activate(FNavigationTool& InTool)
+{
+	if (const TSharedPtr<FUICommandList> BaseCommandList = InTool.GetBaseCommandList())
+	{
+		BindCommands(BaseCommandList.ToSharedRef());
+	}
+
+	LoadState(InTool);
+
+	const TSharedRef<FNavigationToolProvider> SharedThisRef = SharedThis(this);
+
+	InTool.ForEachToolView([&SharedThisRef](const TSharedRef<FNavigationToolView>& InToolView)
+		{
+			//InToolView->LoadViewState(InProvider);
+
+			InToolView->CreateColumns(SharedThisRef);
+			InToolView->CreateDefaultColumnViews(SharedThisRef);
+		});
+
+	OnActivate();
+}
+
+void FNavigationToolProvider::Deactivate(FNavigationTool& InTool)
+{
+	CleanupExtendedColumnViews();
+
+	OnDeactivate();
+}
+
+void FNavigationToolProvider::CleanupExtendedColumnViews()
+{
+	UNavigationToolSettings* const Settings = GetMutableDefault<UNavigationToolSettings>();
+	if (!ensure(IsValid(Settings)))
+	{
+		return;
+	}
+
+	TSet<FNavigationToolColumnView>& CustomColumnViews = Settings->GetCustomColumnViews();
+
+	for (const FName ColumnViewName : ExtendedColumnViewNames)
+	{
+		CustomColumnViews.Remove(ColumnViewName);
+	}
+
+	Settings->SaveConfig();
+}
+
+bool FNavigationToolProvider::IsSequenceSupported(UMovieSceneSequence* const InSequence) const
+{
+	if (IsValid(InSequence))
+	{
+		return GetSupportedSequenceClasses().Contains(InSequence->GetClass());
+	}
+	return false;
+}
+
+void FNavigationToolProvider::SaveState(FNavigationTool& InTool)
+{
+	SaveSerializedTree(InTool, /*bInResetTree*/true);
+
+	const TSharedRef<FNavigationToolProvider> SharedThisRef = SharedThis(this);
+
+	InTool.ForEachToolView([&SharedThisRef](const TSharedPtr<FNavigationToolView>& InToolView)
+		{
+			InToolView->SaveViewState(SharedThisRef);
+		});
+
+	if (FNavigationToolSaveState* const SaveState = GetSaveState(InTool))
+	{
+		// Remove any saved item colors that can no longer be found
+		TArray<FString> ItemIds;
+		SaveState->ItemColorMap.GetKeys(ItemIds);
+		for (const FString& ItemId : ItemIds)
+		{
+			if (!InTool.FindItem(ItemId).IsValid())
+			{
+				SaveState->ItemColorMap.Remove(ItemId);
+			}
+		}
+	}
+
+	UpdateItemIdContexts(InTool);
+}
+
+void FNavigationToolProvider::LoadState(FNavigationTool& InTool)
+{
+	const UNavigationToolSettings* const ToolSettings = GetDefault<UNavigationToolSettings>();
+	if (!ensure(IsValid(ToolSettings)))
+	{
+		return;
+	}
+
+	UpdateItemIdContexts(InTool);
+
+	if (FNavigationToolSaveState* const SaveState = GetSaveState(InTool))
+	{
+		LoadSerializedTree(InTool.GetTreeRoot(), &SaveState->SerializedTree);
+	}
+
+	const TSharedRef<FNavigationToolProvider> SharedThisRef = SharedThis(this);
+
+	InTool.ForEachToolView([&SharedThisRef](const TSharedPtr<FNavigationToolView>& InToolView)
+		{
+			InToolView->LoadViewState(SharedThisRef);
+		});
+}
+
+void FNavigationToolProvider::SaveSerializedTree(FNavigationTool& InTool, const bool bInResetTree)
+{
+	FNavigationToolSaveState* const SaveState = GetSaveState(InTool);
+	if (!SaveState)
+	{
+		return;
+	}
+
+	if (bInResetTree)
+	{
+		SaveState->SerializedTree.Reset();
+	}
+
+	SaveSerializedTreeRecursive(InTool.GetTreeRoot(), SaveState->SerializedTree);
+}
+
+void FNavigationToolProvider::SaveSerializedTreeRecursive(const FNavigationToolItemPtr& InParentItem
+	, FNavigationToolSerializedTree& InSerializedTree)
+{
+	const TArray<FNavigationToolItemPtr>& Children = InParentItem->GetChildren();
+
+	const FNavigationToolSerializedItem ParentSceneItem = InParentItem->MakeSerializedItem();//MakeSerializedItemFromToolItem(InParentItem);
+
+	for (int32 Index = 0; Index < Children.Num(); ++Index)
+	{
+		if (const FNavigationToolItemPtr& ChildItem = Children[Index])
+		{
+			if (ChildItem->ShouldSort())
+			{
+				const FNavigationToolSerializedItem SceneItem = ChildItem->MakeSerializedItem();
+				if (SceneItem.IsValid())
+				{
+					InSerializedTree.GetOrAddTreeNode(SceneItem, ParentSceneItem);
+				}
+			}
+			SaveSerializedTreeRecursive(ChildItem, InSerializedTree);
+		}
+	}
+}
+
+void FNavigationToolProvider::LoadSerializedTree(const FNavigationToolItemPtr& InParentItem
+	, FNavigationToolSerializedTree* const InSerializedTree)
+{
+	TArray<FNavigationToolItemPtr>& Children = InParentItem->GetChildrenMutable();
+	
+	TArray<FNavigationToolItemPtr> Sortable;
+	TArray<FNavigationToolItemPtr> Unsortable;
+	SplitSortableAndUnsortableItems(Children, Sortable, Unsortable);
+
+	// If Scene Tree is valid, Item Sorting should be empty as this function only takes a valid Scene Tree if
+	// loaded version supports Scene Trees (i.e. when Item Sorting stops being loaded in)
+	if (InSerializedTree)
+	{
+		Sortable.Sort([InSerializedTree](const FNavigationToolItemPtr& InItemA, const FNavigationToolItemPtr& InItemB)
+			{
+				const FNavigationToolSerializedTreeNode* const NodeA = InSerializedTree->FindTreeNode(InItemA->MakeSerializedItem());
+				const FNavigationToolSerializedTreeNode* const NodeB = InSerializedTree->FindTreeNode(InItemB->MakeSerializedItem());
+				return FNavigationToolSerializedTree::CompareTreeItemOrder(NodeA, NodeB);
+			});
+	}
+
+	Children = MoveTemp(Unsortable);
+	Children.Append(MoveTemp(Sortable));
+
+	for (FNavigationToolItemPtr& Child : Children)
+	{
+		LoadSerializedTree(Child, InSerializedTree);
+	}
+}
+
 }
diff --git a/Engine/Source/Editor/Sequencer/Public/Tools/NavigationTool/Customization/NavigationToolExtender.h b/Engine/Source/Editor/Sequencer/Public/Tools/NavigationTool/Customization/NavigationToolExtender.h
--- a/Engine/Source/Editor/Sequencer/Public/Tools/NavigationTool/Customization/NavigationToolExtender.h
+++ b/Engine/Source/Editor/Sequencer/Public/Tools/NavigationTool/Customization/NavigationToolExtender.h
@@ -19,7 +19,7 @@ enum class ENavigationToolProvidersChangeType
 };
 
 DECLARE_MULTICAST_DELEGATE_ThreeParams(FOnNavigationToolProvidersChanged
-	, const TSharedRef<INavigationTool>& /*InTool*/
+	, const FName /*InToolId*/
 	, const TSharedRef<FNavigationToolProvider>& /*InProvider*/
 	, const ENavigationToolProvidersChangeType /*InChangeType*/);
 
@@ -40,52 +40,38 @@ public:
 	 * @param InTool The Navigation Tool instance for which the name is being retrieved.
 	 * @return The name of the Navigation Tool instance if available; otherwise, NAME_None.
 	 */
-	SEQUENCER_API static FName GetToolInstanceName(const INavigationTool& InTool);
+	SEQUENCER_API static FName GetToolInstanceId(const INavigationTool& InTool);
 
 	/**
 	 * Registers a Navigation Tool provider for a specified Navigation Tool instance.
 	 * This method ensures that the provider is associated with the given tool and replaces any
 	 * existing provider with the same identifier, if one exists.
-	 * @param InTool The Navigation Tool instance for which the provider will be registered.
+	 * @param InToolId The Navigation Tool Id for which the provider will be registered.
 	 * @param InProvider The Navigation Tool provider being registered.
 	 * @return True if the provider was successfully registered, false otherwise.
 	 */
-	SEQUENCER_API static bool RegisterToolProvider(const TSharedRef<INavigationTool>& InTool
-		, const TSharedRef<FNavigationToolProvider>& InProvider);
-
-	/**
-	 * Registers a Navigation Tool provider for a specified Navigation Tool instance.
-	 * Unregisters all other providers before registering the given provider.
-	 * @param InTool The Navigation Tool instance for which the provider will be registered.
-	 * @param InProvider The Navigation Tool provider being registered.
-	 * @return True if the provider was successfully registered, false otherwise.
-	 */
-	SEQUENCER_API static bool RegisterOnlyToolProvider(const TSharedRef<INavigationTool>& InTool
+	SEQUENCER_API static bool RegisterToolProvider(const FName InToolId
 		, const TSharedRef<FNavigationToolProvider>& InProvider);
 
 	/**
 	 * Unregisters a Navigation Tool provider associated with a specific instance and identifier.
-	 * @param InTool A shared reference to the Navigation Tool from which the provider needs to be unregistered.
+	 * @param InToolId The Id of the Navigation Tool from which the provider needs to be unregistered.
 	 * @param InProviderId The identifier of the provider to be unregistered.
 	 * @return True if the provider was successfully unregistered; false otherwise.
 	 */
-	SEQUENCER_API static bool UnregisterToolProvider(const TSharedRef<INavigationTool>& InTool, const FName InProviderId);
+	SEQUENCER_API static bool UnregisterToolProvider(const FName InToolId, const FName InProviderId);
 
-
-	/**
-	 * Unregisters all Navigation Tool providers associated with a given instance.
-	 * @param InTool The Navigation Tool for which all associated providers should be unregistered.
-	 */
-	SEQUENCER_API static void UnregisterAllToolProviders(const TSharedRef<INavigationTool>& InTool);
+	/**  */
+	SEQUENCER_API static TSharedPtr<FNavigationToolProvider> FindToolProvider(const FName InToolId, const FName InProviderId);
 
 	/**
 	 * Finds the tool providers associated with the given Navigation Tool instance.
-	 * @param InTool The Navigation Tool instance for which to find associated providers.
+	 * @param InToolId The Navigation Tool Id for which to find associated providers.
 	 * @param OutProviders A set that will be populated with the found tool providers as weak pointers.
 	 * @return True if tool providers were found and populated into OutProviders; otherwise, false.
 	 */
-	SEQUENCER_API static bool FindToolProviders(const INavigationTool& InTool
-		, TSet<TWeakPtr<FNavigationToolProvider>>& OutProviders);
+	SEQUENCER_API static bool FindToolProviders(const FName InToolId
+		, TSet<TSharedRef<FNavigationToolProvider>>& OutProviders);
 
 	/**
 	 * Iterates over each Navigation Tool provider and executes the predicate function.
@@ -93,18 +79,18 @@ public:
 	 * @param InPredicate A function that accepts a reference to an INavigationTool and a shared reference
 	 *                    to FNavigationToolProvider. Should return true to continue iteration or false to stop.
 	 */
-	SEQUENCER_API static void ForEachProvider(const TFunction<bool(INavigationTool& InTool
+	SEQUENCER_API static void ForEachProvider(const TFunction<bool(const FName InToolId
 		, const TSharedRef<FNavigationToolProvider>& InProvider)>& InPredicate);
 
 	/**
 	 * Iterates over each Navigation Tool provider and executes the predicate function associated with the given
 	 * Navigation Tool instance. The iteration is controlled by the provided predicate function, which determines
 	 * whether it continues or stops.
-	 * @param InTool The Navigation Tool instance whose providers are to be iterated over.
+	 * @param InToolId The Navigation Tool Id whose providers are to be iterated over.
 	 * @param InPredicate A predicate function that takes a shared reference to a Navigation Tool provider as input.
 	 *                    The function should return true to continue iteration or false to stop it.
 	 */
-	SEQUENCER_API static void ForEachToolProvider(INavigationTool& InTool
+	SEQUENCER_API static void ForEachToolProvider(const FName InToolId
 		, const TFunction<bool(const TSharedRef<FNavigationToolProvider>& InProvider)>& InPredicate);
 
 	/** Event called when a provider has been added or removed from a tool instance. */
@@ -132,36 +118,9 @@ public:
 	SEQUENCER_API static FSlateIcon FindOverrideIcon(TSharedPtr<const FNavigationToolItem> InItem);
 
 protected:
-	struct FNavigationToolProviderEntry
-	{
-		FNavigationToolProviderEntry() = delete;
-		FNavigationToolProviderEntry(const TSharedRef<INavigationTool>& InTool
-			, const TSet<TWeakPtr<FNavigationToolProvider>>& InWeakProviders)
-			: WeakTool(InTool)
-			, WeakProviders(InWeakProviders)
-		{}
-
-		bool operator==(const FNavigationToolProviderEntry& InOther) const
-		{
-			return (WeakTool.IsValid() && InOther.WeakTool.IsValid())
-				&& (WeakTool.Pin() == InOther.WeakTool.Pin());
-		}
-
-		friend uint32 GetTypeHash(const FNavigationToolProviderEntry& InEntry)
-		{
-			return GetTypeHash(InEntry.WeakTool);
-		}
-
-		TWeakPtr<INavigationTool> WeakTool;
-
-		TSet<TWeakPtr<FNavigationToolProvider>> WeakProviders;
-	};
-
 	static FNavigationToolExtender& Get();
 
-	static FNavigationToolProviderEntry* FindToolProvidersEntry(const INavigationTool& InTool);
-
-	static TSharedPtr<FNavigationToolProvider> FindToolProvider(const INavigationTool& InTool, const FName InProviderId);
+	static TSet<TSharedRef<FNavigationToolProvider>>* FindToolProviders(const FName InToolId);
 
 	void RegisterOverridenIcon_Internal(const FNavigationToolItemTypeId& InItemTypeId, const TSharedRef<INavigationToolIconCustomization>& InNavigationToolIconCustomization);
 	void UnregisterOverriddenIcon_Internal(const FNavigationToolItemTypeId& InItemTypeId, const FName& InSpecializationIdentifier);
@@ -174,7 +133,7 @@ protected:
 	TSharedPtr<INavigationToolIconCustomization> GetCustomizationForItem(const TSharedPtr<const FNavigationToolItem>& InItem) const;
 
 	/** Providers from modules that specify their extensions */
-	TSet<FNavigationToolProviderEntry> ToolProviders;
+	TMap<FName, TSet<TSharedRef<FNavigationToolProvider>>> ToolProviders;
 
 	FOnNavigationToolProvidersChanged ProvidersChangedDelegate;
 
diff --git a/Engine/Source/Editor/Sequencer/Public/Tools/NavigationTool/INavigationTool.h b/Engine/Source/Editor/Sequencer/Public/Tools/NavigationTool/INavigationTool.h
--- a/Engine/Source/Editor/Sequencer/Public/Tools/NavigationTool/INavigationTool.h
+++ b/Engine/Source/Editor/Sequencer/Public/Tools/NavigationTool/INavigationTool.h
@@ -55,6 +55,8 @@ public:
 	SEQUENCER_API virtual void ShowHideToolTab(const bool bInVisible) = 0;
 	SEQUENCER_API virtual void ToggleToolTabVisible() = 0;
 
+	virtual void ForEachProvider(const TFunction<bool(const TSharedRef<FNavigationToolProvider>& InToolProvider)>& InPredicate) const = 0;
+
 	/** Gets the Command List that the Navigation Tool Views will use to append their Command Lists to */
 	virtual TSharedPtr<FUICommandList> GetBaseCommandList() const = 0;
 
@@ -189,8 +191,8 @@ public:
  	 */
 	virtual void SetIgnoreNotify(const ENavigationToolIgnoreNotifyFlags InFlag, const bool bInIgnore) = 0;
 
-	/** Called when the objects have been selected and notified through USelection Instances in Mode Tools */
-	virtual void OnObjectSelectionChanged(const FNavigationToolEditorSelection& InEditorSelection) = 0;
+	/** Called when the Sequencer selection has changed */
+	virtual void OnSequencerSelectionChanged() = 0;
 
 	/** Returns the currently selected items in the most recent outliner view (since this list can vary between outliner views) */
 	virtual TArray<FNavigationToolItemPtr> GetSelectedItems(const bool bInNormalizeToTopLevelSelections = false) const = 0;
diff --git a/Engine/Source/Editor/Sequencer/Public/Tools/NavigationTool/Items/INavigationToolItem.h b/Engine/Source/Editor/Sequencer/Public/Tools/NavigationTool/Items/INavigationToolItem.h
--- a/Engine/Source/Editor/Sequencer/Public/Tools/NavigationTool/Items/INavigationToolItem.h
+++ b/Engine/Source/Editor/Sequencer/Public/Tools/NavigationTool/Items/INavigationToolItem.h
@@ -53,6 +53,9 @@ public:
 	/** Gets the module provider responsible for the creation of this item */
 	virtual TSharedPtr<FNavigationToolProvider> GetProvider() const = 0;
 
+	/** Gets the providers saved state */
+	virtual FNavigationToolSaveState* GetProviderSaveState() const = 0;
+
 	/** Used to signal the Scoped Selection that this Item should be Selected */
 	virtual void Select(FNavigationToolScopedSelection& InSelection) const {}
 
@@ -544,6 +547,12 @@ public:
 
 	/** Sets the expansion state of this item */
 	virtual void SetExpansion(const bool bInIsExpanded) = 0;
+
+	/** Converts this Navigation Tool item to a serialized item that can be saved in the sequence */
+	virtual FNavigationToolSerializedItem MakeSerializedItem()
+	{
+		return FNavigationToolSerializedItem(GetItemId().GetStringId());
+	}
 };
 
 /** Adds Scoped Item Flags, removes them when out of scope. Useful for temp checks like IgnorePendingKill */
diff --git a/Engine/Source/Editor/Sequencer/Public/Tools/NavigationTool/Items/NavigationToolItem.h b/Engine/Source/Editor/Sequencer/Public/Tools/NavigationTool/Items/NavigationToolItem.h
--- a/Engine/Source/Editor/Sequencer/Public/Tools/NavigationTool/Items/NavigationToolItem.h
+++ b/Engine/Source/Editor/Sequencer/Public/Tools/NavigationTool/Items/NavigationToolItem.h
@@ -34,6 +34,7 @@ public:
 	//~ Begin INavigationToolItem
 	virtual INavigationTool& GetOwnerTool() const final override;
 	virtual TSharedPtr<FNavigationToolProvider> GetProvider() const final override;
+	virtual FNavigationToolSaveState* GetProviderSaveState() const override;
 	virtual bool IsItemValid() const override;
 	virtual void RefreshChildren() override;
 	virtual void ResetChildren() override;
diff --git a/Engine/Source/Editor/Sequencer/Public/Tools/NavigationTool/Providers/INavigationToolProvider.h b/Engine/Source/Editor/Sequencer/Public/Tools/NavigationTool/Providers/INavigationToolProvider.h
--- a/Engine/Source/Editor/Sequencer/Public/Tools/NavigationTool/Providers/INavigationToolProvider.h
+++ b/Engine/Source/Editor/Sequencer/Public/Tools/NavigationTool/Providers/INavigationToolProvider.h
@@ -2,13 +2,15 @@
 
 #pragma once
 
+#include "Containers/Set.h"
 #include "Templates/SharedPointer.h"
+#include "Templates/SubclassOf.h"
 #include "Tools/NavigationTool/Filters/Filters/NavigationToolBuiltInFilterParams.h"
 #include "Tools/NavigationTool/NavigationToolDefines.h"
 
 class FDragDropEvent;
 class FReply;
-class ISequencer;
+class UMovieSceneSequence;
 struct FNavigationToolColumnView;
 struct FNavigationToolSaveState;
 enum class EItemDropZone;
@@ -31,14 +33,20 @@ public:
 	/** @return Unique identifier name used to distinguish this provider from others */
 	virtual FName GetIdentifier() const = 0;
 
+	/** Gets the sequence classes that this provider is designed to support */
+	virtual TSet<TSubclassOf<UMovieSceneSequence>> GetSupportedSequenceClasses() const = 0;
+
 	/** @return Saved editor only state data to be restored */
 	virtual FNavigationToolSaveState* GetSaveState(const INavigationTool& InTool) const = 0;
 
 	/** Saves the editor only state data to be restored */
 	virtual void SetSaveState(const INavigationTool& InTool, const FNavigationToolSaveState& InSaveState) const = 0;
 
-	virtual void OnRegistered() {}
-	virtual void OnUnregistered() {}
+	/** Event called when this provider is activated. A provider is activated when a Sequencer opens a sequence. */
+	virtual void OnActivate() {}
+
+	/** Event called when this provider is deactivated. A provider is deactivated when a Sequencer closes. */
+	virtual void OnDeactivate() {}
 
 	virtual void BindCommands(const TSharedRef<FUICommandList>& InCommandList) {}
 
diff --git a/Engine/Source/Editor/Sequencer/Public/Tools/NavigationTool/Providers/NavigationToolProvider.h b/Engine/Source/Editor/Sequencer/Public/Tools/NavigationTool/Providers/NavigationToolProvider.h
--- a/Engine/Source/Editor/Sequencer/Public/Tools/NavigationTool/Providers/NavigationToolProvider.h
+++ b/Engine/Source/Editor/Sequencer/Public/Tools/NavigationTool/Providers/NavigationToolProvider.h
@@ -2,13 +2,17 @@
 
 #pragma once
 
+#include "Tools/NavigationTool/NavigationToolDefines.h"
 #include "Tools/NavigationTool/Providers/INavigationToolProvider.h"
 
+class UMovieSceneSequence;
+struct FNavigationToolSerializedTree;
 struct FNavigationToolViewSaveState;
 
 namespace UE::Sequencer::NavigationTool
 {
 
+class FNavigationTool;
 class INavigationTool;
 
 /** Base Navigation Tool Provider to extend from */
@@ -36,12 +40,30 @@ public:
 	/** Ensure saved editor only data contains an entry for a specific view, creating it if necessary */
 	void EnsureToolViewCount(const INavigationTool& InTool, const int32 InToolViewId);
 
-	const TArray<FName>& GetExtendedColumnViewNames() const;
-	const TArray<FName>& GetExtendedBuiltInFilterNames() const;
+	bool IsSequenceSupported(UMovieSceneSequence* const InSequence) const;
 
 protected:
+	friend class FNavigationTool;
+
+	void Activate(FNavigationTool& InTool);
+	void Deactivate(FNavigationTool& InTool);
+
+	void CleanupExtendedColumnViews();
+
 	TArray<FName> ExtendedColumnViewNames;
 	TArray<FName> ExtendedBuiltInFilterNames;
+
+private:
+	void SaveState(FNavigationTool& InTool);
+	void LoadState(FNavigationTool& InTool);
+
+	void SaveSerializedTree(FNavigationTool& InTool, const bool bInResetTree);
+
+	void SaveSerializedTreeRecursive(const FNavigationToolItemPtr& InParentItem
+		, FNavigationToolSerializedTree& InSerializedTree);
+
+	void LoadSerializedTree(const FNavigationToolItemPtr& InParentItem
+		, FNavigationToolSerializedTree* const InSerializedTree);
 };
 
 }
