[Sequencer] Navigation Tool: Change FNavigationToolAddItem and FNavigationToolRemoveItem constructor param from INavigationTool to FNavigationTool

diff --git a/Engine/Source/Editor/Sequencer/Private/Tools/NavigationTool/ItemActions/NavigationToolAddItem.cpp b/Engine/Source/Editor/Sequencer/Private/Tools/NavigationTool/ItemActions/NavigationToolAddItem.cpp
--- a/Engine/Source/Editor/Sequencer/Private/Tools/NavigationTool/ItemActions/NavigationToolAddItem.cpp
+++ b/Engine/Source/Editor/Sequencer/Private/Tools/NavigationTool/ItemActions/NavigationToolAddItem.cpp
@@ -19,7 +19,7 @@ bool FNavigationToolAddItem::ShouldTransact() const
 	return EnumHasAnyFlags(AddParams.Flags, ENavigationToolAddItemFlags::Transact);
 }
 
-void FNavigationToolAddItem::Execute(INavigationTool& InTool)
+void FNavigationToolAddItem::Execute(FNavigationTool& InTool)
 {
 	if (!AddParams.Item.IsValid())
 	{
@@ -105,7 +105,7 @@ void FNavigationToolAddItem::Execute(INavigationTool& InTool)
 		InTool.SelectItems({ AddParams.Item }, AddParams.SelectionFlags);
 	}
 
-	//InTool.SetToolModified();
+	InTool.SetToolModified();
 }
 
 void FNavigationToolAddItem::OnObjectsReplaced(const TMap<UObject*, UObject*>& InReplacementMap, bool bRecursive)
diff --git a/Engine/Source/Editor/Sequencer/Private/Tools/NavigationTool/ItemActions/NavigationToolRemoveItem.cpp b/Engine/Source/Editor/Sequencer/Private/Tools/NavigationTool/ItemActions/NavigationToolRemoveItem.cpp
--- a/Engine/Source/Editor/Sequencer/Private/Tools/NavigationTool/ItemActions/NavigationToolRemoveItem.cpp
+++ b/Engine/Source/Editor/Sequencer/Private/Tools/NavigationTool/ItemActions/NavigationToolRemoveItem.cpp
@@ -18,7 +18,7 @@ FNavigationToolRemoveItem::FNavigationToolRemoveItem(const FNavigationToolRemove
 	}
 }
 
-void FNavigationToolRemoveItem::Execute(INavigationTool& InTool)
+void FNavigationToolRemoveItem::Execute(FNavigationTool& InTool)
 {
 	if (RemoveParams.Item.IsValid())
 	{
@@ -80,7 +80,7 @@ void FNavigationToolRemoveItem::Execute(INavigationTool& InTool)
 		}
 
 		InTool.UnregisterItem(RemoveParams.Item->GetItemId());
-		//InTool.SetToolModified();
+		InTool.SetToolModified();
 	}
 }
 
diff --git a/Engine/Source/Editor/Sequencer/Public/Tools/NavigationTool/ItemActions/NavigationToolAddItem.h b/Engine/Source/Editor/Sequencer/Public/Tools/NavigationTool/ItemActions/NavigationToolAddItem.h
--- a/Engine/Source/Editor/Sequencer/Public/Tools/NavigationTool/ItemActions/NavigationToolAddItem.h
+++ b/Engine/Source/Editor/Sequencer/Public/Tools/NavigationTool/ItemActions/NavigationToolAddItem.h
@@ -8,7 +8,7 @@
 namespace UE::Sequencer::NavigationTool
 {
 
-class INavigationTool;
+class FNavigationTool;
 
 /**
  * Item action responsible for adding an item to the Tree under a given optional Parent.
@@ -23,7 +23,7 @@ public:
 
 	//~ Begin INavigationToolAction
 	virtual bool ShouldTransact() const override;
-	virtual void Execute(INavigationTool& InTool) override;
+	virtual void Execute(FNavigationTool& InTool) override;
 	virtual void OnObjectsReplaced(const TMap<UObject*, UObject*>& InReplacementMap, bool bRecursive) override;
 	//~ End INavigationToolAction
 
diff --git a/Engine/Source/Editor/Sequencer/Public/Tools/NavigationTool/ItemActions/NavigationToolItemAction.h b/Engine/Source/Editor/Sequencer/Public/Tools/NavigationTool/ItemActions/NavigationToolItemAction.h
--- a/Engine/Source/Editor/Sequencer/Public/Tools/NavigationTool/ItemActions/NavigationToolItemAction.h
+++ b/Engine/Source/Editor/Sequencer/Public/Tools/NavigationTool/ItemActions/NavigationToolItemAction.h
@@ -9,7 +9,7 @@ class UObject;
 namespace UE::Sequencer::NavigationTool
 {
 
-class INavigationTool;
+class FNavigationTool;
 
 /**
  * Interface class for an Action in the Navigation Tool (e.g. Add/Delete/Move Tree Item)
@@ -23,7 +23,7 @@ public:
 	virtual bool ShouldTransact() const { return false; }
 
 	/** The Action to execute on the given Navigation Tool */
-	virtual void Execute(INavigationTool& InTool) = 0;
+	virtual void Execute(FNavigationTool& InTool) = 0;
 
 	/** Replace any Objects that might be held in this Action that has been killed and replaced by a new object (e.g. BP Components) */
 	virtual void OnObjectsReplaced(const TMap<UObject*, UObject*>& InReplacementMap, const bool bInRecursive) = 0;
diff --git a/Engine/Source/Editor/Sequencer/Public/Tools/NavigationTool/ItemActions/NavigationToolRemoveItem.h b/Engine/Source/Editor/Sequencer/Public/Tools/NavigationTool/ItemActions/NavigationToolRemoveItem.h
--- a/Engine/Source/Editor/Sequencer/Public/Tools/NavigationTool/ItemActions/NavigationToolRemoveItem.h
+++ b/Engine/Source/Editor/Sequencer/Public/Tools/NavigationTool/ItemActions/NavigationToolRemoveItem.h
@@ -21,7 +21,7 @@ public:
 	FNavigationToolRemoveItem(const FNavigationToolRemoveItemParams& InRemoveItemParams);
 
 	//~ Begin INavigationToolAction
-	virtual void Execute(INavigationTool& InTool) override;
+	virtual void Execute(FNavigationTool& InTool) override;
 	virtual void OnObjectsReplaced(const TMap<UObject*, UObject*>& InReplacementMap, bool bInRecursive) override;
 	//~ End INavigationToolAction
 
