Sequencer: Rename Snap Keys to Elements and Snap Sections to Elements

diff --git a/Engine/Source/Editor/Sequencer/Private/SSequencer.cpp b/Engine/Source/Editor/Sequencer/Private/SSequencer.cpp
--- a/Engine/Source/Editor/Sequencer/Private/SSequencer.cpp
+++ b/Engine/Source/Editor/Sequencer/Private/SSequencer.cpp
@@ -2523,8 +2523,8 @@ TSharedRef<SWidget> SSequencer::MakeSnapMenu()
 		Toolbar->bSearchable = false;
 		{
 			FToolMenuSection& Section = Toolbar->AddSection( "KeyAndSectionsSnapping", LOCTEXT( "SnappingMenuKeyAndSectionsHeader", "Key and Sections Snapping" ) );
-			Section.AddMenuEntry( FSequencerCommands::Get().ToggleSnapKeyTimesToKeys );
-			Section.AddMenuEntry( FSequencerCommands::Get().ToggleSnapSectionTimesToSections );
+			Section.AddMenuEntry( FSequencerCommands::Get().ToggleSnapKeyTimesToElements );
+			Section.AddMenuEntry( FSequencerCommands::Get().ToggleSnapSectionTimesToElements);
 			Section.AddMenuEntry( FSequencerCommands::Get().ToggleSnapKeysAndSectionsToPlayRange );
 		}
 
diff --git a/Engine/Source/Editor/Sequencer/Private/Sequencer.cpp b/Engine/Source/Editor/Sequencer/Private/Sequencer.cpp
--- a/Engine/Source/Editor/Sequencer/Private/Sequencer.cpp
+++ b/Engine/Source/Editor/Sequencer/Private/Sequencer.cpp
@@ -11414,16 +11414,16 @@ void FSequencer::BindCommands()
 		FIsActionChecked::CreateLambda( [this]{ return Settings->GetIsSnapEnabled(); } ) );
 
 	SequencerCommandBindings->MapAction(
-		Commands.ToggleSnapKeyTimesToKeys,
-		FExecuteAction::CreateLambda( [this]{ Settings->SetSnapKeyTimesToKeys( !Settings->GetSnapKeyTimesToKeys() ); } ),
+		Commands.ToggleSnapKeyTimesToElements,
+		FExecuteAction::CreateLambda( [this]{ Settings->SetSnapKeyTimesToElements( !Settings->GetSnapKeyTimesToElements() ); } ),
 		FCanExecuteAction::CreateLambda( []{ return true; } ),
-		FIsActionChecked::CreateLambda( [this]{ return Settings->GetSnapKeyTimesToKeys(); } ) );
+		FIsActionChecked::CreateLambda( [this]{ return Settings->GetSnapKeyTimesToElements(); } ) );
 
 	SequencerCommandBindings->MapAction(
-		Commands.ToggleSnapSectionTimesToSections,
-		FExecuteAction::CreateLambda( [this]{ Settings->SetSnapSectionTimesToSections( !Settings->GetSnapSectionTimesToSections() ); } ),
+		Commands.ToggleSnapSectionTimesToElements,
+		FExecuteAction::CreateLambda( [this]{ Settings->SetSnapSectionTimesToElements( !Settings->GetSnapSectionTimesToElements() ); } ),
 		FCanExecuteAction::CreateLambda( []{ return true; } ),
-		FIsActionChecked::CreateLambda( [this]{ return Settings->GetSnapSectionTimesToSections(); } ) );
+		FIsActionChecked::CreateLambda( [this]{ return Settings->GetSnapSectionTimesToElements(); } ) );
 
 	SequencerCommandBindings->MapAction(
 		Commands.ToggleSnapKeysAndSectionsToPlayRange,
diff --git a/Engine/Source/Editor/Sequencer/Private/SequencerCommands.cpp b/Engine/Source/Editor/Sequencer/Private/SequencerCommands.cpp
--- a/Engine/Source/Editor/Sequencer/Private/SequencerCommands.cpp
+++ b/Engine/Source/Editor/Sequencer/Private/SequencerCommands.cpp
@@ -133,8 +133,8 @@ void FSequencerCommands::RegisterCommands()
 	UI_COMMAND( ToggleShowInfoButton, "Info Button", "Enables and disables showing the information button in the playback controls", EUserInterfaceActionType::ToggleButton, FInputChord());
 	UI_COMMAND( ToggleIsSnapEnabled, "Enable Snapping", "Enables and disables snapping", EUserInterfaceActionType::ToggleButton, FInputChord() );
 
-	UI_COMMAND( ToggleSnapKeyTimesToKeys, "Snap Keys to Keys and Sections", "Snap keys to other keys and sections in this sequence", EUserInterfaceActionType::ToggleButton, FInputChord() );
-	UI_COMMAND( ToggleSnapSectionTimesToSections, "Snap Sections to Keys and Sections", "Snap sections to other keys and sections in this sequence", EUserInterfaceActionType::ToggleButton, FInputChord() );
+	UI_COMMAND( ToggleSnapKeyTimesToElements, "Snap Keys to Elements", "Snap keys to other keys, sections, markers, playback range, selection range, playhead, and previous position", EUserInterfaceActionType::ToggleButton, FInputChord() );
+	UI_COMMAND( ToggleSnapSectionTimesToElements, "Snap Sections to Elements", "Snap sections to other sections, keys, markers, playback range, selection range, and playhead", EUserInterfaceActionType::ToggleButton, FInputChord() );
 	UI_COMMAND( ToggleSnapKeysAndSectionsToPlayRange, "Snap Keys and Sections to the Playback Range", "When checked, keys and sections will be snapped to the playback range bounds", EUserInterfaceActionType::ToggleButton, FInputChord() );
 
 	UI_COMMAND( ToggleSnapPlayTimeToKeys, "Snap to Keys While Scrubbing", "Snap the playhead to keys of the selected track while scrubbing", EUserInterfaceActionType::ToggleButton, FInputChord() );
diff --git a/Engine/Source/Editor/Sequencer/Private/SequencerSettings.cpp b/Engine/Source/Editor/Sequencer/Private/SequencerSettings.cpp
--- a/Engine/Source/Editor/Sequencer/Private/SequencerSettings.cpp
+++ b/Engine/Source/Editor/Sequencer/Private/SequencerSettings.cpp
@@ -19,8 +19,8 @@ USequencerSettings::USequencerSettings( const FObjectInitializer& ObjectInitiali
 	bCreateSpawnableCameras = true;
 	bShowRangeSlider = false;
 	bIsSnapEnabled = true;
-	bSnapKeyTimesToKeys = false;
-	bSnapSectionTimesToSections = false;
+	bSnapKeyTimesToElements = false;
+	bSnapSectionTimesToElements = false;
 	bSnapPlayTimeToKeys = false;
 	bSnapPlayTimeToSections = false;
 	bSnapPlayTimeToMarkers = false;
@@ -205,16 +205,16 @@ void USequencerSettings::SetIsSnapEnabled(bool InbIsSnapEnabled)
 	}
 }
 
-bool USequencerSettings::GetSnapKeyTimesToKeys() const
+bool USequencerSettings::GetSnapKeyTimesToElements() const
 {
-	return bSnapKeyTimesToKeys;
+	return bSnapKeyTimesToElements;
 }
 
-void USequencerSettings::SetSnapKeyTimesToKeys(bool InbSnapKeyTimesToKeys)
+void USequencerSettings::SetSnapKeyTimesToElements(bool InbSnapKeyTimesToElements)
 {
-	if ( bSnapKeyTimesToKeys != InbSnapKeyTimesToKeys )
+	if ( bSnapKeyTimesToElements != InbSnapKeyTimesToElements)
 	{
-		bSnapKeyTimesToKeys = InbSnapKeyTimesToKeys;
+		bSnapKeyTimesToElements = InbSnapKeyTimesToElements;
 		SaveConfig();
 	}
 }
@@ -232,16 +232,16 @@ void USequencerSettings::SetLeftMouseDragDoesMarque(bool bDoMarque)
 	}
 }
 
-bool USequencerSettings::GetSnapSectionTimesToSections() const
+bool USequencerSettings::GetSnapSectionTimesToElements() const
 {
-	return bSnapSectionTimesToSections;
+	return bSnapSectionTimesToElements;
 }
 
-void USequencerSettings::SetSnapSectionTimesToSections( bool InbSnapSectionTimesToSections )
+void USequencerSettings::SetSnapSectionTimesToElements( bool InbSnapSectionTimesToElements )
 {
-	if ( bSnapSectionTimesToSections != InbSnapSectionTimesToSections )
+	if (bSnapSectionTimesToElements != InbSnapSectionTimesToElements)
 	{
-		bSnapSectionTimesToSections = InbSnapSectionTimesToSections;
+		bSnapSectionTimesToElements = InbSnapSectionTimesToElements;
 		SaveConfig();
 	}
 }
diff --git a/Engine/Source/Editor/Sequencer/Private/Tools/EditToolDragOperations.cpp b/Engine/Source/Editor/Sequencer/Private/Tools/EditToolDragOperations.cpp
--- a/Engine/Source/Editor/Sequencer/Private/Tools/EditToolDragOperations.cpp
+++ b/Engine/Source/Editor/Sequencer/Private/Tools/EditToolDragOperations.cpp
@@ -292,7 +292,7 @@ void FResizeSection::OnDrag(const FPointerEvent& MouseEvent, FVector2D LocalMous
 
 	// Snapping
 	const bool bSnapToInterval = Settings->GetForceWholeFrames();
-	const bool bSnapToLikeTypes = ((Settings->GetIsSnapEnabled() || MouseEvent.IsShiftDown()) && Settings->GetSnapSectionTimesToSections());
+	const bool bSnapToLikeTypes = ((Settings->GetIsSnapEnabled() || MouseEvent.IsShiftDown()) && Settings->GetSnapSectionTimesToElements());
 	
 	SnapField.GetValue().SetSnapToInterval(bSnapToInterval);
 	SnapField.GetValue().SetSnapToLikeTypes(bSnapToLikeTypes);
@@ -635,7 +635,7 @@ void FManipulateSectionEasing::OnDrag(const FPointerEvent& MouseEvent, FVector2D
 
 	// Snapping
 	const bool bSnapToInterval = Settings->GetForceWholeFrames();
-	const bool bSnapToLikeTypes = ((Settings->GetIsSnapEnabled() || MouseEvent.IsShiftDown()) && Settings->GetSnapSectionTimesToSections());
+	const bool bSnapToLikeTypes = ((Settings->GetIsSnapEnabled() || MouseEvent.IsShiftDown()) && Settings->GetSnapSectionTimesToElements());
 
 	SnapField.GetValue().SetSnapToInterval(bSnapToInterval);
 	SnapField.GetValue().SetSnapToLikeTypes(bSnapToLikeTypes);
@@ -847,7 +847,7 @@ void FMoveKeysAndSections::OnDrag(const FPointerEvent& MouseEvent, FVector2D Loc
 
 	// If they have both keys and settings selected then we snap to the interval if either one of them is enabled, otherwise respect the individual setting.
 	const bool bSnapToInterval = Settings->GetForceWholeFrames() && (KeysAsArray.Num() > 0 || Sections.Num());
-	const bool bSnapToLikeTypes = (Settings->GetIsSnapEnabled() || MouseEvent.IsShiftDown()) && ( (KeysAsArray.Num() > 0 && Settings->GetSnapKeyTimesToKeys()) || (Sections.Num() > 0 && Settings->GetSnapSectionTimesToSections()) );
+	const bool bSnapToLikeTypes = (Settings->GetIsSnapEnabled() || MouseEvent.IsShiftDown()) && ( (KeysAsArray.Num() > 0 && Settings->GetSnapKeyTimesToElements()) || (Sections.Num() > 0 && Settings->GetSnapSectionTimesToElements()) );
 
 	SnapField.GetValue().SetSnapToInterval(bSnapToInterval);
 	SnapField.GetValue().SetSnapToLikeTypes(bSnapToLikeTypes);
diff --git a/Engine/Source/Editor/Sequencer/Public/SequencerCommands.h b/Engine/Source/Editor/Sequencer/Public/SequencerCommands.h
--- a/Engine/Source/Editor/Sequencer/Public/SequencerCommands.h
+++ b/Engine/Source/Editor/Sequencer/Public/SequencerCommands.h
@@ -340,11 +340,11 @@ public:
 	/** Turns snapping on and off. */
 	TSharedPtr< FUICommandInfo > ToggleIsSnapEnabled;
 
-	/** Toggles whether or not keys should snap to other keys in the section. */
-	TSharedPtr< FUICommandInfo > ToggleSnapKeyTimesToKeys;
+	/** Toggles whether or not keys should snap to other elements. */
+	TSharedPtr< FUICommandInfo > ToggleSnapKeyTimesToElements;
 
-	/** Toggles whether or not sections should snap to other sections. */
-	TSharedPtr< FUICommandInfo > ToggleSnapSectionTimesToSections;
+	/** Toggles whether or not sections should snap to other elements. */
+	TSharedPtr< FUICommandInfo > ToggleSnapSectionTimesToElements;
 
 	/** Toggle constraining keys and sections in the play range */
 	TSharedPtr< FUICommandInfo > ToggleSnapKeysAndSectionsToPlayRange;
diff --git a/Engine/Source/Editor/Sequencer/Public/SequencerSettings.h b/Engine/Source/Editor/Sequencer/Public/SequencerSettings.h
--- a/Engine/Source/Editor/Sequencer/Public/SequencerSettings.h
+++ b/Engine/Source/Editor/Sequencer/Public/SequencerSettings.h
@@ -183,15 +183,15 @@ public:
 	/** Sets whether or not snapping is enabled. */
 	void SetIsSnapEnabled(bool InbIsSnapEnabled);
 
-	/** Gets whether or not to snap keys to other keys. */
-	bool GetSnapKeyTimesToKeys() const;
-	/** Sets whether or not to snap keys to other keys. */
-	void SetSnapKeyTimesToKeys(bool InbSnapKeyTimesToKeys);
+	/** Gets whether or not to snap keys to other elements. */
+	bool GetSnapKeyTimesToElements() const;
+	/** Sets whether or not to snap keys to other elements. */
+	void SetSnapKeyTimesToElements(bool InbSnapKeyTimesToElements);
 
-	/** Gets whether or not to snap sections to other sections. */
-	bool GetSnapSectionTimesToSections() const;
-	/** sets whether or not to snap sections to other sections. */
-	void SetSnapSectionTimesToSections( bool InbSnapSectionTimesToSections );
+	/** Gets whether or not to snap sections to other elements. */
+	bool GetSnapSectionTimesToElements() const;
+	/** sets whether or not to snap sections to other elements. */
+	void SetSnapSectionTimesToElements( bool InbSnapSectionTimesToElements );
 
 	/** @return true if keys and sections should be kept within the playback range when moving them */
 	bool GetSnapKeysAndSectionsToPlayRange() const;
@@ -591,13 +591,13 @@ protected:
 	UPROPERTY( config, EditAnywhere, Category=Snapping )
 	bool bIsSnapEnabled;
 
-	/** Enable or disable snapping keys to other keys. */
+	/** Enable or disable snapping keys to other elements. */
 	UPROPERTY( config, EditAnywhere, Category=Snapping )
-	bool bSnapKeyTimesToKeys;
+	bool bSnapKeyTimesToElements;
 	
-	/** Enable or disable snapping sections to other sections. */
+	/** Enable or disable snapping sections to other elements. */
 	UPROPERTY( config, EditAnywhere, Category=Snapping )
-	bool bSnapSectionTimesToSections;
+	bool bSnapSectionTimesToElements;
 
 	/** Enable or disable keeping keys and sections in the playback range. */
 	UPROPERTY(config, EditAnywhere, Category = Timeline)
