Sequencer: Fix display of ease handles

diff --git a/Engine/Source/Editor/Sequencer/Private/SequencerHotspots.cpp b/Engine/Source/Editor/Sequencer/Private/SequencerHotspots.cpp
--- a/Engine/Source/Editor/Sequencer/Private/SequencerHotspots.cpp
+++ b/Engine/Source/Editor/Sequencer/Private/SequencerHotspots.cpp
@@ -619,11 +619,11 @@ TOptional<FFrameNumber> FSectionEasingHandleHotspot::GetTime() const
 	{
 		if (HandleType == ESequencerEasingType::In && !ThisSection->GetEaseInRange().IsEmpty())
 		{
-			return UE::MovieScene::DiscreteExclusiveUpper(ThisSection->GetEaseInRange());
+			return ThisSection->GetEaseInRange().GetUpperBoundValue();
 		}
 		else if (HandleType == ESequencerEasingType::Out && !ThisSection->GetEaseOutRange().IsEmpty())
 		{
-			return UE::MovieScene::DiscreteInclusiveLower(ThisSection->GetEaseOutRange());
+			return ThisSection->GetEaseOutRange().GetLowerBoundValue();
 		}
 	}
 	return TOptional<FFrameNumber>();
