Sequencer: Fix some issues related to the group and level filters

Group and level filters should respond to CanResetFilters() and ResetFilters()
Level filters has its own ResetFilter() since it's meaning is a bit opposite (the filters are by default active, meaning it should show actors in those levels) and only is active when there's more than 1 level.

Fixed an issue where deleting a level would not reset cached AllWorldLevels
Removed InGroupName argument to IsGroupFilterActive since it's a noop

diff --git a/Engine/Source/Editor/Sequencer/Private/Filters/Filters/SequencerTrackFilter_Level.cpp b/Engine/Source/Editor/Sequencer/Private/Filters/Filters/SequencerTrackFilter_Level.cpp
--- a/Engine/Source/Editor/Sequencer/Private/Filters/Filters/SequencerTrackFilter_Level.cpp
+++ b/Engine/Source/Editor/Sequencer/Private/Filters/Filters/SequencerTrackFilter_Level.cpp
@@ -1,4 +1,4 @@
-﻿// Copyright Epic Games, Inc. All Rights Reserved.
+// Copyright Epic Games, Inc. All Rights Reserved.
 
 #include "Filters/Filters/SequencerTrackFilter_Level.h"
 #include "Engine/Level.h"
@@ -210,6 +210,8 @@ void FSequencerTrackFilter_Level::UpdateWorld(UWorld* const InWorld)
 
 void FSequencerTrackFilter_Level::HandleLevelsChanged()
 {
+	AllWorldLevels.Empty();
+
 	if (!CachedWorld.IsValid())
 	{
 		ResetFilter();
@@ -226,7 +228,6 @@ void FSequencerTrackFilter_Level::HandleLevelsChanged()
 
 	// Build a list of level names contained in the current world
 	TSet<FString> OldWorldLevels = AllWorldLevels;
-	AllWorldLevels.Empty();
 	for (const ULevel* const Level : WorldLevels)
 	{
 		if (IsValid(Level))
diff --git a/Engine/Source/Editor/Sequencer/Private/Filters/Menus/SequencerTrackFilterMenu.cpp b/Engine/Source/Editor/Sequencer/Private/Filters/Menus/SequencerTrackFilterMenu.cpp
--- a/Engine/Source/Editor/Sequencer/Private/Filters/Menus/SequencerTrackFilterMenu.cpp
+++ b/Engine/Source/Editor/Sequencer/Private/Filters/Menus/SequencerTrackFilterMenu.cpp
@@ -699,14 +699,7 @@ bool FSequencerTrackFilterMenu::CanResetFilters() const
 		return true;
 	}
 
-	const TArray<TSharedRef<FFilterCategory>> Categories = {
-		FilterBar->GetClassTypeCategory(),
-		FilterBar->GetComponentTypeCategory(),
-		FilterBar->GetMiscCategory()
-	};
-	const TArray<TSharedRef<FSequencerTrackFilter>> ClassAndCompFilters = FilterBar->GetCommonFilters(Categories);
-
-	return FilterBar->HasEnabledFilter(ClassAndCompFilters) || FilterBar->HasEnabledCustomTextFilters();
+	return FilterBar->CanResetFilters();
 }
 
 void FSequencerTrackFilterMenu::ResetFilters()
diff --git a/Engine/Source/Editor/Sequencer/Private/Filters/SequencerFilterBar.cpp b/Engine/Source/Editor/Sequencer/Private/Filters/SequencerFilterBar.cpp
--- a/Engine/Source/Editor/Sequencer/Private/Filters/SequencerFilterBar.cpp
+++ b/Engine/Source/Editor/Sequencer/Private/Filters/SequencerFilterBar.cpp
@@ -344,15 +344,15 @@ void FSequencerFilterBar::ToggleMuteFilters()
 
 bool FSequencerFilterBar::CanResetFilters() const
 {
-	const TArray<TSharedRef<FFilterCategory>> Categories = { GetClassTypeCategory(), GetComponentTypeCategory(), GetMiscCategory() };
-	const TArray<TSharedRef<FSequencerTrackFilter>> ClassAndCompFilters = GetCommonFilters(Categories);
-	return HasEnabledFilter(ClassAndCompFilters);
+	return HasAnyFiltersEnabled();
 }
 
 void FSequencerFilterBar::ResetFilters()
 {
 	EnableAllFilters(false, {});
 	EnableCustomTextFilters(false);
+	EnableAllGroupFilters(false);
+	LevelFilter->ResetFilter();
 }
 
 FString FSequencerFilterBar::GetTextFilterString() const
@@ -501,7 +501,7 @@ TSharedPtr<FSequencerTrackFilter_CustomText> FSequencerFilterBar::FindCustomText
 
 bool FSequencerFilterBar::HasAnyFiltersEnabled() const
 {
-	return HasEnabledCommonFilters() || HasEnabledCustomTextFilters();
+	return HasEnabledCommonFilters() || HasEnabledCustomTextFilters() || AnyInternalFilterActive();
 }
 
 bool FSequencerFilterBar::IsFilterActiveByDisplayName(const FString& InFilterName) const
@@ -1283,7 +1283,7 @@ void FSequencerFilterBar::EnableAllGroupFilters(const bool bInEnable)
 	}
 }
 
-bool FSequencerFilterBar::IsGroupFilterActive(const FString& InGroupName) const
+bool FSequencerFilterBar::IsGroupFilterActive() const
 {
 	return GroupFilter->HasActiveGroupFilter();
 }
diff --git a/Engine/Source/Editor/Sequencer/Private/Filters/SequencerFilterBar.h b/Engine/Source/Editor/Sequencer/Private/Filters/SequencerFilterBar.h
--- a/Engine/Source/Editor/Sequencer/Private/Filters/SequencerFilterBar.h
+++ b/Engine/Source/Editor/Sequencer/Private/Filters/SequencerFilterBar.h
@@ -216,7 +216,7 @@ public:
 	/** Group Filter Functions */
 
 	void EnableAllGroupFilters(const bool bInEnable);
-	bool IsGroupFilterActive(const FString& InGroupName) const;
+	bool IsGroupFilterActive() const;
 
 	/** Misc Functions */
 
