[Sequencer] Remove track properties from showing in sidebar when track area is selected. (Removes double scroll bar)

diff --git a/Engine/Source/Editor/Sequencer/Private/Widgets/Sidebar/SequencerSelectionDrawer.cpp b/Engine/Source/Editor/Sequencer/Private/Widgets/Sidebar/SequencerSelectionDrawer.cpp
--- a/Engine/Source/Editor/Sequencer/Private/Widgets/Sidebar/SequencerSelectionDrawer.cpp
+++ b/Engine/Source/Editor/Sequencer/Private/Widgets/Sidebar/SequencerSelectionDrawer.cpp
@@ -299,25 +299,15 @@ void FSequencerSelectionDrawer::BuildKeySelectionDetails(const TSharedRef<FSeque
 void FSequencerSelectionDrawer::BuildTrackAreaDetails(const TSharedRef<FSequencerSelection>& InSelection, FMenuBuilder& MenuBuilder)
 {
 	TArray<TWeakObjectPtr<>> AllSectionObjects; 
-	TArray<TWeakObjectPtr<>> WeakTrackObjects;
-	TArray<TPair<TWeakObjectPtr<UMovieSceneTrack>, int32>> SelectedTrackRows;
 
 	for (const FViewModelPtr TrackAreaItem : InSelection->TrackArea)
 	{
 		if (const TViewModelPtr<FSectionModel> SectionModel = TrackAreaItem.ImplicitCast())
 		{
-			if (UMovieSceneSection* Section = SectionModel->GetSection())
+			UMovieSceneSection* const Section = SectionModel->GetSection();
+			if (IsValid(Section))
 			{
 				AllSectionObjects.Add(SectionModel->GetSection());
-
-				if (UMovieSceneTrack* Track = Section->GetTypedOuter<UMovieSceneTrack>())
-				{
-					WeakTrackObjects.Add(Track);
-					if (Track->SupportsMultipleRows())
-					{
-						SelectedTrackRows.Add(TPair<TWeakObjectPtr<UMovieSceneTrack>, int32>(Track, Section->GetRowIndex()));
-					}
-				}
 			}
 		}
 	}
@@ -326,21 +316,6 @@ void FSequencerSelectionDrawer::BuildTrackAreaDetails(const TSharedRef<FSequence
 	{
 		SequencerHelpers::BuildEditSectionMenu(WeakSequencer, AllSectionObjects, MenuBuilder, false);
 	}
-
-	if (!SelectedTrackRows.IsEmpty() && !Algo::AnyOf(SelectedTrackRows, [](const TPair<TWeakObjectPtr<UMovieSceneTrack>, int32> TrackRow) {
-		return TrackRow.Key.IsValid() && !TrackRow.Key->SupportsMultipleRows();
-		}))
-	{
-		MenuBuilder.BeginSection(TEXT("TrackRowMetadata"));
-		{
-			// Empty here, will be implemented by extension.
-		}
-		MenuBuilder.EndSection();
-	}
-	if (!WeakTrackObjects.IsEmpty())
-	{
-		SequencerHelpers::BuildEditTrackMenu(WeakSequencer, WeakTrackObjects, MenuBuilder, false);
-	}
 }
 
 void FSequencerSelectionDrawer::BuildOutlinerDetails(const TSharedRef<FSequencerSelection>& InSelection, FMenuBuilder& MenuBuilder)
@@ -363,7 +338,6 @@ void FSequencerSelectionDrawer::BuildOutlinerDetails(const TSharedRef<FSequencer
 	TSet<TViewModelPtr<FChannelGroupOutlinerModel>> ChannelGroups;
 	TArray<TPair<TWeakObjectPtr<UMovieSceneTrack>, int32>> SelectedTrackRows;
 
-
 	for (const FViewModelPtr OutlinerItem : InSelection->Outliner)
 	{
 		if (const TViewModelPtr<ITrackExtension> TrackExtension = OutlinerItem.ImplicitCast())
