[Sequencer] Navigation Tool: Add new SVG for marker item icon

diff --git a/Engine/Source/Editor/Sequencer/Private/Tools/NavigationTool/Items/NavigationToolMarker.cpp b/Engine/Source/Editor/Sequencer/Private/Tools/NavigationTool/Items/NavigationToolMarker.cpp
--- a/Engine/Source/Editor/Sequencer/Private/Tools/NavigationTool/Items/NavigationToolMarker.cpp
+++ b/Engine/Source/Editor/Sequencer/Private/Tools/NavigationTool/Items/NavigationToolMarker.cpp
@@ -9,8 +9,8 @@
 #include "Tools/NavigationTool/Items/NavigationToolSequence.h"
 #include "Tools/NavigationTool/NavigationTool.h"
 #include "Tools/NavigationTool/NavigationToolScopedSelection.h"
+#include "Tools/NavigationTool/NavigationToolStyle.h"
 #include "Tools/NavigationTool/Utils/NavigationToolMiscUtils.h"
-#include "Tools/NavigationTool/Utils/NavigationToolMovieSceneUtils.h"
 
 #define LOCTEXT_NAMESPACE "NavigationToolMarker"
 
@@ -58,12 +58,12 @@ FText FNavigationToolMarker::GetClassName() const
 
 FSlateIcon FNavigationToolMarker::GetIcon() const
 {
-	return FSlateIcon(FAppStyle::GetAppStyleSetName(), TEXT("AnimTimeline.SectionMarker"));
+	return FSlateIcon(FNavigationToolStyle::Get().GetStyleSetName(), TEXT("Item.Marker.Icon"));
 }
 
 const FSlateBrush* FNavigationToolMarker::GetIconBrush() const
 {
-	return FAppStyle::GetBrush(TEXT("AnimTimeline.SectionMarker"));
+	return FNavigationToolStyle::Get().GetBrush(TEXT("Item.Marker.Icon"));
 }
 
 FSlateColor FNavigationToolMarker::GetIconColor() const
diff --git a/Engine/Source/Editor/Sequencer/Private/Tools/NavigationTool/NavigationToolStyle.cpp b/Engine/Source/Editor/Sequencer/Private/Tools/NavigationTool/NavigationToolStyle.cpp
--- a/Engine/Source/Editor/Sequencer/Private/Tools/NavigationTool/NavigationToolStyle.cpp
+++ b/Engine/Source/Editor/Sequencer/Private/Tools/NavigationTool/NavigationToolStyle.cpp
@@ -17,6 +17,8 @@
 FNavigationToolStyle::FNavigationToolStyle()
 	: FSlateStyleSet(TEXT("NavigationToolStyle"))
 {
+	SetParentStyleName(FAppStyle::GetAppStyleSetName());
+
 	FSlateStyleSet::SetContentRoot(FPaths::EngineContentDir() / TEXT("Editor/Slate"));
 	FSlateStyleSet::SetCoreContentRoot(FPaths::EngineContentDir() / TEXT("Slate"));
 
@@ -50,6 +52,8 @@ FNavigationToolStyle::FNavigationToolStyle()
 		.SetBackgroundImageReadOnly(FSlateRoundedBoxBrush(FStyleColors::Background, 4.0f, FStyleColors::InputOutline, 1.0f));
 	Set("NonEditableTextBox", NonEditableTextBoxStyle);
 
+	Set("Item.Marker.Icon", new IMAGE_BRUSH_SVG("Sequencer/SequenceNavigator/Marker_16", Icon16x16));
+
 	FSlateStyleRegistry::RegisterSlateStyle(*this);
 }
 
