Sequencer: improved selection of all keys
- avoid using the same parent node repeatedly, as it will redo the same thing needlessly
- use range based selection + delay selection event once done

#rnx

diff --git a/Engine/Source/Editor/Sequencer/Private/MVVM/Selection.cpp b/Engine/Source/Editor/Sequencer/Private/MVVM/Selection.cpp
--- a/Engine/Source/Editor/Sequencer/Private/MVVM/Selection.cpp
+++ b/Engine/Source/Editor/Sequencer/Private/MVVM/Selection.cpp
@@ -168,11 +168,28 @@ void FSequencerSelection::PreBroadcastChangeEvent()
 	}
 
 	// Gather selection states from selected keys
-	for (FKeyHandle Key : KeySelection)
 	{
-		TViewModelPtr<FChannelModel>      Channel            = KeySelection.GetModelForKey(Key);
-		TViewModelPtr<IOutlinerExtension> ParentOutlinerNode = Channel ? Channel->GetLinkedOutlinerItem() : nullptr;
-		if (ParentOutlinerNode)
+		TSet<TViewModelPtr<FChannelModel>> Channels;
+		for (const FKeyHandle& Key : KeySelection)
+		{
+			if (TViewModelPtr<FChannelModel> Channel = KeySelection.GetModelForKey(Key))
+			{
+				Channels.Add(Channel);
+			}
+		}
+
+		TSet<TViewModelPtr<IOutlinerExtension>> ParentOutlinerNodes;
+		ParentOutlinerNodes.Reserve(Channels.Num());
+		for (const TViewModelPtr<FChannelModel>& Channel : Channels)
+		{
+			if (TViewModelPtr<IOutlinerExtension> ParentOutlinerNode = Channel ? Channel->GetLinkedOutlinerItem() : nullptr)
+			{
+				ParentOutlinerNodes.Add(ParentOutlinerNode);			
+			}
+		}
+
+		NodesWithKeysOrSections.Reserve(NodesWithKeysOrSections.Num() + ParentOutlinerNodes.Num());
+		for (TViewModelPtr<IOutlinerExtension> ParentOutlinerNode: ParentOutlinerNodes)
 		{
 			ParentOutlinerNode->ToggleSelectionState(EOutlinerSelectionState::HasSelectedKeys, true);
 			NodesWithKeysOrSections.Add(ParentOutlinerNode);
diff --git a/Engine/Source/Editor/Sequencer/Private/SequencerContextMenus.cpp b/Engine/Source/Editor/Sequencer/Private/SequencerContextMenus.cpp
--- a/Engine/Source/Editor/Sequencer/Private/SequencerContextMenus.cpp
+++ b/Engine/Source/Editor/Sequencer/Private/SequencerContextMenus.cpp
@@ -112,6 +112,36 @@ static void CreateKeyStructForSelection(TWeakPtr<ISequencer> InWeakSequencer, TS
 	}
 }
 
+namespace UE::Sequencer::Private
+{
+	
+TSet<TSharedPtr<FChannelModel>> GetChannelModels(const TWeakPtr<FSequencer>& InWeakSequencer)
+{
+	TSet<TSharedPtr<FChannelModel>> Channels;
+	
+	const TSharedPtr<FSequencer> Sequencer = InWeakSequencer.IsValid() ? InWeakSequencer.Pin() : nullptr;
+	if (Sequencer.IsValid())
+	{
+		const FSequencerSelection& Selection = *Sequencer->GetViewModel()->GetSelection();
+		for (const TViewModelPtr<IOutlinerExtension>& Item : Selection.Outliner)
+		{
+			SequencerHelpers::GetAllChannels(Item, Channels);
+		}
+
+		if (Channels.IsEmpty())
+		{
+			for (const TWeakViewModelPtr<IOutlinerExtension>& DisplayNode : Selection.GetNodesWithSelectedKeysOrSections())
+			{
+				SequencerHelpers::GetAllChannels(DisplayNode.Pin(), Channels);
+			}
+		}
+	}
+	
+	return MoveTemp(Channels);
+}
+
+}
+
 void FKeyContextMenu::BuildMenu(FMenuBuilder& MenuBuilder, TSharedPtr<FExtender> MenuExtender, TWeakPtr<FSequencer> InWeakSequencer)
 {
 	TSharedRef<FKeyContextMenu> Menu = MakeShareable(new FKeyContextMenu(InWeakSequencer));
@@ -855,27 +885,20 @@ void FSectionContextMenu::SelectAllKeys()
 		return;
 	}
 
-	TArray<FKeyHandle> HandlesScratch;
-
-	FSequencerSelection& Selection = *Sequencer->GetViewModel()->GetSelection();
-	for (TViewModelPtr<FSectionModel> Section : Selection.TrackArea.Filter<FSectionModel>())
+	const TSet<TSharedPtr<FChannelModel>> Channels = Private::GetChannelModels(WeakSequencer);
+	if (!Channels.IsEmpty())
 	{
-		UMovieSceneSection* SectionObject = Section->GetSection();
-		if (SectionObject)
+		FSequencerSelection& Selection = *Sequencer->GetViewModel()->GetSelection();
+		FSelectionEventSuppressor EventSuppressor = Selection.SuppressEvents();
+		
+		TArray<FKeyHandle> HandlesScratch;
+		for (const TSharedPtr<FChannelModel>& Channel: Channels)
 		{
-			// Iterate all channels
-			for (const TViewModelPtr<FChannelModel>& Channel : Section->GetDescendantsOfType<FChannelModel>())
+			if (Channel->GetLinkedOutlinerItem() && !Channel->GetLinkedOutlinerItem()->IsFilteredOut())
 			{
-				if (Channel->GetLinkedOutlinerItem() && !Channel->GetLinkedOutlinerItem()->IsFilteredOut())
-				{
-					HandlesScratch.Reset();
-					Channel->GetKeyArea()->GetKeyHandles(HandlesScratch);
-
-					for (FKeyHandle KeyHandle : HandlesScratch)
-					{
-						Selection.KeySelection.Select(Channel, KeyHandle);
-					}
-				}
+				HandlesScratch.Reset();
+				Channel->GetKeyArea()->GetKeyHandles(HandlesScratch);
+				Selection.KeySelection.SelectRange(Channel, HandlesScratch);
 			}
 		}
 	}
@@ -1013,45 +1036,55 @@ void FSectionContextMenu::ReduceKeys()
 	{
 		return;
 	}
-
-	FScopedTransaction ReduceKeysTransaction(LOCTEXT("ReduceKeys_Transaction", "Reduce Keys"));
-
-	TSet<TSharedPtr<IKeyArea> > KeyAreas;
-	for (TViewModelPtr<IOutlinerExtension> Item : Sequencer->GetViewModel()->GetSelection()->Outliner)
+	
+	const TSet<TSharedPtr<FChannelModel>> ChannelModels = Private::GetChannelModels(WeakSequencer);
+	if (ChannelModels.IsEmpty())
 	{
-		SequencerHelpers::GetAllKeyAreas(Item, KeyAreas);
+		return;
 	}
 
-	if (KeyAreas.Num() == 0)
+	TSet<UMovieSceneSection*> Sections;
+	TSet<FMovieSceneChannel*> Channels;
+	Channels.Reserve(ChannelModels.Num());
+	
+	for (const TSharedPtr<FChannelModel>& ChannelModel: ChannelModels)
 	{
-		for (const TWeakViewModelPtr<IOutlinerExtension>& DisplayNode : Sequencer->GetViewModel()->GetSelection()->GetNodesWithSelectedKeysOrSections())
+		if (ChannelModel->GetLinkedOutlinerItem() && !ChannelModel->GetLinkedOutlinerItem()->IsFilteredOut())
 		{
-			SequencerHelpers::GetAllKeyAreas(DisplayNode.Pin(), KeyAreas);
+			const TSharedPtr<IKeyArea> KeyArea = ChannelModel->GetKeyArea();
+			if (KeyArea.IsValid())
+			{
+				FMovieSceneChannel* Channel = KeyArea->GetChannel().Get();
+				UMovieSceneSection* Section = KeyArea->GetOwningSection();
+				if (Channel && Section)
+				{
+					Channels.Add(Channel);
+					Sections.Add(Section);
+				}
+			}
 		}
 	}
 
-	FKeyDataOptimizationParams Params;
-	Params.bAutoSetInterpolation = true;
-	Params.Tolerance = Sequencer->GetSequencerSettings()->GetReduceKeysTolerance();
-
-	for (TSharedPtr<IKeyArea> KeyArea : KeyAreas)
+	if (!Sections.IsEmpty())
 	{
-		if (KeyArea.IsValid())
+		FScopedTransaction ReduceKeysTransaction(LOCTEXT("ReduceKeys_Transaction", "Reduce Keys"));
+
+		for (UMovieSceneSection* Section: Sections)
 		{
-			UMovieSceneSection* Section = KeyArea->GetOwningSection();
-			if (Section)
-			{
-				Section->Modify();
+			Section->Modify();
+		}
 
-				if (FMovieSceneChannel* Channel = KeyArea->GetChannel().Get())
-				{
-					Channel->Optimize(Params);
-				}
-			}
+		FKeyDataOptimizationParams Params;
+		Params.bAutoSetInterpolation = true;
+		Params.Tolerance = Sequencer->GetSequencerSettings()->GetReduceKeysTolerance();
+	
+		for (FMovieSceneChannel* Channel: Channels)
+		{
+			Channel->Optimize(Params);
 		}
-	}
 
-	Sequencer->NotifyMovieSceneDataChanged( EMovieSceneDataChangeType::TrackValueChanged );
+		Sequencer->NotifyMovieSceneDataChanged( EMovieSceneDataChangeType::TrackValueChanged );
+	}
 }
 
 bool FSectionContextMenu::CanAutoSize() const
