[Sequencer] Navigation Tool:
- Fix column view save state from provider refactor regression
- Save column view names as FText instead of FName

diff --git a/Engine/Plugins/MovieScene/LevelSequenceEditor/Source/LevelSequenceEditor/Private/LevelSequenceNavigationToolProvider.cpp b/Engine/Plugins/MovieScene/LevelSequenceEditor/Source/LevelSequenceEditor/Private/LevelSequenceNavigationToolProvider.cpp
--- a/Engine/Plugins/MovieScene/LevelSequenceEditor/Source/LevelSequenceEditor/Private/LevelSequenceNavigationToolProvider.cpp
+++ b/Engine/Plugins/MovieScene/LevelSequenceEditor/Source/LevelSequenceEditor/Private/LevelSequenceNavigationToolProvider.cpp
@@ -34,7 +34,7 @@ namespace UE::Sequencer::NavigationTool
 
 const FName FLevelSequenceNavigationToolProvider::Identifier = TEXT("LevelSequence");
 
-const FName FLevelSequenceNavigationToolProvider::AnimationColumnViewId = TEXT("Animation");
+const FText FLevelSequenceNavigationToolProvider::AnimationColumnViewName = LOCTEXT("AnimationColumnViewName", "Animation");
 
 FLevelSequenceNavigationToolProvider::FLevelSequenceNavigationToolProvider()
 	: ToolCommands(MakeShared<FUICommandList>())
@@ -51,9 +51,9 @@ TSet<TSubclassOf<UMovieSceneSequence>> FLevelSequenceNavigationToolProvider::Get
 	return { ULevelSequence::StaticClass() };
 }
 
-FName FLevelSequenceNavigationToolProvider::GetDefaultColumnView() const
+FText FLevelSequenceNavigationToolProvider::GetDefaultColumnView() const
 {
-	return AnimationColumnViewId;
+	return AnimationColumnViewName;
 }
 
 void FLevelSequenceNavigationToolProvider::BindCommands(const TSharedRef<FUICommandList>& InCommandList)
@@ -110,17 +110,17 @@ void FLevelSequenceNavigationToolProvider::OnExtendColumns(FNavigationToolColumn
 void FLevelSequenceNavigationToolProvider::OnExtendColumnViews(TSet<FNavigationToolColumnView>& OutColumnViews)
 {
 	FNavigationToolColumnView& AnimationColumnView = OutColumnViews
-		.FindOrAdd(FNavigationToolColumnView(AnimationColumnViewId));
+		.FindOrAdd(FNavigationToolColumnView(AnimationColumnViewName));
 
 	AnimationColumnView.VisibleColumns =
 		{
-			FNavigationToolPlayheadColumn::GetStaticTypeId().ToName(),
-			FNavigationToolColorColumn::GetStaticTypeId().ToName(),
-			FNavigationToolLabelColumn::GetStaticTypeId().ToName(),
-			FNavigationToolItemsColumn::GetStaticTypeId().ToName(),
-			FNavigationToolInTimeColumn::GetStaticTypeId().ToName(),
-			FNavigationToolOutTimeColumn::GetStaticTypeId().ToName(),
-			FNavigationToolHBiasColumn::GetStaticTypeId().ToName()
+			FNavigationToolPlayheadColumn::StaticColumnId(),
+			FNavigationToolColorColumn::StaticColumnId(),
+			FNavigationToolLabelColumn::StaticColumnId(),
+			FNavigationToolItemsColumn::StaticColumnId(),
+			FNavigationToolInTimeColumn::StaticColumnId(),
+			FNavigationToolOutTimeColumn::StaticColumnId(),
+			FNavigationToolHBiasColumn::StaticColumnId()
 		};
 
 	FNavigationToolProvider::OnExtendColumnViews(OutColumnViews);
diff --git a/Engine/Plugins/MovieScene/LevelSequenceEditor/Source/LevelSequenceEditor/Private/LevelSequenceNavigationToolProvider.h b/Engine/Plugins/MovieScene/LevelSequenceEditor/Source/LevelSequenceEditor/Private/LevelSequenceNavigationToolProvider.h
--- a/Engine/Plugins/MovieScene/LevelSequenceEditor/Source/LevelSequenceEditor/Private/LevelSequenceNavigationToolProvider.h
+++ b/Engine/Plugins/MovieScene/LevelSequenceEditor/Source/LevelSequenceEditor/Private/LevelSequenceNavigationToolProvider.h
@@ -21,7 +21,7 @@ class FLevelSequenceNavigationToolProvider : public FNavigationToolProvider
 public:
 	static const FName Identifier;
 
-	static const FName AnimationColumnViewId;
+	static const FText AnimationColumnViewName;
 
 	FLevelSequenceNavigationToolProvider();
 	virtual ~FLevelSequenceNavigationToolProvider() override = default;
@@ -32,7 +32,7 @@ public:
 
 	virtual TSet<TSubclassOf<UMovieSceneSequence>> GetSupportedSequenceClasses() const override;
 
-	virtual FName GetDefaultColumnView() const override;
+	virtual FText GetDefaultColumnView() const override;
 
 	virtual void BindCommands(const TSharedRef<FUICommandList>& InCommandList) override;
 
diff --git a/Engine/Source/Editor/Sequencer/Private/Tools/NavigationTool/Columns/NavigationToolColumn.cpp b/Engine/Source/Editor/Sequencer/Private/Tools/NavigationTool/Columns/NavigationToolColumn.cpp
--- a/Engine/Source/Editor/Sequencer/Private/Tools/NavigationTool/Columns/NavigationToolColumn.cpp
+++ b/Engine/Source/Editor/Sequencer/Private/Tools/NavigationTool/Columns/NavigationToolColumn.cpp
@@ -1,19 +1,8 @@
 ﻿// Copyright Epic Games, Inc. All Rights Reserved.
 
 #include "Tools/NavigationTool/Columns/NavigationToolColumn.h"
-#include "Tools/NavigationTool/Providers/NavigationToolProvider.h"
 
 namespace UE::Sequencer::NavigationTool
 {
 
-TSharedPtr<FNavigationToolProvider> FNavigationToolColumn::GetProvider() const
-{
-	return WeakProvider.Pin();
 }
-
-void FNavigationToolColumn::SetProvider(const TWeakPtr<FNavigationToolProvider>& InWeakProvider)
-{
-	WeakProvider = InWeakProvider;
-}
-
-}
\ No newline at end of file
diff --git a/Engine/Source/Editor/Sequencer/Private/Tools/NavigationTool/Menus/NavigationToolToolbarMenu.cpp b/Engine/Source/Editor/Sequencer/Private/Tools/NavigationTool/Menus/NavigationToolToolbarMenu.cpp
--- a/Engine/Source/Editor/Sequencer/Private/Tools/NavigationTool/Menus/NavigationToolToolbarMenu.cpp
+++ b/Engine/Source/Editor/Sequencer/Private/Tools/NavigationTool/Menus/NavigationToolToolbarMenu.cpp
@@ -312,14 +312,19 @@ void FNavigationToolToolbarMenu::CreateColumnViewOptionsMenu(UToolMenu* const In
 	ColumnViewsSection.AddMenuEntryWithCommandList(ToolCommands.ResetVisibleColumnSizes, ViewCommandList);
 	ColumnViewsSection.AddMenuEntryWithCommandList(ToolCommands.SaveCurrentColumnView, ViewCommandList);
 
-	const TSet<FNavigationToolColumnView>& CustomColumnViews = ToolSettings->GetCustomColumnViews();
+	TSet<FNavigationToolColumnView>& CustomColumnViews = ToolSettings->GetCustomColumnViews();
+	CustomColumnViews.Sort([](const FNavigationToolColumnView& InA, const FNavigationToolColumnView& InB)
+		{
+			return InA.ViewName.CompareTo(InB.ViewName) > 0;
+		});
+
 	if (!CustomColumnViews.IsEmpty())
 	{
 		ColumnViewsSection.AddSeparator(NAME_None);
 
 		for (const FNavigationToolColumnView& ColumnView : CustomColumnViews)
 		{
-			const TSharedRef<SWidget> MenuItemWidget = FNavigationToolFilterMenu::ConstructCustomMenuItemWidget(FText::FromName(ColumnView.ViewName),
+			const TSharedRef<SWidget> MenuItemWidget = FNavigationToolFilterMenu::ConstructCustomMenuItemWidget(ColumnView.ViewName,
 				FSimpleDelegate::CreateSP(ToolViewRef, &FNavigationToolView::ApplyCustomColumnView, ColumnView.ViewName),
 				TAttribute<ECheckBoxState>::CreateSP(this, &FNavigationToolToolbarMenu::GetCustomColumnViewMenuItemCheckState, ColumnView.ViewName, ToolViewRef),
 				FAppStyle::GetBrush(TEXT("Icons.Delete")),
@@ -351,7 +356,7 @@ void FNavigationToolToolbarMenu::CreateFilterBarOptionsMenu(UToolMenu* const InM
 	FilterBarLayoutSection.AddMenuEntry(FilterCommands.SetToHorizontalLayout);
 }
 
-ECheckBoxState FNavigationToolToolbarMenu::GetCustomColumnViewMenuItemCheckState(const FName InColumnId, const TSharedRef<FNavigationToolView> InToolView) const
+ECheckBoxState FNavigationToolToolbarMenu::GetCustomColumnViewMenuItemCheckState(const FText InColumnViewName, const TSharedRef<FNavigationToolView> InToolView) const
 {
 	UNavigationToolSettings* const ToolSettings = GetMutableDefault<UNavigationToolSettings>();
 	if (!ensure(IsValid(ToolSettings)))
@@ -359,7 +364,7 @@ ECheckBoxState FNavigationToolToolbarMenu::GetCustomColumnViewMenuItemCheckState
 		return ECheckBoxState::Unchecked;
 	}
 
-	FNavigationToolColumnView* const SavedColumnView = ToolSettings->FindCustomColumnView(InColumnId);
+	FNavigationToolColumnView* const SavedColumnView = ToolSettings->FindCustomColumnView(InColumnViewName);
 	if (!SavedColumnView)
 	{
 		return ECheckBoxState::Unchecked;
@@ -378,7 +383,7 @@ ECheckBoxState FNavigationToolToolbarMenu::GetCustomColumnViewMenuItemCheckState
 	return ECheckBoxState::Checked;
 }
 
-void FNavigationToolToolbarMenu::OnDeleteCustomColumnViewMenuItemClick(const FName InColumnId)
+void FNavigationToolToolbarMenu::OnDeleteCustomColumnViewMenuItemClick(const FText InColumnViewName)
 {
 	UNavigationToolSettings* const ToolSettings = GetMutableDefault<UNavigationToolSettings>();
 	if (!ensure(IsValid(ToolSettings)))
@@ -388,7 +393,7 @@ void FNavigationToolToolbarMenu::OnDeleteCustomColumnViewMenuItemClick(const FNa
 
 	TSet<FNavigationToolColumnView>& CustomColumnViews = ToolSettings->GetCustomColumnViews();
 
-	CustomColumnViews.Remove(InColumnId);
+	CustomColumnViews.Remove(InColumnViewName);
 
 	ToolSettings->SaveConfig();
 
diff --git a/Engine/Source/Editor/Sequencer/Private/Tools/NavigationTool/Menus/NavigationToolToolbarMenu.h b/Engine/Source/Editor/Sequencer/Private/Tools/NavigationTool/Menus/NavigationToolToolbarMenu.h
--- a/Engine/Source/Editor/Sequencer/Private/Tools/NavigationTool/Menus/NavigationToolToolbarMenu.h
+++ b/Engine/Source/Editor/Sequencer/Private/Tools/NavigationTool/Menus/NavigationToolToolbarMenu.h
@@ -31,8 +31,8 @@ protected:
 	void CreateColumnViewOptionsMenu(UToolMenu* const InMenu);
 	void CreateFilterBarOptionsMenu(UToolMenu* const InMenu);
 
-	ECheckBoxState GetCustomColumnViewMenuItemCheckState(const FName InColumnId, const TSharedRef<FNavigationToolView> InToolView) const;
-	void OnDeleteCustomColumnViewMenuItemClick(const FName InColumnId);
+	ECheckBoxState GetCustomColumnViewMenuItemCheckState(const FText InColumnViewName, const TSharedRef<FNavigationToolView> InToolView) const;
+	void OnDeleteCustomColumnViewMenuItemClick(const FText InColumnViewName);
 
 	bool IsGlobalFilterActive(const TSharedRef<FNavigationToolBuiltInFilter> InFilter) const;
 	void OnToggleGlobalFilter(const TSharedRef<FNavigationToolBuiltInFilter> InFilter, const TSharedRef<FNavigationTool> InTool);
diff --git a/Engine/Source/Editor/Sequencer/Private/Tools/NavigationTool/NavigationTool.cpp b/Engine/Source/Editor/Sequencer/Private/Tools/NavigationTool/NavigationTool.cpp
--- a/Engine/Source/Editor/Sequencer/Private/Tools/NavigationTool/NavigationTool.cpp
+++ b/Engine/Source/Editor/Sequencer/Private/Tools/NavigationTool/NavigationTool.cpp
@@ -1278,14 +1278,14 @@ void FNavigationTool::OnSequencerActivate(FMovieSceneSequenceIDRef InSequenceId)
 	// We apply default views *after* all columns have been created for all providers
 	if (ToolSettings->ShouldApplyDefaultColumnView())
 	{
-		TArray<FName> DefaultColumnViews;
+		TArray<FText> DefaultColumnViews;
 
 		ForEachProvider([&DefaultColumnViews](const TSharedRef<FNavigationToolProvider>& InProvider)
 			{
-				const FName DefaultColumnView = InProvider->GetDefaultColumnView();
-				if (!DefaultColumnView.IsNone())
+				const FText DefaultColumnView = InProvider->GetDefaultColumnView();
+				if (!DefaultColumnView.IsEmptyOrWhitespace())
 				{
-					DefaultColumnViews.AddUnique(DefaultColumnView);
+					DefaultColumnViews.Add(DefaultColumnView);
 				}
 				return true;
 			});
diff --git a/Engine/Source/Editor/Sequencer/Private/Tools/NavigationTool/NavigationToolSettings.cpp b/Engine/Source/Editor/Sequencer/Private/Tools/NavigationTool/NavigationToolSettings.cpp
--- a/Engine/Source/Editor/Sequencer/Private/Tools/NavigationTool/NavigationToolSettings.cpp
+++ b/Engine/Source/Editor/Sequencer/Private/Tools/NavigationTool/NavigationToolSettings.cpp
@@ -123,11 +123,11 @@ ENavigationToolItemViewMode UNavigationToolSettings::GetItemProxyViewMode() cons
 	return ItemProxyViewMode;
 }
 
-FNavigationToolColumnView* UNavigationToolSettings::FindCustomColumnView(const FName InColumnId)
+FNavigationToolColumnView* UNavigationToolSettings::FindCustomColumnView(const FText& InColumnViewName)
 {
 	for (FNavigationToolColumnView& ColumnView : CustomColumnViews)
 	{
-		if (ColumnView.ViewName == InColumnId)
+		if (ColumnView.ViewName.EqualTo(InColumnViewName))
 		{
 			return &ColumnView;
 		}
diff --git a/Engine/Source/Editor/Sequencer/Private/Tools/NavigationTool/NavigationToolView.cpp b/Engine/Source/Editor/Sequencer/Private/Tools/NavigationTool/NavigationToolView.cpp
--- a/Engine/Source/Editor/Sequencer/Private/Tools/NavigationTool/NavigationToolView.cpp
+++ b/Engine/Source/Editor/Sequencer/Private/Tools/NavigationTool/NavigationToolView.cpp
@@ -97,12 +97,6 @@ void FNavigationToolView::CreateColumns(const TSharedRef<FNavigationToolProvider
 	FNavigationToolColumnExtender ColumnExtender;
 	InProvider->OnExtendColumns(ColumnExtender);
 
-	// Automatically set provider for the extended columns
-	for (const TSharedPtr<FNavigationToolColumn>& Column : ColumnExtender.GetColumns())
-	{
-		Column->SetProvider(InProvider);
-	}
-
 	// Sort and re-cache columns
 	const TArray<TSharedPtr<FNavigationToolColumn>>& ColumnsToAdd = ColumnExtender.GetColumns();
 
@@ -336,31 +330,17 @@ bool FNavigationToolView::ShouldShowColumnByDefault(const TSharedPtr<INavigation
 		return false;
 	}
 
-	const TSharedPtr<FNavigationToolProvider> Provider = InColumn->GetProvider();
-	if (!Provider.IsValid())
-	{
-		return false;
-	}
-
 	const TSharedPtr<FNavigationTool> Tool = WeakTool.Pin();
 	if (!Tool.IsValid())
 	{
 		return false;
 	}
 
-	const FNavigationToolViewSaveState* const ViewSaveState = Provider->GetViewSaveState(*Tool, ToolViewId);
-	if (!ViewSaveState)
-	{
-		return false;
-	}
+	bool bShouldShow = IsColumnVisible(InColumn);
 
-	const FName ColumnId = InColumn->GetColumnId();
-	if (const FNavigationToolViewColumnSaveState* const FoundColumnState = ViewSaveState->ColumnsState.Find(ColumnId))
-	{
-		return FoundColumnState->bVisible;
-	}
+	bShouldShow |= InColumn->ShouldShowColumnByDefault();
 
-	return InColumn->ShouldShowColumnByDefault();
+	return bShouldShow;
 }
 
 void FNavigationToolView::RequestRefresh()
@@ -1036,31 +1016,29 @@ void FNavigationToolView::HideColumn(const TSharedPtr<INavigationToolColumn>& In
 
 bool FNavigationToolView::IsColumnVisible(const TSharedPtr<INavigationToolColumn>& InColumn) const
 {
-	const TSharedPtr<FNavigationToolProvider> Provider = InColumn->GetProvider();
-	if (!Provider.IsValid())
-	{
-		return false;
-	}
-
 	const TSharedPtr<FNavigationTool> Tool = WeakTool.Pin();
 	if (!Tool.IsValid())
 	{
 		return false;
 	}
 
-	FNavigationToolViewSaveState* const ViewSaveState = Provider->GetViewSaveState(*Tool, ToolViewId);
-	if (!ViewSaveState)
-	{
-		return false;
-	}
+	bool bShouldShow = false;
 
-	const FName ColumnId = InColumn->GetColumnId();
-	if (const FNavigationToolViewColumnSaveState* const FoundColumnState = ViewSaveState->ColumnsState.Find(ColumnId))
-	{
-		return FoundColumnState->bVisible;
-	}
+	Tool->ForEachProvider([this, &InColumn, &Tool, &bShouldShow]
+		(const TSharedRef<FNavigationToolProvider>& InProvider)
+		{
+			if (FNavigationToolViewSaveState* const ViewSaveState = InProvider->GetViewSaveState(*Tool, ToolViewId))
+			{
+				const FName ColumnId = InColumn->GetColumnId();
+				if (const FNavigationToolViewColumnSaveState* const FoundColumnState = ViewSaveState->ColumnsState.Find(ColumnId))
+				{
+					bShouldShow |= FoundColumnState->bVisible;
+				}
+			}
+			return true;
+		});
 
-	return false;
+	return bShouldShow;
 }
 
 ENavigationToolItemViewMode FNavigationToolView::GetItemDefaultViewMode() const
@@ -1094,7 +1072,7 @@ void FNavigationToolView::ToggleItemProxyViewModeSupport(ENavigationToolItemView
 ECheckBoxState FNavigationToolView::GetViewModeCheckState(ENavigationToolItemViewMode InViewMode, ENavigationToolItemViewMode InFlags) const
 {
 	const ENavigationToolItemViewMode Result = InViewMode & InFlags;
-	
+
 	if (Result == InFlags)
 	{
 		return ECheckBoxState::Checked;
@@ -1707,7 +1685,7 @@ void FNavigationToolView::SaveViewState(const TSharedRef<FNavigationToolProvider
 			ViewSaveState->ActiveItemFilters.Add(*ActiveItemFilter->GetDisplayName().ToString());
 		}
 
-		SaveColumnState(NAME_None);
+		SaveColumnState();
 		SaveToolViewItems(*ViewSaveState);
 	}
 }
@@ -1756,7 +1734,8 @@ void FNavigationToolView::SaveColumnState(const FName InColumnId)
 	// Save all columns unless a specific column is specified
 	if (InColumnId.IsNone())
 	{
-		// Save each column to their respective provider save data
+		// Save each column to their respective providers save data
+		// Note: Some columns may have multiple providers
 		for (const TTuple<FName, TSharedPtr<INavigationToolColumn>>& Column : Columns)
 		{
 			if (!Column.Value.IsValid())
@@ -1764,32 +1743,34 @@ void FNavigationToolView::SaveColumnState(const FName InColumnId)
 				continue;
 			}
 
-			const TSharedPtr<FNavigationToolProvider> Provider = Column.Value->GetProvider();
-			if (!Provider.IsValid())
-			{
-				continue;
-			}
+			const FName ColumnId = Column.Value->GetColumnId();
 
-			if (FNavigationToolViewSaveState* const ViewSaveState = Provider->GetViewSaveState(ToolRef, ToolViewId))
-			{
-				const FName ColumnId = Column.Value->GetColumnId();
-				FNavigationToolViewColumnSaveState& FoundColumnState = ViewSaveState->ColumnsState.FindOrAdd(ColumnId);
-				ToolViewWidget->GenerateColumnState(ColumnId, FoundColumnState);
-			}
+			Tool->ForEachProvider([this, &ColumnId, &ToolRef]
+				(const TSharedRef<FNavigationToolProvider>& InProvider)
+				{
+					if (FNavigationToolViewSaveState* const ViewSaveState = InProvider->GetViewSaveState(ToolRef, ToolViewId))
+					{
+						FNavigationToolViewColumnSaveState& FoundColumnState = ViewSaveState->ColumnsState.FindOrAdd(ColumnId);
+						ToolViewWidget->GenerateColumnState(ColumnId, FoundColumnState);
+					}
+					return true;
+				});
 		}
 	}
 	else if (ensure(Columns.Contains(InColumnId)))
 	{
-		// Save the specific column to its providers save data
-		const TSharedPtr<INavigationToolColumn>& Column = Columns[InColumnId];
-		if (const TSharedPtr<FNavigationToolProvider> Provider = Column->GetProvider())
-		{
-			if (FNavigationToolViewSaveState* const ViewSaveState = Provider->GetViewSaveState(ToolRef, ToolViewId))
+		Tool->ForEachProvider([this, &InColumnId, &ToolRef]
+			(const TSharedRef<FNavigationToolProvider>& InProvider)
 			{
-				FNavigationToolViewColumnSaveState& FoundColumnState = ViewSaveState->ColumnsState.FindOrAdd(InColumnId);
-				ToolViewWidget->GenerateColumnState(InColumnId, FoundColumnState);
-			}
-		}
+				// Save the specific column to its providers save data
+				const TSharedPtr<INavigationToolColumn>& Column = Columns[InColumnId];
+				if (FNavigationToolViewSaveState* const ViewSaveState = InProvider->GetViewSaveState(ToolRef, ToolViewId))
+				{
+					FNavigationToolViewColumnSaveState& FoundColumnState = ViewSaveState->ColumnsState.FindOrAdd(InColumnId);
+					ToolViewWidget->GenerateColumnState(InColumnId, FoundColumnState);
+				}
+				return true;
+			});
 	}
 }
 
@@ -1916,12 +1897,6 @@ void FNavigationToolView::ResetColumnSize(const FName InColumnId)
 		return;
 	}
 
-	const TSharedPtr<FNavigationToolProvider> Provider = Columns[InColumnId]->GetProvider();
-	if (!Provider.IsValid())
-	{
-		return;
-	}
-
 	const TSharedPtr<FNavigationTool> Tool = WeakTool.Pin();
 	if (!Tool.IsValid())
 	{
@@ -1930,23 +1905,23 @@ void FNavigationToolView::ResetColumnSize(const FName InColumnId)
 
 	const FNavigationTool& ToolRef = *Tool;
 
-	FNavigationToolViewSaveState* const ViewSaveState = Provider->GetViewSaveState(ToolRef, ToolViewId);
-	if (!ViewSaveState)
-	{
-		return;
-	}
-
-	const float DefaultSize = Columns[InColumnId]->GetFillWidth();
-	if (DefaultSize <= 0.f)
-	{
-		return;
-	}
-
-	ViewSaveState->ColumnsState.FindOrAdd(InColumnId).Size = DefaultSize;
+	Tool->ForEachProvider([this, &InColumnId, &ToolRef]
+		(const TSharedRef<FNavigationToolProvider>& InProvider)
+		{
+			if (FNavigationToolViewSaveState* const ViewSaveState = InProvider->GetViewSaveState(ToolRef, ToolViewId))
+			{
+				const float DefaultSize = Columns[InColumnId]->GetFillWidth();
+				if (DefaultSize > 0.f)
+				{
+					ViewSaveState->ColumnsState.FindOrAdd(InColumnId).Size = DefaultSize;
 
-	ToolViewWidget->SetColumnWidth(InColumnId, DefaultSize);
+					ToolViewWidget->SetColumnWidth(InColumnId, DefaultSize);
 
-	ToolViewWidget->GenerateColumnState(InColumnId, ViewSaveState->ColumnsState.FindOrAdd(InColumnId));
+					ToolViewWidget->GenerateColumnState(InColumnId, ViewSaveState->ColumnsState.FindOrAdd(InColumnId));
+				}
+			}
+			return true;
+		});
 }
 
 bool FNavigationToolView::CanResetColumnSize(const FName InColumnId) const
@@ -1956,12 +1931,6 @@ bool FNavigationToolView::CanResetColumnSize(const FName InColumnId) const
 		return false;
 	}
 
-	const TSharedPtr<FNavigationToolProvider> Provider = Columns[InColumnId]->GetProvider();
-	if (!Provider.IsValid())
-	{
-		return false;
-	}
-
 	const TSharedPtr<FNavigationTool> Tool = WeakTool.Pin();
 	if (!Tool.IsValid())
 	{
@@ -1970,19 +1939,29 @@ bool FNavigationToolView::CanResetColumnSize(const FName InColumnId) const
 
 	const FNavigationTool& ToolRef = *Tool;
 
-	const FNavigationToolViewSaveState* const ViewSaveState = Provider->GetViewSaveState(ToolRef, ToolViewId);
-	if (!ViewSaveState || !ViewSaveState->ColumnsState.Contains(InColumnId))
-	{
-		return false;
-	}
+	bool bCanReset = false;
 
-	const float DefaultSize = Columns[InColumnId]->GetFillWidth();
-	if (DefaultSize <= 0.f)
-	{
-		return false;
-	}
+	Tool->ForEachProvider([this, &InColumnId, &ToolRef, &bCanReset]
+		(const TSharedRef<FNavigationToolProvider>& InProvider)
+		{
+			const FNavigationToolViewSaveState* const ViewSaveState = InProvider->GetViewSaveState(ToolRef, ToolViewId);
+			if (!ViewSaveState || !ViewSaveState->ColumnsState.Contains(InColumnId))
+			{
+				return true;
+			}
+
+			const float DefaultSize = Columns[InColumnId]->GetFillWidth();
+			if (DefaultSize <= 0.f)
+			{
+				return true;
+			}
+
+			bCanReset |= ViewSaveState->ColumnsState[InColumnId].Size != DefaultSize;
 
-	return ViewSaveState->ColumnsState[InColumnId].Size != DefaultSize;
+			return true;
+		});
+
+	return bCanReset;
 }
 
 void FNavigationToolView::ResetVisibleColumnSizes()
@@ -2015,21 +1994,16 @@ void FNavigationToolView::ResetVisibleColumnSizes()
 			continue;
 		}
 
-		const TSharedPtr<FNavigationToolProvider> Provider = Pair.Value->GetProvider();
-		if (!Provider.IsValid())
-		{
-			continue;
-		}
-
-		FNavigationToolViewSaveState* const ViewSaveState = Provider->GetViewSaveState(ToolRef, ToolViewId);
-		if (!ViewSaveState)
-		{
-			return;
-		}
-
-		ToolViewWidget->SetColumnWidth(ColumnId, DefaultSize);
-
-		ToolViewWidget->GenerateColumnState(ColumnId, ViewSaveState->ColumnsState.FindOrAdd(ColumnId));
+		Tool->ForEachProvider([this, &ColumnId, &ToolRef, DefaultSize]
+			(const TSharedRef<FNavigationToolProvider>& InProvider)
+			{
+				if (FNavigationToolViewSaveState* const ViewSaveState = InProvider->GetViewSaveState(ToolRef, ToolViewId))
+				{
+					ToolViewWidget->SetColumnWidth(ColumnId, DefaultSize);
+					ToolViewWidget->GenerateColumnState(ColumnId, ViewSaveState->ColumnsState.FindOrAdd(ColumnId));
+				}
+				return true;
+			});
 	}
 }
 
@@ -2042,9 +2016,33 @@ void FNavigationToolView::SaveNewCustomColumnView()
 {
 	TSet<FNavigationToolColumnView>& CustomColumnViews = ToolSettings->GetCustomColumnViews();
 
+	SaveColumnState();
+
+	// Create a unique column view name suggestion
+	auto DoesColumnViewExist = [&CustomColumnViews](const FText& InViewName)
+	{
+		for (const FNavigationToolColumnView& ColumnView : CustomColumnViews)
+		{
+			if (ColumnView.ViewName.EqualTo(InViewName))
+			{
+				return true;
+			}
+		}
+		return false;
+	};
+
 	FNavigationToolColumnView NewColumnView;
-	NewColumnView.ViewName = TEXT("Column View 1");
-	for (const TPair<FName, TSharedPtr<INavigationToolColumn>>& Column : GetColumns())
+
+	for (int32 Index = 1; Index < INT_MAX; ++Index)
+	{
+		NewColumnView.ViewName = FText::Format(LOCTEXT("ColumnViewName", "Column View {0}"), { Index });
+		if (!DoesColumnViewExist(NewColumnView.ViewName))
+		{
+			break;
+		}
+	}
+
+	for (const TPair<FName, TSharedPtr<INavigationToolColumn>>& Column : Columns)
 	{
 		const FName ColumnId = Column.Value->GetColumnId();
 		if (IsColumnVisible(Column.Value))
@@ -2053,17 +2051,14 @@ void FNavigationToolView::SaveNewCustomColumnView()
 		}
 	}
 
-	FString UserEnteredText;
-
+	// Prompt user for name, using the generated unique suggestion as the default name
 	FModalTextInputDialog InputDialog;
 	InputDialog.InputLabel = LOCTEXT("NewColumnViewName_InputLabel", "New Column View Name");
-	if (!InputDialog.Open(NewColumnView.ViewName.ToString(), UserEnteredText))
+	if (!InputDialog.Open(NewColumnView.ViewName, NewColumnView.ViewName))
 	{
 		return;
 	}
 
-	NewColumnView.ViewName = FName(UserEnteredText);
-
 	bool bAlreadyExists = false;
 	CustomColumnViews.Add(NewColumnView, &bAlreadyExists);
 
@@ -2076,26 +2071,26 @@ void FNavigationToolView::SaveNewCustomColumnView()
 
 	CustomColumnViews.Sort([](const FNavigationToolColumnView& InA, const FNavigationToolColumnView& InB)
 		{
-			return InA.ViewName.FastLess(InB.ViewName);
+			return InA.ViewName.CompareTo(InB.ViewName) < 0;
 		});
 
 	ToolSettings->SaveConfig();
 }
 
-void FNavigationToolView::ApplyCustomColumnView(const FName InColumnViewName)
+void FNavigationToolView::ApplyCustomColumnView(const FText InColumnViewName)
 {
-	if (!ToolViewWidget.IsValid() || InColumnViewName.IsNone())
+	if (!ToolViewWidget.IsValid() || InColumnViewName.IsEmptyOrWhitespace())
 	{
 		return;
 	}
 
-	FNavigationToolColumnView* const SavedColumnView = ToolSettings->FindCustomColumnView(InColumnViewName);
+	const FNavigationToolColumnView* const SavedColumnView = ToolSettings->FindCustomColumnView(InColumnViewName);
 	if (!SavedColumnView)
 	{
 		return;
 	}
 
-	for (const TPair<FName, TSharedPtr<INavigationToolColumn>>& Column : GetColumns())
+	for (const TPair<FName, TSharedPtr<INavigationToolColumn>>& Column : Columns)
 	{
 		const FName ColumnId = Column.Value->GetColumnId();
 		const bool bColumnVisible = SavedColumnView->VisibleColumns.Contains(ColumnId);
diff --git a/Engine/Source/Editor/Sequencer/Private/Tools/NavigationTool/NavigationToolView.h b/Engine/Source/Editor/Sequencer/Private/Tools/NavigationTool/NavigationToolView.h
--- a/Engine/Source/Editor/Sequencer/Private/Tools/NavigationTool/NavigationToolView.h
+++ b/Engine/Source/Editor/Sequencer/Private/Tools/NavigationTool/NavigationToolView.h
@@ -327,7 +327,7 @@ public:
 	void SaveNewCustomColumnView();
 
 	/** Apply a saved custom column view by name */
-	void ApplyCustomColumnView(const FName InColumnViewName);
+	void ApplyCustomColumnView(const FText InColumnViewName);
 
 protected:
 	/** Initializes the Navigation Tool Instance. Only executed once */
diff --git a/Engine/Source/Editor/Sequencer/Private/Tools/NavigationTool/Providers/NavigationToolProvider.cpp b/Engine/Source/Editor/Sequencer/Private/Tools/NavigationTool/Providers/NavigationToolProvider.cpp
--- a/Engine/Source/Editor/Sequencer/Private/Tools/NavigationTool/Providers/NavigationToolProvider.cpp
+++ b/Engine/Source/Editor/Sequencer/Private/Tools/NavigationTool/Providers/NavigationToolProvider.cpp
@@ -149,8 +149,6 @@ void FNavigationToolProvider::Activate(FNavigationTool& InTool)
 
 	InTool.ForEachToolView([&SharedThisRef](const TSharedRef<FNavigationToolView>& InToolView)
 		{
-			//InToolView->LoadViewState(InProvider);
-
 			InToolView->CreateColumns(SharedThisRef);
 			InToolView->CreateDefaultColumnViews(SharedThisRef);
 		});
@@ -175,7 +173,7 @@ void FNavigationToolProvider::CleanupExtendedColumnViews()
 
 	TSet<FNavigationToolColumnView>& CustomColumnViews = Settings->GetCustomColumnViews();
 
-	for (const FName ColumnViewName : ExtendedColumnViewNames)
+	for (const FText& ColumnViewName : ExtendedColumnViewNames)
 	{
 		CustomColumnViews.Remove(ColumnViewName);
 	}
diff --git a/Engine/Source/Editor/Sequencer/Private/Tools/NavigationTool/Widgets/ModalTextInputDialog.cpp b/Engine/Source/Editor/Sequencer/Private/Tools/NavigationTool/Widgets/ModalTextInputDialog.cpp
--- a/Engine/Source/Editor/Sequencer/Private/Tools/NavigationTool/Widgets/ModalTextInputDialog.cpp
+++ b/Engine/Source/Editor/Sequencer/Private/Tools/NavigationTool/Widgets/ModalTextInputDialog.cpp
@@ -187,7 +187,7 @@ FModalTextInputDialog::FModalTextInputDialog()
 	InputLabel = FText::GetEmpty();
 }
 
-bool FModalTextInputDialog::Open(const FString& InDefaultText, FString& OutText, const TSharedPtr<SWindow>& InParentWindow)
+bool FModalTextInputDialog::Open(const FText& InDefaultText, FText& OutText, const TSharedPtr<SWindow>& InParentWindow)
 {
 	if (WindowWidget.IsValid())
 	{
@@ -198,7 +198,7 @@ bool FModalTextInputDialog::Open(const FString& InDefaultText, FString& OutText,
 	const TSharedRef<SModalTextInputDialog> DialogWidget = SNew(SModalTextInputDialog)
 		.DialogTitle(DialogTitle)
 		.InputLabel(InputLabel)
-		.DefaultText(FText::FromString(InDefaultText));
+		.DefaultText(InDefaultText);
 	WindowWidget = DialogWidget;
 
 	WindowWidget->SetOnWindowClosed(FOnWindowClosed::CreateLambda([this](const TSharedRef<SWindow>& InWindow)
@@ -210,7 +210,7 @@ bool FModalTextInputDialog::Open(const FString& InDefaultText, FString& OutText,
 
 	if (DialogWidget->WasLastAccepted())
 	{
-		OutText = DialogWidget->GetText().ToString();
+		OutText = DialogWidget->GetText();
 		return true;
 	}
 
diff --git a/Engine/Source/Editor/Sequencer/Private/Tools/NavigationTool/Widgets/ModalTextInputDialog.h b/Engine/Source/Editor/Sequencer/Private/Tools/NavigationTool/Widgets/ModalTextInputDialog.h
--- a/Engine/Source/Editor/Sequencer/Private/Tools/NavigationTool/Widgets/ModalTextInputDialog.h
+++ b/Engine/Source/Editor/Sequencer/Private/Tools/NavigationTool/Widgets/ModalTextInputDialog.h
@@ -17,7 +17,7 @@ public:
 	FModalTextInputDialog();
 
 	/** @return True if the text input was accepted. False if cancel was clicked. */
-	bool Open(const FString& InDefaultText, FString& OutText, const TSharedPtr<SWindow>& InParentWindow = nullptr);
+	bool Open(const FText& InDefaultText, FText& OutText, const TSharedPtr<SWindow>& InParentWindow = nullptr);
 
 	bool IsOpen();
 
diff --git a/Engine/Source/Editor/Sequencer/Private/Tools/NavigationTool/Widgets/SNavigationToolView.cpp b/Engine/Source/Editor/Sequencer/Private/Tools/NavigationTool/Widgets/SNavigationToolView.cpp
--- a/Engine/Source/Editor/Sequencer/Private/Tools/NavigationTool/Widgets/SNavigationToolView.cpp
+++ b/Engine/Source/Editor/Sequencer/Private/Tools/NavigationTool/Widgets/SNavigationToolView.cpp
@@ -358,34 +358,34 @@ void SNavigationToolView::ReconstructColumns()
 	// Add the Columns allocated by the Owning Instance
 	for (const TPair<FName, TSharedPtr<INavigationToolColumn>>& Pair : ToolView->GetColumns())
 	{
-		const TSharedPtr<INavigationToolColumn>& Column = Pair.Value;
-		if (!Column.IsValid())
-		{
-			continue;
-		}
+		OwnerTool->ForEachProvider([this, &ToolView, &OwnerTool, &Pair]
+			(const TSharedRef<FNavigationToolProvider>& InProvider)
+			{
+				const TSharedPtr<INavigationToolColumn>& Column = Pair.Value;
+				if (!Column.IsValid())
+				{
+					return true;
+				}
 
-		const TSharedPtr<FNavigationToolProvider> Provider = Column->GetProvider();
-		if (!Provider.IsValid())
-		{
-			continue;
-		}
+				const FName ColumnId = Pair.Key;
 
-		const FName ColumnId = Pair.Key;
+				bool bColumnVisible = Column->ShouldShowColumnByDefault();
+				float ColumnWidth = Column->GetFillWidth();
 
-		bool bColumnVisible = Column->ShouldShowColumnByDefault();
-		float ColumnWidth = Column->GetFillWidth();
+				if (FNavigationToolViewSaveState* const ToolViewSaveState = InProvider->GetViewSaveState(*OwnerTool, ToolView->GetToolViewId()))
+				{
+					if (const FNavigationToolViewColumnSaveState* const ColumnState = ToolViewSaveState->ColumnsState.Find(ColumnId))
+					{
+						bColumnVisible = ColumnState->bVisible;
+						ColumnWidth = ColumnState->Size;
+					}
+				}
 
-		if (FNavigationToolViewSaveState* const ToolViewSaveState = Provider->GetViewSaveState(*OwnerTool, ToolView->GetToolViewId()))
-		{
-			if (const FNavigationToolViewColumnSaveState* const ColumnState = ToolViewSaveState->ColumnsState.Find(ColumnId))
-			{
-				bColumnVisible = ColumnState->bVisible;
-				ColumnWidth = ColumnState->Size;
-			}
-		}
+				HeaderRowWidget->AddColumn(Column->ConstructHeaderRowColumn(ToolView.ToSharedRef(), ColumnWidth));
+				HeaderRowWidget->SetShowGeneratedColumn(ColumnId, bColumnVisible);
 
-		HeaderRowWidget->AddColumn(Column->ConstructHeaderRowColumn(ToolView.ToSharedRef(), ColumnWidth));
-		HeaderRowWidget->SetShowGeneratedColumn(ColumnId, bColumnVisible);
+				return true;
+			});
 	}
 }
 
diff --git a/Engine/Source/Editor/Sequencer/Public/Tools/NavigationTool/Columns/INavigationToolColumn.h b/Engine/Source/Editor/Sequencer/Public/Tools/NavigationTool/Columns/INavigationToolColumn.h
--- a/Engine/Source/Editor/Sequencer/Public/Tools/NavigationTool/Columns/INavigationToolColumn.h
+++ b/Engine/Source/Editor/Sequencer/Public/Tools/NavigationTool/Columns/INavigationToolColumn.h
@@ -9,16 +9,13 @@
 namespace UE::Sequencer::NavigationTool
 {
 
-class FNavigationToolProvider;
 class INavigationToolView;
 class SNavigationToolTreeRow;
 
 class INavigationToolColumn : public INavigationToolItemTypeCastable, public TSharedFromThis<INavigationToolColumn>
 {
 public:
-	FName GetColumnId() const { return GetTypeId().ToName(); }
-
-	virtual TSharedPtr<FNavigationToolProvider> GetProvider() const = 0;
+	virtual FName GetColumnId() const { return GetTypeId().ToName(); }
 
 	virtual FText GetColumnDisplayNameText() const = 0;
 
diff --git a/Engine/Source/Editor/Sequencer/Public/Tools/NavigationTool/Columns/NavigationToolColumn.h b/Engine/Source/Editor/Sequencer/Public/Tools/NavigationTool/Columns/NavigationToolColumn.h
--- a/Engine/Source/Editor/Sequencer/Public/Tools/NavigationTool/Columns/NavigationToolColumn.h
+++ b/Engine/Source/Editor/Sequencer/Public/Tools/NavigationTool/Columns/NavigationToolColumn.h
@@ -10,17 +10,7 @@ namespace UE::Sequencer::NavigationTool
 /** Base Navigation Tool Column to extend from */
 class FNavigationToolColumn : public INavigationToolColumn
 {
-public:
-	//~ Begin INavigationToolColumn
-	SEQUENCER_API virtual TSharedPtr<FNavigationToolProvider> GetProvider() const final override;
-	//~ End INavigationToolColumn
 
-protected:
-	friend class FNavigationToolView;
-
-	void SetProvider(const TWeakPtr<FNavigationToolProvider>& InWeakProvider);
-
-	TWeakPtr<FNavigationToolProvider> WeakProvider;
 };
 
 }
diff --git a/Engine/Source/Editor/Sequencer/Public/Tools/NavigationTool/NavigationToolSettings.h b/Engine/Source/Editor/Sequencer/Public/Tools/NavigationTool/NavigationToolSettings.h
--- a/Engine/Source/Editor/Sequencer/Public/Tools/NavigationTool/NavigationToolSettings.h
+++ b/Engine/Source/Editor/Sequencer/Public/Tools/NavigationTool/NavigationToolSettings.h
@@ -28,24 +28,24 @@ struct FNavigationToolColumnView
 	GENERATED_BODY()
 
 	FNavigationToolColumnView() = default;
-	FNavigationToolColumnView(const FName InViewName)
+	FNavigationToolColumnView(const FText& InViewName)
 		: ViewName(InViewName)
 	{}
 
 	UPROPERTY()
-	FName ViewName;
+	FText ViewName;
 
 	UPROPERTY()
 	TSet<FName> VisibleColumns;
 
 	bool operator==(const FNavigationToolColumnView& InOtherView) const
 	{
-		return ViewName == InOtherView.ViewName;
+		return ViewName.EqualTo(InOtherView.ViewName);
 	}
 
 	friend uint32 GetTypeHash(const FNavigationToolColumnView& InOtherView)
 	{
-		return GetTypeHash(InOtherView.ViewName);
+		return GetTypeHash(InOtherView.ViewName.ToString());
 	} 
 };
 
@@ -82,7 +82,7 @@ public:
 	void ToggleItemProxyViewModeSupport(const ENavigationToolItemViewMode InFlags);
 
 	TSet<FNavigationToolColumnView>& GetCustomColumnViews() { return CustomColumnViews; }
-	FNavigationToolColumnView* FindCustomColumnView(const FName InColumnId);
+	FNavigationToolColumnView* FindCustomColumnView(const FText& InColumnViewName);
 
 	const TSet<FName>& GetEnabledBuiltInFilters() const { return EnabledBuiltInFilters; }
 	void SetBuiltInFilterEnabled(const FName InFilterName, const bool bInEnabled);
diff --git a/Engine/Source/Editor/Sequencer/Public/Tools/NavigationTool/Providers/INavigationToolProvider.h b/Engine/Source/Editor/Sequencer/Public/Tools/NavigationTool/Providers/INavigationToolProvider.h
--- a/Engine/Source/Editor/Sequencer/Public/Tools/NavigationTool/Providers/INavigationToolProvider.h
+++ b/Engine/Source/Editor/Sequencer/Public/Tools/NavigationTool/Providers/INavigationToolProvider.h
@@ -5,6 +5,7 @@
 #include "Containers/Set.h"
 #include "Templates/SharedPointer.h"
 #include "Templates/SubclassOf.h"
+#include "Internationalization/Text.h"
 #include "Tools/NavigationTool/Filters/Filters/NavigationToolBuiltInFilterParams.h"
 #include "Tools/NavigationTool/NavigationToolDefines.h"
 
@@ -102,7 +103,7 @@ public:
 	virtual bool ShouldHideItem(const FNavigationToolItemPtr& InItem) const = 0;
 
 	/** @return The name of the default column view to load when this provider is registered  */
-	virtual FName GetDefaultColumnView() const { return NAME_None; }
+	virtual FText GetDefaultColumnView() const { return FText::GetEmpty(); }
 
 	/** Determines whether an external Drag Drop event (i.e. not an Navigation Tool one) can be accepted by the Navigation Tool for a given Target Item */
 	virtual TOptional<EItemDropZone> OnToolItemCanAcceptDrop(const FDragDropEvent& InDragDropEvent
diff --git a/Engine/Source/Editor/Sequencer/Public/Tools/NavigationTool/Providers/NavigationToolProvider.h b/Engine/Source/Editor/Sequencer/Public/Tools/NavigationTool/Providers/NavigationToolProvider.h
--- a/Engine/Source/Editor/Sequencer/Public/Tools/NavigationTool/Providers/NavigationToolProvider.h
+++ b/Engine/Source/Editor/Sequencer/Public/Tools/NavigationTool/Providers/NavigationToolProvider.h
@@ -50,7 +50,7 @@ protected:
 
 	void CleanupExtendedColumnViews();
 
-	TArray<FName> ExtendedColumnViewNames;
+	TArray<FText> ExtendedColumnViewNames;
 	TArray<FName> ExtendedBuiltInFilterNames;
 
 private:
