[Sequencer] Navigation Tool: Hardcode column Ids instead of using class type names to reduce maintainability if class names change

diff --git a/Engine/Source/Editor/Sequencer/Public/Tools/NavigationTool/Columns/NavigationToolColorColumn.h b/Engine/Source/Editor/Sequencer/Public/Tools/NavigationTool/Columns/NavigationToolColorColumn.h
--- a/Engine/Source/Editor/Sequencer/Public/Tools/NavigationTool/Columns/NavigationToolColorColumn.h
+++ b/Engine/Source/Editor/Sequencer/Public/Tools/NavigationTool/Columns/NavigationToolColorColumn.h
@@ -14,8 +14,11 @@ class FNavigationToolColorColumn : public FNavigationToolColumn
 public:
 	UE_NAVIGATIONTOOL_INHERITS(FNavigationToolColorColumn, FNavigationToolColumn);
 
+	static FName StaticColumnId() { return TEXT("Color"); }
+
 protected:
 	//~ Begin INavigationToolColumn
+	virtual FName GetColumnId() const override { return StaticColumnId(); }
 	SEQUENCER_API virtual FText GetColumnDisplayNameText() const override;
 	virtual bool ShouldShowColumnByDefault() const override { return true; }
 	SEQUENCER_API virtual SHeaderRow::FColumn::FArguments ConstructHeaderRowColumn(const TSharedRef<INavigationToolView>& InToolView, const float InFillSize) override;
diff --git a/Engine/Source/Editor/Sequencer/Public/Tools/NavigationTool/Columns/NavigationToolCommentColumn.h b/Engine/Source/Editor/Sequencer/Public/Tools/NavigationTool/Columns/NavigationToolCommentColumn.h
--- a/Engine/Source/Editor/Sequencer/Public/Tools/NavigationTool/Columns/NavigationToolCommentColumn.h
+++ b/Engine/Source/Editor/Sequencer/Public/Tools/NavigationTool/Columns/NavigationToolCommentColumn.h
@@ -14,8 +14,11 @@ class FNavigationToolCommentColumn : public FNavigationToolColumn
 public:
 	UE_NAVIGATIONTOOL_INHERITS(FNavigationToolCommentColumn, FNavigationToolColumn);
 
+	static FName StaticColumnId() { return TEXT("Comment"); }
+
 protected:
 	//~ Begin INavigationToolColumn
+	virtual FName GetColumnId() const override { return StaticColumnId(); }
 	SEQUENCER_API virtual FText GetColumnDisplayNameText() const override;
 	SEQUENCER_API virtual const FSlateBrush* GetIconBrush() const override;
 	virtual bool ShouldShowColumnByDefault() const override { return false; }
diff --git a/Engine/Source/Editor/Sequencer/Public/Tools/NavigationTool/Columns/NavigationToolDeactiveStateColumn.h b/Engine/Source/Editor/Sequencer/Public/Tools/NavigationTool/Columns/NavigationToolDeactiveStateColumn.h
--- a/Engine/Source/Editor/Sequencer/Public/Tools/NavigationTool/Columns/NavigationToolDeactiveStateColumn.h
+++ b/Engine/Source/Editor/Sequencer/Public/Tools/NavigationTool/Columns/NavigationToolDeactiveStateColumn.h
@@ -14,7 +14,11 @@ class FNavigationToolDeactiveStateColumn : public FNavigationToolColumn
 public:
 	UE_NAVIGATIONTOOL_INHERITS(FNavigationToolDeactiveStateColumn, FNavigationToolColumn);
 
+	static FName StaticColumnId() { return TEXT("DeactiveState"); }
+
+protected:
 	//~ Begin INavigationToolColumn
+	virtual FName GetColumnId() const override { return StaticColumnId(); }
 	SEQUENCER_API virtual FText GetColumnDisplayNameText() const override;
 	SEQUENCER_API virtual const FSlateBrush* GetIconBrush() const override;
 	SEQUENCER_API virtual SHeaderRow::FColumn::FArguments ConstructHeaderRowColumn(const TSharedRef<INavigationToolView>& InToolView, const float InFillSize) override;
diff --git a/Engine/Source/Editor/Sequencer/Public/Tools/NavigationTool/Columns/NavigationToolHBiasColumn.h b/Engine/Source/Editor/Sequencer/Public/Tools/NavigationTool/Columns/NavigationToolHBiasColumn.h
--- a/Engine/Source/Editor/Sequencer/Public/Tools/NavigationTool/Columns/NavigationToolHBiasColumn.h
+++ b/Engine/Source/Editor/Sequencer/Public/Tools/NavigationTool/Columns/NavigationToolHBiasColumn.h
@@ -14,8 +14,11 @@ class FNavigationToolHBiasColumn : public FNavigationToolColumn
 public:
 	UE_NAVIGATIONTOOL_INHERITS(FNavigationToolHBiasColumn, FNavigationToolColumn);
 
+	static FName StaticColumnId() { return TEXT("HBias"); }
+
 protected:
 	//~ Begin INavigationToolColumn
+	virtual FName GetColumnId() const override { return StaticColumnId(); }
 	SEQUENCER_API virtual FText GetColumnDisplayNameText() const override;
 	SEQUENCER_API virtual const FSlateBrush* GetIconBrush() const override;
 	virtual bool ShouldShowColumnByDefault() const override { return false; }
diff --git a/Engine/Source/Editor/Sequencer/Public/Tools/NavigationTool/Columns/NavigationToolInTimeColumn.h b/Engine/Source/Editor/Sequencer/Public/Tools/NavigationTool/Columns/NavigationToolInTimeColumn.h
--- a/Engine/Source/Editor/Sequencer/Public/Tools/NavigationTool/Columns/NavigationToolInTimeColumn.h
+++ b/Engine/Source/Editor/Sequencer/Public/Tools/NavigationTool/Columns/NavigationToolInTimeColumn.h
@@ -14,8 +14,11 @@ class FNavigationToolInTimeColumn : public FNavigationToolColumn
 public:
 	UE_NAVIGATIONTOOL_INHERITS(FNavigationToolInTimeColumn, FNavigationToolColumn);
 
+	static FName StaticColumnId() { return TEXT("InTime"); }
+
 protected:
 	//~ Begin INavigationToolColumn
+	virtual FName GetColumnId() const override { return StaticColumnId(); }
 	SEQUENCER_API virtual FText GetColumnDisplayNameText() const override;
 	SEQUENCER_API virtual const FSlateBrush* GetIconBrush() const override;
 	virtual bool ShouldShowColumnByDefault() const override { return true; }
diff --git a/Engine/Source/Editor/Sequencer/Public/Tools/NavigationTool/Columns/NavigationToolItemsColumn.h b/Engine/Source/Editor/Sequencer/Public/Tools/NavigationTool/Columns/NavigationToolItemsColumn.h
--- a/Engine/Source/Editor/Sequencer/Public/Tools/NavigationTool/Columns/NavigationToolItemsColumn.h
+++ b/Engine/Source/Editor/Sequencer/Public/Tools/NavigationTool/Columns/NavigationToolItemsColumn.h
@@ -14,8 +14,11 @@ class FNavigationToolItemsColumn : public FNavigationToolColumn
 public:
 	UE_NAVIGATIONTOOL_INHERITS(FNavigationToolItemsColumn, FNavigationToolColumn);
 
+	static FName StaticColumnId() { return TEXT("Items"); }
+
 protected:
 	//~ Begin INavigationToolColumn
+	virtual FName GetColumnId() const override { return StaticColumnId(); }
 	SEQUENCER_API virtual FText GetColumnDisplayNameText() const override;
 	SEQUENCER_API virtual const FSlateBrush* GetIconBrush() const override;
 	virtual bool ShouldShowColumnByDefault() const override { return true; }
diff --git a/Engine/Source/Editor/Sequencer/Public/Tools/NavigationTool/Columns/NavigationToolLabelColumn.h b/Engine/Source/Editor/Sequencer/Public/Tools/NavigationTool/Columns/NavigationToolLabelColumn.h
--- a/Engine/Source/Editor/Sequencer/Public/Tools/NavigationTool/Columns/NavigationToolLabelColumn.h
+++ b/Engine/Source/Editor/Sequencer/Public/Tools/NavigationTool/Columns/NavigationToolLabelColumn.h
@@ -14,8 +14,11 @@ class FNavigationToolLabelColumn : public FNavigationToolColumn
 public:
 	UE_NAVIGATIONTOOL_INHERITS(FNavigationToolLabelColumn, FNavigationToolColumn);
 
+	static FName StaticColumnId() { return TEXT("Label"); }
+
 protected:
 	//~ Begin INavigationToolColumn
+	virtual FName GetColumnId() const override { return StaticColumnId(); }
 	SEQUENCER_API virtual FText GetColumnDisplayNameText() const override;
 	SEQUENCER_API virtual const FSlateBrush* GetIconBrush() const override;
 	SEQUENCER_API virtual SHeaderRow::FColumn::FArguments ConstructHeaderRowColumn(const TSharedRef<INavigationToolView>& InToolView, const float InFillSize) override;
diff --git a/Engine/Source/Editor/Sequencer/Public/Tools/NavigationTool/Columns/NavigationToolLengthColumn.h b/Engine/Source/Editor/Sequencer/Public/Tools/NavigationTool/Columns/NavigationToolLengthColumn.h
--- a/Engine/Source/Editor/Sequencer/Public/Tools/NavigationTool/Columns/NavigationToolLengthColumn.h
+++ b/Engine/Source/Editor/Sequencer/Public/Tools/NavigationTool/Columns/NavigationToolLengthColumn.h
@@ -14,8 +14,11 @@ class FNavigationToolLengthColumn : public FNavigationToolColumn
 public:
 	UE_NAVIGATIONTOOL_INHERITS(FNavigationToolLengthColumn, FNavigationToolColumn);
 
+	static FName StaticColumnId() { return TEXT("Length"); }
+
 protected:
 	//~ Begin INavigationToolColumn
+	virtual FName GetColumnId() const override { return StaticColumnId(); }
 	SEQUENCER_API virtual FText GetColumnDisplayNameText() const override;
 	SEQUENCER_API virtual const FSlateBrush* GetIconBrush() const override;
 	virtual bool ShouldShowColumnByDefault() const override { return false; }
diff --git a/Engine/Source/Editor/Sequencer/Public/Tools/NavigationTool/Columns/NavigationToolLockColumn.h b/Engine/Source/Editor/Sequencer/Public/Tools/NavigationTool/Columns/NavigationToolLockColumn.h
--- a/Engine/Source/Editor/Sequencer/Public/Tools/NavigationTool/Columns/NavigationToolLockColumn.h
+++ b/Engine/Source/Editor/Sequencer/Public/Tools/NavigationTool/Columns/NavigationToolLockColumn.h
@@ -14,8 +14,11 @@ class FNavigationToolLockColumn : public FNavigationToolColumn
 public:
 	UE_NAVIGATIONTOOL_INHERITS(FNavigationToolLockColumn, FNavigationToolColumn);
 
+	static FName StaticColumnId() { return TEXT("Lock"); }
+
 protected:
 	//~ Begin INavigationToolColumn
+	virtual FName GetColumnId() const override { return StaticColumnId(); }
 	SEQUENCER_API virtual FText GetColumnDisplayNameText() const override;
 	SEQUENCER_API virtual const FSlateBrush* GetIconBrush() const override;
 	virtual bool ShouldShowColumnByDefault() const override { return true; }
diff --git a/Engine/Source/Editor/Sequencer/Public/Tools/NavigationTool/Columns/NavigationToolMarkerVisibilityColumn.h b/Engine/Source/Editor/Sequencer/Public/Tools/NavigationTool/Columns/NavigationToolMarkerVisibilityColumn.h
--- a/Engine/Source/Editor/Sequencer/Public/Tools/NavigationTool/Columns/NavigationToolMarkerVisibilityColumn.h
+++ b/Engine/Source/Editor/Sequencer/Public/Tools/NavigationTool/Columns/NavigationToolMarkerVisibilityColumn.h
@@ -14,7 +14,11 @@ class FNavigationToolMarkerVisibilityColumn : public FNavigationToolColumn
 public:
 	UE_NAVIGATIONTOOL_INHERITS(FNavigationToolMarkerVisibilityColumn, FNavigationToolColumn);
 
+	static FName StaticColumnId() { return TEXT("MarkerVisibility"); }
+
+protected:
 	//~ Begin INavigationToolColumn
+	virtual FName GetColumnId() const override { return StaticColumnId(); }
 	SEQUENCER_API virtual FText GetColumnDisplayNameText() const override;
 	SEQUENCER_API virtual const FSlateBrush* GetIconBrush() const override;
 	SEQUENCER_API virtual SHeaderRow::FColumn::FArguments ConstructHeaderRowColumn(const TSharedRef<INavigationToolView>& InToolView, const float InFillSize) override;
diff --git a/Engine/Source/Editor/Sequencer/Public/Tools/NavigationTool/Columns/NavigationToolOutTimeColumn.h b/Engine/Source/Editor/Sequencer/Public/Tools/NavigationTool/Columns/NavigationToolOutTimeColumn.h
--- a/Engine/Source/Editor/Sequencer/Public/Tools/NavigationTool/Columns/NavigationToolOutTimeColumn.h
+++ b/Engine/Source/Editor/Sequencer/Public/Tools/NavigationTool/Columns/NavigationToolOutTimeColumn.h
@@ -14,8 +14,11 @@ class FNavigationToolOutTimeColumn : public FNavigationToolColumn
 public:
 	UE_NAVIGATIONTOOL_INHERITS(FNavigationToolOutTimeColumn, FNavigationToolColumn);
 
+	static FName StaticColumnId() { return TEXT("OutTime"); }
+
 protected:
 	//~ Begin INavigationToolColumn
+	virtual FName GetColumnId() const override { return StaticColumnId(); }
 	SEQUENCER_API virtual FText GetColumnDisplayNameText() const override;
 	SEQUENCER_API virtual const FSlateBrush* GetIconBrush() const override;
 	virtual bool ShouldShowColumnByDefault() const override { return true; }
diff --git a/Engine/Source/Editor/Sequencer/Public/Tools/NavigationTool/Columns/NavigationToolPlayheadColumn.h b/Engine/Source/Editor/Sequencer/Public/Tools/NavigationTool/Columns/NavigationToolPlayheadColumn.h
--- a/Engine/Source/Editor/Sequencer/Public/Tools/NavigationTool/Columns/NavigationToolPlayheadColumn.h
+++ b/Engine/Source/Editor/Sequencer/Public/Tools/NavigationTool/Columns/NavigationToolPlayheadColumn.h
@@ -14,7 +14,11 @@ class FNavigationToolPlayheadColumn : public FNavigationToolColumn
 public:
 	UE_NAVIGATIONTOOL_INHERITS(FNavigationToolPlayheadColumn, FNavigationToolColumn);
 
+	static FName StaticColumnId() { return TEXT("Playhead"); }
+
+protected:
 	//~ Begin INavigationToolColumn
+	virtual FName GetColumnId() const override { return StaticColumnId(); }
 	SEQUENCER_API virtual FText GetColumnDisplayNameText() const override;
 	SEQUENCER_API virtual const FSlateBrush* GetIconBrush() const override;
 	SEQUENCER_API virtual SHeaderRow::FColumn::FArguments ConstructHeaderRowColumn(const TSharedRef<INavigationToolView>& InToolView, const float InFillSize) override;
diff --git a/Engine/Source/Editor/Sequencer/Public/Tools/NavigationTool/Columns/NavigationToolRevisionControlColumn.h b/Engine/Source/Editor/Sequencer/Public/Tools/NavigationTool/Columns/NavigationToolRevisionControlColumn.h
--- a/Engine/Source/Editor/Sequencer/Public/Tools/NavigationTool/Columns/NavigationToolRevisionControlColumn.h
+++ b/Engine/Source/Editor/Sequencer/Public/Tools/NavigationTool/Columns/NavigationToolRevisionControlColumn.h
@@ -14,7 +14,11 @@ class FNavigationToolRevisionControlColumn : public FNavigationToolColumn
 public:
 	UE_NAVIGATIONTOOL_INHERITS(FNavigationToolRevisionControlColumn, FNavigationToolColumn);
 
+	static FName StaticColumnId() { return TEXT("RevisionControl"); }
+
+protected:
 	//~ Begin INavigationToolColumn
+	virtual FName GetColumnId() const override { return StaticColumnId(); }
 	SEQUENCER_API virtual FText GetColumnDisplayNameText() const override;
 	SEQUENCER_API virtual const FSlateBrush* GetIconBrush() const override;
 	SEQUENCER_API virtual SHeaderRow::FColumn::FArguments ConstructHeaderRowColumn(const TSharedRef<INavigationToolView>& InToolView, const float InFillSize) override;
diff --git a/Engine/Source/Editor/Sequencer/Public/Tools/NavigationTool/Columns/NavigationToolStartFrameOffsetColumn.h b/Engine/Source/Editor/Sequencer/Public/Tools/NavigationTool/Columns/NavigationToolStartFrameOffsetColumn.h
--- a/Engine/Source/Editor/Sequencer/Public/Tools/NavigationTool/Columns/NavigationToolStartFrameOffsetColumn.h
+++ b/Engine/Source/Editor/Sequencer/Public/Tools/NavigationTool/Columns/NavigationToolStartFrameOffsetColumn.h
@@ -14,8 +14,11 @@ class FNavigationToolStartFrameOffsetColumn : public FNavigationToolColumn
 public:
 	UE_NAVIGATIONTOOL_INHERITS(FNavigationToolStartFrameOffsetColumn, FNavigationToolColumn);
 
+	static FName StaticColumnId() { return TEXT("StartFrameOffset"); }
+
 protected:
 	//~ Begin INavigationToolColumn
+	virtual FName GetColumnId() const override { return StaticColumnId(); }
 	SEQUENCER_API virtual FText GetColumnDisplayNameText() const override;
 	SEQUENCER_API virtual const FSlateBrush* GetIconBrush() const override;
 	virtual bool ShouldShowColumnByDefault() const override { return false; }
diff --git a/Engine/Source/Editor/Sequencer/Public/Tools/NavigationTool/Columns/NavigationToolTakeColumn.h b/Engine/Source/Editor/Sequencer/Public/Tools/NavigationTool/Columns/NavigationToolTakeColumn.h
--- a/Engine/Source/Editor/Sequencer/Public/Tools/NavigationTool/Columns/NavigationToolTakeColumn.h
+++ b/Engine/Source/Editor/Sequencer/Public/Tools/NavigationTool/Columns/NavigationToolTakeColumn.h
@@ -14,8 +14,11 @@ class FNavigationToolTakeColumn : public FNavigationToolColumn
 public:
 	UE_NAVIGATIONTOOL_INHERITS(FNavigationToolTakeColumn, FNavigationToolColumn);
 
+	static FName StaticColumnId() { return TEXT("ShotTake"); }
+
 protected:
 	//~ Begin INavigationToolColumn
+	virtual FName GetColumnId() const override { return StaticColumnId(); }
 	SEQUENCER_API virtual FText GetColumnDisplayNameText() const override;
 	SEQUENCER_API virtual const FSlateBrush* GetIconBrush() const override;
 	virtual bool ShouldShowColumnByDefault() const override { return false; }
