[Sequencer] Navigation Tool:
- Hide filter bar config property
- Remve redundant "NavigationTool" from config property name

diff --git a/Engine/Source/Editor/Sequencer/Private/Tools/NavigationTool/NavigationToolSettings.cpp b/Engine/Source/Editor/Sequencer/Private/Tools/NavigationTool/NavigationToolSettings.cpp
--- a/Engine/Source/Editor/Sequencer/Private/Tools/NavigationTool/NavigationToolSettings.cpp
+++ b/Engine/Source/Editor/Sequencer/Private/Tools/NavigationTool/NavigationToolSettings.cpp
@@ -163,7 +163,7 @@ void UNavigationToolSettings::SetBuiltInFilterEnabled(const FName InFilterName,
 
 FSequencerFilterBarConfig& UNavigationToolSettings::FindOrAddFilterBar(const FName InIdentifier, const bool bInSaveConfig)
 {
-	FSequencerFilterBarConfig& FilterBarSettings = NavigationToolFilterBars.FindOrAdd(InIdentifier);
+	FSequencerFilterBarConfig& FilterBarSettings = FilterBars.FindOrAdd(InIdentifier);
 
 	if (bInSaveConfig)
 	{
@@ -175,12 +175,12 @@ FSequencerFilterBarConfig& UNavigationToolSettings::FindOrAddFilterBar(const FNa
 
 FSequencerFilterBarConfig* UNavigationToolSettings::FindFilterBar(const FName InIdentifier)
 {
-	return NavigationToolFilterBars.Find(InIdentifier);
+	return FilterBars.Find(InIdentifier);
 }
 
 bool UNavigationToolSettings::RemoveFilterBar(const FName InIdentifier)
 {
-	const int32 RemovedCount = NavigationToolFilterBars.Remove(InIdentifier);
+	const int32 RemovedCount = FilterBars.Remove(InIdentifier);
 	const bool bRemoveSuccess = RemovedCount > 0;
 
 	if (bRemoveSuccess)
diff --git a/Engine/Source/Editor/Sequencer/Public/Tools/NavigationTool/NavigationToolSettings.h b/Engine/Source/Editor/Sequencer/Public/Tools/NavigationTool/NavigationToolSettings.h
--- a/Engine/Source/Editor/Sequencer/Public/Tools/NavigationTool/NavigationToolSettings.h
+++ b/Engine/Source/Editor/Sequencer/Public/Tools/NavigationTool/NavigationToolSettings.h
@@ -149,33 +149,33 @@ private:
 
 	/** Saved settings for each unique filter bar instance mapped by instance identifier */
 	UPROPERTY(config, EditAnywhere, Category = "Filtering")
-	TMap<FName, FSequencerFilterBarConfig> NavigationToolFilterBars;
+	TMap<FName, FSequencerFilterBarConfig> FilterBars;
 
 	/** Automatically expand tracks that pass filters */
-	UPROPERTY(config, EditAnywhere, Category = "Filtering")
+	UPROPERTY(Config, EditAnywhere, Category = "Filtering")
 	bool bAutoExpandNodesOnFilterPass = false;
 	
 	/** Display the filter menu categories as submenus instead of sections */
-	UPROPERTY(config, EditAnywhere, Category = "Filtering")
+	UPROPERTY(Config, EditAnywhere, Category = "Filtering")
 	bool bUseFilterSubmenusForCategories = false;
 
 	/** Last saved visibility of the filter bar to restore after closed */
-	UPROPERTY(config, EditAnywhere, Category = "Filtering")
+	UPROPERTY(Config, EditAnywhere, Category = "Filtering")
 	bool bFilterBarVisible = false;
 
 	/** Last saved layout orientation of the filter bar to restore after closed */
-	UPROPERTY(config, EditAnywhere, Category = "Filtering")
+	UPROPERTY(Config, EditAnywhere, Category = "Filtering")
 	EFilterBarLayout LastFilterBarLayout = EFilterBarLayout::Horizontal;
 
 	/** Last saved size of the filter bar to restore after closed */
-	UPROPERTY(config)
+	UPROPERTY(Config)
 	float LastFilterBarSizeCoefficient = 0.05f;
 
 	/** If true, syncs selections in Sequencer to Navigation Tool */
-	UPROPERTY(config, EditAnywhere, Category = "Editor UX")
+	UPROPERTY(Config, EditAnywhere, Category = "Editor UX")
 	bool bSyncSelectionToNavigationTool = true;
 
 	/** If true, syncs selections in Navigation Tool to Sequencer */
-	UPROPERTY(config, EditAnywhere, Category = "Editor UX")
+	UPROPERTY(Config, EditAnywhere, Category = "Editor UX")
 	bool bSyncSelectionToSequencer = true;
 };
