[Sequencer] Navigation Tool:
- Remove old forward declarations and includes
- Changes needed to move Navigation Tool to its own plugin
- General cleanup

diff --git a/Engine/Plugins/Experimental/Avalanche/Source/AvalancheSequencer/Private/NavigationTool/Columns/AvaNavigationToolStatusColumn.cpp b/Engine/Plugins/Experimental/Avalanche/Source/AvalancheSequencer/Private/NavigationTool/Columns/AvaNavigationToolStatusColumn.cpp
--- a/Engine/Plugins/Experimental/Avalanche/Source/AvalancheSequencer/Private/NavigationTool/Columns/AvaNavigationToolStatusColumn.cpp
+++ b/Engine/Plugins/Experimental/Avalanche/Source/AvalancheSequencer/Private/NavigationTool/Columns/AvaNavigationToolStatusColumn.cpp
@@ -38,8 +38,8 @@ SHeaderRow::FColumn::FArguments FAvaNavigationToolStatusColumn::ConstructHeaderR
 		.OnGetMenuContent(InToolView, &INavigationToolView::GetColumnMenuContent, ColumnId);
 }
 
-TSharedRef<SWidget> FAvaNavigationToolStatusColumn::ConstructRowWidget(const FNavigationToolItemPtr& InItem
-	, const TSharedRef<INavigationToolView>& InToolView
+TSharedRef<SWidget> FAvaNavigationToolStatusColumn::ConstructRowWidget(const FNavigationToolItemRef& InItem
+	, const TSharedRef<INavigationToolView>& InView
 	, const TSharedRef<SNavigationToolTreeRow>& InRow)
 {
 	if (!InItem->IsA<FNavigationToolAvaSequence>())
@@ -47,7 +47,7 @@ TSharedRef<SWidget> FAvaNavigationToolStatusColumn::ConstructRowWidget(const FNa
 		return SNullWidget::NullWidget;
 	}
 
-	const TSharedPtr<INavigationTool> OwnerTool = InToolView->GetOwnerTool();
+	const TSharedPtr<INavigationTool> OwnerTool = InView->GetOwnerTool();
 	if (!OwnerTool.IsValid())
 	{
 		return SNullWidget::NullWidget;
@@ -59,7 +59,7 @@ TSharedRef<SWidget> FAvaNavigationToolStatusColumn::ConstructRowWidget(const FNa
 		return SNullWidget::NullWidget;
 	}
 
-	return SNew(SAvaNavigationToolStatus, InItem, InToolView, InRow, Sequencer.ToSharedRef());
+	return SNew(SAvaNavigationToolStatus, InItem, InView, InRow, Sequencer.ToSharedRef());
 }
 
 } // UE::Sequencer::NavigationTool
diff --git a/Engine/Plugins/Experimental/Avalanche/Source/AvalancheSequencer/Private/NavigationTool/Columns/AvaNavigationToolStatusColumn.h b/Engine/Plugins/Experimental/Avalanche/Source/AvalancheSequencer/Private/NavigationTool/Columns/AvaNavigationToolStatusColumn.h
--- a/Engine/Plugins/Experimental/Avalanche/Source/AvalancheSequencer/Private/NavigationTool/Columns/AvaNavigationToolStatusColumn.h
+++ b/Engine/Plugins/Experimental/Avalanche/Source/AvalancheSequencer/Private/NavigationTool/Columns/AvaNavigationToolStatusColumn.h
@@ -22,8 +22,8 @@ protected:
 	virtual bool ShouldShowColumnByDefault() const override { return false; }
 	virtual float GetFillWidth() const override { return 10.f; }
 	virtual SHeaderRow::FColumn::FArguments ConstructHeaderRowColumn(const TSharedRef<INavigationToolView>& InToolView, const float InFillSize) override;
-	virtual TSharedRef<SWidget> ConstructRowWidget(const FNavigationToolItemPtr& InItem
-		, const TSharedRef<INavigationToolView>& InToolView
+	virtual TSharedRef<SWidget> ConstructRowWidget(const FNavigationToolItemRef& InItem
+		, const TSharedRef<INavigationToolView>& InView
 		, const TSharedRef<SNavigationToolTreeRow>& InRow) override;
 	//~ End INavigationToolColumn
 };
diff --git a/Engine/Source/Editor/Sequencer/Private/Tools/NavigationTool/Columns/NavigationToolColorColumn.cpp b/Engine/Source/Editor/Sequencer/Private/Tools/NavigationTool/Columns/NavigationToolColorColumn.cpp
--- a/Engine/Source/Editor/Sequencer/Private/Tools/NavigationTool/Columns/NavigationToolColorColumn.cpp
+++ b/Engine/Source/Editor/Sequencer/Private/Tools/NavigationTool/Columns/NavigationToolColorColumn.cpp
@@ -34,11 +34,11 @@ SHeaderRow::FColumn::FArguments FNavigationToolColorColumn::ConstructHeaderRowCo
 		];
 }
 
-TSharedRef<SWidget> FNavigationToolColorColumn::ConstructRowWidget(const FNavigationToolItemPtr& InItem
-	, const TSharedRef<INavigationToolView>& InToolView
+TSharedRef<SWidget> FNavigationToolColorColumn::ConstructRowWidget(const FNavigationToolItemRef& InItem
+	, const TSharedRef<INavigationToolView>& InView
 	, const TSharedRef<SNavigationToolTreeRow>& InRow)
 {
-	return SNew(SNavigationToolColor, InItem, InToolView, InRow);
+	return SNew(SNavigationToolColor, InItem, InView, InRow);
 }
 
 }
diff --git a/Engine/Source/Editor/Sequencer/Private/Tools/NavigationTool/Columns/NavigationToolCommentColumn.cpp b/Engine/Source/Editor/Sequencer/Private/Tools/NavigationTool/Columns/NavigationToolCommentColumn.cpp
--- a/Engine/Source/Editor/Sequencer/Private/Tools/NavigationTool/Columns/NavigationToolCommentColumn.cpp
+++ b/Engine/Source/Editor/Sequencer/Private/Tools/NavigationTool/Columns/NavigationToolCommentColumn.cpp
@@ -33,11 +33,11 @@ SHeaderRow::FColumn::FArguments FNavigationToolCommentColumn::ConstructHeaderRow
 		.OnGetMenuContent(InToolView, &INavigationToolView::GetColumnMenuContent, ColumnId);
 }
 
-TSharedRef<SWidget> FNavigationToolCommentColumn::ConstructRowWidget(const FNavigationToolItemPtr& InItem
-	, const TSharedRef<INavigationToolView>& InToolView
+TSharedRef<SWidget> FNavigationToolCommentColumn::ConstructRowWidget(const FNavigationToolItemRef& InItem
+	, const TSharedRef<INavigationToolView>& InView
 	, const TSharedRef<SNavigationToolTreeRow>& InRow)
 {
-	return SNew(SNavigationToolComment, InItem, InToolView, InRow);
+	return SNew(SNavigationToolComment, InItem, InView, InRow);
 }
 
 }
diff --git a/Engine/Source/Editor/Sequencer/Private/Tools/NavigationTool/Columns/NavigationToolDeactiveStateColumn.cpp b/Engine/Source/Editor/Sequencer/Private/Tools/NavigationTool/Columns/NavigationToolDeactiveStateColumn.cpp
--- a/Engine/Source/Editor/Sequencer/Private/Tools/NavigationTool/Columns/NavigationToolDeactiveStateColumn.cpp
+++ b/Engine/Source/Editor/Sequencer/Private/Tools/NavigationTool/Columns/NavigationToolDeactiveStateColumn.cpp
@@ -46,8 +46,8 @@ SHeaderRow::FColumn::FArguments FNavigationToolDeactiveStateColumn::ConstructHea
 		];
 }
 
-TSharedRef<SWidget> FNavigationToolDeactiveStateColumn::ConstructRowWidget(const FNavigationToolItemPtr& InItem
-	, const TSharedRef<INavigationToolView>& InToolView
+TSharedRef<SWidget> FNavigationToolDeactiveStateColumn::ConstructRowWidget(const FNavigationToolItemRef& InItem
+	, const TSharedRef<INavigationToolView>& InView
 	, const TSharedRef<SNavigationToolTreeRow>& InRow)
 {
 	if (!InItem->IsA<ISequenceInactivableExtension>())
@@ -60,7 +60,7 @@ TSharedRef<SWidget> FNavigationToolDeactiveStateColumn::ConstructRowWidget(const
 		.AutoWidth()
 		.VAlign(VAlign_Center)
 		[
-			SNew(SNavigationToolDeactiveState, SharedThis(this), InToolView, InItem, InRow)
+			SNew(SNavigationToolDeactiveState, SharedThis(this), InItem, InView, InRow)
 		];
 }
 
diff --git a/Engine/Source/Editor/Sequencer/Private/Tools/NavigationTool/Columns/NavigationToolHBiasColumn.cpp b/Engine/Source/Editor/Sequencer/Private/Tools/NavigationTool/Columns/NavigationToolHBiasColumn.cpp
--- a/Engine/Source/Editor/Sequencer/Private/Tools/NavigationTool/Columns/NavigationToolHBiasColumn.cpp
+++ b/Engine/Source/Editor/Sequencer/Private/Tools/NavigationTool/Columns/NavigationToolHBiasColumn.cpp
@@ -35,8 +35,8 @@ SHeaderRow::FColumn::FArguments FNavigationToolHBiasColumn::ConstructHeaderRowCo
 		.OnGetMenuContent(InToolView, &INavigationToolView::GetColumnMenuContent, ColumnId);
 }
 
-TSharedRef<SWidget> FNavigationToolHBiasColumn::ConstructRowWidget(const FNavigationToolItemPtr& InItem
-	, const TSharedRef<INavigationToolView>& InToolView
+TSharedRef<SWidget> FNavigationToolHBiasColumn::ConstructRowWidget(const FNavigationToolItemRef& InItem
+	, const TSharedRef<INavigationToolView>& InView
 	, const TSharedRef<SNavigationToolTreeRow>& InRow)
 {
 	if (!InItem->IsA<FNavigationToolSequence>())
@@ -44,7 +44,7 @@ TSharedRef<SWidget> FNavigationToolHBiasColumn::ConstructRowWidget(const FNaviga
 		return SNullWidget::NullWidget;
 	}
 
-	return SNew(SNavigationToolHBias, InItem, InToolView, InRow);
+	return SNew(SNavigationToolHBias, InItem, InView, InRow);
 }
 
 }
diff --git a/Engine/Source/Editor/Sequencer/Private/Tools/NavigationTool/Columns/NavigationToolInTimeColumn.cpp b/Engine/Source/Editor/Sequencer/Private/Tools/NavigationTool/Columns/NavigationToolInTimeColumn.cpp
--- a/Engine/Source/Editor/Sequencer/Private/Tools/NavigationTool/Columns/NavigationToolInTimeColumn.cpp
+++ b/Engine/Source/Editor/Sequencer/Private/Tools/NavigationTool/Columns/NavigationToolInTimeColumn.cpp
@@ -1,6 +1,7 @@
 // Copyright Epic Games, Inc. All Rights Reserved.
 
 #include "Tools/NavigationTool/Columns/NavigationToolInTimeColumn.h"
+#include "Tools/NavigationTool/Extensions/IInTimeExtension.h"
 #include "Tools/NavigationTool/Items/NavigationToolSequence.h"
 #include "Tools/NavigationTool/NavigationToolView.h"
 #include "Tools/NavigationTool/Widgets/Columns/SNavigationToolInTime.h"
@@ -35,11 +36,16 @@ SHeaderRow::FColumn::FArguments FNavigationToolInTimeColumn::ConstructHeaderRowC
 		.OnGetMenuContent(InToolView, &INavigationToolView::GetColumnMenuContent, ColumnId);
 }
 
-TSharedRef<SWidget> FNavigationToolInTimeColumn::ConstructRowWidget(const FNavigationToolItemPtr& InItem
-	, const TSharedRef<INavigationToolView>& InToolView
+TSharedRef<SWidget> FNavigationToolInTimeColumn::ConstructRowWidget(const FNavigationToolItemRef& InItem
+	, const TSharedRef<INavigationToolView>& InView
 	, const TSharedRef<SNavigationToolTreeRow>& InRow)
 {
-	return SNew(SNavigationToolInTime, InItem, InToolView, InRow);
+	if (!InItem->IsA<IInTimeExtension>())
+	{
+		return SNullWidget::NullWidget;
+	}
+
+	return SNew(SNavigationToolInTime, InItem, InView, InRow);
 }
 
 }
diff --git a/Engine/Source/Editor/Sequencer/Private/Tools/NavigationTool/Columns/NavigationToolItemsColumn.cpp b/Engine/Source/Editor/Sequencer/Private/Tools/NavigationTool/Columns/NavigationToolItemsColumn.cpp
--- a/Engine/Source/Editor/Sequencer/Private/Tools/NavigationTool/Columns/NavigationToolItemsColumn.cpp
+++ b/Engine/Source/Editor/Sequencer/Private/Tools/NavigationTool/Columns/NavigationToolItemsColumn.cpp
@@ -34,11 +34,11 @@ SHeaderRow::FColumn::FArguments FNavigationToolItemsColumn::ConstructHeaderRowCo
 		.OnGetMenuContent(InToolView, &INavigationToolView::GetColumnMenuContent, ColumnId);
 }
 
-TSharedRef<SWidget> FNavigationToolItemsColumn::ConstructRowWidget(const FNavigationToolItemPtr& InItem
-	, const TSharedRef<INavigationToolView>& InToolView
+TSharedRef<SWidget> FNavigationToolItemsColumn::ConstructRowWidget(const FNavigationToolItemRef& InItem
+	, const TSharedRef<INavigationToolView>& InView
 	, const TSharedRef<SNavigationToolTreeRow>& InRow)
 {
-	return SNew(SNavigationToolItemList, InItem, InToolView, InRow);
+	return SNew(SNavigationToolItemList, InItem, InView, InRow);
 }
 
 }
diff --git a/Engine/Source/Editor/Sequencer/Private/Tools/NavigationTool/Columns/NavigationToolLabelColumn.cpp b/Engine/Source/Editor/Sequencer/Private/Tools/NavigationTool/Columns/NavigationToolLabelColumn.cpp
--- a/Engine/Source/Editor/Sequencer/Private/Tools/NavigationTool/Columns/NavigationToolLabelColumn.cpp
+++ b/Engine/Source/Editor/Sequencer/Private/Tools/NavigationTool/Columns/NavigationToolLabelColumn.cpp
@@ -31,8 +31,8 @@ SHeaderRow::FColumn::FArguments FNavigationToolLabelColumn::ConstructHeaderRowCo
 		.OnGetMenuContent(InToolView, &INavigationToolView::GetColumnMenuContent, ColumnId);
 }
 
-TSharedRef<SWidget> FNavigationToolLabelColumn::ConstructRowWidget(const FNavigationToolItemPtr& InItem
-	, const TSharedRef<INavigationToolView>& InToolView
+TSharedRef<SWidget> FNavigationToolLabelColumn::ConstructRowWidget(const FNavigationToolItemRef& InItem
+	, const TSharedRef<INavigationToolView>& InView
 	, const TSharedRef<SNavigationToolTreeRow>& InRow)
 {
 	return InItem->GenerateLabelWidget(InRow);
diff --git a/Engine/Source/Editor/Sequencer/Private/Tools/NavigationTool/Columns/NavigationToolLengthColumn.cpp b/Engine/Source/Editor/Sequencer/Private/Tools/NavigationTool/Columns/NavigationToolLengthColumn.cpp
--- a/Engine/Source/Editor/Sequencer/Private/Tools/NavigationTool/Columns/NavigationToolLengthColumn.cpp
+++ b/Engine/Source/Editor/Sequencer/Private/Tools/NavigationTool/Columns/NavigationToolLengthColumn.cpp
@@ -34,11 +34,11 @@ SHeaderRow::FColumn::FArguments FNavigationToolLengthColumn::ConstructHeaderRowC
 		.OnGetMenuContent(InToolView, &INavigationToolView::GetColumnMenuContent, ColumnId);
 }
 
-TSharedRef<SWidget> FNavigationToolLengthColumn::ConstructRowWidget(const FNavigationToolItemPtr& InItem
-	, const TSharedRef<INavigationToolView>& InToolView
+TSharedRef<SWidget> FNavigationToolLengthColumn::ConstructRowWidget(const FNavigationToolItemRef& InItem
+	, const TSharedRef<INavigationToolView>& InView
 	, const TSharedRef<SNavigationToolTreeRow>& InRow)
 {
-	return SNew(SNavigationToolLength, InItem, InToolView, InRow);
+	return SNew(SNavigationToolLength, InItem, InView, InRow);
 }
 
 }
diff --git a/Engine/Source/Editor/Sequencer/Private/Tools/NavigationTool/Columns/NavigationToolLockColumn.cpp b/Engine/Source/Editor/Sequencer/Private/Tools/NavigationTool/Columns/NavigationToolLockColumn.cpp
--- a/Engine/Source/Editor/Sequencer/Private/Tools/NavigationTool/Columns/NavigationToolLockColumn.cpp
+++ b/Engine/Source/Editor/Sequencer/Private/Tools/NavigationTool/Columns/NavigationToolLockColumn.cpp
@@ -43,8 +43,8 @@ SHeaderRow::FColumn::FArguments FNavigationToolLockColumn::ConstructHeaderRowCol
 		];
 }
 
-TSharedRef<SWidget> FNavigationToolLockColumn::ConstructRowWidget(const FNavigationToolItemPtr& InItem
-	, const TSharedRef<INavigationToolView>& InToolView
+TSharedRef<SWidget> FNavigationToolLockColumn::ConstructRowWidget(const FNavigationToolItemRef& InItem
+	, const TSharedRef<INavigationToolView>& InView
 	, const TSharedRef<SNavigationToolTreeRow>& InRow)
 {
 	if (!InItem->IsA<ISequenceLockableExtension>())
@@ -57,7 +57,7 @@ TSharedRef<SWidget> FNavigationToolLockColumn::ConstructRowWidget(const FNavigat
 		.AutoWidth()
 		.VAlign(VAlign_Center)
 		[
-			SNew(SNavigationToolLock, InItem, InToolView, InRow)
+			SNew(SNavigationToolLock, InItem, InView, InRow)
 		];
 }
 
diff --git a/Engine/Source/Editor/Sequencer/Private/Tools/NavigationTool/Columns/NavigationToolMarkerVisibilityColumn.cpp b/Engine/Source/Editor/Sequencer/Private/Tools/NavigationTool/Columns/NavigationToolMarkerVisibilityColumn.cpp
--- a/Engine/Source/Editor/Sequencer/Private/Tools/NavigationTool/Columns/NavigationToolMarkerVisibilityColumn.cpp
+++ b/Engine/Source/Editor/Sequencer/Private/Tools/NavigationTool/Columns/NavigationToolMarkerVisibilityColumn.cpp
@@ -44,8 +44,8 @@ SHeaderRow::FColumn::FArguments FNavigationToolMarkerVisibilityColumn::Construct
 		];
 }
 
-TSharedRef<SWidget> FNavigationToolMarkerVisibilityColumn::ConstructRowWidget(const FNavigationToolItemPtr& InItem
-	, const TSharedRef<INavigationToolView>& InToolView
+TSharedRef<SWidget> FNavigationToolMarkerVisibilityColumn::ConstructRowWidget(const FNavigationToolItemRef& InItem
+	, const TSharedRef<INavigationToolView>& InView
 	, const TSharedRef<SNavigationToolTreeRow>& InRow)
 {
 	if (!InItem->IsA<IMarkerVisibilityExtension>())
@@ -53,7 +53,7 @@ TSharedRef<SWidget> FNavigationToolMarkerVisibilityColumn::ConstructRowWidget(co
 		return SNullWidget::NullWidget;
 	}
 
-	return SNew(SNavigationToolMarkerVisibility, SharedThis(this), InItem, InToolView, InRow);
+	return SNew(SNavigationToolMarkerVisibility, SharedThis(this), InItem, InView, InRow);
 }
 
 }
diff --git a/Engine/Source/Editor/Sequencer/Private/Tools/NavigationTool/Columns/NavigationToolOutTimeColumn.cpp b/Engine/Source/Editor/Sequencer/Private/Tools/NavigationTool/Columns/NavigationToolOutTimeColumn.cpp
--- a/Engine/Source/Editor/Sequencer/Private/Tools/NavigationTool/Columns/NavigationToolOutTimeColumn.cpp
+++ b/Engine/Source/Editor/Sequencer/Private/Tools/NavigationTool/Columns/NavigationToolOutTimeColumn.cpp
@@ -1,7 +1,7 @@
 // Copyright Epic Games, Inc. All Rights Reserved.
 
 #include "Tools/NavigationTool/Columns/NavigationToolOutTimeColumn.h"
-#include "Tools/NavigationTool/Items/NavigationToolSequence.h"
+#include "Tools/NavigationTool/Extensions/IOutTimeExtension.h"
 #include "Tools/NavigationTool/NavigationTool.h"
 #include "Tools/NavigationTool/NavigationToolView.h"
 #include "Tools/NavigationTool/Widgets/Columns/SNavigationToolOutTime.h"
@@ -37,11 +37,16 @@ SHeaderRow::FColumn::FArguments FNavigationToolOutTimeColumn::ConstructHeaderRow
 		.OnGetMenuContent(InToolView, &INavigationToolView::GetColumnMenuContent, ColumnId);
 }
 
-TSharedRef<SWidget> FNavigationToolOutTimeColumn::ConstructRowWidget(const FNavigationToolItemPtr& InItem
-	, const TSharedRef<INavigationToolView>& InToolView
+TSharedRef<SWidget> FNavigationToolOutTimeColumn::ConstructRowWidget(const FNavigationToolItemRef& InItem
+	, const TSharedRef<INavigationToolView>& InView
 	, const TSharedRef<SNavigationToolTreeRow>& InRow)
 {
-	return SNew(SNavigationToolOutTime, InItem, InToolView, InRow);
+	if (!InItem->IsA<IOutTimeExtension>())
+	{
+		return SNullWidget::NullWidget;
+	}
+
+	return SNew(SNavigationToolOutTime, InItem, InView, InRow);
 }
 
 }
diff --git a/Engine/Source/Editor/Sequencer/Private/Tools/NavigationTool/Columns/NavigationToolPlayheadColumn.cpp b/Engine/Source/Editor/Sequencer/Private/Tools/NavigationTool/Columns/NavigationToolPlayheadColumn.cpp
--- a/Engine/Source/Editor/Sequencer/Private/Tools/NavigationTool/Columns/NavigationToolPlayheadColumn.cpp
+++ b/Engine/Source/Editor/Sequencer/Private/Tools/NavigationTool/Columns/NavigationToolPlayheadColumn.cpp
@@ -41,8 +41,8 @@ SHeaderRow::FColumn::FArguments FNavigationToolPlayheadColumn::ConstructHeaderRo
 		];
 }
 
-TSharedRef<SWidget> FNavigationToolPlayheadColumn::ConstructRowWidget(const FNavigationToolItemPtr& InItem
-	, const TSharedRef<INavigationToolView>& InToolView
+TSharedRef<SWidget> FNavigationToolPlayheadColumn::ConstructRowWidget(const FNavigationToolItemRef& InItem
+	, const TSharedRef<INavigationToolView>& InView
 	, const TSharedRef<SNavigationToolTreeRow>& InRow)
 {
 	if (!InItem->IsA<IPlayheadExtension>())
@@ -50,7 +50,7 @@ TSharedRef<SWidget> FNavigationToolPlayheadColumn::ConstructRowWidget(const FNav
 		return SNullWidget::NullWidget;
 	}
 
-	return SNew(SNavigationToolPlayhead, SharedThis(this), InItem, InToolView, InRow);
+	return SNew(SNavigationToolPlayhead, SharedThis(this), InItem, InView, InRow);
 }
 
 }
diff --git a/Engine/Source/Editor/Sequencer/Private/Tools/NavigationTool/Columns/NavigationToolRevisionControlColumn.cpp b/Engine/Source/Editor/Sequencer/Private/Tools/NavigationTool/Columns/NavigationToolRevisionControlColumn.cpp
--- a/Engine/Source/Editor/Sequencer/Private/Tools/NavigationTool/Columns/NavigationToolRevisionControlColumn.cpp
+++ b/Engine/Source/Editor/Sequencer/Private/Tools/NavigationTool/Columns/NavigationToolRevisionControlColumn.cpp
@@ -46,8 +46,8 @@ SHeaderRow::FColumn::FArguments FNavigationToolRevisionControlColumn::ConstructH
 		];
 }
 
-TSharedRef<SWidget> FNavigationToolRevisionControlColumn::ConstructRowWidget(const FNavigationToolItemPtr& InItem
-	, const TSharedRef<INavigationToolView>& InToolView
+TSharedRef<SWidget> FNavigationToolRevisionControlColumn::ConstructRowWidget(const FNavigationToolItemRef& InItem
+	, const TSharedRef<INavigationToolView>& InView
 	, const TSharedRef<SNavigationToolTreeRow>& InRow)
 {
 	if (!InItem->IsA<FNavigationToolSequence>())
@@ -55,7 +55,7 @@ TSharedRef<SWidget> FNavigationToolRevisionControlColumn::ConstructRowWidget(con
 		return SNullWidget::NullWidget;
 	}
 
-	return SNew(SNavigationToolRevisionControl, SharedThis(this), InItem, InToolView, InRow);
+	return SNew(SNavigationToolRevisionControl, SharedThis(this), InItem, InView, InRow);
 }
 
 }
diff --git a/Engine/Source/Editor/Sequencer/Private/Tools/NavigationTool/Columns/NavigationToolStartFrameOffsetColumn.cpp b/Engine/Source/Editor/Sequencer/Private/Tools/NavigationTool/Columns/NavigationToolStartFrameOffsetColumn.cpp
--- a/Engine/Source/Editor/Sequencer/Private/Tools/NavigationTool/Columns/NavigationToolStartFrameOffsetColumn.cpp
+++ b/Engine/Source/Editor/Sequencer/Private/Tools/NavigationTool/Columns/NavigationToolStartFrameOffsetColumn.cpp
@@ -34,11 +34,11 @@ SHeaderRow::FColumn::FArguments FNavigationToolStartFrameOffsetColumn::Construct
 		.OnGetMenuContent(InToolView, &INavigationToolView::GetColumnMenuContent, ColumnId);
 }
 
-TSharedRef<SWidget> FNavigationToolStartFrameOffsetColumn::ConstructRowWidget(const FNavigationToolItemPtr& InItem
-	, const TSharedRef<INavigationToolView>& InToolView
+TSharedRef<SWidget> FNavigationToolStartFrameOffsetColumn::ConstructRowWidget(const FNavigationToolItemRef& InItem
+	, const TSharedRef<INavigationToolView>& InView
 	, const TSharedRef<SNavigationToolTreeRow>& InRow)
 {
-	return SNew(SNavigationToolStartFrameOffset, InItem, InToolView, InRow);
+	return SNew(SNavigationToolStartFrameOffset, InItem, InView, InRow);
 }
 
 }
diff --git a/Engine/Source/Editor/Sequencer/Private/Tools/NavigationTool/Columns/NavigationToolTakeColumn.cpp b/Engine/Source/Editor/Sequencer/Private/Tools/NavigationTool/Columns/NavigationToolTakeColumn.cpp
--- a/Engine/Source/Editor/Sequencer/Private/Tools/NavigationTool/Columns/NavigationToolTakeColumn.cpp
+++ b/Engine/Source/Editor/Sequencer/Private/Tools/NavigationTool/Columns/NavigationToolTakeColumn.cpp
@@ -36,8 +36,8 @@ SHeaderRow::FColumn::FArguments FNavigationToolTakeColumn::ConstructHeaderRowCol
 		.OnGetMenuContent(InToolView, &INavigationToolView::GetColumnMenuContent, ColumnId);
 }
 
-TSharedRef<SWidget> FNavigationToolTakeColumn::ConstructRowWidget(const FNavigationToolItemPtr& InItem
-	, const TSharedRef<INavigationToolView>& InToolView
+TSharedRef<SWidget> FNavigationToolTakeColumn::ConstructRowWidget(const FNavigationToolItemRef& InItem
+	, const TSharedRef<INavigationToolView>& InView
 	, const TSharedRef<SNavigationToolTreeRow>& InRow)
 {
 	if (!InItem->IsA<FNavigationToolSequence>())
@@ -45,7 +45,7 @@ TSharedRef<SWidget> FNavigationToolTakeColumn::ConstructRowWidget(const FNavigat
 		return SNullWidget::NullWidget;
 	}
 
-	return SNew(SNavigationToolTake, InItem, InToolView, InRow);
+	return SNew(SNavigationToolTake, InItem, InView, InRow);
 }
 
 }
diff --git a/Engine/Source/Editor/Sequencer/Private/Tools/NavigationTool/Filters/Filters/NavigationToolFilter_Dirty.cpp b/Engine/Source/Editor/Sequencer/Private/Tools/NavigationTool/Filters/Filters/NavigationToolFilter_Dirty.cpp
--- a/Engine/Source/Editor/Sequencer/Private/Tools/NavigationTool/Filters/Filters/NavigationToolFilter_Dirty.cpp
+++ b/Engine/Source/Editor/Sequencer/Private/Tools/NavigationTool/Filters/Filters/NavigationToolFilter_Dirty.cpp
@@ -1,7 +1,6 @@
 ﻿// Copyright Epic Games, Inc. All Rights Reserved.
 
 #include "Tools/NavigationTool/Filters/Filters/NavigationToolFilter_Dirty.h"
-#include "Sequencer.h"
 #include "Tools/NavigationTool/Filters/NavigationToolFilterCommands.h"
 #include "Tools/NavigationTool/Items/NavigationToolBinding.h"
 #include "Tools/NavigationTool/Items/NavigationToolSequence.h"
diff --git a/Engine/Source/Editor/Sequencer/Private/Tools/NavigationTool/Filters/Filters/NavigationToolFilter_Unbound.cpp b/Engine/Source/Editor/Sequencer/Private/Tools/NavigationTool/Filters/Filters/NavigationToolFilter_Unbound.cpp
--- a/Engine/Source/Editor/Sequencer/Private/Tools/NavigationTool/Filters/Filters/NavigationToolFilter_Unbound.cpp
+++ b/Engine/Source/Editor/Sequencer/Private/Tools/NavigationTool/Filters/Filters/NavigationToolFilter_Unbound.cpp
@@ -1,7 +1,6 @@
 ﻿// Copyright Epic Games, Inc. All Rights Reserved.
 
 #include "Tools/NavigationTool/Filters/Filters/NavigationToolFilter_Unbound.h"
-#include "Sequencer.h"
 #include "Tools/NavigationTool/Filters/NavigationToolFilterCommands.h"
 #include "Tools/NavigationTool/Items/NavigationToolBinding.h"
 
diff --git a/Engine/Source/Editor/Sequencer/Private/Tools/NavigationTool/Filters/NavigationToolFilterBar.cpp b/Engine/Source/Editor/Sequencer/Private/Tools/NavigationTool/Filters/NavigationToolFilterBar.cpp
--- a/Engine/Source/Editor/Sequencer/Private/Tools/NavigationTool/Filters/NavigationToolFilterBar.cpp
+++ b/Engine/Source/Editor/Sequencer/Private/Tools/NavigationTool/Filters/NavigationToolFilterBar.cpp
@@ -5,7 +5,8 @@
 #include "Filters/Widgets/SFilterExpressionHelpDialog.h"
 #include "Framework/Application/SlateApplication.h"
 #include "Framework/MultiBox/SToolBarButtonBlock.h"
-#include "Sequencer.h"
+#include "ISequencer.h"
+#include "SequencerSettings.h"
 #include "Stats/StatsMisc.h"
 #include "Tools/NavigationTool/Filters/Filters/NavigationToolFilter_CustomText.h"
 #include "Tools/NavigationTool/Filters/Filters/NavigationToolFilter_Dirty.h"
@@ -351,7 +352,7 @@ FText FNavigationToolFilterBar::GetFilterErrorText() const
 
 void FNavigationToolFilterBar::RequestFilterUpdate()
 {
-	const TSharedPtr<FSequencer> Sequencer = StaticCastSharedPtr<FSequencer>(Tool.GetSequencer());
+	const TSharedPtr<ISequencer> Sequencer = Tool.GetSequencer();
 	if (!Sequencer.IsValid())
 	{
 		return;
diff --git a/Engine/Source/Editor/Sequencer/Private/Tools/NavigationTool/Filters/TextExpressions/NavigationToolFilterTextExpressionContext.cpp b/Engine/Source/Editor/Sequencer/Private/Tools/NavigationTool/Filters/TextExpressions/NavigationToolFilterTextExpressionContext.cpp
--- a/Engine/Source/Editor/Sequencer/Private/Tools/NavigationTool/Filters/TextExpressions/NavigationToolFilterTextExpressionContext.cpp
+++ b/Engine/Source/Editor/Sequencer/Private/Tools/NavigationTool/Filters/TextExpressions/NavigationToolFilterTextExpressionContext.cpp
@@ -1,9 +1,8 @@
 // Copyright Epic Games, Inc. All Rights Reserved.
 
 #include "Tools/NavigationTool/Filters/TextExpressions/NavigationToolFilterTextExpressionContext.h"
-#include "Sequencer.h"
 #include "Tools/NavigationTool/Filters/INavigationToolFilterBar.h"
-#include "Tools/NavigationTool/Items/NavigationToolItem.h"
+#include "Tools/NavigationTool/Items/INavigationToolItem.h"
 
 namespace UE::Sequencer::NavigationTool
 {
diff --git a/Engine/Source/Editor/Sequencer/Private/Tools/NavigationTool/Items/NavigationToolItemId.cpp b/Engine/Source/Editor/Sequencer/Private/Tools/NavigationTool/Items/NavigationToolItemId.cpp
--- a/Engine/Source/Editor/Sequencer/Private/Tools/NavigationTool/Items/NavigationToolItemId.cpp
+++ b/Engine/Source/Editor/Sequencer/Private/Tools/NavigationTool/Items/NavigationToolItemId.cpp
@@ -5,7 +5,9 @@
 #include "MovieSceneSection.h"
 #include "MovieSceneSequence.h"
 #include "MVVM/Extensions/IObjectBindingExtension.h"
-#include "Sequencer.h"
+#include "MVVM/Extensions/ITrackExtension.h"
+#include "MVVM/ViewModels/SectionModel.h"
+#include "MVVM/ViewModels/SequenceModel.h"
 #include "Tools/NavigationTool/Items/INavigationToolItem.h"
 
 namespace UE::Sequencer::NavigationTool
diff --git a/Engine/Source/Editor/Sequencer/Private/Tools/NavigationTool/Items/NavigationToolMarker.cpp b/Engine/Source/Editor/Sequencer/Private/Tools/NavigationTool/Items/NavigationToolMarker.cpp
--- a/Engine/Source/Editor/Sequencer/Private/Tools/NavigationTool/Items/NavigationToolMarker.cpp
+++ b/Engine/Source/Editor/Sequencer/Private/Tools/NavigationTool/Items/NavigationToolMarker.cpp
@@ -1,10 +1,12 @@
 // Copyright Epic Games, Inc. All Rights Reserved.
 
 #include "Tools/NavigationTool/Items/NavigationToolMarker.h"
+#include "ISequencer.h"
 #include "MovieScene.h"
 #include "ScopedTransaction.h"
 #include "Sections/MovieSceneSubSection.h"
 #include "Sequencer.h"
+#include "SequencerSettings.h"
 #include "Styling/StyleColors.h"
 #include "Tools/NavigationTool/Items/NavigationToolSequence.h"
 #include "Tools/NavigationTool/NavigationTool.h"
diff --git a/Engine/Source/Editor/Sequencer/Private/Tools/NavigationTool/Items/NavigationToolSequence.cpp b/Engine/Source/Editor/Sequencer/Private/Tools/NavigationTool/Items/NavigationToolSequence.cpp
--- a/Engine/Source/Editor/Sequencer/Private/Tools/NavigationTool/Items/NavigationToolSequence.cpp
+++ b/Engine/Source/Editor/Sequencer/Private/Tools/NavigationTool/Items/NavigationToolSequence.cpp
@@ -1,8 +1,8 @@
 // Copyright Epic Games, Inc. All Rights Reserved.
 
 #include "Tools/NavigationTool/Items/NavigationToolSequence.h"
-#include "LevelSequenceActor.h"
 #include "GameFramework/Actor.h"
+#include "LevelSequenceActor.h"
 #include "Misc/QualifiedFrameTime.h"
 #include "MovieScene.h"
 #include "MVVM/SectionModelStorageExtension.h"
@@ -10,8 +10,8 @@
 #include "ScopedTransaction.h"
 #include "Sections/MovieSceneCinematicShotSection.h"
 #include "Sections/MovieSceneSubSection.h"
-#include "Sequencer.h"
 #include "SequencerChannelInterface.h"
+#include "SequencerSettings.h"
 #include "Styling/SlateIconFinder.h"
 #include "Tools/NavigationTool/INavigationTool.h"
 #include "Tools/NavigationTool/Items/NavigationToolActor.h"
diff --git a/Engine/Source/Editor/Sequencer/Private/Tools/NavigationTool/Menus/NavigationToolFilterMenu.cpp b/Engine/Source/Editor/Sequencer/Private/Tools/NavigationTool/Menus/NavigationToolFilterMenu.cpp
--- a/Engine/Source/Editor/Sequencer/Private/Tools/NavigationTool/Menus/NavigationToolFilterMenu.cpp
+++ b/Engine/Source/Editor/Sequencer/Private/Tools/NavigationTool/Menus/NavigationToolFilterMenu.cpp
@@ -1,7 +1,6 @@
 ﻿// Copyright Epic Games, Inc. All Rights Reserved.
 
 #include "Tools/NavigationTool/Menus/NavigationToolFilterMenu.h"
-#include "Sequencer.h"
 #include "ToolMenu.h"
 #include "ToolMenus.h"
 #include "ToolMenuSection.h"
diff --git a/Engine/Source/Editor/Sequencer/Private/Tools/NavigationTool/NavigationTool.cpp b/Engine/Source/Editor/Sequencer/Private/Tools/NavigationTool/NavigationTool.cpp
--- a/Engine/Source/Editor/Sequencer/Private/Tools/NavigationTool/NavigationTool.cpp
+++ b/Engine/Source/Editor/Sequencer/Private/Tools/NavigationTool/NavigationTool.cpp
@@ -4,19 +4,15 @@
 #include "Algo/Reverse.h"
 #include "BlueprintEditorSettings.h"
 #include "Components/ActorComponent.h"
+#include "Editor.h"
 #include "Editor/Transactor.h"
-#include "Engine/Level.h"
-#include "Engine/LevelStreaming.h"
-#include "Engine/World.h"
-#include "Framework/Application/SlateApplication.h"
 #include "Framework/MultiBox/MultiBoxBuilder.h"
 #include "GameFramework/Actor.h"
-#include "GameFramework/WorldSettings.h"
+#include "ISequencer.h"
 #include "Kismet2/ComponentEditorUtils.h"
 #include "Misc/Optional.h"
 #include "ScopedTransaction.h"
 #include "Sections/MovieSceneSubSection.h"
-#include "Sequencer.h"
 #include "MVVM/Selection/Selection.h"
 #include "MVVM/ViewModels/TrackModel.h"
 #include "MVVM/ViewModels/TrackRowModel.h"
@@ -58,7 +54,7 @@ void FNavigationTool::Init()
 
 	FNavigationToolExtender::OnProvidersChanged().AddSP(this, &FNavigationTool::OnProvidersChanged);
 
-	if (const TSharedPtr<FSequencer> Sequencer = StaticCastSharedPtr<FSequencer>(WeakSequencer.Pin()))
+	if (const TSharedPtr<ISequencer> Sequencer = WeakSequencer.Pin())
 	{
 		Sequencer->OnActivateSequence().AddSP(this, &FNavigationTool::OnSequencerActivate);
 		Sequencer->OnCloseEvent().AddSP(this, &FNavigationTool::OnSequencerClose);
@@ -75,7 +71,7 @@ void FNavigationTool::Shutdown()
 
 	FNavigationToolExtender::OnProvidersChanged().RemoveAll(this);
 
-	if (const TSharedPtr<FSequencer> Sequencer = StaticCastSharedPtr<FSequencer>(WeakSequencer.Pin()))
+	if (const TSharedPtr<ISequencer> Sequencer = WeakSequencer.Pin())
 	{
 		Sequencer->OnActivateSequence().RemoveAll(this);
 		Sequencer->OnCloseEvent().RemoveAll(this);
@@ -508,7 +504,7 @@ void FNavigationTool::SetIgnoreNotify(const ENavigationToolIgnoreNotifyFlags InF
 
 void FNavigationTool::PostUndo(const bool bInSuccess)
 {
-	if (bInSuccess)
+	if (GEditor && bInSuccess)
 	{
 		const int32 QueueIndex = GEditor->Trans->GetQueueLength() - GEditor->Trans->GetUndoCount();
 		const FTransaction* Transaction = GEditor->Trans->GetTransaction(QueueIndex);
@@ -518,7 +514,7 @@ void FNavigationTool::PostUndo(const bool bInSuccess)
 
 void FNavigationTool::PostRedo(const bool bInSuccess)
 {
-	if (bInSuccess)
+	if (GEditor && bInSuccess)
 	{
 		const int32 QueueIndex = GEditor->Trans->GetQueueLength() - GEditor->Trans->GetUndoCount();
 		const FTransaction* Transaction = GEditor->Trans->GetTransaction(QueueIndex);
@@ -1058,7 +1054,7 @@ void FNavigationTool::OnTabVisibilityChanged(const bool bInVisible)
 
 void FNavigationTool::BindEvents()
 {
-	if (const TSharedPtr<FSequencer> Sequencer = StaticCastSharedPtr<FSequencer>(WeakSequencer.Pin()))
+	if (const TSharedPtr<ISequencer> Sequencer = WeakSequencer.Pin())
 	{
 		Sequencer->OnChannelChanged().AddSPLambda(this, [this](const FMovieSceneChannelMetaData*, UMovieSceneSection*)
 			{
@@ -1077,8 +1073,14 @@ void FNavigationTool::BindEvents()
 		Sequencer->OnTreeViewChanged().AddSP(this, &FNavigationTool::OnTreeViewChanged);
 		Sequencer->OnEndScrubbingEvent().AddSP(this, &FNavigationTool::OnTreeViewChanged);
 
-		Sequencer->GetSelection().Outliner.OnChanged.AddSP(this, &FNavigationTool::OnSequencerSelectionChanged);
-		Sequencer->GetSelection().TrackArea.OnChanged.AddSP(this, &FNavigationTool::OnSequencerSelectionChanged);
+		if (const TSharedPtr<FSequencerEditorViewModel> ViewModel = Sequencer->GetViewModel())
+		{
+			if (const TSharedPtr<FSequencerSelection> SequencerSelection = ViewModel->GetSelection())
+			{
+				SequencerSelection->Outliner.OnChanged.AddSP(this, &FNavigationTool::OnSequencerSelectionChanged);
+				SequencerSelection->TrackArea.OnChanged.AddSP(this, &FNavigationTool::OnSequencerSelectionChanged);
+			}
+		}
 	}
 
 	// Listen to object replacement changes
@@ -1088,7 +1090,7 @@ void FNavigationTool::BindEvents()
 
 void FNavigationTool::UnbindEvents()
 {
-	if (const TSharedPtr<FSequencer> Sequencer = StaticCastSharedPtr<FSequencer>(WeakSequencer.Pin()))
+	if (const TSharedPtr<ISequencer> Sequencer = WeakSequencer.Pin())
 	{
 		Sequencer->OnChannelChanged().RemoveAll(this);
 		Sequencer->OnMovieSceneBindingsChanged().RemoveAll(this);
@@ -1098,8 +1100,14 @@ void FNavigationTool::UnbindEvents()
 		Sequencer->OnTreeViewChanged().RemoveAll(this);
 		Sequencer->OnEndScrubbingEvent().RemoveAll(this);
 
-		Sequencer->GetSelection().Outliner.OnChanged.RemoveAll(this);
-		Sequencer->GetSelection().TrackArea.OnChanged.RemoveAll(this);
+		if (const TSharedPtr<FSequencerEditorViewModel> ViewModel = Sequencer->GetViewModel())
+		{
+			if (const TSharedPtr<FSequencerSelection> SequencerSelection = ViewModel->GetSelection())
+			{
+				SequencerSelection->Outliner.OnChanged.RemoveAll(this);
+				SequencerSelection->TrackArea.OnChanged.RemoveAll(this);
+			}
+		}
 	}
 
 	FCoreUObjectDelegates::OnObjectsReplaced.RemoveAll(this);
@@ -1187,22 +1195,32 @@ void FNavigationTool::OnSequencerSelectionChanged()
 		return;
 	}
 
-	const TSharedPtr<FSequencer> Sequencer = StaticCastSharedPtr<FSequencer>(WeakSequencer.Pin());
+	const TSharedPtr<ISequencer> Sequencer = WeakSequencer.Pin();
 	if (!Sequencer.IsValid())
 	{
 		return;
 	}
 
-	FSequencerSelection& SequencerSelection = Sequencer->GetSelection();
+	const TSharedPtr<FSequencerEditorViewModel> ViewModel = Sequencer->GetViewModel();
+	if (!ViewModel.IsValid())
+	{
+		return;
+	}
+
+	const TSharedPtr<FSequencerSelection> SequencerSelection = ViewModel->GetSelection();
+	if (!SequencerSelection.IsValid())
+	{
+		return;
+	}
 
 	if (!ItemsLastSelected.IsValid())
 	{
 		ItemsLastSelected = MakeShared<TArray<FNavigationToolItemWeakPtr>>();
 	}
 
-	ItemsLastSelected->Reserve(ItemsLastSelected->Num() + SequencerSelection.Outliner.Num());
+	ItemsLastSelected->Reserve(ItemsLastSelected->Num() + SequencerSelection->Outliner.Num());
 
-	for (const FViewModelPtr OutlinerItem : SequencerSelection.Outliner)
+	for (const FViewModelPtr OutlinerItem : SequencerSelection->Outliner)
 	{
 		ItemsLastSelected->Append(TryFindItems(OutlinerItem));
 	}
diff --git a/Engine/Source/Editor/Sequencer/Private/Tools/NavigationTool/NavigationToolScopedSelection.cpp b/Engine/Source/Editor/Sequencer/Private/Tools/NavigationTool/NavigationToolScopedSelection.cpp
--- a/Engine/Source/Editor/Sequencer/Private/Tools/NavigationTool/NavigationToolScopedSelection.cpp
+++ b/Engine/Source/Editor/Sequencer/Private/Tools/NavigationTool/NavigationToolScopedSelection.cpp
@@ -1,12 +1,17 @@
 // Copyright Epic Games, Inc. All Rights Reserved.
 
 #include "Tools/NavigationTool/NavigationToolScopedSelection.h"
+#include "ISequencer.h"
 #include "MVVM/ObjectBindingModelStorageExtension.h"
 #include "MVVM/SectionModelStorageExtension.h"
 #include "MVVM/Selection/Selection.h"
 #include "MVVM/TrackModelStorageExtension.h"
+#include "MVVM/ViewModels/ObjectBindingModel.h"
+#include "MVVM/ViewModels/SectionModel.h"
+#include "MVVM/ViewModels/SequencerEditorViewModel.h"
+#include "MVVM/ViewModels/TrackModel.h"
 #include "MVVM/Views/SOutlinerView.h"
-#include "Sequencer.h"
+#include "Tools/NavigationTool/NavigationToolSettings.h"
 #include "Utils/NavigationToolMiscUtils.h"
 
 namespace UE::Sequencer::NavigationTool
diff --git a/Engine/Source/Editor/Sequencer/Private/Tools/NavigationTool/NavigationToolView.cpp b/Engine/Source/Editor/Sequencer/Private/Tools/NavigationTool/NavigationToolView.cpp
--- a/Engine/Source/Editor/Sequencer/Private/Tools/NavigationTool/NavigationToolView.cpp
+++ b/Engine/Source/Editor/Sequencer/Private/Tools/NavigationTool/NavigationToolView.cpp
@@ -6,10 +6,9 @@
 #include "Framework/Commands/GenericCommands.h"
 #include "Framework/MultiBox/MultiBoxBuilder.h"
 #include "IContentBrowserSingleton.h"
+#include "ISequencer.h"
 #include "LevelSequence.h"
 #include "Misc/MessageDialog.h"
-#include "MVVM/Selection/Selection.h"
-#include "Sequencer.h"
 #include "SequencerLog.h"
 #include "Filters/Filters/NavigationToolBuiltInFilter.h"
 #include "Styling/SlateTypes.h"
diff --git a/Engine/Source/Editor/Sequencer/Private/Tools/NavigationTool/Utils/NavigationToolMovieSceneUtils.cpp b/Engine/Source/Editor/Sequencer/Private/Tools/NavigationTool/Utils/NavigationToolMovieSceneUtils.cpp
--- a/Engine/Source/Editor/Sequencer/Private/Tools/NavigationTool/Utils/NavigationToolMovieSceneUtils.cpp
+++ b/Engine/Source/Editor/Sequencer/Private/Tools/NavigationTool/Utils/NavigationToolMovieSceneUtils.cpp
@@ -1,6 +1,7 @@
 ﻿// Copyright Epic Games, Inc. All Rights Reserved.
 
 #include "NavigationToolMovieSceneUtils.h"
+#include "ISequencer.h"
 #include "Misc/FrameTime.h"
 #include "MovieScene.h"
 #include "Sections/MovieSceneSubSection.h"
@@ -54,7 +55,7 @@ bool IsGloballyMarkedFramesForSequence(UMovieSceneSequence* const InSequence)
 		return false;
 	}
 
-	UMovieScene* const MovieScene = InSequence->GetMovieScene();
+	const UMovieScene* const MovieScene = InSequence->GetMovieScene();
 	if (!IsValid(MovieScene))
 	{
 		return false;
@@ -117,7 +118,7 @@ void GetSequenceSubSections(UMovieSceneSequence* const InSequence, TArray<UMovie
 		return;
 	}
 
-	UMovieScene* const MovieScene = InSequence->GetMovieScene();
+	const UMovieScene* const MovieScene = InSequence->GetMovieScene();
 	if (!IsValid(MovieScene))
 	{
 		return;
diff --git a/Engine/Source/Editor/Sequencer/Private/Tools/NavigationTool/Widgets/Columns/SNavigationToolColor.cpp b/Engine/Source/Editor/Sequencer/Private/Tools/NavigationTool/Widgets/Columns/SNavigationToolColor.cpp
--- a/Engine/Source/Editor/Sequencer/Private/Tools/NavigationTool/Widgets/Columns/SNavigationToolColor.cpp
+++ b/Engine/Source/Editor/Sequencer/Private/Tools/NavigationTool/Widgets/Columns/SNavigationToolColor.cpp
@@ -112,7 +112,7 @@ public:
 };
 
 void SNavigationToolColor::Construct(const FArguments& InArgs
-	, const FNavigationToolItemPtr& InItem
+	, const FNavigationToolItemRef& InItem
 	, const TSharedRef<INavigationToolView>& InView
 	, const TSharedRef<SNavigationToolTreeRow>& InRowWidget)
 {
diff --git a/Engine/Source/Editor/Sequencer/Private/Tools/NavigationTool/Widgets/Columns/SNavigationToolColor.h b/Engine/Source/Editor/Sequencer/Private/Tools/NavigationTool/Widgets/Columns/SNavigationToolColor.h
--- a/Engine/Source/Editor/Sequencer/Private/Tools/NavigationTool/Widgets/Columns/SNavigationToolColor.h
+++ b/Engine/Source/Editor/Sequencer/Private/Tools/NavigationTool/Widgets/Columns/SNavigationToolColor.h
@@ -23,7 +23,7 @@ public:
 	virtual ~SNavigationToolColor() override {}
 
 	void Construct(const FArguments& InArgs
-		, const FNavigationToolItemPtr& InItem
+		, const FNavigationToolItemRef& InItem
 		, const TSharedRef<INavigationToolView>& InView
 		, const TSharedRef<SNavigationToolTreeRow>& InRowWidget);
 
diff --git a/Engine/Source/Editor/Sequencer/Private/Tools/NavigationTool/Widgets/Columns/SNavigationToolComment.cpp b/Engine/Source/Editor/Sequencer/Private/Tools/NavigationTool/Widgets/Columns/SNavigationToolComment.cpp
--- a/Engine/Source/Editor/Sequencer/Private/Tools/NavigationTool/Widgets/Columns/SNavigationToolComment.cpp
+++ b/Engine/Source/Editor/Sequencer/Private/Tools/NavigationTool/Widgets/Columns/SNavigationToolComment.cpp
@@ -16,7 +16,7 @@ namespace UE::Sequencer::NavigationTool
 {
 
 void SNavigationToolComment::Construct(const FArguments& InArgs
-	, const FNavigationToolItemPtr& InItem
+	, const FNavigationToolItemRef& InItem
 	, const TSharedRef<INavigationToolView>& InView
 	, const TSharedRef<SNavigationToolTreeRow>& InRowWidget)
 {
@@ -24,24 +24,9 @@ void SNavigationToolComment::Construct(const FArguments& InArgs
 	WeakView = InView;
 	WeakRowWidget = InRowWidget;
 
-	auto EmptyContent = [this]()
-	{
-		ChildSlot
-		[
-			SNullWidget::NullWidget
-		];
-	};
-
-	if (!WeakItem.IsValid() || !WeakView.IsValid())
-	{
-		EmptyContent();
-		return;
-	}
-
-	const FNavigationToolSequence* const SequenceItem = WeakItem.Pin()->CastTo<FNavigationToolSequence>();
+	const FNavigationToolSequence* const SequenceItem = InItem->CastTo<FNavigationToolSequence>();
 	if (!SequenceItem || !IsValid(SequenceItem->GetSubSection()))
 	{
-		EmptyContent();
 		return;
 	}
 
diff --git a/Engine/Source/Editor/Sequencer/Private/Tools/NavigationTool/Widgets/Columns/SNavigationToolComment.h b/Engine/Source/Editor/Sequencer/Private/Tools/NavigationTool/Widgets/Columns/SNavigationToolComment.h
--- a/Engine/Source/Editor/Sequencer/Private/Tools/NavigationTool/Widgets/Columns/SNavigationToolComment.h
+++ b/Engine/Source/Editor/Sequencer/Private/Tools/NavigationTool/Widgets/Columns/SNavigationToolComment.h
@@ -21,7 +21,7 @@ public:
 	virtual ~SNavigationToolComment() override;
 
 	void Construct(const FArguments& InArgs
-		, const FNavigationToolItemPtr& InItem
+		, const FNavigationToolItemRef& InItem
 		, const TSharedRef<INavigationToolView>& InView
 		, const TSharedRef<SNavigationToolTreeRow>& InRowWidget);
 
diff --git a/Engine/Source/Editor/Sequencer/Private/Tools/NavigationTool/Widgets/Columns/SNavigationToolDeactiveState.cpp b/Engine/Source/Editor/Sequencer/Private/Tools/NavigationTool/Widgets/Columns/SNavigationToolDeactiveState.cpp
--- a/Engine/Source/Editor/Sequencer/Private/Tools/NavigationTool/Widgets/Columns/SNavigationToolDeactiveState.cpp
+++ b/Engine/Source/Editor/Sequencer/Private/Tools/NavigationTool/Widgets/Columns/SNavigationToolDeactiveState.cpp
@@ -73,8 +73,8 @@ void SetItemInactive(const FNavigationToolItemPtr& InItem, const bool bInInactiv
 
 void SNavigationToolDeactiveState::Construct(const FArguments& InArgs
 	, const TSharedRef<FNavigationToolDeactiveStateColumn>& InColumn
+	, const FNavigationToolItemRef& InItem
 	, const TSharedRef<INavigationToolView>& InView
-	, const FNavigationToolItemPtr& InItem
 	, const TSharedRef<SNavigationToolTreeRow>& InRowWidget)
 {
 	WeakColumn = InColumn;
diff --git a/Engine/Source/Editor/Sequencer/Private/Tools/NavigationTool/Widgets/Columns/SNavigationToolDeactiveState.h b/Engine/Source/Editor/Sequencer/Private/Tools/NavigationTool/Widgets/Columns/SNavigationToolDeactiveState.h
--- a/Engine/Source/Editor/Sequencer/Private/Tools/NavigationTool/Widgets/Columns/SNavigationToolDeactiveState.h
+++ b/Engine/Source/Editor/Sequencer/Private/Tools/NavigationTool/Widgets/Columns/SNavigationToolDeactiveState.h
@@ -12,8 +12,6 @@ namespace UE::Sequencer::NavigationTool
 
 class FNavigationToolItem;
 class FNavigationToolDeactiveStateColumn;
-class FNavigationToolView;
-class INavigationToolItem;
 class INavigationToolView;
 class SNavigationToolTreeRow;
 enum class EItemSequenceInactiveState;
@@ -27,8 +25,8 @@ public:
 
 	void Construct(const FArguments& InArgs
 		, const TSharedRef<FNavigationToolDeactiveStateColumn>& InColumn
+		, const FNavigationToolItemRef& InItem
 		, const TSharedRef<INavigationToolView>& InView
-		, const FNavigationToolItemPtr& InItem
 		, const TSharedRef<SNavigationToolTreeRow>& InRowWidget);
 
 private:
diff --git a/Engine/Source/Editor/Sequencer/Private/Tools/NavigationTool/Widgets/Columns/SNavigationToolEndFrameOffset.cpp b/Engine/Source/Editor/Sequencer/Private/Tools/NavigationTool/Widgets/Columns/SNavigationToolEndFrameOffset.cpp
--- a/Engine/Source/Editor/Sequencer/Private/Tools/NavigationTool/Widgets/Columns/SNavigationToolEndFrameOffset.cpp
+++ b/Engine/Source/Editor/Sequencer/Private/Tools/NavigationTool/Widgets/Columns/SNavigationToolEndFrameOffset.cpp
@@ -1,20 +1,18 @@
 // Copyright Epic Games, Inc. All Rights Reserved.
 
 #include "SNavigationToolEndFrameOffset.h"
-#include "Sequencer.h"
+#include "Sections/MovieSceneSubSection.h"
 #include "Tools/NavigationTool/Items/NavigationToolItemUtils.h"
 #include "Tools/NavigationTool/Items/NavigationToolSequence.h"
 #include "Tools/NavigationTool/NavigationToolView.h"
 #include "Tools/NavigationTool/Widgets/SNavigationToolTreeRow.h"
-#include "TrackEditors/SubTrackEditorBase.h"
-#include "Widgets/Input/SSpinBox.h"
 
 using namespace UE::Sequencer::NavigationTool;
 
 #define LOCTEXT_NAMESPACE "SNavigationToolEndFrameOffset"
 
 void SNavigationToolEndFrameOffset::Construct(const FArguments& InArgs
-	, const FNavigationToolItemPtr& InItem
+	, const FNavigationToolItemRef& InItem
 	, const TSharedRef<INavigationToolView>& InView
 	, const TSharedRef<SNavigationToolTreeRow>& InRowWidget)
 {
@@ -22,14 +20,7 @@ void SNavigationToolEndFrameOffset::Construct(const FArguments& InArgs
 	WeakView = InView;
 	WeakRowWidget = InRowWidget;
 
-	if (!WeakItem.IsValid()
-		|| !WeakView.IsValid()
-		|| !WeakRowWidget.IsValid())
-	{
-		return;
-	}
-
-	const FNavigationToolSequence* const SequenceItem = WeakItem.Pin()->CastTo<FNavigationToolSequence>();
+	const FNavigationToolSequence* const SequenceItem = InItem->CastTo<FNavigationToolSequence>();
 	if (!SequenceItem || !IsValid(SequenceItem->GetSubSection()))
 	{
 		return;
diff --git a/Engine/Source/Editor/Sequencer/Private/Tools/NavigationTool/Widgets/Columns/SNavigationToolEndFrameOffset.h b/Engine/Source/Editor/Sequencer/Private/Tools/NavigationTool/Widgets/Columns/SNavigationToolEndFrameOffset.h
--- a/Engine/Source/Editor/Sequencer/Private/Tools/NavigationTool/Widgets/Columns/SNavigationToolEndFrameOffset.h
+++ b/Engine/Source/Editor/Sequencer/Private/Tools/NavigationTool/Widgets/Columns/SNavigationToolEndFrameOffset.h
@@ -20,7 +20,7 @@ public:
 	virtual ~SNavigationToolEndFrameOffset() override {}
 
 	void Construct(const FArguments& InArgs
-		, const FNavigationToolItemPtr& InItem
+		, const FNavigationToolItemRef& InItem
 		, const TSharedRef<INavigationToolView>& InView
 		, const TSharedRef<SNavigationToolTreeRow>& InRowWidget);
 
diff --git a/Engine/Source/Editor/Sequencer/Private/Tools/NavigationTool/Widgets/Columns/SNavigationToolHBias.cpp b/Engine/Source/Editor/Sequencer/Private/Tools/NavigationTool/Widgets/Columns/SNavigationToolHBias.cpp
--- a/Engine/Source/Editor/Sequencer/Private/Tools/NavigationTool/Widgets/Columns/SNavigationToolHBias.cpp
+++ b/Engine/Source/Editor/Sequencer/Private/Tools/NavigationTool/Widgets/Columns/SNavigationToolHBias.cpp
@@ -16,7 +16,7 @@ namespace UE::Sequencer::NavigationTool
 {
 
 void SNavigationToolHBias::Construct(const FArguments& InArgs
-	, const FNavigationToolItemPtr& InItem
+	, const FNavigationToolItemRef& InItem
 	, const TSharedRef<INavigationToolView>& InView
 	, const TSharedRef<SNavigationToolTreeRow>& InRowWidget)
 {
@@ -24,31 +24,9 @@ void SNavigationToolHBias::Construct(const FArguments& InArgs
 	WeakView = InView;
 	WeakRowWidget = InRowWidget;
 
-	auto EmptyContent = [this]()
+	const FNavigationToolSequence* const SequenceItem = InItem->CastTo<FNavigationToolSequence>();
+	if (!SequenceItem || !IsValid(SequenceItem->GetSubSection()))
 	{
-		ChildSlot
-		[
-			SNullWidget::NullWidget
-		];
-	};
-
-	if (!WeakItem.IsValid() || !WeakView.IsValid())
-	{
-		EmptyContent();
-		return;
-	}
-
-	const FNavigationToolSequence* const SequenceItem = WeakItem.Pin()->CastTo<FNavigationToolSequence>();
-	if (!SequenceItem)
-	{
-		EmptyContent();
-		return;
-	}
-
-	UMovieSceneSubSection* const SubSection = SequenceItem->GetSubSection();
-	if (!IsValid(SubSection))
-	{
-		EmptyContent();
 		return;
 	}
 
diff --git a/Engine/Source/Editor/Sequencer/Private/Tools/NavigationTool/Widgets/Columns/SNavigationToolHBias.h b/Engine/Source/Editor/Sequencer/Private/Tools/NavigationTool/Widgets/Columns/SNavigationToolHBias.h
--- a/Engine/Source/Editor/Sequencer/Private/Tools/NavigationTool/Widgets/Columns/SNavigationToolHBias.h
+++ b/Engine/Source/Editor/Sequencer/Private/Tools/NavigationTool/Widgets/Columns/SNavigationToolHBias.h
@@ -22,7 +22,7 @@ public:
 	virtual ~SNavigationToolHBias() override {}
 
 	void Construct(const FArguments& InArgs
-		, const FNavigationToolItemPtr& InItem
+		, const FNavigationToolItemRef& InItem
 		, const TSharedRef<INavigationToolView>& InView
 		, const TSharedRef<SNavigationToolTreeRow>& InRowWidget);
 
diff --git a/Engine/Source/Editor/Sequencer/Private/Tools/NavigationTool/Widgets/Columns/SNavigationToolInTime.cpp b/Engine/Source/Editor/Sequencer/Private/Tools/NavigationTool/Widgets/Columns/SNavigationToolInTime.cpp
--- a/Engine/Source/Editor/Sequencer/Private/Tools/NavigationTool/Widgets/Columns/SNavigationToolInTime.cpp
+++ b/Engine/Source/Editor/Sequencer/Private/Tools/NavigationTool/Widgets/Columns/SNavigationToolInTime.cpp
@@ -13,7 +13,7 @@ namespace UE::Sequencer::NavigationTool
 {
 
 void SNavigationToolInTime::Construct(const FArguments& InArgs
-	, const FNavigationToolItemPtr& InItem
+	, const FNavigationToolItemRef& InItem
 	, const TSharedRef<INavigationToolView>& InView
 	, const TSharedRef<SNavigationToolTreeRow>& InRowWidget)
 {
@@ -21,10 +21,7 @@ void SNavigationToolInTime::Construct(const FArguments& InArgs
 	WeakView = InView;
 	WeakRowWidget = InRowWidget;
 
-	if (!WeakItem.IsValid()
-		|| !WeakView.IsValid()
-		|| !WeakRowWidget.IsValid()
-		|| !WeakItem.Pin()->IsA<IInTimeExtension>())
+	if (!InItem->IsA<IInTimeExtension>())
 	{
 		return;
 	}
diff --git a/Engine/Source/Editor/Sequencer/Private/Tools/NavigationTool/Widgets/Columns/SNavigationToolInTime.h b/Engine/Source/Editor/Sequencer/Private/Tools/NavigationTool/Widgets/Columns/SNavigationToolInTime.h
--- a/Engine/Source/Editor/Sequencer/Private/Tools/NavigationTool/Widgets/Columns/SNavigationToolInTime.h
+++ b/Engine/Source/Editor/Sequencer/Private/Tools/NavigationTool/Widgets/Columns/SNavigationToolInTime.h
@@ -21,7 +21,7 @@ public:
 	virtual ~SNavigationToolInTime() override {}
 
 	void Construct(const FArguments& InArgs
-		, const FNavigationToolItemPtr& InItem
+		, const FNavigationToolItemRef& InItem
 		, const TSharedRef<INavigationToolView>& InView
 		, const TSharedRef<SNavigationToolTreeRow>& InRowWidget);
 
diff --git a/Engine/Source/Editor/Sequencer/Private/Tools/NavigationTool/Widgets/Columns/SNavigationToolItemChip.cpp b/Engine/Source/Editor/Sequencer/Private/Tools/NavigationTool/Widgets/Columns/SNavigationToolItemChip.cpp
--- a/Engine/Source/Editor/Sequencer/Private/Tools/NavigationTool/Widgets/Columns/SNavigationToolItemChip.cpp
+++ b/Engine/Source/Editor/Sequencer/Private/Tools/NavigationTool/Widgets/Columns/SNavigationToolItemChip.cpp
@@ -17,7 +17,7 @@ namespace UE::Sequencer::NavigationTool
 {
 
 void SNavigationToolItemChip::Construct(const FArguments& InArgs
-	, const TSharedRef<INavigationToolItem>& InItem
+	, const FNavigationToolItemRef& InItem
 	, const TSharedPtr<INavigationToolView>& InView)
 {
 	WeakItem = InItem;
diff --git a/Engine/Source/Editor/Sequencer/Private/Tools/NavigationTool/Widgets/Columns/SNavigationToolItemChip.h b/Engine/Source/Editor/Sequencer/Private/Tools/NavigationTool/Widgets/Columns/SNavigationToolItemChip.h
--- a/Engine/Source/Editor/Sequencer/Private/Tools/NavigationTool/Widgets/Columns/SNavigationToolItemChip.h
+++ b/Engine/Source/Editor/Sequencer/Private/Tools/NavigationTool/Widgets/Columns/SNavigationToolItemChip.h
@@ -34,7 +34,7 @@ public:
 	SLATE_END_ARGS()
 
 	void Construct(const FArguments& InArgs
-		, const TSharedRef<INavigationToolItem>& InItem
+		, const FNavigationToolItemRef& InItem
 		, const TSharedPtr<INavigationToolView>& InView);
 
 	//~ Begin SWidget
diff --git a/Engine/Source/Editor/Sequencer/Private/Tools/NavigationTool/Widgets/Columns/SNavigationToolItemList.cpp b/Engine/Source/Editor/Sequencer/Private/Tools/NavigationTool/Widgets/Columns/SNavigationToolItemList.cpp
--- a/Engine/Source/Editor/Sequencer/Private/Tools/NavigationTool/Widgets/Columns/SNavigationToolItemList.cpp
+++ b/Engine/Source/Editor/Sequencer/Private/Tools/NavigationTool/Widgets/Columns/SNavigationToolItemList.cpp
@@ -3,8 +3,6 @@
 #include "SNavigationToolItemList.h"
 #include "Tools/NavigationTool/Items/NavigationToolItem.h"
 #include "Tools/NavigationTool/NavigationToolView.h"
-#include "Tools/NavigationTool/Items/NavigationToolBinding.h"
-#include "Tools/NavigationTool/Items/NavigationToolSequence.h"
 #include "Tools/NavigationTool/Utils/NavigationToolMiscUtils.h"
 #include "Tools/NavigationTool/Widgets/Columns/SNavigationToolItemChip.h"
 #include "Tools/NavigationTool/Widgets/SNavigationToolTreeRow.h"
@@ -16,7 +14,7 @@ namespace UE::Sequencer::NavigationTool
 {
 
 void SNavigationToolItemList::Construct(const FArguments& InArgs
-	, const FNavigationToolItemPtr& InItem
+	, const FNavigationToolItemRef& InItem
 	, const TSharedRef<INavigationToolView>& InView
 	, const TSharedRef<SNavigationToolTreeRow>& InRowWidget)
 {
diff --git a/Engine/Source/Editor/Sequencer/Private/Tools/NavigationTool/Widgets/Columns/SNavigationToolItemList.h b/Engine/Source/Editor/Sequencer/Private/Tools/NavigationTool/Widgets/Columns/SNavigationToolItemList.h
--- a/Engine/Source/Editor/Sequencer/Private/Tools/NavigationTool/Widgets/Columns/SNavigationToolItemList.h
+++ b/Engine/Source/Editor/Sequencer/Private/Tools/NavigationTool/Widgets/Columns/SNavigationToolItemList.h
@@ -23,7 +23,7 @@ public:
 	virtual ~SNavigationToolItemList() override;
 
 	void Construct(const FArguments& InArgs
-		, const FNavigationToolItemPtr& InItem
+		, const FNavigationToolItemRef& InItem
 		, const TSharedRef<INavigationToolView>& InView
 		, const TSharedRef<SNavigationToolTreeRow>& InRowWidget);
 
diff --git a/Engine/Source/Editor/Sequencer/Private/Tools/NavigationTool/Widgets/Columns/SNavigationToolLabelItem.cpp b/Engine/Source/Editor/Sequencer/Private/Tools/NavigationTool/Widgets/Columns/SNavigationToolLabelItem.cpp
--- a/Engine/Source/Editor/Sequencer/Private/Tools/NavigationTool/Widgets/Columns/SNavigationToolLabelItem.cpp
+++ b/Engine/Source/Editor/Sequencer/Private/Tools/NavigationTool/Widgets/Columns/SNavigationToolLabelItem.cpp
@@ -18,14 +18,12 @@ namespace UE::Sequencer::NavigationTool
 {
 
 void SNavigationToolLabelItem::Construct(const FArguments& InArgs
-	, const FNavigationToolItemPtr& InItem
+	, const FNavigationToolItemRef& InItem
 	, const TSharedRef<SNavigationToolTreeRow>& InRowWidget)
 {
 	WeakItem = InItem;
 	WeakView = InRowWidget->GetToolView();
 
-	const auto ItemRef = InItem.ToSharedRef();
-
 	InItem->OnRenameAction().AddSP(this, &SNavigationToolLabelItem::OnRenameAction);
 
 	ChildSlot
@@ -60,9 +58,9 @@ void SNavigationToolLabelItem::Construct(const FArguments& InArgs
 					[
 						SNew(SImage)
 						.IsEnabled(this, &SNavigationToolLabelItem::IsItemEnabled)
-						.ColorAndOpacity(ItemRef, &INavigationToolItem::GetIconColor)
-						.Image(ItemRef, &INavigationToolItem::GetIconBrush)
-						.ToolTipText(ItemRef, &INavigationToolItem::GetIconTooltipText)
+						.ColorAndOpacity(InItem, &INavigationToolItem::GetIconColor)
+						.Image(InItem, &INavigationToolItem::GetIconBrush)
+						.ToolTipText(InItem, &INavigationToolItem::GetIconTooltipText)
 					]
 				]
 				+ SHorizontalBox::Slot()
diff --git a/Engine/Source/Editor/Sequencer/Private/Tools/NavigationTool/Widgets/Columns/SNavigationToolLabelItem.h b/Engine/Source/Editor/Sequencer/Private/Tools/NavigationTool/Widgets/Columns/SNavigationToolLabelItem.h
--- a/Engine/Source/Editor/Sequencer/Private/Tools/NavigationTool/Widgets/Columns/SNavigationToolLabelItem.h
+++ b/Engine/Source/Editor/Sequencer/Private/Tools/NavigationTool/Widgets/Columns/SNavigationToolLabelItem.h
@@ -23,7 +23,7 @@ public:
 	SLATE_END_ARGS()
 
 	void Construct(const FArguments& InArgs
-		, const FNavigationToolItemPtr& InItem
+		, const FNavigationToolItemRef& InItem
 		, const TSharedRef<SNavigationToolTreeRow>& InRowWidget);
 
 	bool IsReadOnly() const;
diff --git a/Engine/Source/Editor/Sequencer/Private/Tools/NavigationTool/Widgets/Columns/SNavigationToolLength.cpp b/Engine/Source/Editor/Sequencer/Private/Tools/NavigationTool/Widgets/Columns/SNavigationToolLength.cpp
--- a/Engine/Source/Editor/Sequencer/Private/Tools/NavigationTool/Widgets/Columns/SNavigationToolLength.cpp
+++ b/Engine/Source/Editor/Sequencer/Private/Tools/NavigationTool/Widgets/Columns/SNavigationToolLength.cpp
@@ -14,7 +14,7 @@ namespace UE::Sequencer::NavigationTool
 {
 
 void SNavigationToolLength::Construct(const FArguments& InArgs
-	, const FNavigationToolItemPtr& InItem
+	, const FNavigationToolItemRef& InItem
 	, const TSharedRef<INavigationToolView>& InView
 	, const TSharedRef<SNavigationToolTreeRow>& InRowWidget)
 {
@@ -22,14 +22,7 @@ void SNavigationToolLength::Construct(const FArguments& InArgs
 	WeakView = InView;
 	WeakRowWidget = InRowWidget;
 
-	if (!WeakItem.IsValid()
-		|| !WeakView.IsValid()
-		|| !WeakRowWidget.IsValid())
-	{
-		return;
-	}
-
-	const FNavigationToolSequence* const SequenceItem = WeakItem.Pin()->CastTo<FNavigationToolSequence>();
+	const FNavigationToolSequence* const SequenceItem = InItem->CastTo<FNavigationToolSequence>();
 	if (!SequenceItem || !IsValid(SequenceItem->GetSubSection()))
 	{
 		return;
diff --git a/Engine/Source/Editor/Sequencer/Private/Tools/NavigationTool/Widgets/Columns/SNavigationToolLength.h b/Engine/Source/Editor/Sequencer/Private/Tools/NavigationTool/Widgets/Columns/SNavigationToolLength.h
--- a/Engine/Source/Editor/Sequencer/Private/Tools/NavigationTool/Widgets/Columns/SNavigationToolLength.h
+++ b/Engine/Source/Editor/Sequencer/Private/Tools/NavigationTool/Widgets/Columns/SNavigationToolLength.h
@@ -20,7 +20,7 @@ public:
 	virtual ~SNavigationToolLength() override {}
 
 	void Construct(const FArguments& InArgs
-		, const FNavigationToolItemPtr& InItem
+		, const FNavigationToolItemRef& InItem
 		, const TSharedRef<INavigationToolView>& InView
 		, const TSharedRef<SNavigationToolTreeRow>& InRowWidget);
 
diff --git a/Engine/Source/Editor/Sequencer/Private/Tools/NavigationTool/Widgets/Columns/SNavigationToolLock.cpp b/Engine/Source/Editor/Sequencer/Private/Tools/NavigationTool/Widgets/Columns/SNavigationToolLock.cpp
--- a/Engine/Source/Editor/Sequencer/Private/Tools/NavigationTool/Widgets/Columns/SNavigationToolLock.cpp
+++ b/Engine/Source/Editor/Sequencer/Private/Tools/NavigationTool/Widgets/Columns/SNavigationToolLock.cpp
@@ -71,7 +71,7 @@ void SetItemLocked(const FNavigationToolItemPtr& InItem, const bool bInIsLocked)
 }
 
 void SNavigationToolLock::Construct(const FArguments& InArgs
-	, const FNavigationToolItemPtr& InItem
+	, const FNavigationToolItemRef& InItem
 	, const TSharedRef<INavigationToolView>& InView
 	, const TSharedRef<SNavigationToolTreeRow>& InRowWidget)
 {
diff --git a/Engine/Source/Editor/Sequencer/Private/Tools/NavigationTool/Widgets/Columns/SNavigationToolLock.h b/Engine/Source/Editor/Sequencer/Private/Tools/NavigationTool/Widgets/Columns/SNavigationToolLock.h
--- a/Engine/Source/Editor/Sequencer/Private/Tools/NavigationTool/Widgets/Columns/SNavigationToolLock.h
+++ b/Engine/Source/Editor/Sequencer/Private/Tools/NavigationTool/Widgets/Columns/SNavigationToolLock.h
@@ -21,7 +21,6 @@ namespace UE::Sequencer::NavigationTool
 
 class INavigationToolView;
 class SNavigationToolTreeRow;
-enum class ESequenceLockState;
 
 class SNavigationToolLock : public SImage
 {
@@ -32,7 +31,7 @@ public:
 	virtual ~SNavigationToolLock() override {}
 
 	void Construct(const FArguments& InArgs
-		, const FNavigationToolItemPtr& InItem
+		, const FNavigationToolItemRef& InItem
 		, const TSharedRef<INavigationToolView>& InView
 		, const TSharedRef<SNavigationToolTreeRow>& InRowWidget);
 
diff --git a/Engine/Source/Editor/Sequencer/Private/Tools/NavigationTool/Widgets/Columns/SNavigationToolMarkerVisibility.cpp b/Engine/Source/Editor/Sequencer/Private/Tools/NavigationTool/Widgets/Columns/SNavigationToolMarkerVisibility.cpp
--- a/Engine/Source/Editor/Sequencer/Private/Tools/NavigationTool/Widgets/Columns/SNavigationToolMarkerVisibility.cpp
+++ b/Engine/Source/Editor/Sequencer/Private/Tools/NavigationTool/Widgets/Columns/SNavigationToolMarkerVisibility.cpp
@@ -74,7 +74,7 @@ void SetItemMarkerVisibility(const FNavigationToolItemPtr& InItem, const bool bI
 
 void SNavigationToolMarkerVisibility::Construct(const FArguments& InArgs
 	, const TSharedRef<FNavigationToolMarkerVisibilityColumn>& InColumn
-	, const FNavigationToolItemPtr& InItem
+	, const FNavigationToolItemRef& InItem
 	, const TSharedRef<INavigationToolView>& InView
 	, const TSharedRef<SNavigationToolTreeRow>& InRowWidget)
 {
diff --git a/Engine/Source/Editor/Sequencer/Private/Tools/NavigationTool/Widgets/Columns/SNavigationToolMarkerVisibility.h b/Engine/Source/Editor/Sequencer/Private/Tools/NavigationTool/Widgets/Columns/SNavigationToolMarkerVisibility.h
--- a/Engine/Source/Editor/Sequencer/Private/Tools/NavigationTool/Widgets/Columns/SNavigationToolMarkerVisibility.h
+++ b/Engine/Source/Editor/Sequencer/Private/Tools/NavigationTool/Widgets/Columns/SNavigationToolMarkerVisibility.h
@@ -24,7 +24,7 @@ public:
 
 	void Construct(const FArguments& InArgs
 		, const TSharedRef<FNavigationToolMarkerVisibilityColumn>& InColumn
-		, const FNavigationToolItemPtr& InItem
+		, const FNavigationToolItemRef& InItem
 		, const TSharedRef<INavigationToolView>& InView
 		, const TSharedRef<SNavigationToolTreeRow>& InRowWidget);
 
diff --git a/Engine/Source/Editor/Sequencer/Private/Tools/NavigationTool/Widgets/Columns/SNavigationToolOutTime.cpp b/Engine/Source/Editor/Sequencer/Private/Tools/NavigationTool/Widgets/Columns/SNavigationToolOutTime.cpp
--- a/Engine/Source/Editor/Sequencer/Private/Tools/NavigationTool/Widgets/Columns/SNavigationToolOutTime.cpp
+++ b/Engine/Source/Editor/Sequencer/Private/Tools/NavigationTool/Widgets/Columns/SNavigationToolOutTime.cpp
@@ -13,7 +13,7 @@ namespace UE::Sequencer::NavigationTool
 {
 
 void SNavigationToolOutTime::Construct(const FArguments& InArgs
-	, const FNavigationToolItemPtr& InItem
+	, const FNavigationToolItemRef& InItem
 	, const TSharedRef<INavigationToolView>& InView
 	, const TSharedRef<SNavigationToolTreeRow>& InRowWidget)
 {
@@ -21,10 +21,7 @@ void SNavigationToolOutTime::Construct(const FArguments& InArgs
 	WeakView = InView;
 	WeakRowWidget = InRowWidget;
 
-	if (!WeakItem.IsValid()
-		|| !WeakView.IsValid()
-		|| !WeakRowWidget.IsValid()
-		|| !WeakItem.Pin()->IsA<IOutTimeExtension>())
+	if (!InItem->IsA<IOutTimeExtension>())
 	{
 		return;
 	}
diff --git a/Engine/Source/Editor/Sequencer/Private/Tools/NavigationTool/Widgets/Columns/SNavigationToolOutTime.h b/Engine/Source/Editor/Sequencer/Private/Tools/NavigationTool/Widgets/Columns/SNavigationToolOutTime.h
--- a/Engine/Source/Editor/Sequencer/Private/Tools/NavigationTool/Widgets/Columns/SNavigationToolOutTime.h
+++ b/Engine/Source/Editor/Sequencer/Private/Tools/NavigationTool/Widgets/Columns/SNavigationToolOutTime.h
@@ -21,7 +21,7 @@ public:
 	virtual ~SNavigationToolOutTime() override {}
 
 	void Construct(const FArguments& InArgs
-		, const FNavigationToolItemPtr& InItem
+		, const FNavigationToolItemRef& InItem
 		, const TSharedRef<INavigationToolView>& InView
 		, const TSharedRef<SNavigationToolTreeRow>& InRowWidget);
 
diff --git a/Engine/Source/Editor/Sequencer/Private/Tools/NavigationTool/Widgets/Columns/SNavigationToolPlayhead.cpp b/Engine/Source/Editor/Sequencer/Private/Tools/NavigationTool/Widgets/Columns/SNavigationToolPlayhead.cpp
--- a/Engine/Source/Editor/Sequencer/Private/Tools/NavigationTool/Widgets/Columns/SNavigationToolPlayhead.cpp
+++ b/Engine/Source/Editor/Sequencer/Private/Tools/NavigationTool/Widgets/Columns/SNavigationToolPlayhead.cpp
@@ -33,7 +33,7 @@ EItemContainsPlayhead GetItemContainsPlayhead(const FNavigationToolItemPtr& InIt
 
 void SNavigationToolPlayhead::Construct(const FArguments& InArgs
 	, const TSharedRef<FNavigationToolPlayheadColumn>& InColumn
-	, const FNavigationToolItemPtr& InItem
+	, const FNavigationToolItemRef& InItem
 	, const TSharedRef<INavigationToolView>& InView
 	, const TSharedRef<SNavigationToolTreeRow>& InRowWidget)
 {
diff --git a/Engine/Source/Editor/Sequencer/Private/Tools/NavigationTool/Widgets/Columns/SNavigationToolPlayhead.h b/Engine/Source/Editor/Sequencer/Private/Tools/NavigationTool/Widgets/Columns/SNavigationToolPlayhead.h
--- a/Engine/Source/Editor/Sequencer/Private/Tools/NavigationTool/Widgets/Columns/SNavigationToolPlayhead.h
+++ b/Engine/Source/Editor/Sequencer/Private/Tools/NavigationTool/Widgets/Columns/SNavigationToolPlayhead.h
@@ -22,7 +22,7 @@ public:
 
 	void Construct(const FArguments& InArgs
 		, const TSharedRef<FNavigationToolPlayheadColumn>& InColumn
-		, const FNavigationToolItemPtr& InItem
+		, const FNavigationToolItemRef& InItem
 		, const TSharedRef<INavigationToolView>& InView
 		, const TSharedRef<SNavigationToolTreeRow>& InRowWidget);
 
diff --git a/Engine/Source/Editor/Sequencer/Private/Tools/NavigationTool/Widgets/Columns/SNavigationToolRevisionControl.cpp b/Engine/Source/Editor/Sequencer/Private/Tools/NavigationTool/Widgets/Columns/SNavigationToolRevisionControl.cpp
--- a/Engine/Source/Editor/Sequencer/Private/Tools/NavigationTool/Widgets/Columns/SNavigationToolRevisionControl.cpp
+++ b/Engine/Source/Editor/Sequencer/Private/Tools/NavigationTool/Widgets/Columns/SNavigationToolRevisionControl.cpp
@@ -63,7 +63,7 @@ FText GetRevisionControlStatusText(const FNavigationToolItemPtr& InItem)
 
 void SNavigationToolRevisionControl::Construct(const FArguments& InArgs
 	, const TSharedRef<FNavigationToolRevisionControlColumn>& InColumn
-	, const FNavigationToolItemPtr& InItem
+	, const FNavigationToolItemRef& InItem
 	, const TSharedRef<INavigationToolView>& InView
 	, const TSharedRef<SNavigationToolTreeRow>& InRowWidget)
 {
diff --git a/Engine/Source/Editor/Sequencer/Private/Tools/NavigationTool/Widgets/Columns/SNavigationToolRevisionControl.h b/Engine/Source/Editor/Sequencer/Private/Tools/NavigationTool/Widgets/Columns/SNavigationToolRevisionControl.h
--- a/Engine/Source/Editor/Sequencer/Private/Tools/NavigationTool/Widgets/Columns/SNavigationToolRevisionControl.h
+++ b/Engine/Source/Editor/Sequencer/Private/Tools/NavigationTool/Widgets/Columns/SNavigationToolRevisionControl.h
@@ -20,7 +20,7 @@ public:
 
 	void Construct(const FArguments& InArgs
 		, const TSharedRef<FNavigationToolRevisionControlColumn>& InColumn
-		, const FNavigationToolItemPtr& InItem
+		, const FNavigationToolItemRef& InItem
 		, const TSharedRef<INavigationToolView>& InView
 		, const TSharedRef<SNavigationToolTreeRow>& InRowWidget);
 
diff --git a/Engine/Source/Editor/Sequencer/Private/Tools/NavigationTool/Widgets/Columns/SNavigationToolStartFrameOffset.cpp b/Engine/Source/Editor/Sequencer/Private/Tools/NavigationTool/Widgets/Columns/SNavigationToolStartFrameOffset.cpp
--- a/Engine/Source/Editor/Sequencer/Private/Tools/NavigationTool/Widgets/Columns/SNavigationToolStartFrameOffset.cpp
+++ b/Engine/Source/Editor/Sequencer/Private/Tools/NavigationTool/Widgets/Columns/SNavigationToolStartFrameOffset.cpp
@@ -1,11 +1,10 @@
 // Copyright Epic Games, Inc. All Rights Reserved.
 
 #include "SNavigationToolStartFrameOffset.h"
-#include "Sequencer.h"
+#include "Sections/MovieSceneSubSection.h"
 #include "Tools/NavigationTool/Items/NavigationToolItemUtils.h"
 #include "Tools/NavigationTool/NavigationToolView.h"
 #include "Tools/NavigationTool/Widgets/SNavigationToolTreeRow.h"
-#include "TrackEditors/SubTrackEditorBase.h"
 
 #define LOCTEXT_NAMESPACE "SNavigationToolStartFrameOffset"
 
@@ -13,7 +12,7 @@ namespace UE::Sequencer::NavigationTool
 {
 
 void SNavigationToolStartFrameOffset::Construct(const FArguments& InArgs
-	, const FNavigationToolItemPtr& InItem
+	, const FNavigationToolItemRef& InItem
 	, const TSharedRef<INavigationToolView>& InView
 	, const TSharedRef<SNavigationToolTreeRow>& InRowWidget)
 {
@@ -21,14 +20,7 @@ void SNavigationToolStartFrameOffset::Construct(const FArguments& InArgs
 	WeakView = InView;
 	WeakRowWidget = InRowWidget;
 
-	if (!WeakItem.IsValid()
-		|| !WeakView.IsValid()
-		|| !WeakRowWidget.IsValid())
-	{
-		return;
-	}
-
-	const FNavigationToolSequence* const SequenceItem = WeakItem.Pin()->CastTo<FNavigationToolSequence>();
+	const FNavigationToolSequence* const SequenceItem = InItem->CastTo<FNavigationToolSequence>();
 	if (!SequenceItem || !IsValid(SequenceItem->GetSubSection()))
 	{
 		return;
diff --git a/Engine/Source/Editor/Sequencer/Private/Tools/NavigationTool/Widgets/Columns/SNavigationToolStartFrameOffset.h b/Engine/Source/Editor/Sequencer/Private/Tools/NavigationTool/Widgets/Columns/SNavigationToolStartFrameOffset.h
--- a/Engine/Source/Editor/Sequencer/Private/Tools/NavigationTool/Widgets/Columns/SNavigationToolStartFrameOffset.h
+++ b/Engine/Source/Editor/Sequencer/Private/Tools/NavigationTool/Widgets/Columns/SNavigationToolStartFrameOffset.h
@@ -20,7 +20,7 @@ public:
 	virtual ~SNavigationToolStartFrameOffset() override {}
 
 	void Construct(const FArguments& InArgs
-		, const FNavigationToolItemPtr& InItem
+		, const FNavigationToolItemRef& InItem
 		, const TSharedRef<INavigationToolView>& InView
 		, const TSharedRef<SNavigationToolTreeRow>& InRowWidget);
 
diff --git a/Engine/Source/Editor/Sequencer/Private/Tools/NavigationTool/Widgets/Columns/SNavigationToolTake.cpp b/Engine/Source/Editor/Sequencer/Private/Tools/NavigationTool/Widgets/Columns/SNavigationToolTake.cpp
--- a/Engine/Source/Editor/Sequencer/Private/Tools/NavigationTool/Widgets/Columns/SNavigationToolTake.cpp
+++ b/Engine/Source/Editor/Sequencer/Private/Tools/NavigationTool/Widgets/Columns/SNavigationToolTake.cpp
@@ -1,11 +1,11 @@
 // Copyright Epic Games, Inc. All Rights Reserved.
 
 #include "SNavigationToolTake.h"
+#include "AssetRegistry/AssetData.h"
 #include "Framework/MultiBox/MultiBoxBuilder.h"
+#include "ISequencer.h"
 #include "MovieSceneSequence.h"
 #include "MovieSceneToolHelpers.h"
-#include "ScopedTransaction.h"
-#include "Sequencer.h"
 #include "Styling/StyleColors.h"
 #include "Tools/NavigationTool/Items/NavigationToolItemUtils.h"
 #include "Tools/NavigationTool/NavigationTool.h"
@@ -101,7 +101,7 @@ public:
 };
 
 void SNavigationToolTake::Construct(const FArguments& InArgs
-	, const FNavigationToolItemPtr& InItem
+	, const FNavigationToolItemRef& InItem
 	, const TSharedRef<INavigationToolView>& InView
 	, const TSharedRef<SNavigationToolTreeRow>& InRowWidget)
 {
@@ -115,11 +115,6 @@ void SNavigationToolTake::Construct(const FArguments& InArgs
 
 	if (!ActiveTakeInfo.IsValid())
 	{
-		ChildSlot
-		[
-			SNullWidget::NullWidget
-		];
-
 		return;
 	}
 
@@ -330,7 +325,6 @@ void SNavigationToolTake::SetActiveTake(UMovieSceneSequence* const InSequence)
 		return;
 	}
 
-	//FSequencer& Sequencer = Tool->GetSequencer();
 	const TSharedPtr<ISequencer> Sequencer = Tool->GetSequencer();
 	if (!Sequencer.IsValid())
 	{
diff --git a/Engine/Source/Editor/Sequencer/Private/Tools/NavigationTool/Widgets/Columns/SNavigationToolTake.h b/Engine/Source/Editor/Sequencer/Private/Tools/NavigationTool/Widgets/Columns/SNavigationToolTake.h
--- a/Engine/Source/Editor/Sequencer/Private/Tools/NavigationTool/Widgets/Columns/SNavigationToolTake.h
+++ b/Engine/Source/Editor/Sequencer/Private/Tools/NavigationTool/Widgets/Columns/SNavigationToolTake.h
@@ -8,14 +8,12 @@
 #include "Tools/NavigationTool/Columns/NavigationToolTakeColumn.h"
 #include "Widgets/SCompoundWidget.h"
 
-class SMenuAnchor;
 class UMovieSceneSequence;
 
 namespace UE::Sequencer::NavigationTool
 {
 
 class INavigationTool;
-class INavigationToolItem;
 class INavigationToolView;
 class SNavigationToolTreeRow;
 
@@ -34,7 +32,7 @@ public:
 	SLATE_END_ARGS()
 
 	void Construct(const FArguments& InArgs
-		, const FNavigationToolItemPtr& InItem
+		, const FNavigationToolItemRef& InItem
 		, const TSharedRef<INavigationToolView>& InView
 		, const TSharedRef<SNavigationToolTreeRow>& InRowWidget);
 
diff --git a/Engine/Source/Editor/Sequencer/Private/Tools/NavigationTool/Widgets/Columns/SNavigationToolTime.cpp b/Engine/Source/Editor/Sequencer/Private/Tools/NavigationTool/Widgets/Columns/SNavigationToolTime.cpp
--- a/Engine/Source/Editor/Sequencer/Private/Tools/NavigationTool/Widgets/Columns/SNavigationToolTime.cpp
+++ b/Engine/Source/Editor/Sequencer/Private/Tools/NavigationTool/Widgets/Columns/SNavigationToolTime.cpp
@@ -1,7 +1,7 @@
 ﻿// Copyright Epic Games, Inc. All Rights Reserved.
 
 #include "SNavigationToolTime.h"
-#include "Sequencer.h"
+#include "ISequencer.h"
 #include "Tools/NavigationTool/Items/NavigationToolItemUtils.h"
 #include "Tools/NavigationTool/NavigationToolStyle.h"
 #include "Tools/NavigationTool/NavigationToolView.h"
@@ -14,7 +14,7 @@ namespace UE::Sequencer::NavigationTool
 {
 
 void SNavigationToolTime::Construct(const FArguments& InArgs
-	, const FNavigationToolItemPtr& InItem
+	, const FNavigationToolItemRef& InItem
 	, const TSharedRef<INavigationToolView>& InView
 	, const TSharedRef<SNavigationToolTreeRow>& InRowWidget)
 {
@@ -22,13 +22,6 @@ void SNavigationToolTime::Construct(const FArguments& InArgs
 	WeakView = InView;
 	WeakRowWidget = InRowWidget;
 
-	if (!WeakItem.IsValid()
-		|| !WeakView.IsValid()
-		|| !WeakRowWidget.IsValid())
-	{
-		return;
-	}
-
 	ChildSlot
 	[
 		SNew(SSpinBox<double>)
@@ -95,9 +88,9 @@ TSharedPtr<INumericTypeInterface<double>> SNavigationToolTime::GetNumericTypeInt
 
 double SNavigationToolTime::GetDisplayRateDeltaFrameCount() const
 {
-	if (const TSharedPtr<FSequencer> Sequencer = StaticCastSharedPtr<FSequencer>(GetSequencer()))
+	if (const TSharedPtr<ISequencer> Sequencer = GetSequencer())
 	{
-		return Sequencer->GetDisplayRateDeltaFrameCount();
+		return Sequencer->GetFocusedTickResolution().AsDecimal() * Sequencer->GetFocusedDisplayRate().AsInterval();
 	}
 	return 0.;
 }
diff --git a/Engine/Source/Editor/Sequencer/Private/Tools/NavigationTool/Widgets/Columns/SNavigationToolTime.h b/Engine/Source/Editor/Sequencer/Private/Tools/NavigationTool/Widgets/Columns/SNavigationToolTime.h
--- a/Engine/Source/Editor/Sequencer/Private/Tools/NavigationTool/Widgets/Columns/SNavigationToolTime.h
+++ b/Engine/Source/Editor/Sequencer/Private/Tools/NavigationTool/Widgets/Columns/SNavigationToolTime.h
@@ -25,7 +25,7 @@ public:
 	virtual ~SNavigationToolTime() override {}
 
 	void Construct(const FArguments& InArgs
-		, const FNavigationToolItemPtr& InItem
+		, const FNavigationToolItemRef& InItem
 		, const TSharedRef<INavigationToolView>& InView
 		, const TSharedRef<SNavigationToolTreeRow>& InRowWidget);
 
diff --git a/Engine/Source/Editor/Sequencer/Private/Tools/NavigationTool/Widgets/SNavigationToolFilterBar.cpp b/Engine/Source/Editor/Sequencer/Private/Tools/NavigationTool/Widgets/SNavigationToolFilterBar.cpp
--- a/Engine/Source/Editor/Sequencer/Private/Tools/NavigationTool/Widgets/SNavigationToolFilterBar.cpp
+++ b/Engine/Source/Editor/Sequencer/Private/Tools/NavigationTool/Widgets/SNavigationToolFilterBar.cpp
@@ -5,7 +5,6 @@
 #include "Filters/Widgets/SFilterExpressionHelpDialog.h"
 #include "Filters/Widgets/SSequencerFilter.h"
 #include "MovieScene.h"
-#include "Sequencer.h"
 #include "SequencerLog.h"
 #include "Tools/NavigationTool/Filters/Filters/NavigationToolFilter_CustomText.h"
 #include "Tools/NavigationTool/Filters/Filters/NavigationToolFilterBase.h"
diff --git a/Engine/Source/Editor/Sequencer/Public/Tools/NavigationTool/Columns/INavigationToolColumn.h b/Engine/Source/Editor/Sequencer/Public/Tools/NavigationTool/Columns/INavigationToolColumn.h
--- a/Engine/Source/Editor/Sequencer/Public/Tools/NavigationTool/Columns/INavigationToolColumn.h
+++ b/Engine/Source/Editor/Sequencer/Public/Tools/NavigationTool/Columns/INavigationToolColumn.h
@@ -33,7 +33,7 @@ public:
 
 	virtual bool CanHideColumn(const FName InColumnId) const { return true; }
 
-	virtual TSharedRef<SWidget> ConstructRowWidget(const FNavigationToolItemPtr& InItem
+	virtual TSharedRef<SWidget> ConstructRowWidget(const FNavigationToolItemRef& InItem
 		, const TSharedRef<INavigationToolView>& InToolView
 		, const TSharedRef<SNavigationToolTreeRow>& InRow) = 0;
 
diff --git a/Engine/Source/Editor/Sequencer/Public/Tools/NavigationTool/Columns/NavigationToolColorColumn.h b/Engine/Source/Editor/Sequencer/Public/Tools/NavigationTool/Columns/NavigationToolColorColumn.h
--- a/Engine/Source/Editor/Sequencer/Public/Tools/NavigationTool/Columns/NavigationToolColorColumn.h
+++ b/Engine/Source/Editor/Sequencer/Public/Tools/NavigationTool/Columns/NavigationToolColorColumn.h
@@ -22,7 +22,7 @@ protected:
 	SEQUENCER_API virtual FText GetColumnDisplayNameText() const override;
 	virtual bool ShouldShowColumnByDefault() const override { return true; }
 	SEQUENCER_API virtual SHeaderRow::FColumn::FArguments ConstructHeaderRowColumn(const TSharedRef<INavigationToolView>& InToolView, const float InFillSize) override;
-	SEQUENCER_API virtual TSharedRef<SWidget> ConstructRowWidget(const FNavigationToolItemPtr& InItem
+	SEQUENCER_API virtual TSharedRef<SWidget> ConstructRowWidget(const FNavigationToolItemRef& InItem
 		, const TSharedRef<INavigationToolView>& InToolView
 		, const TSharedRef<SNavigationToolTreeRow>& InRow) override;
 	//~ End INavigationToolColumn
diff --git a/Engine/Source/Editor/Sequencer/Public/Tools/NavigationTool/Columns/NavigationToolCommentColumn.h b/Engine/Source/Editor/Sequencer/Public/Tools/NavigationTool/Columns/NavigationToolCommentColumn.h
--- a/Engine/Source/Editor/Sequencer/Public/Tools/NavigationTool/Columns/NavigationToolCommentColumn.h
+++ b/Engine/Source/Editor/Sequencer/Public/Tools/NavigationTool/Columns/NavigationToolCommentColumn.h
@@ -24,7 +24,7 @@ protected:
 	virtual bool ShouldShowColumnByDefault() const override { return false; }
 	virtual float GetFillWidth() const override { return 5.f; }
 	SEQUENCER_API virtual SHeaderRow::FColumn::FArguments ConstructHeaderRowColumn(const TSharedRef<INavigationToolView>& InToolView, const float InFillSize) override;
-	SEQUENCER_API virtual TSharedRef<SWidget> ConstructRowWidget(const FNavigationToolItemPtr& InItem
+	SEQUENCER_API virtual TSharedRef<SWidget> ConstructRowWidget(const FNavigationToolItemRef& InItem
 		, const TSharedRef<INavigationToolView>& InToolView
 		, const TSharedRef<SNavigationToolTreeRow>& InRow) override;
 	//~ End INavigationToolColumn
diff --git a/Engine/Source/Editor/Sequencer/Public/Tools/NavigationTool/Columns/NavigationToolDeactiveStateColumn.h b/Engine/Source/Editor/Sequencer/Public/Tools/NavigationTool/Columns/NavigationToolDeactiveStateColumn.h
--- a/Engine/Source/Editor/Sequencer/Public/Tools/NavigationTool/Columns/NavigationToolDeactiveStateColumn.h
+++ b/Engine/Source/Editor/Sequencer/Public/Tools/NavigationTool/Columns/NavigationToolDeactiveStateColumn.h
@@ -23,7 +23,7 @@ protected:
 	SEQUENCER_API virtual const FSlateBrush* GetIconBrush() const override;
 	SEQUENCER_API virtual SHeaderRow::FColumn::FArguments ConstructHeaderRowColumn(const TSharedRef<INavigationToolView>& InToolView, const float InFillSize) override;
 	virtual bool ShouldShowColumnByDefault() const override { return true; }
-	SEQUENCER_API virtual TSharedRef<SWidget> ConstructRowWidget(const FNavigationToolItemPtr& InItem
+	SEQUENCER_API virtual TSharedRef<SWidget> ConstructRowWidget(const FNavigationToolItemRef& InItem
 		, const TSharedRef<INavigationToolView>& InToolView
 		, const TSharedRef<SNavigationToolTreeRow>& InRow) override;
 	//~ End INavigationToolColumn
diff --git a/Engine/Source/Editor/Sequencer/Public/Tools/NavigationTool/Columns/NavigationToolHBiasColumn.h b/Engine/Source/Editor/Sequencer/Public/Tools/NavigationTool/Columns/NavigationToolHBiasColumn.h
--- a/Engine/Source/Editor/Sequencer/Public/Tools/NavigationTool/Columns/NavigationToolHBiasColumn.h
+++ b/Engine/Source/Editor/Sequencer/Public/Tools/NavigationTool/Columns/NavigationToolHBiasColumn.h
@@ -24,7 +24,7 @@ protected:
 	virtual bool ShouldShowColumnByDefault() const override { return false; }
 	virtual float GetFillWidth() const override { return 5.f; }
 	SEQUENCER_API virtual SHeaderRow::FColumn::FArguments ConstructHeaderRowColumn(const TSharedRef<INavigationToolView>& InToolView, const float InFillSize) override;
-	SEQUENCER_API virtual TSharedRef<SWidget> ConstructRowWidget(const FNavigationToolItemPtr& InItem
+	SEQUENCER_API virtual TSharedRef<SWidget> ConstructRowWidget(const FNavigationToolItemRef& InItem
 		, const TSharedRef<INavigationToolView>& InToolView
 		, const TSharedRef<SNavigationToolTreeRow>& InRow) override;
 	//~ End INavigationToolColumn
diff --git a/Engine/Source/Editor/Sequencer/Public/Tools/NavigationTool/Columns/NavigationToolInTimeColumn.h b/Engine/Source/Editor/Sequencer/Public/Tools/NavigationTool/Columns/NavigationToolInTimeColumn.h
--- a/Engine/Source/Editor/Sequencer/Public/Tools/NavigationTool/Columns/NavigationToolInTimeColumn.h
+++ b/Engine/Source/Editor/Sequencer/Public/Tools/NavigationTool/Columns/NavigationToolInTimeColumn.h
@@ -24,7 +24,7 @@ protected:
 	virtual bool ShouldShowColumnByDefault() const override { return true; }
 	virtual float GetFillWidth() const override { return 5.f; }
 	SEQUENCER_API virtual SHeaderRow::FColumn::FArguments ConstructHeaderRowColumn(const TSharedRef<INavigationToolView>& InToolView, const float InFillSize) override;
-	SEQUENCER_API virtual TSharedRef<SWidget> ConstructRowWidget(const FNavigationToolItemPtr& InItem
+	SEQUENCER_API virtual TSharedRef<SWidget> ConstructRowWidget(const FNavigationToolItemRef& InItem
 		, const TSharedRef<INavigationToolView>& InToolView
 		, const TSharedRef<SNavigationToolTreeRow>& InRow) override;
 	//~ End INavigationToolColumn
diff --git a/Engine/Source/Editor/Sequencer/Public/Tools/NavigationTool/Columns/NavigationToolItemsColumn.h b/Engine/Source/Editor/Sequencer/Public/Tools/NavigationTool/Columns/NavigationToolItemsColumn.h
--- a/Engine/Source/Editor/Sequencer/Public/Tools/NavigationTool/Columns/NavigationToolItemsColumn.h
+++ b/Engine/Source/Editor/Sequencer/Public/Tools/NavigationTool/Columns/NavigationToolItemsColumn.h
@@ -24,7 +24,7 @@ protected:
 	virtual bool ShouldShowColumnByDefault() const override { return true; }
 	virtual float GetFillWidth() const override { return 3.f; }
 	SEQUENCER_API virtual SHeaderRow::FColumn::FArguments ConstructHeaderRowColumn(const TSharedRef<INavigationToolView>& InToolView, const float InFillSize) override;
-	SEQUENCER_API virtual TSharedRef<SWidget> ConstructRowWidget(const FNavigationToolItemPtr& InItem
+	SEQUENCER_API virtual TSharedRef<SWidget> ConstructRowWidget(const FNavigationToolItemRef& InItem
 		, const TSharedRef<INavigationToolView>& InToolView
 		, const TSharedRef<SNavigationToolTreeRow>& InRow) override;
 	//~ End INavigationToolColumn
diff --git a/Engine/Source/Editor/Sequencer/Public/Tools/NavigationTool/Columns/NavigationToolLabelColumn.h b/Engine/Source/Editor/Sequencer/Public/Tools/NavigationTool/Columns/NavigationToolLabelColumn.h
--- a/Engine/Source/Editor/Sequencer/Public/Tools/NavigationTool/Columns/NavigationToolLabelColumn.h
+++ b/Engine/Source/Editor/Sequencer/Public/Tools/NavigationTool/Columns/NavigationToolLabelColumn.h
@@ -24,7 +24,7 @@ protected:
 	SEQUENCER_API virtual SHeaderRow::FColumn::FArguments ConstructHeaderRowColumn(const TSharedRef<INavigationToolView>& InToolView, const float InFillSize) override;
 	virtual float GetFillWidth() const override { return 10.f; }
 	virtual bool CanHideColumn(const FName InColumnId) const override { return false; }
-	SEQUENCER_API virtual TSharedRef<SWidget> ConstructRowWidget(const FNavigationToolItemPtr& InItem
+	SEQUENCER_API virtual TSharedRef<SWidget> ConstructRowWidget(const FNavigationToolItemRef& InItem
 		, const TSharedRef<INavigationToolView>& InToolView
 		, const TSharedRef<SNavigationToolTreeRow>& InRow) override;
 	//~ End INavigationToolColumn
diff --git a/Engine/Source/Editor/Sequencer/Public/Tools/NavigationTool/Columns/NavigationToolLengthColumn.h b/Engine/Source/Editor/Sequencer/Public/Tools/NavigationTool/Columns/NavigationToolLengthColumn.h
--- a/Engine/Source/Editor/Sequencer/Public/Tools/NavigationTool/Columns/NavigationToolLengthColumn.h
+++ b/Engine/Source/Editor/Sequencer/Public/Tools/NavigationTool/Columns/NavigationToolLengthColumn.h
@@ -24,7 +24,7 @@ protected:
 	virtual bool ShouldShowColumnByDefault() const override { return false; }
 	virtual float GetFillWidth() const override { return 5.f; }
 	SEQUENCER_API virtual SHeaderRow::FColumn::FArguments ConstructHeaderRowColumn(const TSharedRef<INavigationToolView>& InToolView, const float InFillSize) override;
-	SEQUENCER_API virtual TSharedRef<SWidget> ConstructRowWidget(const FNavigationToolItemPtr& InItem
+	SEQUENCER_API virtual TSharedRef<SWidget> ConstructRowWidget(const FNavigationToolItemRef& InItem
 		, const TSharedRef<INavigationToolView>& InToolView
 		, const TSharedRef<SNavigationToolTreeRow>& InRow) override;
 	//~ End INavigationToolColumn
diff --git a/Engine/Source/Editor/Sequencer/Public/Tools/NavigationTool/Columns/NavigationToolLockColumn.h b/Engine/Source/Editor/Sequencer/Public/Tools/NavigationTool/Columns/NavigationToolLockColumn.h
--- a/Engine/Source/Editor/Sequencer/Public/Tools/NavigationTool/Columns/NavigationToolLockColumn.h
+++ b/Engine/Source/Editor/Sequencer/Public/Tools/NavigationTool/Columns/NavigationToolLockColumn.h
@@ -23,7 +23,7 @@ protected:
 	SEQUENCER_API virtual const FSlateBrush* GetIconBrush() const override;
 	virtual bool ShouldShowColumnByDefault() const override { return true; }
 	SEQUENCER_API virtual SHeaderRow::FColumn::FArguments ConstructHeaderRowColumn(const TSharedRef<INavigationToolView>& InToolView, const float InFillSize) override;
-	SEQUENCER_API virtual TSharedRef<SWidget> ConstructRowWidget(const FNavigationToolItemPtr& InItem
+	SEQUENCER_API virtual TSharedRef<SWidget> ConstructRowWidget(const FNavigationToolItemRef& InItem
 		, const TSharedRef<INavigationToolView>& InToolView
 		, const TSharedRef<SNavigationToolTreeRow>& InRow) override;
 	//~ End INavigationToolColumn
diff --git a/Engine/Source/Editor/Sequencer/Public/Tools/NavigationTool/Columns/NavigationToolMarkerVisibilityColumn.h b/Engine/Source/Editor/Sequencer/Public/Tools/NavigationTool/Columns/NavigationToolMarkerVisibilityColumn.h
--- a/Engine/Source/Editor/Sequencer/Public/Tools/NavigationTool/Columns/NavigationToolMarkerVisibilityColumn.h
+++ b/Engine/Source/Editor/Sequencer/Public/Tools/NavigationTool/Columns/NavigationToolMarkerVisibilityColumn.h
@@ -23,7 +23,7 @@ protected:
 	SEQUENCER_API virtual const FSlateBrush* GetIconBrush() const override;
 	SEQUENCER_API virtual SHeaderRow::FColumn::FArguments ConstructHeaderRowColumn(const TSharedRef<INavigationToolView>& InToolView, const float InFillSize) override;
 	virtual bool ShouldShowColumnByDefault() const override { return true; }
-	SEQUENCER_API virtual TSharedRef<SWidget> ConstructRowWidget(const FNavigationToolItemPtr& InItem
+	SEQUENCER_API virtual TSharedRef<SWidget> ConstructRowWidget(const FNavigationToolItemRef& InItem
 		, const TSharedRef<INavigationToolView>& InToolView
 		, const TSharedRef<SNavigationToolTreeRow>& InRow) override;
 	//~ End INavigationToolColumn
diff --git a/Engine/Source/Editor/Sequencer/Public/Tools/NavigationTool/Columns/NavigationToolOutTimeColumn.h b/Engine/Source/Editor/Sequencer/Public/Tools/NavigationTool/Columns/NavigationToolOutTimeColumn.h
--- a/Engine/Source/Editor/Sequencer/Public/Tools/NavigationTool/Columns/NavigationToolOutTimeColumn.h
+++ b/Engine/Source/Editor/Sequencer/Public/Tools/NavigationTool/Columns/NavigationToolOutTimeColumn.h
@@ -24,7 +24,7 @@ protected:
 	virtual bool ShouldShowColumnByDefault() const override { return true; }
 	virtual float GetFillWidth() const override { return 5.f; }
 	SEQUENCER_API virtual SHeaderRow::FColumn::FArguments ConstructHeaderRowColumn(const TSharedRef<INavigationToolView>& InToolView, const float InFillSize) override;
-	SEQUENCER_API virtual TSharedRef<SWidget> ConstructRowWidget(const FNavigationToolItemPtr& InItem
+	SEQUENCER_API virtual TSharedRef<SWidget> ConstructRowWidget(const FNavigationToolItemRef& InItem
 		, const TSharedRef<INavigationToolView>& InToolView
 		, const TSharedRef<SNavigationToolTreeRow>& InRow) override;
 	//~ End INavigationToolColumn
diff --git a/Engine/Source/Editor/Sequencer/Public/Tools/NavigationTool/Columns/NavigationToolPlayheadColumn.h b/Engine/Source/Editor/Sequencer/Public/Tools/NavigationTool/Columns/NavigationToolPlayheadColumn.h
--- a/Engine/Source/Editor/Sequencer/Public/Tools/NavigationTool/Columns/NavigationToolPlayheadColumn.h
+++ b/Engine/Source/Editor/Sequencer/Public/Tools/NavigationTool/Columns/NavigationToolPlayheadColumn.h
@@ -23,7 +23,7 @@ protected:
 	SEQUENCER_API virtual const FSlateBrush* GetIconBrush() const override;
 	SEQUENCER_API virtual SHeaderRow::FColumn::FArguments ConstructHeaderRowColumn(const TSharedRef<INavigationToolView>& InToolView, const float InFillSize) override;
 	virtual bool ShouldShowColumnByDefault() const override { return true; }
-	SEQUENCER_API virtual TSharedRef<SWidget> ConstructRowWidget(const FNavigationToolItemPtr& InItem
+	SEQUENCER_API virtual TSharedRef<SWidget> ConstructRowWidget(const FNavigationToolItemRef& InItem
 		, const TSharedRef<INavigationToolView>& InToolView
 		, const TSharedRef<SNavigationToolTreeRow>& InRow) override;
 	//~ End INavigationToolColumn
diff --git a/Engine/Source/Editor/Sequencer/Public/Tools/NavigationTool/Columns/NavigationToolRevisionControlColumn.h b/Engine/Source/Editor/Sequencer/Public/Tools/NavigationTool/Columns/NavigationToolRevisionControlColumn.h
--- a/Engine/Source/Editor/Sequencer/Public/Tools/NavigationTool/Columns/NavigationToolRevisionControlColumn.h
+++ b/Engine/Source/Editor/Sequencer/Public/Tools/NavigationTool/Columns/NavigationToolRevisionControlColumn.h
@@ -23,7 +23,7 @@ protected:
 	SEQUENCER_API virtual const FSlateBrush* GetIconBrush() const override;
 	SEQUENCER_API virtual SHeaderRow::FColumn::FArguments ConstructHeaderRowColumn(const TSharedRef<INavigationToolView>& InToolView, const float InFillSize) override;
 	virtual bool ShouldShowColumnByDefault() const override { return true; }
-	SEQUENCER_API virtual TSharedRef<SWidget> ConstructRowWidget(const FNavigationToolItemPtr& InItem
+	SEQUENCER_API virtual TSharedRef<SWidget> ConstructRowWidget(const FNavigationToolItemRef& InItem
 		, const TSharedRef<INavigationToolView>& InToolView
 		, const TSharedRef<SNavigationToolTreeRow>& InRow) override;
 	//~ End INavigationToolColumn
diff --git a/Engine/Source/Editor/Sequencer/Public/Tools/NavigationTool/Columns/NavigationToolStartFrameOffsetColumn.h b/Engine/Source/Editor/Sequencer/Public/Tools/NavigationTool/Columns/NavigationToolStartFrameOffsetColumn.h
--- a/Engine/Source/Editor/Sequencer/Public/Tools/NavigationTool/Columns/NavigationToolStartFrameOffsetColumn.h
+++ b/Engine/Source/Editor/Sequencer/Public/Tools/NavigationTool/Columns/NavigationToolStartFrameOffsetColumn.h
@@ -24,8 +24,8 @@ protected:
 	virtual bool ShouldShowColumnByDefault() const override { return false; }
 	virtual float GetFillWidth() const override { return 5.f; }
 	SEQUENCER_API virtual SHeaderRow::FColumn::FArguments ConstructHeaderRowColumn(const TSharedRef<INavigationToolView>& InToolView, const float InFillSize) override;
-	SEQUENCER_API virtual TSharedRef<SWidget> ConstructRowWidget(const FNavigationToolItemPtr& InItem
-		, const TSharedRef<INavigationToolView>& InToolView
+	SEQUENCER_API virtual TSharedRef<SWidget> ConstructRowWidget(const FNavigationToolItemRef& InItem
+		, const TSharedRef<INavigationToolView>& InView
 		, const TSharedRef<SNavigationToolTreeRow>& InRow) override;
 	//~ End INavigationToolColumn
 };
diff --git a/Engine/Source/Editor/Sequencer/Public/Tools/NavigationTool/Columns/NavigationToolTakeColumn.h b/Engine/Source/Editor/Sequencer/Public/Tools/NavigationTool/Columns/NavigationToolTakeColumn.h
--- a/Engine/Source/Editor/Sequencer/Public/Tools/NavigationTool/Columns/NavigationToolTakeColumn.h
+++ b/Engine/Source/Editor/Sequencer/Public/Tools/NavigationTool/Columns/NavigationToolTakeColumn.h
@@ -24,7 +24,7 @@ protected:
 	virtual bool ShouldShowColumnByDefault() const override { return false; }
 	virtual float GetFillWidth() const override { return 5.f; }
 	SEQUENCER_API virtual SHeaderRow::FColumn::FArguments ConstructHeaderRowColumn(const TSharedRef<INavigationToolView>& InToolView, const float InFillSize) override;
-	SEQUENCER_API virtual TSharedRef<SWidget> ConstructRowWidget(const FNavigationToolItemPtr& InItem
+	SEQUENCER_API virtual TSharedRef<SWidget> ConstructRowWidget(const FNavigationToolItemRef& InItem
 		, const TSharedRef<INavigationToolView>& InToolView
 		, const TSharedRef<SNavigationToolTreeRow>& InRow) override;
 	//~ End INavigationToolColumn
diff --git a/Engine/Source/Editor/Sequencer/Public/Tools/NavigationTool/INavigationTool.h b/Engine/Source/Editor/Sequencer/Public/Tools/NavigationTool/INavigationTool.h
--- a/Engine/Source/Editor/Sequencer/Public/Tools/NavigationTool/INavigationTool.h
+++ b/Engine/Source/Editor/Sequencer/Public/Tools/NavigationTool/INavigationTool.h
@@ -16,22 +16,14 @@
 #include "Tools/NavigationTool/NavigationToolDefines.h"
 #include "Templates/SharedPointer.h"
 
-class AActor;
-class FArchive;
-class FSequencer;
 class FUICommandList;
-class UMovieSceneSequence;
-enum class EItemDropZone;
-struct FAttachmentTransformRules;
 template<typename OptionalType> struct TOptional;
 
 namespace UE::Sequencer::NavigationTool
 {
 
-class FNavigationToolEditorSelection;
 class FNavigationToolProvider;
 class INavigationToolAction;
-class INavigationToolFilterBar;
 class INavigationToolView;
 struct FNavigationToolItemId;
 
diff --git a/Engine/Source/Editor/Sequencer/Public/Tools/NavigationTool/NavigationToolDefines.h b/Engine/Source/Editor/Sequencer/Public/Tools/NavigationTool/NavigationToolDefines.h
--- a/Engine/Source/Editor/Sequencer/Public/Tools/NavigationTool/NavigationToolDefines.h
+++ b/Engine/Source/Editor/Sequencer/Public/Tools/NavigationTool/NavigationToolDefines.h
@@ -14,6 +14,7 @@ namespace UE::Sequencer::NavigationTool
 
 class INavigationToolItem;
 
+using FNavigationToolItemRef = TSharedRef<INavigationToolItem>;
 using FNavigationToolItemPtr = TSharedPtr<INavigationToolItem>;
 using FNavigationToolItemWeakPtr = TWeakPtr<INavigationToolItem>;
 
