Sequencer: FMovieSceneTrackEditor::AnimatablePropertyChanged no loner creates a transaction if a transaction is already ongoing

#rnx

diff --git a/Engine/Source/Editor/Sequencer/Private/MovieSceneTrackEditor.cpp b/Engine/Source/Editor/Sequencer/Private/MovieSceneTrackEditor.cpp
--- a/Engine/Source/Editor/Sequencer/Private/MovieSceneTrackEditor.cpp
+++ b/Engine/Source/Editor/Sequencer/Private/MovieSceneTrackEditor.cpp
@@ -112,9 +112,14 @@ void FMovieSceneTrackEditor::AnimatablePropertyChanged( FOnKeyProperty OnKeyProp
 		// @todo Sequencer - The sequencer probably should have taken care of this
 		MovieSceneSequence->SetFlags(RF_Transactional);
 	
-		// Create a transaction record because we are about to add keys
-		const bool bShouldActuallyTransact = !GIsTransacting;		// Don't transact if we're recording in a PIE world.  That type of keyframe capture cannot be undone.
-		FScopedTransaction AutoKeyTransaction( NSLOCTEXT("AnimatablePropertyTool", "PropertyChanged", "Animatable Property Changed"), bShouldActuallyTransact );
+		// Only create a transaction if the editor is not transacting already
+		TOptional<FScopedTransaction> AutoKeyTransaction;
+		if (!GEditor->IsTransactionActive())
+		{
+			// Create a transaction record because we are about to add keys
+			const bool bShouldActuallyTransact = !GIsTransacting;		// Don't transact if we're recording in a PIE world.  That type of keyframe capture cannot be undone.
+			AutoKeyTransaction = FScopedTransaction(NSLOCTEXT("AnimatablePropertyTool", "PropertyChanged", "Animatable Property Changed"), bShouldActuallyTransact);
+		}
 
 		FKeyPropertyResult KeyPropertyResult = OnKeyProperty.Execute( KeyTime );
 
