Sequencer: Check for hidden function before Interp

The Camera Component has FieldOfView as Interp property but the CineCameraComponent declares SetFieldOfView as a HiddenFunction since it wants to control the field of view value through the focal length. This fixes FieldOfView being considered animatable by checking for HiddenFunction before blindly allowing all Interp properties to be animatable.

diff --git a/Engine/Source/Editor/Sequencer/Private/SequencerObjectChangeListener.cpp b/Engine/Source/Editor/Sequencer/Private/SequencerObjectChangeListener.cpp
--- a/Engine/Source/Editor/Sequencer/Private/SequencerObjectChangeListener.cpp
+++ b/Engine/Source/Editor/Sequencer/Private/SequencerObjectChangeListener.cpp
@@ -219,12 +219,6 @@ const FOnAnimatablePropertyChanged* FSequencerObjectChangeListener::FindProperty
 			return nullptr;
 		}
 
-		// Interp properties are always keyable
-		if (PropertyOrContainer->HasAnyPropertyFlags(CPF_Interp))
-		{
-			return DelegatePtr;
-		}
-
 		if (bCanApplyFunction)
 		{
 			// If the function is hidden, we cannot use it
@@ -232,7 +226,16 @@ const FOnAnimatablePropertyChanged* FSequencerObjectChangeListener::FindProperty
 			{
 				return nullptr;
 			}
+		}
 
+		// Interp properties are always keyable
+		if (PropertyOrContainer->HasAnyPropertyFlags(CPF_Interp))
+		{
+			return DelegatePtr;
+		}
+
+		if (bCanApplyFunction)
+		{
 			// If there is a native setter we can always animate the property
 			if (Property.HasSetter())
 			{
