[Sequencer] Changes needed to move Navigation Tool to it's own plugin

Add interface functions:
InvalidateGlobalMarkedFramesCache()
GetOutlinerViewWidget()

Widgets moved to public:
SFilterBarClippingHorizontalBox
SSequencerSearchBox
SFilterExpressionHelpDialog
SSequencerFilter

diff --git a/Engine/Source/Editor/Sequencer/Private/Filters/SequencerFilterBar.cpp b/Engine/Source/Editor/Sequencer/Private/Filters/SequencerFilterBar.cpp
--- a/Engine/Source/Editor/Sequencer/Private/Filters/SequencerFilterBar.cpp
+++ b/Engine/Source/Editor/Sequencer/Private/Filters/SequencerFilterBar.cpp
@@ -19,7 +19,9 @@
 #include "Filters/SequencerTrackFilterCollection.h"
 #include "Filters/SequencerTrackFilterCommands.h"
 #include "Filters/SequencerTrackFilterExtension.h"
+#include "Filters/Widgets/SFilterExpressionHelpDialog.h"
 #include "Filters/Widgets/SSequencerFilterBar.h"
+#include "Filters/Widgets/SSequencerSearchBox.h"
 #include "Framework/Application/SlateApplication.h"
 #include "Framework/MultiBox/SToolBarButtonBlock.h"
 #include "Menus/SequencerTrackFilterMenu.h"
@@ -33,8 +35,6 @@
 #include "Stats/StatsMisc.h"
 #include "UObject/UObjectIterator.h"
 #include "Widgets/SFilterBarIsolateHideShow.h"
-#include "Widgets/SFilterExpressionHelpDialog.h"
-#include "Widgets/SSequencerSearchBox.h"
 
 using namespace UE::Sequencer;
 
diff --git a/Engine/Source/Editor/Sequencer/Private/Filters/Widgets/SFilterBarClippingHorizontalBox.cpp b/Engine/Source/Editor/Sequencer/Private/Filters/Widgets/SFilterBarClippingHorizontalBox.cpp
--- a/Engine/Source/Editor/Sequencer/Private/Filters/Widgets/SFilterBarClippingHorizontalBox.cpp
+++ b/Engine/Source/Editor/Sequencer/Private/Filters/Widgets/SFilterBarClippingHorizontalBox.cpp
@@ -1,6 +1,6 @@
 ﻿// Copyright Epic Games, Inc. All Rights Reserved.
 
-#include "SFilterBarClippingHorizontalBox.h"
+#include "Filters/Widgets/SFilterBarClippingHorizontalBox.h"
 #include "Framework/Application/SlateApplication.h"
 #include "Layout/ArrangedChildren.h"
 #include "Styling/CoreStyle.h"
diff --git a/Engine/Source/Editor/Sequencer/Private/Filters/Widgets/SFilterBarClippingHorizontalBox.h b/Engine/Source/Editor/Sequencer/Private/Filters/Widgets/SFilterBarClippingHorizontalBox.h
deleted file mode 100644
diff --git a/Engine/Source/Editor/Sequencer/Private/Filters/Widgets/SFilterExpressionHelpDialog.cpp b/Engine/Source/Editor/Sequencer/Private/Filters/Widgets/SFilterExpressionHelpDialog.cpp
--- a/Engine/Source/Editor/Sequencer/Private/Filters/Widgets/SFilterExpressionHelpDialog.cpp
+++ b/Engine/Source/Editor/Sequencer/Private/Filters/Widgets/SFilterExpressionHelpDialog.cpp
@@ -1,6 +1,6 @@
 ﻿// Copyright Epic Games, Inc. All Rights Reserved.
 
-#include "SFilterExpressionHelpDialog.h"
+#include "Filters/Widgets/SFilterExpressionHelpDialog.h"
 #include "Filters/SequencerTextFilterExpressionContext.h"
 #include "Framework/Application/SlateApplication.h"
 #include "Interfaces/IMainFrameModule.h"
diff --git a/Engine/Source/Editor/Sequencer/Private/Filters/Widgets/SFilterExpressionHelpDialog.h b/Engine/Source/Editor/Sequencer/Private/Filters/Widgets/SFilterExpressionHelpDialog.h
deleted file mode 100644
diff --git a/Engine/Source/Editor/Sequencer/Private/Filters/Widgets/SSequencerFilter.h b/Engine/Source/Editor/Sequencer/Private/Filters/Widgets/SSequencerFilter.h
deleted file mode 100644
diff --git a/Engine/Source/Editor/Sequencer/Private/Filters/Widgets/SSequencerFilterBar.cpp b/Engine/Source/Editor/Sequencer/Private/Filters/Widgets/SSequencerFilterBar.cpp
--- a/Engine/Source/Editor/Sequencer/Private/Filters/Widgets/SSequencerFilterBar.cpp
+++ b/Engine/Source/Editor/Sequencer/Private/Filters/Widgets/SSequencerFilterBar.cpp
@@ -7,13 +7,12 @@
 #include "Filters/SequencerFilterBar.h"
 #include "Filters/Widgets/SFilterBarClippingHorizontalBox.h"
 #include "Filters/Widgets/SFilterExpressionHelpDialog.h"
+#include "Filters/Widgets/SSequencerFilter.h"
+#include "Filters/Widgets/SSequencerSearchBox.h"
 #include "MovieScene.h"
 #include "Sequencer.h"
 #include "SequencerLog.h"
-#include "SSequencerFilter.h"
 #include "SSequencerCustomTextFilterDialog.h"
-#include "Tools/NavigationTool/Menus/NavigationToolFilterContextMenu.h"
-#include "Widgets/SSequencerSearchBox.h"
 
 #define LOCTEXT_NAMESPACE "SSequencerFilterBar"
 
diff --git a/Engine/Source/Editor/Sequencer/Private/Filters/Widgets/SSequencerSearchBox.cpp b/Engine/Source/Editor/Sequencer/Private/Filters/Widgets/SSequencerSearchBox.cpp
new file mode 100644
--- /dev/null
+++ b/Engine/Source/Editor/Sequencer/Private/Filters/Widgets/SSequencerSearchBox.cpp
@@ -0,0 +1,683 @@
+﻿// Copyright Epic Games, Inc. All Rights Reserved.
+
+#include "Filters/Widgets/SSequencerSearchBox.h"
+#include "Filters/ISequencerFilterBar.h"
+#include "Filters/Filters/SequencerTrackFilter_Text.h"
+#include "Filters/Widgets/SSequencerFilterSuggestionListRow.h"
+#include "Framework/Application/SlateApplication.h"
+#include "SequencerSettings.h"
+#include "Widgets/Input/SMenuAnchor.h"
+#include "Widgets/Layout/SBorder.h"
+#include "Widgets/Layout/SBox.h"
+#include "Widgets/SBoxPanel.h"
+#include "Widgets/Views/SListView.h"
+
+class ITableRow;
+class STableViewBase;
+class SWidget;
+struct FGeometry;
+
+#define LOCTEXT_NAMESPACE "SSequencerSearchBox"
+
+/** Case sensitive hashing function for TMap */
+template <typename ValueType>
+struct FSequencerSearchCategoryKeyMapFuncs : BaseKeyFuncs<ValueType, FText, /*bInAllowDuplicateKeys*/false>
+{
+	static FORCEINLINE const FString& GetSourceString(const FText& InText)
+	{
+		const FString* SourceString = FTextInspector::GetSourceString(InText);
+		check(SourceString);
+		return *SourceString;
+	}
+	static FORCEINLINE const FText& GetSetKey(const TPair<FText, ValueType>& Element)
+	{
+		return Element.Key;
+	}
+	static FORCEINLINE bool Matches(const FText& A, const FText& B)
+	{
+		return GetSourceString(A).Equals(GetSourceString(B), ESearchCase::CaseSensitive);
+	}
+	static FORCEINLINE uint32 GetKeyHash(const FText& Key)
+	{
+		return FLocKey::ProduceHash(GetSourceString(Key));
+	}
+};
+
+void SSequencerSearchBox::Construct(const FArguments& InArgs, const TWeakPtr<ISequencerFilterBar>& InWeakFilterBar)
+{
+	WeakFilterBar = InWeakFilterBar;
+
+	OnTextChanged = InArgs._OnTextChanged;
+	OnTextCommitted = InArgs._OnTextCommitted;
+	OnKeyDownHandler = InArgs._OnKeyDownHandler;
+	PossibleSuggestions = InArgs._PossibleSuggestions;
+	OnSuggestionFilter = InArgs._OnSuggestionFilter;
+	OnSuggestionChosen = InArgs._OnSuggestionChosen;
+	PreCommittedText = InArgs._InitialText.Get();
+	bMustMatchPossibleSuggestions = InArgs._MustMatchPossibleSuggestions.Get();
+
+	if (!OnSuggestionFilter.IsBound())
+	{
+		OnSuggestionFilter.BindSP(this, &SSequencerSearchBox::DefaultSuggestionFilterImpl);
+	}
+
+	if (!OnSuggestionChosen.IsBound())
+	{
+		OnSuggestionChosen.BindSP(this, &SSequencerSearchBox::DefaultSuggestionChosenImpl);
+	}
+
+	ChildSlot
+		[
+			SAssignNew(SuggestionMenuAnchor, SMenuAnchor)
+			.Placement(InArgs._SuggestionListPlacement)
+			[
+				/** Use SFilterSearchBox internally to add the ability to show search history and save searches as filters */
+				SAssignNew(SearchBox, SFilterSearchBox)
+				.InitialText(InArgs._InitialText)
+				.HintText(InArgs._HintText)
+				.ShowSearchHistory(InArgs._ShowSearchHistory)
+				.DelayChangeNotificationsWhileTyping(InArgs._DelayChangeNotificationsWhileTyping)
+				.OnTextChanged(this, &SSequencerSearchBox::HandleTextChanged)
+				.OnTextCommitted(this, &SSequencerSearchBox::HandleTextCommitted)
+				.OnKeyDownHandler(this, &SSequencerSearchBox::HandleKeyDown)
+				.OnSaveSearchClicked(InArgs._OnSaveSearchClicked)
+			]
+			.MenuContent(GetSuggestionListMenuContent())
+		];
+}
+
+TSharedRef<SWidget> SSequencerSearchBox::GetSuggestionListMenuContent()
+{
+	const USequencerSettings* const SequencerSettings = GetDefault<USequencerSettings>();
+	check(IsValid(SequencerSettings));
+
+	return SNew(SBorder)
+		.BorderImage(FAppStyle::GetBrush(TEXT("Menu.Background")))
+		.Padding(FMargin(2))
+		[
+			SNew(SBox)
+			.MinDesiredWidth(180)
+			.MinDesiredHeight(16)
+			.MaxDesiredHeight(260)
+			[
+				SAssignNew(SuggestionListView, SListView<TSharedPtr<FSequencerFilterSuggestionListEntryBase>>)
+				.ListItemsSource(&SuggestionItems)
+				.SelectionMode(ESelectionMode::Single) // Ideally the mouse over would not highlight while keyboard controls the UI
+				.OnGenerateRow(this, &SSequencerSearchBox::MakeSuggestionListItemWidget)
+				.OnSelectionChanged(this, &SSequencerSearchBox::OnSelectionChanged)
+				.ScrollbarDragFocusCause(EFocusCause::SetDirectly) // Use SetDirect so that clicking the scrollbar doesn't close the suggestions list
+				.HeaderRow
+				(
+					SNew(SHeaderRow)
+					.Visibility(EVisibility::Collapsed)
+					+ SHeaderRow::Column("Suggestion")
+					.DefaultLabel(LOCTEXT("SuggestionColumnHeader", "Suggestion"))
+					.FixedWidth(180.f)
+					.VAlignCell(VAlign_Center)
+					.HAlignCell(HAlign_Left)
+					+ SHeaderRow::Column("Description")
+					.DefaultLabel(LOCTEXT("DescriptionColumnHeader", "Description"))
+					.FillWidth(1.f)
+					.VAlignCell(VAlign_Center)
+					.HAlignCell(HAlign_Left)
+				)
+			]
+		];
+}
+
+FText SSequencerSearchBox::GetText() const
+{
+	return SearchBox->GetText();
+}
+
+void SSequencerSearchBox::SetText(const TAttribute<FText>& InNewText)
+{
+	SearchBox->SetText(InNewText);
+	PreCommittedText = InNewText.Get();
+
+	// Set the cursor interaction location to the end of the newly inserted suggestion
+	if (LastChosenCursorOffset != INDEX_NONE)
+	{
+		const FTextLocation CursorLocation = FTextLocation(0, LastChosenCursorOffset);
+		SearchBox->SelectText(CursorLocation, CursorLocation);
+
+		LastChosenCursorOffset = INDEX_NONE;
+	}
+}
+
+void SSequencerSearchBox::SetError(const FText& InError)
+{
+	SearchBox->SetError(InError);
+}
+
+void SSequencerSearchBox::SetError(const FString& InError)
+{
+	SearchBox->SetError(InError);
+}
+
+FReply SSequencerSearchBox::OnPreviewKeyDown(const FGeometry& InGeometry, const FKeyEvent& InKeyEvent)
+{
+	if (SuggestionMenuAnchor->IsOpen() && InKeyEvent.GetKey() == EKeys::Escape)
+	{
+		// Clear any selection first to prevent the currently selection being set in the text box
+		SuggestionListView->ClearSelection();
+		SuggestionMenuAnchor->SetIsOpen(false, false);
+
+		return FReply::Handled();
+	}
+
+	return FReply::Unhandled();
+}
+
+FReply SSequencerSearchBox::HandleKeyDown(const FGeometry& InGeometry, const FKeyEvent& InKeyEvent)
+{
+	if (!SearchBox->GetText().IsEmpty()
+		&& (InKeyEvent.GetKey() == EKeys::Up || InKeyEvent.GetKey() == EKeys::Down))
+	{
+		if (!SuggestionMenuAnchor->IsOpen())
+		{
+			SuggestionMenuAnchor->SetIsOpen(true, false);
+		}
+
+		const bool bSelectingUp = InKeyEvent.GetKey() == EKeys::Up;
+		const TSharedPtr<FSequencerFilterSuggestionListEntryBase> SelectedSuggestion = GetSelectedSuggestion();
+
+		int32 TargetIdx = INDEX_NONE;
+		if (SelectedSuggestion.IsValid())
+		{
+			const int32 SelectionDirection = bSelectingUp ? -1 : 1;
+
+			// Select the next non-header suggestion, based on the direction of travel
+			TargetIdx = SuggestionItems.IndexOfByKey(SelectedSuggestion);
+			if (SuggestionItems.IsValidIndex(TargetIdx))
+			{
+				do
+				{
+					TargetIdx += SelectionDirection;
+				}
+				while (SuggestionItems.IsValidIndex(TargetIdx) && SuggestionItems[TargetIdx]->IsHeader());
+			}
+		}
+		else if (!bSelectingUp && SuggestionItems.Num() > 0)
+		{
+			// Nothing selected and pressed down, select the first non-header suggestion
+			TargetIdx = 0;
+			while (SuggestionItems.IsValidIndex(TargetIdx) && SuggestionItems[TargetIdx]->IsHeader())
+			{
+				TargetIdx += 1;
+			}
+		}
+
+		if (SuggestionItems.IsValidIndex(TargetIdx))
+		{
+			SuggestionListView->SetSelection(SuggestionItems[TargetIdx]);
+			SuggestionListView->RequestScrollIntoView(SuggestionItems[TargetIdx]);
+		}
+
+		return FReply::Handled();
+	}
+
+	if (OnKeyDownHandler.IsBound())
+	{
+		return OnKeyDownHandler.Execute(InGeometry, InKeyEvent);
+	}
+
+	return FReply::Unhandled();
+}
+
+bool SSequencerSearchBox::SupportsKeyboardFocus() const
+{
+	return SearchBox->SupportsKeyboardFocus();
+}
+
+bool SSequencerSearchBox::HasKeyboardFocus() const
+{
+	// Since keyboard focus is forwarded to our editable text, we will test it instead
+	return SearchBox->HasKeyboardFocus();
+}
+
+FReply SSequencerSearchBox::OnFocusReceived(const FGeometry& InGeometry, const FFocusEvent& InFocusEvent)
+{
+	// Forward keyboard focus to our editable text widget
+	return SearchBox->OnFocusReceived(InGeometry, InFocusEvent);
+}
+
+void SSequencerSearchBox::HandleTextChanged(const FText& InNewText)
+{
+	LastCursorLocation = SearchBox->GetSelection().GetEnd();
+
+	OnTextChanged.ExecuteIfBound(InNewText);
+
+	UpdateSuggestionList();
+}
+
+void SSequencerSearchBox::HandleTextCommitted(const FText& InNewText, const ETextCommit::Type InCommitType)
+{
+	TSharedPtr<FSequencerFilterSuggestionListEntryBase> SelectedSuggestion = GetSelectedSuggestion();
+
+	FText CommittedText;
+	if (SelectedSuggestion.IsValid() && !SelectedSuggestion->IsHeader() && InCommitType != ETextCommit::OnCleared)
+	{
+		CommittedText = OnSuggestionChosen.Execute(InNewText, SelectedSuggestion->AsSuggestionEntry()->Suggestion.Suggestion);
+	}
+	else
+	{
+		if (InCommitType == ETextCommit::OnCleared)
+		{
+			// Clear text when escape is pressed then commit an empty string
+			CommittedText = FText::GetEmpty();
+		}
+		else if (bMustMatchPossibleSuggestions)
+		{
+			const FString NewTextStr = InNewText.ToString();
+			const bool bIsSuggestion = PossibleSuggestions.Get().ContainsByPredicate(
+				[NewTextStr](const FSequencerFilterSuggestion& InSuggestion)
+				{
+					return InSuggestion.Suggestion == NewTextStr;
+				});
+			if (bIsSuggestion)
+			{
+				CommittedText = InNewText;
+			}
+			else
+			{
+				// commit the original text if we have to match a suggestion
+				CommittedText = PreCommittedText;
+			}
+		}
+		else
+		{
+			// otherwise, set the typed text
+			CommittedText = InNewText;
+		}
+	}
+
+	// Set the text and execute the delegate
+	bDisableOpeningSuggestions = true;
+	SetText(CommittedText);
+	OnTextCommitted.ExecuteIfBound(CommittedText, InCommitType);
+
+	if (InCommitType != ETextCommit::Default)
+	{
+		// Clear the suggestion box if the user has navigated away or set their own text
+		SuggestionMenuAnchor->SetIsOpen(false, false);
+	}
+}
+
+void SSequencerSearchBox::OnSelectionChanged(TSharedPtr<FSequencerFilterSuggestionListEntryBase> InNewValue, const ESelectInfo::Type InSelectInfo)
+{
+	// If the user clicked directly on an item to select it, then accept the choice and close the window
+	if (InSelectInfo == ESelectInfo::OnMouseClick && !InNewValue->IsHeader())
+	{
+		const FText NewText = OnSuggestionChosen.Execute(SearchBox->GetText(), InNewValue->AsSuggestionEntry()->Suggestion.Suggestion);
+		SetText(NewText);
+
+		SuggestionMenuAnchor->SetIsOpen(false, false);
+		FocusEditBox();
+	}
+}
+
+TSharedRef<ITableRow> SSequencerSearchBox::MakeSuggestionListItemWidget(const TSharedPtr<FSequencerFilterSuggestionListEntryBase> InSuggestion, const TSharedRef<STableViewBase>& InOwnerTable)
+{
+	return SNew(SSequencerFilterSuggestionListRow, InOwnerTable)
+		.ListItem(InSuggestion)
+		.HighlightText(this, &SSequencerSearchBox::GetHighlightText);
+}
+
+FText SSequencerSearchBox::GetHighlightText() const
+{
+	return SuggestionHighlightText;
+}
+
+void SSequencerSearchBox::UpdateSuggestionList()
+{
+	FText SearchText = SearchBox->GetText();
+
+	SuggestionItems.Reset();
+	SuggestionHighlightText = FText::GetEmpty();
+
+	if (!SearchText.IsEmpty())
+	{
+		typedef TMap<FText
+			, TArray<TSharedPtr<FSequencerFilterSuggestionListEntryBase>>
+			, FDefaultSetAllocator
+			, FSequencerSearchCategoryKeyMapFuncs<TArray<TSharedPtr<FSequencerFilterSuggestionListEntryBase>>>> FCategorizedSuggestionsMap;
+
+		FParsedExpression ParsedExpression;
+		ExtractSearchFilterTerms(SearchText, ParsedExpression);
+
+		// Get the potential suggestions
+		TArray<FSequencerFilterSuggestion> FilteredSuggestions;
+		if (PossibleSuggestions.IsBound())
+		{
+			FilteredSuggestions = PossibleSuggestions.Get();
+		}
+		else
+		{
+			if (!ParsedExpression.Key.IsSet() && ParsedExpression.Value.IsSet())
+			{
+				DefaultKeySuggestions(SearchText.ToString(), FilteredSuggestions);
+			}
+			else if (ParsedExpression.Key.IsSet())
+			{
+				DefaultValueSuggestions(ParsedExpression.Key.GetValue(), FilteredSuggestions);
+			}
+
+			if (ParsedExpression.Value.IsSet())
+			{
+				SearchText = FText::FromString(ParsedExpression.Value.GetValue());
+			}
+		}
+
+		// Run them through the filter
+		OnSuggestionFilter.Execute(SearchText, FilteredSuggestions, SuggestionHighlightText);
+
+		// Split the suggestions list into categories
+		FCategorizedSuggestionsMap CategorizedSuggestions;
+		for (const FSequencerFilterSuggestion& Suggestion : FilteredSuggestions)
+		{
+			TArray<TSharedPtr<FSequencerFilterSuggestionListEntryBase>>& CategorySuggestions = CategorizedSuggestions.FindOrAdd(Suggestion.CategoryName);
+
+			const TSharedRef<FSequencerFilterSuggestionListEntry> NewSuggestionEntry = MakeShared<FSequencerFilterSuggestionListEntry>();
+			NewSuggestionEntry->Suggestion = Suggestion;
+
+			CategorySuggestions.Add(NewSuggestionEntry);
+		}
+
+		// Rebuild the flat list in categorized groups
+		// If there is only one category, and that category is empty (undefined), then skip adding the category headers
+		const bool bSkipCategoryHeaders = CategorizedSuggestions.Num() == 1 && CategorizedSuggestions.Contains(FText::GetEmpty());
+		for (const TPair<FText, TArray<TSharedPtr<FSequencerFilterSuggestionListEntryBase>>>& CategorySuggestionsPair : CategorizedSuggestions)
+		{
+			if (!bSkipCategoryHeaders)
+			{
+				const FText CategoryDisplayName = CategorySuggestionsPair.Key.IsEmpty()
+					? LOCTEXT("UndefinedCategory", "Undefined")
+					: CategorySuggestionsPair.Key;
+
+				const TSharedRef<FSequencerFilterSuggestionListHeaderEntry> NewHeaderEntry = MakeShared<FSequencerFilterSuggestionListHeaderEntry>(CategoryDisplayName);
+
+				SuggestionItems.Add(NewHeaderEntry);
+			}
+			SuggestionItems.Append(CategorySuggestionsPair.Value);
+		}
+	}
+
+	if (!bDisableOpeningSuggestions && SuggestionItems.Num() > 0 && HasKeyboardFocus())
+	{
+		SuggestionMenuAnchor->SetIsOpen(true, false);
+	}
+	else
+	{
+		SuggestionMenuAnchor->SetIsOpen(false, false);
+	}
+
+	bDisableOpeningSuggestions = false;
+
+	SuggestionListView->RequestListRefresh();
+}
+
+void SSequencerSearchBox::FocusEditBox()
+{
+	FWidgetPath WidgetToFocusPath;
+	FSlateApplication::Get().GeneratePathToWidgetUnchecked(SearchBox.ToSharedRef(), WidgetToFocusPath);
+	FSlateApplication::Get().SetKeyboardFocus(WidgetToFocusPath, EFocusCause::SetDirectly);
+}
+
+TSharedPtr<FSequencerFilterSuggestionListEntryBase> SSequencerSearchBox::GetSelectedSuggestion() const
+{
+	TSharedPtr<FSequencerFilterSuggestionListEntryBase> SelectedSuggestion;
+
+	if (SuggestionMenuAnchor->IsOpen())
+	{
+		const TArray<TSharedPtr<FSequencerFilterSuggestionListEntryBase>>& SelectedSuggestionList = SuggestionListView->GetSelectedItems();
+		if (SelectedSuggestionList.Num() > 0)
+		{
+			// Selection mode is Single, so there should only be one suggestion at the most
+			SelectedSuggestion = SelectedSuggestionList[0];
+		}
+	}
+
+	return SelectedSuggestion;
+}
+
+void SSequencerSearchBox::SetOnSaveSearchHandler(SFilterSearchBox::FOnSaveSearchClicked InOnSaveSearchHandler)
+{
+	SearchBox->SetOnSaveSearchHandler(InOnSaveSearchHandler);
+}
+
+void SSequencerSearchBox::ExtractSearchFilterTerms(const FText& InSearchText, FParsedExpression& OutParsedPair) const
+{
+	const TSharedPtr<ISequencerFilterBar> FilterBar = WeakFilterBar.Pin();
+	if (!FilterBar.IsValid())
+	{
+		return;
+	}
+
+	const FTextFilterExpressionEvaluator& TextFilterExpressionEvaluator = FilterBar->GetTextFilterExpressionEvaluator();
+	const TArray<FExpressionToken>& ExpressionTokens = TextFilterExpressionEvaluator.GetFilterExpressionTokens();
+	if (ExpressionTokens.Num() == 0)
+	{
+		return;
+	}
+
+	const FString SearchString = InSearchText.ToString();
+
+	OutParsedPair.SuggestionIndex = SearchString.Len();
+
+	const int32 CaretLocationIndex = LastCursorLocation.GetOffset();
+	const int32 CaretTokenIndex = FindTokenIndex(ExpressionTokens, CaretLocationIndex);
+
+	if (CaretTokenIndex == INDEX_NONE)
+	{
+		return;
+	}
+
+	// Inspect the tokens to see what the last part of the search term was
+	// If it was a key->value pair then we'll use that to control what kinds of results we show
+	// For anything else we just use the text from the last token as our filter term to allow incremental auto-complete
+	const FExpressionToken& LastToken = ExpressionTokens[CaretTokenIndex];
+
+	const int32 SecondToLastIndex = CaretTokenIndex - 1;
+	const int32 ThirdToLastIndex = CaretTokenIndex - 2;
+
+	// If the last token is a text token, then consider it as a value and walk back to see if we also have a key
+	if (LastToken.Node.Cast<TextFilterExpressionParser::FTextToken>())
+	{
+		OutParsedPair.Value = LastToken.Context.GetString();
+		OutParsedPair.SuggestionIndex = LastToken.Context.GetCharacterIndex();
+
+		if (ExpressionTokens.IsValidIndex(ThirdToLastIndex))
+		{
+			// Check if the second from last token is a operator (=, !=, <, >, etc)
+			const FExpressionToken& ComparisonToken = ExpressionTokens[SecondToLastIndex];
+			if (IsOperatorToken(ComparisonToken))
+			{
+				const FExpressionToken& KeyToken = ExpressionTokens[ThirdToLastIndex];
+				if (KeyToken.Node.Cast<TextFilterExpressionParser::FTextToken>())
+				{
+					OutParsedPair.Key = KeyToken.Context.GetString();
+					OutParsedPair.SuggestionIndex = KeyToken.Context.GetCharacterIndex();
+				}
+			}
+			else if (IsLogicalOperatorToken(ComparisonToken))
+			{
+				const FExpressionToken& KeyToken = ExpressionTokens[SecondToLastIndex];
+				if (KeyToken.Node.Cast<TextFilterExpressionParser::FTextToken>())
+				{
+					OutParsedPair.Key = KeyToken.Context.GetString();
+					OutParsedPair.SuggestionIndex = LastToken.Context.GetCharacterIndex();
+				}
+			}
+		}
+	}
+	// If the last token is a comparison operator, then walk back and see if we have a key
+	else if (IsOperatorToken(LastToken))
+	{
+		if (ExpressionTokens.IsValidIndex(SecondToLastIndex))
+		{
+			const FExpressionToken& KeyToken = ExpressionTokens[SecondToLastIndex];
+			if (KeyToken.Node.Cast<TextFilterExpressionParser::FTextToken>())
+			{
+				OutParsedPair.Key = KeyToken.Context.GetString();
+				OutParsedPair.Value = TEXT("");
+				OutParsedPair.SuggestionIndex = LastToken.Context.GetCharacterIndex();
+			}
+		}
+	}
+}
+
+void SSequencerSearchBox::DefaultSuggestionFilterImpl(const FText& InSearchText, TArray<FSequencerFilterSuggestion>& OutPossibleSuggestions, FText& OutSuggestionHighlightText)
+{
+	OutPossibleSuggestions.RemoveAll([SearchStr = InSearchText.ToString()](const FSequencerFilterSuggestion& InSuggestion)
+		{
+			return !InSuggestion.Suggestion.Contains(SearchStr);
+		});
+
+	OutSuggestionHighlightText = InSearchText;
+}
+
+FText SSequencerSearchBox::DefaultSuggestionChosenImpl(const FText& InSearchText, const FString& InSuggestion)
+{
+	const FTextFilterExpressionEvaluator* const TextFilterExpressionEvaluator = GetTextFilterExpressionEvaluator();
+	if (!TextFilterExpressionEvaluator)
+	{
+		return InSearchText;
+	}
+
+	const TArray<FExpressionToken>& ExpressionTokens = TextFilterExpressionEvaluator->GetFilterExpressionTokens();
+	const int32 CaretLocationIndex = LastCursorLocation.GetOffset();
+	const int32 CaretTokenIndex = FindTokenIndex(ExpressionTokens, CaretLocationIndex);
+	const FExpressionToken& CaretToken = ExpressionTokens[CaretTokenIndex];
+	const int32 SuggestionInsertionIndex = CaretToken.Context.GetCharacterIndex();
+	const FString CaretTokenString = CaretToken.Context.GetString();
+
+	// Replace the value with the suggestion
+	FString SearchString = InSearchText.ToString();
+	SearchString.RemoveAt(SuggestionInsertionIndex, CaretTokenString.Len(), EAllowShrinking::No);
+	SearchString.InsertAt(SuggestionInsertionIndex, InSuggestion);
+
+	// Set the cursor interaction location to the end of the newly inserted suggestion
+	LastChosenCursorOffset = SuggestionInsertionIndex + InSuggestion.Len();
+
+	return FText::FromString(SearchString);
+}
+
+void SSequencerSearchBox::DefaultKeySuggestions(const FString& InKeyValue, TArray<FSequencerFilterSuggestion>& OutPossibleSuggestions)
+{
+	const TSharedPtr<ISequencerFilterBar> FilterBar = WeakFilterBar.Pin();
+	if (!FilterBar.IsValid())
+	{
+		return;
+	}
+
+	const TArray<TSharedRef<ISequencerTextFilterExpressionContext>>& FilterExpressions = FilterBar->GetTextFilterExpressionContexts();
+
+	for (const TSharedRef<ISequencerTextFilterExpressionContext>& Expression : FilterExpressions)
+	{
+		const TArray<FName> Keys = Expression->GetKeys().Array();
+		for (const FName Key : Keys)
+		{
+			const FString KeyString = Key.ToString();
+
+			FSequencerFilterSuggestion NewSuggestion;
+			NewSuggestion.Suggestion = FString::Printf(TEXT("%s"), *KeyString);
+			NewSuggestion.DisplayName = FText::FromString(KeyString);
+			NewSuggestion.Description = Expression->GetDescription();
+
+			OutPossibleSuggestions.Add(NewSuggestion);
+		}
+	}
+
+	OutPossibleSuggestions.Sort([](const FSequencerFilterSuggestion& InA, const FSequencerFilterSuggestion& InB)
+		{
+			return InA.DisplayName.CompareTo(InB.DisplayName) < 0;
+		});
+}
+
+void SSequencerSearchBox::DefaultValueSuggestions(const FString& InKeyValue, TArray<FSequencerFilterSuggestion>& OutPossibleSuggestions)
+{
+	const TSharedPtr<ISequencerFilterBar> FilterBar = WeakFilterBar.Pin();
+	if (!FilterBar.IsValid())
+	{
+		return;
+	}
+
+	const TArray<TSharedRef<ISequencerTextFilterExpressionContext>> FilterExpressions = FilterBar->GetTextFilterExpressionContexts();
+
+	FSequencerTextFilterKeyword FoundKeyword;
+	TSharedPtr<ISequencerTextFilterExpressionContext> FoundExpression;
+
+	for (const TSharedRef<ISequencerTextFilterExpressionContext>& Expression : FilterExpressions)
+	{
+		for (const FName Key : Expression->GetKeys())
+		{
+			if (InKeyValue.Equals(Key.ToString(), ESearchCase::IgnoreCase))
+			{
+				FoundExpression = Expression;
+				goto end;
+			}
+		}
+	}
+
+end:
+	if (FoundExpression.IsValid())
+	{
+		const TArray<FSequencerTextFilterKeyword> TextFilterKeywords = FoundExpression->GetValueKeywords();
+		for (const FSequencerTextFilterKeyword& TextFilterKeyword : TextFilterKeywords)
+		{
+			FSequencerFilterSuggestion NewSuggestion;
+			NewSuggestion.Suggestion = FString::Printf(TEXT("%s"), *TextFilterKeyword.Keyword);
+			NewSuggestion.DisplayName = FText::FromString(TextFilterKeyword.Keyword);
+			NewSuggestion.Description = TextFilterKeyword.Description;
+
+			OutPossibleSuggestions.Add(NewSuggestion);
+		}
+
+		OutPossibleSuggestions.Sort([](const FSequencerFilterSuggestion& InA, const FSequencerFilterSuggestion& InB)
+			{
+				return InA.DisplayName.CompareTo(InB.DisplayName) < 0;
+			});
+	}
+}
+
+bool SSequencerSearchBox::IsOperatorToken(const FExpressionToken& InToken)
+{
+	return InToken.Node.Cast<TextFilterExpressionParser::FEqual>()
+		|| InToken.Node.Cast<TextFilterExpressionParser::FNotEqual>()
+		|| InToken.Node.Cast<TextFilterExpressionParser::FLess>()
+		|| InToken.Node.Cast<TextFilterExpressionParser::FLessOrEqual>()
+		|| InToken.Node.Cast<TextFilterExpressionParser::FGreater>()
+		|| InToken.Node.Cast<TextFilterExpressionParser::FGreaterOrEqual>();
+};
+
+bool SSequencerSearchBox::IsLogicalOperatorToken(const FExpressionToken& InToken)
+{
+	return InToken.Node.Cast<TextFilterExpressionParser::FAnd>()
+		|| InToken.Node.Cast<TextFilterExpressionParser::FOr>();
+};
+
+int32 SSequencerSearchBox::FindTokenIndex(const TArray<FExpressionToken>& ExpressionTokens, const int32 InIndexToFind)
+{
+	for (int32 Index = ExpressionTokens.Num() - 1; Index > 0; --Index)
+	{
+		const int32 CharacterIndex = ExpressionTokens[Index].Context.GetCharacterIndex();
+		const int32 TokenLength = ExpressionTokens[Index].Context.GetString().Len();
+		if (InIndexToFind > CharacterIndex)// && InIndexToFind <= CharacterIndex + TokenLength)
+		{
+			return Index;
+		}
+	}
+	return ExpressionTokens.IsEmpty() ? INDEX_NONE : 0;
+}
+
+const FTextFilterExpressionEvaluator* SSequencerSearchBox::GetTextFilterExpressionEvaluator() const
+{
+	if (const TSharedPtr<ISequencerFilterBar> FilterBar = WeakFilterBar.Pin())
+	{
+		return &FilterBar->GetTextFilterExpressionEvaluator();
+	}
+	return nullptr;
+}
+
+#undef LOCTEXT_NAMESPACE
diff --git a/Engine/Source/Editor/Sequencer/Private/SSequencer.cpp b/Engine/Source/Editor/Sequencer/Private/SSequencer.cpp
--- a/Engine/Source/Editor/Sequencer/Private/SSequencer.cpp
+++ b/Engine/Source/Editor/Sequencer/Private/SSequencer.cpp
@@ -138,7 +138,7 @@
 #include "Sidebar/SSidebar.h"
 #include "Sidebar/SSidebarContainer.h"
 #include "Widgets/SOverlay.h"
-#include "Widgets/SSequencerSearchBox.h"
+#include "Filters/Widgets/SSequencerSearchBox.h"
 #include "IContentBrowserSingleton.h"
 #include "ContentBrowserModule.h"
 #include "Menus/SequencerToolbarUtils.h"
diff --git a/Engine/Source/Editor/Sequencer/Private/Sequencer.cpp b/Engine/Source/Editor/Sequencer/Private/Sequencer.cpp
--- a/Engine/Source/Editor/Sequencer/Private/Sequencer.cpp
+++ b/Engine/Source/Editor/Sequencer/Private/Sequencer.cpp
@@ -12988,4 +12988,9 @@ TSharedPtr<UE::Sequencer::NavigationTool::INavigationTool> FSequencer::GetNaviga
 	return NavigationTool;
 }
 
+TSharedPtr<UE::Sequencer::SOutlinerView> FSequencer::GetOutlinerViewWidget() const
+{
+	return SequencerWidget->GetTreeView();
+}
+
 #undef LOCTEXT_NAMESPACE
diff --git a/Engine/Source/Editor/Sequencer/Private/Sequencer.h b/Engine/Source/Editor/Sequencer/Private/Sequencer.h
--- a/Engine/Source/Editor/Sequencer/Private/Sequencer.h
+++ b/Engine/Source/Editor/Sequencer/Private/Sequencer.h
@@ -341,7 +341,7 @@ public:
 	TArray<FMovieSceneMarkedFrame> GetMarkedFrames() const override;
 
 	TArray<FMovieSceneMarkedFrame> GetGlobalMarkedFrames() const;
-	void InvalidateGlobalMarkedFramesCache() { bGlobalMarkedFramesCached = false; }
+	virtual void InvalidateGlobalMarkedFramesCache() override { bGlobalMarkedFramesCached = false; }
 	void UpdateGlobalMarkedFramesCache();
 
 	/** Toggle whether to show marked frames globally */
@@ -1281,6 +1281,8 @@ public:
 
 	virtual TSharedPtr<UE::Sequencer::NavigationTool::INavigationTool> GetNavigationTool() const override;
 
+	virtual TSharedPtr<UE::Sequencer::SOutlinerView> GetOutlinerViewWidget() const override;
+
 private:
 
 	/** Update the time bases for the current movie scene */
diff --git a/Engine/Source/Editor/Sequencer/Private/Tools/NavigationTool/Filters/NavigationToolFilterBar.cpp b/Engine/Source/Editor/Sequencer/Private/Tools/NavigationTool/Filters/NavigationToolFilterBar.cpp
--- a/Engine/Source/Editor/Sequencer/Private/Tools/NavigationTool/Filters/NavigationToolFilterBar.cpp
+++ b/Engine/Source/Editor/Sequencer/Private/Tools/NavigationTool/Filters/NavigationToolFilterBar.cpp
@@ -3,6 +3,7 @@
 #include "Tools/NavigationTool/Filters/NavigationToolFilterBar.h"
 #include "Filters/SequencerFilterBarConfig.h"
 #include "Filters/Widgets/SFilterExpressionHelpDialog.h"
+#include "Filters/Widgets/SSequencerSearchBox.h"
 #include "Framework/Application/SlateApplication.h"
 #include "Framework/MultiBox/SToolBarButtonBlock.h"
 #include "ISequencer.h"
@@ -25,7 +26,6 @@
 #include "Tools/NavigationTool/Widgets/SNavigationToolCustomTextFilterDialog.h"
 #include "Tools/NavigationTool/Widgets/SNavigationToolFilterBar.h"
 #include "UObject/UObjectIterator.h"
-#include "Widgets/SSequencerSearchBox.h"
 
 #define LOCTEXT_NAMESPACE "NavigationToolFilterBar"
 
diff --git a/Engine/Source/Editor/Sequencer/Private/Tools/NavigationTool/Menus/NavigationToolToolbarMenu.cpp b/Engine/Source/Editor/Sequencer/Private/Tools/NavigationTool/Menus/NavigationToolToolbarMenu.cpp
--- a/Engine/Source/Editor/Sequencer/Private/Tools/NavigationTool/Menus/NavigationToolToolbarMenu.cpp
+++ b/Engine/Source/Editor/Sequencer/Private/Tools/NavigationTool/Menus/NavigationToolToolbarMenu.cpp
@@ -1,6 +1,7 @@
 // Copyright Epic Games, Inc. All Rights Reserved.
 
 #include "NavigationToolToolbarMenu.h"
+#include "Filters/Widgets/SSequencerSearchBox.h"
 #include "ToolMenu.h"
 #include "ToolMenus.h"
 #include "Tools/NavigationTool/Columns/INavigationToolColumn.h"
@@ -13,7 +14,6 @@
 #include "Tools/NavigationTool/NavigationToolCommands.h"
 #include "Tools/NavigationTool/NavigationToolSettings.h"
 #include "Tools/NavigationTool/NavigationToolView.h"
-#include "Widgets/SSequencerSearchBox.h"
 
 #define LOCTEXT_NAMESPACE "NavigationToolToolbarMenu"
 
diff --git a/Engine/Source/Editor/Sequencer/Private/Tools/NavigationTool/Widgets/SNavigationToolFilterBar.cpp b/Engine/Source/Editor/Sequencer/Private/Tools/NavigationTool/Widgets/SNavigationToolFilterBar.cpp
--- a/Engine/Source/Editor/Sequencer/Private/Tools/NavigationTool/Widgets/SNavigationToolFilterBar.cpp
+++ b/Engine/Source/Editor/Sequencer/Private/Tools/NavigationTool/Widgets/SNavigationToolFilterBar.cpp
@@ -4,6 +4,7 @@
 #include "Filters/Widgets/SFilterBarClippingHorizontalBox.h"
 #include "Filters/Widgets/SFilterExpressionHelpDialog.h"
 #include "Filters/Widgets/SSequencerFilter.h"
+#include "Filters/Widgets/SSequencerSearchBox.h"
 #include "MovieScene.h"
 #include "SequencerLog.h"
 #include "Tools/NavigationTool/Filters/Filters/NavigationToolFilter_CustomText.h"
@@ -13,7 +14,6 @@
 #include "Tools/NavigationTool/Menus/NavigationToolFilterContextMenu.h"
 #include "Tools/NavigationTool/NavigationToolSettings.h"
 #include "Tools/NavigationTool/Widgets/SNavigationToolCustomTextFilterDialog.h"
-#include "Widgets/SSequencerSearchBox.h"
 #include "Widgets/SWindow.h"
 
 using namespace UE::Sequencer;
diff --git a/Engine/Source/Editor/Sequencer/Private/Widgets/SSequencerSearchBox.cpp b/Engine/Source/Editor/Sequencer/Private/Widgets/SSequencerSearchBox.cpp
deleted file mode 100644
diff --git a/Engine/Source/Editor/Sequencer/Private/Widgets/SSequencerSearchBox.h b/Engine/Source/Editor/Sequencer/Private/Widgets/SSequencerSearchBox.h
deleted file mode 100644
diff --git a/Engine/Source/Editor/Sequencer/Public/Filters/Widgets/SFilterBarClippingHorizontalBox.h b/Engine/Source/Editor/Sequencer/Public/Filters/Widgets/SFilterBarClippingHorizontalBox.h
new file mode 100644
--- /dev/null
+++ b/Engine/Source/Editor/Sequencer/Public/Filters/Widgets/SFilterBarClippingHorizontalBox.h
@@ -0,0 +1,55 @@
+﻿// Copyright Epic Games, Inc. All Rights Reserved.
+
+#pragma once
+
+#include "Widgets/DeclarativeSyntaxSupport.h"
+#include "Framework/SlateDelegates.h"
+#include "Widgets/SBoxPanel.h"
+
+class SComboButton;
+
+/**
+ * Specialized widget for handling horizontal clipping of filter widgets that go beyond the filter bar widget width.
+ */
+class SEQUENCER_API SFilterBarClippingHorizontalBox : public SHorizontalBox
+{
+public:
+	static TSharedRef<SWidget> WrapVerticalListWithHeading(const TSharedRef<SWidget>& InWidget, const FPointerEventHandler InMouseButtonUpEvent);
+
+	SLATE_BEGIN_ARGS(SFilterBarClippingHorizontalBox) 
+		: _IsFocusable(true)
+	{}
+		SLATE_ARGUMENT(FOnGetContent, OnWrapButtonClicked)
+		SLATE_ARGUMENT(bool, IsFocusable)
+	SLATE_END_ARGS()
+
+	void Construct(const FArguments& InArgs);
+
+	TSharedRef<SComboButton> CreateWrapButton();
+
+	/** Returns to index of the first clipped child */
+	int32 GetClippedIndex() const { return ClippedIndex; }
+
+protected:
+	//~ Begin SPanel
+	virtual void OnArrangeChildren(const FGeometry& InAllottedGeometry, FArrangedChildren& ArrangedChildren) const override;
+	//~ End SPanel
+
+	/** The button that is displayed when a toolbar or menubar is clipped */
+	TWeakPtr<SComboButton> WeakWrapButton;
+
+	/** Callback for when the wrap button is clicked */
+	FOnGetContent OnWrapButtonClicked;
+
+	/** Index of the first clipped child/block */
+	mutable int32 ClippedIndex = INDEX_NONE;
+	/** Index of the last clipped child/block */
+	mutable int32 LastClippedIndex = INDEX_NONE;
+
+	/** Number of clipped children not including the wrap button */
+	mutable int32 NumClippedChildren = 0;
+
+	TSharedPtr<FActiveTimerHandle> WrapButtonOpenTimer;
+
+	bool bIsFocusable = true;
+};
diff --git a/Engine/Source/Editor/Sequencer/Public/Filters/Widgets/SFilterExpressionHelpDialog.h b/Engine/Source/Editor/Sequencer/Public/Filters/Widgets/SFilterExpressionHelpDialog.h
new file mode 100644
--- /dev/null
+++ b/Engine/Source/Editor/Sequencer/Public/Filters/Widgets/SFilterExpressionHelpDialog.h
@@ -0,0 +1,54 @@
+﻿// Copyright Epic Games, Inc. All Rights Reserved.
+
+#pragma once
+
+#include "Widgets/SWindow.h"
+
+class ISequencerTextFilterExpressionContext;
+enum class ESequencerTextFilterValueType : uint8;
+
+struct SEQUENCER_API FFilterExpressionHelpDialogConfig
+{
+	static constexpr float DefaultMaxDesiredWidth = 460.f;
+	static constexpr float DefaultMaxDesiredHeight = 560.f;
+
+	FFilterExpressionHelpDialogConfig();
+
+	FName IdentifierName;
+
+	FText DialogTitle;
+	FString DocumentationLink;
+
+	TArray<TSharedRef<ISequencerTextFilterExpressionContext>> TextFilterExpressionContexts;
+
+	float MaxDesiredWidth = DefaultMaxDesiredWidth;
+	float MaxDesiredHeight = DefaultMaxDesiredHeight;
+};
+
+class SFilterExpressionHelpDialog : public SWindow
+{
+public:
+	SEQUENCER_API static void Open(FFilterExpressionHelpDialogConfig&& InConfig);
+
+	static bool IsOpen(const FName InName);
+
+	static void CloseWindow(const FName InName);
+
+	void Construct(const FArguments& InArgs, FFilterExpressionHelpDialogConfig&& InConfig);
+
+protected:
+	static const FSlateColor KeyColor;
+	static const FSlateColor ValueColor;
+
+	static TMap<FName, TSharedPtr<SFilterExpressionHelpDialog>> DialogInstance;
+
+	TSharedRef<SWidget> ConstructDialogHeader();
+	TSharedRef<SWidget> ConstructExpressionWidgetList();
+	TSharedRef<SWidget> ConstructExpressionWidget(const TSharedPtr<ISequencerTextFilterExpressionContext>& InExpressionContext);
+	TSharedRef<SWidget> ConstructKeysWidget(const TSet<FName>& InKeys);
+	TSharedRef<SWidget> ConstructValueWidget(const ESequencerTextFilterValueType InValueType);
+
+	void OpenDocumentationLink() const;
+
+	FFilterExpressionHelpDialogConfig Config;
+};
diff --git a/Engine/Source/Editor/Sequencer/Public/Filters/Widgets/SSequencerFilter.h b/Engine/Source/Editor/Sequencer/Public/Filters/Widgets/SSequencerFilter.h
new file mode 100644
--- /dev/null
+++ b/Engine/Source/Editor/Sequencer/Public/Filters/Widgets/SSequencerFilter.h
@@ -0,0 +1,82 @@
+﻿// Copyright Epic Games, Inc. All Rights Reserved.
+
+#pragma once
+
+#include "Filters/SBasicFilterBar.h"
+
+class SSequencerFilterCheckBox;
+
+/**
+ * Generic Sequencer filter class used by Sequencer and Navigation Tool filter bars.
+ */
+class SSequencerFilter : public SCompoundWidget
+{
+public:
+	DECLARE_DELEGATE_RetVal(bool, FOnIsFilterActive);
+	DECLARE_DELEGATE_OneParam(FOnFilterToggle, const ECheckBoxState /*InNewState*/);
+
+	DECLARE_DELEGATE_RetVal(TSharedRef<SWidget>, FOnGetMenuContent);
+
+	SLATE_BEGIN_ARGS(SSequencerFilter) {}
+
+		/** Determines how each individual filter pill looks like */
+		SLATE_ARGUMENT(EFilterPillStyle, FilterPillStyle)
+
+		SLATE_EVENT(FOnIsFilterActive, OnIsFilterActive)
+		SLATE_EVENT(FOnFilterToggle, OnFilterToggle)
+		SLATE_EVENT(FSimpleDelegate, OnCtrlClick)
+		SLATE_EVENT(FSimpleDelegate, OnAltClick)
+		SLATE_EVENT(FSimpleDelegate, OnMiddleClick)
+		SLATE_EVENT(FSimpleDelegate, OnDoubleClick)
+
+		SLATE_ATTRIBUTE(FText, DisplayName)
+		SLATE_ATTRIBUTE(FText, ToolTipText)
+		SLATE_ATTRIBUTE(FSlateColor, BlockColor)
+
+		SLATE_EVENT(FOnGetMenuContent, OnGetMenuContent)
+
+	SLATE_END_ARGS()
+
+	SEQUENCER_API void Construct(const FArguments& InArgs);
+
+protected:
+	TSharedRef<SWidget> ConstructBasicFilterWidget();
+	TSharedRef<SWidget> ConstructDefaultFilterWidget();
+
+	bool IsActive() const;
+
+	void OnFilterToggled(const ECheckBoxState NewState);
+
+	FReply OnFilterCtrlClick();
+	FReply OnFilterAltClick();
+	FReply OnFilterMiddleButtonClick();
+	FReply OnFilterDoubleClick();
+
+	TSharedRef<SWidget> GetRightClickMenuContent();
+
+	ECheckBoxState IsChecked() const;
+
+	FSlateColor GetFilterImageColorAndOpacity() const;
+
+	EVisibility GetFilterOverlayVisibility() const;
+
+	FMargin GetFilterNamePadding() const;
+
+	bool IsButtonEnabled() const;
+
+	FOnIsFilterActive IsFilterActiveDelegate;
+
+	FOnFilterToggle FilterToggleDelegate;
+	FSimpleDelegate CtrlClickDelegate;
+	FSimpleDelegate AltClickDelegate;
+	FSimpleDelegate MiddleClickDelegate;
+	FSimpleDelegate DoubleClickDelegate;
+
+	TAttribute<FText> DisplayName;
+	TAttribute<FText> ToolTipText;
+	TAttribute<FSlateColor> BlockColor;
+
+	FOnGetMenuContent GetMenuContentDelegate;
+
+	TSharedPtr<SSequencerFilterCheckBox> ToggleButtonPtr;
+};
diff --git a/Engine/Source/Editor/Sequencer/Public/Filters/Widgets/SSequencerSearchBox.h b/Engine/Source/Editor/Sequencer/Public/Filters/Widgets/SSequencerSearchBox.h
new file mode 100644
--- /dev/null
+++ b/Engine/Source/Editor/Sequencer/Public/Filters/Widgets/SSequencerSearchBox.h
@@ -0,0 +1,203 @@
+// Copyright Epic Games, Inc. All Rights Reserved.
+
+#pragma once
+
+#include "Filters/SequencerFilterSuggestion.h"
+#include "Filters/SFilterSearchBox.h"
+#include "Framework/Text/TextLayout.h"
+#include "Misc/ExpressionParserTypesFwd.h"
+#include "Widgets/DeclarativeSyntaxSupport.h"
+#include "Widgets/SCompoundWidget.h"
+
+class FSequencerFilterSuggestionListEntryBase;
+class FTextFilterExpressionEvaluator;
+class ISequencerFilterBar;
+class ITableRow;
+class SMenuAnchor;
+class STableViewBase;
+struct FFocusEvent;
+struct FGeometry;
+struct FKeyEvent;
+
+/** A delegate for a callback to filter the given suggestion list, to allow custom filtering behavior */
+DECLARE_DELEGATE_ThreeParams(FOnSequencerSearchBoxSuggestionFilter, const FText& /*InSearchText*/, TArray<FSequencerFilterSuggestion>& /*InPossibleSuggestions*/, FText& /*InSuggestionHighlightText*/);
+
+/** A delegate for a callback when a suggestion entry is chosen during a search, to allow custom compositing behavior of the suggestion into the search text */
+DECLARE_DELEGATE_RetVal_TwoParams(FText, FOnSequencerSearchBoxSuggestionChosen, const FText& /*InSearchText*/, const FString& /*InSuggestion*/);
+
+/**
+ * A wrapper widget around SFilterSearchBox to provide filter text expression suggestions in a dropdown menu.
+ */
+class SEQUENCER_API SSequencerSearchBox : public SCompoundWidget
+{
+public:
+	SLATE_BEGIN_ARGS(SSequencerSearchBox)
+		: _SuggestionListPlacement(MenuPlacement_BelowAnchor)
+		, _PossibleSuggestions(TArray<FSequencerFilterSuggestion>())
+		, _DelayChangeNotificationsWhileTyping(true)
+		, _MustMatchPossibleSuggestions(false)
+		, _ShowSearchHistory(true)
+	{}
+
+		/** Where to place the suggestion list */
+		SLATE_ARGUMENT(EMenuPlacement, SuggestionListPlacement)
+
+		/** Invoked whenever the text changes */
+		SLATE_EVENT(FOnTextChanged, OnTextChanged)
+
+		/** Invoked whenever the text is committed (e.g. user presses enter) */
+		SLATE_EVENT(FOnTextCommitted, OnTextCommitted)
+
+		/** Initial text to display for the search text */
+		SLATE_ATTRIBUTE(FText, InitialText)
+
+		/** Hint text to display for the search text when there is no value */
+		SLATE_ATTRIBUTE(FText, HintText)
+
+		/** All possible suggestions for the search text */
+		SLATE_ATTRIBUTE(TArray<FSequencerFilterSuggestion>, PossibleSuggestions)
+
+		/** Whether the SearchBox should delay notifying listeners of text changed events until the user is done typing */
+		SLATE_ATTRIBUTE(bool, DelayChangeNotificationsWhileTyping)
+
+		/** Whether the SearchBox allows entries that don't match the possible suggestions */
+		SLATE_ATTRIBUTE(bool, MustMatchPossibleSuggestions)
+
+		/** Callback to filter the given suggestion list, to allow custom filtering behavior */
+		SLATE_EVENT(FOnSequencerSearchBoxSuggestionFilter, OnSuggestionFilter)
+
+		/** Callback when a suggestion entry is chosen during an asset search, to allow custom compositing behavior of the suggestion into the search text */
+		SLATE_EVENT(FOnSequencerSearchBoxSuggestionChosen, OnSuggestionChosen)
+
+		/** Callback delegate to have first chance handling of the OnKeyDown event */
+		SLATE_EVENT(FOnKeyDown, OnKeyDownHandler)
+	
+		/** Whether we should show a dropdown containing the last few searches */
+		SLATE_ATTRIBUTE(bool, ShowSearchHistory)
+
+		/** Handler for when the + Button next to a search is clicked */
+		SLATE_EVENT(SFilterSearchBox::FOnSaveSearchClicked, OnSaveSearchClicked)
+
+	SLATE_END_ARGS()
+
+	/** Constructs this widget with InArgs */
+	void Construct(const FArguments& InArgs, const TWeakPtr<ISequencerFilterBar>& InWeakFilterBar);
+
+	/** Get the text string currently being edited */
+	FText GetText() const;
+
+	/** Sets the text string currently being edited */
+	void SetText(const TAttribute<FText>& InNewText);
+
+	/** Set or clear the current error reporting information for this search box */
+	void SetError(const FText& InError);
+	void SetError(const FString& InError);
+
+	/** Show a + button next to the current search and set the handler for when that is clicked */
+	void SetOnSaveSearchHandler(SFilterSearchBox::FOnSaveSearchClicked InOnSaveSearchHandler);
+
+	//~ Begin SWidget
+	virtual FReply OnPreviewKeyDown(const FGeometry& InGeometry, const FKeyEvent& InKeyEvent) override;
+	virtual bool SupportsKeyboardFocus() const override;
+	virtual bool HasKeyboardFocus() const override;
+	virtual FReply OnFocusReceived(const FGeometry& InGeometry, const FFocusEvent& InFocusEvent) override;
+	//~ End SWidget
+
+private:
+	struct FParsedExpression
+	{
+		TOptional<FString> Key;
+		TOptional<FString> Value;
+		int32 SuggestionIndex = 0;
+	};
+
+	void ExtractSearchFilterTerms(const FText& InSearchText, FParsedExpression& OutParsedPair) const;
+
+	TSharedRef<SWidget> GetSuggestionListMenuContent();
+
+	/** First chance handler for key down events to the editable text widget */
+	FReply HandleKeyDown(const FGeometry& InGeometry, const FKeyEvent& InKeyEvent);
+
+	/** Handler for when text in the editable text box changed */
+	void HandleTextChanged(const FText& InNewText);
+
+	/** Handler for when text in the editable text box changed */
+	void HandleTextCommitted(const FText& InNewText, const ETextCommit::Type InCommitType);
+
+	/** Called by SListView when the selection changes in the suggestion list */
+	void OnSelectionChanged(TSharedPtr<FSequencerFilterSuggestionListEntryBase> InNewValue, const ESelectInfo::Type InSelectInfo);
+
+	/** Makes the widget for a suggestion message in the list view */
+	TSharedRef<ITableRow> MakeSuggestionListItemWidget(const TSharedPtr<FSequencerFilterSuggestionListEntryBase> InSuggestion, const TSharedRef<STableViewBase>& InOwnerTable);
+
+	/** Gets the text to highlight in the suggestion list */
+	FText GetHighlightText() const;
+
+	/** Updates and shows or hides the suggestion list */
+	void UpdateSuggestionList();
+
+	/** Sets the focus to the InputText box */
+	void FocusEditBox();
+
+	/** Returns the currently selected suggestion */
+	TSharedPtr<FSequencerFilterSuggestionListEntryBase> GetSelectedSuggestion() const;
+
+	/** Default implementation of OnSuggestionFilter, if no external implementation is provided */
+	void DefaultSuggestionFilterImpl(const FText& InSearchText, TArray<FSequencerFilterSuggestion>& OutPossibleSuggestions, FText& OutSuggestionHighlightText);
+
+	/** Default implementation of OnSuggestionChosen, if no external implementation is provided */
+	FText DefaultSuggestionChosenImpl(const FText& InSearchText, const FString& InSuggestion);
+
+	/** Provides default key suggestions, if no suggestions are provided */
+	void DefaultKeySuggestions(const FString& InKeyValue, TArray<FSequencerFilterSuggestion>& OutPossibleSuggestions);
+
+	/** Provides default value suggestions, if no suggestions are provided */
+	void DefaultValueSuggestions(const FString& InKeyValue, TArray<FSequencerFilterSuggestion>& OutPossibleSuggestions);
+
+	static bool IsOperatorToken(const TExpressionToken<TCHAR>& InToken);
+	static bool IsLogicalOperatorToken(const TExpressionToken<TCHAR>& InToken);
+
+	static int32 FindTokenIndex(const TArray<TExpressionToken<TCHAR>>& ExpressionTokens, const int32 InIndexToFind);
+
+	const FTextFilterExpressionEvaluator* GetTextFilterExpressionEvaluator() const;
+
+	TWeakPtr<ISequencerFilterBar> WeakFilterBar;
+
+	TSharedPtr<SFilterSearchBox> SearchBox;
+
+	TSharedPtr<SListView<TSharedPtr<FSequencerFilterSuggestionListEntryBase>>> SuggestionListView;
+
+	TSharedPtr<SMenuAnchor> SuggestionMenuAnchor;
+
+	TArray<TSharedPtr<FSequencerFilterSuggestionListEntryBase>> SuggestionItems;
+
+	/** The state of the text prior to being committed */
+	FText PreCommittedText;
+
+	/** The highlight text to use for the suggestions list */
+	FText SuggestionHighlightText;
+
+	FOnTextChanged OnTextChanged;
+	FOnTextCommitted OnTextCommitted;
+
+	/** Delegate to filter the given suggestion list, to allow custom filtering behavior */
+	FOnSequencerSearchBoxSuggestionFilter OnSuggestionFilter;
+
+	/** Delegate when a suggestion entry is chosen during a search, to allow custom compositing behavior of the suggestion into the search text */
+	FOnSequencerSearchBoxSuggestionChosen OnSuggestionChosen;
+
+	/** Delegate for first chance handling for key down events */
+	FOnKeyDown OnKeyDownHandler;
+
+	/** All possible suggestions for the search text */
+	TAttribute<TArray<FSequencerFilterSuggestion>> PossibleSuggestions;
+
+	/** Determines whether the committed text should match a suggestion */
+	bool bMustMatchPossibleSuggestions = false;
+
+	bool bDisableOpeningSuggestions = true;
+
+	FTextLocation LastCursorLocation = FTextLocation(0, 0);
+
+	int32 LastChosenCursorOffset = INDEX_NONE;
+};
diff --git a/Engine/Source/Editor/Sequencer/Public/ISequencer.h b/Engine/Source/Editor/Sequencer/Public/ISequencer.h
--- a/Engine/Source/Editor/Sequencer/Public/ISequencer.h
+++ b/Engine/Source/Editor/Sequencer/Public/ISequencer.h
@@ -91,6 +91,7 @@ namespace Sequencer
 
 enum class ETimeDomain : uint8;
 class FSequencerEditorViewModel;
+class SOutlinerView;
 struct FCreateBindingParams;
 struct FTimeDomainOverride;
 struct FSequencerNumericTypeInterface;
@@ -1000,6 +1001,10 @@ public:
 
 	virtual TSharedPtr<UE::Sequencer::NavigationTool::INavigationTool> GetNavigationTool() const = 0;
 
+	virtual void InvalidateGlobalMarkedFramesCache() = 0;
+
+	virtual TSharedPtr<UE::Sequencer::SOutlinerView> GetOutlinerViewWidget() const = 0;
+
 protected:
 	FOnInitializeDetailsPanel InitializeDetailsPanelEvent;
 	FOnCameraAddedToSequencer CameraAddedToSequencer;
diff --git a/Engine/Source/Editor/Sequencer/Public/MVVM/TrackModelStorageExtension.h b/Engine/Source/Editor/Sequencer/Public/MVVM/TrackModelStorageExtension.h
--- a/Engine/Source/Editor/Sequencer/Public/MVVM/TrackModelStorageExtension.h
+++ b/Engine/Source/Editor/Sequencer/Public/MVVM/TrackModelStorageExtension.h
@@ -24,7 +24,7 @@ class FTrackModelStorageExtension
 {
 public:
 
-	UE_SEQUENCER_DECLARE_VIEW_MODEL_TYPE_ID(FTrackModelStorageExtension)
+	UE_SEQUENCER_DECLARE_VIEW_MODEL_TYPE_ID_API(SEQUENCER_API, FTrackModelStorageExtension)
 
 	FTrackModelStorageExtension(const TArray<FOnCreateTrackModel>& InTrackModelCreators);
 
@@ -33,7 +33,7 @@ public:
 
 	TSharedPtr<FTrackModel> CreateModelForTrack(UMovieSceneTrack* InTrack, TSharedPtr<FViewModel> DesiredParent = nullptr);
 
-	TSharedPtr<FTrackModel> FindModelForTrack(UMovieSceneTrack* InTrack) const;
+	SEQUENCER_API TSharedPtr<FTrackModel> FindModelForTrack(UMovieSceneTrack* InTrack) const;
 
 private:
 
