Sequencer: Fix conflicting shortcuts

Ctrl+Shift+F can't be used because it conflicts with Open Content Browser.
The changes here are:

Toggle Filter Bar Visibility - Alt-F
Scroll to Next - Alt-N

#rnx

diff --git a/Engine/Plugins/Experimental/Avalanche/Source/AvalancheOutliner/Private/AvaOutlinerCommands.cpp b/Engine/Plugins/Experimental/Avalanche/Source/AvalancheOutliner/Private/AvaOutlinerCommands.cpp
--- a/Engine/Plugins/Experimental/Avalanche/Source/AvalancheOutliner/Private/AvaOutlinerCommands.cpp
+++ b/Engine/Plugins/Experimental/Avalanche/Source/AvalancheOutliner/Private/AvaOutlinerCommands.cpp
@@ -58,7 +58,7 @@ void FAvaOutlinerCommands::RegisterCommands()
 		, "Scroll to Next"
 		, "Scrolls the next selection into view"
 		, EUserInterfaceActionType::Button
-		, FInputChord(EModifierKey::Alt, EKeys::F));
+		, FInputChord(EModifierKey::Alt, EKeys::N));
 
 	UI_COMMAND(ToggleMutedHierarchy
 		, "Muted Hierarchy"
diff --git a/Engine/Plugins/Experimental/SequenceNavigator/Source/SequenceNavigator/Private/Filters/NavigationToolFilterCommands.cpp b/Engine/Plugins/Experimental/SequenceNavigator/Source/SequenceNavigator/Private/Filters/NavigationToolFilterCommands.cpp
--- a/Engine/Plugins/Experimental/SequenceNavigator/Source/SequenceNavigator/Private/Filters/NavigationToolFilterCommands.cpp
+++ b/Engine/Plugins/Experimental/SequenceNavigator/Source/SequenceNavigator/Private/Filters/NavigationToolFilterCommands.cpp
@@ -10,7 +10,7 @@ void FNavigationToolFilterCommands::RegisterCommands()
 		, "Filter Bar Visibility"
 		, "Toggle the visibility of the filter bar. The filter bar will only be displayed when there are enabled filters"
 		, EUserInterfaceActionType::ToggleButton
-		, FInputChord(EModifierKey::Control | EModifierKey::Shift, EKeys::F));
+		, FInputChord(EModifierKey::Alt, EKeys::F));
 
 	UI_COMMAND(SetToVerticalLayout
 		, "Vertical"
diff --git a/Engine/Plugins/Experimental/SequenceNavigator/Source/SequenceNavigator/Private/NavigationToolCommands.cpp b/Engine/Plugins/Experimental/SequenceNavigator/Source/SequenceNavigator/Private/NavigationToolCommands.cpp
--- a/Engine/Plugins/Experimental/SequenceNavigator/Source/SequenceNavigator/Private/NavigationToolCommands.cpp
+++ b/Engine/Plugins/Experimental/SequenceNavigator/Source/SequenceNavigator/Private/NavigationToolCommands.cpp
@@ -91,7 +91,7 @@ void FNavigationToolCommands::RegisterCommands()
 		, "Scroll to Next"
 		, "Scrolls the next selection into view"
 		, EUserInterfaceActionType::Button
-		, FInputChord(EModifierKey::Alt, EKeys::F));
+		, FInputChord(EModifierKey::Alt, EKeys::N));
 
 	UI_COMMAND(ToggleMutedHierarchy
 		, "Muted Hierarchy"
diff --git a/Engine/Source/Editor/Sequencer/Private/Filters/SequencerTrackFilterCommands.cpp b/Engine/Source/Editor/Sequencer/Private/Filters/SequencerTrackFilterCommands.cpp
--- a/Engine/Source/Editor/Sequencer/Private/Filters/SequencerTrackFilterCommands.cpp
+++ b/Engine/Source/Editor/Sequencer/Private/Filters/SequencerTrackFilterCommands.cpp
@@ -15,7 +15,7 @@ FSequencerTrackFilterCommands::FSequencerTrackFilterCommands()
 
 void FSequencerTrackFilterCommands::RegisterCommands()
 {
-	UI_COMMAND(ToggleFilterBarVisibility, "Filter Bar Visibility", "Toggle the visibility of the filter bar. The filter bar will only be displayed when there are enabled filters", EUserInterfaceActionType::ToggleButton, FInputChord(EModifierKey::Control | EModifierKey::Shift, EKeys::F));
+	UI_COMMAND(ToggleFilterBarVisibility, "Filter Bar Visibility", "Toggle the visibility of the filter bar. The filter bar will only be displayed when there are enabled filters", EUserInterfaceActionType::ToggleButton, FInputChord(EModifierKey::Alt, EKeys::F));
 
 	UI_COMMAND(SetToVerticalLayout, "Vertical", "Swap to a vertical layout for the filter bar", EUserInterfaceActionType::RadioButton, FInputChord());
 	UI_COMMAND(SetToHorizontalLayout, "Horizontal", "Swap to a horizontal layout for the filter bar", EUserInterfaceActionType::RadioButton, FInputChord());
