[Sequencer] Remove unneeded instances of IsValid() from filtering classes to maximize filter logic processing

diff --git a/Engine/Source/Editor/Sequencer/Private/Filters/Filters/SequencerTrackFilterBase.cpp b/Engine/Source/Editor/Sequencer/Private/Filters/Filters/SequencerTrackFilterBase.cpp
--- a/Engine/Source/Editor/Sequencer/Private/Filters/Filters/SequencerTrackFilterBase.cpp
+++ b/Engine/Source/Editor/Sequencer/Private/Filters/Filters/SequencerTrackFilterBase.cpp
@@ -31,5 +31,5 @@ UMovieSceneSequence* FSequencerTrackFilter::GetFocusedMovieSceneSequence() const
 UMovieScene* FSequencerTrackFilter::GetFocusedGetMovieScene() const
 {
 	const UMovieSceneSequence* const FocusedMovieSceneSequence = GetFocusedMovieSceneSequence();
-	return IsValid(FocusedMovieSceneSequence) ? FocusedMovieSceneSequence->GetMovieScene() : nullptr;
+	return FocusedMovieSceneSequence ? FocusedMovieSceneSequence->GetMovieScene() : nullptr;
 }
diff --git a/Engine/Source/Editor/Sequencer/Private/Filters/Filters/SequencerTrackFilter_Group.cpp b/Engine/Source/Editor/Sequencer/Private/Filters/Filters/SequencerTrackFilter_Group.cpp
--- a/Engine/Source/Editor/Sequencer/Private/Filters/Filters/SequencerTrackFilter_Group.cpp
+++ b/Engine/Source/Editor/Sequencer/Private/Filters/Filters/SequencerTrackFilter_Group.cpp
@@ -17,8 +17,7 @@ FSequencerTrackFilter_Group::FSequencerTrackFilter_Group(ISequencerTrackFilters&
 
 FSequencerTrackFilter_Group::~FSequencerTrackFilter_Group()
 {
-	UMovieScene* const MovieScene = MovieSceneWeak.Get();
-	if (IsValid(MovieScene))
+	if (UMovieScene* const MovieScene = MovieSceneWeak.Get())
 	{
 		MovieScene->GetNodeGroups().OnNodeGroupCollectionChanged().RemoveAll(this);
 		MovieSceneWeak.Reset();
@@ -53,13 +52,13 @@ FString FSequencerTrackFilter_Group::GetName() const
 bool FSequencerTrackFilter_Group::PassesFilter(FSequencerTrackFilterType InItem) const
 {
 	const UMovieSceneSequence* const FocusedSequence = GetFocusedMovieSceneSequence();
-	if (!IsValid(FocusedSequence))
+	if (!FocusedSequence)
 	{
 		return true;
 	}
 
 	UMovieScene* const FocusedMovieScene = FocusedSequence->GetMovieScene();
-	if (!IsValid(FocusedMovieScene))
+	if (!FocusedMovieScene)
 	{
 		return true;
 	}
@@ -96,16 +95,16 @@ void FSequencerTrackFilter_Group::UpdateMovieScene(UMovieScene* const InMovieSce
 {
 	UMovieScene* const OldMovieScene = MovieSceneWeak.Get();
 
-	if (!IsValid(OldMovieScene) || OldMovieScene != InMovieScene)
+	if (!OldMovieScene || OldMovieScene != InMovieScene)
 	{
-		if (IsValid(OldMovieScene))
+		if (OldMovieScene)
 		{
 			OldMovieScene->GetNodeGroups().OnNodeGroupCollectionChanged().RemoveAll(this);
 		}
 
 		MovieSceneWeak.Reset();
 
-		if (IsValid(InMovieScene))
+		if (InMovieScene)
 		{
 			MovieSceneWeak = InMovieScene;
 			InMovieScene->GetNodeGroups().OnNodeGroupCollectionChanged().AddRaw(this, &FSequencerTrackFilter_Group::HandleGroupsChanged);
@@ -137,7 +136,7 @@ void FSequencerTrackFilter_Group::ForEachMovieSceneNodeGroup(UMovieScene* const
 	, FSequencerTrackFilterType InItem
 	, const TFunctionRef<bool(const TViewModelPtr<IOutlinerExtension>& InParent, UMovieSceneNodeGroup*)>& InFunction)
 {
-	if (!IsValid(InMovieScene))
+	if (!InMovieScene)
 	{
 		return;
 	}
diff --git a/Engine/Source/Editor/Sequencer/Private/Filters/Filters/SequencerTrackFilter_Keyed.cpp b/Engine/Source/Editor/Sequencer/Private/Filters/Filters/SequencerTrackFilter_Keyed.cpp
--- a/Engine/Source/Editor/Sequencer/Private/Filters/Filters/SequencerTrackFilter_Keyed.cpp
+++ b/Engine/Source/Editor/Sequencer/Private/Filters/Filters/SequencerTrackFilter_Keyed.cpp
@@ -77,7 +77,7 @@ bool FSequencerTrackFilter_Keyed::DoesTrackExtensionHaveKeys(const TWeakViewMode
 	}
 
 	UMovieSceneTrack* const TrackObject = Track->GetTrack();
-	if (!IsValid(TrackObject))
+	if (!TrackObject)
 	{
 		return false;
 	}
diff --git a/Engine/Source/Editor/Sequencer/Private/Filters/Filters/SequencerTrackFilter_Level.cpp b/Engine/Source/Editor/Sequencer/Private/Filters/Filters/SequencerTrackFilter_Level.cpp
--- a/Engine/Source/Editor/Sequencer/Private/Filters/Filters/SequencerTrackFilter_Level.cpp
+++ b/Engine/Source/Editor/Sequencer/Private/Filters/Filters/SequencerTrackFilter_Level.cpp
@@ -80,7 +80,7 @@ bool FSequencerTrackFilter_Level::PassesFilter(FSequencerTrackFilterType InItem)
 		}
 
 		const UPackage* const Package = Object->GetPackage();
-		if (!IsValid(Package))
+		if (!Package)
 		{
 			continue;
 		}
@@ -198,7 +198,7 @@ void FSequencerTrackFilter_Level::UpdateWorld(UWorld* const InWorld)
 
 		CachedWorld.Reset();
 
-		if (IsValid(InWorld))
+		if (InWorld)
 		{
 			CachedWorld = InWorld;
 			CachedWorld.Get()->OnLevelsChanged().AddRaw(this, &FSequencerTrackFilter_Level::HandleLevelsChanged);
@@ -230,7 +230,7 @@ void FSequencerTrackFilter_Level::HandleLevelsChanged()
 	TSet<FString> OldWorldLevels = AllWorldLevels;
 	for (const ULevel* const Level : WorldLevels)
 	{
-		if (IsValid(Level))
+		if (Level)
 		{
 			const FString LevelName = FPackageName::GetShortName(Level->GetPackage()->GetName());
 			AllWorldLevels.Add(LevelName);
diff --git a/Engine/Source/Editor/Sequencer/Private/Filters/Filters/SequencerTrackFilter_Modified.cpp b/Engine/Source/Editor/Sequencer/Private/Filters/Filters/SequencerTrackFilter_Modified.cpp
--- a/Engine/Source/Editor/Sequencer/Private/Filters/Filters/SequencerTrackFilter_Modified.cpp
+++ b/Engine/Source/Editor/Sequencer/Private/Filters/Filters/SequencerTrackFilter_Modified.cpp
@@ -47,7 +47,7 @@ bool FSequencerTrackFilter_Modified::PassesFilter(FSequencerTrackFilterType InIt
 	FSequencerFilterData& FilterData = GetFilterInterface().GetFilterData();
 
 	const UMovieSceneTrack* const TrackObject = FilterData.ResolveMovieSceneTrackObject(InItem);
-	if (!IsValid(TrackObject))
+	if (!TrackObject)
 	{
 		return true;
 	}
diff --git a/Engine/Source/Editor/Sequencer/Private/Filters/Filters/SequencerTrackFilter_Selected.cpp b/Engine/Source/Editor/Sequencer/Private/Filters/Filters/SequencerTrackFilter_Selected.cpp
--- a/Engine/Source/Editor/Sequencer/Private/Filters/Filters/SequencerTrackFilter_Selected.cpp
+++ b/Engine/Source/Editor/Sequencer/Private/Filters/Filters/SequencerTrackFilter_Selected.cpp
@@ -112,7 +112,7 @@ bool FSequencerTrackFilter_Selected::PassesFilter(FSequencerTrackFilterType InIt
 	}
 
 	const UTypedElementSelectionSet* const Selection = LevelEditor->GetElementSelectionSet();
-	if (!IsValid(Selection))
+	if (!Selection)
 	{
 		return false;
 	}
@@ -120,7 +120,7 @@ bool FSequencerTrackFilter_Selected::PassesFilter(FSequencerTrackFilterType InIt
 	FSequencerFilterData& FilterData = GetFilterInterface().GetFilterData();
 
 	const UObject* const TrackObject = FilterData.ResolveTrackBoundObject(GetSequencer(), InItem);
-	if (!IsValid(TrackObject))
+	if (!TrackObject)
 	{
 		return false;
 	}
@@ -133,13 +133,13 @@ bool FSequencerTrackFilter_Selected::PassesFilter(FSequencerTrackFilterType InIt
 	}
 
 	const USceneComponent* const Component = TrackObject->GetTypedOuter<USceneComponent>();
-	if (IsValid(Component) && SelectedObjects.Contains(Component))
+	if (Component && SelectedObjects.Contains(Component))
 	{
 		return true;
 	}
 
 	const AActor* const Actor = TrackObject->GetTypedOuter<AActor>();
-	if (IsValid(Actor) && SelectedObjects.Contains(Actor))
+	if (Actor && SelectedObjects.Contains(Actor))
 	{
 		return true;
 	}
diff --git a/Engine/Source/Editor/Sequencer/Private/Filters/Filters/SequencerTrackFilter_Text.cpp b/Engine/Source/Editor/Sequencer/Private/Filters/Filters/SequencerTrackFilter_Text.cpp
--- a/Engine/Source/Editor/Sequencer/Private/Filters/Filters/SequencerTrackFilter_Text.cpp
+++ b/Engine/Source/Editor/Sequencer/Private/Filters/Filters/SequencerTrackFilter_Text.cpp
@@ -71,7 +71,7 @@ FSequencerTrackFilter_Text::FSequencerTrackFilter_Text(ISequencerTrackFilters& I
 	for (TObjectIterator<USequencerTrackFilterTextExpressionExtension> ExtensionIt(RF_NoFlags); ExtensionIt; ++ExtensionIt)
 	{
 		const USequencerTrackFilterTextExpressionExtension* const PotentialExtension = *ExtensionIt;
-		if (IsValid(PotentialExtension)
+		if (PotentialExtension
 			&& PotentialExtension->HasAnyFlags(RF_ClassDefaultObject)
 			&& !PotentialExtension->GetClass()->HasAnyClassFlags(CLASS_Deprecated | CLASS_Abstract))
 		{
diff --git a/Engine/Source/Editor/Sequencer/Private/Filters/Filters/SequencerTrackFilters.cpp b/Engine/Source/Editor/Sequencer/Private/Filters/Filters/SequencerTrackFilters.cpp
--- a/Engine/Source/Editor/Sequencer/Private/Filters/Filters/SequencerTrackFilters.cpp
+++ b/Engine/Source/Editor/Sequencer/Private/Filters/Filters/SequencerTrackFilters.cpp
@@ -186,7 +186,7 @@ bool FSequencerTrackFilter_Subsequence::PassesFilter(FSequencerTrackFilterType I
 {
 	FSequencerFilterData& FilterData = GetFilterInterface().GetFilterData();
 	const UMovieSceneTrack* const TrackObject = FilterData.ResolveMovieSceneTrackObject(InItem);
-	return IsValid(TrackObject)
+	return TrackObject
 		&& TrackObject->IsA(UMovieSceneSubTrack::StaticClass())
 		&& !TrackObject->IsA(UMovieSceneCinematicShotTrack::StaticClass());
 }
diff --git a/Engine/Source/Editor/Sequencer/Private/Filters/Menus/SequencerFilterBarContextMenu.cpp b/Engine/Source/Editor/Sequencer/Private/Filters/Menus/SequencerFilterBarContextMenu.cpp
--- a/Engine/Source/Editor/Sequencer/Private/Filters/Menus/SequencerFilterBarContextMenu.cpp
+++ b/Engine/Source/Editor/Sequencer/Private/Filters/Menus/SequencerFilterBarContextMenu.cpp
@@ -35,13 +35,13 @@ TSharedRef<SWidget> FSequencerFilterBarContextMenu::CreateMenu(const TSharedRef<
 
 void FSequencerFilterBarContextMenu::PopulateMenu(UToolMenu* const InMenu)
 {
-	if (!IsValid(InMenu))
+	if (!InMenu)
 	{
 		return;
 	}
 
 	USequencerFilterBarContext* const Context = InMenu->FindContext<USequencerFilterBarContext>();
-	if (!IsValid(Context))
+	if (!Context)
 	{
 		return;
 	}
diff --git a/Engine/Source/Editor/Sequencer/Private/Filters/Menus/SequencerTrackFilterContextMenu.cpp b/Engine/Source/Editor/Sequencer/Private/Filters/Menus/SequencerTrackFilterContextMenu.cpp
--- a/Engine/Source/Editor/Sequencer/Private/Filters/Menus/SequencerTrackFilterContextMenu.cpp
+++ b/Engine/Source/Editor/Sequencer/Private/Filters/Menus/SequencerTrackFilterContextMenu.cpp
@@ -16,7 +16,7 @@
 TSharedRef<SWidget> FSequencerTrackFilterContextMenu::CreateMenuWidget(const TSharedRef<FSequencerTrackFilter>& InFilter)
 {
 	UToolMenus* const ToolMenus = UToolMenus::Get();
-	check(IsValid(ToolMenus));
+	check(ToolMenus);
 
 	const FName FilterMenuName = TEXT("Sequencer.TrackFilterContextMenu");
 	if (!ToolMenus->IsMenuRegistered(FilterMenuName))
@@ -43,13 +43,13 @@ TSharedRef<SWidget> FSequencerTrackFilterContextMenu::CreateMenuWidget(const TSh
 
 void FSequencerTrackFilterContextMenu::PopulateMenu(UToolMenu* const InMenu)
 {
-	if (!IsValid(InMenu))
+	if (!InMenu)
 	{
 		return;
 	}
 
 	USequencerFilterMenuContext* const Context = InMenu->FindContext<USequencerFilterMenuContext>();
-	if (!IsValid(Context))
+	if (!Context)
 	{
 		return;
 	}
@@ -237,8 +237,7 @@ void FSequencerTrackFilterContextMenu::OnDeleteFilter()
 	FilterInterface.SetFilterActiveByDisplayName(FilterName, false);
 	FilterInterface.RemoveCustomTextFilter(CustomTextFilter.ToSharedRef(), false);
 
-	USequencerSettings* const SequencerSettings = FilterInterface.GetSequencer().GetSequencerSettings();
-	if (IsValid(SequencerSettings))
+	if (USequencerSettings* const SequencerSettings = FilterInterface.GetSequencer().GetSequencerSettings())
 	{
 		FSequencerFilterBarConfig& Config = SequencerSettings->FindOrAddTrackFilterBar(FilterInterface.GetIdentifier(), false);
 
diff --git a/Engine/Source/Editor/Sequencer/Private/Filters/Menus/SequencerTrackFilterMenu.cpp b/Engine/Source/Editor/Sequencer/Private/Filters/Menus/SequencerTrackFilterMenu.cpp
--- a/Engine/Source/Editor/Sequencer/Private/Filters/Menus/SequencerTrackFilterMenu.cpp
+++ b/Engine/Source/Editor/Sequencer/Private/Filters/Menus/SequencerTrackFilterMenu.cpp
@@ -45,13 +45,13 @@ TSharedRef<SWidget> FSequencerTrackFilterMenu::CreateMenu(const TSharedRef<FSequ
 
 void FSequencerTrackFilterMenu::PopulateMenu(UToolMenu* const InMenu)
 {
-	if (!IsValid(InMenu))
+	if (!InMenu)
 	{
 		return;
 	}
 
 	USequencerFilterBarContext* const Context = InMenu->FindContext<USequencerFilterBarContext>();
-	if (!IsValid(Context))
+	if (!Context)
 	{
 		return;
 	}
@@ -105,8 +105,7 @@ void FSequencerTrackFilterMenu::PopulateFilterOptionsSection(UToolMenu& InMenu)
 		));
 
 	const USequencerSettings* const SequencerSettings = FilterBar->GetSequencer().GetSequencerSettings();
-	check(IsValid(SequencerSettings));
-	if (SequencerSettings->GetUseFilterSubmenusForCategories())
+	if (SequencerSettings && SequencerSettings->GetUseFilterSubmenusForCategories())
 	{
 		Section.AddSeparator(NAME_None);
 	}
@@ -121,7 +120,7 @@ void FSequencerTrackFilterMenu::PopulateCommonFilterSections(UToolMenu& InMenu)
 	}
 
 	const USequencerSettings* const SequencerSettings = FilterBar->GetSequencer().GetSequencerSettings();
-	if (!IsValid(SequencerSettings))
+	if (!SequencerSettings)
 	{
 		return;
 	}
@@ -204,7 +203,7 @@ void FSequencerTrackFilterMenu::PopulateOtherFilterSections(UToolMenu& InMenu)
 
 void FSequencerTrackFilterMenu::FillLevelFilterMenu(UToolMenu* const InMenu)
 {
-	if (!IsValid(InMenu))
+	if (!InMenu)
 	{
 		return;
 	}
@@ -218,7 +217,7 @@ void FSequencerTrackFilterMenu::FillLevelFilterMenu(UToolMenu* const InMenu)
 	FToolMenuSection& LevelsSection = InMenu->FindOrAddSection(TEXT("Levels"), LOCTEXT("LevelFilters", "Level Filters"));
 
 	const UWorld* const World = FilterBar->GetWorld();
-	if (!IsValid(World))
+	if (!World)
 	{
 		return;
 	}
@@ -241,7 +240,7 @@ void FSequencerTrackFilterMenu::FillLevelFilterMenu(UToolMenu* const InMenu)
 
 void FSequencerTrackFilterMenu::FillGroupFilterMenu(UToolMenu* const InMenu)
 {
-	if (!IsValid(InMenu))
+	if (!InMenu)
 	{
 		return;
 	}
@@ -264,13 +263,13 @@ void FSequencerTrackFilterMenu::FillGroupFilterMenu(UToolMenu* const InMenu)
 		FUIAction(FExecuteAction::CreateSP(this, &FSequencerTrackFilterMenu::OpenNodeGroupsManager)));
 
 	const UMovieSceneSequence* const FocusedMovieSequence = Sequencer.GetRootMovieSceneSequence();
-	if (!IsValid(FocusedMovieSequence))
+	if (!FocusedMovieSequence)
 	{
 		return;
 	}
 
 	UMovieScene* const FocusedMovieScene = FocusedMovieSequence->GetMovieScene();
-	if (!IsValid(FocusedMovieScene))
+	if (!FocusedMovieScene)
 	{
 		return;
 	}
@@ -295,7 +294,7 @@ void FSequencerTrackFilterMenu::FillGroupFilterMenu(UToolMenu* const InMenu)
 
 void FSequencerTrackFilterMenu::FillCustomTextFiltersMenu(UToolMenu* const InMenu)
 {
-	if (!IsValid(InMenu))
+	if (!InMenu)
 	{
 		return;
 	}
@@ -394,7 +393,7 @@ void FSequencerTrackFilterMenu::FillFiltersMenuCategory(FToolMenuSection& InOutS
 
 void FSequencerTrackFilterMenu::FillFiltersMenuCategory(UToolMenu* const InMenu, const TSharedRef<FFilterCategory> InMenuCategory)
 {
-	if (IsValid(InMenu))
+	if (InMenu)
 	{
 		FillFiltersMenuCategory(InMenu->AddSection(TEXT("Section")), InMenuCategory);
 	}
@@ -615,7 +614,7 @@ ECheckBoxState FSequencerTrackFilterMenu::GetAllLevelsCheckState() const
 
 void FSequencerTrackFilterMenu::OnNodeGroupFilterClicked(UMovieSceneNodeGroup* const InNodeGroup)
 {
-	if (IsValid(InNodeGroup))
+	if (InNodeGroup)
 	{
 		InNodeGroup->SetEnableFilter(!InNodeGroup->GetEnableFilter());
 	}
@@ -651,13 +650,13 @@ ECheckBoxState FSequencerTrackFilterMenu::GetAllGroupsCheckState() const
 	}
 
 	const UMovieSceneSequence* const RootSequence = FilterBar->GetSequencer().GetRootMovieSceneSequence();
-	if (!IsValid(RootSequence))
+	if (!RootSequence)
 	{
 		return ECheckBoxState::Unchecked;
 	}
 
 	UMovieScene* const FocusedMovieScene = RootSequence->GetMovieScene();
-	if (!IsValid(FocusedMovieScene))
+	if (!FocusedMovieScene)
 	{
 		return ECheckBoxState::Unchecked;
 	}
diff --git a/Engine/Source/Editor/Sequencer/Private/Filters/Menus/SequencerViewOptionsMenu.cpp b/Engine/Source/Editor/Sequencer/Private/Filters/Menus/SequencerViewOptionsMenu.cpp
--- a/Engine/Source/Editor/Sequencer/Private/Filters/Menus/SequencerViewOptionsMenu.cpp
+++ b/Engine/Source/Editor/Sequencer/Private/Filters/Menus/SequencerViewOptionsMenu.cpp
@@ -43,13 +43,13 @@ TSharedRef<SWidget> FSequencerViewOptionsMenu::CreateMenu(const TWeakPtr<FSequen
 
 void FSequencerViewOptionsMenu::PopulateMenu(UToolMenu* const InMenu)
 {
-	if (!IsValid(InMenu))
+	if (!InMenu)
 	{
 		return;
 	}
 
 	USequencerMenuContext* const Context = InMenu->FindContext<USequencerMenuContext>();
-	if (!IsValid(Context))
+	if (!Context)
 	{
 		return;
 	}
@@ -152,7 +152,7 @@ void FSequencerViewOptionsMenu::PopulateFilterOptionsSection(UToolMenu& InMenu)
 
 void FSequencerViewOptionsMenu::CreateFilterBarOptionsSection(UToolMenu* const InMenu)
 {
-	if (!IsValid(InMenu))
+	if (!InMenu)
 	{
 		return;
 	}
@@ -204,7 +204,7 @@ bool FSequencerViewOptionsMenu::IsIncludePinnedInFilter() const
 	}
 
 	const USequencerSettings* const SequencerSettings = Sequencer->GetSequencerSettings();
-	if (!IsValid(SequencerSettings))
+	if (!SequencerSettings)
 	{
 		return false;
 	}
@@ -221,7 +221,7 @@ void FSequencerViewOptionsMenu::ToggleIncludePinnedInFilter()
 	}
 
 	USequencerSettings* const SequencerSettings = Sequencer->GetSequencerSettings();
-	if (!IsValid(SequencerSettings))
+	if (!SequencerSettings)
 	{
 		return;
 	}
@@ -240,7 +240,7 @@ bool FSequencerViewOptionsMenu::IsAutoExpandPassedFilterNodes() const
 	}
 
 	const USequencerSettings* const SequencerSettings = Sequencer->GetSequencerSettings();
-	if (!IsValid(SequencerSettings))
+	if (!SequencerSettings)
 	{
 		return false;
 	}
@@ -257,7 +257,7 @@ void FSequencerViewOptionsMenu::ToggleAutoExpandPassedFilterNodes()
 	}
 
 	USequencerSettings* const SequencerSettings = Sequencer->GetSequencerSettings();
-	if (!IsValid(SequencerSettings))
+	if (!SequencerSettings)
 	{
 		return;
 	}
@@ -276,7 +276,7 @@ bool FSequencerViewOptionsMenu::IsLinkFiltersWithCurveEditor() const
 	}
 
 	const USequencerSettings* const SequencerSettings = Sequencer->GetSequencerSettings();
-	if (!IsValid(SequencerSettings))
+	if (!SequencerSettings)
 	{
 		return false;
 	}
@@ -293,7 +293,7 @@ void FSequencerViewOptionsMenu::ToggleLinkFiltersWithCurveEditor()
 	}
 
 	USequencerSettings* const SequencerSettings = Sequencer->GetSequencerSettings();
-	if (!IsValid(SequencerSettings))
+	if (!SequencerSettings)
 	{
 		return;
 	}
diff --git a/Engine/Source/Editor/Sequencer/Private/Filters/SequencerFilterBar.cpp b/Engine/Source/Editor/Sequencer/Private/Filters/SequencerFilterBar.cpp
--- a/Engine/Source/Editor/Sequencer/Private/Filters/SequencerFilterBar.cpp
+++ b/Engine/Source/Editor/Sequencer/Private/Filters/SequencerFilterBar.cpp
@@ -152,7 +152,7 @@ void FSequencerFilterBar::CreateDefaultFilters()
 	for (TObjectIterator<USequencerTrackFilterExtension> ExtensionIt(RF_NoFlags); ExtensionIt; ++ExtensionIt)
 	{
 		const USequencerTrackFilterExtension* const PotentialExtension = *ExtensionIt;
-		if (IsValid(PotentialExtension)
+		if (PotentialExtension
 			&& PotentialExtension->HasAnyFlags(RF_ClassDefaultObject)
 			&& !PotentialExtension->GetClass()->HasAnyClassFlags(CLASS_Deprecated | CLASS_Abstract))
 		{
@@ -225,7 +225,7 @@ void FSequencerFilterBar::BindCommands()
 
 	// Bind all filter actions
 	UMovieSceneSequence* const FocusedSequence = Sequencer.GetFocusedMovieSceneSequence();
-	if (!IsValid(FocusedSequence))
+	if (!FocusedSequence)
 	{
 		return;
 	}
@@ -278,7 +278,10 @@ void FSequencerFilterBar::BindCommands()
 void FSequencerFilterBar::CreateCustomTextFiltersFromConfig()
 {
 	USequencerSettings* const SequencerSettings = Sequencer.GetSequencerSettings();
-	check(IsValid(SequencerSettings));
+	if (!SequencerSettings)
+	{
+		return;
+	}
 
 	CustomTextFilters.Empty();
 
@@ -304,8 +307,7 @@ TSharedPtr<FUICommandList> FSequencerFilterBar::GetCommandList() const
 
 FName FSequencerFilterBar::GetIdentifier() const
 {
-	USequencerSettings* const SequencerSettings = Sequencer.GetSequencerSettings();
-	if (IsValid(SequencerSettings))
+	if (const USequencerSettings* const SequencerSettings = Sequencer.GetSequencerSettings())
 	{
 		return *SequencerSettings->GetName();
 	}
@@ -604,7 +606,10 @@ bool FSequencerFilterBar::HasAnyFilterActive(const bool bCheckTextFilter
 bool FSequencerFilterBar::IsFilterActive(const TSharedRef<FSequencerTrackFilter> InFilter) const
 {
 	USequencerSettings* const SequencerSettings = Sequencer.GetSequencerSettings();
-	check(IsValid(SequencerSettings));
+	if (!SequencerSettings)
+	{
+		return false;
+	}
 
 	const FSequencerFilterBarConfig& Config = SequencerSettings->FindOrAddTrackFilterBar(GetIdentifier(), false);
 
@@ -615,7 +620,10 @@ bool FSequencerFilterBar::IsFilterActive(const TSharedRef<FSequencerTrackFilter>
 bool FSequencerFilterBar::SetFilterActive(const TSharedRef<FSequencerTrackFilter>& InFilter, const bool bInActive, const bool bInRequestFilterUpdate)
 {
 	USequencerSettings* const SequencerSettings = Sequencer.GetSequencerSettings();
-	check(IsValid(SequencerSettings));
+	if (!SequencerSettings)
+	{
+		return false;
+	}
 
 	const bool bNewActive = InFilter->IsInverseFilter() ? !bInActive : bInActive;
 
@@ -685,7 +693,10 @@ void FSequencerFilterBar::ActivateCommonFilters(const bool bInActivate
     , const TArray<TSharedRef<FSequencerTrackFilter>>& InExceptions)
 {
 	USequencerSettings* const SequencerSettings = Sequencer.GetSequencerSettings();
-	check(IsValid(SequencerSettings));
+	if (!SequencerSettings)
+	{
+		return;
+	}
 
 	FSequencerFilterBarConfig& Config = SequencerSettings->FindOrAddTrackFilterBar(GetIdentifier(), false);
 
@@ -877,7 +888,10 @@ bool FSequencerFilterBar::HasAnyFilterEnabled() const
 bool FSequencerFilterBar::IsFilterEnabled(TSharedRef<FSequencerTrackFilter> InFilter) const
 {
 	USequencerSettings* const SequencerSettings = Sequencer.GetSequencerSettings();
-	check(IsValid(SequencerSettings));
+	if (!SequencerSettings)
+	{
+		return false;
+	}
 
 	FSequencerFilterBarConfig& Config = SequencerSettings->FindOrAddTrackFilterBar(GetIdentifier(), false);
 
@@ -888,7 +902,10 @@ bool FSequencerFilterBar::IsFilterEnabled(TSharedRef<FSequencerTrackFilter> InFi
 bool FSequencerFilterBar::SetFilterEnabled(const TSharedRef<FSequencerTrackFilter> InFilter, const bool bInEnabled, const bool bInRequestFilterUpdate)
 {
 	USequencerSettings* const SequencerSettings = Sequencer.GetSequencerSettings();
-	check(IsValid(SequencerSettings));
+	if (!SequencerSettings)
+	{
+		return false;
+	}
 
 	FSequencerFilterBarConfig& Config = SequencerSettings->FindOrAddTrackFilterBar(GetIdentifier(), true);
 
@@ -923,7 +940,10 @@ void FSequencerFilterBar::EnableFilters(const bool bInEnable
 	, const TArray<TSharedRef<FSequencerTrackFilter>> InExceptions)
 {
 	USequencerSettings* const SequencerSettings = Sequencer.GetSequencerSettings();
-	check(IsValid(SequencerSettings));
+	if (!SequencerSettings)
+	{
+		return;
+	}
 
 	FSequencerFilterBarConfig& Config = SequencerSettings->FindOrAddTrackFilterBar(GetIdentifier(), true);
 
@@ -1073,8 +1093,7 @@ bool FSequencerFilterBar::AddCustomTextFilter(const TSharedRef<FSequencerTrackFi
 
 	if (bInAddToConfig)
 	{
-		USequencerSettings* const SequencerSettings = Sequencer.GetSequencerSettings();
-		if (ensure(IsValid(SequencerSettings)))
+		if (USequencerSettings* const SequencerSettings = Sequencer.GetSequencerSettings())
 		{
 			FSequencerFilterBarConfig& Config = SequencerSettings->FindOrAddTrackFilterBar(GetIdentifier(), false);
 
@@ -1099,8 +1118,7 @@ bool FSequencerFilterBar::RemoveCustomTextFilter(const TSharedRef<FSequencerTrac
 
 	if (bInAddToConfig)
 	{
-		USequencerSettings* const SequencerSettings = Sequencer.GetSequencerSettings();
-		if (ensure(IsValid(SequencerSettings)))
+		if (USequencerSettings* const SequencerSettings = Sequencer.GetSequencerSettings())
 		{
 			FSequencerFilterBarConfig& Config = SequencerSettings->FindOrAddTrackFilterBar(GetIdentifier(), false);
 
@@ -1120,7 +1138,10 @@ bool FSequencerFilterBar::RemoveCustomTextFilter(const TSharedRef<FSequencerTrac
 void FSequencerFilterBar::ActivateCustomTextFilters(const bool bInActivate, const TArray<TSharedRef<FSequencerTrackFilter_CustomText>> InExceptions)
 {
 	USequencerSettings* const SequencerSettings = Sequencer.GetSequencerSettings();
-	check(IsValid(SequencerSettings));
+	if (!SequencerSettings)
+	{
+		return;
+	}
 
 	FSequencerFilterBarConfig& Config = SequencerSettings->FindOrAddTrackFilterBar(GetIdentifier(), false);
 
@@ -1152,7 +1173,10 @@ void FSequencerFilterBar::ActivateCustomTextFilters(const bool bInActivate, cons
 void FSequencerFilterBar::EnableCustomTextFilters(const bool bInEnable, const TArray<TSharedRef<FSequencerTrackFilter_CustomText>> InExceptions)
 {
 	USequencerSettings* const SequencerSettings = Sequencer.GetSequencerSettings();
-	check(IsValid(SequencerSettings));
+	if (!SequencerSettings)
+	{
+		return;
+	}
 
 	FSequencerFilterBarConfig& Config = SequencerSettings->FindOrAddTrackFilterBar(GetIdentifier(), false);
 
@@ -1266,13 +1290,13 @@ bool FSequencerFilterBar::CanEnableAllLevelFilters(const bool bInEnable)
 void FSequencerFilterBar::EnableAllGroupFilters(const bool bInEnable)
 {
 	UMovieSceneSequence* const FocusedMovieSequence = Sequencer.GetFocusedMovieSceneSequence();
-	if (!IsValid(FocusedMovieSequence))
+	if (!FocusedMovieSequence)
 	{
 		return;
 	}
 
 	UMovieScene* const FocusedMovieScene = FocusedMovieSequence->GetMovieScene();
-	if (!IsValid(FocusedMovieScene))
+	if (!FocusedMovieScene)
 	{
 		return;
 	}
@@ -1291,7 +1315,7 @@ bool FSequencerFilterBar::IsGroupFilterActive() const
 bool FSequencerFilterBar::PassesAnyCommonFilter(const TViewModelPtr<IOutlinerExtension>& InNode)
 {
 	UMovieSceneSequence* const FocusedSequence = Sequencer.GetFocusedMovieSceneSequence();
-	if (!IsValid(FocusedSequence))
+	if (!FocusedSequence)
 	{
 		return true;
 	}
@@ -1328,7 +1352,7 @@ bool FSequencerFilterBar::PassesAnyCommonFilter(const TViewModelPtr<IOutlinerExt
 bool FSequencerFilterBar::PassesAllInternalFilters(const TViewModelPtr<IOutlinerExtension>& InNode)
 {
 	UMovieSceneSequence* const FocusedSequence = Sequencer.GetFocusedMovieSceneSequence();
-	if (!IsValid(FocusedSequence))
+	if (!FocusedSequence)
 	{
 		return true;
 	}
@@ -1367,8 +1391,7 @@ bool FSequencerFilterBar::PassesAllCustomTextFilters(const TViewModelPtr<IOutlin
 
 UWorld* FSequencerFilterBar::GetWorld() const
 {
-	UObject* const PlaybackContext = Sequencer.GetPlaybackContext();
-	if (IsValid(PlaybackContext))
+	if (UObject* const PlaybackContext = Sequencer.GetPlaybackContext())
 	{
 		return PlaybackContext->GetWorld();
 	}
@@ -1385,8 +1408,7 @@ const FSequencerFilterData& FSequencerFilterBar::FilterNodes()
 	LevelFilter->UpdateWorld(World);
 
 	// Update the group filter
-	const UMovieSceneSequence* const FocusedMovieSceneSequence = Sequencer.GetFocusedMovieSceneSequence();
-	if (IsValid(FocusedMovieSceneSequence))
+	if (const UMovieSceneSequence* const FocusedMovieSceneSequence = Sequencer.GetFocusedMovieSceneSequence())
 	{
 		GroupFilter->UpdateMovieScene(FocusedMovieSceneSequence->GetMovieScene());
 	}
@@ -1467,7 +1489,10 @@ bool FSequencerFilterBar::FilterNodesRecursive(const bool bInHasActiveFilter, co
 	}
 
 	const USequencerSettings* const SequencerSettings = Sequencer.GetSequencerSettings();
-    check(IsValid(SequencerSettings));
+	if (!SequencerSettings)
+	{
+		return false;
+	}
 
 	// Pinning overrides all other filters
 	if (!SequencerSettings->GetIncludePinnedInFilter())
@@ -1718,8 +1743,7 @@ TSharedRef<SComboButton> FSequencerFilterBar::MakeAddFilterButton()
 
 bool FSequencerFilterBar::ShouldShowFilterBarWidget() const
 {
-	const USequencerSettings* const Settings = Sequencer.GetSequencerSettings();
-	if (ensure(IsValid(Settings)))
+	if (const USequencerSettings* const Settings = Sequencer.GetSequencerSettings())
 	{
 		return Settings->IsFilterBarVisible() && HasAnyFiltersEnabled();
 	}
@@ -1728,8 +1752,7 @@ bool FSequencerFilterBar::ShouldShowFilterBarWidget() const
 
 bool FSequencerFilterBar::IsFilterBarVisible() const
 {
-	const USequencerSettings* const Settings = Sequencer.GetSequencerSettings();
-	if (ensure(IsValid(Settings)))
+	if (const USequencerSettings* const Settings = Sequencer.GetSequencerSettings())
 	{
 		return Settings->IsFilterBarVisible();
 	}
@@ -1739,7 +1762,7 @@ bool FSequencerFilterBar::IsFilterBarVisible() const
 void FSequencerFilterBar::ToggleFilterBarVisibility()
 {
 	USequencerSettings* const Settings = Sequencer.GetSequencerSettings();
-	if (!ensure(IsValid(Settings)))
+	if (!Settings)
 	{
 		return;
 	}
@@ -1753,8 +1776,7 @@ void FSequencerFilterBar::ToggleFilterBarVisibility()
 
 bool FSequencerFilterBar::IsFilterBarLayout(const EFilterBarLayout InLayout) const
 {
-	const USequencerSettings* const Settings = Sequencer.GetSequencerSettings();
-	if (ensure(IsValid(Settings)))
+	if (const USequencerSettings* const Settings = Sequencer.GetSequencerSettings())
 	{
 		return (Settings->GetFilterBarLayout() == InLayout);
 	}
@@ -1764,7 +1786,7 @@ bool FSequencerFilterBar::IsFilterBarLayout(const EFilterBarLayout InLayout) con
 void FSequencerFilterBar::SetToVerticalLayout()
 {
 	USequencerSettings* const Settings = Sequencer.GetSequencerSettings();
-	if (!ensure(IsValid(Settings)))
+	if (!Settings)
 	{
 		return;
 	}
@@ -1777,7 +1799,7 @@ void FSequencerFilterBar::SetToVerticalLayout()
 void FSequencerFilterBar::SetToHorizontalLayout()
 {
 	USequencerSettings* const Settings = Sequencer.GetSequencerSettings();
-	if (!ensure(IsValid(Settings)))
+	if (!Settings)
 	{
 		return;
 	}
@@ -1812,7 +1834,7 @@ TSharedPtr<SSequencerFilterBar> FSequencerFilterBar::GetWidget() const
 bool FSequencerFilterBar::IsFilterSupported(const TSharedRef<FSequencerTrackFilter>& InFilter) const
 {
 	UMovieSceneSequence* const MovieSceneSequence = Sequencer.GetFocusedMovieSceneSequence();
-	if (!IsValid(MovieSceneSequence))
+	if (!MovieSceneSequence)
 	{
 		return false;
 	}
diff --git a/Engine/Source/Editor/Sequencer/Private/Filters/SequencerFilterData.cpp b/Engine/Source/Editor/Sequencer/Private/Filters/SequencerFilterData.cpp
--- a/Engine/Source/Editor/Sequencer/Private/Filters/SequencerFilterData.cpp
+++ b/Engine/Source/Editor/Sequencer/Private/Filters/SequencerFilterData.cpp
@@ -182,8 +182,7 @@ UMovieSceneTrack* FSequencerFilterData::ResolveMovieSceneTrackObject(FSequencerT
 
 	if (const TViewModelPtr<ITrackExtension> AncestorTrackModel = InNode->FindAncestorOfType<ITrackExtension>(true))
 	{
-		UMovieSceneTrack* const TrackObject = AncestorTrackModel->GetTrack();
-		if (IsValid(TrackObject))
+		if (UMovieSceneTrack* const TrackObject = AncestorTrackModel->GetTrack())
 		{
 			ResolvedTrackObjects.Add(WeakOutlinerNode, TrackObject);
 			return TrackObject;
@@ -215,8 +214,7 @@ UObject* FSequencerFilterData::ResolveTrackBoundObject(ISequencer& InSequencer,
 
 	if (const TViewModelPtr<IObjectBindingExtension> ObjectBindingModel = InNode->FindAncestorOfType<IObjectBindingExtension>(true))
 	{
-		UObject* const BoundObject = InSequencer.FindSpawnedObjectOrTemplate(ObjectBindingModel->GetObjectGuid());
-		if (IsValid(BoundObject))
+		if (UObject* const BoundObject = InSequencer.FindSpawnedObjectOrTemplate(ObjectBindingModel->GetObjectGuid()))
 		{
 			ResolvedObjects.Add(WeakOutlinerNode, BoundObject);
 			return BoundObject;
diff --git a/Engine/Source/Editor/Sequencer/Private/Filters/SequencerTrackFilterCollection.cpp b/Engine/Source/Editor/Sequencer/Private/Filters/SequencerTrackFilterCollection.cpp
--- a/Engine/Source/Editor/Sequencer/Private/Filters/SequencerTrackFilterCollection.cpp
+++ b/Engine/Source/Editor/Sequencer/Private/Filters/SequencerTrackFilterCollection.cpp
@@ -185,7 +185,7 @@ void FSequencerTrackFilterCollection::ForEachFilter(const TFunctionRef<bool(cons
 	using namespace UE::Sequencer;
 
 	UMovieSceneSequence* const RootSequence = FilterInterface.GetSequencer().GetRootMovieSceneSequence();
-	if (!IsValid(RootSequence))
+	if (!RootSequence)
 	{
 		return;
 	}
diff --git a/Engine/Source/Editor/Sequencer/Private/Filters/TextExpressions/SequencerTextFilterExpressionContext.cpp b/Engine/Source/Editor/Sequencer/Private/Filters/TextExpressions/SequencerTextFilterExpressionContext.cpp
--- a/Engine/Source/Editor/Sequencer/Private/Filters/TextExpressions/SequencerTextFilterExpressionContext.cpp
+++ b/Engine/Source/Editor/Sequencer/Private/Filters/TextExpressions/SequencerTextFilterExpressionContext.cpp
@@ -93,7 +93,7 @@ UMovieSceneSequence* FSequencerTextFilterExpressionContext::GetFocusedMovieScene
 UMovieScene* FSequencerTextFilterExpressionContext::GetFocusedGetMovieScene() const
 {
 	const UMovieSceneSequence* const FocusedMovieSceneSequence = GetFocusedMovieSceneSequence();
-	return IsValid(FocusedMovieSceneSequence) ? FocusedMovieSceneSequence->GetMovieScene() : nullptr;
+	return FocusedMovieSceneSequence ? FocusedMovieSceneSequence->GetMovieScene() : nullptr;
 }
 
 UMovieSceneTrack* FSequencerTextFilterExpressionContext::GetMovieSceneTrack() const
diff --git a/Engine/Source/Editor/Sequencer/Private/Filters/TextExpressions/SequencerTextFilterExpression_BindingType.cpp b/Engine/Source/Editor/Sequencer/Private/Filters/TextExpressions/SequencerTextFilterExpression_BindingType.cpp
--- a/Engine/Source/Editor/Sequencer/Private/Filters/TextExpressions/SequencerTextFilterExpression_BindingType.cpp
+++ b/Engine/Source/Editor/Sequencer/Private/Filters/TextExpressions/SequencerTextFilterExpression_BindingType.cpp
@@ -52,7 +52,7 @@ bool FSequencerTextFilterExpression_BindingType::TestComplexExpression(const FNa
 	FGuid ObjectBindingID = ObjectBindingExtension->GetObjectGuid();
 
 	const UMovieSceneSequence* FocusedSequence = GetFocusedMovieSceneSequence();
-	if (!IsValid(FocusedSequence))
+	if (!FocusedSequence)
 	{
 		return true;
 	}
diff --git a/Engine/Source/Editor/Sequencer/Private/Filters/TextExpressions/SequencerTextFilterExpression_CustomBinding.cpp b/Engine/Source/Editor/Sequencer/Private/Filters/TextExpressions/SequencerTextFilterExpression_CustomBinding.cpp
--- a/Engine/Source/Editor/Sequencer/Private/Filters/TextExpressions/SequencerTextFilterExpression_CustomBinding.cpp
+++ b/Engine/Source/Editor/Sequencer/Private/Filters/TextExpressions/SequencerTextFilterExpression_CustomBinding.cpp
@@ -50,7 +50,7 @@ bool FSequencerTextFilterExpression_CustomBinding::TestComplexExpression(const F
 	FGuid ObjectBindingID = ObjectBindingExtension->GetObjectGuid();
 
 	const UMovieSceneSequence* FocusedSequence = GetFocusedMovieSceneSequence();
-	if (!IsValid(FocusedSequence))
+	if (!FocusedSequence)
 	{
 		return true;
 	}
diff --git a/Engine/Source/Editor/Sequencer/Private/Filters/TextExpressions/SequencerTextFilterExpression_EmptyBinding.cpp b/Engine/Source/Editor/Sequencer/Private/Filters/TextExpressions/SequencerTextFilterExpression_EmptyBinding.cpp
--- a/Engine/Source/Editor/Sequencer/Private/Filters/TextExpressions/SequencerTextFilterExpression_EmptyBinding.cpp
+++ b/Engine/Source/Editor/Sequencer/Private/Filters/TextExpressions/SequencerTextFilterExpression_EmptyBinding.cpp
@@ -50,7 +50,7 @@ bool FSequencerTextFilterExpression_EmptyBinding::TestComplexExpression(const FN
 	FGuid ObjectBindingID = ObjectBindingExtension->GetObjectGuid();
 
 	const UMovieSceneSequence* FocusedSequence = GetFocusedMovieSceneSequence();
-	if (!IsValid(FocusedSequence))
+	if (!FocusedSequence)
 	{
 		return true;
 	}
diff --git a/Engine/Source/Editor/Sequencer/Private/Filters/TextExpressions/SequencerTextFilterExpression_Group.cpp b/Engine/Source/Editor/Sequencer/Private/Filters/TextExpressions/SequencerTextFilterExpression_Group.cpp
--- a/Engine/Source/Editor/Sequencer/Private/Filters/TextExpressions/SequencerTextFilterExpression_Group.cpp
+++ b/Engine/Source/Editor/Sequencer/Private/Filters/TextExpressions/SequencerTextFilterExpression_Group.cpp
@@ -42,13 +42,13 @@ bool FSequencerTextFilterExpression_Group::TestComplexExpression(const FName& In
 	const ISequencer& Sequencer = FilterInterface.GetSequencer();
 
 	const UMovieSceneSequence* const FocusedSequence = Sequencer.GetFocusedMovieSceneSequence();
-	if (!IsValid(FocusedSequence))
+	if (!FocusedSequence)
 	{
 		return true;
 	}
 
 	UMovieScene* const FocusedMovieScene = FocusedSequence->GetMovieScene();
-	if (!IsValid(FocusedMovieScene))
+	if (!FocusedMovieScene)
 	{
 		return true;
 	}
diff --git a/Engine/Source/Editor/Sequencer/Private/Filters/TextExpressions/SequencerTextFilterExpression_Level.cpp b/Engine/Source/Editor/Sequencer/Private/Filters/TextExpressions/SequencerTextFilterExpression_Level.cpp
--- a/Engine/Source/Editor/Sequencer/Private/Filters/TextExpressions/SequencerTextFilterExpression_Level.cpp
+++ b/Engine/Source/Editor/Sequencer/Private/Filters/TextExpressions/SequencerTextFilterExpression_Level.cpp
@@ -57,7 +57,7 @@ bool FSequencerTextFilterExpression_Level::TestComplexExpression(const FName& In
 		}
 
 		const UPackage* const Package = Object->GetPackage();
-		if (!IsValid(Package))
+		if (!Package)
 		{
 			continue;
 		}
diff --git a/Engine/Source/Editor/Sequencer/Private/Filters/TextExpressions/SequencerTextFilterExpression_ObjectClass.cpp b/Engine/Source/Editor/Sequencer/Private/Filters/TextExpressions/SequencerTextFilterExpression_ObjectClass.cpp
--- a/Engine/Source/Editor/Sequencer/Private/Filters/TextExpressions/SequencerTextFilterExpression_ObjectClass.cpp
+++ b/Engine/Source/Editor/Sequencer/Private/Filters/TextExpressions/SequencerTextFilterExpression_ObjectClass.cpp
@@ -39,8 +39,7 @@ bool FSequencerTextFilterExpression_ObjectClass::TestComplexExpression(const FNa
 
 	ISequencer& Sequencer = FilterInterface.GetSequencer();
 
-	UObject* const BoundObject = FilterInterface.GetFilterData().ResolveTrackBoundObject(Sequencer, FilterItem);
-	if (IsValid(BoundObject))
+	if (UObject* const BoundObject = FilterInterface.GetFilterData().ResolveTrackBoundObject(Sequencer, FilterItem))
 	{
 		const FString BoundObjectClassName = BoundObject->GetClass()->GetName();
 		if (TextFilterUtils::TestComplexExpression(BoundObjectClassName, InValue, InComparisonOperation, InTextComparisonMode))
diff --git a/Engine/Source/Editor/Sequencer/Private/Filters/TextExpressions/SequencerTextFilterExpression_Tag.cpp b/Engine/Source/Editor/Sequencer/Private/Filters/TextExpressions/SequencerTextFilterExpression_Tag.cpp
--- a/Engine/Source/Editor/Sequencer/Private/Filters/TextExpressions/SequencerTextFilterExpression_Tag.cpp
+++ b/Engine/Source/Editor/Sequencer/Private/Filters/TextExpressions/SequencerTextFilterExpression_Tag.cpp
@@ -41,7 +41,7 @@ bool FSequencerTextFilterExpression_Tag::TestComplexExpression(const FName& InKe
 	}
 
 	const UMovieScene* const MovieScene = GetFocusedGetMovieScene();
-	if (!IsValid(MovieScene))
+	if (!MovieScene)
 	{
 		return true;
 	}
diff --git a/Engine/Source/Editor/Sequencer/Private/Filters/TextExpressions/SequencerTextFilterExpression_Time.cpp b/Engine/Source/Editor/Sequencer/Private/Filters/TextExpressions/SequencerTextFilterExpression_Time.cpp
--- a/Engine/Source/Editor/Sequencer/Private/Filters/TextExpressions/SequencerTextFilterExpression_Time.cpp
+++ b/Engine/Source/Editor/Sequencer/Private/Filters/TextExpressions/SequencerTextFilterExpression_Time.cpp
@@ -52,7 +52,7 @@ bool FSequencerTextFilterExpression_Time::TestComplexExpression(const FName& InK
 	}
 
 	const UMovieScene* const FocusedMovieScene = GetFocusedGetMovieScene();
-	if (!IsValid(FocusedMovieScene))
+	if (!FocusedMovieScene)
 	{
 		return true;
 	}
@@ -113,13 +113,13 @@ bool FSequencerTextFilterExpression_Time::CompareTime(ISequencer& InSequencer
 	}
 
 	const UMovieSceneSequence* const MovieSceneSequence = InSequencer.GetFocusedMovieSceneSequence();
-	if (!IsValid(MovieSceneSequence))
+	if (!MovieSceneSequence)
 	{
 		return false;
 	}
 
 	const UMovieScene* const FocusedMovieScene = MovieSceneSequence->GetMovieScene();
-	if (!IsValid(FocusedMovieScene))
+	if (!FocusedMovieScene)
 	{
 		return false;
 	}
diff --git a/Engine/Source/Editor/Sequencer/Private/Filters/TextExpressions/SequencerTextFilterExpression_TrackClass.cpp b/Engine/Source/Editor/Sequencer/Private/Filters/TextExpressions/SequencerTextFilterExpression_TrackClass.cpp
--- a/Engine/Source/Editor/Sequencer/Private/Filters/TextExpressions/SequencerTextFilterExpression_TrackClass.cpp
+++ b/Engine/Source/Editor/Sequencer/Private/Filters/TextExpressions/SequencerTextFilterExpression_TrackClass.cpp
@@ -38,8 +38,7 @@ bool FSequencerTextFilterExpression_TrackClass::TestComplexExpression(const FNam
 		return true;
 	}
 
-	UMovieSceneTrack* const TrackObject = GetMovieSceneTrack();
-	if (IsValid(TrackObject))
+	if (UMovieSceneTrack* const TrackObject = GetMovieSceneTrack())
 	{
 		const FString TrackClassName = TrackObject->GetClass()->GetName();
 		if (TextFilterUtils::TestComplexExpression(TrackClassName, InValue, InComparisonOperation, InTextComparisonMode))
diff --git a/Engine/Source/Editor/Sequencer/Private/Filters/Widgets/SSequencerSearchBox.cpp b/Engine/Source/Editor/Sequencer/Private/Filters/Widgets/SSequencerSearchBox.cpp
--- a/Engine/Source/Editor/Sequencer/Private/Filters/Widgets/SSequencerSearchBox.cpp
+++ b/Engine/Source/Editor/Sequencer/Private/Filters/Widgets/SSequencerSearchBox.cpp
@@ -89,7 +89,7 @@ void SSequencerSearchBox::Construct(const FArguments& InArgs, const TWeakPtr<ISe
 TSharedRef<SWidget> SSequencerSearchBox::GetSuggestionListMenuContent()
 {
 	const USequencerSettings* const SequencerSettings = GetDefault<USequencerSettings>();
-	check(IsValid(SequencerSettings));
+	check(SequencerSettings);
 
 	return SNew(SBorder)
 		.BorderImage(FAppStyle::GetBrush(TEXT("Menu.Background")))
diff --git a/Engine/Source/Editor/Sequencer/Public/Filters/SequencerFilterBase.h b/Engine/Source/Editor/Sequencer/Public/Filters/SequencerFilterBase.h
--- a/Engine/Source/Editor/Sequencer/Public/Filters/SequencerFilterBase.h
+++ b/Engine/Source/Editor/Sequencer/Public/Filters/SequencerFilterBase.h
@@ -17,24 +17,24 @@ public:
 	static bool SupportsLevelSequence(UMovieSceneSequence* const InSequence)
 	{
 		static UClass* const LevelSequenceClass = FindObject<UClass>(nullptr, TEXT("/Script/LevelSequence.LevelSequence"), true);
-		return IsValid(InSequence)
-			&& IsValid(LevelSequenceClass)
+		return InSequence
+			&& LevelSequenceClass
 			&& InSequence->GetClass()->IsChildOf(LevelSequenceClass);
 	}
 
 	static bool SupportsUMGSequence(UMovieSceneSequence* const InSequence)
 	{
 		static UClass* const WidgetAnimationClass = FindObject<UClass>(nullptr, TEXT("/Script/UMG.WidgetAnimation"), true);
-		return IsValid(InSequence)
-			&& IsValid(WidgetAnimationClass)
+		return InSequence
+			&& WidgetAnimationClass
 			&& InSequence->GetClass()->IsChildOf(WidgetAnimationClass);
 	}
 
 	static bool SupportsDaySequence(UMovieSceneSequence* const InSequence)
 	{
 		static UClass* const DaySequenceClass = FindObject<UClass>(nullptr, TEXT("/Script/DaySequence.DaySequence"), true);
-		return IsValid(InSequence)
-			&& IsValid(DaySequenceClass)
+		return InSequence
+			&& DaySequenceClass
 			&& InSequence->GetClass()->IsChildOf(DaySequenceClass);
 	}
 
diff --git a/Engine/Source/Editor/Sequencer/Public/Filters/SequencerTrackFilterBase.h b/Engine/Source/Editor/Sequencer/Public/Filters/SequencerTrackFilterBase.h
--- a/Engine/Source/Editor/Sequencer/Public/Filters/SequencerTrackFilterBase.h
+++ b/Engine/Source/Editor/Sequencer/Public/Filters/SequencerTrackFilterBase.h
@@ -40,7 +40,7 @@ public:
 	template<typename InTrackClassType>
 	static bool IsSequenceTrackSupported(UMovieSceneSequence* const InSequence)
 	{
-		const ETrackSupport Support = IsValid(InSequence)
+		const ETrackSupport Support = InSequence
 			? InSequence->IsTrackSupported(InTrackClassType::StaticClass())
 			: ETrackSupport::NotSupported;
 		return Support == ETrackSupport::Supported;
@@ -105,7 +105,7 @@ public:
 	{
 		FSequencerFilterData& FilterData = GetFilterInterface().GetFilterData();
 		const UMovieSceneTrack* const TrackObject = FilterData.ResolveMovieSceneTrackObject(InItem);
-		return IsValid(TrackObject) && TrackObject->IsA(InClassType::StaticClass());
+		return TrackObject && TrackObject->IsA(InClassType::StaticClass());
 	}
 	//~ End IFilter
 
@@ -135,7 +135,7 @@ public:
 		FSequencerFilterData& FilterData = GetFilterInterface().GetFilterData();
 
 		UObject* const BoundObject = FilterData.ResolveTrackBoundObject(GetSequencer(), InItem);
-		if (!IsValid(BoundObject))
+		if (!BoundObject)
 		{
 			return false;
 		}
