Sequencer: Fix time display format not getting the current settings

This is a regresion from.38446563.
This fixes Take Recorder not setting/getting Take Recorder Settings since they are set after initialization.

#rnx

diff --git a/Engine/Source/Editor/Sequencer/Private/Sequencer.cpp b/Engine/Source/Editor/Sequencer/Private/Sequencer.cpp
--- a/Engine/Source/Editor/Sequencer/Private/Sequencer.cpp
+++ b/Engine/Source/Editor/Sequencer/Private/Sequencer.cpp
@@ -392,15 +392,15 @@ void FSequencer::InitSequencer(const FSequencerInitParams& InitParams
 	{
 
 		TAttribute<EFrameNumberDisplayFormats> GetDisplayFormatAttr = MakeAttributeLambda(
-			[Settings = this->Settings]
+			[this]
 			{
-				return Settings->GetTimeDisplayFormat();
+				return GetSequencerSettings()->GetTimeDisplayFormat();
 			}
 		);
 		TAttribute<uint8> GetZeroPadFramesAttr = MakeAttributeLambda(
-			[Settings = this->Settings]()->uint8
+			[this]()->uint8
 			{
-				return Settings->GetZeroPadFrames();
+				return GetSequencerSettings()->GetZeroPadFrames();
 			}
 		);
 
