Sequencer: preview camera cuts in simulate

This also lets users enable/disable seeing camera cuts in PIE after ejecting from the player controller.

diff --git a/Engine/Source/Editor/Sequencer/Private/SSequencer.cpp b/Engine/Source/Editor/Sequencer/Private/SSequencer.cpp
--- a/Engine/Source/Editor/Sequencer/Private/SSequencer.cpp
+++ b/Engine/Source/Editor/Sequencer/Private/SSequencer.cpp
@@ -1918,6 +1918,7 @@ TSharedRef<SWidget> SSequencer::MakeViewMenu()
 		{
 			MenuBuilder.AddMenuEntry(FSequencerCommands::Get().TogglePilotCamera);
 			MenuBuilder.AddMenuEntry(FSequencerCommands::Get().ToggleRestoreOriginalViewportOnCameraCutUnlock);
+			MenuBuilder.AddMenuEntry(FSequencerCommands::Get().TogglePreviewCameraCutsInSimulate);
 		}
 		MenuBuilder.EndSection();
 	}
diff --git a/Engine/Source/Editor/Sequencer/Private/Sequencer.cpp b/Engine/Source/Editor/Sequencer/Private/Sequencer.cpp
--- a/Engine/Source/Editor/Sequencer/Private/Sequencer.cpp
+++ b/Engine/Source/Editor/Sequencer/Private/Sequencer.cpp
@@ -86,6 +86,7 @@
 #include "MovieSceneBindingProxy.h"
 #include "MovieSceneSequence.h"
 #include "MovieSceneFolder.h"
+#include "MovieSceneTracksSettings.h"
 #include "PropertyEditorModule.h"
 #include "PropertyHandle.h"
 #include "EditorWidgetsModule.h"
@@ -387,6 +388,8 @@ void FSequencer::InitSequencer(const FSequencerInitParams& InitParams
 	Settings->GetOnShowSelectedNodesOnlyChanged().AddSP(this, &FSequencer::OnSelectedNodesOnlyChanged);
 	Settings->GetOnTimeDisplayFormatChanged().AddSP(this, &FSequencer::OnTimeDisplayFormatChanged);
 
+	const FString TracksSettingsUniqueName = FString::Format(TEXT("{0}.TracksSettings"), { InitParams.ViewParams.UniqueName });
+	TracksSettings = USequencerSettingsContainer::GetOrCreate<UMovieSceneTracksSettings>(*TracksSettingsUniqueName);
 
 	// Initialize numeric type interfaces
 	{
@@ -11104,6 +11107,12 @@ void FSequencer::BindCommands()
 		FCanExecuteAction::CreateLambda([] { return true; }),
 		FIsActionChecked::CreateLambda([this] { return Settings->GetRestoreOriginalViewportOnCameraCutUnlock(); }) );
 
+	SequencerCommandBindings->MapAction(
+		Commands.TogglePreviewCameraCutsInSimulate,
+		FExecuteAction::CreateLambda([this] { TracksSettings->SetPreviewCameraCutsInSimulate(!TracksSettings->GetPreviewCameraCutsInSimulate()); }),
+		FCanExecuteAction::CreateLambda([] { return true; }),
+		FIsActionChecked::CreateLambda([this] { return TracksSettings->GetPreviewCameraCutsInSimulate(); }) );
+
 	SequencerCommandBindings->MapAction(
 		Commands.ToggleExpandCollapseNodes,
 		FExecuteAction::CreateSP(this, &FSequencer::ToggleExpandCollapseNodes));
diff --git a/Engine/Source/Editor/Sequencer/Private/Sequencer.h b/Engine/Source/Editor/Sequencer/Private/Sequencer.h
--- a/Engine/Source/Editor/Sequencer/Private/Sequencer.h
+++ b/Engine/Source/Editor/Sequencer/Private/Sequencer.h
@@ -71,6 +71,7 @@ class UMovieSceneCopyableBinding;
 class UMovieSceneCompiledDataManager;
 class UMovieSceneCopyableTrack;
 class UMovieSceneNodeGroup;
+class UMovieSceneTracksSettings;
 template <class TClass> class TSubclassOf;
 class UMovieSceneCustomBinding;
 
@@ -1298,6 +1299,9 @@ private:
 	/** User-supplied settings object for this sequencer */
 	TObjectPtr<USequencerSettings> Settings;
 
+	/** User-supplied settings object for this sequencer */
+	TObjectPtr<UMovieSceneTracksSettings> TracksSettings;
+
 	/** Command list for sequencer commands (Sequencer widgets only). */
 	TSharedRef<FUICommandList> SequencerCommandBindings;
 
diff --git a/Engine/Source/Editor/Sequencer/Private/SequencerCommands.cpp b/Engine/Source/Editor/Sequencer/Private/SequencerCommands.cpp
--- a/Engine/Source/Editor/Sequencer/Private/SequencerCommands.cpp
+++ b/Engine/Source/Editor/Sequencer/Private/SequencerCommands.cpp
@@ -59,6 +59,7 @@ void FSequencerCommands::RegisterCommands()
 
 	UI_COMMAND( ToggleAutoExpandNodesOnSelection, "Auto Expand Nodes on Selection", "Toggle auto expanding the outliner tree on child selection", EUserInterfaceActionType::ToggleButton, FInputChord());
 	UI_COMMAND( ToggleRestoreOriginalViewportOnCameraCutUnlock, "Restore Cinematic Viewports", "Toggle whether unlocking a camera cut track should return the viewport to its original location, or keep it where the camera cut was", EUserInterfaceActionType::ToggleButton, FInputChord() );
+	UI_COMMAND( TogglePreviewCameraCutsInSimulate, "Preview Camera Cuts in Simulate", "Toggle whether camera cuts are evaluated in SIE, or after ejecting from the player controller in PIE", EUserInterfaceActionType::ToggleButton, FInputChord() );
 	UI_COMMAND( ToggleExpandCollapseNodes, "Expand/Collapse Nodes", "Toggle expand or collapse selected nodes", EUserInterfaceActionType::Button, FInputChord(EKeys::V) );
 	UI_COMMAND( ToggleExpandCollapseNodesAndDescendants, "Expand/Collapse Nodes and Descendants", "Toggle expand or collapse selected nodes and descendants", EUserInterfaceActionType::Button, FInputChord(EModifierKey::Shift, EKeys::V) );
 	UI_COMMAND( ExpandAllNodes, "Expand All Nodes", "Expand all nodes and descendants", EUserInterfaceActionType::Button, FInputChord() );
diff --git a/Engine/Source/Editor/Sequencer/Public/SequencerCommands.h b/Engine/Source/Editor/Sequencer/Public/SequencerCommands.h
--- a/Engine/Source/Editor/Sequencer/Public/SequencerCommands.h
+++ b/Engine/Source/Editor/Sequencer/Public/SequencerCommands.h
@@ -160,6 +160,9 @@ public:
 	 */
 	TSharedPtr< FUICommandInfo > ToggleRestoreOriginalViewportOnCameraCutUnlock;
 
+	/** Toggle whether to preview camera cuts while in SIE, or after ejecting from the player controller in PIE. */
+	TSharedPtr< FUICommandInfo > TogglePreviewCameraCutsInSimulate;
+
 	/** Expand/collapse nodes */
 	TSharedPtr< FUICommandInfo > ToggleExpandCollapseNodes;
 
diff --git a/Engine/Source/Runtime/MovieSceneTracks/Private/MovieSceneTracksSettings.cpp b/Engine/Source/Runtime/MovieSceneTracks/Private/MovieSceneTracksSettings.cpp
new file mode 100644
--- /dev/null
+++ b/Engine/Source/Runtime/MovieSceneTracks/Private/MovieSceneTracksSettings.cpp
@@ -0,0 +1,21 @@
+// Copyright Epic Games, Inc. All Rights Reserved.
+
+#include "MovieSceneTracksSettings.h"
+
+#include UE_INLINE_GENERATED_CPP_BY_NAME(MovieSceneTracksSettings)
+
+UMovieSceneTracksSettings::UMovieSceneTracksSettings(const FObjectInitializer& ObjectInitializer)
+	: Super(ObjectInitializer)
+{
+	bPreviewCameraCutsInSimulate = true;
+}
+
+void UMovieSceneTracksSettings::SetPreviewCameraCutsInSimulate(bool bInPreviewCameraCutsInSimulate)
+{
+	if (bInPreviewCameraCutsInSimulate != bPreviewCameraCutsInSimulate)
+	{
+		bPreviewCameraCutsInSimulate = bInPreviewCameraCutsInSimulate;
+		SaveConfig();
+	}
+}
+
diff --git a/Engine/Source/Runtime/MovieSceneTracks/Private/TrackInstances/MovieSceneCameraCutEditorHandler.cpp b/Engine/Source/Runtime/MovieSceneTracks/Private/TrackInstances/MovieSceneCameraCutEditorHandler.cpp
--- a/Engine/Source/Runtime/MovieSceneTracks/Private/TrackInstances/MovieSceneCameraCutEditorHandler.cpp
+++ b/Engine/Source/Runtime/MovieSceneTracks/Private/TrackInstances/MovieSceneCameraCutEditorHandler.cpp
@@ -8,6 +8,7 @@
 #include "Evaluation/CameraCutPlaybackCapability.h"
 #include "IMovieScenePlayer.h"
 #include "LevelEditorViewport.h"
+#include "MovieSceneTracksSettings.h"
 #include "Systems/MovieSceneMotionVectorSimulationSystem.h"
 #include "TrackInstances/MovieSceneCameraCutTrackInstance.h"
 #include "TrackInstances/MovieSceneCameraCutViewportPreviewer.h"
@@ -17,12 +18,17 @@ namespace UE::MovieScene
 
 bool FPreAnimatedCameraCutEditorTraits::ShouldHandleViewportCameraCuts(UWorld* ViewportWorld)
 {
+	const bool bSupportsSimulateWorlds = GetDefault<UMovieSceneTracksSettings>()->GetPreviewCameraCutsInSimulate();
+
 	return ViewportWorld && 
 		// We can handle editor worlds, and game worlds that don't have an active player controller/pawn,
-		// such as PIE/SIE where the user has "ejected" out of the player controller.
+		// such as PIE/SIE where the user has "ejected" out of the player controller, as long as the user
+		// hasn't opt out of that in the settings.
 		(ViewportWorld->WorldType == EWorldType::Editor || 
 			ViewportWorld->WorldType == EWorldType::EditorPreview ||
-			(ViewportWorld->WorldType == EWorldType::PIE && 
+			(
+				bSupportsSimulateWorlds &&
+				ViewportWorld->WorldType == EWorldType::PIE && 
 				GEditor && GEditor->bIsSimulatingInEditor));
 }
 
diff --git a/Engine/Source/Runtime/MovieSceneTracks/Public/MovieSceneTracksSettings.h b/Engine/Source/Runtime/MovieSceneTracks/Public/MovieSceneTracksSettings.h
new file mode 100644
--- /dev/null
+++ b/Engine/Source/Runtime/MovieSceneTracks/Public/MovieSceneTracksSettings.h
@@ -0,0 +1,39 @@
+// Copyright Epic Games, Inc. All Rights Reserved.
+
+#pragma once
+
+#include "UObject/Object.h"
+#include "UObject/ObjectMacros.h"
+
+#include "MovieSceneTracksSettings.generated.h"
+
+/** Options for some of the Sequencer systems in this module. */
+UCLASS(config=EditorPerProjectUserSettings, PerObjectConfig)
+class MOVIESCENETRACKS_API UMovieSceneTracksSettings : public UObject
+{
+	GENERATED_BODY()
+
+public:
+
+	UMovieSceneTracksSettings(const FObjectInitializer& ObjectInitializer);
+
+	/**
+	 * Gets whether camera cut tracks should take control of the viewport in SIE, or PIE while ejected from the player controller.
+	 */
+	bool GetPreviewCameraCutsInSimulate() const { return bPreviewCameraCutsInSimulate; }
+
+	/**
+	 * Sets whether camera cut tracks should take control of the viewport in SIE, or PIE while ejected from the player controller.
+	 */
+	void SetPreviewCameraCutsInSimulate(bool bInPreviewCameraCutsInSimulate);
+
+protected:
+
+	/**
+	 * Whether camera cut tracks should take control of the viewport in SIE (Simulate in Editor) or after ejecting
+	 * from the player controller while in PIE.
+	 */
+	UPROPERTY(config, EditAnywhere, Category=General)
+	bool bPreviewCameraCutsInSimulate;
+};
+
