[Backout] - CL42237728
#rnx
Original CL Desc
-----------------------------------------------------------------
Sequencer: Fix time display format not getting the current settings

This is a regresion from.38446563.
This fixes Take Recorder not setting/getting Take Recorder Settings since they are set after initialization.

#rnx

diff --git a/Engine/Source/Editor/Sequencer/Private/Sequencer.cpp b/Engine/Source/Editor/Sequencer/Private/Sequencer.cpp
--- a/Engine/Source/Editor/Sequencer/Private/Sequencer.cpp
+++ b/Engine/Source/Editor/Sequencer/Private/Sequencer.cpp
@@ -395,15 +395,15 @@ void FSequencer::InitSequencer(const FSequencerInitParams& InitParams
 	{
 
 		TAttribute<EFrameNumberDisplayFormats> GetDisplayFormatAttr = MakeAttributeLambda(
-			[this]
+			[Settings = this->Settings]
 			{
-				return GetSequencerSettings()->GetTimeDisplayFormat();
+				return Settings->GetTimeDisplayFormat();
 			}
 		);
 		TAttribute<uint8> GetZeroPadFramesAttr = MakeAttributeLambda(
-			[this]()->uint8
+			[Settings = this->Settings]()->uint8
 			{
-				return GetSequencerSettings()->GetZeroPadFrames();
+				return Settings->GetZeroPadFrames();
 			}
 		);
 
