[Sequencer] Include children in Name and BindingName text expression filters

diff --git a/Engine/Source/Editor/Sequencer/Private/Filters/TextExpressions/SequencerTextFilterExpression_BindingName.cpp b/Engine/Source/Editor/Sequencer/Private/Filters/TextExpressions/SequencerTextFilterExpression_BindingName.cpp
--- a/Engine/Source/Editor/Sequencer/Private/Filters/TextExpressions/SequencerTextFilterExpression_BindingName.cpp
+++ b/Engine/Source/Editor/Sequencer/Private/Filters/TextExpressions/SequencerTextFilterExpression_BindingName.cpp
@@ -38,15 +38,16 @@ bool FSequencerTextFilterExpression_BindingName::TestComplexExpression(const FNa
 		return true;
 	}
 
-	const TViewModelPtr<FObjectBindingModel> ObjectBindingModel = FilterItem->FindAncestorOfType<FObjectBindingModel>(true);
-	if (!ObjectBindingModel.IsValid())
+	for (const TViewModelPtr<FObjectBindingModel>& ObjectBindingModel : FilterItem->GetAncestorsOfType<FObjectBindingModel>())
 	{
-		return false;
+		const FString Label = ObjectBindingModel->GetLabel().ToString();
+		if (TextFilterUtils::TestComplexExpression(Label, InValue, InComparisonOperation, InTextComparisonMode))
+		{
+			return true;
+		}
 	}
 
-	const FString BindingLabel = ObjectBindingModel->GetLabel().ToString();
-	const bool bPassed = TextFilterUtils::TestComplexExpression(BindingLabel, InValue, InComparisonOperation, InTextComparisonMode);
-	return bPassed;
+	return false;
 }
 
 #undef LOCTEXT_NAMESPACE
diff --git a/Engine/Source/Editor/Sequencer/Private/Filters/TextExpressions/SequencerTextFilterExpression_Name.cpp b/Engine/Source/Editor/Sequencer/Private/Filters/TextExpressions/SequencerTextFilterExpression_Name.cpp
--- a/Engine/Source/Editor/Sequencer/Private/Filters/TextExpressions/SequencerTextFilterExpression_Name.cpp
+++ b/Engine/Source/Editor/Sequencer/Private/Filters/TextExpressions/SequencerTextFilterExpression_Name.cpp
@@ -1,9 +1,9 @@
 ﻿// Copyright Epic Games, Inc. All Rights Reserved.
 
 #include "SequencerTextFilterExpression_Name.h"
-#include "MovieSceneNameableTrack.h"
-#include "MVVM/ViewModelPtr.h"
 #include "MVVM/Extensions/IOutlinerExtension.h"
+#include "MVVM/ViewModelPtr.h"
+#include "MVVM/ViewModels/ViewModelIterators.h"
 
 using namespace UE::Sequencer;
 
@@ -39,11 +39,13 @@ bool FSequencerTextFilterExpression_Name::TestComplexExpression(const FName& InK
 		return true;
 	}
 
-	if (const TViewModelPtr<IOutlinerExtension> OutlinerExtension = FilterItem.ImplicitCast())
+	for (const TViewModelPtr<IOutlinerExtension>& OutlinerExtension : FilterItem->GetAncestorsOfType<IOutlinerExtension>())
 	{
 		const FString Label = OutlinerExtension->GetLabel().ToString();
-		const bool bPassed = TextFilterUtils::TestComplexExpression(Label, InValue, InComparisonOperation, InTextComparisonMode);
-		return bPassed;
+		if (TextFilterUtils::TestComplexExpression(Label, InValue, InComparisonOperation, InTextComparisonMode))
+		{
+			return true;
+		}
 	}
 
 	return false;
