[Sequencer] Filtering: Show children of passed outliner item text filter

diff --git a/Engine/Plugins/Experimental/SequenceNavigator/Source/SequenceNavigator/Private/Filters/TextExpressions/NavigationToolFilterTextExpressionContext.cpp b/Engine/Plugins/Experimental/SequenceNavigator/Source/SequenceNavigator/Private/Filters/TextExpressions/NavigationToolFilterTextExpressionContext.cpp
--- a/Engine/Plugins/Experimental/SequenceNavigator/Source/SequenceNavigator/Private/Filters/TextExpressions/NavigationToolFilterTextExpressionContext.cpp
+++ b/Engine/Plugins/Experimental/SequenceNavigator/Source/SequenceNavigator/Private/Filters/TextExpressions/NavigationToolFilterTextExpressionContext.cpp
@@ -22,10 +22,25 @@ void FNavigationToolFilterTextExpressionContext::SetFilterItem(const FNavigation
 bool FNavigationToolFilterTextExpressionContext::TestBasicStringExpression(const FTextFilterString& InValue
 	, const ETextFilterTextComparisonMode InTextComparisonMode) const
 {
-	if (const TSharedPtr<INavigationToolItem> FilterItem = WeakFilterItem.Pin())
+	const TSharedPtr<INavigationToolItem> FilterItem = WeakFilterItem.Pin();
+	if (!FilterItem.IsValid())
 	{
-		const FString Label = FilterItem->GetDisplayName().ToString();
-		if (TextFilterUtils::TestBasicStringExpression(Label, InValue, ETextFilterTextComparisonMode::Partial))
+		return false;
+	}
+
+	constexpr ETextFilterTextComparisonMode TextComparisonMode = ETextFilterTextComparisonMode::Partial;
+
+	FString Label = FilterItem->GetDisplayName().ToString();
+
+	if (TextFilterUtils::TestBasicStringExpression(Label, InValue, TextComparisonMode))
+	{
+		return true;
+	}
+
+	for (const FNavigationToolItemPtr& ParentItem : FilterItem->GetParents())
+	{
+		Label = ParentItem->GetDisplayName().ToString();
+		if (TextFilterUtils::TestBasicStringExpression(Label, InValue, TextComparisonMode))
 		{
 			return true;
 		}
diff --git a/Engine/Source/Editor/Sequencer/Private/Filters/TextExpressions/SequencerTextFilterExpressionContext.cpp b/Engine/Source/Editor/Sequencer/Private/Filters/TextExpressions/SequencerTextFilterExpressionContext.cpp
--- a/Engine/Source/Editor/Sequencer/Private/Filters/TextExpressions/SequencerTextFilterExpressionContext.cpp
+++ b/Engine/Source/Editor/Sequencer/Private/Filters/TextExpressions/SequencerTextFilterExpressionContext.cpp
@@ -17,49 +17,52 @@ FSequencerTextFilterExpressionContext::FSequencerTextFilterExpressionContext(ISe
 void FSequencerTextFilterExpressionContext::SetFilterItem(FSequencerTrackFilterType InFilterItem)
 {
 	FilterItem = InFilterItem;
-
-	FSequencerFilterData& FilterData = FilterInterface.GetFilterData();
 }
 
 bool FSequencerTextFilterExpressionContext::TestBasicStringExpression(const FTextFilterString& InValue
 	, const ETextFilterTextComparisonMode InTextComparisonMode) const
 {
-	if (const TViewModelPtr<FChannelGroupOutlinerModel> ChannelGroupOutlinerModel = FilterItem.ImplicitCast())
+	constexpr ETextFilterTextComparisonMode TextComparisonMode = ETextFilterTextComparisonMode::Partial;
+
+	for (const FViewModelPtr& OutlinerItem : FilterItem->GetAncestors(/*bIncludeThis=*/true))
 	{
-		const FTextFilterString TrackLabel = ChannelGroupOutlinerModel->GetLabel().ToString();
-		if (TextFilterUtils::TestBasicStringExpression(TrackLabel, InValue, ETextFilterTextComparisonMode::Partial))
+		if (const TViewModelPtr<FChannelGroupOutlinerModel> ChannelGroupOutlinerModel = OutlinerItem.ImplicitCast())
 		{
-			return true;
-		}
+			const FTextFilterString TrackLabel = ChannelGroupOutlinerModel->GetLabel().ToString();
+			if (TextFilterUtils::TestBasicStringExpression(TrackLabel, InValue, TextComparisonMode))
+			{
+				return true;
+			}
 
-		const FTextFilterString ChannelName = ChannelGroupOutlinerModel->GetChannelName();
-		if (TextFilterUtils::TestBasicStringExpression(ChannelName, InValue, ETextFilterTextComparisonMode::Partial))
-		{
-			return true;
+			const FTextFilterString ChannelName = ChannelGroupOutlinerModel->GetChannelName();
+			if (TextFilterUtils::TestBasicStringExpression(ChannelName, InValue, TextComparisonMode))
+			{
+				return true;
+			}
 		}
-	}
-	else if (const TViewModelPtr<FCategoryGroupModel> CategoryGroupModel = FilterItem.ImplicitCast())
-	{
-		const FTextFilterString CategoryName = CategoryGroupModel->GetCategoryName();
-		if (TextFilterUtils::TestBasicStringExpression(CategoryName, InValue, ETextFilterTextComparisonMode::Partial))
+		else if (const TViewModelPtr<FCategoryGroupModel> CategoryGroupModel = OutlinerItem.ImplicitCast())
 		{
-			return true;
+			const FTextFilterString CategoryName = CategoryGroupModel->GetCategoryName();
+			if (TextFilterUtils::TestBasicStringExpression(CategoryName, InValue, TextComparisonMode))
+			{
+				return true;
+			}
 		}
-	}
-	else if (const TViewModelPtr<FChannelGroupModel> ChannelGroupModel = FilterItem.ImplicitCast())
-	{
-		const FTextFilterString CategoryName = ChannelGroupModel->GetChannelName();
-		if (TextFilterUtils::TestBasicStringExpression(CategoryName, InValue, ETextFilterTextComparisonMode::Partial))
+		else if (const TViewModelPtr<FChannelGroupModel> ChannelGroupModel = OutlinerItem.ImplicitCast())
 		{
-			return true;
+			const FTextFilterString CategoryName = ChannelGroupModel->GetChannelName();
+			if (TextFilterUtils::TestBasicStringExpression(CategoryName, InValue, TextComparisonMode))
+			{
+				return true;
+			}
 		}
-	}
-	else if (const TViewModelPtr<IOutlinerExtension> OutlinerExtension = FilterItem.ImplicitCast())
-	{
-		const FTextFilterString TrackLabel = OutlinerExtension->GetLabel().ToString();
-		if (TextFilterUtils::TestBasicStringExpression(TrackLabel, InValue, ETextFilterTextComparisonMode::Partial))
+		else if (const TViewModelPtr<IOutlinerExtension> OutlinerExtension = OutlinerItem.ImplicitCast())
 		{
-			return true;
+			const FTextFilterString TrackLabel = OutlinerExtension->GetLabel().ToString();
+			if (TextFilterUtils::TestBasicStringExpression(TrackLabel, InValue, TextComparisonMode))
+			{
+				return true;
+			}
 		}
 	}
 
