[Sequencer] All tracks are muted when Solo one of the tracks of animation track

diff --git a/Engine/Source/Editor/Sequencer/Private/SSequencer.cpp b/Engine/Source/Editor/Sequencer/Private/SSequencer.cpp
--- a/Engine/Source/Editor/Sequencer/Private/SSequencer.cpp
+++ b/Engine/Source/Editor/Sequencer/Private/SSequencer.cpp
@@ -2683,15 +2683,16 @@ void SSequencer::UpdateLayoutTree()
 			NewNodePathsToIsolate.Empty();
 		}
 
+		// Local Mute/Solo state (transient, non dirtying, mute/solo evaluation)
 		if (Sequencer->GetFocusedMovieSceneSequence())
 		{
 			bool bAnyChanged = false;
 
-			TSharedPtr<FSharedViewModelData> SharedData = Sequencer->GetViewModel()->GetRootModel()->GetSharedData();
-			FOutlinerCacheExtension*           OutlinerCache = SharedData->CastThis<FOutlinerCacheExtension>();
-			FDeactiveStateCacheExtension* DeactiveState = SharedData->CastThis<FDeactiveStateCacheExtension>();
-			FMuteStateCacheExtension*          MuteState     = SharedData->CastThis<FMuteStateCacheExtension>();
-			FSoloStateCacheExtension*          SoloState     = SharedData->CastThis<FSoloStateCacheExtension>();
+			TSharedPtr<FSharedViewModelData> SharedData    = Sequencer->GetViewModel()->GetRootModel()->GetSharedData();
+			FOutlinerCacheExtension*         OutlinerCache = SharedData->CastThis<FOutlinerCacheExtension>();
+			FDeactiveStateCacheExtension*    DeactiveState = SharedData->CastThis<FDeactiveStateCacheExtension>();
+			FMuteStateCacheExtension*        MuteState     = SharedData->CastThis<FMuteStateCacheExtension>();
+			FSoloStateCacheExtension*        SoloState     = SharedData->CastThis<FSoloStateCacheExtension>();
 
 			check(OutlinerCache && DeactiveState && MuteState && SoloState);
 
@@ -2720,23 +2721,16 @@ void SSequencer::UpdateLayoutTree()
 
 				if (const TViewModelPtr<FTrackRowModel> TrackRowModel = TrackNode.ImplicitCast())
 				{
-					// Mute/Solo state (local, non dirtying, mute/solo evaluation)
-					if (bLocalEvalDisabled != Track->IsLocalRowEvalDisabled(TrackNode->GetRowIndex()))
-					{
-						Track->MarkAsChanged();
-						Track->SetLocalRowEvalDisabled(bLocalEvalDisabled, TrackNode->GetRowIndex());
-						bAnyChanged = true;
-					}
+					Track->MarkAsChanged();
+					Track->SetLocalEvalDisabled(false);
+					Track->SetLocalRowEvalDisabled(bLocalEvalDisabled, TrackNode->GetRowIndex());
+					bAnyChanged = true;
 				}
 				else
 				{
-					// Mute/Solo state (local, non dirtying, mute/solo evaluation)
-					if (bLocalEvalDisabled != Track->IsLocalEvalDisabled())
-					{
-						Track->MarkAsChanged();
-						Track->SetLocalEvalDisabled(bLocalEvalDisabled);
-						bAnyChanged = true;
-					}
+					Track->MarkAsChanged();
+					Track->SetLocalEvalDisabled(bLocalEvalDisabled);
+					bAnyChanged = true;
 				}
 			}
 
