Fix Custom Clock Snapping misbehaving when ForceWholeFrames is enabled.

#rnx

diff --git a/Engine/Source/Editor/Sequencer/Private/Sequencer.cpp b/Engine/Source/Editor/Sequencer/Private/Sequencer.cpp
--- a/Engine/Source/Editor/Sequencer/Private/Sequencer.cpp
+++ b/Engine/Source/Editor/Sequencer/Private/Sequencer.cpp
@@ -24,6 +24,7 @@
 #include "MVVM/ViewModels/PossessableModel.h"
 #include "MVVM/Views/SOutlinerView.h"
 #include "MVVM/Extensions/IDeletableExtension.h"
+#include "MVVM/Extensions/IClockExtension.h"
 #include "MVVM/CurveEditorExtension.h"
 #include "MVVM/CurveEditorIntegrationExtension.h"
 #include "MVVM/ObjectBindingModelStorageExtension.h"
@@ -3627,6 +3628,16 @@ void FSequencer::SnapSequencerTime(FFrameTime& ScrubTime)
 			ScrubTime = FFrameRate::Snap(ScrubTime, TickResolution, DisplayRate);
 		}
 	}
+	
+	if (GetSequencerSettings()->GetIsSnapEnabled())
+	{
+		using namespace UE::Sequencer;
+		TViewModelPtr<IClockExtension> Clock = GetViewModel()->GetRootSequenceModel().ImplicitCast();
+		if (Clock.IsValid() && Clock->SupportsSnapping() && Clock->ShouldSnapFrameTime())
+		{
+			ScrubTime = Clock->SnapFrameTime(ScrubTime);
+		}
+	}
 }
 
 void FSequencer::ForceEvaluate()
