Make TimeSliderController menu extendable

#rnx

diff --git a/Engine/Source/Editor/Sequencer/Private/SequencerTimeSliderController.cpp b/Engine/Source/Editor/Sequencer/Private/SequencerTimeSliderController.cpp
--- a/Engine/Source/Editor/Sequencer/Private/SequencerTimeSliderController.cpp
+++ b/Engine/Source/Editor/Sequencer/Private/SequencerTimeSliderController.cpp
@@ -9,6 +9,7 @@
 #include "TrackEditors/TimeWarpTrackEditor.h"
 #include "Channels/MovieSceneTimeWarpChannel.h"
 #include "SequencerTimeDomainOverride.h"
+#include "SequencerToolMenuContext.h"
 #include "SSequencer.h"
 #include "Fonts/SlateFontInfo.h"
 #include "Rendering/DrawElements.h"
@@ -1845,261 +1846,369 @@ ETimeSliderPlaybackStatus FSequencerTimeSliderController::GetPlaybackStatus() co
 TSharedRef<SWidget> FSequencerTimeSliderController::OpenSetPlaybackRangeMenu(const FGeometry& MyGeometry, const FPointerEvent& MouseEvent)
 {
 	using namespace UE::Sequencer;
-
-	TSharedPtr<FSequencer> Sequencer = WeakSequencer.Pin();
-	const bool bReadOnly = Sequencer && Sequencer->IsReadOnly();
 	
-	FScrubRangeToScreen RangeToScreen = FScrubRangeToScreen(GetViewRange(), MyGeometry.Size);
-	const float MousePixel = MyGeometry.AbsoluteToLocal(MouseEvent.GetScreenSpacePosition()).X;
-	FFrameNumber FrameNumber = ComputeFrameTimeFromMouse(MyGeometry, MouseEvent.GetScreenSpacePosition(), RangeToScreen).FrameNumber;
-
-	const bool bShouldCloseWindowAfterMenuSelection = true;
-	FMenuBuilder MenuBuilder(bShouldCloseWindowAfterMenuSelection, nullptr);
-
-	FText CurrentTimeText;
-	CurrentTimeText = FText::FromString(Sequencer->GetNumericTypeInterface()->ToString(FrameNumber.Value));
-	
-	TRange<FFrameNumber> PlaybackRange = TimeSliderArgs.PlaybackRange.Get();
-	TOptional<TRange<FFrameNumber>> SubSequenceRange = TimeSliderArgs.SubSequenceRange.Get();
-
-	MenuBuilder.BeginSection("SequencerPlaybackRangeMenu", FText::Format(LOCTEXT("PlaybackRangeTextFormat", "Playback Range ({0}):"), CurrentTimeText));
-	{
-		MenuBuilder.AddMenuEntry(
-			LOCTEXT("SetPlaybackStart", "Set Start Time"),
-			FText(),
-			FSlateIcon(),
-			FUIAction(
-				FExecuteAction::CreateLambda([this, FrameNumber]{ SetPlaybackRangeStart(FrameNumber); }),
-				FCanExecuteAction::CreateLambda([this, FrameNumber, PlaybackRange]{ return !TimeSliderArgs.IsPlaybackRangeLocked.Get() && FrameNumber < UE::MovieScene::DiscreteExclusiveUpper(PlaybackRange); })
-			)
-		);
-
-		MenuBuilder.AddMenuEntry(
-			LOCTEXT("SetPlaybackEnd", "Set End Time"),
-			FText(),
-			FSlateIcon(),
-			FUIAction(
-				FExecuteAction::CreateLambda([this, FrameNumber]{ SetPlaybackRangeEnd(FrameNumber); }),
-				FCanExecuteAction::CreateLambda([this, FrameNumber, PlaybackRange]{ return !TimeSliderArgs.IsPlaybackRangeLocked.Get() && FrameNumber >= UE::MovieScene::DiscreteInclusiveLower(PlaybackRange); })
-			)
-		);
-
-		MenuBuilder.AddMenuEntry(
-			LOCTEXT("ConformToSubsequenceRange", "Conform to Range"),
-			LOCTEXT("ConformToSubsequenceRangeTooltip", "Conform the start and end time to the extents of the subsequence range"),
-			FSlateIcon(),
-			FUIAction(
-				FExecuteAction::CreateLambda([this, SubSequenceRange] { SetPlaybackRangeStart(SubSequenceRange.GetValue().GetLowerBoundValue()); SetPlaybackRangeEnd(SubSequenceRange.GetValue().GetUpperBoundValue()); }),
-				FCanExecuteAction::CreateLambda([this, SubSequenceRange] { return !TimeSliderArgs.IsPlaybackRangeLocked.Get() && SubSequenceRange.IsSet(); })
-			)
-		);
-
-		MenuBuilder.AddMenuEntry(
-			LOCTEXT("ToggleLocked", "Locked"),
-			LOCTEXT("ToggleLockedTooltip", "Lock/Unlock the playback range"),
-			FSlateIcon(),
-			FUIAction(
-				FExecuteAction::CreateLambda([this] { TimeSliderArgs.OnTogglePlaybackRangeLocked.ExecuteIfBound(); }),
-				FCanExecuteAction::CreateLambda([bReadOnly]{ return !bReadOnly; }),
-				FIsActionChecked::CreateLambda([this] { return TimeSliderArgs.IsPlaybackRangeLocked.Get(); })
-			),
-			NAME_None,
-			EUserInterfaceActionType::ToggleButton
-		);
-	}
-	MenuBuilder.EndSection(); // SequencerPlaybackRangeMenu
+	class SMarkedFramePropertyWidget : public SCompoundWidget, public FNotifyHook
+	{
+	public:
+		UMovieScene* MovieSceneToModify;
+		TSharedPtr<IStructureDetailsView> DetailsView;
+		TWeakPtr<FSequencer> WeakSequencer;
 
-	TRange<FFrameNumber> SelectionRange = TimeSliderArgs.SelectionRange.Get();
-	MenuBuilder.BeginSection("SequencerSelectionRangeMenu", FText::Format(LOCTEXT("SelectionRangeTextFormat", "Selection Range ({0}):"), CurrentTimeText));
-	{
-		MenuBuilder.AddMenuEntry(
-			LOCTEXT("SetSelectionStart", "Set Selection Start"),
-			FText(),
-			FSlateIcon(),
-			FUIAction(
-				FExecuteAction::CreateLambda([this, FrameNumber]{ SetSelectionRangeStart(FrameNumber); }),
-				FCanExecuteAction::CreateLambda([=]{ return SelectionRange.IsEmpty() || FrameNumber < UE::MovieScene::DiscreteExclusiveUpper(SelectionRange); })
-			)
-		);
+		SLATE_BEGIN_ARGS(SMarkedFramePropertyWidget){}
+		SLATE_END_ARGS()
 
-		MenuBuilder.AddMenuEntry(
-			LOCTEXT("SetSelectionEnd", "Set Selection End"),
-			FText(),
-			FSlateIcon(),
-			FUIAction(
-				FExecuteAction::CreateLambda([this, FrameNumber]{ SetSelectionRangeEnd(FrameNumber); }),
-				FCanExecuteAction::CreateLambda([=]{ return SelectionRange.IsEmpty() || FrameNumber >= UE::MovieScene::DiscreteInclusiveLower(SelectionRange); })
-			)
-		);
+		void Construct(const FArguments& InArgs, UMovieScene* InMovieScene, int32 InMarkedFrameIndex, TWeakPtr<FSequencer> InWeakSequencer)
+		{
+			MovieSceneToModify = InMovieScene;
+			WeakSequencer = InWeakSequencer;
+
+			FPropertyEditorModule& PropertyEditorModule = FModuleManager::GetModuleChecked<FPropertyEditorModule>("PropertyEditor");
+			FDetailsViewArgs DetailsViewArgs;
+			DetailsViewArgs.bAllowSearch = false;
+			DetailsViewArgs.bShowScrollBar = false;
+			DetailsViewArgs.NameAreaSettings = FDetailsViewArgs::HideNameArea;
+			DetailsViewArgs.NotifyHook = this;
+
+			FStructureDetailsViewArgs StructureDetailsViewArgs;
+			StructureDetailsViewArgs.bShowObjects = true;
+			StructureDetailsViewArgs.bShowAssets = true;
+			StructureDetailsViewArgs.bShowClasses = true;
+			StructureDetailsViewArgs.bShowInterfaces = true;
+
+			TSharedPtr<FStructOnScope> StructOnScope = MakeShared<FStructOnScope>(FMovieSceneMarkedFrame::StaticStruct(), (uint8 *)&InMovieScene->GetMarkedFrames()[InMarkedFrameIndex]);
+
+			DetailsView = PropertyEditorModule.CreateStructureDetailView(DetailsViewArgs, StructureDetailsViewArgs, nullptr);
+			DetailsView->GetDetailsView()->RegisterInstancedCustomPropertyTypeLayout("FrameNumber", 
+				FOnGetPropertyTypeCustomizationInstance::CreateSP(WeakSequencer.Pin().ToSharedRef(), &FSequencer::MakeFrameNumberDetailsCustomization));
+			DetailsView->SetStructureData(StructOnScope);
+
+			ChildSlot
+			[
+				DetailsView->GetWidget().ToSharedRef()
+			];
+		}
 
-		MenuBuilder.AddMenuEntry(
-			LOCTEXT("ClearSelectionRange", "Clear Selection Range"),
-			FText(),
-			FSlateIcon(),
-			FUIAction(
-				FExecuteAction::CreateLambda([this]{ TimeSliderArgs.OnSelectionRangeChanged.ExecuteIfBound(TRange<FFrameNumber>::Empty()); }),
-				FCanExecuteAction::CreateLambda([=]{ return !SelectionRange.IsEmpty(); })
-			)
-		);
-	}
-	MenuBuilder.EndSection(); // SequencerPlaybackRangeMenu
+		virtual void NotifyPreChange( FProperty* PropertyAboutToChange ) override
+		{
+			MovieSceneToModify->Modify();
+		}
 
-	UMovieSceneCompiledDataManager* CompiledDataManager = Sequencer->GetEvaluationTemplate().GetCompiledDataManager();
-	const FMovieSceneSequenceHierarchy* Hierarchy = CompiledDataManager->FindHierarchy(Sequencer->GetEvaluationTemplate().GetCompiledDataID());
+		virtual void NotifyPreChange( class FEditPropertyChain* PropertyAboutToChange ) override
+		{
+			MovieSceneToModify->Modify();
+		}
+	};
 
-	if (TimeSliderArgs.ScrubPositionParentChain.IsSet() && Hierarchy)
+	static const FName MenuName("Sequencer.TimelineMenu");
+	if (!UToolMenus::Get()->IsMenuRegistered(MenuName))
 	{
-		MenuBuilder.BeginSection("SequencerParentChainMenu");
+		UToolMenu* TimelineMenu = UToolMenus::Get()->RegisterMenu(MenuName, NAME_None, EMultiBoxType::Menu);
+		TimelineMenu->bSearchable = false;
+		TimelineMenu->bShouldCloseWindowAfterMenuSelection = true;
+		
+		TimelineMenu->AddDynamicSection(NAME_None, FNewToolMenuDelegate::CreateLambda([](UToolMenu* const InMenu)
 		{
-			TArray<FMovieSceneSequenceID> ParentChain = TimeSliderArgs.ScrubPositionParentChain.Get();
-			for (FMovieSceneSequenceID ParentID : ParentChain)
+			USequencerTimeSliderControllerMenuContext* Context = InMenu->FindContext<USequencerTimeSliderControllerMenuContext>();
+			if (!Context)
 			{
-				FText ParentText = Sequencer->GetRootMovieSceneSequence()->GetDisplayName();
-
-				for (const TTuple<FMovieSceneSequenceID, FMovieSceneSubSequenceData>& Pair : Hierarchy->AllSubSequenceData())
-				{
-					if (Pair.Key == ParentID && Pair.Value.GetSequence())
-					{
-						ParentText = Pair.Value.GetSequence()->GetDisplayName();
-						break;
-					}
-				}
+				return;
+			}
+			
+			TSharedPtr<FSequencerTimeSliderController> TimeSliderController = Context->WeakTimeSliderController.Pin();
+            if (!TimeSliderController.IsValid())
+            {
+            	return;
+            }
+
+            TSharedPtr<FSequencer> Sequencer = TimeSliderController->WeakSequencer.Pin();
+            if (!Sequencer.IsValid())
+            {
+            	return;
+            }
+
+			const bool bReadOnly = Sequencer->IsReadOnly();
+			const FGeometry& MyGeometry = Context->Geometry;
+			const FPointerEvent& MouseEvent = Context->PointerEvent;
+			const FFrameNumber& FrameNumber = Context->FrameNumber;
+				
+			FScrubRangeToScreen RangeToScreen = FScrubRangeToScreen(TimeSliderController->GetViewRange(), MyGeometry.Size);
+			const float MousePixel = MyGeometry.AbsoluteToLocal(MouseEvent.GetScreenSpacePosition()).X;
+			
+			FText CurrentTimeText = FText::FromString(Sequencer->GetNumericTypeInterface()->ToString(FrameNumber.Value));
+			// SequencerPlaybackRangeMenu
+			{
+				FToolMenuSection& Section = InMenu->AddSection("SequencerPlaybackRangeMenu", FText::Format(LOCTEXT("PlaybackRangeTextFormat", "Playback Range ({0}):"), CurrentTimeText));
+			
+				TRange<FFrameNumber> PlaybackRange = TimeSliderController->TimeSliderArgs.PlaybackRange.Get();
+				TOptional<TRange<FFrameNumber>> SubSequenceRange = TimeSliderController->TimeSliderArgs.SubSequenceRange.Get();
+						
+				Section.AddMenuEntry(
+					"SetPlaybackEnd",
+					LOCTEXT("SetPlaybackEnd", "Set End Time"),
+					FText(),
+					FSlateIcon(),
+					FUIAction(
+						FExecuteAction::CreateLambda([TimeSliderController, FrameNumber]
+						{
+							TimeSliderController->SetPlaybackRangeEnd(FrameNumber);
+						}),
+						FCanExecuteAction::CreateLambda([TimeSliderController, FrameNumber, PlaybackRange]
+						{
+							return !TimeSliderController->TimeSliderArgs.IsPlaybackRangeLocked.Get() && FrameNumber >= UE::MovieScene::DiscreteInclusiveLower(PlaybackRange);
+						})
+					)
+				);
 
-				MenuBuilder.AddMenuEntry(
-					ParentText,
-					FText::Format(LOCTEXT("DisplayTimeSpace", "Display time in the space of {0}"), ParentText),
+				Section.AddMenuEntry(
+					"ConformToSubsequenceRange",
+					LOCTEXT("ConformToSubsequenceRange", "Conform to Range"),
+					LOCTEXT("ConformToSubsequenceRangeTooltip", "Conform the start and end time to the extents of the subsequence range"),
 					FSlateIcon(),
 					FUIAction(
-						FExecuteAction::CreateLambda([this, ParentID] { TimeSliderArgs.OnScrubPositionParentChanged.ExecuteIfBound(ParentID); }),
-						FCanExecuteAction(),
-						FIsActionChecked::CreateLambda([this, ParentID] { return TimeSliderArgs.ScrubPositionParent.Get() == MovieSceneSequenceID::Invalid ? ParentID == TimeSliderArgs.ScrubPositionParentChain.Get().Last() : TimeSliderArgs.ScrubPositionParent.Get() == ParentID; })
-					),
-					NAME_None,
-					EUserInterfaceActionType::RadioButton
+						FExecuteAction::CreateLambda([TimeSliderController, SubSequenceRange]
+						{
+							TimeSliderController->SetPlaybackRangeStart(SubSequenceRange.GetValue().GetLowerBoundValue());
+							TimeSliderController->SetPlaybackRangeEnd(SubSequenceRange.GetValue().GetUpperBoundValue());
+						}),
+						FCanExecuteAction::CreateLambda([TimeSliderController, SubSequenceRange]
+						{
+							return !TimeSliderController->TimeSliderArgs.IsPlaybackRangeLocked.Get() && SubSequenceRange.IsSet();
+						})
+					)
 				);
+
+				Section.AddMenuEntry(
+					"ToggleLocked",
+						LOCTEXT("ToggleLocked", "Locked"),
+						LOCTEXT("ToggleLockedTooltip", "Lock/Unlock the playback range"),
+						FSlateIcon(),
+						FUIAction(
+							FExecuteAction::CreateLambda([TimeSliderController]
+							{
+								TimeSliderController->TimeSliderArgs.OnTogglePlaybackRangeLocked.ExecuteIfBound();
+							}),
+							FCanExecuteAction::CreateLambda([bReadOnly]{ return !bReadOnly; }),
+							FIsActionChecked::CreateLambda([TimeSliderController]
+							{
+								return TimeSliderController->TimeSliderArgs.IsPlaybackRangeLocked.Get();
+							})
+						),
+						EUserInterfaceActionType::ToggleButton
+					);
 			}
-		}
-		MenuBuilder.EndSection(); // Sequencer Parent Chain
-	}
 
-	MenuBuilder.BeginSection("SequencerMarkMenu", FText::Format(LOCTEXT("MarkTextFormat", "Mark ({0}):"), CurrentTimeText));
-	{
-		UMovieScene* MovieScene = Sequencer->GetFocusedMovieSceneSequence()->GetMovieScene();
-		bool bHasMarks = MovieScene->GetMarkedFrames().Num() > 0;
+			// SequencerPlaybackRangeMenu
+			{
+				TRange<FFrameNumber> SelectionRange = TimeSliderController->TimeSliderArgs.SelectionRange.Get();
+				FToolMenuSection& Section = InMenu->AddSection("SequencerSelectionRangeMenu", FText::Format(LOCTEXT("SelectionRangeTextFormat", "Playback Range ({0}):"), CurrentTimeText));
 
-		int32 MarkedIndex = INDEX_NONE;
-		const bool bTestLabelBox = true;
-		HitTestMark(MyGeometry, RangeToScreen, MousePixel, bTestLabelBox, &MarkedIndex);
+				Section.AddMenuEntry(
+					"SetSelectionStart",
+					LOCTEXT("SetSelectionStart", "Set Selection Start"),
+					FText(),
+					FSlateIcon(),
+					FUIAction(
+						FExecuteAction::CreateLambda([TimeSliderController, FrameNumber]
+						{
+							TimeSliderController->SetSelectionRangeStart(FrameNumber);
+						}),
+						FCanExecuteAction::CreateLambda([SelectionRange, FrameNumber]
+						{
+							return SelectionRange.IsEmpty() || FrameNumber < UE::MovieScene::DiscreteExclusiveUpper(SelectionRange);
+						})
+					)
+				);
 
-		FSequencerSelection& SequencerSelection = Sequencer->GetSelection();
-		SequencerSelection.Empty();
+				Section.AddMenuEntry(
+					"SetSelectionEnd",
+					LOCTEXT("SetSelectionEnd", "Set Selection End"),
+					FText(),
+					FSlateIcon(),
+					FUIAction(
+						FExecuteAction::CreateLambda([TimeSliderController, FrameNumber]
+						{
+							TimeSliderController->SetSelectionRangeEnd(FrameNumber);
+						}),
+						FCanExecuteAction::CreateLambda([SelectionRange, FrameNumber]
+						{
+							return SelectionRange.IsEmpty() || FrameNumber >= UE::MovieScene::DiscreteInclusiveLower(SelectionRange);
+						})
+					)
+				);
 
-		if (MarkedIndex != INDEX_NONE)
-		{
-			SequencerSelection.MarkedFrames.Empty();
-			SequencerSelection.MarkedFrames.Select(MarkedIndex);
+				Section.AddMenuEntry(
+					"ClearSelectionRange",
+					LOCTEXT("ClearSelectionRange", "Clear Selection Range"),
+					FText(),
+					FSlateIcon(),
+					FUIAction(
+						FExecuteAction::CreateLambda([TimeSliderController]
+						{
+							TimeSliderController->TimeSliderArgs.OnSelectionRangeChanged.ExecuteIfBound(TRange<FFrameNumber>::Empty());
+						}),
+						FCanExecuteAction::CreateLambda([SelectionRange]
+						{
+							return !SelectionRange.IsEmpty();
+						})
+					)
+				);
+			}
 
-			class SMarkedFramePropertyWidget : public SCompoundWidget, public FNotifyHook
+			// SequencerParentChainMenu
 			{
-			public:
-				UMovieScene* MovieSceneToModify;
-				TSharedPtr<IStructureDetailsView> DetailsView;
-				TWeakPtr<FSequencer> WeakSequencer;
+				UMovieSceneCompiledDataManager* CompiledDataManager = Sequencer->GetEvaluationTemplate().GetCompiledDataManager();
+				const FMovieSceneSequenceHierarchy* Hierarchy = CompiledDataManager->FindHierarchy(Sequencer->GetEvaluationTemplate().GetCompiledDataID());
 
-				SLATE_BEGIN_ARGS(SMarkedFramePropertyWidget){}
-				SLATE_END_ARGS()
-
-				void Construct(const FArguments& InArgs, UMovieScene* InMovieScene, int32 InMarkedFrameIndex, TWeakPtr<FSequencer> InWeakSequencer)
+				if (Hierarchy)
 				{
-					MovieSceneToModify = InMovieScene;
-					WeakSequencer = InWeakSequencer;
-
-					FPropertyEditorModule& PropertyEditorModule = FModuleManager::GetModuleChecked<FPropertyEditorModule>("PropertyEditor");
-					FDetailsViewArgs DetailsViewArgs;
-					DetailsViewArgs.bAllowSearch = false;
-					DetailsViewArgs.bShowScrollBar = false;
-					DetailsViewArgs.NameAreaSettings = FDetailsViewArgs::HideNameArea;
-					DetailsViewArgs.NotifyHook = this;
-
-					FStructureDetailsViewArgs StructureDetailsViewArgs;
-					StructureDetailsViewArgs.bShowObjects = true;
-					StructureDetailsViewArgs.bShowAssets = true;
-					StructureDetailsViewArgs.bShowClasses = true;
-					StructureDetailsViewArgs.bShowInterfaces = true;
-
-					TSharedPtr<FStructOnScope> StructOnScope = MakeShared<FStructOnScope>(FMovieSceneMarkedFrame::StaticStruct(), (uint8 *)&InMovieScene->GetMarkedFrames()[InMarkedFrameIndex]);
-
-					DetailsView = PropertyEditorModule.CreateStructureDetailView(DetailsViewArgs, StructureDetailsViewArgs, nullptr);
-					DetailsView->GetDetailsView()->RegisterInstancedCustomPropertyTypeLayout("FrameNumber", 
-						FOnGetPropertyTypeCustomizationInstance::CreateSP(WeakSequencer.Pin().ToSharedRef(), &FSequencer::MakeFrameNumberDetailsCustomization));
-					DetailsView->SetStructureData(StructOnScope);
-
-					ChildSlot
-					[
-						DetailsView->GetWidget().ToSharedRef()
-					];
-				}
+					FToolMenuSection& Section = InMenu->AddSection("SequencerParentChainMenu");
 
-				virtual void NotifyPreChange( FProperty* PropertyAboutToChange ) override
-				{
-					MovieSceneToModify->Modify();
+					TArray<FMovieSceneSequenceID> ParentChain = TimeSliderController->TimeSliderArgs.ScrubPositionParentChain.Get();
+					for (FMovieSceneSequenceID ParentID : ParentChain)
+					{
+						FText ParentText = Sequencer->GetRootMovieSceneSequence()->GetDisplayName();
+
+						for (const TTuple<FMovieSceneSequenceID, FMovieSceneSubSequenceData>& Pair : Hierarchy->AllSubSequenceData())
+						{
+							if (Pair.Key == ParentID && Pair.Value.GetSequence())
+							{
+								ParentText = Pair.Value.GetSequence()->GetDisplayName();
+								break;
+							}
+						}
+
+						Section.AddMenuEntry(
+							"DisplayTimeSpace",
+							ParentText,
+							FText::Format(LOCTEXT("DisplayTimeSpace", "Display time in the space of {0}"), ParentText),
+							FSlateIcon(),
+							FUIAction(
+								FExecuteAction::CreateLambda([TimeSliderController, ParentID]
+								{
+									TimeSliderController->TimeSliderArgs.OnScrubPositionParentChanged.ExecuteIfBound(ParentID);
+								}),
+								FCanExecuteAction(),
+								FIsActionChecked::CreateLambda([TimeSliderController, ParentID]
+								{
+									return TimeSliderController->TimeSliderArgs.ScrubPositionParent.Get() == MovieSceneSequenceID::Invalid
+									? ParentID == TimeSliderController->TimeSliderArgs.ScrubPositionParentChain.Get().Last()
+									: TimeSliderController->TimeSliderArgs.ScrubPositionParent.Get() == ParentID;
+								})
+							),
+							EUserInterfaceActionType::RadioButton
+						);
+					}
 				}
+			}
+
+			// SequencerMarkMenu
+			{
+				FToolMenuSection& Section = InMenu->AddSection("SequencerMarkMenu", FText::Format(LOCTEXT("MarkTextFormat", "Mark ({0}):"), CurrentTimeText));
+
+				UMovieScene* MovieScene = Sequencer->GetFocusedMovieSceneSequence()->GetMovieScene();
+				bool bHasMarks = MovieScene->GetMarkedFrames().Num() > 0;
 
-				virtual void NotifyPreChange( class FEditPropertyChain* PropertyAboutToChange ) override
+				int32 MarkedIndex = INDEX_NONE;
+				const bool bTestLabelBox = true;
+				TimeSliderController->HitTestMark(MyGeometry, RangeToScreen, MousePixel, bTestLabelBox, &MarkedIndex);
+
+				FSequencerSelection& SequencerSelection = Sequencer->GetSelection();
+				SequencerSelection.Empty();
+
+				if (MarkedIndex != INDEX_NONE)
 				{
-					MovieSceneToModify->Modify();
-				}
-			};
+					SequencerSelection.MarkedFrames.Empty();
+					SequencerSelection.MarkedFrames.Select(MarkedIndex);
 
-			const bool bLockedMarkedFrames = TimeSliderArgs.AreMarkedFramesLocked.Get();
+					const bool bLockedMarkedFrames = TimeSliderController->TimeSliderArgs.AreMarkedFramesLocked.Get();
 
-			TSharedRef<SMarkedFramePropertyWidget> Widget = SNew(SMarkedFramePropertyWidget, MovieScene, MarkedIndex, WeakSequencer);
-			Widget->SetEnabled(!bLockedMarkedFrames);
-			MenuBuilder.AddWidget(Widget, FText::GetEmpty(), false);
-		}
+					TSharedRef<SMarkedFramePropertyWidget> Widget = SNew(SMarkedFramePropertyWidget, MovieScene, MarkedIndex, Sequencer.ToWeakPtr());
+					Widget->SetEnabled(!bLockedMarkedFrames);
+					Section.AddEntry(FToolMenuEntry::InitWidget(NAME_None, Widget, FText::GetEmpty(), false));
+				}
 
-		MenuBuilder.AddMenuEntry( 
-			LOCTEXT("AddMark", "Add Mark"),
-			FText(),
-			FSlateIcon(),
-			FUIAction(
-				FExecuteAction::CreateLambda( [this, FrameNumber]{ AddMarkAtFrame(FrameNumber); }),
-				FCanExecuteAction::CreateLambda([this, MarkedIndex]{ return !TimeSliderArgs.AreMarkedFramesLocked.Get() && MarkedIndex == INDEX_NONE; }))
-		);
+				Section.AddMenuEntry( 
+					"AddMark",
+					LOCTEXT("AddMark", "Add Mark"),
+					FText(),
+					FSlateIcon(),
+					FUIAction(
+						FExecuteAction::CreateLambda( [TimeSliderController, FrameNumber]
+						{
+							TimeSliderController->AddMarkAtFrame(FrameNumber);
+						}),
+						FCanExecuteAction::CreateLambda([TimeSliderController, MarkedIndex]
+						{
+							return !TimeSliderController->TimeSliderArgs.AreMarkedFramesLocked.Get() && MarkedIndex == INDEX_NONE;
+						}))
+				);
 
-		MenuBuilder.AddMenuEntry( 
-			LOCTEXT("DeleteMark", "Delete Mark"),
-			FText(),
-			FSlateIcon(),
-			FUIAction(
-				FExecuteAction::CreateLambda([this, MarkedIndex]{ DeleteMarkAtIndex(MarkedIndex); }),
-				FCanExecuteAction::CreateLambda([this, MarkedIndex]{ return !TimeSliderArgs.AreMarkedFramesLocked.Get() && MarkedIndex != INDEX_NONE; }))
-		);
+				Section.AddMenuEntry(
+					"DeleteMark",
+					LOCTEXT("DeleteMark", "Delete Mark"),
+					FText(),
+					FSlateIcon(),
+					FUIAction(
+						FExecuteAction::CreateLambda([TimeSliderController, MarkedIndex]
+						{
+							TimeSliderController->DeleteMarkAtIndex(MarkedIndex);
+						}),
+						FCanExecuteAction::CreateLambda([TimeSliderController, MarkedIndex]
+						{
+							return !TimeSliderController->TimeSliderArgs.AreMarkedFramesLocked.Get() && MarkedIndex != INDEX_NONE;
+						}))
+				);
 
-		MenuBuilder.AddMenuEntry( 
-			LOCTEXT("DeleteAllMarks", "Delete All Marks"),
-			FText(),
-			FSlateIcon(),
-			FUIAction(
-				FExecuteAction::CreateLambda([this]{ DeleteAllMarks(); }),
-				FCanExecuteAction::CreateLambda([this, bHasMarks]{ return !TimeSliderArgs.AreMarkedFramesLocked.Get() && bHasMarks; }))
-			);
+				Section.AddMenuEntry(
+					"DeleteAllMarks",
+					LOCTEXT("DeleteAllMarks", "Delete All Marks"),
+					FText(),
+					FSlateIcon(),
+					FUIAction(
+						FExecuteAction::CreateLambda([TimeSliderController]
+						{
+							TimeSliderController->DeleteAllMarks();
+						}),
+						FCanExecuteAction::CreateLambda([TimeSliderController, bHasMarks]
+						{
+							return !TimeSliderController->TimeSliderArgs.AreMarkedFramesLocked.Get() && bHasMarks;
+						}))
+					);
 
-		MenuBuilder.AddMenuEntry(
-			LOCTEXT("ToggleLockedMarks", "Locked"),
-			LOCTEXT("ToggleLockedMarksTooltip", "Lock/Unlock all marked frames"),
-			FSlateIcon(),
-			FUIAction(
-				FExecuteAction::CreateLambda([this]{ TimeSliderArgs.OnToggleMarkedFramesLocked.ExecuteIfBound(); }),
-				FCanExecuteAction::CreateLambda([=]{ return !bReadOnly; }),
-				FIsActionChecked::CreateLambda([this]{ return TimeSliderArgs.AreMarkedFramesLocked.Get(); })),
-			NAME_None,
-			EUserInterfaceActionType::ToggleButton
-			);
+				Section.AddMenuEntry(
+					"ToggleLockedMarks",
+					LOCTEXT("ToggleLockedMarks", "Locked"),
+					LOCTEXT("ToggleLockedMarksTooltip", "Lock/Unlock all marked frames"),
+					FSlateIcon(),
+					FUIAction(
+						FExecuteAction::CreateLambda([TimeSliderController]
+						{
+							TimeSliderController->TimeSliderArgs.OnToggleMarkedFramesLocked.ExecuteIfBound();
+						}),
+						FCanExecuteAction::CreateLambda([bReadOnly]{ return !bReadOnly; }),
+						FIsActionChecked::CreateLambda([TimeSliderController]
+						{
+							return TimeSliderController->TimeSliderArgs.AreMarkedFramesLocked.Get();
+						})),
+					EUserInterfaceActionType::ToggleButton
+					);
+			}
+		}));
 	}
-	MenuBuilder.EndSection(); // SequencerMarkMenu
-
-	return MenuBuilder.MakeWidget();
+	
+	USequencerTimeSliderControllerMenuContext* TimeSliderContext = NewObject<USequencerTimeSliderControllerMenuContext>();
+	TimeSliderContext->WeakTimeSliderController = AsWeak();
+	TimeSliderContext->Geometry = MyGeometry;
+	TimeSliderContext->PointerEvent = MouseEvent;
+	TimeSliderContext->FrameNumber = ComputeFrameTimeFromMouse(MyGeometry, MouseEvent.GetScreenSpacePosition(), FScrubRangeToScreen(GetViewRange(), MyGeometry.Size)).FrameNumber;
+
+	USequencerToolMenuContext* SequencerContext = NewObject<USequencerToolMenuContext>();
+	SequencerContext->WeakSequencer = WeakSequencer;
+
+	FToolMenuContext MenuContext;
+	MenuContext.AddObject(TimeSliderContext);
+	MenuContext.AddObject(SequencerContext);
+	return UToolMenus::Get()->GenerateWidget(MenuName, MenuContext);
 }
 
 void FSequencerTimeSliderController::ClampViewRange(double& NewRangeMin, double& NewRangeMax)
diff --git a/Engine/Source/Editor/Sequencer/Public/SequencerToolMenuContext.h b/Engine/Source/Editor/Sequencer/Public/SequencerToolMenuContext.h
--- a/Engine/Source/Editor/Sequencer/Public/SequencerToolMenuContext.h
+++ b/Engine/Source/Editor/Sequencer/Public/SequencerToolMenuContext.h
@@ -3,8 +3,11 @@
 #pragma once
 
 #include "CoreTypes.h"
+#include "Layout/Geometry.h"
+#include "Input/Events.h"
 #include "SequencerToolMenuContext.generated.h"
 
+class FSequencerTimeSliderController;
 class ISequencer;
 class SSequencerPlayRateCombo;
 
@@ -23,3 +26,20 @@ class USequencerClockSourceMenuContext : public USequencerToolMenuContext
 public:
 	TWeakPtr<SSequencerPlayRateCombo> WeakSequencerCombo;
 };
+
+UCLASS(MinimalAPI)
+class USequencerTimeSliderControllerMenuContext : public UObject
+{
+	GENERATED_BODY()
+	
+public:
+
+	// The time slider controller associated with this context
+	TWeakPtr<FSequencerTimeSliderController> WeakTimeSliderController;
+	
+	FGeometry Geometry;
+
+	FPointerEvent PointerEvent;
+
+	FFrameNumber FrameNumber;
+};
\ No newline at end of file
