Sequencer: Disable vertical and horizontal zoom out limiting in curve editor

diff --git a/Engine/Source/Editor/Sequencer/Private/SequencerSettings.cpp b/Engine/Source/Editor/Sequencer/Private/SequencerSettings.cpp
--- a/Engine/Source/Editor/Sequencer/Private/SequencerSettings.cpp
+++ b/Engine/Source/Editor/Sequencer/Private/SequencerSettings.cpp
@@ -91,8 +91,8 @@ USequencerSettings::USequencerSettings( const FObjectInitializer& ObjectInitiali
 	// Don't add any factors for zoom-in. Zoom-in is already very fast with a zoom factor of 1.
 	CurveEditorZoomScaling.HorizontalZoomScale.EditorCurveData.AddKey(100, 6); // ... e.g. showing 100 seconds worth of frames, zoom out factor is 6.
 	CurveEditorZoomScaling.VerticalZoomScale.EditorCurveData.AddKey(100, 6);
-	CurveEditorZoomScaling.bLimitHorizontalZoomOut = true;
-	CurveEditorZoomScaling.bLimitVerticalZoomOut = true;
+	CurveEditorZoomScaling.bLimitHorizontalZoomOut = false;
+	CurveEditorZoomScaling.bLimitVerticalZoomOut = false;
 	CurveEditorZoomScaling.MaxHorizontalZoomOut = 900; // Cannot zoom out more than 900 seconds = 15 minutes.
 	CurveEditorZoomScaling.MaxVerticalZoomOut = 10000;
 }
