Sequencer: In FTrackModel::ForceUpdate stop Auto-expansions when there is only one row layout by not have NumRows ==1 be a special case, the default n > 1 section correctly keeps track of existing states.

diff --git a/Engine/Source/Editor/Sequencer/Private/MVVM/ViewModels/TrackModel.cpp b/Engine/Source/Editor/Sequencer/Private/MVVM/ViewModels/TrackModel.cpp
--- a/Engine/Source/Editor/Sequencer/Private/MVVM/ViewModels/TrackModel.cpp
+++ b/Engine/Source/Editor/Sequencer/Private/MVVM/ViewModels/TrackModel.cpp
@@ -194,63 +194,6 @@ void FTrackModel::ForceUpdate()
 		SectionChildren.Empty();
 		TopLevelChannelChildren.Empty();
 	}
-	else if (NumRows == 1)
-	{
-		// Keep sections alive by retaining the previous list temporarily
-		TSharedPtr<FViewModel> SectionsTail;
-
-		FScopedViewModelListHead RecycledModels(AsShared(), EViewModelListType::Recycled);
-		GetChildrenForList(&SectionList).MoveChildrenTo<IRecyclableExtension>(RecycledModels.GetChildren(), IRecyclableExtension::CallOnRecycle);
-
-		bool bNeedsLayout = NumRows != PreviousLayoutNumRows;
-
-		// Add all sections directly to this track row
-		for (UMovieSceneSection* Section : Track->GetAllSections())
-		{
-			TSharedPtr<FSectionModel> SectionModel = SectionModelStorage->FindModelForSection(Section);
-			if (!SectionModel && TrackEditor)
-			{
-				TSharedRef<ISequencerSection> SectionInterface = TrackEditor->MakeSectionInterface(*Section, *Track, ObjectBinding);
-				SectionModel = SectionModelStorage->CreateModelForSection(Section, SectionInterface);
-				bNeedsLayout = true;
-			}
-
-			if (ensure(SectionModel))
-			{
-				bNeedsLayout |= SectionModel->NeedsLayout();
-
-				// Move the child back into the real section list
-				SectionChildren.InsertChild(SectionModel, SectionsTail);
-				SectionsTail = SectionModel;
-			}
-		}
-
-		// If we are discarding any sections (because they still remain in the recycled list) we must run the layout
-		bNeedsLayout |= !RecycledModels.GetChildren().IsEmpty();
-
-		if (bNeedsLayout)
-		{
-
-			OutlinerChildren.MoveChildrenTo<IRecyclableExtension>(RecycledModels.GetChildren(), IRecyclableExtension::CallOnRecycle);
-			TopLevelChannelChildren.MoveChildrenTo<IRecyclableExtension>(RecycledModels.GetChildren(), IRecyclableExtension::CallOnRecycle);
-
-			// Rebuild the outliner layout for this track. This will clear our children and rebuild them if needed
-			// (with potentially recycled children), so if we went from, say, 2 rows to 1 row, it should correctly
-			// discard any children we don't need anymore.
-			FTrackModelLayoutBuilder LayoutBuilder(AsShared());
-
-			for (TSharedPtr<FSectionModel> Section : TViewModelListIterator<FSectionModel>(&SectionList))
-			{
-				LayoutBuilder.RefreshLayout(Section);
-			}
-
-			if (OutlinerChildren.IsEmpty())
-			{
-				// Reset expansion state if this track can no longer be expanded
-				SetExpansion(false);
-			}
-		}
-	}
 	else
 	{
 		// Always expand parent tracks
