Sequencer Mode: only draw the marquee selection in the focused viewport

#rnx

diff --git a/Engine/Source/Editor/Sequencer/Private/SequencerEdMode.cpp b/Engine/Source/Editor/Sequencer/Private/SequencerEdMode.cpp
--- a/Engine/Source/Editor/Sequencer/Private/SequencerEdMode.cpp
+++ b/Engine/Source/Editor/Sequencer/Private/SequencerEdMode.cpp
@@ -552,7 +552,12 @@ void FSequencerEdMode::Render(const FSceneView* View, FViewport* Viewport, FPrim
 {
 	FEdMode::Render(View, Viewport, PDI);
 
-	DragToolHandler.Render3DDragTool(View, PDI);
+	const FViewportClient* ViewportClient = Viewport ? Viewport->GetClient() : nullptr;
+	const FViewportClient* FocusedViewportClient = Owner ? Owner->GetFocusedViewportClient() : nullptr;
+	if (ViewportClient && ViewportClient == FocusedViewportClient)
+	{
+		DragToolHandler.Render3DDragTool(View, PDI);
+	}
 
 #if WITH_EDITORONLY_DATA
 	if (PDI)
@@ -578,19 +583,26 @@ void FSequencerEdMode::DrawHUD(FEditorViewportClient* ViewportClient, FViewport*
 {
 	FEdMode::DrawHUD(ViewportClient,Viewport,View,Canvas);
 
-	DragToolHandler.RenderDragTool(View, Canvas);
-
-	if( ViewportClient->AllowsCinematicControl() )
+	if (ViewportClient)
 	{
-		// Get the size of the viewport
-		const int32 SizeX = Viewport->GetSizeXY().X;
-		const int32 SizeY = Viewport->GetSizeXY().Y;
-
-		// Draw subtitles (toggle is handled internally)
-		FVector2D MinPos(0.f, 0.f);
-		FVector2D MaxPos(1.f, .9f);
-		FIntRect SubtitleRegion(FMath::TruncToInt(SizeX * MinPos.X), FMath::TruncToInt(SizeY * MinPos.Y), FMath::TruncToInt(SizeX * MaxPos.X), FMath::TruncToInt(SizeY * MaxPos.Y));
-		FSubtitleManager::GetSubtitleManager()->DisplaySubtitles( Canvas, SubtitleRegion, ViewportClient->GetWorld()->GetAudioTimeSeconds() );
+		const FViewportClient* FocusedViewportClient = Owner ? Owner->GetFocusedViewportClient() : nullptr;
+		if (ViewportClient == FocusedViewportClient)
+		{
+			DragToolHandler.RenderDragTool(View, Canvas);
+		}
+
+		if ( ViewportClient->AllowsCinematicControl() )
+		{
+			// Get the size of the viewport
+			const int32 SizeX = Viewport->GetSizeXY().X;
+			const int32 SizeY = Viewport->GetSizeXY().Y;
+
+			// Draw subtitles (toggle is handled internally)
+			FVector2D MinPos(0.f, 0.f);
+			FVector2D MaxPos(1.f, .9f);
+			FIntRect SubtitleRegion(FMath::TruncToInt(SizeX * MinPos.X), FMath::TruncToInt(SizeY * MinPos.Y), FMath::TruncToInt(SizeX * MaxPos.X), FMath::TruncToInt(SizeY * MaxPos.Y));
+			FSubtitleManager::GetSubtitleManager()->DisplaySubtitles( Canvas, SubtitleRegion, ViewportClient->GetWorld()->GetAudioTimeSeconds() );
+		}
 	}
 
 	if (SelectabilityTool->IsSelectionLimited())
