[Sequencer] Fix mute/solo for animation track rows

diff --git a/Engine/Source/Editor/Sequencer/Private/SSequencer.cpp b/Engine/Source/Editor/Sequencer/Private/SSequencer.cpp
--- a/Engine/Source/Editor/Sequencer/Private/SSequencer.cpp
+++ b/Engine/Source/Editor/Sequencer/Private/SSequencer.cpp
@@ -2683,15 +2683,14 @@ void SSequencer::UpdateLayoutTree()
 			NewNodePathsToIsolate.Empty();
 		}
 
+		// Local Mute/Solo state (transient, non dirtying, mute/solo evaluation)
 		if (Sequencer->GetFocusedMovieSceneSequence())
 		{
-			bool bAnyChanged = false;
-
-			TSharedPtr<FSharedViewModelData> SharedData = Sequencer->GetViewModel()->GetRootModel()->GetSharedData();
-			FOutlinerCacheExtension*           OutlinerCache = SharedData->CastThis<FOutlinerCacheExtension>();
-			FDeactiveStateCacheExtension* DeactiveState = SharedData->CastThis<FDeactiveStateCacheExtension>();
-			FMuteStateCacheExtension*          MuteState     = SharedData->CastThis<FMuteStateCacheExtension>();
-			FSoloStateCacheExtension*          SoloState     = SharedData->CastThis<FSoloStateCacheExtension>();
+			TSharedPtr<FSharedViewModelData> SharedData    = Sequencer->GetViewModel()->GetRootModel()->GetSharedData();
+			FOutlinerCacheExtension*         OutlinerCache = SharedData->CastThis<FOutlinerCacheExtension>();
+			FDeactiveStateCacheExtension*    DeactiveState = SharedData->CastThis<FDeactiveStateCacheExtension>();
+			FMuteStateCacheExtension*        MuteState     = SharedData->CastThis<FMuteStateCacheExtension>();
+			FSoloStateCacheExtension*        SoloState     = SharedData->CastThis<FSoloStateCacheExtension>();
 
 			check(OutlinerCache && DeactiveState && MuteState && SoloState);
 
@@ -2701,6 +2700,8 @@ void SSequencer::UpdateLayoutTree()
 
 			const bool bAnySoloNodes = EnumHasAnyFlags(SoloState->GetRootFlags(), ECachedSoloState::Soloed | ECachedSoloState::PartiallySoloedChildren);
 
+			bool bAnyChanged = false;
+
 			for (TViewModelPtr<ITrackExtension> TrackNode
 				: Sequencer->GetNodeTree()->GetRootNode()->GetDescendantsOfType<ITrackExtension>())
 			{
@@ -2716,25 +2717,35 @@ void SSequencer::UpdateLayoutTree()
 				const bool bIsMuted = EnumHasAnyFlags(MuteFlags, ECachedMuteState::Muted | ECachedMuteState::ImplicitlyMutedByParent);
 				const bool bIsSoloed = EnumHasAnyFlags(SoloFlags, ECachedSoloState::Soloed | ECachedSoloState::ImplicitlySoloedByParent);
 
-				const bool bLocalEvalDisabled = bIsMuted || (bAnySoloNodes && !bIsSoloed);
+				const bool bLocalEvalDisabled = bIsSoloed ? false : (bIsMuted || (bAnySoloNodes && !bIsSoloed));
 
 				if (const TViewModelPtr<FTrackRowModel> TrackRowModel = TrackNode.ImplicitCast())
 				{
-					// Mute/Solo state (local, non dirtying, mute/solo evaluation)
-					if (bLocalEvalDisabled != Track->IsLocalRowEvalDisabled(TrackNode->GetRowIndex()))
+					const int32 TrackRowIndex = TrackNode->GetRowIndex();
+					const bool bCurrentLocalRowEvalDisabled = Track->IsLocalRowEvalDisabled(TrackRowIndex);
+
+					if (bLocalEvalDisabled != bCurrentLocalRowEvalDisabled
+						|| (bIsSoloed && (bIsSoloed != bLocalEvalDisabled)))
 					{
 						Track->MarkAsChanged();
-						Track->SetLocalRowEvalDisabled(bLocalEvalDisabled, TrackNode->GetRowIndex());
+						if (!bLocalEvalDisabled)
+						{
+							Track->SetLocalEvalDisabled(false);
+						}
+						Track->SetLocalRowEvalDisabled(bLocalEvalDisabled, TrackRowIndex);
+
 						bAnyChanged = true;
 					}
 				}
 				else
 				{
-					// Mute/Solo state (local, non dirtying, mute/solo evaluation)
-					if (bLocalEvalDisabled != Track->IsLocalEvalDisabled())
+					const bool bCurrentLocalEvalDisabled = Track->IsLocalEvalDisabled();
+
+					if (bLocalEvalDisabled != bCurrentLocalEvalDisabled)
 					{
 						Track->MarkAsChanged();
 						Track->SetLocalEvalDisabled(bLocalEvalDisabled);
+
 						bAnyChanged = true;
 					}
 				}
