Sequencer: Add back missing Set Start Time

#rnx

diff --git a/Engine/Source/Editor/Sequencer/Private/SequencerTimeSliderController.cpp b/Engine/Source/Editor/Sequencer/Private/SequencerTimeSliderController.cpp
--- a/Engine/Source/Editor/Sequencer/Private/SequencerTimeSliderController.cpp
+++ b/Engine/Source/Editor/Sequencer/Private/SequencerTimeSliderController.cpp
@@ -1941,6 +1941,23 @@ TSharedRef<SWidget> FSequencerTimeSliderController::OpenSetPlaybackRangeMenu(con
 			
 				TRange<FFrameNumber> PlaybackRange = TimeSliderController->TimeSliderArgs.PlaybackRange.Get();
 				TOptional<TRange<FFrameNumber>> SubSequenceRange = TimeSliderController->TimeSliderArgs.SubSequenceRange.Get();
+
+				Section.AddMenuEntry(
+					"SetPlaybackStart",
+					LOCTEXT("SetPlaybackStart", "Set Start Time"),
+					FText(),
+					FSlateIcon(),
+					FUIAction(
+						FExecuteAction::CreateLambda([TimeSliderController, FrameNumber]
+						{
+							TimeSliderController->SetPlaybackRangeStart(FrameNumber);
+						}),
+						FCanExecuteAction::CreateLambda([TimeSliderController, FrameNumber, PlaybackRange]
+						{
+							return !TimeSliderController->TimeSliderArgs.IsPlaybackRangeLocked.Get() && FrameNumber < UE::MovieScene::DiscreteExclusiveUpper(PlaybackRange);
+						})
+					)
+				);
 						
 				Section.AddMenuEntry(
 					"SetPlaybackEnd",
