Sequencer: Don't empty selection when right clicking in the timeline. Instead, just deselect

diff --git a/Engine/Source/Editor/Sequencer/Private/SequencerTimeSliderController.cpp b/Engine/Source/Editor/Sequencer/Private/SequencerTimeSliderController.cpp
--- a/Engine/Source/Editor/Sequencer/Private/SequencerTimeSliderController.cpp
+++ b/Engine/Source/Editor/Sequencer/Private/SequencerTimeSliderController.cpp
@@ -2130,12 +2130,15 @@ TSharedRef<SWidget> FSequencerTimeSliderController::OpenSetPlaybackRangeMenu(con
 				TimeSliderController->HitTestMark(MyGeometry, RangeToScreen, MousePixel, bTestLabelBox, &MarkedIndex);
 
 				FSequencerSelection& SequencerSelection = Sequencer->GetSelection();
-				SequencerSelection.Empty();
 
 				if (MarkedIndex != INDEX_NONE)
 				{
-					SequencerSelection.MarkedFrames.Empty();
-					SequencerSelection.MarkedFrames.Select(MarkedIndex);
+					const bool bIsSelected = SequencerSelection.MarkedFrames.IsSelected(MarkedIndex);
+					if (!bIsSelected)
+					{
+						SequencerSelection.MarkedFrames.Empty();
+						SequencerSelection.MarkedFrames.Select(MarkedIndex);
+					}
 
 					const bool bLockedMarkedFrames = TimeSliderController->TimeSliderArgs.AreMarkedFramesLocked.Get();
 
