Sequencer: Restore section model expansion state to be closed if it was before, since RestoreLayout resets it to expanded.

diff --git a/Engine/Source/Editor/Sequencer/Private/MVVM/ViewModels/TrackModel.cpp b/Engine/Source/Editor/Sequencer/Private/MVVM/ViewModels/TrackModel.cpp
--- a/Engine/Source/Editor/Sequencer/Private/MVVM/ViewModels/TrackModel.cpp
+++ b/Engine/Source/Editor/Sequencer/Private/MVVM/ViewModels/TrackModel.cpp
@@ -204,6 +204,8 @@ void FTrackModel::ForceUpdate()
 
 		bool bNeedsLayout = NumRows != PreviousLayoutNumRows;
 
+		//Sections that weren't expanded
+		TSet<TSharedPtr<FSectionModel>> CollapsedSections;
 		// Add all sections directly to this track row
 		for (UMovieSceneSection* Section : Track->GetAllSections())
 		{
@@ -214,7 +216,14 @@ void FTrackModel::ForceUpdate()
 				SectionModel = SectionModelStorage->CreateModelForSection(Section, SectionInterface);
 				bNeedsLayout = true;
 			}
-
+			else
+			{
+				TViewModelPtr<IOutlinerExtension> Outliner = SectionModel->FindAncestorOfType<IOutlinerExtension>();
+				if (Outliner && (Outliner->IsExpanded() == false))
+				{
+					CollapsedSections.Add(SectionModel);
+				}
+			}
 			if (ensure(SectionModel))
 			{
 				bNeedsLayout |= SectionModel->NeedsLayout();
@@ -242,6 +251,14 @@ void FTrackModel::ForceUpdate()
 			for (TSharedPtr<FSectionModel> Section : TViewModelListIterator<FSectionModel>(&SectionList))
 			{
 				LayoutBuilder.RefreshLayout(Section);
+				if (CollapsedSections.Contains(Section))
+				{
+					TViewModelPtr<IOutlinerExtension> Outliner = Section->FindAncestorOfType<IOutlinerExtension>();
+					if (Outliner)
+					{
+						Outliner->SetExpansion(false);
+					}
+				}
 			}
 
 			if (OutlinerChildren.IsEmpty())
