Updated headers using UnrealCodeFixup to make sure dllstorage is on methods/staticvars instead of types. Part 5/n

diff --git a/Engine/Source/Editor/Sequencer/Public/Filters/ISequencerTrackFilters.h b/Engine/Source/Editor/Sequencer/Public/Filters/ISequencerTrackFilters.h
--- a/Engine/Source/Editor/Sequencer/Public/Filters/ISequencerTrackFilters.h
+++ b/Engine/Source/Editor/Sequencer/Public/Filters/ISequencerTrackFilters.h
@@ -1,4 +1,4 @@
-﻿// Copyright Epic Games, Inc. All Rights Reserved.
+// Copyright Epic Games, Inc. All Rights Reserved.
 
 #pragma once
 
@@ -13,23 +13,23 @@ class ISequencerTrackFilters : public ISequencerFilterBar
 public:
 	virtual ~ISequencerTrackFilters() override = default;
 
-	SEQUENCER_API virtual bool IsFilterEnabled(TSharedRef<FSequencerTrackFilter> InFilter) const = 0;
-	SEQUENCER_API virtual bool SetFilterEnabled(const TSharedRef<FSequencerTrackFilter> InFilter, const bool bInEnabled, const bool bInRequestFilterUpdate = true) = 0;
+	virtual bool IsFilterEnabled(TSharedRef<FSequencerTrackFilter> InFilter) const = 0;
+	virtual bool SetFilterEnabled(const TSharedRef<FSequencerTrackFilter> InFilter, const bool bInEnabled, const bool bInRequestFilterUpdate = true) = 0;
 
-	SEQUENCER_API virtual bool IsFilterActive(const TSharedRef<FSequencerTrackFilter> InFilter) const = 0;
-	SEQUENCER_API virtual bool SetFilterActive(const TSharedRef<FSequencerTrackFilter>& InFilter, const bool bInActive, const bool bInRequestFilterUpdate = true) = 0;
+	virtual bool IsFilterActive(const TSharedRef<FSequencerTrackFilter> InFilter) const = 0;
+	virtual bool SetFilterActive(const TSharedRef<FSequencerTrackFilter>& InFilter, const bool bInActive, const bool bInRequestFilterUpdate = true) = 0;
 
-	SEQUENCER_API virtual bool AddCustomTextFilter(const TSharedRef<FSequencerTrackFilter_CustomText>& InFilter, const bool bInAddToConfig) = 0;
-	SEQUENCER_API virtual bool RemoveCustomTextFilter(const TSharedRef<FSequencerTrackFilter_CustomText>& InFilter, const bool bInAddToConfig) = 0;
+	virtual bool AddCustomTextFilter(const TSharedRef<FSequencerTrackFilter_CustomText>& InFilter, const bool bInAddToConfig) = 0;
+	virtual bool RemoveCustomTextFilter(const TSharedRef<FSequencerTrackFilter_CustomText>& InFilter, const bool bInAddToConfig) = 0;
 
-	SEQUENCER_API virtual void HideSelectedTracks() = 0;
-	SEQUENCER_API virtual void IsolateSelectedTracks() = 0;
+	virtual void HideSelectedTracks() = 0;
+	virtual void IsolateSelectedTracks() = 0;
 
-	SEQUENCER_API virtual void ShowOnlyLocationCategoryGroups() = 0;
-	SEQUENCER_API virtual void ShowOnlyRotationCategoryGroups() = 0;
-	SEQUENCER_API virtual void ShowOnlyScaleCategoryGroups() = 0;
+	virtual void ShowOnlyLocationCategoryGroups() = 0;
+	virtual void ShowOnlyRotationCategoryGroups() = 0;
+	virtual void ShowOnlyScaleCategoryGroups() = 0;
 
-	SEQUENCER_API virtual bool HasSelectedTracks() const = 0;
+	virtual bool HasSelectedTracks() const = 0;
 
-	SEQUENCER_API virtual FSequencerFilterData& GetFilterData() = 0;
+	virtual FSequencerFilterData& GetFilterData() = 0;
 };
diff --git a/Engine/Source/Editor/Sequencer/Public/Filters/SequencerFilterData.h b/Engine/Source/Editor/Sequencer/Public/Filters/SequencerFilterData.h
--- a/Engine/Source/Editor/Sequencer/Public/Filters/SequencerFilterData.h
+++ b/Engine/Source/Editor/Sequencer/Public/Filters/SequencerFilterData.h
@@ -1,10 +1,12 @@
-﻿// Copyright Epic Games, Inc. All Rights Reserved.
+// Copyright Epic Games, Inc. All Rights Reserved.
 
 #pragma once
 
 #include "MovieSceneTrack.h"
 #include "MVVM/ViewModelPtr.h"
 
+#define UE_API SEQUENCER_API
+
 class ISequencer;
 
 namespace UE::Sequencer
@@ -17,40 +19,40 @@ namespace UE::Sequencer
 using FSequencerTrackFilterType = UE::Sequencer::FViewModelPtr;
 
 /** Represents a cache between nodes for a filter operation. */
-struct SEQUENCER_API FSequencerFilterData
+struct FSequencerFilterData
 {
-	FSequencerFilterData(const FString& InRawFilterText);
+	UE_API FSequencerFilterData(const FString& InRawFilterText);
 
-	bool operator==(const FSequencerFilterData& InRhs) const;
-	bool operator!=(const FSequencerFilterData& InRhs) const;
+	UE_API bool operator==(const FSequencerFilterData& InRhs) const;
+	UE_API bool operator!=(const FSequencerFilterData& InRhs) const;
 
-	void Reset();
+	UE_API void Reset();
 
-	FString GetRawFilterText() const;
+	UE_API FString GetRawFilterText() const;
 
-	uint32 GetDisplayNodeCount() const;
-	uint32 GetTotalNodeCount() const;
+	UE_API uint32 GetDisplayNodeCount() const;
+	UE_API uint32 GetTotalNodeCount() const;
 
-	uint32 GetFilterInCount() const;
-	uint32 GetFilterOutCount() const;
+	UE_API uint32 GetFilterInCount() const;
+	UE_API uint32 GetFilterOutCount() const;
 
-	void IncrementTotalNodeCount();
+	UE_API void IncrementTotalNodeCount();
 
-	void FilterInNode(UE::Sequencer::TWeakViewModelPtr<UE::Sequencer::IOutlinerExtension> InNodeWeak);
-	void FilterOutNode(UE::Sequencer::TWeakViewModelPtr<UE::Sequencer::IOutlinerExtension> InNodeWeak);
+	UE_API void FilterInNode(UE::Sequencer::TWeakViewModelPtr<UE::Sequencer::IOutlinerExtension> InNodeWeak);
+	UE_API void FilterOutNode(UE::Sequencer::TWeakViewModelPtr<UE::Sequencer::IOutlinerExtension> InNodeWeak);
 
-	void FilterInParentChildNodes(const UE::Sequencer::TViewModelPtr<UE::Sequencer::IOutlinerExtension>& InNode
+	UE_API void FilterInParentChildNodes(const UE::Sequencer::TViewModelPtr<UE::Sequencer::IOutlinerExtension>& InNode
 		, const bool bInIncludeSelf, const bool bInIncludeParents, const bool bInIncludeChildren = false);
 
-	void FilterInNodeWithAncestors(const UE::Sequencer::TViewModelPtr<UE::Sequencer::IOutlinerExtension>& InNode);
+	UE_API void FilterInNodeWithAncestors(const UE::Sequencer::TViewModelPtr<UE::Sequencer::IOutlinerExtension>& InNode);
 
-	bool ContainsFilterInNodes(const FSequencerFilterData& InOtherData) const;
+	UE_API bool ContainsFilterInNodes(const FSequencerFilterData& InOtherData) const;
 
-	bool IsFilteredOut(const UE::Sequencer::TViewModelPtr<UE::Sequencer::IOutlinerExtension>& InNode) const;
+	UE_API bool IsFilteredOut(const UE::Sequencer::TViewModelPtr<UE::Sequencer::IOutlinerExtension>& InNode) const;
 
-	UE::Sequencer::TWeakViewModelPtr<UE::Sequencer::ITrackExtension> ResolveTrack(FSequencerTrackFilterType InNode);
-	UMovieSceneTrack* ResolveMovieSceneTrackObject(FSequencerTrackFilterType InNode);
-	UObject* ResolveTrackBoundObject(ISequencer& InSequencer, FSequencerTrackFilterType InNode);
+	UE_API UE::Sequencer::TWeakViewModelPtr<UE::Sequencer::ITrackExtension> ResolveTrack(FSequencerTrackFilterType InNode);
+	UE_API UMovieSceneTrack* ResolveMovieSceneTrackObject(FSequencerTrackFilterType InNode);
+	UE_API UObject* ResolveTrackBoundObject(ISequencer& InSequencer, FSequencerTrackFilterType InNode);
 
 	TMap<UE::Sequencer::TWeakViewModelPtr<UE::Sequencer::IOutlinerExtension>, UE::Sequencer::TWeakViewModelPtr<UE::Sequencer::ITrackExtension>> ResolvedTracks;
 	TMap<UE::Sequencer::TWeakViewModelPtr<UE::Sequencer::IOutlinerExtension>, TWeakObjectPtr<UMovieSceneTrack>> ResolvedTrackObjects;
@@ -64,3 +66,5 @@ protected:
 	/** Nodes to be displayed in the UI */
 	TSet<UE::Sequencer::TWeakViewModelPtr<UE::Sequencer::IOutlinerExtension>> FilterInNodes;
 };
+
+#undef UE_API
diff --git a/Engine/Source/Editor/Sequencer/Public/Filters/SequencerTrackFilterBase.h b/Engine/Source/Editor/Sequencer/Public/Filters/SequencerTrackFilterBase.h
--- a/Engine/Source/Editor/Sequencer/Public/Filters/SequencerTrackFilterBase.h
+++ b/Engine/Source/Editor/Sequencer/Public/Filters/SequencerTrackFilterBase.h
@@ -1,4 +1,4 @@
-﻿// Copyright Epic Games, Inc. All Rights Reserved.
+// Copyright Epic Games, Inc. All Rights Reserved.
 
 #pragma once
 
@@ -55,12 +55,12 @@ public:
 	//~ End FSequencerFilterBase
 
 	/** Returns whether this filter needs reevaluating any time track values have been modified, not just tree changes */
-	SEQUENCER_API virtual bool ShouldUpdateOnTrackValueChanged() const { return false; }
+	virtual bool ShouldUpdateOnTrackValueChanged() const { return false; }
 
 	/** Returns whether the filter supports the sequence type */
 	SEQUENCER_API virtual bool SupportsSequence(UMovieSceneSequence* const InSequence) const override;
 
-	SEQUENCER_API virtual TSubclassOf<UMovieSceneTrack> GetTrackClass() const { return nullptr; }
+	virtual TSubclassOf<UMovieSceneTrack> GetTrackClass() const { return nullptr; }
 
 protected:
 	UMovieSceneSequence* GetFocusedMovieSceneSequence() const;
diff --git a/Engine/Source/Editor/Sequencer/Public/Filters/SequencerTrackFilterExtension.h b/Engine/Source/Editor/Sequencer/Public/Filters/SequencerTrackFilterExtension.h
--- a/Engine/Source/Editor/Sequencer/Public/Filters/SequencerTrackFilterExtension.h
+++ b/Engine/Source/Editor/Sequencer/Public/Filters/SequencerTrackFilterExtension.h
@@ -1,4 +1,4 @@
-﻿// Copyright Epic Games, Inc. All Rights Reserved.
+// Copyright Epic Games, Inc. All Rights Reserved.
 
 #pragma once
 
@@ -11,8 +11,8 @@ class ISequencerTrackFilters;
 /**
  * Derive from this class to make additional track filters available in Sequencer.
  */
-UCLASS(Abstract)
-class SEQUENCER_API USequencerTrackFilterExtension : public UObject
+UCLASS(MinimalAPI, Abstract)
+class USequencerTrackFilterExtension : public UObject
 {
 	GENERATED_BODY()
 
diff --git a/Engine/Source/Editor/Sequencer/Public/Filters/SequencerTrackFilterTextExpressionExtension.h b/Engine/Source/Editor/Sequencer/Public/Filters/SequencerTrackFilterTextExpressionExtension.h
--- a/Engine/Source/Editor/Sequencer/Public/Filters/SequencerTrackFilterTextExpressionExtension.h
+++ b/Engine/Source/Editor/Sequencer/Public/Filters/SequencerTrackFilterTextExpressionExtension.h
@@ -10,8 +10,8 @@ class ISequencerTrackFilters;
 /**
  * Derive from this class to make additional track filter text expressions available in Sequencer.
  */
-UCLASS(Abstract)
-class SEQUENCER_API USequencerTrackFilterTextExpressionExtension : public UObject
+UCLASS(MinimalAPI, Abstract)
+class USequencerTrackFilterTextExpressionExtension : public UObject
 {
 	GENERATED_BODY()
 
diff --git a/Engine/Source/Editor/Sequencer/Public/Filters/Widgets/SFilterBarClippingHorizontalBox.h b/Engine/Source/Editor/Sequencer/Public/Filters/Widgets/SFilterBarClippingHorizontalBox.h
--- a/Engine/Source/Editor/Sequencer/Public/Filters/Widgets/SFilterBarClippingHorizontalBox.h
+++ b/Engine/Source/Editor/Sequencer/Public/Filters/Widgets/SFilterBarClippingHorizontalBox.h
@@ -1,4 +1,4 @@
-﻿// Copyright Epic Games, Inc. All Rights Reserved.
+// Copyright Epic Games, Inc. All Rights Reserved.
 
 #pragma once
 
@@ -6,15 +6,17 @@
 #include "Framework/SlateDelegates.h"
 #include "Widgets/SBoxPanel.h"
 
+#define UE_API SEQUENCER_API
+
 class SComboButton;
 
 /**
  * Specialized widget for handling horizontal clipping of filter widgets that go beyond the filter bar widget width.
  */
-class SEQUENCER_API SFilterBarClippingHorizontalBox : public SHorizontalBox
+class SFilterBarClippingHorizontalBox : public SHorizontalBox
 {
 public:
-	static TSharedRef<SWidget> WrapVerticalListWithHeading(const TSharedRef<SWidget>& InWidget, const FPointerEventHandler InMouseButtonUpEvent);
+	static UE_API TSharedRef<SWidget> WrapVerticalListWithHeading(const TSharedRef<SWidget>& InWidget, const FPointerEventHandler InMouseButtonUpEvent);
 
 	SLATE_BEGIN_ARGS(SFilterBarClippingHorizontalBox) 
 		: _IsFocusable(true)
@@ -23,16 +25,16 @@ public:
 		SLATE_ARGUMENT(bool, IsFocusable)
 	SLATE_END_ARGS()
 
-	void Construct(const FArguments& InArgs);
+	UE_API void Construct(const FArguments& InArgs);
 
-	TSharedRef<SComboButton> CreateWrapButton();
+	UE_API TSharedRef<SComboButton> CreateWrapButton();
 
 	/** Returns to index of the first clipped child */
 	int32 GetClippedIndex() const { return ClippedIndex; }
 
 protected:
 	//~ Begin SPanel
-	virtual void OnArrangeChildren(const FGeometry& InAllottedGeometry, FArrangedChildren& ArrangedChildren) const override;
+	UE_API virtual void OnArrangeChildren(const FGeometry& InAllottedGeometry, FArrangedChildren& ArrangedChildren) const override;
 	//~ End SPanel
 
 	/** The button that is displayed when a toolbar or menubar is clipped */
@@ -53,3 +55,5 @@ protected:
 
 	bool bIsFocusable = true;
 };
+
+#undef UE_API
diff --git a/Engine/Source/Editor/Sequencer/Public/Filters/Widgets/SFilterExpressionHelpDialog.h b/Engine/Source/Editor/Sequencer/Public/Filters/Widgets/SFilterExpressionHelpDialog.h
--- a/Engine/Source/Editor/Sequencer/Public/Filters/Widgets/SFilterExpressionHelpDialog.h
+++ b/Engine/Source/Editor/Sequencer/Public/Filters/Widgets/SFilterExpressionHelpDialog.h
@@ -1,18 +1,20 @@
-﻿// Copyright Epic Games, Inc. All Rights Reserved.
+// Copyright Epic Games, Inc. All Rights Reserved.
 
 #pragma once
 
 #include "Widgets/SWindow.h"
 
+#define UE_API SEQUENCER_API
+
 class ISequencerTextFilterExpressionContext;
 enum class ESequencerTextFilterValueType : uint8;
 
-struct SEQUENCER_API FFilterExpressionHelpDialogConfig
+struct FFilterExpressionHelpDialogConfig
 {
 	static constexpr float DefaultMaxDesiredWidth = 460.f;
 	static constexpr float DefaultMaxDesiredHeight = 560.f;
 
-	FFilterExpressionHelpDialogConfig();
+	UE_API FFilterExpressionHelpDialogConfig();
 
 	FName IdentifierName;
 
@@ -52,3 +54,5 @@ protected:
 
 	FFilterExpressionHelpDialogConfig Config;
 };
+
+#undef UE_API
diff --git a/Engine/Source/Editor/Sequencer/Public/Filters/Widgets/SSequencerSearchBox.h b/Engine/Source/Editor/Sequencer/Public/Filters/Widgets/SSequencerSearchBox.h
--- a/Engine/Source/Editor/Sequencer/Public/Filters/Widgets/SSequencerSearchBox.h
+++ b/Engine/Source/Editor/Sequencer/Public/Filters/Widgets/SSequencerSearchBox.h
@@ -9,6 +9,8 @@
 #include "Widgets/DeclarativeSyntaxSupport.h"
 #include "Widgets/SCompoundWidget.h"
 
+#define UE_API SEQUENCER_API
+
 class FSequencerFilterSuggestionListEntryBase;
 class FTextFilterExpressionEvaluator;
 class ISequencerFilterBar;
@@ -28,7 +30,7 @@ DECLARE_DELEGATE_RetVal_TwoParams(FText, FOnSequencerSearchBoxSuggestionChosen,
 /**
  * A wrapper widget around SFilterSearchBox to provide filter text expression suggestions in a dropdown menu.
  */
-class SEQUENCER_API SSequencerSearchBox : public SCompoundWidget
+class SSequencerSearchBox : public SCompoundWidget
 {
 public:
 	SLATE_BEGIN_ARGS(SSequencerSearchBox)
@@ -81,26 +83,26 @@ public:
 	SLATE_END_ARGS()
 
 	/** Constructs this widget with InArgs */
-	void Construct(const FArguments& InArgs, const TWeakPtr<ISequencerFilterBar>& InWeakFilterBar);
+	UE_API void Construct(const FArguments& InArgs, const TWeakPtr<ISequencerFilterBar>& InWeakFilterBar);
 
 	/** Get the text string currently being edited */
-	FText GetText() const;
+	UE_API FText GetText() const;
 
 	/** Sets the text string currently being edited */
-	void SetText(const TAttribute<FText>& InNewText);
+	UE_API void SetText(const TAttribute<FText>& InNewText);
 
 	/** Set or clear the current error reporting information for this search box */
-	void SetError(const FText& InError);
-	void SetError(const FString& InError);
+	UE_API void SetError(const FText& InError);
+	UE_API void SetError(const FString& InError);
 
 	/** Show a + button next to the current search and set the handler for when that is clicked */
-	void SetOnSaveSearchHandler(SFilterSearchBox::FOnSaveSearchClicked InOnSaveSearchHandler);
+	UE_API void SetOnSaveSearchHandler(SFilterSearchBox::FOnSaveSearchClicked InOnSaveSearchHandler);
 
 	//~ Begin SWidget
-	virtual FReply OnPreviewKeyDown(const FGeometry& InGeometry, const FKeyEvent& InKeyEvent) override;
-	virtual bool SupportsKeyboardFocus() const override;
-	virtual bool HasKeyboardFocus() const override;
-	virtual FReply OnFocusReceived(const FGeometry& InGeometry, const FFocusEvent& InFocusEvent) override;
+	UE_API virtual FReply OnPreviewKeyDown(const FGeometry& InGeometry, const FKeyEvent& InKeyEvent) override;
+	UE_API virtual bool SupportsKeyboardFocus() const override;
+	UE_API virtual bool HasKeyboardFocus() const override;
+	UE_API virtual FReply OnFocusReceived(const FGeometry& InGeometry, const FFocusEvent& InFocusEvent) override;
 	//~ End SWidget
 
 private:
@@ -111,55 +113,55 @@ private:
 		int32 SuggestionIndex = 0;
 	};
 
-	void ExtractSearchFilterTerms(const FText& InSearchText, FParsedExpression& OutParsedPair) const;
+	UE_API void ExtractSearchFilterTerms(const FText& InSearchText, FParsedExpression& OutParsedPair) const;
 
-	TSharedRef<SWidget> GetSuggestionListMenuContent();
+	UE_API TSharedRef<SWidget> GetSuggestionListMenuContent();
 
 	/** First chance handler for key down events to the editable text widget */
-	FReply HandleKeyDown(const FGeometry& InGeometry, const FKeyEvent& InKeyEvent);
+	UE_API FReply HandleKeyDown(const FGeometry& InGeometry, const FKeyEvent& InKeyEvent);
 
 	/** Handler for when text in the editable text box changed */
-	void HandleTextChanged(const FText& InNewText);
+	UE_API void HandleTextChanged(const FText& InNewText);
 
 	/** Handler for when text in the editable text box changed */
-	void HandleTextCommitted(const FText& InNewText, const ETextCommit::Type InCommitType);
+	UE_API void HandleTextCommitted(const FText& InNewText, const ETextCommit::Type InCommitType);
 
 	/** Called by SListView when the selection changes in the suggestion list */
-	void OnSelectionChanged(TSharedPtr<FSequencerFilterSuggestionListEntryBase> InNewValue, const ESelectInfo::Type InSelectInfo);
+	UE_API void OnSelectionChanged(TSharedPtr<FSequencerFilterSuggestionListEntryBase> InNewValue, const ESelectInfo::Type InSelectInfo);
 
 	/** Makes the widget for a suggestion message in the list view */
-	TSharedRef<ITableRow> MakeSuggestionListItemWidget(const TSharedPtr<FSequencerFilterSuggestionListEntryBase> InSuggestion, const TSharedRef<STableViewBase>& InOwnerTable);
+	UE_API TSharedRef<ITableRow> MakeSuggestionListItemWidget(const TSharedPtr<FSequencerFilterSuggestionListEntryBase> InSuggestion, const TSharedRef<STableViewBase>& InOwnerTable);
 
 	/** Gets the text to highlight in the suggestion list */
-	FText GetHighlightText() const;
+	UE_API FText GetHighlightText() const;
 
 	/** Updates and shows or hides the suggestion list */
-	void UpdateSuggestionList();
+	UE_API void UpdateSuggestionList();
 
 	/** Sets the focus to the InputText box */
-	void FocusEditBox();
+	UE_API void FocusEditBox();
 
 	/** Returns the currently selected suggestion */
-	TSharedPtr<FSequencerFilterSuggestionListEntryBase> GetSelectedSuggestion() const;
+	UE_API TSharedPtr<FSequencerFilterSuggestionListEntryBase> GetSelectedSuggestion() const;
 
 	/** Default implementation of OnSuggestionFilter, if no external implementation is provided */
-	void DefaultSuggestionFilterImpl(const FText& InSearchText, TArray<FSequencerFilterSuggestion>& OutPossibleSuggestions, FText& OutSuggestionHighlightText);
+	UE_API void DefaultSuggestionFilterImpl(const FText& InSearchText, TArray<FSequencerFilterSuggestion>& OutPossibleSuggestions, FText& OutSuggestionHighlightText);
 
 	/** Default implementation of OnSuggestionChosen, if no external implementation is provided */
-	FText DefaultSuggestionChosenImpl(const FText& InSearchText, const FString& InSuggestion);
+	UE_API FText DefaultSuggestionChosenImpl(const FText& InSearchText, const FString& InSuggestion);
 
 	/** Provides default key suggestions, if no suggestions are provided */
-	void DefaultKeySuggestions(const FString& InKeyValue, TArray<FSequencerFilterSuggestion>& OutPossibleSuggestions);
+	UE_API void DefaultKeySuggestions(const FString& InKeyValue, TArray<FSequencerFilterSuggestion>& OutPossibleSuggestions);
 
 	/** Provides default value suggestions, if no suggestions are provided */
-	void DefaultValueSuggestions(const FString& InKeyValue, TArray<FSequencerFilterSuggestion>& OutPossibleSuggestions);
+	UE_API void DefaultValueSuggestions(const FString& InKeyValue, TArray<FSequencerFilterSuggestion>& OutPossibleSuggestions);
 
-	static bool IsOperatorToken(const TExpressionToken<TCHAR>& InToken);
-	static bool IsLogicalOperatorToken(const TExpressionToken<TCHAR>& InToken);
+	static UE_API bool IsOperatorToken(const TExpressionToken<TCHAR>& InToken);
+	static UE_API bool IsLogicalOperatorToken(const TExpressionToken<TCHAR>& InToken);
 
-	static int32 FindTokenIndex(const TArray<TExpressionToken<TCHAR>>& ExpressionTokens, const int32 InIndexToFind);
+	static UE_API int32 FindTokenIndex(const TArray<TExpressionToken<TCHAR>>& ExpressionTokens, const int32 InIndexToFind);
 
-	const FTextFilterExpressionEvaluator* GetTextFilterExpressionEvaluator() const;
+	UE_API const FTextFilterExpressionEvaluator* GetTextFilterExpressionEvaluator() const;
 
 	TWeakPtr<ISequencerFilterBar> WeakFilterBar;
 
@@ -201,3 +203,5 @@ private:
 
 	int32 LastChosenCursorOffset = INDEX_NONE;
 };
+
+#undef UE_API
diff --git a/Engine/Source/Editor/Sequencer/Public/IKeyArea.h b/Engine/Source/Editor/Sequencer/Public/IKeyArea.h
--- a/Engine/Source/Editor/Sequencer/Public/IKeyArea.h
+++ b/Engine/Source/Editor/Sequencer/Public/IKeyArea.h
@@ -12,6 +12,8 @@
 #include "MovieSceneCommonHelpers.h"
 #include "Channels/MovieSceneChannelHandle.h"
 
+#define UE_API SEQUENCER_API
+
 struct FCurveModelID;
 struct FRichCurve;
 struct FKeyDrawParams;
@@ -37,7 +39,7 @@ class FSequencerSectionPainter;
 /**
  * Interface that should be implemented for the UI portion of a key area within a section
  */
-class SEQUENCER_API IKeyArea : public TSharedFromThis<IKeyArea>
+class IKeyArea : public TSharedFromThis<IKeyArea>
 {
 public:
 
@@ -50,7 +52,7 @@ public:
 	 * @param InSection The section that owns the channel that this key area represents
 	 * @param InChannel Handle to the channel this key area represents
 	 */
-	IKeyArea(TWeakPtr<ISequencerSection> InSection, FMovieSceneChannelHandle InChannel, const FCurveModelID& InCurveId);
+	UE_API IKeyArea(TWeakPtr<ISequencerSection> InSection, FMovieSceneChannelHandle InChannel, const FCurveModelID& InCurveId);
 
 	/**
 	 * Reinitialize this key area with a new channel
@@ -58,7 +60,7 @@ public:
 	 * @param InSection The section that owns the channel that this key area represents
 	 * @param InChannel Handle to the channel this key area represents
 	 */
-	void Reinitialize(TWeakPtr<ISequencerSection> InSection, FMovieSceneChannelHandle InChannel);
+	UE_API void Reinitialize(TWeakPtr<ISequencerSection> InSection, FMovieSceneChannelHandle InChannel);
 
 public:
 
@@ -68,7 +70,7 @@ public:
 	 *
 	 * @return The channel interface, or nullptr if one was not found
 	 */
-	ISequencerChannelInterface* FindChannelEditorInterface() const;
+	UE_API ISequencerChannelInterface* FindChannelEditorInterface() const;
 
 	/**
 	 * Access the channel type identifier for the channel that this key area wraps
@@ -95,21 +97,21 @@ public:
 	 *
 	 * @return A (potentially invalid) interface to this key area's channel
 	 */
-	FMovieSceneChannel* ResolveChannel() const;
+	UE_API FMovieSceneChannel* ResolveChannel() const;
 
 	/**
 	 * Get this key area's name
 	 *
 	 * @return This key area's name
 	 */
-	FName GetName() const;
+	UE_API FName GetName() const;
 
 	/**
 	 * Set this key area's name
 	 *
 	 * @param InName This key area's name
 	 */
-	void SetName(FName InName);
+	UE_API void SetName(FName InName);
 
 	/**
 	 * Get the color of this channel that should be drawn underneath its keys
@@ -126,7 +128,7 @@ public:
 	 *
 	 * @return The owning section, or nullptr if it has been destroyed
 	 */
-	UMovieSceneSection* GetOwningSection() const;
+	UE_API UMovieSceneSection* GetOwningSection() const;
 
 
 	/**
@@ -135,21 +137,21 @@ public:
 	 *
 	 * @return The owning object, or nullptr if it has been destroyed
 	 */
-	UObject* GetOwningObject() const;
+	UE_API UObject* GetOwningObject() const;
 
 	/**
 	 * Access section editor interface
 	 *
 	 * @return The owning section interface, or nullptr if it has been destroyed
 	 */
-	TSharedPtr<ISequencerSection> GetSectionInterface() const;
+	UE_API TSharedPtr<ISequencerSection> GetSectionInterface() const;
 
 	/**
 	 * Check whether we should show this key area's curve on its channel UI or not
 	 *
 	 * @return True to show the curve, false otherwise
 	 */
-	bool ShouldShowCurve() const;
+	UE_API bool ShouldShowCurve() const;
 
 public:
 
@@ -161,7 +163,7 @@ public:
 	 * @param InSequencer     The currently active sequencer
 	 * @return A handle to the key that was added or updated
 	 */
-	FKeyHandle AddOrUpdateKey(FFrameNumber Time, const FGuid& ObjectBindingID, ISequencer& InSequencer);
+	UE_API FKeyHandle AddOrUpdateKey(FFrameNumber Time, const FGuid& ObjectBindingID, ISequencer& InSequencer);
 
 	/**
 	 * Delete the specified keys. If all keys are removed, the current value at that time will be set as the default value for the channel
@@ -170,7 +172,7 @@ public:
 	 * @param InHandles             The key handles to delete
 	 * @param InTime                The time at which to evaluate for the default value if there are no keys remaining
 	 */
-	void DeleteKeys(TArrayView<const FKeyHandle> InHandles, FFrameNumber InTime);
+	UE_API void DeleteKeys(TArrayView<const FKeyHandle> InHandles, FFrameNumber InTime);
 
 	/**
 	 * Duplicate the key represented by the specified handle
@@ -178,7 +180,7 @@ public:
 	 * @param InKeyHandle     A handle to the key to duplicate
 	 * @return A handle to the new duplicated key
 	 */
-	FKeyHandle DuplicateKey(FKeyHandle InKeyHandle) const;
+	UE_API FKeyHandle DuplicateKey(FKeyHandle InKeyHandle) const;
 
 	/**
 	 * Get the time of the key represented by the specified handle
@@ -199,7 +201,7 @@ public:
 	 * @param InKeyHandles     A handle to the key to query for
 	 * @param OutTimes         A pre-sized array view to populate with key times
 	 */
-	void GetKeyTimes(TArrayView<const FKeyHandle> InKeyHandles, TArrayView<FFrameNumber> OutTimes) const;
+	UE_API void GetKeyTimes(TArrayView<const FKeyHandle> InKeyHandles, TArrayView<FFrameNumber> OutTimes) const;
 
 	/**
 	 * Get all key handles that exist within the given time range
@@ -230,7 +232,7 @@ public:
 	 * @param OutTimes         (Optional) An array to populate with key times, any key times will be appended to this array
 	 * @param WithinRange      (Optional) A predicate range to search within
 	 */
-	void GetKeyInfo(TArray<FKeyHandle>* OutHandles, TArray<FFrameNumber>* OutTimes, const TRange<FFrameNumber>& WithinRange = TRange<FFrameNumber>::All()) const;
+	UE_API void GetKeyInfo(TArray<FKeyHandle>* OutHandles, TArray<FFrameNumber>* OutTimes, const TRange<FFrameNumber>& WithinRange = TRange<FFrameNumber>::All()) const;
 
 	/**
 	 * Set the time of the key with the specified handle
@@ -249,7 +251,7 @@ public:
 	 * @param InKeyHandles     An array of handles that should have their time set to times in the corresponding InKeyTimes array
 	 * @param InKeyTimes       Array of times to set to, one per key handle. Must match the size of InKeyHandles
 	 */
-	void SetKeyTimes(TArrayView<const FKeyHandle> InKeyHandles, TArrayView<const FFrameNumber> InKeyTimes) const;
+	UE_API void SetKeyTimes(TArrayView<const FKeyHandle> InKeyHandles, TArrayView<const FFrameNumber> InKeyTimes) const;
 
 public:
 
@@ -259,7 +261,7 @@ public:
 	 * @param InKeyHandles     An array of handles to draw
 	 * @param OutKeyDrawParams Pre-sized array view of parameters to populate with each key
 	 */
-	void DrawKeys(TArrayView<const FKeyHandle> InKeyHandles, TArrayView<FKeyDrawParams> OutKeyDrawParams);
+	UE_API void DrawKeys(TArrayView<const FKeyHandle> InKeyHandles, TArrayView<FKeyDrawParams> OutKeyDrawParams);
 
 	/**
 	 * Copy all the keys specified in KeyMask to the specified clipboard
@@ -267,7 +269,7 @@ public:
 	 * @param ClipboardBuilder The structure responsible for building clipboard information for each key
 	 * @param KeyMask          A specific set of keys to copy
 	 */
-	void CopyKeys(FMovieSceneClipboardBuilder& ClipboardBuilder, TArrayView<const FKeyHandle> KeyMask) const;
+	UE_API void CopyKeys(FMovieSceneClipboardBuilder& ClipboardBuilder, TArrayView<const FKeyHandle> KeyMask) const;
 
 	/**
 	 * Draw additional content in addition to keys for a particular IKeyArea
@@ -276,7 +278,7 @@ public:
 	 * @param LayerId          The slate layer to paint onto
 	 * @return The new slate layer ID for subsequent elements to paint onto
 	 */
-	int32 DrawExtra(const FSequencerChannelPaintArgs& PaintArgs, int32 LayerId) const;
+	UE_API int32 DrawExtra(const FSequencerChannelPaintArgs& PaintArgs, int32 LayerId) const;
 
 	/**
 	 * Paste the specified key track into this key area
@@ -285,7 +287,7 @@ public:
 	 * @param SrcEnvironment   The environment the source data was copied from
 	 * @param DstEnvironment   The environment we're pasting into
 	 */
-	TArray<FKeyHandle> PasteKeys(const FMovieSceneClipboardKeyTrack& KeyTrack, const FMovieSceneClipboardEnvironment& SrcEnvironment, const FSequencerPasteEnvironment& DstEnvironment);
+	UE_API TArray<FKeyHandle> PasteKeys(const FMovieSceneClipboardKeyTrack& KeyTrack, const FMovieSceneClipboardEnvironment& SrcEnvironment, const FSequencerPasteEnvironment& DstEnvironment);
 
 
 	/**
@@ -293,7 +295,7 @@ public:
 	 *
 	 * @return (Optional) A new model to be added to a curve editor
 	 */
-	TUniquePtr<FCurveModel> CreateCurveEditorModel(TSharedRef<ISequencer> InSequencer) const;
+	UE_API TUniquePtr<FCurveModel> CreateCurveEditorModel(TSharedRef<ISequencer> InSequencer) const;
 
 public:
 
@@ -304,12 +306,12 @@ public:
 	 *
 	 * @return The key structure for this channel, or nullptr
 	 */
-	TSharedPtr<FStructOnScope> GetKeyStruct(FKeyHandle KeyHandle) const;
+	UE_API TSharedPtr<FStructOnScope> GetKeyStruct(FKeyHandle KeyHandle) const;
 
 	/**
 	 * Check whether this key area can create an editor on the sequencer node tree
 	 */
-	bool CanCreateKeyEditor() const;
+	UE_API bool CanCreateKeyEditor() const;
 
 	/**
 	 * Create an editor on the sequencer node tree
@@ -318,7 +320,7 @@ public:
 	 * @param ObjectBindingID  The ID of the object this key area's track is bound to
 	 * @return The editor widget to display on the key area's node
 	 */
-	TSharedRef<SWidget> CreateKeyEditor(TWeakPtr<ISequencer> Sequencer, const FGuid& ObjectBindingID);
+	UE_API TSharedRef<SWidget> CreateKeyEditor(TWeakPtr<ISequencer> Sequencer, const FGuid& ObjectBindingID);
 
 private:
 
@@ -342,4 +344,6 @@ private:
 
 	/** Consistent ID that is reused and injected into the curve model every time it is created. Pimpl to avoid public build.cs dependency. */
 	const TPimplPtr<FCurveModelID> CurveModelID;
-};
\ No newline at end of file
+};
+
+#undef UE_API
diff --git a/Engine/Source/Editor/Sequencer/Public/ISequencer.h b/Engine/Source/Editor/Sequencer/Public/ISequencer.h
--- a/Engine/Source/Editor/Sequencer/Public/ISequencer.h
+++ b/Engine/Source/Editor/Sequencer/Public/ISequencer.h
@@ -34,9 +34,6 @@
 #include "Widgets/SWidget.h"
 #include "ISequencerNumericTypeInterface.h"
 
-#if UE_ENABLE_INCLUDE_ORDER_DEPRECATED_IN_5_3
-#include "ITimeSlider.h"
-#endif
 #include "Bindings/MovieSceneCustomBinding.h"
 
 #include "ISequencer.generated.h"
@@ -331,7 +328,7 @@ public:
 	/**
 	* Create a new binding for the specified object
 	*/
-	SEQUENCER_API virtual FGuid CreateBinding(UObject& InObject, const UE::Sequencer::FCreateBindingParams& InParams) = 0;
+	virtual FGuid CreateBinding(UObject& InObject, const UE::Sequencer::FCreateBindingParams& InParams) = 0;
 
 	// Override from IMovieScenePlayer
 	FGuid CreateBinding(UMovieSceneSequence* InSequence, UObject* InObject) override;
@@ -894,7 +891,7 @@ public:
 	virtual TArray<FText> GetTrackFilterNames() const = 0;
 
 	/** Gets the Sequencer filter interface used to manage filters */
-	SEQUENCER_API virtual TSharedRef<ISequencerTrackFilters> GetFilterInterface() const = 0;
+	virtual TSharedRef<ISequencerTrackFilters> GetFilterInterface() const = 0;
 
 public:
 
diff --git a/Engine/Source/Editor/Sequencer/Public/ISequencerObjectSchema.h b/Engine/Source/Editor/Sequencer/Public/ISequencerObjectSchema.h
--- a/Engine/Source/Editor/Sequencer/Public/ISequencerObjectSchema.h
+++ b/Engine/Source/Editor/Sequencer/Public/ISequencerObjectSchema.h
@@ -6,6 +6,8 @@
 #include "Templates/SharedPointer.h"
 #include "Containers/ContainersFwd.h"
 
+#define UE_API SEQUENCER_API
+
 class UObject;
 class FExtender;
 class ISequencer;
@@ -35,11 +37,11 @@ struct FObjectSchemaRelevancy
 	friend bool operator>(const FObjectSchemaRelevancy& A, const FObjectSchemaRelevancy& B);
 };
 
-class SEQUENCER_API IObjectSchema : public TSharedFromThis<IObjectSchema>
+class IObjectSchema : public TSharedFromThis<IObjectSchema>
 {
 public:
 
-	static TMap<const IObjectSchema*, TArray<UObject*>> ComputeRelevancy(TArrayView<UObject* const> InObjects);
+	static UE_API TMap<const IObjectSchema*, TArray<UObject*>> ComputeRelevancy(TArrayView<UObject* const> InObjects);
 
 	virtual ~IObjectSchema(){}
 
@@ -49,7 +51,9 @@ public:
 
 	virtual TSharedPtr<FExtender> ExtendObjectBindingMenu(TSharedRef<FUICommandList> CommandList, TWeakPtr<ISequencer> WeakSequencer, TArrayView<UObject* const> ContextSensitiveObjects) const = 0;
 
-	virtual FText GetPrettyName(const UObject* Object) const;
+	UE_API virtual FText GetPrettyName(const UObject* Object) const;
 };
 
-} // namespace UE::Sequencer
\ No newline at end of file
+} // namespace UE::Sequencer
+
+#undef UE_API
diff --git a/Engine/Source/Editor/Sequencer/Public/ISequencerSection.h b/Engine/Source/Editor/Sequencer/Public/ISequencerSection.h
--- a/Engine/Source/Editor/Sequencer/Public/ISequencerSection.h
+++ b/Engine/Source/Editor/Sequencer/Public/ISequencerSection.h
@@ -318,13 +318,13 @@ public:
 	/**
 	Dilation starts with a drag operation
 	*/
-	SEQUENCER_API virtual void BeginDilateSection() {};
+	virtual void BeginDilateSection() {};
 	/**
 	New Range that's set as we Dilate
 	@param NewRange The NewRange.
 	@param DilationFactor The factor we have dilated from the beginning of the drag
 	*/
-	SEQUENCER_API virtual void DilateSection(const TRange<FFrameNumber>& NewRange, float DilationFactor) {};
+	virtual void DilateSection(const TRange<FFrameNumber>& NewRange, float DilationFactor) {};
 
 
 
diff --git a/Engine/Source/Editor/Sequencer/Public/KeyPropertyParams.h b/Engine/Source/Editor/Sequencer/Public/KeyPropertyParams.h
--- a/Engine/Source/Editor/Sequencer/Public/KeyPropertyParams.h
+++ b/Engine/Source/Editor/Sequencer/Public/KeyPropertyParams.h
@@ -12,6 +12,8 @@
 #include "UObject/UnrealType.h"
 #include "UObject/WeakFieldPtr.h"
 
+#define UE_API SEQUENCER_API
+
 class IPropertyHandle;
 class UClass;
 class UObject;
@@ -22,27 +24,27 @@ enum class ESequencerKeyMode;
 /**
  * Parameters for determining if a property can be keyed.
  */
-struct SEQUENCER_API FCanKeyPropertyParams
+struct FCanKeyPropertyParams
 {
 	/**
 	 * Creates new can key property parameters.
 	 * @param InObjectClass the class of the object which has the property to be keyed.
 	 * @param InPropertyPath path get from the root object to the property to be keyed.
 	 */
-	FCanKeyPropertyParams(const UClass* InObjectClass, const FPropertyPath& InPropertyPath);
+	UE_API FCanKeyPropertyParams(const UClass* InObjectClass, const FPropertyPath& InPropertyPath);
 
 	/**
 	* Creates new can key property parameters.
 	* @param InObjectClass the class of the object which has the property to be keyed.
 	* @param InPropertyHandle a handle to the property to be keyed.
 	*/
-	FCanKeyPropertyParams(const UClass* InObjectClass, const IPropertyHandle& InPropertyHandle);
+	UE_API FCanKeyPropertyParams(const UClass* InObjectClass, const IPropertyHandle& InPropertyHandle);
 
 	/** The owner struct */
-	const UStruct* FindPropertyOwner(const FProperty* ForProperty) const;
+	UE_API const UStruct* FindPropertyOwner(const FProperty* ForProperty) const;
 
 	/** The container from which to find setter functions */
-	const UStruct* FindPropertyContainer(const FProperty* ForProperty) const;
+	UE_API const UStruct* FindPropertyContainer(const FProperty* ForProperty) const;
 
 	/** The class of the object which has the property to be keyed. */
 	const UClass* ObjectClass;
@@ -54,21 +56,21 @@ struct SEQUENCER_API FCanKeyPropertyParams
 /**
  * Parameters for keying a property.
  */
-struct SEQUENCER_API FKeyPropertyParams
+struct FKeyPropertyParams
 {
 	/**
 	* Creates new key property parameters for a manually triggered property change.
 	* @param InObjectsToKey an array of the objects who's property will be keyed.
 	* @param InPropertyPath path get from the root object to the property to be keyed.
 	*/
-	FKeyPropertyParams(TArray<UObject*> InObjectsToKey, const FPropertyPath& InPropertyPath, ESequencerKeyMode InKeyMode);
+	UE_API FKeyPropertyParams(TArray<UObject*> InObjectsToKey, const FPropertyPath& InPropertyPath, ESequencerKeyMode InKeyMode);
 
 	/**
 	* Creates new key property parameters from an actual property change notification with a property handle.
 	* @param InObjectsToKey an array of the objects who's property will be keyed.
 	* @param InPropertyHandle a handle to the property to be keyed.
 	*/
-	FKeyPropertyParams(TArray<UObject*> InObjectsToKey, const IPropertyHandle& InPropertyHandle, ESequencerKeyMode InKeyMode);
+	UE_API FKeyPropertyParams(TArray<UObject*> InObjectsToKey, const IPropertyHandle& InPropertyHandle, ESequencerKeyMode InKeyMode);
 
 	/** An array of the objects who's property will be keyed. */
 	const TArray<UObject*> ObjectsToKey;
@@ -83,10 +85,10 @@ struct SEQUENCER_API FKeyPropertyParams
 /**
  * Parameters for the property changed callback.
  */
-class SEQUENCER_API FPropertyChangedParams
+class FPropertyChangedParams
 {
 public:
-	FPropertyChangedParams(TArray<UObject*> InObjectsThatChanged, const FPropertyPath& InPropertyPath, const FPropertyPath& InStructPathToKey, ESequencerKeyMode InKeyMode);
+	UE_API FPropertyChangedParams(TArray<UObject*> InObjectsThatChanged, const FPropertyPath& InPropertyPath, const FPropertyPath& InStructPathToKey, ESequencerKeyMode InKeyMode);
 
 	/**
 	 * Gets the value of the property that changed.
@@ -146,7 +148,7 @@ public:
 	}
 
 	/** Gets the property path as a period seperated string of property names. */
-	FString GetPropertyPathString() const;
+	UE_API FString GetPropertyPathString() const;
 
 	/** An array of the objects that changed. */
 	const TArray<UObject*> ObjectsThatChanged;
@@ -172,3 +174,5 @@ private:
 
 template<> SEQUENCER_API bool FPropertyChangedParams::GetPropertyValueImpl<bool>(void* Data, const FPropertyInfo& Info);
 template<> SEQUENCER_API UObject* FPropertyChangedParams::GetPropertyValueImpl<UObject*>(void* Data, const FPropertyInfo& Info);
+
+#undef UE_API
diff --git a/Engine/Source/Editor/Sequencer/Public/LevelEditorSequencerIntegration.h b/Engine/Source/Editor/Sequencer/Public/LevelEditorSequencerIntegration.h
--- a/Engine/Source/Editor/Sequencer/Public/LevelEditorSequencerIntegration.h
+++ b/Engine/Source/Editor/Sequencer/Public/LevelEditorSequencerIntegration.h
@@ -12,6 +12,8 @@
 #include "Engine/EngineBaseTypes.h"
 #include "LevelEditorSequencerIntegration.generated.h"
 
+#define UE_API SEQUENCER_API
+
 class AActor;
 class FExtender;
 class FMenuBuilder;
@@ -112,124 +114,124 @@ struct TStructOpsTypeTraits<FLevelEditorSequencerUpdateGizmoTickFunction> : publ
 };
 
 
-class SEQUENCER_API FLevelEditorSequencerIntegration
+class FLevelEditorSequencerIntegration
 {
 public:
 
-	static FLevelEditorSequencerIntegration& Get();
+	static UE_API FLevelEditorSequencerIntegration& Get();
 
-	void Initialize(const FLevelEditorSequencerIntegrationOptions& Options);
+	UE_API void Initialize(const FLevelEditorSequencerIntegrationOptions& Options);
 
-	void AddSequencer(TSharedRef<ISequencer> InSequencer, const FLevelEditorSequencerIntegrationOptions& Options);
+	UE_API void AddSequencer(TSharedRef<ISequencer> InSequencer, const FLevelEditorSequencerIntegrationOptions& Options);
 
-	void OnSequencerReceivedFocus(TSharedRef<ISequencer> InSequencer);
+	UE_API void OnSequencerReceivedFocus(TSharedRef<ISequencer> InSequencer);
 
-	void RemoveSequencer(TSharedRef<ISequencer> InSequencer);
+	UE_API void RemoveSequencer(TSharedRef<ISequencer> InSequencer);
 
-	TArray<TWeakPtr<ISequencer>> GetSequencers();
+	UE_API TArray<TWeakPtr<ISequencer>> GetSequencers();
 	DECLARE_MULTICAST_DELEGATE(FOnSequencersChanged);
 	FOnSequencersChanged& GetOnSequencersChanged() { return OnSequencersChanged; };
 
 private:
 
 	/** Called before the world is going to be saved. The sequencer puts everything back to its initial state. */
-	void OnPreSaveWorld(UWorld* World, FObjectPreSaveContext ObjectSaveContext);
+	UE_API void OnPreSaveWorld(UWorld* World, FObjectPreSaveContext ObjectSaveContext);
 
 	/** Called after the world has been saved. The sequencer updates to the animated state. */
-	void OnPostSaveWorld(UWorld* World, FObjectPostSaveContext ObjectSaveContext);
+	UE_API void OnPostSaveWorld(UWorld* World, FObjectPostSaveContext ObjectSaveContext);
 
 	/** Called before any number of external actors are going to be saved. The sequencer puts everything back to its initial state. */
-	void OnPreSaveExternalActors(UWorld* World);
+	UE_API void OnPreSaveExternalActors(UWorld* World);
 
 	/** Called after any number of external actors has been saved. The sequencer puts everything back to its initial state. */
-	void OnPostSaveExternalActors(UWorld* World);
+	UE_API void OnPostSaveExternalActors(UWorld* World);
 
 	/** Called before asset validation is run on assets. The sequencer puts everything back to its initial state. */
-	void OnPreAssetValidation();
+	UE_API void OnPreAssetValidation();
 	
 	/** Called after asset validation has finished. The sequencer re-evaluates to hide the fact we did this from users. */
-	void OnPostAssetValidation();
+	UE_API void OnPostAssetValidation();
 
 	/** Called after a level has been added */
-	void OnLevelAdded(ULevel* InLevel, UWorld* InWorld);
+	UE_API void OnLevelAdded(ULevel* InLevel, UWorld* InWorld);
 
 	/** Called after a level has been removed */
-	void OnLevelRemoved(ULevel* InLevel, UWorld* InWorld);
+	UE_API void OnLevelRemoved(ULevel* InLevel, UWorld* InWorld);
 
 	/** Called after a new level has been created. The sequencer editor mode needs to be enabled. */
-	void OnNewCurrentLevel();
+	UE_API void OnNewCurrentLevel();
 
 	/** Called after a map has been opened. The sequencer editor mode needs to be enabled. */
-	void OnMapOpened(const FString& Filename, bool bLoadAsTemplate);
+	UE_API void OnMapOpened(const FString& Filename, bool bLoadAsTemplate);
 
 	/** Called when new actors are dropped in the viewport. */
-	void OnNewActorsDropped(const TArray<UObject*>& DroppedObjects, const TArray<AActor*>& DroppedActors);
+	UE_API void OnNewActorsDropped(const TArray<UObject*>& DroppedObjects, const TArray<AActor*>& DroppedActors);
 
 	/** Called when viewport tab content changes. */
-	void OnTabContentChanged();
+	UE_API void OnTabContentChanged();
 
 	/** Called when the map is changed. */
-	void OnMapChanged(UWorld* World, EMapChangeType MapChangeType);
+	UE_API void OnMapChanged(UWorld* World, EMapChangeType MapChangeType);
 
 	/** Called before a PIE session begins. */
-	void OnPreBeginPIE(bool bIsSimulating);
+	UE_API void OnPreBeginPIE(bool bIsSimulating);
 
 	/** Called after a PIE session ends. */
-	void OnEndPIE(bool bIsSimulating);
+	UE_API void OnEndPIE(bool bIsSimulating);
 
 	/** Called after PIE session ends and maps have been cleaned up */
-	void OnEndPlayMap();
+	UE_API void OnEndPlayMap();
 
 	/** Handles the actor selection changing externally .*/
-	void OnActorSelectionChanged( UObject* );
+	UE_API void OnActorSelectionChanged( UObject* );
 
 	/** Called when an actor label has changed */
-	void OnActorLabelChanged(AActor* ChangedActor);
+	UE_API void OnActorLabelChanged(AActor* ChangedActor);
 
 	/** Called when sequencer has been evaluated */
-	void OnSequencerEvaluated();
+	UE_API void OnSequencerEvaluated();
 
 	/** Called when bindings have changed */
-	void OnMovieSceneBindingsChanged();
+	UE_API void OnMovieSceneBindingsChanged();
 
 	/** Called when data has changed */
-	void OnMovieSceneDataChanged(EMovieSceneDataChangeType DataChangeType);
+	UE_API void OnMovieSceneDataChanged(EMovieSceneDataChangeType DataChangeType);
 
 	/** Called when allow edits mode has changed */
-	void OnAllowEditsModeChanged(EAllowEditsMode AllowEditsMode);
+	UE_API void OnAllowEditsModeChanged(EAllowEditsMode AllowEditsMode);
 
 	/** Called when the user begins playing/scrubbing */
-	void OnBeginDeferUpdates();
+	UE_API void OnBeginDeferUpdates();
 
 	/** Called when the user stops playing/scrubbing */
-	void OnEndDeferUpdates();
+	UE_API void OnEndDeferUpdates();
 
 	/** Called to determine whether a binding is visible in the tree view */
-	bool IsBindingVisible(const FMovieSceneBinding& InBinding);
+	UE_API bool IsBindingVisible(const FMovieSceneBinding& InBinding);
 
-	void RegisterMenus();
+	UE_API void RegisterMenus();
 
-	void MakeBrowseToSelectedActorSubMenu(UToolMenu* Menu);
-	void BrowseToSelectedActor(AActor* Actor, FSequencer* Sequencer, FMovieSceneSequenceID SequenceId);
+	UE_API void MakeBrowseToSelectedActorSubMenu(UToolMenu* Menu);
+	UE_API void BrowseToSelectedActor(AActor* Actor, FSequencer* Sequencer, FMovieSceneSequenceID SequenceId);
 
-	bool IsPropertyReadOnly(const FPropertyAndParent& InPropertyAndParent);
+	UE_API bool IsPropertyReadOnly(const FPropertyAndParent& InPropertyAndParent);
 
 private:
 
-	void ActivateSequencerEditorMode();
-	void DeactivateSequencerEditorMode();
-	void AddLevelViewportMenuExtender();
-	void BindDetailHandler(const FLevelEditorSequencerIntegrationOptions& Options);
-	void ActivateDetailHandler();
-	void AttachOutlinerColumn();
-	void DetachOutlinerColumn();
-	void ActivateRealtimeViewports();
-	void RestoreRealtimeViewports();
-	void RestoreToSavedState(UWorld* World);
-	void ResetToAnimatedState(UWorld* World);
-
-	void BackupSpawnablePilotData();
-	void RestoreSpawnablePilotData();
+	UE_API void ActivateSequencerEditorMode();
+	UE_API void DeactivateSequencerEditorMode();
+	UE_API void AddLevelViewportMenuExtender();
+	UE_API void BindDetailHandler(const FLevelEditorSequencerIntegrationOptions& Options);
+	UE_API void ActivateDetailHandler();
+	UE_API void AttachOutlinerColumn();
+	UE_API void DetachOutlinerColumn();
+	UE_API void ActivateRealtimeViewports();
+	UE_API void RestoreRealtimeViewports();
+	UE_API void RestoreToSavedState(UWorld* World);
+	UE_API void ResetToAnimatedState(UWorld* World);
+
+	UE_API void BackupSpawnablePilotData();
+	UE_API void RestoreSpawnablePilotData();
 
 	struct FSequencerAndOptions
 	{
@@ -242,16 +244,16 @@ private:
 
 public:
 	
-	TSharedRef< ISceneOutlinerColumn > CreateSequencerInfoColumn( ISceneOutliner& SceneOutliner ) const;
-	TSharedRef< ISceneOutlinerColumn > CreateSequencerSpawnableColumn( ISceneOutliner& SceneOutliner ) const;
+	UE_API TSharedRef< ISceneOutlinerColumn > CreateSequencerInfoColumn( ISceneOutliner& SceneOutliner ) const;
+	UE_API TSharedRef< ISceneOutlinerColumn > CreateSequencerSpawnableColumn( ISceneOutliner& SceneOutliner ) const;
 
 private:
 
-	void IterateAllSequencers(TFunctionRef<void(FSequencer&, const FLevelEditorSequencerIntegrationOptions& Options)>) const;
-	void UpdateDetails(bool bForceRefresh = false);
+	UE_API void IterateAllSequencers(TFunctionRef<void(FSequencer&, const FLevelEditorSequencerIntegrationOptions& Options)>) const;
+	UE_API void UpdateDetails(bool bForceRefresh = false);
 
-	FLevelEditorSequencerIntegration();
-	~FLevelEditorSequencerIntegration();
+	UE_API FLevelEditorSequencerIntegration();
+	UE_API ~FLevelEditorSequencerIntegration();
 
 private:
 	FAcquiredResources AcquiredResources;
@@ -267,3 +269,5 @@ private:
 	FLevelEditorSequencerUpdateGizmoTickFunction UpdateGizmoTickFunction;
 
 };
+
+#undef UE_API
diff --git a/Engine/Source/Editor/Sequencer/Public/MVVM/CurveEditorExtension.h b/Engine/Source/Editor/Sequencer/Public/MVVM/CurveEditorExtension.h
--- a/Engine/Source/Editor/Sequencer/Public/MVVM/CurveEditorExtension.h
+++ b/Engine/Source/Editor/Sequencer/Public/MVVM/CurveEditorExtension.h
@@ -11,6 +11,8 @@
 #include "Tree/ICurveEditorTreeItem.h"
 #include "STemporarilyFocusedSpinBox.h"
 
+#define UE_API SEQUENCER_API
+
 class ISequencer;
 class IPropertyTypeCustomization;
 class SCurveEditorPanel;
@@ -25,28 +27,28 @@ namespace Sequencer
 
 class FSequencerEditorViewModel;
 
-class SEQUENCER_API FCurveEditorExtension : public IDynamicExtension
+class FCurveEditorExtension : public IDynamicExtension
 {
 public:
 
-	UE_SEQUENCER_DECLARE_VIEW_MODEL_TYPE_ID(FCurveEditorExtension)
+	UE_SEQUENCER_DECLARE_VIEW_MODEL_TYPE_ID_API(UE_API, FCurveEditorExtension)
 
-	FCurveEditorExtension();
+	UE_API FCurveEditorExtension();
 
-	virtual void OnCreated(TSharedRef<FViewModel> InWeakOwner) override;
+	UE_API virtual void OnCreated(TSharedRef<FViewModel> InWeakOwner) override;
 
 	/** Creates the curve editor view-model and widget */
-	void CreateCurveEditor(const FTimeSliderArgs& TimeSliderArgs);
+	UE_API void CreateCurveEditor(const FTimeSliderArgs& TimeSliderArgs);
 
 	/** Gets the curve editor view-model */
 	TSharedPtr<FCurveEditor> GetCurveEditor() const { return CurveEditorModel; }
 
 	/** Opens the curve editor */
-	void OpenCurveEditor();
+	UE_API void OpenCurveEditor();
 	/** Returns whether the curve editor is open */
-	bool IsCurveEditorOpen() const;
+	UE_API bool IsCurveEditorOpen() const;
 	/** Closes the curve editor */
-	void CloseCurveEditor();
+	UE_API void CloseCurveEditor();
 
 	/** Curve editor tree widget */
 	TSharedPtr<SCurveEditorTree> GetCurveEditorTreeView() const { return CurveEditorTreeView; }
@@ -54,23 +56,23 @@ public:
 	/**
 	 * Synchronize curve editor selection with sequencer outliner selection on the next update.
 	 */
-	void RequestSyncSelection();
+	UE_API void RequestSyncSelection();
 
 
 public:
 
-	static const FName CurveEditorTabName;
+	static UE_API const FName CurveEditorTabName;
 
 private:
 
 	/** Get the default key attributes to apply to newly created keys on the curve editor */
-	FKeyAttributes GetDefaultKeyAttributes() const;
+	UE_API FKeyAttributes GetDefaultKeyAttributes() const;
 
 	/** Syncs the selection of the curve editor with that of the sequencer */
-	void SyncSelection();
+	UE_API void SyncSelection();
 
 	/** Executed when the filter class has changed */
-	void FilterClassChanged();
+	UE_API void FilterClassChanged();
 
 private:
 
@@ -98,3 +100,4 @@ private:
 } // namespace Sequencer
 } // namespace UE
 
+#undef UE_API
diff --git a/Engine/Source/Editor/Sequencer/Public/MVVM/CurveEditorIntegrationExtension.h b/Engine/Source/Editor/Sequencer/Public/MVVM/CurveEditorIntegrationExtension.h
--- a/Engine/Source/Editor/Sequencer/Public/MVVM/CurveEditorIntegrationExtension.h
+++ b/Engine/Source/Editor/Sequencer/Public/MVVM/CurveEditorIntegrationExtension.h
@@ -9,6 +9,8 @@
 #include "CurveEditorTypes.h"
 #include "Tree/ICurveEditorTreeItem.h"
 
+#define UE_API SEQUENCER_API
+
 namespace UE
 {
 namespace Sequencer
@@ -28,37 +30,37 @@ class FCurveEditorExtension;
  * - Outliner items implementing ICurveEditorTreeItemExtension (or its default shim) if they 
  *   want to show up in the curve editor.
  */
-class SEQUENCER_API FCurveEditorIntegrationExtension : public IDynamicExtension
+class FCurveEditorIntegrationExtension : public IDynamicExtension
 {
 public:
 
-	UE_SEQUENCER_DECLARE_VIEW_MODEL_TYPE_ID(FCurveEditorIntegrationExtension)
+	UE_SEQUENCER_DECLARE_VIEW_MODEL_TYPE_ID_API(UE_API, FCurveEditorIntegrationExtension)
 
-	FCurveEditorIntegrationExtension();
+	UE_API FCurveEditorIntegrationExtension();
 
-	virtual void OnCreated(TSharedRef<FViewModel> InWeakOwner) override;
+	UE_API virtual void OnCreated(TSharedRef<FViewModel> InWeakOwner) override;
 
 	/** 
 	 * Keeps the curve editor items up-to-date with the sequencer outliner by adding/removing 
 	 * entries as needed.
 	 */
-	void UpdateCurveEditor();
+	UE_API void UpdateCurveEditor();
 
 	/**
 	 * Clears the curve editor of all contents.
 	 */
-	void ResetCurveEditor();
+	UE_API void ResetCurveEditor();
 
 private:
 
 	/** Update curve editor items when sequencer outliner items change */
-	void OnHierarchyChanged();
+	UE_API void OnHierarchyChanged();
 
 	/** Adds the given view-model to the curve editor */
-	FCurveEditorTreeItemID AddToCurveEditor(TViewModelPtr<ICurveEditorTreeItemExtension> InViewModel, TSharedPtr<FCurveEditor> InCurveEditor);
+	UE_API FCurveEditorTreeItemID AddToCurveEditor(TViewModelPtr<ICurveEditorTreeItemExtension> InViewModel, TSharedPtr<FCurveEditor> InCurveEditor);
 
 	/** Finds the curve editor extension on the top-level sequencer editor view-model */
-	FCurveEditorExtension* GetCurveEditorExtension();
+	UE_API FCurveEditorExtension* GetCurveEditorExtension();
 
 private:
 
@@ -69,3 +71,4 @@ private:
 } // namespace Sequencer
 } // namespace UE
 
+#undef UE_API
diff --git a/Engine/Source/Editor/Sequencer/Public/MVVM/Extensions/IBindingLifetimeExtension.h b/Engine/Source/Editor/Sequencer/Public/MVVM/Extensions/IBindingLifetimeExtension.h
--- a/Engine/Source/Editor/Sequencer/Public/MVVM/Extensions/IBindingLifetimeExtension.h
+++ b/Engine/Source/Editor/Sequencer/Public/MVVM/Extensions/IBindingLifetimeExtension.h
@@ -5,6 +5,8 @@
 #include "MVVM/ViewModelTypeID.h"
 #include "Math/Range.h"
 
+#define UE_API SEQUENCER_API
+
 struct FGuid;
 
 namespace UE
@@ -12,11 +14,11 @@ namespace UE
 	namespace Sequencer
 	{
 
-		class SEQUENCER_API IBindingLifetimeExtension
+		class IBindingLifetimeExtension
 		{
 		public:
 
-			UE_SEQUENCER_DECLARE_VIEW_MODEL_TYPE_ID(IBindingLifetimeExtension)
+			UE_SEQUENCER_DECLARE_VIEW_MODEL_TYPE_ID_API(UE_API, IBindingLifetimeExtension)
 
 			virtual ~IBindingLifetimeExtension() {}
 
@@ -26,3 +28,4 @@ namespace UE
 	} // namespace Sequencer
 } // namespace UE
 
+#undef UE_API
diff --git a/Engine/Source/Editor/Sequencer/Public/MVVM/Extensions/IConditionableExtension.h b/Engine/Source/Editor/Sequencer/Public/MVVM/Extensions/IConditionableExtension.h
--- a/Engine/Source/Editor/Sequencer/Public/MVVM/Extensions/IConditionableExtension.h
+++ b/Engine/Source/Editor/Sequencer/Public/MVVM/Extensions/IConditionableExtension.h
@@ -6,6 +6,8 @@
 #include "MVVM/ViewModelTypeID.h"
 #include "MVVM/Extensions/HierarchicalCacheExtension.h"
 
+#define UE_API SEQUENCER_API
+
 class UMovieSceneCondition;
 
 namespace UE
@@ -23,11 +25,11 @@ namespace UE
 		/**
 		 * An extension for models that can have conditions
 		 */
-		class SEQUENCER_API IConditionableExtension
+		class IConditionableExtension
 		{
 		public:
 
-			UE_SEQUENCER_DECLARE_VIEW_MODEL_TYPE_ID(IConditionableExtension)
+			UE_SEQUENCER_DECLARE_VIEW_MODEL_TYPE_ID_API(UE_API, IConditionableExtension)
 
 			virtual ~IConditionableExtension() {}
 
@@ -57,19 +59,20 @@ namespace UE
 		};
 		ENUM_CLASS_FLAGS(ECachedConditionState)
 
-		class SEQUENCER_API FConditionStateCacheExtension
+		class FConditionStateCacheExtension
 			: public TFlagStateCacheExtension<ECachedConditionState>
 		{
 		public:
 
-			UE_SEQUENCER_DECLARE_VIEW_MODEL_TYPE_ID(FConditionStateCacheExtension);
+			UE_SEQUENCER_DECLARE_VIEW_MODEL_TYPE_ID_API(UE_API, FConditionStateCacheExtension);
 
 		private:
 
-			ECachedConditionState ComputeFlagsForModel(const FViewModelPtr& ViewModel) override;
-			void PostComputeChildrenFlags(const FViewModelPtr& ViewModel, ECachedConditionState& OutThisModelFlags, ECachedConditionState& OutPropagateToParentFlags) override;
+			UE_API ECachedConditionState ComputeFlagsForModel(const FViewModelPtr& ViewModel) override;
+			UE_API void PostComputeChildrenFlags(const FViewModelPtr& ViewModel, ECachedConditionState& OutThisModelFlags, ECachedConditionState& OutPropagateToParentFlags) override;
 		};
 
 	} // namespace Sequencer
 } // namespace UE
 
+#undef UE_API
diff --git a/Engine/Source/Editor/Sequencer/Public/MVVM/Extensions/ICurveEditorTreeItemExtension.h b/Engine/Source/Editor/Sequencer/Public/MVVM/Extensions/ICurveEditorTreeItemExtension.h
--- a/Engine/Source/Editor/Sequencer/Public/MVVM/Extensions/ICurveEditorTreeItemExtension.h
+++ b/Engine/Source/Editor/Sequencer/Public/MVVM/Extensions/ICurveEditorTreeItemExtension.h
@@ -6,6 +6,8 @@
 #include "CurveEditorTypes.h"
 #include "Tree/ICurveEditorTreeItem.h"
 
+#define UE_API SEQUENCER_API
+
 struct FGuid;
 class FCurveEditor;
 
@@ -17,11 +19,11 @@ namespace Sequencer
 /**
  * Extension interface for view-models that want to show up in the curve editor.
  */
-class SEQUENCER_API ICurveEditorTreeItemExtension
+class ICurveEditorTreeItemExtension
 {
 public:
 
-	UE_SEQUENCER_DECLARE_VIEW_MODEL_TYPE_ID(ICurveEditorTreeItemExtension)
+	UE_SEQUENCER_DECLARE_VIEW_MODEL_TYPE_ID_API(UE_API, ICurveEditorTreeItemExtension)
 
 	virtual ~ICurveEditorTreeItemExtension(){}
 
@@ -46,7 +48,7 @@ public:
 /**
  * Default partial implementation of ICurveEditorTreeItemExtension
  */
-class SEQUENCER_API FCurveEditorTreeItemExtensionShim : public ICurveEditorTreeItemExtension
+class FCurveEditorTreeItemExtensionShim : public ICurveEditorTreeItemExtension
 {
 public:
 
@@ -62,3 +64,4 @@ private:
 } // namespace Sequencer
 } // namespace UE
 
+#undef UE_API
diff --git a/Engine/Source/Editor/Sequencer/Public/MVVM/Extensions/IObjectBindingExtension.h b/Engine/Source/Editor/Sequencer/Public/MVVM/Extensions/IObjectBindingExtension.h
--- a/Engine/Source/Editor/Sequencer/Public/MVVM/Extensions/IObjectBindingExtension.h
+++ b/Engine/Source/Editor/Sequencer/Public/MVVM/Extensions/IObjectBindingExtension.h
@@ -4,6 +4,8 @@
 
 #include "MVVM/ViewModelTypeID.h"
 
+#define UE_API SEQUENCER_API
+
 struct FGuid;
 
 namespace UE
@@ -11,11 +13,11 @@ namespace UE
 namespace Sequencer
 {
 
-class SEQUENCER_API IObjectBindingExtension
+class IObjectBindingExtension
 {
 public:
 
-	UE_SEQUENCER_DECLARE_VIEW_MODEL_TYPE_ID(IObjectBindingExtension)
+	UE_SEQUENCER_DECLARE_VIEW_MODEL_TYPE_ID_API(UE_API, IObjectBindingExtension)
 
 	virtual ~IObjectBindingExtension(){}
 
@@ -25,3 +27,4 @@ public:
 } // namespace Sequencer
 } // namespace UE
 
+#undef UE_API
diff --git a/Engine/Source/Editor/Sequencer/Public/MVVM/Extensions/ISnappableExtension.h b/Engine/Source/Editor/Sequencer/Public/MVVM/Extensions/ISnappableExtension.h
--- a/Engine/Source/Editor/Sequencer/Public/MVVM/Extensions/ISnappableExtension.h
+++ b/Engine/Source/Editor/Sequencer/Public/MVVM/Extensions/ISnappableExtension.h
@@ -7,6 +7,8 @@
 #include "Curves/KeyHandle.h"
 #include "Templates/MemoryOps.h"
 
+#define UE_API SEQUENCER_API
+
 struct FGuid;
 
 class IKeyArea;
@@ -18,7 +20,7 @@ namespace Sequencer
 {
 
 /** Structure defining a point to snap to in the sequencer */
-struct SEQUENCER_API FSnapPoint
+struct FSnapPoint
 {
 	enum ESnapType { Key, SectionBounds, CustomSection, PlaybackRange, CurrentTime, InOutRange, Mark };
 
@@ -55,7 +57,7 @@ private:
 	{}
 };
 
-class SEQUENCER_API ISnapField
+class ISnapField
 {
 public:
 	virtual ~ISnapField() {}
@@ -64,7 +66,7 @@ public:
 };
 
 /** Interface that defines how to construct an FSequencerSnapField */
-struct SEQUENCER_API ISnapCandidate
+struct ISnapCandidate
 {
 	enum class EDefaultApplicability
 	{
@@ -103,11 +105,11 @@ struct SEQUENCER_API ISnapCandidate
 	virtual bool AreSectionCustomSnapsApplicable(UMovieSceneSection* Section) const { return true; }
 };
 
-class SEQUENCER_API ISnappableExtension
+class ISnappableExtension
 {
 public:
 
-	UE_SEQUENCER_DECLARE_VIEW_MODEL_TYPE_ID(ISnappableExtension)
+	UE_SEQUENCER_DECLARE_VIEW_MODEL_TYPE_ID_API(UE_API, ISnappableExtension)
 
 	virtual ~ISnappableExtension(){}
 
@@ -119,3 +121,4 @@ ENUM_CLASS_FLAGS(ISnapCandidate::ESnapCapabilities);
 } // namespace Sequencer
 } // namespace UE
 
+#undef UE_API
diff --git a/Engine/Source/Editor/Sequencer/Public/MVVM/Extensions/ITrackExtension.h b/Engine/Source/Editor/Sequencer/Public/MVVM/Extensions/ITrackExtension.h
--- a/Engine/Source/Editor/Sequencer/Public/MVVM/Extensions/ITrackExtension.h
+++ b/Engine/Source/Editor/Sequencer/Public/MVVM/Extensions/ITrackExtension.h
@@ -8,6 +8,8 @@
 #include "MVVM/ViewModels/ViewModelHierarchy.h"
 #include "MVVM/ViewModels/ViewModelIterators.h"
 
+#define UE_API SEQUENCER_API
+
 class UMovieSceneSection;
 class UMovieSceneTrack;
 class ISequencerTrackEditor;
@@ -17,11 +19,11 @@ namespace UE::Sequencer
 
 struct FViewModelChildren;
 
-class SEQUENCER_API ITrackExtension
+class ITrackExtension
 {
 public:
 
-	UE_SEQUENCER_DECLARE_VIEW_MODEL_TYPE_ID(ITrackExtension)
+	UE_SEQUENCER_DECLARE_VIEW_MODEL_TYPE_ID_API(UE_API, ITrackExtension)
 
 	virtual ~ITrackExtension(){}
 
@@ -48,3 +50,4 @@ public:
 
 } // namespace UE::Sequencer
 
+#undef UE_API
diff --git a/Engine/Source/Editor/Sequencer/Public/MVVM/ObjectBindingModelStorageExtension.h b/Engine/Source/Editor/Sequencer/Public/MVVM/ObjectBindingModelStorageExtension.h
--- a/Engine/Source/Editor/Sequencer/Public/MVVM/ObjectBindingModelStorageExtension.h
+++ b/Engine/Source/Editor/Sequencer/Public/MVVM/ObjectBindingModelStorageExtension.h
@@ -10,6 +10,8 @@
 #include "Misc/Guid.h"
 #include "Templates/SharedPointer.h"
 
+#define UE_API SEQUENCER_API
+
 class UMovieSceneTrack;
 namespace UE::MovieScene { class ISequenceDataEventHandler; }
 struct FMovieSceneBinding;
@@ -25,48 +27,48 @@ class FSequenceModel;
 class FViewModel;
 struct FViewModelChildren;
 
-class SEQUENCER_API FObjectBindingModelStorageExtension
+class FObjectBindingModelStorageExtension
 	: public IDynamicExtension
 	, private UE::MovieScene::TIntrusiveEventHandler<UE::MovieScene::ISequenceDataEventHandler>
 {
 public:
 
-	UE_SEQUENCER_DECLARE_VIEW_MODEL_TYPE_ID(FObjectBindingModelStorageExtension);
+	UE_SEQUENCER_DECLARE_VIEW_MODEL_TYPE_ID_API(UE_API, FObjectBindingModelStorageExtension);
 
-	FObjectBindingModelStorageExtension();
+	UE_API FObjectBindingModelStorageExtension();
 
 	/**
 	 * Implementation function for creating a new model for a binding from its ID. May return a placeholder if the binding does not yet exist.
 	 * @param Binding The binding to create a model for
 	 */
-	TSharedPtr<FViewModel> GetOrCreateModelForBinding(const FGuid& Binding);
+	UE_API TSharedPtr<FViewModel> GetOrCreateModelForBinding(const FGuid& Binding);
 
 	/**
 	 * Implementation function for creating a new model for a binding
 	 * @param Binding The binding to create a model for
 	 */
-	TSharedPtr<FViewModel> GetOrCreateModelForBinding(const FMovieSceneBinding& Binding);
+	UE_API TSharedPtr<FViewModel> GetOrCreateModelForBinding(const FMovieSceneBinding& Binding);
 
-	TSharedPtr<FObjectBindingModel> FindModelForObjectBinding(const FGuid& InObjectBindingID) const;
+	UE_API TSharedPtr<FObjectBindingModel> FindModelForObjectBinding(const FGuid& InObjectBindingID) const;
 
-	virtual void OnCreated(TSharedRef<FViewModel> InWeakOwner) override;
-	virtual void OnReinitialize() override;
+	UE_API virtual void OnCreated(TSharedRef<FViewModel> InWeakOwner) override;
+	UE_API virtual void OnReinitialize() override;
 
 private:
 
-	void OnBindingAdded(const FMovieSceneBinding& Binding) override;
-	void OnBindingRemoved(const FGuid& ObjectBindingID) override;
-	void OnTrackAddedToBinding(UMovieSceneTrack* Track, const FGuid& Binding) override;
-	void OnTrackRemovedFromBinding(UMovieSceneTrack* Track, const FGuid& Binding) override;
-	void OnBindingParentChanged(const FGuid& Binding, const FGuid& NewParent) override;
+	UE_API void OnBindingAdded(const FMovieSceneBinding& Binding) override;
+	UE_API void OnBindingRemoved(const FGuid& ObjectBindingID) override;
+	UE_API void OnTrackAddedToBinding(UMovieSceneTrack* Track, const FGuid& Binding) override;
+	UE_API void OnTrackRemovedFromBinding(UMovieSceneTrack* Track, const FGuid& Binding) override;
+	UE_API void OnBindingParentChanged(const FGuid& Binding, const FGuid& NewParent) override;
 
 
-	TSharedPtr<FObjectBindingModel> CreateModelForObjectBinding(const FMovieSceneBinding& Binding);
+	UE_API TSharedPtr<FObjectBindingModel> CreateModelForObjectBinding(const FMovieSceneBinding& Binding);
 
-	TSharedPtr<FViewModel> CreatePlaceholderForObjectBinding(const FGuid& ObjectID);
-	TSharedPtr<FViewModel> FindPlaceholderForObjectBinding(const FGuid& InObjectBindingID) const;
+	UE_API TSharedPtr<FViewModel> CreatePlaceholderForObjectBinding(const FGuid& ObjectID);
+	UE_API TSharedPtr<FViewModel> FindPlaceholderForObjectBinding(const FGuid& InObjectBindingID) const;
 
-	void Compact();
+	UE_API void Compact();
 
 private:
 
@@ -79,3 +81,4 @@ private:
 } // namespace Sequencer
 } // namespace UE
 
+#undef UE_API
diff --git a/Engine/Source/Editor/Sequencer/Public/MVVM/SectionModelStorageExtension.h b/Engine/Source/Editor/Sequencer/Public/MVVM/SectionModelStorageExtension.h
--- a/Engine/Source/Editor/Sequencer/Public/MVVM/SectionModelStorageExtension.h
+++ b/Engine/Source/Editor/Sequencer/Public/MVVM/SectionModelStorageExtension.h
@@ -8,6 +8,8 @@
 #include "Templates/SharedPointer.h"
 #include "UObject/ObjectKey.h"
 
+#define UE_API SEQUENCER_API
+
 class ISequencerSection;
 class UMovieSceneSection;
 
@@ -19,20 +21,20 @@ namespace Sequencer
 class FSectionModel;
 class FSequenceModel;
 
-class SEQUENCER_API FSectionModelStorageExtension
+class FSectionModelStorageExtension
 	: public IDynamicExtension
 {
 public:
 
-	UE_SEQUENCER_DECLARE_VIEW_MODEL_TYPE_ID(FSectionModelStorageExtension);
+	UE_SEQUENCER_DECLARE_VIEW_MODEL_TYPE_ID_API(UE_API, FSectionModelStorageExtension);
 
-	FSectionModelStorageExtension();
+	UE_API FSectionModelStorageExtension();
 
-	TSharedPtr<FSectionModel> CreateModelForSection(UMovieSceneSection* InSection, TSharedRef<ISequencerSection> SectionInterface);
+	UE_API TSharedPtr<FSectionModel> CreateModelForSection(UMovieSceneSection* InSection, TSharedRef<ISequencerSection> SectionInterface);
 
-	TSharedPtr<FSectionModel> FindModelForSection(const UMovieSceneSection* InSection) const;
+	UE_API TSharedPtr<FSectionModel> FindModelForSection(const UMovieSceneSection* InSection) const;
 
-	virtual void OnReinitialize() override;
+	UE_API virtual void OnReinitialize() override;
 
 private:
 
@@ -42,3 +44,4 @@ private:
 } // namespace Sequencer
 } // namespace UE
 
+#undef UE_API
diff --git a/Engine/Source/Editor/Sequencer/Public/MVVM/TrackRowModelStorageExtension.h b/Engine/Source/Editor/Sequencer/Public/MVVM/TrackRowModelStorageExtension.h
--- a/Engine/Source/Editor/Sequencer/Public/MVVM/TrackRowModelStorageExtension.h
+++ b/Engine/Source/Editor/Sequencer/Public/MVVM/TrackRowModelStorageExtension.h
@@ -10,6 +10,8 @@
 #include "Templates/Tuple.h"
 #include "UObject/ObjectKey.h"
 
+#define UE_API SEQUENCER_API
+
 class UMovieSceneTrack;
 
 namespace UE
@@ -21,21 +23,21 @@ class FSequenceModel;
 class FTrackRowModel;
 class FViewModel;
 
-class SEQUENCER_API FTrackRowModelStorageExtension
+class FTrackRowModelStorageExtension
 	: public IDynamicExtension
 {
 public:
 
-	UE_SEQUENCER_DECLARE_VIEW_MODEL_TYPE_ID(FTrackRowModelStorageExtension)
+	UE_SEQUENCER_DECLARE_VIEW_MODEL_TYPE_ID_API(UE_API, FTrackRowModelStorageExtension)
 
-	FTrackRowModelStorageExtension();
+	UE_API FTrackRowModelStorageExtension();
 
-	virtual void OnCreated(TSharedRef<FViewModel> InWeakOwner) override;
-	virtual void OnReinitialize() override;
+	UE_API virtual void OnCreated(TSharedRef<FViewModel> InWeakOwner) override;
+	UE_API virtual void OnReinitialize() override;
 
-	TSharedPtr<FTrackRowModel> CreateModelForTrackRow(UMovieSceneTrack* InTrack, int32 InRowIndex, TSharedPtr<FViewModel> DesiredParent = nullptr);
+	UE_API TSharedPtr<FTrackRowModel> CreateModelForTrackRow(UMovieSceneTrack* InTrack, int32 InRowIndex, TSharedPtr<FViewModel> DesiredParent = nullptr);
 
-	TSharedPtr<FTrackRowModel> FindModelForTrackRow(UMovieSceneTrack* InTrack, int32 InRowIndex) const;
+	UE_API TSharedPtr<FTrackRowModel> FindModelForTrackRow(UMovieSceneTrack* InTrack, int32 InRowIndex) const;
 
 private:
 
@@ -49,3 +51,4 @@ private:
 } // namespace Sequencer
 } // namespace UE
 
+#undef UE_API
diff --git a/Engine/Source/Editor/Sequencer/Public/MVVM/ViewModels/BindingLifetimeOverlayModel.h b/Engine/Source/Editor/Sequencer/Public/MVVM/ViewModels/BindingLifetimeOverlayModel.h
--- a/Engine/Source/Editor/Sequencer/Public/MVVM/ViewModels/BindingLifetimeOverlayModel.h
+++ b/Engine/Source/Editor/Sequencer/Public/MVVM/ViewModels/BindingLifetimeOverlayModel.h
@@ -12,6 +12,8 @@
 #include "Templates/SharedPointer.h"
 #include "SequencerEditorViewModel.h"
 
+#define UE_API SEQUENCER_API
+
 namespace UE::Sequencer { template <typename T> struct TAutoRegisterViewModelTypeID; }
 struct FFrameNumber;
 
@@ -19,26 +21,26 @@ namespace UE
 {
 	namespace Sequencer
 	{
-		class SEQUENCER_API FBindingLifetimeOverlayModel
+		class FBindingLifetimeOverlayModel
 			: public FViewModel
 			, public FLinkedOutlinerExtension
 			, public ITrackLaneExtension
 		{
 		public:
 
-			UE_SEQUENCER_DECLARE_CASTABLE(FBindingLifetimeOverlayModel, FViewModel
+			UE_SEQUENCER_DECLARE_CASTABLE_API(UE_API, FBindingLifetimeOverlayModel, FViewModel
 				, FLinkedOutlinerExtension
 				, ITrackLaneExtension
 			);
 
-			FBindingLifetimeOverlayModel(TWeakPtr<FViewModel> LayerRoot, TWeakPtr<FSequencerEditorViewModel> InEditorViewModel, TViewModelPtr<IBindingLifetimeExtension> InBindingLifetimeTrack);
-			~FBindingLifetimeOverlayModel();
+			UE_API FBindingLifetimeOverlayModel(TWeakPtr<FViewModel> LayerRoot, TWeakPtr<FSequencerEditorViewModel> InEditorViewModel, TViewModelPtr<IBindingLifetimeExtension> InBindingLifetimeTrack);
+			UE_API ~FBindingLifetimeOverlayModel();
 
 			/*~ ITrackLaneExtension Interface */
-			TSharedPtr<ITrackLaneWidget> CreateTrackLaneView(const FCreateTrackLaneViewParams& InParams) override;
-			FTrackLaneVirtualAlignment ArrangeVirtualTrackLaneView() const override;
+			UE_API TSharedPtr<ITrackLaneWidget> CreateTrackLaneView(const FCreateTrackLaneViewParams& InParams) override;
+			UE_API FTrackLaneVirtualAlignment ArrangeVirtualTrackLaneView() const override;
 
-			const TArray<FFrameNumberRange>& GetInverseLifetimeRange() const;
+			UE_API const TArray<FFrameNumberRange>& GetInverseLifetimeRange() const;
 
 		private:
 
@@ -49,3 +51,4 @@ namespace UE
 	} // namespace Sequencer
 } // namespace UE
 
+#undef UE_API
diff --git a/Engine/Source/Editor/Sequencer/Public/MVVM/ViewModels/CategoryModel.h b/Engine/Source/Editor/Sequencer/Public/MVVM/ViewModels/CategoryModel.h
--- a/Engine/Source/Editor/Sequencer/Public/MVVM/ViewModels/CategoryModel.h
+++ b/Engine/Source/Editor/Sequencer/Public/MVVM/ViewModels/CategoryModel.h
@@ -23,6 +23,8 @@
 #include "Templates/SharedPointer.h"
 #include "UObject/NameTypes.h"
 
+#define UE_API SEQUENCER_API
+
 class FSequencerSectionCategoryNode;
 class SWidget;
 class UMovieSceneSection;
@@ -38,7 +40,7 @@ namespace Sequencer
  * For instance, this represents the "Location" category of a single transform section, which would
  * contain the X, Y, and Z translation channels.
  */
-class SEQUENCER_API FCategoryModel
+class FCategoryModel
 	: public FViewModel
 	, public FLinkedOutlinerExtension
 	, public FGeometryExtensionShim
@@ -46,26 +48,26 @@ class SEQUENCER_API FCategoryModel
 	, public FLinkedOutlinerComputedSizingShim
 {
 public:
-	UE_SEQUENCER_DECLARE_CASTABLE(FCategoryModel, FViewModel, FLinkedOutlinerExtension, IGeometryExtension, ITrackLaneExtension);
+	UE_SEQUENCER_DECLARE_CASTABLE_API(UE_API, FCategoryModel, FViewModel, FLinkedOutlinerExtension, IGeometryExtension, ITrackLaneExtension);
 
-	explicit FCategoryModel(FName InCategoryName);
+	UE_API explicit FCategoryModel(FName InCategoryName);
 
 	/** Whether any of the channels within this category are animated.  */
-	bool IsAnimated() const;
+	UE_API bool IsAnimated() const;
 
 	/** Returns the category's name */
 	FName GetCategoryName() const { return CategoryName; }
 
 	/** Returns the desired sizing for the track area row */
-	virtual FOutlinerSizing GetDesiredSizing() const;
+	UE_API virtual FOutlinerSizing GetDesiredSizing() const;
 
 	/*~ ITrackLaneExtension */
-	TSharedPtr<ITrackLaneWidget> CreateTrackLaneView(const FCreateTrackLaneViewParams& InParams) override;
-	FTrackLaneVirtualAlignment ArrangeVirtualTrackLaneView() const override;
+	UE_API TSharedPtr<ITrackLaneWidget> CreateTrackLaneView(const FCreateTrackLaneViewParams& InParams) override;
+	UE_API FTrackLaneVirtualAlignment ArrangeVirtualTrackLaneView() const override;
 
 private:
 
-	FLinearColor GetKeyBarColor() const;
+	UE_API FLinearColor GetKeyBarColor() const;
 
 private:
 
@@ -78,7 +80,7 @@ private:
  * For instance, this represents the "Location" category entry in the Sequence outliner, which would contain
  * the X, Y, and Z translation channels of all the corresponding sections in the track area.
  */
-class SEQUENCER_API FCategoryGroupModel
+class FCategoryGroupModel
 	: public FOutlinerItemModel
 	, public ITrackAreaExtension
 	, public ICompoundOutlinerExtension
@@ -86,16 +88,16 @@ class SEQUENCER_API FCategoryGroupModel
 	, public IRecyclableExtension
 {
 public:
-	UE_SEQUENCER_DECLARE_CASTABLE(FCategoryGroupModel, FOutlinerItemModel, ITrackAreaExtension, ICompoundOutlinerExtension, IDeletableExtension, IRecyclableExtension);
+	UE_SEQUENCER_DECLARE_CASTABLE_API(UE_API, FCategoryGroupModel, FOutlinerItemModel, ITrackAreaExtension, ICompoundOutlinerExtension, IDeletableExtension, IRecyclableExtension);
 
-	explicit FCategoryGroupModel(FName InCategoryName, const FText& InDisplayText, FGetMovieSceneTooltipText InGetGroupTooltipTextDelegate);
+	UE_API explicit FCategoryGroupModel(FName InCategoryName, const FText& InDisplayText, FGetMovieSceneTooltipText InGetGroupTooltipTextDelegate);
 
-	~FCategoryGroupModel();
+	UE_API ~FCategoryGroupModel();
 
 	/**
 	 * @return Whether any of the channels within this category group have any keyframes on them
 	 */
-	bool IsAnimated() const;
+	UE_API bool IsAnimated() const;
 
 	FName GetCategoryName() const
 	{
@@ -107,35 +109,35 @@ public:
 		return DisplayText;
 	}
 
-	void AddCategory(TWeakViewModelPtr<FCategoryModel> InCategory);
+	UE_API void AddCategory(TWeakViewModelPtr<FCategoryModel> InCategory);
 
-	TArrayView<const TWeakViewModelPtr<FCategoryModel>> GetCategories() const;
+	UE_API TArrayView<const TWeakViewModelPtr<FCategoryModel>> GetCategories() const;
 
 public:
 
 	/*~ ICompoundOutlinerExtension */
-	FOutlinerSizing RecomputeSizing() override;
+	UE_API FOutlinerSizing RecomputeSizing() override;
 
 	/*~ IOutlinerExtension */
-	FOutlinerSizing GetOutlinerSizing() const override;
-	FText GetLabel() const override;
-	FSlateFontInfo GetLabelFont() const override;
-	FText GetLabelToolTipText() const override;
-	TSharedPtr<SWidget> CreateOutlinerViewForColumn(const FCreateOutlinerViewParams& InParams, const FName& InColumnName) override;
+	UE_API FOutlinerSizing GetOutlinerSizing() const override;
+	UE_API FText GetLabel() const override;
+	UE_API FSlateFontInfo GetLabelFont() const override;
+	UE_API FText GetLabelToolTipText() const override;
+	UE_API TSharedPtr<SWidget> CreateOutlinerViewForColumn(const FCreateOutlinerViewParams& InParams, const FName& InColumnName) override;
 
 	/*~ ITrackAreaExtension */
-	FTrackAreaParameters GetTrackAreaParameters() const override;
-	FViewModelVariantIterator GetTrackAreaModelList() const override;
+	UE_API FTrackAreaParameters GetTrackAreaParameters() const override;
+	UE_API FViewModelVariantIterator GetTrackAreaModelList() const override;
 
 	/*~ IDeletableExtension */
-	bool CanDelete(FText* OutErrorMessage) const override;
-	void Delete() override;
+	UE_API bool CanDelete(FText* OutErrorMessage) const override;
+	UE_API void Delete() override;
 
 	/*~ IRecyclableExtension */
-	void OnRecycle() override;
+	UE_API void OnRecycle() override;
 
 	/*~ FOutlinerItemModelMixin */
-	virtual void BuildSidebarMenu(FMenuBuilder& MenuBuilder) override;
+	UE_API virtual void BuildSidebarMenu(FMenuBuilder& MenuBuilder) override;
 
 private:
 
@@ -149,3 +151,4 @@ private:
 } // namespace Sequencer
 } // namespace UE
 
+#undef UE_API
diff --git a/Engine/Source/Editor/Sequencer/Public/MVVM/ViewModels/ChannelModel.h b/Engine/Source/Editor/Sequencer/Public/MVVM/ViewModels/ChannelModel.h
--- a/Engine/Source/Editor/Sequencer/Public/MVVM/ViewModels/ChannelModel.h
+++ b/Engine/Source/Editor/Sequencer/Public/MVVM/ViewModels/ChannelModel.h
@@ -28,6 +28,8 @@
 #include "Tree/ICurveEditorTreeItem.h"
 #include "UObject/NameTypes.h"
 
+#define UE_API SEQUENCER_API
+
 class FCurveModel;
 class FSequencerSectionKeyAreaNode;
 class IKeyArea;
@@ -52,7 +54,7 @@ class FChannelGroupOutlinerModel;
  * Model for a single channel inside a section.
  * For instance, this represents the "Location.X" channel of a single transform section.
  */
-class SEQUENCER_API FChannelModel
+class FChannelModel
 	: public FViewModel
 	, public FLinkedOutlinerExtension
 	, public FGeometryExtensionShim
@@ -63,24 +65,24 @@ class SEQUENCER_API FChannelModel
 {
 public:
 
-	UE_SEQUENCER_DECLARE_CASTABLE(FChannelModel, FViewModel, FLinkedOutlinerExtension, IGeometryExtension, ITrackLaneExtension, IRecyclableExtension, IKeyExtension);
+	UE_SEQUENCER_DECLARE_CASTABLE_API(UE_API, FChannelModel, FViewModel, FLinkedOutlinerExtension, IGeometryExtension, ITrackLaneExtension, IRecyclableExtension, IKeyExtension);
 
-	FChannelModel(FName InChannelName, TWeakPtr<ISequencerSection> InSection, FMovieSceneChannelHandle InChannel);
-	~FChannelModel();
+	UE_API FChannelModel(FName InChannelName, TWeakPtr<ISequencerSection> InSection, FMovieSceneChannelHandle InChannel);
+	UE_API ~FChannelModel();
 
-	void Initialize(TWeakPtr<ISequencerSection> InSection, FMovieSceneChannelHandle InChannel);
+	UE_API void Initialize(TWeakPtr<ISequencerSection> InSection, FMovieSceneChannelHandle InChannel);
 
 	/** Returns the section object that owns the associated channel */
-	UMovieSceneSection* GetSection() const;
+	UE_API UMovieSceneSection* GetSection() const;
 
 	/** Returns the object that owns the associated channel. May return the same as GetSection(). */
-	UObject* GetOwningObject() const;
+	UE_API UObject* GetOwningObject() const;
 
 	/** Returns the associated channel object */
-	FMovieSceneChannel* GetChannel() const;
+	UE_API FMovieSceneChannel* GetChannel() const;
 
 	/** Returns whether this channel has any keyframes on it */
-	bool IsAnimated() const;
+	UE_API bool IsAnimated() const;
 
 	/** Returns the channel's name */
 	FName GetChannelName() const { return ChannelName; }
@@ -89,22 +91,22 @@ public:
 	TSharedPtr<IKeyArea> GetKeyArea() const { return KeyArea; }
 
 	/** Returns the desired sizing for the track area row */
-	FOutlinerSizing GetDesiredSizing() const;
+	UE_API FOutlinerSizing GetDesiredSizing() const;
 
-	int32 CustomPaint(const FSequencerChannelPaintArgs& CustomPaintArgs, int32 LayerId) const;
+	UE_API int32 CustomPaint(const FSequencerChannelPaintArgs& CustomPaintArgs, int32 LayerId) const;
 
 	/*~ ITrackLaneExtension */
-	TSharedPtr<ITrackLaneWidget> CreateTrackLaneView(const FCreateTrackLaneViewParams& InParams) override;
-	FTrackLaneVirtualAlignment ArrangeVirtualTrackLaneView() const override;
+	UE_API TSharedPtr<ITrackLaneWidget> CreateTrackLaneView(const FCreateTrackLaneViewParams& InParams) override;
+	UE_API FTrackLaneVirtualAlignment ArrangeVirtualTrackLaneView() const override;
 
 	/*~ IRecyclableExtension */
-	void OnRecycle() override;
+	UE_API void OnRecycle() override;
 
 	/*~ IKeyExtension */
-	bool UpdateCachedKeys(TSharedPtr<FCachedKeys>& OutCachedKeys) const override;
-	bool GetFixedExtents(double& OutFixedMin, double& OutFixedMax) const override;
-	void DrawKeys(TArrayView<const FKeyHandle> InKeyHandles, TArrayView<FKeyDrawParams> OutKeyDrawParams) override;
-	TUniquePtr<FCurveModel> CreateCurveModel() override;
+	UE_API bool UpdateCachedKeys(TSharedPtr<FCachedKeys>& OutCachedKeys) const override;
+	UE_API bool GetFixedExtents(double& OutFixedMin, double& OutFixedMax) const override;
+	UE_API void DrawKeys(TArrayView<const FKeyHandle> InKeyHandles, TArrayView<FKeyDrawParams> OutKeyDrawParams) override;
+	UE_API TUniquePtr<FCurveModel> CreateCurveModel() override;
 
 	/*~ Begin virtual interface */
 	virtual void BuildContextMenu(FMenuBuilder& MenuBuilder, TViewModelPtr<FChannelGroupOutlinerModel> GroupOwner) {}
@@ -112,7 +114,7 @@ public:
 	/*~ End virtual interface */
 private:
 
-	FLinearColor GetKeyBarColor() const;
+	UE_API FLinearColor GetKeyBarColor() const;
 
 private:
 
@@ -127,22 +129,22 @@ private:
  * Model for the outliner entry associated with all sections' channels of a given common name.
  * For instance, this represents the "Location.X" entry in the Sequencer outliner.
  */
-class SEQUENCER_API FChannelGroupModel
+class FChannelGroupModel
 	: public FViewModel
 	, public ITrackAreaExtension
 	, public IRecyclableExtension
 {
 public:
 
-	UE_SEQUENCER_DECLARE_CASTABLE(FChannelGroupModel, FViewModel, ITrackAreaExtension, IRecyclableExtension);
+	UE_SEQUENCER_DECLARE_CASTABLE_API(UE_API, FChannelGroupModel, FViewModel, ITrackAreaExtension, IRecyclableExtension);
 
-	FChannelGroupModel(FName InChannelName, const FText& InDisplayText);
-	FChannelGroupModel(FName InChannelName, const FText& InDisplayText, const FText& InTooltipText);
-	FChannelGroupModel(FName InChannelName, const FText& InDisplayText, FGetMovieSceneTooltipText InGetTooltipTextDelegate);
-	~FChannelGroupModel();
+	UE_API FChannelGroupModel(FName InChannelName, const FText& InDisplayText);
+	UE_API FChannelGroupModel(FName InChannelName, const FText& InDisplayText, const FText& InTooltipText);
+	UE_API FChannelGroupModel(FName InChannelName, const FText& InDisplayText, FGetMovieSceneTooltipText InGetTooltipTextDelegate);
+	UE_API ~FChannelGroupModel();
 
 	/** Returns whether any of the channels within this group have any keyframes on them */
-	bool IsAnimated() const;
+	UE_API bool IsAnimated() const;
 
 	/** Returns the common name for all channels in this group */
 	FName GetChannelName() const { return ChannelName; }
@@ -151,57 +153,57 @@ public:
 	FText GetDisplayText() const { return DisplayText; }
 
 	/** Returns the tooltip for this group */
-	FText GetTooltipText() const;
+	UE_API FText GetTooltipText() const;
 
 	/** Gets all the channel models in this group */
-	TArrayView<const TWeakViewModelPtr<FChannelModel>> GetChannels() const;
+	UE_API TArrayView<const TWeakViewModelPtr<FChannelModel>> GetChannels() const;
 	/** Adds a channel model to this group */
-	void AddChannel(TWeakViewModelPtr<FChannelModel> InChannel);
+	UE_API void AddChannel(TWeakViewModelPtr<FChannelModel> InChannel);
 
 	/** Get the key area of the channel associated with the given section */
-	TSharedPtr<IKeyArea> GetKeyArea(TSharedPtr<FSectionModel> InOwnerSection) const;
+	UE_API TSharedPtr<IKeyArea> GetKeyArea(TSharedPtr<FSectionModel> InOwnerSection) const;
 
 	/** Get the key area of the channel associated with the given section */
-	TSharedPtr<IKeyArea> GetKeyArea(const UMovieSceneSection* InOwnerSection) const;
+	UE_API TSharedPtr<IKeyArea> GetKeyArea(const UMovieSceneSection* InOwnerSection) const;
 
 	/** Get the channel model at the given index in the list of channels */
-	TSharedPtr<FChannelModel> GetChannel(int32 Index) const;
+	UE_API TSharedPtr<FChannelModel> GetChannel(int32 Index) const;
 
 	/** Get the channel model of the channel associated with the given section */
-	TSharedPtr<FChannelModel> GetChannel(TSharedPtr<FSectionModel> InOwnerSection) const;
+	UE_API TSharedPtr<FChannelModel> GetChannel(TSharedPtr<FSectionModel> InOwnerSection) const;
 
 	/** Get the channel model of the channel associated with the given section */
-	TSharedPtr<FChannelModel> GetChannel(const UMovieSceneSection* InOwnerSection) const;
+	UE_API TSharedPtr<FChannelModel> GetChannel(const UMovieSceneSection* InOwnerSection) const;
 
 	/** Get the key areas of all channels */
-	TArray<TSharedRef<IKeyArea>> GetAllKeyAreas() const;
+	UE_API TArray<TSharedRef<IKeyArea>> GetAllKeyAreas() const;
 
 	/** Gets a serial number representing if the list of channels has changed */
-	uint32 GetChannelsSerialNumber() const;
+	UE_API uint32 GetChannelsSerialNumber() const;
 
 	/** Notifies this view-model that the list of channels has changed */
-	void OnChannelOverridesChanged();
+	UE_API void OnChannelOverridesChanged();
 
 public:
 
 	/*~ ITrackAreaExtension */
-	FTrackAreaParameters GetTrackAreaParameters() const override;
-	FViewModelVariantIterator GetTrackAreaModelList() const override;
+	UE_API FTrackAreaParameters GetTrackAreaParameters() const override;
+	UE_API FViewModelVariantIterator GetTrackAreaModelList() const override;
 
 	/*~ IRecyclableExtension */
-	void OnRecycle() override;
+	UE_API void OnRecycle() override;
 
-	void CreateCurveModels(TArray<TUniquePtr<FCurveModel>>& OutCurveModels);
-	bool HasCurves() const;
-	TOptional<FString> GetUniquePathName() const;
+	UE_API void CreateCurveModels(TArray<TUniquePtr<FCurveModel>>& OutCurveModels);
+	UE_API bool HasCurves() const;
+	UE_API TOptional<FString> GetUniquePathName() const;
 
-	void BuildChannelOverrideMenu(FMenuBuilder& MenuBuilder);
+	UE_API void BuildChannelOverrideMenu(FMenuBuilder& MenuBuilder);
 
-	void CleanupChannels();
+	UE_API void CleanupChannels();
 
 protected:
 
-	void UpdateMutability();
+	UE_API void UpdateMutability();
 
 protected:
 
@@ -217,48 +219,48 @@ protected:
  * Model for the outliner entry associated with all sections' channels of a given common name.
  * For instance, this represents the "Location.X" entry in the Sequencer outliner.
  */
-class SEQUENCER_API FChannelGroupOutlinerModel
+class FChannelGroupOutlinerModel
 	: public TOutlinerModelMixin<FChannelGroupModel>
 	, public ICompoundOutlinerExtension
 	, public IDeletableExtension
 {
 public:
 
-	UE_SEQUENCER_DECLARE_CASTABLE(FChannelGroupOutlinerModel, FChannelGroupModel, FOutlinerItemModelMixin, ICompoundOutlinerExtension, IDeletableExtension);
+	UE_SEQUENCER_DECLARE_CASTABLE_API(UE_API, FChannelGroupOutlinerModel, FChannelGroupModel, FOutlinerItemModelMixin, ICompoundOutlinerExtension, IDeletableExtension);
 
-	FChannelGroupOutlinerModel(FName InChannelName, const FText& InDisplayText, FGetMovieSceneTooltipText InGetTooltipTextDelegate);
-	~FChannelGroupOutlinerModel();
+	UE_API FChannelGroupOutlinerModel(FName InChannelName, const FText& InDisplayText, FGetMovieSceneTooltipText InGetTooltipTextDelegate);
+	UE_API ~FChannelGroupOutlinerModel();
 
 public:
 
 	/*~ ICompoundOutlinerExtension */
-	FOutlinerSizing RecomputeSizing() override;
+	UE_API FOutlinerSizing RecomputeSizing() override;
 
 	/*~ IOutlinerExtension */
-	FOutlinerSizing GetOutlinerSizing() const override;
-	FText GetLabel() const override;
-	FSlateFontInfo GetLabelFont() const override;
-	FText GetLabelToolTipText() const override;
-	TSharedPtr<SWidget> CreateOutlinerViewForColumn(const FCreateOutlinerViewParams& InParams, const FName& InColumnName) override;
+	UE_API FOutlinerSizing GetOutlinerSizing() const override;
+	UE_API FText GetLabel() const override;
+	UE_API FSlateFontInfo GetLabelFont() const override;
+	UE_API FText GetLabelToolTipText() const override;
+	UE_API TSharedPtr<SWidget> CreateOutlinerViewForColumn(const FCreateOutlinerViewParams& InParams, const FName& InColumnName) override;
 
 	/*~ ICurveEditorTreeItem */
-	void CreateCurveModels(TArray<TUniquePtr<FCurveModel>>& OutCurveModels) override;
+	UE_API void CreateCurveModels(TArray<TUniquePtr<FCurveModel>>& OutCurveModels) override;
 
 	/*~ IDeletableExtension */
-	bool CanDelete(FText* OutErrorMessage) const override;
-	void Delete() override;
+	UE_API bool CanDelete(FText* OutErrorMessage) const override;
+	UE_API void Delete() override;
 
 	/*~ ICurveEditorTreeItemExtension */
-	bool HasCurves() const override;
-	void BuildContextMenu(FMenuBuilder& MenuBuilder) override;
+	UE_API bool HasCurves() const override;
+	UE_API void BuildContextMenu(FMenuBuilder& MenuBuilder) override;
 	
-	void OnUpdated();
+	UE_API void OnUpdated();
 	
-	TOptional<FString> GetUniquePathName() const override;
+	UE_API TOptional<FString> GetUniquePathName() const override;
 
 private:
 
-	EVisibility GetKeyEditorVisibility() const;
+	UE_API EVisibility GetKeyEditorVisibility() const;
 	
 private:
 
@@ -270,20 +272,21 @@ private:
 /**
  * Utility class for building menus that add, edit, and remove channel overrides.
  */
-class SEQUENCER_API FChannelGroupOverrideHelper
+class FChannelGroupOverrideHelper
 {
 public:
 
-	static void BuildChannelOverrideMenu(FMenuBuilder& MenuBuilder, TSharedPtr<FSequenceModel> SequenceModel);
+	static UE_API void BuildChannelOverrideMenu(FMenuBuilder& MenuBuilder, TSharedPtr<FSequenceModel> SequenceModel);
 
 private:
 
-	static bool GetChannelsFromOutlinerSelection(TSharedPtr<FSequenceModel> SequenceModel, TArray<TViewModelPtr<FChannelModel>>& OutChannelModels);
-	static void BuildChannelOverrideMenu(FMenuBuilder& MenuBuilder, TSharedPtr<FSequenceModel> SequenceModel, UMovieSceneChannelOverrideContainer::FOverrideCandidates OverrideCandidates);
-	static void OverrideChannels(TSharedPtr<FSequenceModel> SequenceModel, TSubclassOf<UMovieSceneChannelOverrideContainer> OverrideClass);
-	static void RemoveChannelOverrides(TSharedPtr<FSequenceModel> SequenceModel);
+	static UE_API bool GetChannelsFromOutlinerSelection(TSharedPtr<FSequenceModel> SequenceModel, TArray<TViewModelPtr<FChannelModel>>& OutChannelModels);
+	static UE_API void BuildChannelOverrideMenu(FMenuBuilder& MenuBuilder, TSharedPtr<FSequenceModel> SequenceModel, UMovieSceneChannelOverrideContainer::FOverrideCandidates OverrideCandidates);
+	static UE_API void OverrideChannels(TSharedPtr<FSequenceModel> SequenceModel, TSubclassOf<UMovieSceneChannelOverrideContainer> OverrideClass);
+	static UE_API void RemoveChannelOverrides(TSharedPtr<FSequenceModel> SequenceModel);
 };
 
 } // namespace Sequencer
 } // namespace UE
 
+#undef UE_API
diff --git a/Engine/Source/Editor/Sequencer/Public/MVVM/ViewModels/LayerBarModel.h b/Engine/Source/Editor/Sequencer/Public/MVVM/ViewModels/LayerBarModel.h
--- a/Engine/Source/Editor/Sequencer/Public/MVVM/ViewModels/LayerBarModel.h
+++ b/Engine/Source/Editor/Sequencer/Public/MVVM/ViewModels/LayerBarModel.h
@@ -15,6 +15,8 @@
 #include "Math/Range.h"
 #include "Templates/SharedPointer.h"
 
+#define UE_API SEQUENCER_API
+
 namespace UE::Sequencer { class ILayerBarExtension; }
 namespace UE::Sequencer { template <typename T> struct TAutoRegisterViewModelTypeID; }
 struct FFrameNumber;
@@ -26,7 +28,7 @@ namespace Sequencer
 
 class FSequenceModel;
 
-class SEQUENCER_API FLayerBarModel
+class FLayerBarModel
 	: public FViewModel
 	, public FLinkedOutlinerExtension
 	, public ITrackLaneExtension
@@ -38,7 +40,7 @@ class SEQUENCER_API FLayerBarModel
 {
 public:
 
-	UE_SEQUENCER_DECLARE_CASTABLE(FLayerBarModel, FViewModel
+	UE_SEQUENCER_DECLARE_CASTABLE_API(UE_API, FLayerBarModel, FViewModel
 		, FLinkedOutlinerExtension
 		, ITrackLaneExtension
 		, ISelectableExtension
@@ -47,43 +49,44 @@ public:
 		, IStretchableExtension
 	);
 
-	FLayerBarModel(TWeakPtr<FViewModel> LayerRoot);
-	~FLayerBarModel();
+	UE_API FLayerBarModel(TWeakPtr<FViewModel> LayerRoot);
+	UE_API ~FLayerBarModel();
 
 	/*~ ITrackLaneExtension Interface */
-	TSharedPtr<ITrackLaneWidget> CreateTrackLaneView(const FCreateTrackLaneViewParams& InParams) override;
-	FTrackLaneVirtualAlignment ArrangeVirtualTrackLaneView() const override;
+	UE_API TSharedPtr<ITrackLaneWidget> CreateTrackLaneView(const FCreateTrackLaneViewParams& InParams) override;
+	UE_API FTrackLaneVirtualAlignment ArrangeVirtualTrackLaneView() const override;
 
 	/*~ ISelectableExtension Interface */
-	ESelectionIntent IsSelectable() const override;
+	UE_API ESelectionIntent IsSelectable() const override;
 
 	/*~ ISnappableExtension Interface */
-	void AddToSnapField(const ISnapCandidate& Candidate, ISnapField& SnapField) const override;
+	UE_API void AddToSnapField(const ISnapCandidate& Candidate, ISnapField& SnapField) const override;
 
 	/*~ IDraggableTrackAreaExtension Interface */
-	bool CanDrag() const override;
-	void OnBeginDrag(IDragOperation& DragOperation) override;
-	void OnEndDrag(IDragOperation& DragOperation) override;
+	UE_API bool CanDrag() const override;
+	UE_API void OnBeginDrag(IDragOperation& DragOperation) override;
+	UE_API void OnEndDrag(IDragOperation& DragOperation) override;
 
 	/*~ IStretchableExtension Interface */
-	void OnInitiateStretch(IStretchOperation& StretchOperation, EStretchConstraint Constraint, FStretchParameters* InOutGlobalParameters) override;
+	UE_API void OnInitiateStretch(IStretchOperation& StretchOperation, EStretchConstraint Constraint, FStretchParameters* InOutGlobalParameters) override;
 
 public:
 
-	TRange<FFrameNumber> ComputeRange() const;
+	UE_API TRange<FFrameNumber> ComputeRange() const;
 
-	void Offset(FFrameNumber Offset);
+	UE_API void Offset(FFrameNumber Offset);
 
 private:
 
 	/*~ FViewModel interface */
-	void OnConstruct() override;
-	void OnDestruct() override;
+	UE_API void OnConstruct() override;
+	UE_API void OnDestruct() override;
 
 	/*~ TViewModelExtensionCollection Interface */
-	void OnExtensionsDirtied() override;
+	UE_API void OnExtensionsDirtied() override;
 };
 
 } // namespace Sequencer
 } // namespace UE
 
+#undef UE_API
