Updated headers using UnrealCodeFixup to make sure dllstorage is on methods/staticvars instead of types. Part 5/n

diff --git a/Engine/Source/Editor/Sequencer/Public/MVVM/ViewModels/ObjectBindingModel.h b/Engine/Source/Editor/Sequencer/Public/MVVM/ViewModels/ObjectBindingModel.h
--- a/Engine/Source/Editor/Sequencer/Public/MVVM/ViewModels/ObjectBindingModel.h
+++ b/Engine/Source/Editor/Sequencer/Public/MVVM/ViewModels/ObjectBindingModel.h
@@ -15,6 +15,8 @@
 #include "MVVM/Extensions/IDraggableOutlinerExtension.h"
 #include "MVVM/Extensions/IDeletableExtension.h"
 
+#define UE_API SEQUENCER_API
+
 
 struct FMovieSceneBinding;
 
@@ -40,7 +42,7 @@ enum class EObjectBindingType
 	Possessable, Spawnable, Unknown
 };
 
-class SEQUENCER_API FObjectBindingModel
+class FObjectBindingModel
 	: public FEvaluableOutlinerItemModel
 	, public IObjectBindingExtension
 	, public IDraggableOutlinerExtension
@@ -52,7 +54,7 @@ class SEQUENCER_API FObjectBindingModel
 {
 public:
 
-	UE_SEQUENCER_DECLARE_CASTABLE(FObjectBindingModel
+	UE_SEQUENCER_DECLARE_CASTABLE_API(UE_API, FObjectBindingModel
 		, FEvaluableOutlinerItemModel
 		, IObjectBindingExtension
 		, IDraggableOutlinerExtension
@@ -62,73 +64,73 @@ public:
 		, ISortableExtension
 		, IDeletableExtension);
 
-	FObjectBindingModel(FSequenceModel* OwnerModel, const FMovieSceneBinding& InBinding);
-	~FObjectBindingModel();
+	UE_API FObjectBindingModel(FSequenceModel* OwnerModel, const FMovieSceneBinding& InBinding);
+	UE_API ~FObjectBindingModel();
 
-	static EViewModelListType GetTopLevelChildTrackAreaGroupType();
+	static UE_API EViewModelListType GetTopLevelChildTrackAreaGroupType();
 
-	void AddTrack(UMovieSceneTrack* Track);
-	void RemoveTrack(UMovieSceneTrack* Track);
+	UE_API void AddTrack(UMovieSceneTrack* Track);
+	UE_API void RemoveTrack(UMovieSceneTrack* Track);
 
 	/*~ IObjectBindingExtension */
-	FGuid GetObjectGuid() const override;
+	UE_API FGuid GetObjectGuid() const override;
 
 	/*~ IRenameableExtension */
-	bool CanRename() const override;
-	void Rename(const FText& NewName) override;
+	UE_API bool CanRename() const override;
+	UE_API void Rename(const FText& NewName) override;
 
 	/*~ IOutlinerExtension */
-	FOutlinerSizing GetOutlinerSizing() const override;
-	FText GetLabel() const override;
-	FSlateColor GetLabelColor() const override;
-	FText GetLabelToolTipText() const override;
-	const FSlateBrush* GetIconBrush() const override;
-	TSharedPtr<SWidget> CreateOutlinerViewForColumn(const FCreateOutlinerViewParams& InParams, const FName& InColumnName) override;
+	UE_API FOutlinerSizing GetOutlinerSizing() const override;
+	UE_API FText GetLabel() const override;
+	UE_API FSlateColor GetLabelColor() const override;
+	UE_API FText GetLabelToolTipText() const override;
+	UE_API const FSlateBrush* GetIconBrush() const override;
+	UE_API TSharedPtr<SWidget> CreateOutlinerViewForColumn(const FCreateOutlinerViewParams& InParams, const FName& InColumnName) override;
 
 	/*~ ITrackAreaExtension */
-	FTrackAreaParameters GetTrackAreaParameters() const override;
-	FViewModelVariantIterator GetTrackAreaModelList() const override;
-	FViewModelVariantIterator GetTopLevelChildTrackAreaModels() const override;
+	UE_API FTrackAreaParameters GetTrackAreaParameters() const override;
+	UE_API FViewModelVariantIterator GetTrackAreaModelList() const override;
+	UE_API FViewModelVariantIterator GetTopLevelChildTrackAreaModels() const override;
 
 	/*~ IGroupableExtension */
-	void GetIdentifierForGrouping(TStringBuilder<128>& OutString) const override;
+	UE_API void GetIdentifierForGrouping(TStringBuilder<128>& OutString) const override;
 
 	/*~ ISortableExtension */
-	void SortChildren() override;
-	FSortingKey GetSortingKey() const override;
-	void SetCustomOrder(int32 InCustomOrder) override;
+	UE_API void SortChildren() override;
+	UE_API FSortingKey GetSortingKey() const override;
+	UE_API void SetCustomOrder(int32 InCustomOrder) override;
 
 	/*~ FOutlinerItemModel */
-	void BuildContextMenu(FMenuBuilder& MenuBuilder) override;
-	void BuildOrganizeContextMenu(FMenuBuilder& MenuBuilder) override;
-	bool GetDefaultExpansionState() const override;
-	void BuildSidebarMenu(FMenuBuilder& MenuBuilder) override;
+	UE_API void BuildContextMenu(FMenuBuilder& MenuBuilder) override;
+	UE_API void BuildOrganizeContextMenu(FMenuBuilder& MenuBuilder) override;
+	UE_API bool GetDefaultExpansionState() const override;
+	UE_API void BuildSidebarMenu(FMenuBuilder& MenuBuilder) override;
 
 	/*~ IDraggableOutlinerExtension */
-	bool CanDrag() const override;
+	UE_API bool CanDrag() const override;
 
 	/*~ IDeletableExtension */
-	bool CanDelete(FText* OutErrorMessage) const override;
-	void Delete() override;
+	UE_API bool CanDelete(FText* OutErrorMessage) const override;
+	UE_API void Delete() override;
 
 public:
 
-	virtual void SetParentBindingID(const FGuid& InObjectBindingID);
-	virtual FGuid GetDesiredParentBinding() const;
-	virtual EObjectBindingType GetType() const;
-	virtual FText GetTooltipForSingleObjectBinding() const;
-	virtual const UClass* FindObjectClass() const;
-	virtual bool SupportsRebinding() const;
+	UE_API virtual void SetParentBindingID(const FGuid& InObjectBindingID);
+	UE_API virtual FGuid GetDesiredParentBinding() const;
+	UE_API virtual EObjectBindingType GetType() const;
+	UE_API virtual FText GetTooltipForSingleObjectBinding() const;
+	UE_API virtual const UClass* FindObjectClass() const;
+	UE_API virtual bool SupportsRebinding() const;
 	virtual FSlateColor GetInvalidBindingLabelColor() const { return FLinearColor::Red; }
 
 public:
 
-	TSharedRef<SWidget> GetAddTrackMenuContent();
+	UE_API TSharedRef<SWidget> GetAddTrackMenuContent();
 
 protected:
 
 	/*~ FViewModel interface */
-	void OnConstruct() override;
+	UE_API void OnConstruct() override;
 
 private:
 
@@ -144,19 +146,19 @@ private:
 		}
 	};
 
-	void AddPropertyMenuItems(FMenuBuilder& AddTrackMenuBuilder, int32 NumStartingBlocks, TArray<FPropertyPath> KeyablePropertyPaths, int32 PropertyNameIndexStart);
-	void AddPropertyMenuItem(FMenuBuilder& AddTrackMenuBuilder, const FPropertyMenuData& KeyablePropertyMenuData);
-	void HandleAddTrackSubMenuNew(FMenuBuilder& AddTrackMenuBuilder, TArray<FPropertyPath> KeyablePropertyPaths, int32 PropertyNameIndexStart);
-	void HandlePropertyMenuItemExecute(FPropertyPath PropertyPath);
+	UE_API void AddPropertyMenuItems(FMenuBuilder& AddTrackMenuBuilder, int32 NumStartingBlocks, TArray<FPropertyPath> KeyablePropertyPaths, int32 PropertyNameIndexStart);
+	UE_API void AddPropertyMenuItem(FMenuBuilder& AddTrackMenuBuilder, const FPropertyMenuData& KeyablePropertyMenuData);
+	UE_API void HandleAddTrackSubMenuNew(FMenuBuilder& AddTrackMenuBuilder, TArray<FPropertyPath> KeyablePropertyPaths, int32 PropertyNameIndexStart);
+	UE_API void HandlePropertyMenuItemExecute(FPropertyPath PropertyPath);
 
-	void AddTagMenu(FMenuBuilder& MenuBuilder);
-	ECheckBoxState GetTagCheckState(FName TagName);
-	void ToggleTag(FName TagName);
-	void HandleDeleteTag(FName TagName);
-	void HandleAddTag(FName TagName);
-	void HandleTemplateActorClassPicked(UClass* ChosenClass);
+	UE_API void AddTagMenu(FMenuBuilder& MenuBuilder);
+	UE_API ECheckBoxState GetTagCheckState(FName TagName);
+	UE_API void ToggleTag(FName TagName);
+	UE_API void HandleDeleteTag(FName TagName);
+	UE_API void HandleAddTag(FName TagName);
+	UE_API void HandleTemplateActorClassPicked(UClass* ChosenClass);
 
-	void OnFinishedChangingDynamicBindingProperties(const FPropertyChangedEvent& ChangeEvent, TSharedPtr<FStructOnScope> ValueStruct);
+	UE_API void OnFinishedChangingDynamicBindingProperties(const FPropertyChangedEvent& ChangeEvent, TSharedPtr<FStructOnScope> ValueStruct);
 
 protected:
 
@@ -172,3 +174,4 @@ protected:
 } // namespace Sequencer
 } // namespace UE
 
+#undef UE_API
diff --git a/Engine/Source/Editor/Sequencer/Public/MVVM/ViewModels/OutlinerItemModel.h b/Engine/Source/Editor/Sequencer/Public/MVVM/ViewModels/OutlinerItemModel.h
--- a/Engine/Source/Editor/Sequencer/Public/MVVM/ViewModels/OutlinerItemModel.h
+++ b/Engine/Source/Editor/Sequencer/Public/MVVM/ViewModels/OutlinerItemModel.h
@@ -17,6 +17,8 @@
 #include "CurveEditorTypes.h"
 #include "Tree/ICurveEditorTreeItem.h"
 
+#define UE_API SEQUENCER_API
+
 class UMovieSceneSection;
 class UMovieSceneSequence;
 class UMovieSceneTrack;
@@ -31,7 +33,7 @@ namespace Sequencer
 
 class FSequencerEditorViewModel;
 
-class SEQUENCER_API FOutlinerItemModelMixin
+class FOutlinerItemModelMixin
 	: public FOutlinerExtensionShim
 	, public FGeometryExtensionShim
 	, public FPinnableExtensionShim
@@ -44,55 +46,55 @@ public:
 
 	using Implements = TImplements<IOutlinerExtension, IGeometryExtension, IPinnableExtension, IHoveredExtension, IDimmableExtension, ICurveEditorTreeItemExtension>;
 
-	FOutlinerItemModelMixin();
+	UE_API FOutlinerItemModelMixin();
 
-	TSharedPtr<FSequencerEditorViewModel> GetEditor() const;
+	UE_API TSharedPtr<FSequencerEditorViewModel> GetEditor() const;
 	
 	/*~ IOutlinerExtension */
-	FName GetIdentifier() const override;
-	bool IsExpanded() const override;
-	void SetExpansion(bool bInIsExpanded) override;
-	bool IsFilteredOut() const override;
-	TSharedPtr<SWidget> CreateContextMenuWidget(const FCreateOutlinerContextMenuWidgetParams& InParams) override;
-	FSlateColor GetLabelColor() const override;
+	UE_API FName GetIdentifier() const override;
+	UE_API bool IsExpanded() const override;
+	UE_API void SetExpansion(bool bInIsExpanded) override;
+	UE_API bool IsFilteredOut() const override;
+	UE_API TSharedPtr<SWidget> CreateContextMenuWidget(const FCreateOutlinerContextMenuWidgetParams& InParams) override;
+	UE_API FSlateColor GetLabelColor() const override;
 
 	/*~ ICurveEditorTreeItemExtension */
-	virtual bool HasCurves() const override;
-	virtual TSharedPtr<ICurveEditorTreeItem> GetCurveEditorTreeItem() const override;
-	virtual TOptional<FString> GetUniquePathName() const override;
+	UE_API virtual bool HasCurves() const override;
+	UE_API virtual TSharedPtr<ICurveEditorTreeItem> GetCurveEditorTreeItem() const override;
+	UE_API virtual TOptional<FString> GetUniquePathName() const override;
 
 	/*~ ICurveEditorTreeItem */
-	virtual TSharedPtr<SWidget> GenerateCurveEditorTreeWidget(const FName& InColumnName, TWeakPtr<FCurveEditor> InCurveEditor, FCurveEditorTreeItemID InTreeItemID, const TSharedRef<ITableRow>& InTableRow) override;
-	virtual void CreateCurveModels(TArray<TUniquePtr<FCurveModel>>& OutCurveModels) override;
-	virtual bool PassesFilter(const FCurveEditorTreeFilter* InFilter) const override;
+	UE_API virtual TSharedPtr<SWidget> GenerateCurveEditorTreeWidget(const FName& InColumnName, TWeakPtr<FCurveEditor> InCurveEditor, FCurveEditorTreeItemID InTreeItemID, const TSharedRef<ITableRow>& InTableRow) override;
+	UE_API virtual void CreateCurveModels(TArray<TUniquePtr<FCurveModel>>& OutCurveModels) override;
+	UE_API virtual bool PassesFilter(const FCurveEditorTreeFilter* InFilter) const override;
 
 	/*~ IPinnableExtension */
-	bool IsPinned() const override;
+	UE_API bool IsPinned() const override;
 
 	/*~ IDimmableExtension */
-	bool IsDimmed() const override;
+	UE_API bool IsDimmed() const override;
 
 	/** Get context menu contents. */
-	virtual void BuildContextMenu(FMenuBuilder& MenuBuilder);
+	UE_API virtual void BuildContextMenu(FMenuBuilder& MenuBuilder);
 
-	virtual void BuildSidebarMenu(FMenuBuilder& MenuBuilder);
+	UE_API virtual void BuildSidebarMenu(FMenuBuilder& MenuBuilder);
 
 protected:
 
 	/** Set identifier for computing node paths */
-	void SetIdentifier(FName InNewIdentifier);
+	UE_API void SetIdentifier(FName InNewIdentifier);
 
 	/** Get the default expansion state if it wasn't saved in the movie-scene data */
-	virtual bool GetDefaultExpansionState() const;
+	UE_API virtual bool GetDefaultExpansionState() const;
 
 	/** Set expansion state without saving it in the movie-scene data */
-	void SetExpansionWithoutSaving(bool bInIsExpanded);
+	UE_API void SetExpansionWithoutSaving(bool bInIsExpanded);
 
-	virtual void BuildSectionColorTintsMenu(FMenuBuilder& MenuBuilder);
-	virtual void BuildOrganizeContextMenu(FMenuBuilder& MenuBuilder);
-	virtual void BuildDisplayOptionsMenu(FMenuBuilder& MenuBuilder);
-	virtual void BuildTrackOptionsMenu(FMenuBuilder& MenuBuilder, const TArray<UMovieSceneTrack*>& InTracks);
-	virtual void BuildTrackRowOptionsMenu(FMenuBuilder& MenuBuilder);
+	UE_API virtual void BuildSectionColorTintsMenu(FMenuBuilder& MenuBuilder);
+	UE_API virtual void BuildOrganizeContextMenu(FMenuBuilder& MenuBuilder);
+	UE_API virtual void BuildDisplayOptionsMenu(FMenuBuilder& MenuBuilder);
+	UE_API virtual void BuildTrackOptionsMenu(FMenuBuilder& MenuBuilder, const TArray<UMovieSceneTrack*>& InTracks);
+	UE_API virtual void BuildTrackRowOptionsMenu(FMenuBuilder& MenuBuilder);
 
 private:
 
@@ -104,21 +106,21 @@ protected:
 
 private:
 
-	bool IsRootModelPinned() const;
-	void ToggleRootModelPinned();
+	UE_API bool IsRootModelPinned() const;
+	UE_API void ToggleRootModelPinned();
 
-	ECheckBoxState SelectedModelsSoloState() const;
-	void ToggleSelectedModelsSolo();
+	UE_API ECheckBoxState SelectedModelsSoloState() const;
+	UE_API void ToggleSelectedModelsSolo();
 
-	ECheckBoxState SelectedModelsMuteState() const;
-	void ToggleSelectedModelsMuted();
+	UE_API ECheckBoxState SelectedModelsMuteState() const;
+	UE_API void ToggleSelectedModelsMuted();
 
-	TArray<UMovieSceneSection*> GetSelectedSections() const;
-	TArray<UMovieSceneTrack*> GetSelectedTracks() const;
-	TArray<TPair<UMovieSceneTrack*, int32>> GetSelectedTrackRows() const;
+	UE_API TArray<UMovieSceneSection*> GetSelectedSections() const;
+	UE_API TArray<UMovieSceneTrack*> GetSelectedTracks() const;
+	UE_API TArray<TPair<UMovieSceneTrack*, int32>> GetSelectedTrackRows() const;
 
-	void AddEvalOptionsPropertyMenuItem(FMenuBuilder& InMenuBuilder, const FBoolProperty* InProperty, TFunction<bool(UMovieSceneTrack*)> InValidator = nullptr);
-	void AddDisplayOptionsPropertyMenuItem(FMenuBuilder& InMenuBuilder, const FBoolProperty* InProperty, TFunction<bool(UMovieSceneTrack*)> InValidator = nullptr);
+	UE_API void AddEvalOptionsPropertyMenuItem(FMenuBuilder& InMenuBuilder, const FBoolProperty* InProperty, TFunction<bool(UMovieSceneTrack*)> InValidator = nullptr);
+	UE_API void AddDisplayOptionsPropertyMenuItem(FMenuBuilder& InMenuBuilder, const FBoolProperty* InProperty, TFunction<bool(UMovieSceneTrack*)> InValidator = nullptr);
 
 private:
 
@@ -154,13 +156,13 @@ public:
 	virtual const FViewModel* AsViewModel() const { return this; }
 };
 
-class SEQUENCER_API FOutlinerItemModel : public TOutlinerModelMixin<FViewModel>
+class FOutlinerItemModel : public TOutlinerModelMixin<FViewModel>
 {
 public:
-	UE_SEQUENCER_DECLARE_CASTABLE(FOutlinerItemModel, FViewModel, FOutlinerItemModelMixin);
+	UE_SEQUENCER_DECLARE_CASTABLE_API(UE_API, FOutlinerItemModel, FViewModel, FOutlinerItemModelMixin);
 };
 
-class SEQUENCER_API FEvaluableOutlinerItemModel
+class FEvaluableOutlinerItemModel
 	: public FOutlinerItemModel
 	, public IMutableExtension
 	, public ISoloableExtension
@@ -168,31 +170,31 @@ class SEQUENCER_API FEvaluableOutlinerItemModel
 {
 public:
 
-	UE_SEQUENCER_DECLARE_CASTABLE(FEvaluableOutlinerItemModel
+	UE_SEQUENCER_DECLARE_CASTABLE_API(UE_API, FEvaluableOutlinerItemModel
 		, FOutlinerItemModel
 		, IMutableExtension
 		, ISoloableExtension
 		, IDeactivatableExtension);
 
 	/*~ ISoloableExtension */
-	bool IsSolo() const override;
-	void SetIsSoloed(bool bIsSoloed) override;
+	UE_API bool IsSolo() const override;
+	UE_API void SetIsSoloed(bool bIsSoloed) override;
 
 	/*~ IMutableExtension */
-	bool IsMuted() const override;
-	void SetIsMuted(bool bIsMuted) override;
+	UE_API bool IsMuted() const override;
+	UE_API void SetIsMuted(bool bIsMuted) override;
 
 	/*~ IDeactivatableExtension */
-	virtual bool IsDeactivated() const override;
-	virtual void SetIsDeactivated(const bool bInIsDeactivated) override;
+	UE_API virtual bool IsDeactivated() const override;
+	UE_API virtual void SetIsDeactivated(const bool bInIsDeactivated) override;
 };
 
-class SEQUENCER_API FOutlinerCacheExtension
+class FOutlinerCacheExtension
 	: public FHierarchicalCacheExtension
 {
 public:
 
-	UE_SEQUENCER_DECLARE_VIEW_MODEL_TYPE_ID(FOutlinerCacheExtension);
+	UE_SEQUENCER_DECLARE_VIEW_MODEL_TYPE_ID_API(UE_API, FOutlinerCacheExtension);
 
 	FOutlinerCacheExtension()
 	{
@@ -204,3 +206,4 @@ public:
 } // namespace Sequencer
 } // namespace UE
 
+#undef UE_API
diff --git a/Engine/Source/Editor/Sequencer/Public/MVVM/ViewModels/SectionModel.h b/Engine/Source/Editor/Sequencer/Public/MVVM/ViewModels/SectionModel.h
--- a/Engine/Source/Editor/Sequencer/Public/MVVM/ViewModels/SectionModel.h
+++ b/Engine/Source/Editor/Sequencer/Public/MVVM/ViewModels/SectionModel.h
@@ -20,6 +20,8 @@
 #include "EventHandlers/ISignedObjectEventHandler.h"
 #include "EventHandlers/ISectionEventHandler.h"
 
+#define UE_API SEQUENCER_API
+
 class ISequencerSection;
 class UMovieSceneSection;
 struct FMovieSceneChannelProxy;
@@ -36,7 +38,7 @@ struct FOverlappingSections;
 /**
  * Model for a sequencer section
  */
-class SEQUENCER_API FSectionModel
+class FSectionModel
 	: public FViewModel
 	, public FLinkedOutlinerExtension
 	, public IObjectModelExtension
@@ -52,7 +54,7 @@ class SEQUENCER_API FSectionModel
 {
 public:
 
-	UE_SEQUENCER_DECLARE_CASTABLE(FSectionModel, FViewModel
+	UE_SEQUENCER_DECLARE_CASTABLE_API(UE_API, FSectionModel, FViewModel
 		, FLinkedOutlinerExtension
 		, IObjectModelExtension
 		, ILayerBarExtension
@@ -69,33 +71,33 @@ public:
 	UE_DEPRECATED(5.5, "This member is no longer supported, please subscribe to UMovieSceneSignedObject::OnModifiedDirectly.")
 	FOnModelUpdated OnUpdated;
 
-	FSectionModel();
-	~FSectionModel();
+	UE_API FSectionModel();
+	UE_API ~FSectionModel();
 
-	void InitializeSection(TSharedPtr<ISequencerSection> InSectionInterface);
+	UE_API void InitializeSection(TSharedPtr<ISequencerSection> InSectionInterface);
 
 	/** Gets the time range of the section */
-	TRange<FFrameNumber> GetRange() const;
+	UE_API TRange<FFrameNumber> GetRange() const;
 
 	/** Gets the section object */
-	UMovieSceneSection* GetSection() const;
+	UE_API UMovieSceneSection* GetSection() const;
 
 	/** Gets the section interface */
-	TSharedPtr<ISequencerSection> GetSectionInterface() const;
+	UE_API TSharedPtr<ISequencerSection> GetSectionInterface() const;
 
 	/** Gets the parent track or track row model */
-	TViewModelPtr<ITrackExtension> GetParentTrackModel() const;
+	UE_API TViewModelPtr<ITrackExtension> GetParentTrackModel() const;
 
 	/** Gets the parent track or track row model as an ITrackExtension */
-	TViewModelPtr<ITrackExtension> GetParentTrackExtension() const;
+	UE_API TViewModelPtr<ITrackExtension> GetParentTrackExtension() const;
 
 	/** Gets an array of sections that underlap the specified section */
-	TArray<FOverlappingSections> GetUnderlappingSections();
+	UE_API TArray<FOverlappingSections> GetUnderlappingSections();
 	/** Gets an array of sections whose easing bounds underlap the specified section */
-	TArray<FOverlappingSections> GetEasingSegments();
+	UE_API TArray<FOverlappingSections> GetEasingSegments();
 
-	int32 GetPreRollFrames() const;
-	int32 GetPostRollFrames() const;
+	UE_API int32 GetPreRollFrames() const;
+	UE_API int32 GetPostRollFrames() const;
 
 	/** Returns whether this section model needs to be rebuilt, ie. the channel proxy is no longer valid */
 	bool NeedsLayout() const
@@ -118,48 +120,48 @@ public:
 public:
 
 	/*~ IObjectModelExtension Interface */
-	virtual void InitializeObject(TWeakObjectPtr<> InWeakObject) override;
-	virtual UObject* GetObject() const override;
+	UE_API virtual void InitializeObject(TWeakObjectPtr<> InWeakObject) override;
+	UE_API virtual UObject* GetObject() const override;
 
 	/*~ ILayerBarExtension Interface */
-	TRange<FFrameNumber> GetLayerBarRange() const override;
-	void OffsetLayerBar(FFrameNumber Amount) override;
+	UE_API TRange<FFrameNumber> GetLayerBarRange() const override;
+	UE_API void OffsetLayerBar(FFrameNumber Amount) override;
 
 	/*~ ITrackLaneExtension Interface */
-	TSharedPtr<ITrackLaneWidget> CreateTrackLaneView(const FCreateTrackLaneViewParams& InParams) override;
-	FTrackLaneVirtualAlignment ArrangeVirtualTrackLaneView() const override;
+	UE_API TSharedPtr<ITrackLaneWidget> CreateTrackLaneView(const FCreateTrackLaneViewParams& InParams) override;
+	UE_API FTrackLaneVirtualAlignment ArrangeVirtualTrackLaneView() const override;
 
 	/*~ ISelectableExtension Interface */
-	ESelectionIntent IsSelectable() const override;
+	UE_API ESelectionIntent IsSelectable() const override;
 
 	/*~ ISnappableExtension Interface */
-	void AddToSnapField(const ISnapCandidate& Candidate, ISnapField& SnapField) const override;
+	UE_API void AddToSnapField(const ISnapCandidate& Candidate, ISnapField& SnapField) const override;
 
 	/*~ ISignedObjectEventHandler Interface */
-	void OnModifiedDirectly(UMovieSceneSignedObject*) override;
+	UE_API void OnModifiedDirectly(UMovieSceneSignedObject*) override;
 
 	/*~ ISectionEventHandler Interface */
-	void OnRowChanged(UMovieSceneSection*) override;
+	UE_API void OnRowChanged(UMovieSceneSection*) override;
 
 	/*~ IDraggableTrackAreaExtension Interface */
-	bool CanDrag() const override;
-	void OnBeginDrag(IDragOperation& DragOperation) override;
-	void OnEndDrag(IDragOperation& DragOperation) override;
+	UE_API bool CanDrag() const override;
+	UE_API void OnBeginDrag(IDragOperation& DragOperation) override;
+	UE_API void OnEndDrag(IDragOperation& DragOperation) override;
 
 	/*~ IStretchableExtension Interface */
-	void OnInitiateStretch(IStretchOperation& StretchOperation, EStretchConstraint Constraint, FStretchParameters* InOutGlobalParameters) override;
-	EStretchResult OnBeginStretch(const IStretchOperation& StretchOperation, const FStretchScreenParameters& ScreenParameters, FStretchParameters* InOutParameters) override;
-	void OnStretch(const IStretchOperation& StretchOperation, const FStretchScreenParameters& ScreenParameters, FStretchParameters* InOutParameters) override;
-	void OnEndStretch(const IStretchOperation& StretchOperation, const FStretchScreenParameters& ScreenParameters, FStretchParameters* InOutParameters) override;
+	UE_API void OnInitiateStretch(IStretchOperation& StretchOperation, EStretchConstraint Constraint, FStretchParameters* InOutGlobalParameters) override;
+	UE_API EStretchResult OnBeginStretch(const IStretchOperation& StretchOperation, const FStretchScreenParameters& ScreenParameters, FStretchParameters* InOutParameters) override;
+	UE_API void OnStretch(const IStretchOperation& StretchOperation, const FStretchScreenParameters& ScreenParameters, FStretchParameters* InOutParameters) override;
+	UE_API void OnEndStretch(const IStretchOperation& StretchOperation, const FStretchScreenParameters& ScreenParameters, FStretchParameters* InOutParameters) override;
 
 	/*~ IConditionableExtension Interface */
-	const UMovieSceneCondition* GetCondition() const override;
-	EConditionableConditionState GetConditionState() const override;
-	void SetConditionEditorForceTrue(bool bEditorForceTrue) override;
+	UE_API const UMovieSceneCondition* GetCondition() const override;
+	UE_API EConditionableConditionState GetConditionState() const override;
+	UE_API void SetConditionEditorForceTrue(bool bEditorForceTrue) override;
 
 private:
 
-	void UpdateCachedData();
+	UE_API void UpdateCachedData();
 
 private:
 
@@ -184,3 +186,4 @@ struct FOverlappingSections
 } // namespace Sequencer
 } // namespace UE
 
+#undef UE_API
diff --git a/Engine/Source/Editor/Sequencer/Public/MVVM/ViewModels/SequenceModel.h b/Engine/Source/Editor/Sequencer/Public/MVVM/ViewModels/SequenceModel.h
--- a/Engine/Source/Editor/Sequencer/Public/MVVM/ViewModels/SequenceModel.h
+++ b/Engine/Source/Editor/Sequencer/Public/MVVM/ViewModels/SequenceModel.h
@@ -20,6 +20,8 @@
 #include "MVVM/Extensions/ISortableExtension.h"
 #include "MVVM/Extensions/ICurveEditorTreeItemExtension.h"
 
+#define UE_API SEQUENCER_API
+
 class ISequencer;
 class FSequencer;
 class FCurveEditor;
@@ -37,7 +39,7 @@ class FSequencerEditorViewModel;
 
 DECLARE_MULTICAST_DELEGATE_TwoParams(FOnInitializeSequenceModel, TSharedPtr<FEditorViewModel>, TSharedPtr<FSequenceModel>);
 
-class SEQUENCER_API FSequenceModel
+class FSequenceModel
 	: public FViewModel
 	, public ISortableExtension
 	, public IOutlinerDropTargetOutlinerExtension
@@ -46,45 +48,45 @@ class SEQUENCER_API FSequenceModel
 {
 public:
 
-	UE_SEQUENCER_DECLARE_CASTABLE(FSequenceModel, FViewModel, ISortableExtension, IOutlinerDropTargetOutlinerExtension);
+	UE_SEQUENCER_DECLARE_CASTABLE_API(UE_API, FSequenceModel, FViewModel, ISortableExtension, IOutlinerDropTargetOutlinerExtension);
 
-	static FOnInitializeSequenceModel CreateExtensionsEvent;
+	static UE_API FOnInitializeSequenceModel CreateExtensionsEvent;
 
-	FSequenceModel(TWeakPtr<FSequencerEditorViewModel> InEditorViewModel);
+	UE_API FSequenceModel(TWeakPtr<FSequencerEditorViewModel> InEditorViewModel);
 
-	void InitializeExtensions();
+	UE_API void InitializeExtensions();
 
-	FMovieSceneSequenceID GetSequenceID() const;
-	UMovieSceneSequence* GetSequence() const;
-	UMovieScene* GetMovieScene() const;
-	void SetSequence(UMovieSceneSequence* InSequence, FMovieSceneSequenceID InSequenceID);
+	UE_API FMovieSceneSequenceID GetSequenceID() const;
+	UE_API UMovieSceneSequence* GetSequence() const;
+	UE_API UMovieScene* GetMovieScene() const;
+	UE_API void SetSequence(UMovieSceneSequence* InSequence, FMovieSceneSequenceID InSequenceID);
 
-	TSharedPtr<FSequencerEditorViewModel> GetEditor() const;
-	TSharedPtr<ISequencer> GetSequencer() const;
-	TSharedPtr<FSequencer> GetSequencerImpl() const;
+	UE_API TSharedPtr<FSequencerEditorViewModel> GetEditor() const;
+	UE_API TSharedPtr<ISequencer> GetSequencer() const;
+	UE_API TSharedPtr<FSequencer> GetSequencerImpl() const;
 
 	TSharedPtr<FViewModel> GetBottomSpacer() const { return BottomSpacer; }
 
 	/*~ ISignedObjectEventHandler */
-	void OnPostUndo() override;
-	void OnModifiedIndirectly(UMovieSceneSignedObject*) override;
-	void OnModifiedDirectly(UMovieSceneSignedObject*) override;
+	UE_API void OnPostUndo() override;
+	UE_API void OnModifiedIndirectly(UMovieSceneSignedObject*) override;
+	UE_API void OnModifiedDirectly(UMovieSceneSignedObject*) override;
 
 	/*~ ISortableExtension */
-	void SortChildren() override;
-	FSortingKey GetSortingKey() const override;
-	void SetCustomOrder(int32 InCustomOrder) override;
+	UE_API void SortChildren() override;
+	UE_API FSortingKey GetSortingKey() const override;
+	UE_API void SetCustomOrder(int32 InCustomOrder) override;
 
 	/*~ Begin ISequenceDataEventHandler */
-	virtual void OnDecorationAdded(UObject* AddedDecoration) override;
-	virtual void OnDecorationRemoved(UObject* RemovedDecoration) override;
+	UE_API virtual void OnDecorationAdded(UObject* AddedDecoration) override;
+	UE_API virtual void OnDecorationRemoved(UObject* RemovedDecoration) override;
 	/*~ End ISequenceDataEventHandler */
 
 	/*~ IOutlinerDropTargetOutlinerExtension */
-	TOptional<EItemDropZone> CanAcceptDrop(const FViewModelPtr& TargetModel, const FDragDropEvent& DragDropEvent, EItemDropZone InItemDropZone) override;
-	void PerformDrop(const FViewModelPtr& TargetModel, const FDragDropEvent& DragDropEvent, EItemDropZone InItemDropZone) override;
+	UE_API TOptional<EItemDropZone> CanAcceptDrop(const FViewModelPtr& TargetModel, const FDragDropEvent& DragDropEvent, EItemDropZone InItemDropZone) override;
+	UE_API void PerformDrop(const FViewModelPtr& TargetModel, const FDragDropEvent& DragDropEvent, EItemDropZone InItemDropZone) override;
 
-	static EViewModelListType GetDecorationModelListType();
+	static UE_API EViewModelListType GetDecorationModelListType();
 
 private:
 
@@ -105,3 +107,4 @@ private:
 } // namespace Sequencer
 } // namespace UE
 
+#undef UE_API
diff --git a/Engine/Source/Editor/Sequencer/Public/MVVM/ViewModels/SequencerEditorViewModel.h b/Engine/Source/Editor/Sequencer/Public/MVVM/ViewModels/SequencerEditorViewModel.h
--- a/Engine/Source/Editor/Sequencer/Public/MVVM/ViewModels/SequencerEditorViewModel.h
+++ b/Engine/Source/Editor/Sequencer/Public/MVVM/ViewModels/SequencerEditorViewModel.h
@@ -7,6 +7,8 @@
 #include "SequencerCustomizationManager.h"
 #include "Toolkits/AssetEditorToolkit.h"
 
+#define UE_API SEQUENCER_API
+
 struct FMovieSceneSequenceID;
 
 class ISequencer;
@@ -25,61 +27,61 @@ struct ITrackAreaHotspot;
 /**
  * Main view-model for the Sequencer editor.
  */
-class SEQUENCER_API FSequencerEditorViewModel : public FEditorViewModel
+class FSequencerEditorViewModel : public FEditorViewModel
 {
 public:
 
-	UE_SEQUENCER_DECLARE_CASTABLE(FSequencerEditorViewModel, FEditorViewModel);
+	UE_SEQUENCER_DECLARE_CASTABLE_API(UE_API, FSequencerEditorViewModel, FEditorViewModel);
 
-	FSequencerEditorViewModel(TSharedRef<ISequencer> InSequencer, const FSequencerHostCapabilities& InHostCapabilities);
-	~FSequencerEditorViewModel();
+	UE_API FSequencerEditorViewModel(TSharedRef<ISequencer> InSequencer, const FSequencerHostCapabilities& InHostCapabilities);
+	UE_API ~FSequencerEditorViewModel();
 
 	/** Retrieve this editor's selection class */
-	TSharedPtr<FSequencerSelection> GetSelection() const;
+	UE_API TSharedPtr<FSequencerSelection> GetSelection() const;
 
 	// @todo_sequencer_mvvm remove this later
-	TSharedPtr<ISequencer> GetSequencer() const;
+	UE_API TSharedPtr<ISequencer> GetSequencer() const;
 	// @todo_sequencer_mvvm remove this ASAP
-	TSharedPtr<FSequencer> GetSequencerImpl() const;
+	UE_API TSharedPtr<FSequencer> GetSequencerImpl() const;
 
 	// @todo_sequencer_mvvm move this to the root view-model
-	void SetSequence(UMovieSceneSequence* InRootSequence);
+	UE_API void SetSequence(UMovieSceneSequence* InRootSequence);
 
-	TViewModelPtr<FSequenceModel> GetRootSequenceModel() const;
+	UE_API TViewModelPtr<FSequenceModel> GetRootSequenceModel() const;
 
 public:
 
 	/** Adjust sequencer customizations based on the currently focused sequence */
-	bool UpdateSequencerCustomizations(const UMovieSceneSequence* PreviousFocusedSequence);
+	UE_API bool UpdateSequencerCustomizations(const UMovieSceneSequence* PreviousFocusedSequence);
 
 	/** Get the active customization infos */
-	TArrayView<const FSequencerCustomizationInfo> GetActiveCustomizationInfos() const;
+	UE_API TArrayView<const FSequencerCustomizationInfo> GetActiveCustomizationInfos() const;
 
 	/** Build a combined menu extender */
-	TSharedPtr<FExtender> GetSequencerMenuExtender(
+	UE_API TSharedPtr<FExtender> GetSequencerMenuExtender(
 			TSharedPtr<FExtensibilityManager> ExtensibilityManager, const TArray<UObject*>& ContextObjects,
 			TFunctionRef<const FOnGetSequencerMenuExtender&(const FSequencerCustomizationInfo&)> Endpoint, FViewModelPtr InViewModel) const;
 
 	/** Gets the pinned track area view-model. */
-	TSharedPtr<FTrackAreaViewModel> GetPinnedTrackArea() const;
+	UE_API TSharedPtr<FTrackAreaViewModel> GetPinnedTrackArea() const;
 
 	/** Gets the current hotspots across any of our track areas */
-	TSharedPtr<ITrackAreaHotspot> GetHotspot() const;
+	UE_API TSharedPtr<ITrackAreaHotspot> GetHotspot() const;
 
-	void HandleDataHierarchyChanged();
+	UE_API void HandleDataHierarchyChanged();
 
 protected:
 
-	virtual void PreInitializeEditorImpl() override;
-	virtual void InitializeEditorImpl() override;
-	virtual TSharedPtr<FViewModel> CreateRootModelImpl() override;
-	virtual TSharedPtr<FOutlinerViewModel> CreateOutlinerImpl() override;
-	virtual TSharedPtr<FTrackAreaViewModel> CreateTrackAreaImpl() override;
-	virtual TSharedPtr<FSequencerCoreSelection> CreateSelectionImpl() override;
-	virtual USequencerScriptingLayer* CreateScriptingLayerImpl() override;
-	virtual bool IsReadOnly() const override;
+	UE_API virtual void PreInitializeEditorImpl() override;
+	UE_API virtual void InitializeEditorImpl() override;
+	UE_API virtual TSharedPtr<FViewModel> CreateRootModelImpl() override;
+	UE_API virtual TSharedPtr<FOutlinerViewModel> CreateOutlinerImpl() override;
+	UE_API virtual TSharedPtr<FTrackAreaViewModel> CreateTrackAreaImpl() override;
+	UE_API virtual TSharedPtr<FSequencerCoreSelection> CreateSelectionImpl() override;
+	UE_API virtual USequencerScriptingLayer* CreateScriptingLayerImpl() override;
+	UE_API virtual bool IsReadOnly() const override;
 
-	void OnTrackAreaHotspotChanged(TSharedPtr<ITrackAreaHotspot> NewHotspot);
+	UE_API void OnTrackAreaHotspotChanged(TSharedPtr<ITrackAreaHotspot> NewHotspot);
 
 protected:
 
@@ -105,3 +107,4 @@ protected:
 
 } // namespace UE::Sequencer
 
+#undef UE_API
diff --git a/Engine/Source/Editor/Sequencer/Public/MVVM/ViewModels/TrackModel.h b/Engine/Source/Editor/Sequencer/Public/MVVM/ViewModels/TrackModel.h
--- a/Engine/Source/Editor/Sequencer/Public/MVVM/ViewModels/TrackModel.h
+++ b/Engine/Source/Editor/Sequencer/Public/MVVM/ViewModels/TrackModel.h
@@ -24,6 +24,8 @@
 #include "UObject/WeakObjectPtr.h"
 #include "UObject/WeakObjectPtrTemplates.h"
 
+#define UE_API SEQUENCER_API
+
 class ISequencer;
 class ISequencerSection;
 class ISequencerTrackEditor;
@@ -36,7 +38,7 @@ namespace Sequencer
 
 class FSectionModel;
 
-class SEQUENCER_API FTrackModel
+class FTrackModel
 	: public FEvaluableOutlinerItemModel
 	, public IRenameableExtension
 	, public IResizableExtension
@@ -53,7 +55,7 @@ class SEQUENCER_API FTrackModel
 {
 public:
 
-	UE_SEQUENCER_DECLARE_CASTABLE(FTrackModel
+	UE_SEQUENCER_DECLARE_CASTABLE_API(UE_API, FTrackModel
 		, FEvaluableOutlinerItemModel
 		, IRenameableExtension
 		, IResizableExtension
@@ -66,97 +68,97 @@ public:
 		, IDeletableExtension
 		, IConditionableExtension);
 
-	explicit FTrackModel(UMovieSceneTrack* Track);
-	~FTrackModel();
+	UE_API explicit FTrackModel(UMovieSceneTrack* Track);
+	UE_API ~FTrackModel();
 
-	FViewModelChildren GetTopLevelChannels();
+	UE_API FViewModelChildren GetTopLevelChannels();
 
 public:
 
-	static EViewModelListType GetTopLevelChannelType();
-	static EViewModelListType GetTopLevelChannelGroupType();
+	static UE_API EViewModelListType GetTopLevelChannelType();
+	static UE_API EViewModelListType GetTopLevelChannelGroupType();
 
 	/*~ ITrackExtension */
-	UMovieSceneTrack* GetTrack() const override;
-	int32 GetRowIndex() const override;
-	FViewModelChildren GetSectionModels() override;
+	UE_API UMovieSceneTrack* GetTrack() const override;
+	UE_API int32 GetRowIndex() const override;
+	UE_API FViewModelChildren GetSectionModels() override;
 
-	TSharedPtr<ISequencerTrackEditor> GetTrackEditor() const override;
+	UE_API TSharedPtr<ISequencerTrackEditor> GetTrackEditor() const override;
 
 	/*~ ISignedObjectEventHandler */
-	void OnModifiedDirectly(UMovieSceneSignedObject*) override;
-	void OnModifiedIndirectly(UMovieSceneSignedObject*) override;
+	UE_API void OnModifiedDirectly(UMovieSceneSignedObject*) override;
+	UE_API void OnModifiedIndirectly(UMovieSceneSignedObject*) override;
 
 	/* IDeferredSignedObjectFlushSignal */
-	virtual void OnDeferredModifyFlush() override;
+	UE_API virtual void OnDeferredModifyFlush() override;
 
 	/*~ IOutlinerExtension */
-	FOutlinerSizing GetOutlinerSizing() const override;
-	TSharedPtr<SWidget> CreateOutlinerViewForColumn(const FCreateOutlinerViewParams& InParams, const FName& InColumnName) override;
-	FSlateFontInfo GetLabelFont() const override;
-	const FSlateBrush* GetIconBrush() const override;
-	FText GetLabel() const override;
-	FSlateColor GetLabelColor() const override;
-	FText GetLabelToolTipText() const override;
+	UE_API FOutlinerSizing GetOutlinerSizing() const override;
+	UE_API TSharedPtr<SWidget> CreateOutlinerViewForColumn(const FCreateOutlinerViewParams& InParams, const FName& InColumnName) override;
+	UE_API FSlateFontInfo GetLabelFont() const override;
+	UE_API const FSlateBrush* GetIconBrush() const override;
+	UE_API FText GetLabel() const override;
+	UE_API FSlateColor GetLabelColor() const override;
+	UE_API FText GetLabelToolTipText() const override;
 	
 	/*~ IDimmableExtension */
-	bool IsDimmed() const override;
+	UE_API bool IsDimmed() const override;
 
 	/*~ IResizableExtension */
-	bool IsResizable() const override;
-	void Resize(float NewSize) override;
+	UE_API bool IsResizable() const override;
+	UE_API void Resize(float NewSize) override;
 
 	/*~ ILockableExtension Interface */
-	ELockableLockState GetLockState() const override;
-	void SetIsLocked(bool bIsLocked) override;
+	UE_API ELockableLockState GetLockState() const override;
+	UE_API void SetIsLocked(bool bIsLocked) override;
 
 	/*~ IConditionableExtension Interface */
-	const UMovieSceneCondition* GetCondition() const override;
-	EConditionableConditionState GetConditionState() const override;
-	void SetConditionEditorForceTrue(bool bEditorForceTrue) override;
+	UE_API const UMovieSceneCondition* GetCondition() const override;
+	UE_API EConditionableConditionState GetConditionState() const override;
+	UE_API void SetConditionEditorForceTrue(bool bEditorForceTrue) override;
 
 	/*~ ITrackAreaExtension */
-	FTrackAreaParameters GetTrackAreaParameters() const override;
-	FViewModelVariantIterator GetTrackAreaModelList() const override;
-	FViewModelVariantIterator GetTopLevelChildTrackAreaModels() const override;
+	UE_API FTrackAreaParameters GetTrackAreaParameters() const override;
+	UE_API FViewModelVariantIterator GetTrackAreaModelList() const override;
+	UE_API FViewModelVariantIterator GetTopLevelChildTrackAreaModels() const override;
 
 	/*~ ICurveEditorTreeItem */
-	void CreateCurveModels(TArray<TUniquePtr<FCurveModel>>& OutCurveModels) override;
+	UE_API void CreateCurveModels(TArray<TUniquePtr<FCurveModel>>& OutCurveModels) override;
 
 	/*~ IGroupableIdentifier */
-	void GetIdentifierForGrouping(TStringBuilder<128>& OutString) const override;
+	UE_API void GetIdentifierForGrouping(TStringBuilder<128>& OutString) const override;
 
 	/*~ IRenameableExtension */
-	bool CanRename() const override;
-	void Rename(const FText& NewName) override;
-	bool IsRenameValidImpl(const FText& NewName, FText& OutErrorMessage) const override;
+	UE_API bool CanRename() const override;
+	UE_API void Rename(const FText& NewName) override;
+	UE_API bool IsRenameValidImpl(const FText& NewName, FText& OutErrorMessage) const override;
 
 	/*~ ISortableExtension */
-	void SortChildren() override;
-	FSortingKey GetSortingKey() const override;
-	void SetCustomOrder(int32 InCustomOrder) override;
+	UE_API void SortChildren() override;
+	UE_API FSortingKey GetSortingKey() const override;
+	UE_API void SetCustomOrder(int32 InCustomOrder) override;
 
 	/*~ IDraggableOutlinerExtension */
-	bool CanDrag() const override;
+	UE_API bool CanDrag() const override;
 
 	/*~ FOutlinerItemModel */
-	bool HasCurves() const override;
-	void BuildContextMenu(FMenuBuilder& MenuBuilder) override;
-	bool GetDefaultExpansionState() const override;
-	void BuildSidebarMenu(FMenuBuilder& MenuBuilder) override;
+	UE_API bool HasCurves() const override;
+	UE_API void BuildContextMenu(FMenuBuilder& MenuBuilder) override;
+	UE_API bool GetDefaultExpansionState() const override;
+	UE_API void BuildSidebarMenu(FMenuBuilder& MenuBuilder) override;
 
 	/*~ IDeletableExtension */
-	bool CanDelete(FText* OutErrorMessage) const override;
-	void Delete() override;
+	UE_API bool CanDelete(FText* OutErrorMessage) const override;
+	UE_API void Delete() override;
 
 	/*~ FViewModel interface */
-	virtual void OnConstruct() override;
+	UE_API virtual void OnConstruct() override;
 
 private:
 
-	void ForceUpdate();
+	UE_API void ForceUpdate();
 
-	bool FindBoundObjects(TArray<UObject*>& OutBoundObjects) const;
+	UE_API bool FindBoundObjects(TArray<UObject*>& OutBoundObjects) const;
 
 	/** A second children list for the sections inside this track */
 	FViewModelListHead SectionList;
@@ -176,3 +178,4 @@ private:
 } // namespace Sequencer
 } // namespace UE
 
+#undef UE_API
diff --git a/Engine/Source/Editor/Sequencer/Public/MVVM/ViewModels/TrackRowModel.h b/Engine/Source/Editor/Sequencer/Public/MVVM/ViewModels/TrackRowModel.h
--- a/Engine/Source/Editor/Sequencer/Public/MVVM/ViewModels/TrackRowModel.h
+++ b/Engine/Source/Editor/Sequencer/Public/MVVM/ViewModels/TrackRowModel.h
@@ -15,6 +15,8 @@
 #include "UObject/WeakObjectPtr.h"
 #include "UObject/WeakObjectPtrTemplates.h"
 
+#define UE_API SEQUENCER_API
+
 class ISequencer;
 class ISequencerTrackEditor;
 class UMovieSceneTrack;
@@ -26,7 +28,7 @@ namespace Sequencer
 
 class FSectionModel;
 
-class SEQUENCER_API FTrackRowModel
+class FTrackRowModel
 	: public FEvaluableOutlinerItemModel
 	, public ITrackAreaExtension
 	, public ILockableExtension
@@ -38,7 +40,7 @@ class SEQUENCER_API FTrackRowModel
 {
 public:
 
-	UE_SEQUENCER_DECLARE_CASTABLE(FTrackRowModel
+	UE_SEQUENCER_DECLARE_CASTABLE_API(UE_API, FTrackRowModel
 		, FEvaluableOutlinerItemModel
 		, ITrackAreaExtension
 		, ILockableExtension
@@ -48,61 +50,61 @@ public:
 		, IDeletableExtension
 		, IConditionableExtension);
 
-	explicit FTrackRowModel(UMovieSceneTrack* InTrack, int32 InRowIndex);
-	~FTrackRowModel();
+	UE_API explicit FTrackRowModel(UMovieSceneTrack* InTrack, int32 InRowIndex);
+	UE_API ~FTrackRowModel();
 
-	void Initialize();
+	UE_API void Initialize();
 
-	FViewModelChildren GetTopLevelChannels();
+	UE_API FViewModelChildren GetTopLevelChannels();
 
 	/*~ FOutlinerItemModel */
-	void BuildContextMenu(FMenuBuilder& MenuBuilder) override;
-	void BuildSidebarMenu(FMenuBuilder& MenuBuilder) override;
-	void CreateCurveModels(TArray<TUniquePtr<FCurveModel>>& OutCurveModels) override;
+	UE_API void BuildContextMenu(FMenuBuilder& MenuBuilder) override;
+	UE_API void BuildSidebarMenu(FMenuBuilder& MenuBuilder) override;
+	UE_API void CreateCurveModels(TArray<TUniquePtr<FCurveModel>>& OutCurveModels) override;
 
 	/*~ IOutlinerExtension */
-	FOutlinerSizing GetOutlinerSizing() const override;
-	TSharedPtr<SWidget> CreateOutlinerViewForColumn(const FCreateOutlinerViewParams& InParams, const FName& InColumnName) override;
-	FText GetLabel() const override;
-	FSlateColor GetLabelColor() const override;
-	FSlateFontInfo GetLabelFont() const override;
-	const FSlateBrush* GetIconBrush() const override;
+	UE_API FOutlinerSizing GetOutlinerSizing() const override;
+	UE_API TSharedPtr<SWidget> CreateOutlinerViewForColumn(const FCreateOutlinerViewParams& InParams, const FName& InColumnName) override;
+	UE_API FText GetLabel() const override;
+	UE_API FSlateColor GetLabelColor() const override;
+	UE_API FSlateFontInfo GetLabelFont() const override;
+	UE_API const FSlateBrush* GetIconBrush() const override;
 
 	/*~ IRenameableExtension */
-	bool CanRename() const override;
-	void Rename(const FText& NewName) override;
-	bool IsRenameValidImpl(const FText& NewName, FText& OutErrorMessage) const override;
+	UE_API bool CanRename() const override;
+	UE_API void Rename(const FText& NewName) override;
+	UE_API bool IsRenameValidImpl(const FText& NewName, FText& OutErrorMessage) const override;
 
 	/*~ IResizableExtension */
-	bool IsResizable() const override;
-	void Resize(float NewSize) override;
+	UE_API bool IsResizable() const override;
+	UE_API void Resize(float NewSize) override;
 
 	/*~ ITrackExtension */
-	UMovieSceneTrack* GetTrack() const override;
-	int32 GetRowIndex() const override;
-	FViewModelChildren GetSectionModels() override;
+	UE_API UMovieSceneTrack* GetTrack() const override;
+	UE_API int32 GetRowIndex() const override;
+	UE_API FViewModelChildren GetSectionModels() override;
 	TSharedPtr<ISequencerTrackEditor> GetTrackEditor() const override { return TrackEditor; }
 
 	/*~ ITrackAreaExtension */
-	FTrackAreaParameters GetTrackAreaParameters() const override;
-	FViewModelVariantIterator GetTrackAreaModelList() const override;
-	FViewModelVariantIterator GetTopLevelChildTrackAreaModels() const override;
+	UE_API FTrackAreaParameters GetTrackAreaParameters() const override;
+	UE_API FViewModelVariantIterator GetTrackAreaModelList() const override;
+	UE_API FViewModelVariantIterator GetTopLevelChildTrackAreaModels() const override;
 
 	/*~ IDimmableExtension */
-	bool IsDimmed() const override;
+	UE_API bool IsDimmed() const override;
 
 	/*~ IDeletableExtension */
-	bool CanDelete(FText* OutErrorMessage) const override;
-	void Delete() override;
+	UE_API bool CanDelete(FText* OutErrorMessage) const override;
+	UE_API void Delete() override;
 
 	/*~ ILockableExtension Interface */
-	ELockableLockState GetLockState() const override;
-	void SetIsLocked(bool bIsLocked) override;
+	UE_API ELockableLockState GetLockState() const override;
+	UE_API void SetIsLocked(bool bIsLocked) override;
 
 	/*~ IConditionableExtension Interface */
-	const UMovieSceneCondition* GetCondition() const override;
-	EConditionableConditionState GetConditionState() const override;
-	void SetConditionEditorForceTrue(bool bEditorForceTrue) override;
+	UE_API const UMovieSceneCondition* GetCondition() const override;
+	UE_API EConditionableConditionState GetConditionState() const override;
+	UE_API void SetConditionEditorForceTrue(bool bEditorForceTrue) override;
 
 private:
 
@@ -120,3 +122,4 @@ private:
 } // namespace Sequencer
 } // namespace UE
 
+#undef UE_API
diff --git a/Engine/Source/Editor/Sequencer/Public/MVVM/Views/SChannelView.h b/Engine/Source/Editor/Sequencer/Public/MVVM/Views/SChannelView.h
--- a/Engine/Source/Editor/Sequencer/Public/MVVM/Views/SChannelView.h
+++ b/Engine/Source/Editor/Sequencer/Public/MVVM/Views/SChannelView.h
@@ -11,6 +11,8 @@
 
 #include "Templates/SharedPointer.h"
 
+#define UE_API SEQUENCER_API
+
 class FSequencer;
 struct FSequencerSelectedKey;
 
@@ -73,7 +75,7 @@ struct FChannelViewKeyCachedState
 /**
  * Widget for displaying a channel within a section
  */
-class SEQUENCER_API SChannelView
+class SChannelView
 	: public STrackAreaLaneView
 {
 public:
@@ -82,27 +84,27 @@ public:
 		SLATE_DEFAULT_SLOT(FArguments, Content)
 	SLATE_END_ARGS()
 
-	void Construct(const FArguments& InArgs, const FViewModelPtr& InViewModel, TSharedPtr<STrackAreaView> InTrackAreaView);
+	UE_API void Construct(const FArguments& InArgs, const FViewModelPtr& InViewModel, TSharedPtr<STrackAreaView> InTrackAreaView);
 
 	/*~ SWidget */
-	virtual int32 OnPaint(const FPaintArgs& Args, const FGeometry& AllottedGeometry, const FSlateRect& MyCullingRect, FSlateWindowElementList& OutDrawElements, int32 LayerId, const FWidgetStyle& InWidgetStyle, bool bParentEnabled) const override;
-	virtual FReply OnMouseButtonDown(const FGeometry& MyGeometry, const FPointerEvent& MouseEvent) override;
-	virtual FReply OnMouseButtonDoubleClick(const FGeometry& MyGeometry, const FPointerEvent& MouseEvent) override;
-	virtual FReply OnMouseMove(const FGeometry& MyGeometry, const FPointerEvent& MouseEvent) override;
-	virtual FReply OnMouseButtonUp(const FGeometry& MyGeometry, const FPointerEvent& MouseEvent) override;
-	virtual void OnMouseLeave(const FPointerEvent& MouseEvent) override;
+	UE_API virtual int32 OnPaint(const FPaintArgs& Args, const FGeometry& AllottedGeometry, const FSlateRect& MyCullingRect, FSlateWindowElementList& OutDrawElements, int32 LayerId, const FWidgetStyle& InWidgetStyle, bool bParentEnabled) const override;
+	UE_API virtual FReply OnMouseButtonDown(const FGeometry& MyGeometry, const FPointerEvent& MouseEvent) override;
+	UE_API virtual FReply OnMouseButtonDoubleClick(const FGeometry& MyGeometry, const FPointerEvent& MouseEvent) override;
+	UE_API virtual FReply OnMouseMove(const FGeometry& MyGeometry, const FPointerEvent& MouseEvent) override;
+	UE_API virtual FReply OnMouseButtonUp(const FGeometry& MyGeometry, const FPointerEvent& MouseEvent) override;
+	UE_API virtual void OnMouseLeave(const FPointerEvent& MouseEvent) override;
 
 protected:
 
-	int32 DrawLane(const FPaintArgs& Args, const FGeometry& AllottedGeometry, const FSlateRect& MyCullingRect, FSlateWindowElementList& OutDrawElements, int32 LayerId, const FWidgetStyle& InWidgetStyle, bool bParentEnabled) const;
+	UE_API int32 DrawLane(const FPaintArgs& Args, const FGeometry& AllottedGeometry, const FSlateRect& MyCullingRect, FSlateWindowElementList& OutDrawElements, int32 LayerId, const FWidgetStyle& InWidgetStyle, bool bParentEnabled) const;
 
-	void CreateKeysUnderMouse(const FVector2D& MousePosition, const FGeometry& AllottedGeometry, const TSet<FSequencerSelectedKey>& InPressedKeys, TArray<FSequencerSelectedKey>& OutKeys);
-	void GetKeysUnderMouse(const FVector2D& MousePosition, const FGeometry& AllottedGeometry, TArray<FSequencerSelectedKey>& OutKeys) const;
-	void GetKeysAtPixelX(float LocalMousePixelX, TArray<FSequencerSelectedKey>& OutKeys) const;
+	UE_API void CreateKeysUnderMouse(const FVector2D& MousePosition, const FGeometry& AllottedGeometry, const TSet<FSequencerSelectedKey>& InPressedKeys, TArray<FSequencerSelectedKey>& OutKeys);
+	UE_API void GetKeysUnderMouse(const FVector2D& MousePosition, const FGeometry& AllottedGeometry, TArray<FSequencerSelectedKey>& OutKeys) const;
+	UE_API void GetKeysAtPixelX(float LocalMousePixelX, TArray<FSequencerSelectedKey>& OutKeys) const;
 
-	TViewModelPtr<FSectionModel> GetSection() const;
+	UE_API TViewModelPtr<FSectionModel> GetSection() const;
 
-	TSharedPtr<FSequencer> LegacyGetSequencer() const;
+	UE_API TSharedPtr<FSequencer> LegacyGetSequencer() const;
 
 protected:
 
@@ -113,3 +115,5 @@ protected:
 };
 
 } // namespace UE::Sequencer
+
+#undef UE_API
diff --git a/Engine/Source/Editor/Sequencer/Public/MovieSceneSequenceEditor.h b/Engine/Source/Editor/Sequencer/Public/MovieSceneSequenceEditor.h
--- a/Engine/Source/Editor/Sequencer/Public/MovieSceneSequenceEditor.h
+++ b/Engine/Source/Editor/Sequencer/Public/MovieSceneSequenceEditor.h
@@ -5,10 +5,12 @@
 #include "CoreTypes.h"
 #include "UObject/NameTypes.h"
 
+#define UE_API SEQUENCER_API
+
 class UBlueprint;
 class UMovieSceneSequence;
 
-struct SEQUENCER_API FMovieSceneSequenceEditor
+struct FMovieSceneSequenceEditor
 {
 	virtual ~FMovieSceneSequenceEditor(){}
 
@@ -19,7 +21,7 @@ struct SEQUENCER_API FMovieSceneSequenceEditor
 	 * @param Sequence        The sequence to get an editor for
 	 * @return The sequence editor ptr, or null if one is not available for the specified type of sequence
 	 */
-	static FMovieSceneSequenceEditor* Find(UMovieSceneSequence* InSequence);
+	static UE_API FMovieSceneSequenceEditor* Find(UMovieSceneSequence* InSequence);
 
 
 
@@ -28,7 +30,7 @@ struct SEQUENCER_API FMovieSceneSequenceEditor
 	 *
 	 * @param Sequence        The sequence to test
 	 */
-	bool SupportsEvents(UMovieSceneSequence* InSequence) const;
+	UE_API bool SupportsEvents(UMovieSceneSequence* InSequence) const;
 
 
 	/**
@@ -37,7 +39,7 @@ struct SEQUENCER_API FMovieSceneSequenceEditor
 	 * @param Sequence        The sequence to access the director blueprint for
 	 * @return The sequence's director blueprint or nullptr if it does not have one (or cannot)
 	 */
-	UBlueprint* FindDirectorBlueprint(UMovieSceneSequence* Sequence) const;
+	UE_API UBlueprint* FindDirectorBlueprint(UMovieSceneSequence* Sequence) const;
 
 	/**
 	 * Access the director blueprint for the specified sequence
@@ -45,7 +47,7 @@ struct SEQUENCER_API FMovieSceneSequenceEditor
 	 * @param Sequence        The sequence to access the director blueprint for
 	 * @return The sequence's director blueprint or nullptr if it does not have one (or cannot)
 	 */
-	UBlueprint* GetOrCreateDirectorBlueprint(UMovieSceneSequence* Sequence) const;
+	UE_API UBlueprint* GetOrCreateDirectorBlueprint(UMovieSceneSequence* Sequence) const;
 
 private:
 
@@ -54,3 +56,4 @@ private:
 	virtual UBlueprint* CreateBlueprintForSequence(UMovieSceneSequence* InSequence) const { return nullptr; }
 };
 
+#undef UE_API
diff --git a/Engine/Source/Editor/Sequencer/Public/MovieSceneTrackEditor.h b/Engine/Source/Editor/Sequencer/Public/MovieSceneTrackEditor.h
--- a/Engine/Source/Editor/Sequencer/Public/MovieSceneTrackEditor.h
+++ b/Engine/Source/Editor/Sequencer/Public/MovieSceneTrackEditor.h
@@ -13,6 +13,8 @@
 #include "MovieScene.h"
 #include "MovieSceneSequence.h"
 
+#define UE_API SEQUENCER_API
+
 class FMenuBuilder;
 class FPaintArgs;
 class FSlateWindowElementList;
@@ -78,39 +80,39 @@ DECLARE_DELEGATE_RetVal_OneParam(bool, FCanKeyProperty, FFrameNumber)
  *
  * @todo Sequencer Interface needs cleanup
  */
-class SEQUENCER_API FMovieSceneTrackEditor
+class FMovieSceneTrackEditor
 	: public TSharedFromThis<FMovieSceneTrackEditor>
 	, public ISequencerTrackEditor
 {
 public:
 
 	/** Constructor */
-	FMovieSceneTrackEditor(TSharedRef<ISequencer> InSequencer);
+	UE_API FMovieSceneTrackEditor(TSharedRef<ISequencer> InSequencer);
 
 	/** Destructor */
-	virtual ~FMovieSceneTrackEditor();
+	UE_API virtual ~FMovieSceneTrackEditor();
 
 public:
 
 	/** @return The current movie scene */
-	UMovieSceneSequence* GetMovieSceneSequence() const;
+	UE_API UMovieSceneSequence* GetMovieSceneSequence() const;
 
 	/**
 	 * @return The current local time at which we should add a key
 	 */
-	FFrameNumber GetTimeForKey();
+	UE_API FFrameNumber GetTimeForKey();
 	
 	/**
 	 * Initiate keying when there is more than one object to key (ie. drag and drop assets). 
 	 * This will allow for different behaviors, ie. keying all at the current time or keying 
 	 * one after the other.
 	 */
-	static void BeginKeying(FFrameNumber InFrameNumber);
-	static void EndKeying();
+	static UE_API void BeginKeying(FFrameNumber InFrameNumber);
+	static UE_API void EndKeying();
 
-	void UpdatePlaybackRange();
+	UE_API void UpdatePlaybackRange();
 
-	void AnimatablePropertyChanged( FOnKeyProperty OnKeyProperty );
+	UE_API void AnimatablePropertyChanged( FOnKeyProperty OnKeyProperty );
 
 	struct FFindOrCreateHandleResult
 	{
@@ -118,7 +120,7 @@ public:
 		bool bWasCreated;
 	};
 	
-	FFindOrCreateHandleResult FindOrCreateHandleToObject( UObject* Object, bool bCreateHandleIfMissing = true, const FName& CreatedFolderName = NAME_None );
+	UE_API FFindOrCreateHandleResult FindOrCreateHandleToObject( UObject* Object, bool bCreateHandleIfMissing = true, const FName& CreatedFolderName = NAME_None );
 
 	struct FFindOrCreateTrackResult
 	{
@@ -128,7 +130,7 @@ public:
 		bool bWasCreated;
 	};
 
-	FFindOrCreateTrackResult FindOrCreateTrackForObject( const FGuid& ObjectHandle, TSubclassOf<UMovieSceneTrack> TrackClass, FName PropertyName = NAME_None, bool bCreateTrackIfMissing = true );
+	UE_API FFindOrCreateTrackResult FindOrCreateTrackForObject( const FGuid& ObjectHandle, TSubclassOf<UMovieSceneTrack> TrackClass, FName PropertyName = NAME_None, bool bCreateTrackIfMissing = true );
 
 	template<typename TrackClass>
 	struct FFindOrCreateRootTrackResult
@@ -165,38 +167,38 @@ public:
 	}
 
 	/** @return The sequencer bound to this handler */
-	TSharedPtr<ISequencer> GetSequencer() const;
+	UE_API TSharedPtr<ISequencer> GetSequencer() const;
 
 public:
 
 	// ISequencerTrackEditor interface
 
-	virtual FText GetDisplayName() const override;
+	UE_API virtual FText GetDisplayName() const override;
 
-	virtual UMovieSceneTrack* AddTrack(UMovieScene* FocusedMovieScene, const FGuid& ObjectHandle, TSubclassOf<class UMovieSceneTrack> TrackClass, FName UniqueTypeName) override;
+	UE_API virtual UMovieSceneTrack* AddTrack(UMovieScene* FocusedMovieScene, const FGuid& ObjectHandle, TSubclassOf<class UMovieSceneTrack> TrackClass, FName UniqueTypeName) override;
 
-	virtual TSharedRef<ISequencerSection> MakeSectionInterface(UMovieSceneSection& SectionObject, UMovieSceneTrack& Track, FGuid ObjectBinding) override;
+	UE_API virtual TSharedRef<ISequencerSection> MakeSectionInterface(UMovieSceneSection& SectionObject, UMovieSceneTrack& Track, FGuid ObjectBinding) override;
 
-	virtual void BindCommands(TSharedRef<FUICommandList> SequencerCommandBindings) override;
-	virtual void BuildPinnedAddTrackMenu(FMenuBuilder& MenuBuilder) override;
-	virtual void BuildAddTrackMenu(FMenuBuilder& MenuBuilder) override;
-	virtual void BuildObjectBindingColumnWidgets(TFunctionRef<TSharedRef<SHorizontalBox>()> GetEditBox, const UE::Sequencer::TViewModelPtr<UE::Sequencer::FObjectBindingModel>& ObjectBinding, const UE::Sequencer::FCreateOutlinerViewParams& InParams, const FName& InColumnName) override;
-	virtual void BuildObjectBindingTrackMenu(FMenuBuilder& MenuBuilder, const TArray<FGuid>& ObjectBindings, const UClass* ObjectClass) override;
-	virtual TSharedPtr<SWidget> BuildOutlinerEditWidget(const FGuid& ObjectBinding, UMovieSceneTrack* Track, const FBuildEditWidgetParams& Params) override;
-	virtual TSharedPtr<SWidget> BuildOutlinerEditColumnWidget(const FGuid& ObjectBinding, UMovieSceneTrack* Track, const FBuildEditWidgetParams& Params) override;
-	virtual TSharedPtr<SWidget> BuildOutlinerColumnWidget(const FBuildColumnWidgetParams& Params, const FName& ColumnName) override;
-	virtual void BuildTrackContextMenu(FMenuBuilder& MenuBuilder, UMovieSceneTrack* Track) override;
-	virtual void BuildTrackSidebarMenu(FMenuBuilder& MenuBuilder, UMovieSceneTrack* Track) override;
-	virtual bool HandleAssetAdded(UObject* Asset, const FGuid& TargetObjectGuid) override;
-	virtual bool OnAllowDrop(const FDragDropEvent& DragDropEvent, FSequencerDragDropParams& DragDropParams) override;
-	virtual FReply OnDrop(const FDragDropEvent& DragDropEvent, const FSequencerDragDropParams& DragDropParams) override;
+	UE_API virtual void BindCommands(TSharedRef<FUICommandList> SequencerCommandBindings) override;
+	UE_API virtual void BuildPinnedAddTrackMenu(FMenuBuilder& MenuBuilder) override;
+	UE_API virtual void BuildAddTrackMenu(FMenuBuilder& MenuBuilder) override;
+	UE_API virtual void BuildObjectBindingColumnWidgets(TFunctionRef<TSharedRef<SHorizontalBox>()> GetEditBox, const UE::Sequencer::TViewModelPtr<UE::Sequencer::FObjectBindingModel>& ObjectBinding, const UE::Sequencer::FCreateOutlinerViewParams& InParams, const FName& InColumnName) override;
+	UE_API virtual void BuildObjectBindingTrackMenu(FMenuBuilder& MenuBuilder, const TArray<FGuid>& ObjectBindings, const UClass* ObjectClass) override;
+	UE_API virtual TSharedPtr<SWidget> BuildOutlinerEditWidget(const FGuid& ObjectBinding, UMovieSceneTrack* Track, const FBuildEditWidgetParams& Params) override;
+	UE_API virtual TSharedPtr<SWidget> BuildOutlinerEditColumnWidget(const FGuid& ObjectBinding, UMovieSceneTrack* Track, const FBuildEditWidgetParams& Params) override;
+	UE_API virtual TSharedPtr<SWidget> BuildOutlinerColumnWidget(const FBuildColumnWidgetParams& Params, const FName& ColumnName) override;
+	UE_API virtual void BuildTrackContextMenu(FMenuBuilder& MenuBuilder, UMovieSceneTrack* Track) override;
+	UE_API virtual void BuildTrackSidebarMenu(FMenuBuilder& MenuBuilder, UMovieSceneTrack* Track) override;
+	UE_API virtual bool HandleAssetAdded(UObject* Asset, const FGuid& TargetObjectGuid) override;
+	UE_API virtual bool OnAllowDrop(const FDragDropEvent& DragDropEvent, FSequencerDragDropParams& DragDropParams) override;
+	UE_API virtual FReply OnDrop(const FDragDropEvent& DragDropEvent, const FSequencerDragDropParams& DragDropParams) override;
 
-	virtual void OnInitialize() override;
-	virtual void OnRelease() override;
+	UE_API virtual void OnInitialize() override;
+	UE_API virtual void OnRelease() override;
 
 	virtual bool SupportsType( TSubclassOf<class UMovieSceneTrack> TrackClass ) const override = 0;
 	virtual bool SupportsSequence(UMovieSceneSequence* InSequence) const override { return true; }
-	virtual void Tick(float DeltaTime) override;
+	UE_API virtual void Tick(float DeltaTime) override;
 
 protected:
 
@@ -205,7 +207,7 @@ protected:
 	 *
 	 * @return Focused movie scene, or nullptr if no movie scene is focused.
 	 */
-	UMovieScene* GetFocusedMovieScene() const;
+	UE_API UMovieScene* GetFocusedMovieScene() const;
 
 private:
 
@@ -213,8 +215,10 @@ private:
 	TWeakPtr<ISequencer> WeakSequencerPtr;
 
 	/** The key time to use during a multi key operation. Only used if bKeying is true */
-	static TOptional<FFrameNumber> NextKeyTime;
+	static UE_API TOptional<FFrameNumber> NextKeyTime;
 
 	/** Indicates whether we're currently in a keying operation where multiple keys may be created (ie. drag and drop) */
-	static bool bKeying;
+	static UE_API bool bKeying;
 };
+
+#undef UE_API
diff --git a/Engine/Source/Editor/Sequencer/Public/SequencerAddKeyOperation.h b/Engine/Source/Editor/Sequencer/Public/SequencerAddKeyOperation.h
--- a/Engine/Source/Editor/Sequencer/Public/SequencerAddKeyOperation.h
+++ b/Engine/Source/Editor/Sequencer/Public/SequencerAddKeyOperation.h
@@ -12,6 +12,8 @@
 #include "SequencerKeyParams.h"
 #include "Templates/SharedPointer.h"
 
+#define UE_API SEQUENCER_API
+
 class IKeyArea;
 class ISequencer;
 class ISequencerSection;
@@ -34,26 +36,26 @@ class IOutlinerExtension;
  * Temporary structure used for consistent add-key behavior for a set of display nodes
  * Ultimately the operation will call ISequencerTrackEditor::ProcessKeyOperation for each track editor that needs to add keys.
  */
-struct SEQUENCER_API FAddKeyOperation
+struct FAddKeyOperation
 {
 	/**
 	 * Construct an operation from any set of display nodes. Each node in the set will receive keys for all decendant key areas.
 	 *
 	 * @param InModels A set of all the models to key
 	 */
-	static FAddKeyOperation FromNodes(const TSet<TWeakViewModelPtr<IOutlinerExtension>>& InModels);
+	static UE_API FAddKeyOperation FromNodes(const TSet<TWeakViewModelPtr<IOutlinerExtension>>& InModels);
 
 
 	/**
 	 * Construct an operation from a single display node. Every key area underneath this node will receive keys.
 	 */
-	static FAddKeyOperation FromNode(TWeakPtr<FViewModel> InModel);
+	static UE_API FAddKeyOperation FromNode(TWeakPtr<FViewModel> InModel);
 
 
 	/**
 	 * Construct an operation from some key areas on a track.
 	 */
-	static FAddKeyOperation FromKeyAreas(ISequencerTrackEditor* TrackEditor, const TArrayView<TSharedRef<IKeyArea>> InKeyAreas);
+	static UE_API FAddKeyOperation FromKeyAreas(ISequencerTrackEditor* TrackEditor, const TArrayView<TSharedRef<IKeyArea>> InKeyAreas);
 
 
 	/**
@@ -63,7 +65,7 @@ struct SEQUENCER_API FAddKeyOperation
 	 * @param InSequencer   The sequencer instance that is performing this operations
 	 * &param OutResults    The key handles that were added with this operation
 	 */
-	void Commit(FFrameNumber InKeyTime, ISequencer& InSequencer, TArray<FAddKeyResult> *OutResults = nullptr);
+	UE_API void Commit(FFrameNumber InKeyTime, ISequencer& InSequencer, TArray<FAddKeyResult> *OutResults = nullptr);
 
 private:
 
@@ -72,7 +74,7 @@ private:
 	 *
 	 * @param InNodes     A set of nodes to add to this operation that contains no child nodes
 	 */
-	void AddPreFilteredNodes(TArrayView<const TWeakPtr<FViewModel>> InNodes);
+	UE_API void AddPreFilteredNodes(TArrayView<const TWeakPtr<FViewModel>> InNodes);
 
 
 	/**
@@ -81,7 +83,7 @@ private:
 	 * @param InTrackEditor         The track editor
 	 * @param InKeyAnythingBeneath  A node to search within for key areas
 	 */
-	bool ConsiderKeyableAreas(TSharedPtr<ITrackExtension> InTrackModel, FViewModelPtr InKeyAnythingBeneath);
+	UE_API bool ConsiderKeyableAreas(TSharedPtr<ITrackExtension> InTrackModel, FViewModelPtr InKeyAnythingBeneath);
 
 
 	/**
@@ -90,7 +92,7 @@ private:
 	 * @param InTrackNode         The current track node
 	 * @param InKeyAreaNode       The key area node to add key areas from
 	 */
-	bool ProcessKeyArea(TSharedPtr<ITrackExtension> InTrackModel, TViewModelPtr<FChannelGroupModel> InChannelGroupModel);
+	UE_API bool ProcessKeyArea(TSharedPtr<ITrackExtension> InTrackModel, TViewModelPtr<FChannelGroupModel> InChannelGroupModel);
 
 
 	/**
@@ -99,13 +101,13 @@ private:
 	 * @param InTrackEditor       The track editor responsible for the key area
 	 * @param InKeyArea           The key area to add
 	 */
-	bool ProcessKeyArea(ISequencerTrackEditor* InTrackEditor, TSharedPtr<IKeyArea> InKeyArea);
+	UE_API bool ProcessKeyArea(ISequencerTrackEditor* InTrackEditor, TSharedPtr<IKeyArea> InKeyArea);
 
 
 	/**
 	 * Retrieve the operation that relates to a specific track editor instance
 	 */
-	FKeyOperation& GetTrackOperation(ISequencerTrackEditor* TrackEditor);
+	UE_API FKeyOperation& GetTrackOperation(ISequencerTrackEditor* TrackEditor);
 
 private:
 
@@ -118,3 +120,4 @@ private:
 } // namespace Sequencer
 } // namespace UE
 
+#undef UE_API
diff --git a/Engine/Source/Editor/Sequencer/Public/SequencerBakingSetupRestore.h b/Engine/Source/Editor/Sequencer/Public/SequencerBakingSetupRestore.h
--- a/Engine/Source/Editor/Sequencer/Public/SequencerBakingSetupRestore.h
+++ b/Engine/Source/Editor/Sequencer/Public/SequencerBakingSetupRestore.h
@@ -6,17 +6,19 @@
 #include "Templates/SharedPointer.h"
 #include "Misc/Optional.h"
 
+#define UE_API SEQUENCER_API
+
 class ISequencer;
 
 namespace UE::Sequencer
 {
 
-	class SEQUENCER_API FSequencerBakingSetupRestore 
+	class FSequencerBakingSetupRestore 
 	{
 	public:
 		FSequencerBakingSetupRestore() = delete;
-		FSequencerBakingSetupRestore(TSharedPtr<ISequencer>& SequencerPtr);
-		~FSequencerBakingSetupRestore();
+		UE_API FSequencerBakingSetupRestore(TSharedPtr<ISequencer>& SequencerPtr);
+		UE_API ~FSequencerBakingSetupRestore();
 		
 	private:
 		TWeakPtr<ISequencer> WeakSequencer;
@@ -24,4 +26,6 @@ namespace UE::Sequencer
 		
 	};
 
-} // namespace UE::Sequencer
\ No newline at end of file
+} // namespace UE::Sequencer
+
+#undef UE_API
diff --git a/Engine/Source/Editor/Sequencer/Public/SequencerCustomizationManager.h b/Engine/Source/Editor/Sequencer/Public/SequencerCustomizationManager.h
--- a/Engine/Source/Editor/Sequencer/Public/SequencerCustomizationManager.h
+++ b/Engine/Source/Editor/Sequencer/Public/SequencerCustomizationManager.h
@@ -7,6 +7,8 @@
 #include "Misc/EnumClassFlags.h"
 #include "MVVM/ViewModelPtr.h"
 
+#define UE_API SEQUENCER_API
+
 class FAssetDragDropOp;
 class FClassDragDropOp;
 class FActorDragDropOp;
@@ -124,7 +126,7 @@ private:
 };
 
 /** Class that can figure out what customizations to apply to a given sequence. */
-class SEQUENCER_API ISequencerCustomization
+class ISequencerCustomization
 {
 public:
 	virtual ~ISequencerCustomization() {}
@@ -135,29 +137,29 @@ public:
 DECLARE_DELEGATE_RetVal(ISequencerCustomization*, FOnGetSequencerCustomizationInstance);
 
 /** Manager class for sequencer customizations. */
-class SEQUENCER_API FSequencerCustomizationManager
+class FSequencerCustomizationManager
 {
 public:
-	FSequencerCustomizationManager();
+	UE_API FSequencerCustomizationManager();
 
 	/**
 	 * Registers a sequencer customization against a specific type of sequence.
 	 */
-	void RegisterInstancedSequencerCustomization(const UClass* SequenceClass, FOnGetSequencerCustomizationInstance GetCustomizationDelegate);
+	UE_API void RegisterInstancedSequencerCustomization(const UClass* SequenceClass, FOnGetSequencerCustomizationInstance GetCustomizationDelegate);
 	/**
 	 * Unregisters a sequencer customization that was previously registered against a specific type of sequence.
 	 */
-	void UnregisterInstancedSequencerCustomization(const UClass* SequenceClass);
+	UE_API void UnregisterInstancedSequencerCustomization(const UClass* SequenceClass);
 
 	/**
 	 * Gets the sequencer customizations for the given sequence.
 	 */
-	void GetSequencerCustomizations(UMovieSceneSequence* FocusedSequence, TArray<TUniquePtr<ISequencerCustomization>>& OutCustomizations);
+	UE_API void GetSequencerCustomizations(UMovieSceneSequence* FocusedSequence, TArray<TUniquePtr<ISequencerCustomization>>& OutCustomizations);
 
 	/**
 	 * Determines whether the new edited sequence requires changing the customizations.
 	 */
-	bool NeedsCustomizationChange(const UMovieSceneSequence* OldFocusedSequence, const UMovieSceneSequence* NewFocusedSequence) const;
+	UE_API bool NeedsCustomizationChange(const UMovieSceneSequence* OldFocusedSequence, const UMovieSceneSequence* NewFocusedSequence) const;
 
 private:
 	FSequencerCustomizationManager(const FSequencerCustomizationManager&) = delete;
@@ -172,3 +174,4 @@ private:
 	TArray<FCustomizationRegistryEntry> CustomizationRegistryEntries;
 };
 
+#undef UE_API
diff --git a/Engine/Source/Editor/Sequencer/Public/SequencerKeyCollection.h b/Engine/Source/Editor/Sequencer/Public/SequencerKeyCollection.h
--- a/Engine/Source/Editor/Sequencer/Public/SequencerKeyCollection.h
+++ b/Engine/Source/Editor/Sequencer/Public/SequencerKeyCollection.h
@@ -111,7 +111,7 @@ public:
 	 */
 	SEQUENCER_API TOptional<FFrameNumber> FindFirstKeyInRange(const TRange<FFrameNumber>& Range, EFindKeyDirection Direction, EFindKeyType FindKeyType = EFindKeyType::FKT_All) const;
 	UE_DEPRECATED(5.4, "Please use FindFirstKeyInRange which takes EFindKeyType, FindFirstSectionKeyInRange is no longer supported")
-	SEQUENCER_API TOptional<FFrameNumber> FindFirstSectionKeyInRange(const TRange<FFrameNumber>& Range, EFindKeyDirection Direction) const { return FindFirstKeyInRange(Range, Direction, EFindKeyType::FKT_Sections); }
+	TOptional<FFrameNumber> FindFirstSectionKeyInRange(const TRange<FFrameNumber>& Range, EFindKeyDirection Direction) const { return FindFirstKeyInRange(Range, Direction, EFindKeyType::FKT_Sections); }
 
 	/**
 	 * Get a view of all key times that reside within the specified range
@@ -121,7 +121,7 @@ public:
 	 */
 	SEQUENCER_API TArrayView<const FFrameNumber> GetKeysInRange(const TRange<FFrameNumber>& Range, EFindKeyType FindKeyType = EFindKeyType::FKT_All) const;
 	UE_DEPRECATED(5.4, "Please use GetKeysInRange which takes EFindKeyType, GetSectionKeysInRange is no longer supported")
-	SEQUENCER_API TArrayView<const FFrameNumber> GetSectionKeysInRange(const TRange<FFrameNumber>& Range) const { return GetKeysInRange(Range, EFindKeyType::FKT_Sections); }
+	TArrayView<const FFrameNumber> GetSectionKeysInRange(const TRange<FFrameNumber>& Range) const { return GetKeysInRange(Range, EFindKeyType::FKT_Sections); }
 
 	/**
 	* Search forwards or backwards for the next key from the specified frame number
@@ -131,7 +131,7 @@ public:
 	*/
 	SEQUENCER_API TOptional<FFrameNumber> GetNextKey(FFrameNumber FrameNumber, EFindKeyDirection Direction, const TRange<FFrameNumber>& Range, EFindKeyType FindKeyType = EFindKeyType::FKT_All) const;
 	UE_DEPRECATED(5.4, "Please use GetNextKey which takes EFindKeyType, GetNextSectionKey is no longer supported")
-	SEQUENCER_API TOptional<FFrameNumber> GetNextSectionKey(FFrameNumber FrameNumber, EFindKeyDirection Direction, const TRange<FFrameNumber>& Range) const { return GetNextKey(FrameNumber, Direction, Range, EFindKeyType::FKT_Sections); }
+	TOptional<FFrameNumber> GetNextSectionKey(FFrameNumber FrameNumber, EFindKeyDirection Direction, const TRange<FFrameNumber>& Range) const { return GetNextKey(FrameNumber, Direction, Range, EFindKeyType::FKT_Sections); }
 
 	/**
 	 * Access the signature this collection was generated with
diff --git a/Engine/Source/Editor/Sequencer/Public/SequencerKeyParams.h b/Engine/Source/Editor/Sequencer/Public/SequencerKeyParams.h
--- a/Engine/Source/Editor/Sequencer/Public/SequencerKeyParams.h
+++ b/Engine/Source/Editor/Sequencer/Public/SequencerKeyParams.h
@@ -26,7 +26,7 @@ namespace UE
 namespace Sequencer
 {
 
-struct SEQUENCER_API FAddKeyResult
+struct FAddKeyResult
 {
 	TSharedPtr<IKeyArea> KeyArea;
 	FKeyHandle KeyHandle;
diff --git a/Engine/Source/Editor/Sequencer/Public/SequencerKeyStructGenerator.h b/Engine/Source/Editor/Sequencer/Public/SequencerKeyStructGenerator.h
--- a/Engine/Source/Editor/Sequencer/Public/SequencerKeyStructGenerator.h
+++ b/Engine/Source/Editor/Sequencer/Public/SequencerKeyStructGenerator.h
@@ -24,6 +24,8 @@
 
 #include "SequencerKeyStructGenerator.generated.h"
 
+#define UE_API SEQUENCER_API
+
 class FArrayProperty;
 class FProperty;
 class FSequencerKeyStructGenerator;
@@ -33,16 +35,16 @@ class UObject;
 /**
  * Struct type that is generated from an FMovieSceneChannel type to create a single edit interface for a key/value pair
  */
-UCLASS()
-class SEQUENCER_API UMovieSceneKeyStructType : public UScriptStruct
+UCLASS(MinimalAPI)
+class UMovieSceneKeyStructType : public UScriptStruct
 {
 public:
 	GENERATED_BODY()
 
-	UMovieSceneKeyStructType(const FObjectInitializer& ObjInit);
+	UE_API UMovieSceneKeyStructType(const FObjectInitializer& ObjInit);
 
-	void InitializeStruct(void* InDest, int32 ArrayDim) const override;
-	void DestroyStruct(void* Dest, int32 ArrayDim) const override;
+	UE_API void InitializeStruct(void* InDest, int32 ArrayDim) const override;
+	UE_API void DestroyStruct(void* Dest, int32 ArrayDim) const override;
 
 	/**
 	 * Check whether this generated struct is complete and ready to be used
@@ -75,38 +77,38 @@ private:
 /**
  * Singleton class that is used to create, store and instantiate generated structs for editing single keys on channels
  */
-class SEQUENCER_API FSequencerKeyStructGenerator : FGCObject
+class FSequencerKeyStructGenerator : FGCObject
 {
 public:
 
 	/**
 	 * Access the singlton instance of this class
 	 */
-	static FSequencerKeyStructGenerator& Get();
+	static UE_API FSequencerKeyStructGenerator& Get();
 
 
 	/**
 	 * Allocate a brand new, empty key struct type. Must be fully completed and finalized with FinalizeNewKeyStruct
 	 */
-	static UMovieSceneKeyStructType* AllocateNewKeyStruct();
+	static UE_API UMovieSceneKeyStructType* AllocateNewKeyStruct();
 
 
 	/**
 	 * Allocate a brand new, key struct type, automatically discovering time/value properties from property meta-data. Must be finalized with FinalizeNewKeyStruct
 	 */
-	static UMovieSceneKeyStructType* AllocateNewKeyStruct(UScriptStruct* ChannelType);
+	static UE_API UMovieSceneKeyStructType* AllocateNewKeyStruct(UScriptStruct* ChannelType);
 
 
 	/**
 	 * Finalize the specified struct type by prepending a time property and linking the struct
 	 */
-	static void FinalizeNewKeyStruct(UMovieSceneKeyStructType* InStruct);
+	static UE_API void FinalizeNewKeyStruct(UMovieSceneKeyStructType* InStruct);
 
 
 	/**
 	 * Helper function to locate an array property with the specified meta-data tag
 	 */
-	static FArrayProperty* FindArrayPropertyWithTag(UScriptStruct* ChannelStruct, FName MetaDataTag);
+	static UE_API FArrayProperty* FindArrayPropertyWithTag(UScriptStruct* ChannelStruct, FName MetaDataTag);
 
 public:
 
@@ -117,7 +119,7 @@ public:
 	 * @param InstancedStructName       A unique name to associate with this struct, potentially for later retrieval. May include a specific property postfix such as an object or enum type for dynamic channels.
 	 * @param GeneratedStruct           The generated struct to associate with the name
 	 */
-	void AddGeneratedStruct(FName InstancedStructName, UMovieSceneKeyStructType* Struct);
+	UE_API void AddGeneratedStruct(FName InstancedStructName, UMovieSceneKeyStructType* Struct);
 
 
 	/**
@@ -126,7 +128,7 @@ public:
 	 * @param InstancedStructName       The unique name to look for
 	 * @return A generated struct that was added with the specified name, or nullptr if one does not exist.
 	 */
-	UMovieSceneKeyStructType* FindGeneratedStruct(FName InstancedStructName);
+	UE_API UMovieSceneKeyStructType* FindGeneratedStruct(FName InstancedStructName);
 
 
 	/**
@@ -135,7 +137,7 @@ public:
 	 * @param ChannelType     The struct type of the channel (normally retrieved through ::StaticStruct())
 	 * @return A valid pointer to a UMovieSceneKeyStructType struct type, or nullptr if one could not be found or allocated
 	 */
-	UMovieSceneKeyStructType* DefaultInstanceGeneratedStruct(UScriptStruct* ChannelType);
+	UE_API UMovieSceneKeyStructType* DefaultInstanceGeneratedStruct(UScriptStruct* ChannelType);
 
 public:
 
@@ -161,12 +163,12 @@ private:
 	/**
 	 * Create a new struct instance populated with the time and value for the specified key index, but no OnPropertyChangedEvent initialized.
 	 */
-	TSharedPtr<FStructOnScope> CreateInitialStructInstance(const void* SourceChannel, UMovieSceneKeyStructType* GeneratedStructType, int32 InitialKeyIndex);
+	UE_API TSharedPtr<FStructOnScope> CreateInitialStructInstance(const void* SourceChannel, UMovieSceneKeyStructType* GeneratedStructType, int32 InitialKeyIndex);
 
 	/**
 	 * FGCObject Interface
 	 */
-	virtual void AddReferencedObjects(FReferenceCollector& Collector) override;
+	UE_API virtual void AddReferencedObjects(FReferenceCollector& Collector) override;
 	virtual FString GetReferencerName() const override
 	{
 		return TEXT("FSequencerKeyStructGenerator");
@@ -282,4 +284,6 @@ void FSequencerKeyStructGenerator::CopyInstanceToKey(const TMovieSceneChannelHan
 			}
 		}
 	}
-}
\ No newline at end of file
+}
+
+#undef UE_API
diff --git a/Engine/Source/Editor/Sequencer/Public/SequencerSectionPainter.h b/Engine/Source/Editor/Sequencer/Public/SequencerSectionPainter.h
--- a/Engine/Source/Editor/Sequencer/Public/SequencerSectionPainter.h
+++ b/Engine/Source/Editor/Sequencer/Public/SequencerSectionPainter.h
@@ -10,6 +10,8 @@
 #include "Math/Color.h"
 #include "Templates/SharedPointer.h"
 
+#define UE_API SEQUENCER_API
+
 class FSlateWindowElementList;
 class IKeyArea;
 class UMovieSceneSection;
@@ -37,30 +39,30 @@ struct FKeyAreaElement
 };
 
 /** Class that wraps up common section painting functionality */
-class SEQUENCER_API FSequencerSectionPainter
+class FSequencerSectionPainter
 {
 public:
 	UE_NONCOPYABLE(FSequencerSectionPainter)
 
 	/** Constructor */
-	FSequencerSectionPainter(FSlateWindowElementList& OutDrawElements, const FGeometry& InSectionGeometry, TSharedPtr<UE::Sequencer::FSectionModel> Section);
+	UE_API FSequencerSectionPainter(FSlateWindowElementList& OutDrawElements, const FGeometry& InSectionGeometry, TSharedPtr<UE::Sequencer::FSectionModel> Section);
 
 	/** Virtual destructor */
-	virtual ~FSequencerSectionPainter();
+	UE_API virtual ~FSequencerSectionPainter();
 
 	/** Paint the section background with the specified tint override */
 	virtual int32 PaintSectionBackground(const FLinearColor& Tint) = 0;
 
 	/** Paint the section background with the tint stored on the track */
-	int32 PaintSectionBackground();
+	UE_API int32 PaintSectionBackground();
 
 	/** Get the track that this painter is painting sections for */
-	UMovieSceneTrack* GetTrack() const;
+	UE_API UMovieSceneTrack* GetTrack() const;
 
-	FLinearColor GetSectionColor() const;
+	UE_API FLinearColor GetSectionColor() const;
 
 	/** Blend the specified color with the default track color */
-	static FLinearColor BlendColor(FLinearColor InColor);
+	static UE_API FLinearColor BlendColor(FLinearColor InColor);
 
 public:
 
@@ -100,3 +102,4 @@ public:
 	float GhostAlpha;
 };
 
+#undef UE_API
