Updated headers using UnrealCodeFixup to make sure dllstorage is on methods/staticvars instead of types. Part 5/n

diff --git a/Engine/Source/Editor/Sequencer/Public/SequencerSettings.h b/Engine/Source/Editor/Sequencer/Public/SequencerSettings.h
--- a/Engine/Source/Editor/Sequencer/Public/SequencerSettings.h
+++ b/Engine/Source/Editor/Sequencer/Public/SequencerSettings.h
@@ -13,6 +13,8 @@
 #include "Sidebar/SidebarState.h"
 #include "SequencerSettings.generated.h"
 
+#define UE_API SEQUENCER_API
+
 struct FPropertyChangedEvent;
 enum class EAutoChangeMode : uint8;
 enum class EAllowEditsMode : uint8;
@@ -67,8 +69,8 @@ enum class ESequencerTimeWarpDisplay : uint8
 ENUM_CLASS_FLAGS(ESequencerTimeWarpDisplay)
 
 /** Empty class used to house multiple named USequencerSettings */
-UCLASS()
-class SEQUENCER_API USequencerSettingsContainer
+UCLASS(MinimalAPI)
+class USequencerSettingsContainer
 	: public UObject
 {
 public:
@@ -127,8 +129,8 @@ struct FColumnVisibilitySetting
 };
 
 /** Serializable options for sequencer. */
-UCLASS(config=EditorPerProjectUserSettings, PerObjectConfig)
-class SEQUENCER_API USequencerSettings
+UCLASS(MinimalAPI, config=EditorPerProjectUserSettings, PerObjectConfig)
+class USequencerSettings
 	: public UObject
 {
 public:
@@ -140,355 +142,355 @@ public:
 	DECLARE_MULTICAST_DELEGATE(FOnLoopStateChanged);
 	DECLARE_MULTICAST_DELEGATE(FOnTimeDisplayFormatChanged);
 	
-	virtual void PostEditChangeProperty(FPropertyChangedEvent& PropertyChangedEvent) override;
+	UE_API virtual void PostEditChangeProperty(FPropertyChangedEvent& PropertyChangedEvent) override;
 
 	/** Gets the current auto change mode. */
-	EAutoChangeMode GetAutoChangeMode() const;
+	UE_API EAutoChangeMode GetAutoChangeMode() const;
 	/** Sets the current auto change mode. */
-	void SetAutoChangeMode(EAutoChangeMode AutoChangeMode);
+	UE_API void SetAutoChangeMode(EAutoChangeMode AutoChangeMode);
 
 	/** Gets the current allow edits mode. */
-	EAllowEditsMode GetAllowEditsMode() const;
+	UE_API EAllowEditsMode GetAllowEditsMode() const;
 	/** Sets the current auto-key mode. */
-	void SetAllowEditsMode(EAllowEditsMode AllowEditsMode);
+	UE_API void SetAllowEditsMode(EAllowEditsMode AllowEditsMode);
 	/** Gets the multicast delegate which is run whenever the allow edits mode is changed. */
 	FOnAllowEditsModeChanged& GetOnAllowEditsModeChanged() { return OnAllowEditsModeChangedEvent; }
 
 	/** Returns what channels will get keyed when one channel changes */
-	EKeyGroupMode GetKeyGroupMode() const;
+	UE_API EKeyGroupMode GetKeyGroupMode() const;
 	/** Sets which channels are keyed when a channel is keyed */
-	void SetKeyGroupMode(EKeyGroupMode);
+	UE_API void SetKeyGroupMode(EKeyGroupMode);
 
 	/** Get the default Interpolation type for newly created keyframes if the channel does not already have keyframes */
-	EMovieSceneKeyInterpolation GetKeyInterpolation() const;
+	UE_API EMovieSceneKeyInterpolation GetKeyInterpolation() const;
 	/** Sets default key interpolation for creating new keys on empty channels */
-	void SetKeyInterpolation(EMovieSceneKeyInterpolation InKeyInterpolation);
+	UE_API void SetKeyInterpolation(EMovieSceneKeyInterpolation InKeyInterpolation);
 
 	/** Get initial spawn position. */
-	ESequencerSpawnPosition GetSpawnPosition() const;
+	UE_API ESequencerSpawnPosition GetSpawnPosition() const;
 	/** Set initial spawn position. */
-	void SetSpawnPosition(ESequencerSpawnPosition InSpawnPosition);
+	UE_API void SetSpawnPosition(ESequencerSpawnPosition InSpawnPosition);
 
 	/** Get whether to create spawnable cameras. */
-	bool GetCreateSpawnableCameras() const;
+	UE_API bool GetCreateSpawnableCameras() const;
 	/** Set whether to create spawnable cameras. */
-	void SetCreateSpawnableCameras(bool bInCreateSpawnableCameras);
+	UE_API void SetCreateSpawnableCameras(bool bInCreateSpawnableCameras);
 
 	/** Gets whether or not to show the time range slider. */
-	bool GetShowRangeSlider() const;
+	UE_API bool GetShowRangeSlider() const;
 	/** Sets whether or not to show frame numbers. */
-	void SetShowRangeSlider(bool InbShowRangeSlider);
+	UE_API void SetShowRangeSlider(bool InbShowRangeSlider);
 
 	/** Gets whether or not snapping is enabled. */
-	bool GetIsSnapEnabled() const;
+	UE_API bool GetIsSnapEnabled() const;
 	/** Sets whether or not snapping is enabled. */
-	void SetIsSnapEnabled(bool InbIsSnapEnabled);
+	UE_API void SetIsSnapEnabled(bool InbIsSnapEnabled);
 
 	/** Gets whether or not to snap keys to other elements. */
-	bool GetSnapKeyTimesToElements() const;
+	UE_API bool GetSnapKeyTimesToElements() const;
 	/** Sets whether or not to snap keys to other elements. */
-	void SetSnapKeyTimesToElements(bool InbSnapKeyTimesToElements);
+	UE_API void SetSnapKeyTimesToElements(bool InbSnapKeyTimesToElements);
 
 	/** Gets whether or not to snap sections to other elements. */
-	bool GetSnapSectionTimesToElements() const;
+	UE_API bool GetSnapSectionTimesToElements() const;
 	/** sets whether or not to snap sections to other elements. */
-	void SetSnapSectionTimesToElements( bool InbSnapSectionTimesToElements );
+	UE_API void SetSnapSectionTimesToElements( bool InbSnapSectionTimesToElements );
 
 	/** @return true if keys and sections should be kept within the playback range when moving them */
-	bool GetSnapKeysAndSectionsToPlayRange() const;
+	UE_API bool GetSnapKeysAndSectionsToPlayRange() const;
 	/** Set whether or not keys and sections should be kept within the playback range when moving them */
-	void SetSnapKeysAndSectionsToPlayRange(bool bInSnapKeysAndSectionsToPlayRange);
+	UE_API void SetSnapKeysAndSectionsToPlayRange(bool bInSnapKeysAndSectionsToPlayRange);
 
 	/** Gets whether or not to snap the play time to keys while scrubbing. */
-	bool GetSnapPlayTimeToKeys() const;
+	UE_API bool GetSnapPlayTimeToKeys() const;
 	/** Sets whether or not to snap the play time to keys while scrubbing. */
-	void SetSnapPlayTimeToKeys(bool InbSnapPlayTimeToKeys);
+	UE_API void SetSnapPlayTimeToKeys(bool InbSnapPlayTimeToKeys);
 
 	/** Gets whether or not to snap the play time to section bounds while scrubbing. */
-	bool GetSnapPlayTimeToSections() const;
+	UE_API bool GetSnapPlayTimeToSections() const;
 	/** Sets whether or not to snap the play time to section bounds while scrubbing. */
-	void SetSnapPlayTimeToSections(bool InbSnapPlayTimeToSections);
+	UE_API void SetSnapPlayTimeToSections(bool InbSnapPlayTimeToSections);
 
 	/** Gets whether or not to snap the play time to markers while scrubbing. */
-	bool GetSnapPlayTimeToMarkers() const;
+	UE_API bool GetSnapPlayTimeToMarkers() const;
 	/** Sets whether or not to snap the play time to markers while scrubbing. */
-	void SetSnapPlayTimeToMarkers(bool InbSnapPlayTimeToMarkers);
+	UE_API void SetSnapPlayTimeToMarkers(bool InbSnapPlayTimeToMarkers);
 
 	/** Gets whether or not to snap the play time to the pressed key. */
-	bool GetSnapPlayTimeToPressedKey() const;
+	UE_API bool GetSnapPlayTimeToPressedKey() const;
 	/** Sets whether or not to snap the play time to the pressed key. */
-	void SetSnapPlayTimeToPressedKey(bool InbSnapPlayTimeToPressedKey);
+	UE_API void SetSnapPlayTimeToPressedKey(bool InbSnapPlayTimeToPressedKey);
 
 	/** Gets whether or not to snap the play time to the dragged key. */
-	bool GetSnapPlayTimeToDraggedKey() const;
+	UE_API bool GetSnapPlayTimeToDraggedKey() const;
 	/** Sets whether or not to snap the play time to the dragged key. */
-	void SetSnapPlayTimeToDraggedKey(bool InbSnapPlayTimeToDraggedKey);
+	UE_API void SetSnapPlayTimeToDraggedKey(bool InbSnapPlayTimeToDraggedKey);
 
 	/** Gets the snapping interval for curve values. */
-	float GetCurveValueSnapInterval() const;
+	UE_API float GetCurveValueSnapInterval() const;
 	/** Sets the snapping interval for curve values. */
-	void SetCurveValueSnapInterval(float InCurveValueSnapInterval);
+	UE_API void SetCurveValueSnapInterval(float InCurveValueSnapInterval);
 
 	/** Gets whether or not to snap keys/sections/playhead to the interval which forces whole frames */
-	bool GetForceWholeFrames() const;
+	UE_API bool GetForceWholeFrames() const;
 	/** Sets whether or not to snap keys/sections/playhead to the interval which forces whole frames */
-	void SetForceWholeFrames(bool bInForceWholeFrames);
+	UE_API void SetForceWholeFrames(bool bInForceWholeFrames);
 
 	/** Gets the state for spacing between grid lines */
-	TOptional<float> GetGridSpacing() const;
+	UE_API TOptional<float> GetGridSpacing() const;
 	/** Sets the grid line spacing state */
-	void SetGridSpacing(TOptional<float> InGridSpacing);
+	UE_API void SetGridSpacing(TOptional<float> InGridSpacing);
 
 	/** Gets whether or not to snap curve values to the interval. */
-	bool GetSnapCurveValueToInterval() const;
+	UE_API bool GetSnapCurveValueToInterval() const;
 	/** Sets whether or not to snap curve values to the interval. */
-	void SetSnapCurveValueToInterval(bool InbSnapCurveValueToInterval);
+	UE_API void SetSnapCurveValueToInterval(bool InbSnapCurveValueToInterval);
 
 	/** Gets whether or not to show selected nodes only. */
-	bool GetShowSelectedNodesOnly() const;
+	UE_API bool GetShowSelectedNodesOnly() const;
 	/** Sets whether or not to show selected nodes only. */
-	void SetShowSelectedNodesOnly(bool Visible);
+	UE_API void SetShowSelectedNodesOnly(bool Visible);
 	FOnShowSelectedNodesOnlyChanged& GetOnShowSelectedNodesOnlyChanged() { return OnShowSelectedNodesOnlyChangedEvent; }
 
 	/** Gets whether left mouse drag does marquee select instead of camera orbit and ALT always moves the camera */
-	bool GetLeftMouseDragDoesMarquee() const;
+	UE_API bool GetLeftMouseDragDoesMarquee() const;
 	/** Sets whether left mouse drag does marquee select instead of camera orbit and ALT always moves the camera */
-	void SetLeftMouseDragDoesMarque(bool bDoMarque);
+	UE_API void SetLeftMouseDragDoesMarque(bool bDoMarque);
 
 	/** Get zoom in/out position (mouse position or current time). */
-	ESequencerZoomPosition GetZoomPosition() const;
+	UE_API ESequencerZoomPosition GetZoomPosition() const;
 	/** Set zoom in/out position (mouse position or current time). */
-	void SetZoomPosition(ESequencerZoomPosition InZoomPosition);
+	UE_API void SetZoomPosition(ESequencerZoomPosition InZoomPosition);
 
 	/** Gets whether or not auto-scroll is enabled when playing. */
-	bool GetAutoScrollEnabled() const;
+	UE_API bool GetAutoScrollEnabled() const;
 	/** Sets whether or not auto-scroll is enabled when playing. */
-	void SetAutoScrollEnabled(bool bInAutoScrollEnabled);
+	UE_API void SetAutoScrollEnabled(bool bInAutoScrollEnabled);
 	
 	/** Gets whether or not to link the curve editor time range. */
-	bool GetLinkCurveEditorTimeRange() const;
+	UE_API bool GetLinkCurveEditorTimeRange() const;
 	/** Sets whether or not to link the curve editor time range. */
-	void SetLinkCurveEditorTimeRange(bool InbLinkCurveEditorTimeRange);
+	UE_API void SetLinkCurveEditorTimeRange(bool InbLinkCurveEditorTimeRange);
 
 	/** Gets whether or not to link sequence filters with the curve editor tree */
-	bool GetLinkFiltersWithCurveEditor() const;
+	UE_API bool GetLinkFiltersWithCurveEditor() const;
 	/** Sets whether or not to link sequence filters with the curve editor tree */
-	void SetLinkFiltersWithCurveEditor(bool bInLinkFiltersWithCurveEditor);
+	UE_API void SetLinkFiltersWithCurveEditor(bool bInLinkFiltersWithCurveEditor);
 
 	/** Return true if we are to synchronize the curve editor and sequencer trees */
 	bool ShouldSyncCurveEditorSelection() const { return bSynchronizeCurveEditorSelection; }
 	/** Assign whether we are to synchronize the curve editor and sequencer trees */
-	void SyncCurveEditorSelection(bool bInSynchronizeCurveEditorSelection);
+	UE_API void SyncCurveEditorSelection(bool bInSynchronizeCurveEditorSelection);
 
 	/** Return true if we should filter the curve editor tree to only nodes that are relevant to the current sequencer selection */
 	bool ShouldIsolateToCurveEditorSelection() const { return bIsolateCurveEditorToSelection; }
 	/** Assign whether we should filter the curve editor tree to only nodes that are relevant to the current sequencer selection */
-	void IsolateCurveEditorToSelection(bool bInIsolateCurveEditorToSelection);
+	UE_API void IsolateCurveEditorToSelection(bool bInIsolateCurveEditorToSelection);
 
 	/** Return true if we should filter the curve editor tree to only nodes that are relevant to the current sequencer selection */
 	bool GetCurveEditorVisible() const { return bCurveEditorVisible; }
 	/** Assign whether we should filter the curve editor tree to only nodes that are relevant to the current sequencer selection */
-	void SetCurveEditorVisible(bool bCurveEditorVisible);
+	UE_API void SetCurveEditorVisible(bool bCurveEditorVisible);
 
 	/** @return The config affecting the zooming behavior of Sequencer's Curve Editor. */
 	const FCurveEditorZoomScaleConfig& GetCurveEditorZoomScaling() { return CurveEditorZoomScaling; }
 	/** Sets the config affecting the zooming behaviour of Sequencer's Curve Editor. */
-	void SetCurveEditorZoomScaling(const FCurveEditorZoomScaleConfig& Value);
+	UE_API void SetCurveEditorZoomScaling(const FCurveEditorZoomScaleConfig& Value);
 
 	/** Gets the loop mode. */
-	ESequencerLoopMode GetLoopMode() const;
+	UE_API ESequencerLoopMode GetLoopMode() const;
 	/** Sets the loop mode. */
-	void SetLoopMode(ESequencerLoopMode InLoopMode);
+	UE_API void SetLoopMode(ESequencerLoopMode InLoopMode);
 
 	/** @return true if the cursor reset when navigating in and out of subsequences, false otherwise */
-	bool ShouldResetPlayheadWhenNavigating() const;
+	UE_API bool ShouldResetPlayheadWhenNavigating() const;
 	/** Set whether or not the cursor should be reset when navigating in and out of subsequences */
-	void SetResetPlayheadWhenNavigating(bool bInResetPlayheadWhenNavigating);
+	UE_API void SetResetPlayheadWhenNavigating(bool bInResetPlayheadWhenNavigating);
 
 	/** @return true if the cursor should be kept within the playback (or subsequence/shot) range while scrubbing in sequencer, false otherwise */
-	bool ShouldKeepCursorInPlayRangeWhileScrubbing() const;
+	UE_API bool ShouldKeepCursorInPlayRangeWhileScrubbing() const;
 	/** Set whether or not the cursor should be kept within the playback (or subsequence/shot) range while scrubbing in sequencer */
-	void SetKeepCursorInPlayRangeWhileScrubbing(bool bInKeepCursorInPlayRangeWhileScrubbing);
+	UE_API void SetKeepCursorInPlayRangeWhileScrubbing(bool bInKeepCursorInPlayRangeWhileScrubbing);
 
 	/** @return true if the playback range should be synced to the section bounds, false otherwise */
-	bool ShouldKeepPlayRangeInSectionBounds() const;
+	UE_API bool ShouldKeepPlayRangeInSectionBounds() const;
 	/** Set whether or not the playback range should be synced to the section bounds */
-	void SetKeepPlayRangeInSectionBounds(bool bInKeepPlayRangeInSectionBounds);
+	UE_API void SetKeepPlayRangeInSectionBounds(bool bInKeepPlayRangeInSectionBounds);
 
 	/** Get the number of digits we should zero-pad to when showing frame numbers in sequencer */
-	uint8 GetZeroPadFrames() const;
+	UE_API uint8 GetZeroPadFrames() const;
 	/** Set the number of digits we should zero-pad to when showing frame numbers in sequencer */
-	void SetZeroPadFrames(uint8 InZeroPadFrames);
+	UE_API void SetZeroPadFrames(uint8 InZeroPadFrames);
 
 	/** Get the number of frames to increment when jumping forwards/backwards */
-	FFrameNumber GetJumpFrameIncrement() const;
+	UE_API FFrameNumber GetJumpFrameIncrement() const;
 	/** Set the number of frames to increment when jumping forwards/backwards */
-	void SetJumpFrameIncrement(FFrameNumber InJumpFrameIncrement);
+	UE_API void SetJumpFrameIncrement(FFrameNumber InJumpFrameIncrement);
 
 	/** Get the time-warp display mode */
-	ESequencerTimeWarpDisplay GetTimeWarpDisplayMode() const;
+	UE_API ESequencerTimeWarpDisplay GetTimeWarpDisplayMode() const;
 	/** Set the time-warp display mode */
-	void SetTimeWarpDisplayMode(ESequencerTimeWarpDisplay InTimeWarpDisplay);
+	UE_API void SetTimeWarpDisplayMode(ESequencerTimeWarpDisplay InTimeWarpDisplay);
 
 	/** @return true if showing layer bars */
-	bool GetShowLayerBars() const;
+	UE_API bool GetShowLayerBars() const;
 	/** Set whether to show layer bars */ 
-	void SetShowLayerBars(bool bInShowLayerBars);
+	UE_API void SetShowLayerBars(bool bInShowLayerBars);
 
 	/** @return true if showing key bars */
-	bool GetShowKeyBars() const;
+	UE_API bool GetShowKeyBars() const;
 	/** Set whether to show key bars */ 
-	void SetShowKeyBars(bool bInShowKeyBars);
+	UE_API void SetShowKeyBars(bool bInShowKeyBars);
 
 	/** @return true if key areas are infinite */
-	bool GetInfiniteKeyAreas() const;
+	UE_API bool GetInfiniteKeyAreas() const;
 	/** Set whether to show channel colors */
-	void SetInfiniteKeyAreas(bool bInInfiniteKeyAreas);
+	UE_API void SetInfiniteKeyAreas(bool bInInfiniteKeyAreas);
 
 	/** @return true if showing channel colors for the key bars */
-	bool GetShowChannelColors() const;
+	UE_API bool GetShowChannelColors() const;
 	/** Set whether to show channel colors for the key bars */
-	void SetShowChannelColors(bool bInShowChannelColors);
+	UE_API void SetShowChannelColors(bool bInShowChannelColors);
 
 	/** @return true if showing the info button in the playback controls */
-	bool GetShowInfoButton() const;
+	UE_API bool GetShowInfoButton() const;
 	/** Set whether to show the info button in the playback controls */
-	void SetShowInfoButton(bool bInShowInfoButton);
+	UE_API void SetShowInfoButton(bool bInShowInfoButton);
 
 	/** @return true if showing tick lines */
-	bool GetShowTickLines() const;
+	UE_API bool GetShowTickLines() const;
 	/** Set whether to show status bar */
-	void SetShowTickLines(bool bInDrawTickLines);
+	UE_API void SetShowTickLines(bool bInDrawTickLines);
 
 	/** @return true if showing sequencer toolbar */
-	bool GetShowSequencerToolbar() const;
+	UE_API bool GetShowSequencerToolbar() const;
 	/** Set whether to show sequencer toolbar bar */
-	void SetShowSequencerToolbar(bool bInDrawTickLines);
+	UE_API void SetShowSequencerToolbar(bool bInDrawTickLines);
 
 	/** @return true if showing marked frames */
-	bool GetShowMarkedFrames() const;
+	UE_API bool GetShowMarkedFrames() const;
 	/** Set whether to show marked frames */
-	void SetShowMarkedFrames(bool bShowMarkedFrames);
+	UE_API void SetShowMarkedFrames(bool bShowMarkedFrames);
 
 	/** @return true if showing scaling anchors */
-	bool GetShowScalingAnchors() const;
+	UE_API bool GetShowScalingAnchors() const;
 	/** Set whether to show scaling anchors */
-	void SetShowScalingAnchors(bool bShowScalingAnchors);
+	UE_API void SetShowScalingAnchors(bool bShowScalingAnchors);
 
 	/** @return Whether the given channel has curve extents */
-	bool HasKeyAreaCurveExtents(const FString& ChannelName) const;
+	UE_API bool HasKeyAreaCurveExtents(const FString& ChannelName) const;
 	/** @ Remove curve extents for the given channel */
-	void RemoveKeyAreaCurveExtents(const FString& ChannelName);
+	UE_API void RemoveKeyAreaCurveExtents(const FString& ChannelName);
 	/** @return Get the key area curve extents for the given channel */
-	void GetKeyAreaCurveExtents(const FString& ChannelName, double& InMin, double& InMax) const;
+	UE_API void GetKeyAreaCurveExtents(const FString& ChannelName, double& InMin, double& InMax) const;
 	/** Set the key area curve extents for the given channel */
-	void SetKeyAreaCurveExtents(const FString& ChannelName, double InMin, double InMax);
+	UE_API void SetKeyAreaCurveExtents(const FString& ChannelName, double InMin, double InMax);
 
 	/** @return The key area height when showing curves */
-	float GetKeyAreaHeightWithCurves() const;
+	UE_API float GetKeyAreaHeightWithCurves() const;
 	/** Set the key area height when showing curves */
-	void SetKeyAreaHeightWithCurves(float InKeyAreaHeightWithCurves);
+	UE_API void SetKeyAreaHeightWithCurves(float InKeyAreaHeightWithCurves);
 
 	/** @return The tolerance to use when reducing keys */
-	float GetReduceKeysTolerance() const;
+	UE_API float GetReduceKeysTolerance() const;
 	/** Set the tolerance to use when reducing keys */
-	void SetReduceKeysTolerance(float InReduceKeysTolerance);
+	UE_API void SetReduceKeysTolerance(float InReduceKeysTolerance);
 
 	/** @return true if deleting keys that fall beyond the section range when trimming */
-	bool GetDeleteKeysWhenTrimming() const;
+	UE_API bool GetDeleteKeysWhenTrimming() const;
 	/** Set whether to delete keys that fall beyond the section range when trimming */
-	void SetDeleteKeysWhenTrimming(bool bInDeleteKeysWhenTrimming);
+	UE_API void SetDeleteKeysWhenTrimming(bool bInDeleteKeysWhenTrimming);
 
 	/** @return true if disable sections when baking */
-	bool GetDisableSectionsAfterBaking() const;
+	UE_API bool GetDisableSectionsAfterBaking() const;
 	/** Set whether to disable sections when baking, as opposed to deleting */
-	void SetDisableSectionsAfterBaking(bool bInDisableSectionsAfterBaking);
+	UE_API void SetDisableSectionsAfterBaking(bool bInDisableSectionsAfterBaking);
 
 	/** Set the playback range start color */
-	FLinearColor GetPlaybackRangeStartColor() const;
+	UE_API FLinearColor GetPlaybackRangeStartColor() const;
 	/** Set the playback range start color */
-	void SetPlaybackRangeStartColor(const FLinearColor& InColor);
+	UE_API void SetPlaybackRangeStartColor(const FLinearColor& InColor);
 
 	/** Set the playback range end color */
-	FLinearColor GetPlaybackRangeEndColor() const;
+	UE_API FLinearColor GetPlaybackRangeEndColor() const;
 	/** Set the playback range end color */
-	void SetPlaybackRangeEndColor(const FLinearColor& InColor);
+	UE_API void SetPlaybackRangeEndColor(const FLinearColor& InColor);
 
 	/** @return the default marked frame color */
-	FLinearColor GetMarkedFrameColor() const;
+	UE_API FLinearColor GetMarkedFrameColor() const;
 	/** Set the default marked frame color */
-	void SetMarkedFrameColor(const FLinearColor& InColor);
+	UE_API void SetMarkedFrameColor(const FLinearColor& InColor);
 
 	/** @return the section color tints */
-	TArray<FColor> GetSectionColorTints() const;
+	UE_API TArray<FColor> GetSectionColorTints() const;
 	/** Set the section color tints */
-	void SetSectionColorTints(const TArray<FColor>& InSectionColorTints);
+	UE_API void SetSectionColorTints(const TArray<FColor>& InSectionColorTints);
 
 	/** @return Whether to playback in clean mode (game view, hide viewport UI) */
-	bool GetCleanPlaybackMode() const;
+	UE_API bool GetCleanPlaybackMode() const;
 	/** Toggle whether to playback in clean mode */
-	void SetCleanPlaybackMode(bool bInCleanPlaybackMode);
+	UE_API void SetCleanPlaybackMode(bool bInCleanPlaybackMode);
 
 	/** @return Whether to activate realtime viewports when in sequencer */
-	bool ShouldActivateRealtimeViewports() const;
+	UE_API bool ShouldActivateRealtimeViewports() const;
 	/** Toggle whether to allow possession of PIE viewports */
-	void SetActivateRealtimeViewports(bool bInActivateRealtimeViewports);
+	UE_API void SetActivateRealtimeViewports(bool bInActivateRealtimeViewports);
 
 	/** Gets whether or not track defaults will be automatically set when modifying tracks. */
-	bool GetAutoSetTrackDefaults() const;
+	UE_API bool GetAutoSetTrackDefaults() const;
 	/** Sets whether or not track defaults will be automatically set when modifying tracks. */
-	void SetAutoSetTrackDefaults(bool bInAutoSetTrackDefaults);
+	UE_API void SetAutoSetTrackDefaults(bool bInAutoSetTrackDefaults);
 
 	/** @return Whether to show debug vis */
-	bool ShouldShowDebugVisualization() const;
+	UE_API bool ShouldShowDebugVisualization() const;
 	/** Toggle whether to show debug vis */
-	void SetShowDebugVisualization(bool bInShowDebugVisualization);
+	UE_API void SetShowDebugVisualization(bool bInShowDebugVisualization);
 
 	/** @return Whether to evaluate sub sequences in isolation */
-	bool ShouldEvaluateSubSequencesInIsolation() const;
+	UE_API bool ShouldEvaluateSubSequencesInIsolation() const;
 	/** Set whether to evaluate sub sequences in isolation */
-	void SetEvaluateSubSequencesInIsolation(bool bInEvaluateSubSequencesInIsolation);
+	UE_API void SetEvaluateSubSequencesInIsolation(bool bInEvaluateSubSequencesInIsolation);
 	/** Gets the multicast delegate which is run whenever evaluate sub sequences in isolation is changed. */
 	FOnEvaluateSubSequencesInIsolationChanged& GetOnEvaluateSubSequencesInIsolationChanged() { return OnEvaluateSubSequencesInIsolationChangedEvent; }
 
 	/** @return Whether to rerun construction scripts on bound actors every frame */
-	bool ShouldRerunConstructionScripts() const;
+	UE_API bool ShouldRerunConstructionScripts() const;
 	/** Set whether to rerun construction scripts on bound actors every frame */
-	void SetRerunConstructionScripts(bool bInRerunConstructionScripts);
+	UE_API void SetRerunConstructionScripts(bool bInRerunConstructionScripts);
 
 	/** Snaps a time value in seconds to the currently selected interval. */
-	float SnapTimeToInterval(float InTimeValue) const;
+	UE_API float SnapTimeToInterval(float InTimeValue) const;
 
 	/** Check whether to show pre and post roll in sequencer */
-	bool ShouldShowPrePostRoll() const;
+	UE_API bool ShouldShowPrePostRoll() const;
 	/** Toggle whether to show pre and post roll in sequencer */
-	void SetShouldShowPrePostRoll(bool bInVisualizePreAndPostRoll);
+	UE_API void SetShouldShowPrePostRoll(bool bInVisualizePreAndPostRoll);
 
 	/** Check whether whether to recompile the director blueprint when the sequence is evaluated (if one exists) */
-	bool ShouldCompileDirectorOnEvaluate() const;
+	UE_API bool ShouldCompileDirectorOnEvaluate() const;
 	/** Assign whether whether to recompile the director blueprint when the sequence is evaluated (if one exists) */
-	void SetCompileDirectorOnEvaluate(bool bInCompileDirectorOnEvaluate);
+	UE_API void SetCompileDirectorOnEvaluate(bool bInCompileDirectorOnEvaluate);
 
 	uint32 GetTrajectoryPathCap() const { return TrajectoryPathCap; }
 
-	FOnLoopStateChanged& GetOnLoopStateChanged();
+	UE_API FOnLoopStateChanged& GetOnLoopStateChanged();
 
-	FOnTimeDisplayFormatChanged& GetOnTimeDisplayFormatChanged();
+	UE_API FOnTimeDisplayFormatChanged& GetOnTimeDisplayFormatChanged();
 
 	/** What format should we display the UI controls in when representing time in a sequence? */
 	EFrameNumberDisplayFormats GetTimeDisplayFormat() const { return FrameNumberDisplayFormat; }
 	/** Sets the time display format to the specified type. */
-	void SetTimeDisplayFormat(EFrameNumberDisplayFormats InFormat);
+	UE_API void SetTimeDisplayFormat(EFrameNumberDisplayFormats InFormat);
 
 	/** What movie renderer to use */
 	FString GetMovieRendererName() const { return MovieRendererName; }
 	/** Sets the movie renderer to use */
-	void SetMovieRendererName(const FString& InMovieRendererName);
+	UE_API void SetMovieRendererName(const FString& InMovieRendererName);
 
 	/** Gets whether or not to expand the outliner tree view when a child element is selected (from outside of the tree view). */
 	bool GetAutoExpandNodesOnSelection() const { return bAutoExpandNodesOnSelection; }
 	/** Sets whether or not to expand the outliner tree view when a child element is selected (from outside of the tree view). */
-	void SetAutoExpandNodesOnSelection(bool bInAutoExpandNodesOnSelection);
+	UE_API void SetAutoExpandNodesOnSelection(bool bInAutoExpandNodesOnSelection);
 
 
 	/**
@@ -500,74 +502,74 @@ public:
 	 * Sets whether unlocking a camera cut track should return the viewport to its original location, or keep it where
 	 * the camera cut was.
 	 */
-	void SetRestoreOriginalViewportOnCameraCutUnlock(bool bInRestoreOriginalViewportOnCameraCutUnlock);
+	UE_API void SetRestoreOriginalViewportOnCameraCutUnlock(bool bInRestoreOriginalViewportOnCameraCutUnlock);
 
 	/** Gets the tree view width percentage */
 	float GetTreeViewWidth() const { return TreeViewWidth; }
 	/** Sets the tree view width percentage */
-	void SetTreeViewWidth(float InTreeViewWidth);
+	UE_API void SetTreeViewWidth(float InTreeViewWidth);
 
 	/** Gets the saved view density */
-	UE::Sequencer::EViewDensity GetViewDensity() const;
+	UE_API UE::Sequencer::EViewDensity GetViewDensity() const;
 	/** Sets the saved view density */
-	void SetViewDensity(FName InViewDensity);
+	UE_API void SetViewDensity(FName InViewDensity);
 
 	/** Gets the asset browser width */
 	float GetAssetBrowserWidth() const { return AssetBrowserWidth; }
 	/** Sets the asset browser width */
-	void SetAssetBrowserWidth(float InAssetBrowserWidth);
+	UE_API void SetAssetBrowserWidth(float InAssetBrowserWidth);
 
 	/** Gets the asset browser height */
 	float GetAssetBrowserHeight() const { return AssetBrowserHeight; }
 	/** Sets the asset browser width */
-	void SetAssetBrowserHeight(float InAssetBrowserHeight);
+	UE_API void SetAssetBrowserHeight(float InAssetBrowserHeight);
 
 	/** Gets whether the given track filter is enabled */
-	bool IsTrackFilterEnabled(const FString& TrackFilter) const;
+	UE_API bool IsTrackFilterEnabled(const FString& TrackFilter) const;
 	/** Sets whether the track filter should be enabled/disabled */
-	void SetTrackFilterEnabled(const FString& TrackFilter, bool bEnabled);
+	UE_API void SetTrackFilterEnabled(const FString& TrackFilter, bool bEnabled);
 
 	/** Get outliner column visibility in display order */
 	TArray<FColumnVisibilitySetting> GetOutlinerColumnSettings() const { return ColumnVisibilitySettings; }
 	/** Sets the visibility of outliner columns in display order */
-	void SetOutlinerColumnVisibility(const TArray<FColumnVisibilitySetting>& InColumnVisibilitySettings);
+	UE_API void SetOutlinerColumnVisibility(const TArray<FColumnVisibilitySetting>& InColumnVisibilitySettings);
 
 	/** Gets the last saved sidebar state */
-	FSidebarState& GetSidebarState();
+	UE_API FSidebarState& GetSidebarState();
 	/** Sets the sidebar state to be restored on Sequencer initialize */
-	void SetSidebarState(const FSidebarState& InSidebarState);
+	UE_API void SetSidebarState(const FSidebarState& InSidebarState);
 
-	FSequencerFilterBarConfig& FindOrAddTrackFilterBar(const FName InIdentifier, const bool bInSaveConfig);
-	FSequencerFilterBarConfig* FindTrackFilterBar(const FName InIdentifier);
-	bool RemoveTrackFilterBar(const FName InIdentifier);
+	UE_API FSequencerFilterBarConfig& FindOrAddTrackFilterBar(const FName InIdentifier, const bool bInSaveConfig);
+	UE_API FSequencerFilterBarConfig* FindTrackFilterBar(const FName InIdentifier);
+	UE_API bool RemoveTrackFilterBar(const FName InIdentifier);
 
-	bool GetIncludePinnedInFilter() const;
-	void SetIncludePinnedInFilter(const bool bInIncludePinned);
+	UE_API bool GetIncludePinnedInFilter() const;
+	UE_API void SetIncludePinnedInFilter(const bool bInIncludePinned);
 
-	bool GetAutoExpandNodesOnFilterPass() const;
-	void SetAutoExpandNodesOnFilterPass(const bool bInIncludePinned);
+	UE_API bool GetAutoExpandNodesOnFilterPass() const;
+	UE_API void SetAutoExpandNodesOnFilterPass(const bool bInIncludePinned);
 
-	bool GetUseFilterSubmenusForCategories() const;
-	void SetUseFilterSubmenusForCategories(const bool bInUseFilterSubmenusForCategories);
+	UE_API bool GetUseFilterSubmenusForCategories() const;
+	UE_API void SetUseFilterSubmenusForCategories(const bool bInUseFilterSubmenusForCategories);
 
-	bool IsFilterBarVisible() const;
-	void SetFilterBarVisible(const bool bInVisible);
+	UE_API bool IsFilterBarVisible() const;
+	UE_API void SetFilterBarVisible(const bool bInVisible);
 
-	EFilterBarLayout GetFilterBarLayout() const;
-	void SetFilterBarLayout(const EFilterBarLayout InLayout);
+	UE_API EFilterBarLayout GetFilterBarLayout() const;
+	UE_API void SetFilterBarLayout(const EFilterBarLayout InLayout);
 
-	float GetLastFilterBarSizeCoefficient() const;
-	void SetLastFilterBarSizeCoefficient(const float bInSizeCoefficient);
+	UE_API float GetLastFilterBarSizeCoefficient() const;
+	UE_API void SetLastFilterBarSizeCoefficient(const float bInSizeCoefficient);
 	
 	/** Gets the settings that determine how the asset thumbnail is captured when the sequence is saved. */
 	const FSequencerThumbnailCaptureSettings& GetThumbnailCaptureSettings() const { return ThumbnailCaptureSettings; }
 	/** Sets how the asset thumbnail is captured when the sequence is saved.  */
-	void SetThumbnailCaptureSettings(const FSequencerThumbnailCaptureSettings& InNewValue);
+	UE_API void SetThumbnailCaptureSettings(const FSequencerThumbnailCaptureSettings& InNewValue);
 
 	/** @return True if the Navigation Tool is visible */
 	bool IsNavigationToolVisible() const { return bNavigationToolVisible; }
 	/** Sets the visibility of the Navigation Tool */
-	void SetNavigationToolVisible(const bool bInVisible);
+	UE_API void SetNavigationToolVisible(const bool bInVisible);
 
 protected:
 
@@ -913,7 +915,7 @@ protected:
 	bool bNavigationToolVisible;
 
 	UFUNCTION()
-	static bool ShouldShowThumbnailCaptureSettings();
+	static UE_API bool ShouldShowThumbnailCaptureSettings();
 
 	FOnEvaluateSubSequencesInIsolationChanged OnEvaluateSubSequencesInIsolationChangedEvent;
 	FOnShowSelectedNodesOnlyChanged OnShowSelectedNodesOnlyChangedEvent;
@@ -921,3 +923,5 @@ protected:
 	FOnLoopStateChanged OnLoopStateChangedEvent;
 	FOnTimeDisplayFormatChanged OnTimeDisplayFormatChangedEvent;
 };
+
+#undef UE_API
diff --git a/Engine/Source/Editor/Sequencer/Public/SequencerUtilities.h b/Engine/Source/Editor/Sequencer/Public/SequencerUtilities.h
--- a/Engine/Source/Editor/Sequencer/Public/SequencerUtilities.h
+++ b/Engine/Source/Editor/Sequencer/Public/SequencerUtilities.h
@@ -10,6 +10,8 @@
 #include "SequencerCoreFwd.h"
 #include "SequencerUtilities.generated.h"
 
+#define UE_API SEQUENCER_API
+
 class ACineCameraActor;
 template <typename> class TAttribute;
 template <typename> class TSubclassOf;
@@ -196,20 +198,20 @@ struct FSequencerChangeBindingInfo
 /**
  * Helper structure to track when an sequencer is opened and closed.  Note, it will only track from when DoStartup is invoked and will not enumerate existing open sequencers.
  */
-struct SEQUENCER_API FOpenSequencerWatcher
+struct FOpenSequencerWatcher
 {
 	/**
 	 * Begin watching the sequencer.
 	 *
 	 * @param OnStartup - will be called when we start listening for sequencer events. This will always be after engine startup if called during module load; otherwise it will be immediately called.
 	 */
-	void DoStartup(TFunction<void ()> StartupComplete);
+	UE_API void DoStartup(TFunction<void ()> StartupComplete);
 
 	/** Invoked when a new sequencer is created. */
-	void OnSequencerCreated(TSharedRef<ISequencer> InSequencer);
+	UE_API void OnSequencerCreated(TSharedRef<ISequencer> InSequencer);
 
 	/** Registered delegate when a sequencer is closed. */
-	void OnSequencerClosed(TSharedRef<ISequencer> InSequencer);
+	UE_API void OnSequencerClosed(TSharedRef<ISequencer> InSequencer);
 
 	/** Internal structure for tracking sequencers. */
 	struct FOpenSequencerData
@@ -225,122 +227,124 @@ struct SEQUENCER_API FOpenSequencerWatcher
 	TArray<FOpenSequencerData> OpenSequencers;
 };
 
-struct SEQUENCER_API FSequencerUtilities
+struct FSequencerUtilities
 {
 	/* Creates a button (used for +Section) that opens a ComboButton with a user-defined sub-menu content. */
-	static TSharedRef<SWidget> MakeAddButton(FText HoverText, FOnGetContent MenuContent, const TAttribute<bool>& HoverState, TWeakPtr<ISequencer> InSequencer);
+	static UE_API TSharedRef<SWidget> MakeAddButton(FText HoverText, FOnGetContent MenuContent, const TAttribute<bool>& HoverState, TWeakPtr<ISequencer> InSequencer);
 	
 	/* Creates a button (used for +Section) that fires a user-defined OnClick response with no sub-menu. */
-	static TSharedRef<SWidget> MakeAddButton(FText HoverText, FOnClicked OnClicked, const TAttribute<bool>& HoverState, TWeakPtr<ISequencer> InSequencer);
+	static UE_API TSharedRef<SWidget> MakeAddButton(FText HoverText, FOnClicked OnClicked, const TAttribute<bool>& HoverState, TWeakPtr<ISequencer> InSequencer);
 
-	static void MakeTimeWarpMenuEntry(FMenuBuilder& MenuBuilder, UE::Sequencer::TWeakViewModelPtr<UE::Sequencer::ITrackExtension> TrackModel);
+	static UE_API void MakeTimeWarpMenuEntry(FMenuBuilder& MenuBuilder, UE::Sequencer::TWeakViewModelPtr<UE::Sequencer::ITrackExtension> TrackModel);
 
-	static void PopulateTimeWarpSubMenu(FMenuBuilder& MenuBuilder, TFunction<void(TSubclassOf<UMovieSceneTimeWarpGetter>)> OnTimeWarpPicked);
+	static UE_API void PopulateTimeWarpSubMenu(FMenuBuilder& MenuBuilder, TFunction<void(TSubclassOf<UMovieSceneTimeWarpGetter>)> OnTimeWarpPicked);
 
-	static void PopulateTimeWarpChannelSubMenu(FMenuBuilder& MenuBuilder, UE::Sequencer::TWeakViewModelPtr<UE::Sequencer::ITrackExtension> TrackModel);
+	static UE_API void PopulateTimeWarpChannelSubMenu(FMenuBuilder& MenuBuilder, UE::Sequencer::TWeakViewModelPtr<UE::Sequencer::ITrackExtension> TrackModel);
 
-	static void CreateNewSection(UMovieSceneTrack* InTrack, TWeakPtr<ISequencer> InSequencer, int32 InRowIndex, EMovieSceneBlendType InBlendType);
+	static UE_API void CreateNewSection(UMovieSceneTrack* InTrack, TWeakPtr<ISequencer> InSequencer, int32 InRowIndex, EMovieSceneBlendType InBlendType);
 
-	static void PopulateMenu_CreateNewSection(FMenuBuilder& MenuBuilder, int32 RowIndex, UMovieSceneTrack* Track, TWeakPtr<ISequencer> InSequencer);
+	static UE_API void PopulateMenu_CreateNewSection(FMenuBuilder& MenuBuilder, int32 RowIndex, UMovieSceneTrack* Track, TWeakPtr<ISequencer> InSequencer);
 
-	static void PopulateMenu_BlenderSubMenu(FMenuBuilder& MenuBuilder, UMovieSceneTrack* Track, TWeakPtr<ISequencer> InSequencer);
+	static UE_API void PopulateMenu_BlenderSubMenu(FMenuBuilder& MenuBuilder, UMovieSceneTrack* Track, TWeakPtr<ISequencer> InSequencer);
 
-	static void PopulateMenu_SetBlendType(FMenuBuilder& MenuBuilder, UMovieSceneSection* Section, TWeakPtr<ISequencer> InSequencer);
+	static UE_API void PopulateMenu_SetBlendType(FMenuBuilder& MenuBuilder, UMovieSceneSection* Section, TWeakPtr<ISequencer> InSequencer);
 
-	static void PopulateMenu_SetBlendType(FMenuBuilder& MenuBuilder, const TArray<TWeakObjectPtr<UMovieSceneSection>>& InSections, TWeakPtr<ISequencer> InSequencer);
+	static UE_API void PopulateMenu_SetBlendType(FMenuBuilder& MenuBuilder, const TArray<TWeakObjectPtr<UMovieSceneSection>>& InSections, TWeakPtr<ISequencer> InSequencer);
 
-	static TArray<FString> GetAssociatedLevelSequenceMapPackages(const ULevelSequence* InSequence);
-	static TArray<FString> GetAssociatedLevelSequenceMapPackages(FName LevelSequencePackageName);
+	static UE_API TArray<FString> GetAssociatedLevelSequenceMapPackages(const ULevelSequence* InSequence);
+	static UE_API TArray<FString> GetAssociatedLevelSequenceMapPackages(FName LevelSequencePackageName);
 
 	/** 
 	 * Generates a unique FName from a candidate name given a set of already existing names.  
 	 * The name is made unique by appending a number to the end.
 	 */
-	static FName GetUniqueName(FName CandidateName, const TArray<FName>& ExistingNames);
+	static UE_API FName GetUniqueName(FName CandidateName, const TArray<FName>& ExistingNames);
 
 	/** Add existing actors to Sequencer */
-	static TArray<FGuid> AddActors(TSharedRef<ISequencer> Sequencer, const TArray<TWeakObjectPtr<AActor> >& InActors);
+	static UE_API TArray<FGuid> AddActors(TSharedRef<ISequencer> Sequencer, const TArray<TWeakObjectPtr<AActor> >& InActors);
 
 	/** Create a new camera actor and add it to Sequencer */
-	static FGuid CreateCamera(TSharedRef<ISequencer> Sequencer, const bool bSpawnable, ACineCameraActor*& OutActor);
+	static UE_API FGuid CreateCamera(TSharedRef<ISequencer> Sequencer, const bool bSpawnable, ACineCameraActor*& OutActor);
 
 	/** Create a new camera from a rig and add it to Sequencer */
-	static FGuid CreateCameraWithRig(TSharedRef<ISequencer> Sequencer, AActor* Actor, const bool bSpawnable, ACineCameraActor*& OutActor);
+	static UE_API FGuid CreateCameraWithRig(TSharedRef<ISequencer> Sequencer, AActor* Actor, const bool bSpawnable, ACineCameraActor*& OutActor);
 
-	static FGuid MakeNewSpawnable(TSharedRef<ISequencer> Sequencer, UObject& SourceObject, UActorFactory* ActorFactory = nullptr, bool bSetupDefaults = true, FName SpawnableName = NAME_None);
+	static UE_API FGuid MakeNewSpawnable(TSharedRef<ISequencer> Sequencer, UObject& SourceObject, UActorFactory* ActorFactory = nullptr, bool bSetupDefaults = true, FName SpawnableName = NAME_None);
 
 	/** Convert the requested object binding to old-style spawnable. If there are multiple objects assigned to the possessable, multiple spawnables will be created */
-	static TArray<FMovieSceneSpawnable*> ConvertToSpawnable(TSharedRef<ISequencer> Sequencer, FGuid PossessableGuid);
+	static UE_API TArray<FMovieSceneSpawnable*> ConvertToSpawnable(TSharedRef<ISequencer> Sequencer, FGuid PossessableGuid);
 	
 	/** Convert the requested object binding and object binding index to a possessable */
-	static bool CanConvertToPossessable(TSharedRef<ISequencer> Sequencer, FGuid BindingGuid, int32 BindingIndex = 0);
-	static FMovieScenePossessable* ConvertToPossessable(TSharedRef<ISequencer> Sequencer, FGuid BindingGuid, int32 BindingIndex=0);
+	static UE_API bool CanConvertToPossessable(TSharedRef<ISequencer> Sequencer, FGuid BindingGuid, int32 BindingIndex = 0);
+	static UE_API FMovieScenePossessable* ConvertToPossessable(TSharedRef<ISequencer> Sequencer, FGuid BindingGuid, int32 BindingIndex=0);
 
 	/** Convert the selected object binding and object binding index to a custom binding of the chosen type.*/
-	static bool CanConvertToCustomBinding(TSharedRef<ISequencer> Sequencer, FGuid BindingGuid, TSubclassOf<UMovieSceneCustomBinding> CustomBindingType, int32 BindingIndex = 0);
-	static FMovieScenePossessable* ConvertToCustomBinding(TSharedRef<ISequencer> Sequencer, FGuid BindingGuid, TSubclassOf<UMovieSceneCustomBinding> CustomBindingType, int32 BindingIndex = 0);
+	static UE_API bool CanConvertToCustomBinding(TSharedRef<ISequencer> Sequencer, FGuid BindingGuid, TSubclassOf<UMovieSceneCustomBinding> CustomBindingType, int32 BindingIndex = 0);
+	static UE_API FMovieScenePossessable* ConvertToCustomBinding(TSharedRef<ISequencer> Sequencer, FGuid BindingGuid, TSubclassOf<UMovieSceneCustomBinding> CustomBindingType, int32 BindingIndex = 0);
 
 	/** Copy/paste folders */
-	static void CopyFolders(TSharedRef<ISequencer> Sequencer, const TArray<UMovieSceneFolder*>& Folders, FString& FoldersExportedText, FString& TracksExportedText, FString& ObjectsExportedText);
-	static bool PasteFolders(const FString& TextToImport, FMovieScenePasteFoldersParams PasteFoldersParams, TArray<UMovieSceneFolder*>& OutFolders, TArray<FNotificationInfo>& OutErrors);
-	static bool CanPasteFolders(const FString& TextToImport);
+	static UE_API void CopyFolders(TSharedRef<ISequencer> Sequencer, const TArray<UMovieSceneFolder*>& Folders, FString& FoldersExportedText, FString& TracksExportedText, FString& ObjectsExportedText);
+	static UE_API bool PasteFolders(const FString& TextToImport, FMovieScenePasteFoldersParams PasteFoldersParams, TArray<UMovieSceneFolder*>& OutFolders, TArray<FNotificationInfo>& OutErrors);
+	static UE_API bool CanPasteFolders(const FString& TextToImport);
 
 	UE_DEPRECATED(5.5, "CopyFolders now gathers objects and tracks within the folders. Please use CopyFolders that outputs ObjectsExportedText and TracksExportedText")
-	static void CopyFolders(const TArray<UMovieSceneFolder*>& Folders, FString& ExportedText);
+	static UE_API void CopyFolders(const TArray<UMovieSceneFolder*>& Folders, FString& ExportedText);
 
 	/** Copy/paste tracks */
-	static void CopyTracks(const TArray<UMovieSceneTrack*>& Tracks, const TArray<UMovieSceneFolder*>& InFolders, FString& ExportedText);
-	static bool PasteTracks(const FString& TextToImport, FMovieScenePasteTracksParams PasteTracksParams, TArray<UMovieSceneTrack*>& OutTracks, TArray<FNotificationInfo>& OutErrors);
-	static bool CanPasteTracks(const FString& TextToImport);
+	static UE_API void CopyTracks(const TArray<UMovieSceneTrack*>& Tracks, const TArray<UMovieSceneFolder*>& InFolders, FString& ExportedText);
+	static UE_API bool PasteTracks(const FString& TextToImport, FMovieScenePasteTracksParams PasteTracksParams, TArray<UMovieSceneTrack*>& OutTracks, TArray<FNotificationInfo>& OutErrors);
+	static UE_API bool CanPasteTracks(const FString& TextToImport);
 
 	/** Copy/paste sections */
-	static void CopySections(const TArray<UMovieSceneSection*>& Sections, FString& ExportedText);
-	static bool PasteSections(const FString& TextToImport, FMovieScenePasteSectionsParams PasteSectionsParams, TArray<UMovieSceneSection*>& OutSections, TArray<FNotificationInfo>& OutErrors);
-	static bool CanPasteSections(const FString& TextToImport);
+	static UE_API void CopySections(const TArray<UMovieSceneSection*>& Sections, FString& ExportedText);
+	static UE_API bool PasteSections(const FString& TextToImport, FMovieScenePasteSectionsParams PasteSectionsParams, TArray<UMovieSceneSection*>& OutSections, TArray<FNotificationInfo>& OutErrors);
+	static UE_API bool CanPasteSections(const FString& TextToImport);
 
 	/** Copy/paste object bindings */
-	static void CopyBindings(TSharedRef<ISequencer> Sequencer, const TArray<FMovieSceneBindingProxy>& Bindings, const TArray<UMovieSceneFolder*>& InFolders, FString& ExportedText);
-	static void CopyBindings(TSharedRef<ISequencer> Sequencer, const TArray<FMovieSceneBindingProxy>& Bindings, const TArray<UMovieSceneFolder*>& InFolders, FOutputDevice& Ar);
-	static bool PasteBindings(const FString& TextToImport, TSharedRef<ISequencer> Sequencer, FMovieScenePasteBindingsParams PasteBindingsParams, TArray<FMovieSceneBindingProxy>& OutBindings, TArray<FNotificationInfo>& OutErrors);
-	static bool CanPasteBindings(TSharedRef<ISequencer> Sequencer, const FString& TextToImport);
-	static TArray<FString> GetPasteBindingsObjectNames(TSharedRef<ISequencer> Sequencer, const FString& TextToImport);
+	static UE_API void CopyBindings(TSharedRef<ISequencer> Sequencer, const TArray<FMovieSceneBindingProxy>& Bindings, const TArray<UMovieSceneFolder*>& InFolders, FString& ExportedText);
+	static UE_API void CopyBindings(TSharedRef<ISequencer> Sequencer, const TArray<FMovieSceneBindingProxy>& Bindings, const TArray<UMovieSceneFolder*>& InFolders, FOutputDevice& Ar);
+	static UE_API bool PasteBindings(const FString& TextToImport, TSharedRef<ISequencer> Sequencer, FMovieScenePasteBindingsParams PasteBindingsParams, TArray<FMovieSceneBindingProxy>& OutBindings, TArray<FNotificationInfo>& OutErrors);
+	static UE_API bool CanPasteBindings(TSharedRef<ISequencer> Sequencer, const FString& TextToImport);
+	static UE_API TArray<FString> GetPasteBindingsObjectNames(TSharedRef<ISequencer> Sequencer, const FString& TextToImport);
 	/**
 	 * Recursively finds the most appropriate Resolution Context for a given Parent Guid of a Possessable
 	 */
 
-	static UObject* FindResolutionContext(TSharedRef<ISequencer> Sequencer
+	static UE_API UObject* FindResolutionContext(TSharedRef<ISequencer> Sequencer
 		, UMovieSceneSequence& InSequence
 		, UMovieScene& InMovieScene
 		, const FGuid& InParentGuid
 		, UObject* InPlaybackContext);
 
 	/** Utility functions for managing bindings */
-	static FGuid CreateBinding(TSharedRef<ISequencer> Sequencer, UObject& InObject, const UE::Sequencer::FCreateBindingParams& Params = UE::Sequencer::FCreateBindingParams());
-	static FGuid CreateOrReplaceBinding(TSharedRef<ISequencer> Sequencer, UObject* Object, const UE::Sequencer::FCreateBindingParams& Params = UE::Sequencer::FCreateBindingParams());
-	static FGuid CreateOrReplaceBinding(TSharedPtr<ISequencer> Sequencer, UMovieSceneSequence* Sequence, UObject* Object, const UE::Sequencer::FCreateBindingParams& Params = UE::Sequencer::FCreateBindingParams());
-	static void UpdateBindingIDs(TSharedRef<ISequencer> Sequencer, FGuid OldGuid, FGuid NewGuid);
-	static FGuid AssignActor(TSharedRef<ISequencer> Sequencer, AActor* Actor, FGuid InObjectBinding);
-	static void AddActorsToBinding(TSharedRef<ISequencer> Sequencer, const TArray<AActor*>& Actors, const FMovieSceneBindingProxy& ObjectBinding);
-	static void AddObjectsToBinding(TSharedRef<ISequencer> Sequencer, const TArray<UObject*>& Objects, const FMovieSceneBindingProxy& ObjectBinding, UObject* ResolutionContext);
-	static void ReplaceBindingWithActors(TSharedRef<ISequencer> Sequencer, const TArray<AActor*>& Actors, const FMovieSceneBindingProxy& ObjectBinding);
-	static void RemoveActorsFromBinding(TSharedRef<ISequencer> Sequencer, const TArray<AActor*>& Actors, const FMovieSceneBindingProxy& ObjectBinding);
+	static UE_API FGuid CreateBinding(TSharedRef<ISequencer> Sequencer, UObject& InObject, const UE::Sequencer::FCreateBindingParams& Params = UE::Sequencer::FCreateBindingParams());
+	static UE_API FGuid CreateOrReplaceBinding(TSharedRef<ISequencer> Sequencer, UObject* Object, const UE::Sequencer::FCreateBindingParams& Params = UE::Sequencer::FCreateBindingParams());
+	static UE_API FGuid CreateOrReplaceBinding(TSharedPtr<ISequencer> Sequencer, UMovieSceneSequence* Sequence, UObject* Object, const UE::Sequencer::FCreateBindingParams& Params = UE::Sequencer::FCreateBindingParams());
+	static UE_API void UpdateBindingIDs(TSharedRef<ISequencer> Sequencer, FGuid OldGuid, FGuid NewGuid);
+	static UE_API FGuid AssignActor(TSharedRef<ISequencer> Sequencer, AActor* Actor, FGuid InObjectBinding);
+	static UE_API void AddActorsToBinding(TSharedRef<ISequencer> Sequencer, const TArray<AActor*>& Actors, const FMovieSceneBindingProxy& ObjectBinding);
+	static UE_API void AddObjectsToBinding(TSharedRef<ISequencer> Sequencer, const TArray<UObject*>& Objects, const FMovieSceneBindingProxy& ObjectBinding, UObject* ResolutionContext);
+	static UE_API void ReplaceBindingWithActors(TSharedRef<ISequencer> Sequencer, const TArray<AActor*>& Actors, const FMovieSceneBindingProxy& ObjectBinding);
+	static UE_API void RemoveActorsFromBinding(TSharedRef<ISequencer> Sequencer, const TArray<AActor*>& Actors, const FMovieSceneBindingProxy& ObjectBinding);
 
 	/** Show a read only error if the movie scene is locked */
-	static void ShowReadOnlyError();
+	static UE_API void ShowReadOnlyError();
 	
 	/** Show an error if spawnable is not allowed in a movie scene*/
-	static void ShowSpawnableNotAllowedError();
+	static UE_API void ShowSpawnableNotAllowedError();
 
 	// Methods exposing FSequencer functionality outside of the Sequencer module.
 	// This is needed while things are moving into view-models here and there, step by step.
-	static void SaveCurrentMovieSceneAs(TSharedRef<ISequencer> Sequencer);
-	static void SynchronizeExternalSelectionWithSequencerSelection (TSharedRef<ISequencer> Sequencer);
-	static TRange<FFrameNumber> GetTimeBounds(TSharedRef<ISequencer> Sequencer);
+	static UE_API void SaveCurrentMovieSceneAs(TSharedRef<ISequencer> Sequencer);
+	static UE_API void SynchronizeExternalSelectionWithSequencerSelection (TSharedRef<ISequencer> Sequencer);
+	static UE_API TRange<FFrameNumber> GetTimeBounds(TSharedRef<ISequencer> Sequencer);
 
 	// Functions allowing menus to be built for modifying bindings
-	static void AddChangeClassMenu(FMenuBuilder& MenuBuilder, TSharedRef<ISequencer> Sequencer, const TArray<FSequencerChangeBindingInfo>& BindingsToConvert, TFunction<void()> OnBindingChanged);
-	static void HandleTemplateActorClassPicked(UClass* ChosenClass, TSharedRef<ISequencer> Sequencer, const TArray<FSequencerChangeBindingInfo>& BindingsToConvert, TFunction<void()> OnBindingChanged);
+	static UE_API void AddChangeClassMenu(FMenuBuilder& MenuBuilder, TSharedRef<ISequencer> Sequencer, const TArray<FSequencerChangeBindingInfo>& BindingsToConvert, TFunction<void()> OnBindingChanged);
+	static UE_API void HandleTemplateActorClassPicked(UClass* ChosenClass, TSharedRef<ISequencer> Sequencer, const TArray<FSequencerChangeBindingInfo>& BindingsToConvert, TFunction<void()> OnBindingChanged);
 
 	/** Get a movie scene sequence from a FMovideSceneSequenceID */
-	static UMovieSceneSequence* GetMovieSceneSequence(TSharedPtr<ISequencer>& InSequencer, const FMovieSceneSequenceID& SequenceID);
+	static UE_API UMovieSceneSequence* GetMovieSceneSequence(TSharedPtr<ISequencer>& InSequencer, const FMovieSceneSequenceID& SequenceID);
 };
+
+#undef UE_API
diff --git a/Engine/Source/Editor/Sequencer/Public/Tools/MotionTrailOptions.h b/Engine/Source/Editor/Sequencer/Public/Tools/MotionTrailOptions.h
--- a/Engine/Source/Editor/Sequencer/Public/Tools/MotionTrailOptions.h
+++ b/Engine/Source/Editor/Sequencer/Public/Tools/MotionTrailOptions.h
@@ -6,6 +6,8 @@
 #include "Delegates/Delegate.h"
 #include "MotionTrailOptions.generated.h"
 
+#define UE_API SEQUENCER_API
+
 //if true still use old motion trails for sequencer objects.
 extern SEQUENCER_API TAutoConsoleVariable<bool> CVarUseOldSequencerMotionTrails;
 
@@ -24,8 +26,8 @@ enum class EMotionTrailTrailStyle : uint8
 
 
 // TODO: option to make tick size proportional to distance from camera to get a sense of perspective and scale
-UCLASS(config = EditorPerProjectUserSettings)
-class SEQUENCER_API UMotionTrailToolOptions : public UObject
+UCLASS(MinimalAPI, config = EditorPerProjectUserSettings)
+class UMotionTrailToolOptions : public UObject
 {
 	GENERATED_BODY()
 
@@ -176,24 +178,24 @@ public:
 		bool operator==(const FGuid& InId) const { return TrailGuid == InId; }
 
 	};
-	void ResetPinnedItems();
-	int32 GetNumPinned() const;
-	FPinnedTrail* GetPinnedTrail(int32 Index);
-
-	int32 GetIndexFromGuid(FGuid InGuid) const;
-	void PinSelection() const;
-	void UnPinSelection();
-	void AddPinned(const FPinnedTrail& InPinnedTrail);
-	void PinComponent(USceneComponent* InSceneComponent, const FName& InSocketName) const;
-	void DeletePinned(int32 Index);
-	void DeleteAllPinned();
-	void PutPinnnedInSpace(int32 Index, AActor* InActor, const FName& InComponentName);
-	void SetLinearColor(int32 Index, const FLinearColor& Color);
-	void SetHasOffset(int32 Index, bool bHasOffset);
+	UE_API void ResetPinnedItems();
+	UE_API int32 GetNumPinned() const;
+	UE_API FPinnedTrail* GetPinnedTrail(int32 Index);
+
+	UE_API int32 GetIndexFromGuid(FGuid InGuid) const;
+	UE_API void PinSelection() const;
+	UE_API void UnPinSelection();
+	UE_API void AddPinned(const FPinnedTrail& InPinnedTrail);
+	UE_API void PinComponent(USceneComponent* InSceneComponent, const FName& InSocketName) const;
+	UE_API void DeletePinned(int32 Index);
+	UE_API void DeleteAllPinned();
+	UE_API void PutPinnnedInSpace(int32 Index, AActor* InActor, const FName& InComponentName);
+	UE_API void SetLinearColor(int32 Index, const FLinearColor& Color);
+	UE_API void SetHasOffset(int32 Index, bool bHasOffset);
 	//get name and tooltip name for the color types
-	TArray<TPair<FText, FText>>& GetTrailStyles();
-	void SetTrailStyle(int32 Index);
-	int32 GetTrailStyleIndex() const;
+	UE_API TArray<TPair<FText, FText>>& GetTrailStyles();
+	UE_API void SetTrailStyle(int32 Index);
+	UE_API int32 GetTrailStyleIndex() const;
 
 	DECLARE_MULTICAST_DELEGATE(FOnPinSelection);
 	FOnPinSelection OnPinSelection;
@@ -230,3 +232,4 @@ protected:
 
 
 
+#undef UE_API
diff --git a/Engine/Source/Editor/Sequencer/Public/Tools/SMotionTrailOptions.h b/Engine/Source/Editor/Sequencer/Public/Tools/SMotionTrailOptions.h
--- a/Engine/Source/Editor/Sequencer/Public/Tools/SMotionTrailOptions.h
+++ b/Engine/Source/Editor/Sequencer/Public/Tools/SMotionTrailOptions.h
@@ -11,9 +11,11 @@
 #include "Widgets/SCompoundWidget.h"
 #include "Widgets/SWidget.h"
 
+#define UE_API SEQUENCER_API
+
 class IDetailsView;
 
-class SEQUENCER_API SMotionTrailOptions : public SCompoundWidget
+class SMotionTrailOptions : public SCompoundWidget
 {
 	SLATE_BEGIN_ARGS(SMotionTrailOptions) {}
 
@@ -22,8 +24,9 @@ class SEQUENCER_API SMotionTrailOptions : public SCompoundWidget
 	{
 	}
 
-	void Construct(const FArguments& InArgs);
+	UE_API void Construct(const FArguments& InArgs);
 
 	TSharedPtr<IDetailsView> DetailsView;
 };
 
+#undef UE_API
diff --git a/Engine/Source/Editor/SequencerCore/Public/MVVM/Extensions/DynamicExtensionContainer.h b/Engine/Source/Editor/SequencerCore/Public/MVVM/Extensions/DynamicExtensionContainer.h
--- a/Engine/Source/Editor/SequencerCore/Public/MVVM/Extensions/DynamicExtensionContainer.h
+++ b/Engine/Source/Editor/SequencerCore/Public/MVVM/Extensions/DynamicExtensionContainer.h
@@ -10,6 +10,8 @@
 #include "Templates/SharedPointer.h"
 #include "Misc/InlineValue.h"
 
+#define UE_API SEQUENCERCORE_API
+
 namespace UE
 {
 namespace Sequencer
@@ -30,11 +32,11 @@ struct TDynamicExtensionContainerIteratorProxy;
 /**
  * Base class for dynamic extensions that can be added to sequence data models and participate in dynamic casting.
  */
-class SEQUENCERCORE_API IDynamicExtension : public TSharedFromThis<IDynamicExtension>
+class IDynamicExtension : public TSharedFromThis<IDynamicExtension>
 {
 public:
 
-	UE_SEQUENCER_DECLARE_VIEW_MODEL_TYPE_ID(IDynamicExtension)
+	UE_SEQUENCER_DECLARE_VIEW_MODEL_TYPE_ID_API(UE_API, IDynamicExtension)
 
 	virtual ~IDynamicExtension(){}
 
@@ -259,3 +261,4 @@ protected:
 } // namespace Sequencer
 } // namespace UE
 
+#undef UE_API
diff --git a/Engine/Source/Editor/SequencerCore/Public/MVVM/Extensions/HierarchicalCacheExtension.h b/Engine/Source/Editor/SequencerCore/Public/MVVM/Extensions/HierarchicalCacheExtension.h
--- a/Engine/Source/Editor/SequencerCore/Public/MVVM/Extensions/HierarchicalCacheExtension.h
+++ b/Engine/Source/Editor/SequencerCore/Public/MVVM/Extensions/HierarchicalCacheExtension.h
@@ -8,6 +8,8 @@
 #include "MVVM/Extensions/DynamicExtensionContainer.h"
 #include "MVVM/ViewModelPtr.h"
 
+#define UE_API SEQUENCERCORE_API
+
 struct FGuid;
 
 namespace UE::Sequencer
@@ -23,11 +25,11 @@ EnumType CombinePropagatedChildFlags(EnumType ParentFlags, EnumType CombinedChil
 
 DECLARE_MULTICAST_DELEGATE_OneParam(FPreUpdateCachesEvent, FViewModelPtr);
 
-class SEQUENCERCORE_API IHierarchicalCache
+class IHierarchicalCache
 {
 public:
 
-	UE_SEQUENCER_DECLARE_VIEW_MODEL_TYPE_ID(IHierarchicalCache)
+	UE_SEQUENCER_DECLARE_VIEW_MODEL_TYPE_ID_API(UE_API, IHierarchicalCache)
 
 	virtual ~IHierarchicalCache(){}
 
@@ -38,27 +40,27 @@ public:
 };
 
 
-class SEQUENCERCORE_API FHierarchicalCacheExtension
+class FHierarchicalCacheExtension
 	: public IDynamicExtension
 {
 public:
 
-	~FHierarchicalCacheExtension();
+	UE_API ~FHierarchicalCacheExtension();
 
 	FPreUpdateCachesEvent PreUpdateCachesEvent;
 
-	void Initialize(const FViewModelPtr& InRootModel);
+	UE_API void Initialize(const FViewModelPtr& InRootModel);
 
-	void OnCreated(TSharedRef<FViewModel> InWeakOwner) override;
+	UE_API void OnCreated(TSharedRef<FViewModel> InWeakOwner) override;
 
-	void OnHierarchyUpdated();
+	UE_API void OnHierarchyUpdated();
 
-	void UpdateCachedFlags();
-	void UpdateAllCachedFlags(const FViewModelPtr& ViewModel, TArrayView<IHierarchicalCache* const> HierarchicalCaches);
+	UE_API void UpdateCachedFlags();
+	UE_API void UpdateAllCachedFlags(const FViewModelPtr& ViewModel, TArrayView<IHierarchicalCache* const> HierarchicalCaches);
 
 private:
 
-	void UpdateCachedFlagsForModel(const FViewModelPtr& ViewModel, TArrayView<IHierarchicalCache* const> HierarchicalCaches);
+	UE_API void UpdateCachedFlagsForModel(const FViewModelPtr& ViewModel, TArrayView<IHierarchicalCache* const> HierarchicalCaches);
 
 protected:
 
@@ -147,3 +149,4 @@ protected:
 
 } // namespace UE::Sequencer
 
+#undef UE_API
diff --git a/Engine/Source/Editor/SequencerCore/Public/MVVM/Extensions/IDeactivatableExtension.h b/Engine/Source/Editor/SequencerCore/Public/MVVM/Extensions/IDeactivatableExtension.h
--- a/Engine/Source/Editor/SequencerCore/Public/MVVM/Extensions/IDeactivatableExtension.h
+++ b/Engine/Source/Editor/SequencerCore/Public/MVVM/Extensions/IDeactivatableExtension.h
@@ -1,10 +1,12 @@
-﻿// Copyright Epic Games, Inc. All Rights Reserved.
+// Copyright Epic Games, Inc. All Rights Reserved.
 
 #pragma once
 
 #include "MVVM/ViewModelTypeID.h"
 #include "MVVM/Extensions/HierarchicalCacheExtension.h"
 
+#define UE_API SEQUENCERCORE_API
+
 namespace UE::Sequencer
 {
 
@@ -18,10 +20,10 @@ enum class EDeactivatableState
 /**
  * An extension for outliner nodes that can be deactivated
  */
-class SEQUENCERCORE_API IDeactivatableExtension
+class IDeactivatableExtension
 {
 public:
-	UE_SEQUENCER_DECLARE_VIEW_MODEL_TYPE_ID(IDeactivatableExtension);
+	UE_SEQUENCER_DECLARE_VIEW_MODEL_TYPE_ID_API(UE_API, IDeactivatableExtension);
 
 	virtual ~IDeactivatableExtension() {}
 
@@ -55,15 +57,17 @@ ENUM_CLASS_FLAGS(ECachedDeactiveState)
 
 SEQUENCERCORE_API ECachedDeactiveState CombinePropagatedChildFlags(const ECachedDeactiveState ParentFlags, ECachedDeactiveState CombinedChildFlags);
 
-class SEQUENCERCORE_API FDeactiveStateCacheExtension
+class FDeactiveStateCacheExtension
 	: public TFlagStateCacheExtension<ECachedDeactiveState>
 {
 public:
-	UE_SEQUENCER_DECLARE_VIEW_MODEL_TYPE_ID(FDeactiveStateCacheExtension);
+	UE_SEQUENCER_DECLARE_VIEW_MODEL_TYPE_ID_API(UE_API, FDeactiveStateCacheExtension);
 
 private:
-	ECachedDeactiveState ComputeFlagsForModel(const FViewModelPtr& ViewModel) override;
-	void PostComputeChildrenFlags(const FViewModelPtr& ViewModel, ECachedDeactiveState& OutThisModelFlags, ECachedDeactiveState& OutPropagateToParentFlags) override;
+	UE_API ECachedDeactiveState ComputeFlagsForModel(const FViewModelPtr& ViewModel) override;
+	UE_API void PostComputeChildrenFlags(const FViewModelPtr& ViewModel, ECachedDeactiveState& OutThisModelFlags, ECachedDeactiveState& OutPropagateToParentFlags) override;
 };
 
 } // namespace UE::Sequencer
+
+#undef UE_API
diff --git a/Engine/Source/Editor/SequencerCore/Public/MVVM/Extensions/IDeletableExtension.h b/Engine/Source/Editor/SequencerCore/Public/MVVM/Extensions/IDeletableExtension.h
--- a/Engine/Source/Editor/SequencerCore/Public/MVVM/Extensions/IDeletableExtension.h
+++ b/Engine/Source/Editor/SequencerCore/Public/MVVM/Extensions/IDeletableExtension.h
@@ -6,16 +6,18 @@
 #include "Internationalization/Text.h"
 #include "MVVM/ViewModelTypeID.h"
 
+#define UE_API SEQUENCERCORE_API
+
 namespace UE
 {
 namespace Sequencer
 {
 
-class SEQUENCERCORE_API IDeletableExtension
+class IDeletableExtension
 {
 public:
 
-	UE_SEQUENCER_DECLARE_VIEW_MODEL_TYPE_ID(IDeletableExtension)
+	UE_SEQUENCER_DECLARE_VIEW_MODEL_TYPE_ID_API(UE_API, IDeletableExtension)
 
 	virtual ~IDeletableExtension(){}
 
@@ -25,3 +27,5 @@ public:
 
 } // namespace Sequencer
 } // namespace UE
+
+#undef UE_API
diff --git a/Engine/Source/Editor/SequencerCore/Public/MVVM/Extensions/IDimmableExtension.h b/Engine/Source/Editor/SequencerCore/Public/MVVM/Extensions/IDimmableExtension.h
--- a/Engine/Source/Editor/SequencerCore/Public/MVVM/Extensions/IDimmableExtension.h
+++ b/Engine/Source/Editor/SequencerCore/Public/MVVM/Extensions/IDimmableExtension.h
@@ -5,16 +5,18 @@
 #include "CoreMinimal.h"
 #include "MVVM/ViewModelTypeID.h"
 
+#define UE_API SEQUENCERCORE_API
+
 namespace UE
 {
 namespace Sequencer
 {
 
-class SEQUENCERCORE_API IDimmableExtension
+class IDimmableExtension
 {
 public:
 
-	UE_SEQUENCER_DECLARE_VIEW_MODEL_TYPE_ID(IDimmableExtension)
+	UE_SEQUENCER_DECLARE_VIEW_MODEL_TYPE_ID_API(UE_API, IDimmableExtension)
 
 	virtual ~IDimmableExtension(){}
 
@@ -24,3 +26,4 @@ public:
 } // namespace Sequencer
 } // namespace UE
 
+#undef UE_API
diff --git a/Engine/Source/Editor/SequencerCore/Public/MVVM/Extensions/IDraggableOutlinerExtension.h b/Engine/Source/Editor/SequencerCore/Public/MVVM/Extensions/IDraggableOutlinerExtension.h
--- a/Engine/Source/Editor/SequencerCore/Public/MVVM/Extensions/IDraggableOutlinerExtension.h
+++ b/Engine/Source/Editor/SequencerCore/Public/MVVM/Extensions/IDraggableOutlinerExtension.h
@@ -5,6 +5,8 @@
 #include "MVVM/ViewModelTypeID.h"
 #include "Misc/FrameNumber.h"
 
+#define UE_API SEQUENCERCORE_API
+
 struct FGuid;
 
 namespace UE
@@ -16,11 +18,11 @@ namespace Sequencer
 /**
  * Extension for models that can be dragged
  */
-class SEQUENCERCORE_API IDraggableOutlinerExtension
+class IDraggableOutlinerExtension
 {
 public:
 
-	UE_SEQUENCER_DECLARE_VIEW_MODEL_TYPE_ID(IDraggableOutlinerExtension)
+	UE_SEQUENCER_DECLARE_VIEW_MODEL_TYPE_ID_API(UE_API, IDraggableOutlinerExtension)
 
 	virtual ~IDraggableOutlinerExtension(){}
 
@@ -31,3 +33,4 @@ public:
 } // namespace Sequencer
 } // namespace UE
 
+#undef UE_API
diff --git a/Engine/Source/Editor/SequencerCore/Public/MVVM/Extensions/IDraggableTrackAreaExtension.h b/Engine/Source/Editor/SequencerCore/Public/MVVM/Extensions/IDraggableTrackAreaExtension.h
--- a/Engine/Source/Editor/SequencerCore/Public/MVVM/Extensions/IDraggableTrackAreaExtension.h
+++ b/Engine/Source/Editor/SequencerCore/Public/MVVM/Extensions/IDraggableTrackAreaExtension.h
@@ -6,6 +6,8 @@
 #include "Misc/FrameNumber.h"
 #include "Templates/SharedPointer.h"
 
+#define UE_API SEQUENCERCORE_API
+
 struct FGuid;
 
 namespace UE
@@ -18,7 +20,7 @@ class FViewModel;
 /**
  * Interface for drag operations
  */
-class SEQUENCERCORE_API IDragOperation
+class IDragOperation
 {
 public:
 	virtual ~IDragOperation() {}
@@ -33,11 +35,11 @@ public:
 /**
  * Extension for models that can be dragged
  */
-class SEQUENCERCORE_API IDraggableTrackAreaExtension
+class IDraggableTrackAreaExtension
 {
 public:
 
-	UE_SEQUENCER_DECLARE_VIEW_MODEL_TYPE_ID(IDraggableTrackAreaExtension)
+	UE_SEQUENCER_DECLARE_VIEW_MODEL_TYPE_ID_API(UE_API, IDraggableTrackAreaExtension)
 
 	virtual ~IDraggableTrackAreaExtension(){}
 
@@ -52,3 +54,4 @@ public:
 } // namespace Sequencer
 } // namespace UE
 
+#undef UE_API
diff --git a/Engine/Source/Editor/SequencerCore/Public/MVVM/Extensions/IDroppableExtension.h b/Engine/Source/Editor/SequencerCore/Public/MVVM/Extensions/IDroppableExtension.h
--- a/Engine/Source/Editor/SequencerCore/Public/MVVM/Extensions/IDroppableExtension.h
+++ b/Engine/Source/Editor/SequencerCore/Public/MVVM/Extensions/IDroppableExtension.h
@@ -4,6 +4,8 @@
 
 #include "MVVM/ViewModelTypeID.h"
 
+#define UE_API SEQUENCERCORE_API
+
 class FText;
 
 namespace UE
@@ -16,11 +18,11 @@ class IDragOperation;
 /**
  * Extension for models that can be dragged
  */
-class SEQUENCERCORE_API IDroppableExtension
+class IDroppableExtension
 {
 public:
 
-	UE_SEQUENCER_DECLARE_VIEW_MODEL_TYPE_ID(IDroppableExtension)
+	UE_SEQUENCER_DECLARE_VIEW_MODEL_TYPE_ID_API(UE_API, IDroppableExtension)
 
 	virtual ~IDroppableExtension(){}
 
@@ -34,3 +36,4 @@ public:
 } // namespace Sequencer
 } // namespace UE
 
+#undef UE_API
diff --git a/Engine/Source/Editor/SequencerCore/Public/MVVM/Extensions/IGeometryExtension.h b/Engine/Source/Editor/SequencerCore/Public/MVVM/Extensions/IGeometryExtension.h
--- a/Engine/Source/Editor/SequencerCore/Public/MVVM/Extensions/IGeometryExtension.h
+++ b/Engine/Source/Editor/SequencerCore/Public/MVVM/Extensions/IGeometryExtension.h
@@ -5,6 +5,8 @@
 #include "Templates/SharedPointer.h"
 #include "MVVM/ViewModelTypeID.h"
 
+#define UE_API SEQUENCERCORE_API
+
 namespace UE
 {
 namespace Sequencer
@@ -15,7 +17,7 @@ class FViewModel;
 /**
  * Virtual bounding box information for track area elements.
  */
-struct SEQUENCERCORE_API FVirtualGeometry
+struct FVirtualGeometry
 {
 	FVirtualGeometry()
 		: Top(0.f), Height(0.f), NestedBottom(0.f)
@@ -53,24 +55,24 @@ struct SEQUENCERCORE_API FVirtualGeometry
 /**
  * Interface for data models who can store virtual geometry information
  */
-class SEQUENCERCORE_API IGeometryExtension
+class IGeometryExtension
 {
 public:
 
-	UE_SEQUENCER_DECLARE_VIEW_MODEL_TYPE_ID(IGeometryExtension)
+	UE_SEQUENCER_DECLARE_VIEW_MODEL_TYPE_ID_API(UE_API, IGeometryExtension)
 
 	virtual ~IGeometryExtension(){}
 
 	virtual FVirtualGeometry GetVirtualGeometry() const = 0;
 	virtual void ReportVirtualGeometry(const FVirtualGeometry& VirtualGeometry) = 0;
 
-	static float UpdateVirtualGeometry(float InitialVirtualPosition, TSharedPtr<FViewModel> InInitialItem);
+	static UE_API float UpdateVirtualGeometry(float InitialVirtualPosition, TSharedPtr<FViewModel> InInitialItem);
 };
 
 /**
  * Basic utility implementation of the IGeometryExtension interface
  */
-class SEQUENCERCORE_API FGeometryExtensionShim : public IGeometryExtension
+class FGeometryExtensionShim : public IGeometryExtension
 {
 public:
 
@@ -92,3 +94,4 @@ protected:
 } // namespace Sequencer
 } // namespace UE
 
+#undef UE_API
diff --git a/Engine/Source/Editor/SequencerCore/Public/MVVM/Extensions/IGroupableExtension.h b/Engine/Source/Editor/SequencerCore/Public/MVVM/Extensions/IGroupableExtension.h
--- a/Engine/Source/Editor/SequencerCore/Public/MVVM/Extensions/IGroupableExtension.h
+++ b/Engine/Source/Editor/SequencerCore/Public/MVVM/Extensions/IGroupableExtension.h
@@ -6,6 +6,8 @@
 #include "Misc/StringBuilder.h"
 #include "MVVM/ViewModelTypeID.h"
 
+#define UE_API SEQUENCERCORE_API
+
 struct FGuid;
 
 namespace UE
@@ -13,11 +15,11 @@ namespace UE
 namespace Sequencer
 {
 
-class SEQUENCERCORE_API IGroupableExtension
+class IGroupableExtension
 {
 public:
 
-	UE_SEQUENCER_DECLARE_VIEW_MODEL_TYPE_ID(IGroupableExtension)
+	UE_SEQUENCER_DECLARE_VIEW_MODEL_TYPE_ID_API(UE_API, IGroupableExtension)
 
 	virtual ~IGroupableExtension(){}
 
@@ -35,3 +37,4 @@ public:
 } // namespace Sequencer
 } // namespace UE
 
+#undef UE_API
diff --git a/Engine/Source/Editor/SequencerCore/Public/MVVM/Extensions/IHotspotExtension.h b/Engine/Source/Editor/SequencerCore/Public/MVVM/Extensions/IHotspotExtension.h
--- a/Engine/Source/Editor/SequencerCore/Public/MVVM/Extensions/IHotspotExtension.h
+++ b/Engine/Source/Editor/SequencerCore/Public/MVVM/Extensions/IHotspotExtension.h
@@ -5,16 +5,18 @@
 #include "MVVM/ViewModelTypeID.h"
 #include "Templates/SharedPointer.h"
 
+#define UE_API SEQUENCERCORE_API
+
 namespace UE
 {
 namespace Sequencer
 {
 
-class SEQUENCERCORE_API IHotspotExtension
+class IHotspotExtension
 {
 public:
 
-	UE_SEQUENCER_DECLARE_VIEW_MODEL_TYPE_ID(IHotspotExtension)
+	UE_SEQUENCER_DECLARE_VIEW_MODEL_TYPE_ID_API(UE_API, IHotspotExtension)
 
 	virtual ~IHotspotExtension(){}
 };
@@ -22,3 +24,4 @@ public:
 } // namespace Sequencer
 } // namespace UE
 
+#undef UE_API
diff --git a/Engine/Source/Editor/SequencerCore/Public/MVVM/Extensions/IHoveredExtension.h b/Engine/Source/Editor/SequencerCore/Public/MVVM/Extensions/IHoveredExtension.h
--- a/Engine/Source/Editor/SequencerCore/Public/MVVM/Extensions/IHoveredExtension.h
+++ b/Engine/Source/Editor/SequencerCore/Public/MVVM/Extensions/IHoveredExtension.h
@@ -5,16 +5,18 @@
 #include "Templates/SharedPointer.h"
 #include "MVVM/ViewModelTypeID.h"
 
+#define UE_API SEQUENCERCORE_API
+
 namespace UE
 {
 namespace Sequencer
 {
 
-class SEQUENCERCORE_API IHoveredExtension
+class IHoveredExtension
 {
 public:
 
-	UE_SEQUENCER_DECLARE_VIEW_MODEL_TYPE_ID(IHoveredExtension)
+	UE_SEQUENCER_DECLARE_VIEW_MODEL_TYPE_ID_API(UE_API, IHoveredExtension)
 
 	virtual ~IHoveredExtension() {}
 
@@ -24,7 +26,7 @@ public:
 	virtual void OnUnhovered() = 0;
 };
 
-class SEQUENCERCORE_API FHoveredExtensionShim : public IHoveredExtension
+class FHoveredExtensionShim : public IHoveredExtension
 {
 public:
 
@@ -41,3 +43,4 @@ protected:
 } // namespace Sequencer
 } // namespace UE
 
+#undef UE_API
diff --git a/Engine/Source/Editor/SequencerCore/Public/MVVM/Extensions/IKeyExtension.h b/Engine/Source/Editor/SequencerCore/Public/MVVM/Extensions/IKeyExtension.h
--- a/Engine/Source/Editor/SequencerCore/Public/MVVM/Extensions/IKeyExtension.h
+++ b/Engine/Source/Editor/SequencerCore/Public/MVVM/Extensions/IKeyExtension.h
@@ -7,6 +7,8 @@
 #include "Curves/KeyHandle.h"
 #include "Misc/FrameTime.h"
 
+#define UE_API SEQUENCERCORE_API
+
 struct FGuid;
 struct FPointerEvent;
 struct FKeyDrawParams;
@@ -36,11 +38,11 @@ struct FCachedKeys
 	void GetKeysInRangeWithBounds(const TRange<FFrameTime>& Range, TArrayView<const FFrameTime>* OutKeyTimes, TArrayView<const FKeyHandle>* OutHandles, FCachedKey* OutLeadingKey, FCachedKey* OutTrailingKey) const;
 };
 
-class SEQUENCERCORE_API IKeyExtension
+class IKeyExtension
 {
 public:
 
-	UE_SEQUENCER_DECLARE_VIEW_MODEL_TYPE_ID(IKeyExtension)
+	UE_SEQUENCER_DECLARE_VIEW_MODEL_TYPE_ID_API(UE_API, IKeyExtension)
 
 
 	virtual ~IKeyExtension(){}
@@ -58,3 +60,4 @@ public:
 
 } // namespace UE::Sequencer
 
+#undef UE_API
diff --git a/Engine/Source/Editor/SequencerCore/Public/MVVM/Extensions/ILayerBarExtension.h b/Engine/Source/Editor/SequencerCore/Public/MVVM/Extensions/ILayerBarExtension.h
--- a/Engine/Source/Editor/SequencerCore/Public/MVVM/Extensions/ILayerBarExtension.h
+++ b/Engine/Source/Editor/SequencerCore/Public/MVVM/Extensions/ILayerBarExtension.h
@@ -6,16 +6,18 @@
 #include "Misc/FrameNumber.h"
 #include "MVVM/ViewModelTypeID.h"
 
+#define UE_API SEQUENCERCORE_API
+
 namespace UE
 {
 namespace Sequencer
 {
 
-class SEQUENCERCORE_API ILayerBarExtension
+class ILayerBarExtension
 {
 public:
 
-	UE_SEQUENCER_DECLARE_VIEW_MODEL_TYPE_ID(ILayerBarExtension)
+	UE_SEQUENCER_DECLARE_VIEW_MODEL_TYPE_ID_API(UE_API, ILayerBarExtension)
 
 	virtual ~ILayerBarExtension(){}
 
@@ -26,3 +28,4 @@ public:
 } // namespace Sequencer
 } // namespace UE
 
+#undef UE_API
diff --git a/Engine/Source/Editor/SequencerCore/Public/MVVM/Extensions/ILockableExtension.h b/Engine/Source/Editor/SequencerCore/Public/MVVM/Extensions/ILockableExtension.h
--- a/Engine/Source/Editor/SequencerCore/Public/MVVM/Extensions/ILockableExtension.h
+++ b/Engine/Source/Editor/SequencerCore/Public/MVVM/Extensions/ILockableExtension.h
@@ -6,6 +6,8 @@
 #include "MVVM/ViewModelTypeID.h"
 #include "MVVM/Extensions/HierarchicalCacheExtension.h"
 
+#define UE_API SEQUENCERCORE_API
+
 namespace UE
 {
 namespace Sequencer
@@ -21,11 +23,11 @@ enum class ELockableLockState
 /**
  * An extension for models that can be locked
  */
-class SEQUENCERCORE_API ILockableExtension
+class ILockableExtension
 {
 public:
 
-	UE_SEQUENCER_DECLARE_VIEW_MODEL_TYPE_ID(ILockableExtension)
+	UE_SEQUENCER_DECLARE_VIEW_MODEL_TYPE_ID_API(UE_API, ILockableExtension)
 
 	virtual ~ILockableExtension(){}
 
@@ -50,19 +52,20 @@ enum class ECachedLockState
 };
 ENUM_CLASS_FLAGS(ECachedLockState)
 
-class SEQUENCERCORE_API FLockStateCacheExtension
+class FLockStateCacheExtension
 	: public TFlagStateCacheExtension<ECachedLockState>
 {
 public:
 
-	UE_SEQUENCER_DECLARE_VIEW_MODEL_TYPE_ID(FLockStateCacheExtension);
+	UE_SEQUENCER_DECLARE_VIEW_MODEL_TYPE_ID_API(UE_API, FLockStateCacheExtension);
 
 private:
 
-	ECachedLockState ComputeFlagsForModel(const FViewModelPtr& ViewModel) override;
-	void PostComputeChildrenFlags(const FViewModelPtr& ViewModel, ECachedLockState& OutThisModelFlags, ECachedLockState& OutPropagateToParentFlags) override;
+	UE_API ECachedLockState ComputeFlagsForModel(const FViewModelPtr& ViewModel) override;
+	UE_API void PostComputeChildrenFlags(const FViewModelPtr& ViewModel, ECachedLockState& OutThisModelFlags, ECachedLockState& OutPropagateToParentFlags) override;
 };
 
 } // namespace Sequencer
 } // namespace UE
 
+#undef UE_API
diff --git a/Engine/Source/Editor/SequencerCore/Public/MVVM/Extensions/IMutableExtension.h b/Engine/Source/Editor/SequencerCore/Public/MVVM/Extensions/IMutableExtension.h
--- a/Engine/Source/Editor/SequencerCore/Public/MVVM/Extensions/IMutableExtension.h
+++ b/Engine/Source/Editor/SequencerCore/Public/MVVM/Extensions/IMutableExtension.h
@@ -6,6 +6,8 @@
 #include "MVVM/ViewModelTypeID.h"
 #include "MVVM/Extensions/HierarchicalCacheExtension.h"
 
+#define UE_API SEQUENCERCORE_API
+
 namespace UE
 {
 namespace Sequencer
@@ -14,11 +16,11 @@ namespace Sequencer
 /**
  * An extension for outliner nodes that can be muted
  */
-class SEQUENCERCORE_API IMutableExtension
+class IMutableExtension
 {
 public:
 
-	UE_SEQUENCER_DECLARE_VIEW_MODEL_TYPE_ID(IMutableExtension)
+	UE_SEQUENCER_DECLARE_VIEW_MODEL_TYPE_ID_API(UE_API, IMutableExtension)
 
 	virtual ~IMutableExtension(){}
 
@@ -51,20 +53,21 @@ enum class ECachedMuteState
 ENUM_CLASS_FLAGS(ECachedMuteState)
 SEQUENCERCORE_API ECachedMuteState CombinePropagatedChildFlags(ECachedMuteState ParentFlags, ECachedMuteState CombinedChildFlags);
 
-class SEQUENCERCORE_API FMuteStateCacheExtension
+class FMuteStateCacheExtension
 	: public TFlagStateCacheExtension<ECachedMuteState>
 {
 public:
 
-	UE_SEQUENCER_DECLARE_VIEW_MODEL_TYPE_ID(FMuteStateCacheExtension);
+	UE_SEQUENCER_DECLARE_VIEW_MODEL_TYPE_ID_API(UE_API, FMuteStateCacheExtension);
 
 private:
 
-	ECachedMuteState ComputeFlagsForModel(const FViewModelPtr& ViewModel) override;
-	void PostComputeChildrenFlags(const FViewModelPtr& ViewModel, ECachedMuteState& OutThisModelFlags, ECachedMuteState& OutPropagateToParentFlags) override;
+	UE_API ECachedMuteState ComputeFlagsForModel(const FViewModelPtr& ViewModel) override;
+	UE_API void PostComputeChildrenFlags(const FViewModelPtr& ViewModel, ECachedMuteState& OutThisModelFlags, ECachedMuteState& OutPropagateToParentFlags) override;
 };
 
 
 } // namespace Sequencer
 } // namespace UE
 
+#undef UE_API
diff --git a/Engine/Source/Editor/SequencerCore/Public/MVVM/Extensions/IObjectModelExtension.h b/Engine/Source/Editor/SequencerCore/Public/MVVM/Extensions/IObjectModelExtension.h
--- a/Engine/Source/Editor/SequencerCore/Public/MVVM/Extensions/IObjectModelExtension.h
+++ b/Engine/Source/Editor/SequencerCore/Public/MVVM/Extensions/IObjectModelExtension.h
@@ -6,17 +6,19 @@
 #include "MVVM/ViewModelTypeID.h"
 #include "UObject/WeakObjectPtrFwd.h"
 
+#define UE_API SEQUENCERCORE_API
+
 namespace UE::Sequencer
 {
 
 /**
  * An extension that is used for view models that represent UObject data
  */
-class SEQUENCERCORE_API IObjectModelExtension
+class IObjectModelExtension
 {
 public:
 
-	UE_SEQUENCER_DECLARE_VIEW_MODEL_TYPE_ID(IObjectModelExtension)
+	UE_SEQUENCER_DECLARE_VIEW_MODEL_TYPE_ID_API(UE_API, IObjectModelExtension)
 
 	virtual ~IObjectModelExtension(){}
 
@@ -24,4 +26,6 @@ public:
 	virtual UObject* GetObject() const = 0;
 };
 
-} // namespace UE::Sequencer
\ No newline at end of file
+} // namespace UE::Sequencer
+
+#undef UE_API
diff --git a/Engine/Source/Editor/SequencerCore/Public/MVVM/Extensions/IOutlinerExtension.h b/Engine/Source/Editor/SequencerCore/Public/MVVM/Extensions/IOutlinerExtension.h
--- a/Engine/Source/Editor/SequencerCore/Public/MVVM/Extensions/IOutlinerExtension.h
+++ b/Engine/Source/Editor/SequencerCore/Public/MVVM/Extensions/IOutlinerExtension.h
@@ -17,6 +17,8 @@
 #include "Widgets/Views/SListView.h"
 #include "Widgets/Views/STableRow.h"
 
+#define UE_API SEQUENCERCORE_API
+
 class FDragDropEvent;
 class FText;
 class ISequencer;
@@ -56,7 +58,7 @@ enum class EOutlinerSizingFlags
 };
 ENUM_CLASS_FLAGS(EOutlinerSizingFlags);
 
-struct SEQUENCERCORE_API FOutlinerSizing
+struct FOutlinerSizing
 {
 	FOutlinerSizing() = default;
 
@@ -115,7 +117,7 @@ struct SEQUENCERCORE_API FOutlinerSizing
 };
 
 /** Facade class to allow SOutlinerViewRow to be passed around publicly */
-class SEQUENCERCORE_API ISequencerTreeViewRow : public STableRow<TWeakViewModelPtr<IOutlinerExtension>>
+class ISequencerTreeViewRow : public STableRow<TWeakViewModelPtr<IOutlinerExtension>>
 {
 public:
 	virtual bool IsColumnVisible(const FName& InColumnName) const = 0;
@@ -145,11 +147,11 @@ struct FCreateOutlinerContextMenuWidgetParams
 /**
  * Extension interface for models that can be displayed in the outliner view.
  */
-class SEQUENCERCORE_API IOutlinerExtension
+class IOutlinerExtension
 {
 public:
 
-	UE_SEQUENCER_DECLARE_VIEW_MODEL_TYPE_ID(IOutlinerExtension)
+	UE_SEQUENCER_DECLARE_VIEW_MODEL_TYPE_ID_API(UE_API, IOutlinerExtension)
 
 	virtual ~IOutlinerExtension(){}
 
@@ -180,60 +182,60 @@ public:
 	virtual void SetSelectionState(EOutlinerSelectionState InState) = 0;
 
 	/** Toggles a bit in the outliner item's selection state */
-	virtual void ToggleSelectionState(EOutlinerSelectionState InState, bool bInValue);
+	UE_API virtual void ToggleSelectionState(EOutlinerSelectionState InState, bool bInValue);
 
 	/** (deprecated) Create the outliner view for the label column */
-	virtual TSharedRef<SWidget> CreateOutlinerView(const FCreateOutlinerViewParams& InParams);
+	UE_API virtual TSharedRef<SWidget> CreateOutlinerView(const FCreateOutlinerViewParams& InParams);
 
 	/** Create the visual outliner widget for the specified column */
-	virtual TSharedPtr<SWidget> CreateOutlinerViewForColumn(const FCreateOutlinerViewParams& InParams, const FName& InColumnName);
+	UE_API virtual TSharedPtr<SWidget> CreateOutlinerViewForColumn(const FCreateOutlinerViewParams& InParams, const FName& InColumnName);
 
 	/** Build the context menu widget for this outliner item */
 	virtual TSharedPtr<SWidget> CreateContextMenuWidget(const FCreateOutlinerContextMenuWidgetParams& InParams) = 0;
 
 	/** Gets the label for the outliner item */
-	virtual FText GetLabel() const;
+	UE_API virtual FText GetLabel() const;
 
 	/** Gets the font for the outliner item */
-	virtual FSlateFontInfo GetLabelFont() const;
+	UE_API virtual FSlateFontInfo GetLabelFont() const;
 
 	/** Gets the color for the outliner item label */
-	virtual FSlateColor GetLabelColor() const;
+	UE_API virtual FSlateColor GetLabelColor() const;
 
 	/** Gets the tooltip text for the outliner item */
-	virtual FText GetLabelToolTipText() const;
+	UE_API virtual FText GetLabelToolTipText() const;
 
 	/** Gets the icon for the outliner item */
-	virtual const FSlateBrush* GetIconBrush() const;
+	UE_API virtual const FSlateBrush* GetIconBrush() const;
 
 	/** Gets the icon tint color for the outliner item */
-	virtual FSlateColor GetIconTint() const;
+	UE_API virtual FSlateColor GetIconTint() const;
 
 	/** Gets the overlay icon for the outliner item */
-	virtual const FSlateBrush* GetIconOverlayBrush() const;
+	UE_API virtual const FSlateBrush* GetIconOverlayBrush() const;
 
 	/** Gets the tooltip text for the outliner item icon */
-	virtual FText GetIconToolTipText() const;
+	UE_API virtual FText GetIconToolTipText() const;
 
 	/** Gets whether the outliner item has a background */
-	virtual bool HasBackground() const;
+	UE_API virtual bool HasBackground() const;
 
 	/** Utility method to get the full path name of the given model */
-	static FString GetPathName(const FViewModel& Item);
+	static UE_API FString GetPathName(const FViewModel& Item);
 	/** Utility method to get the full path name of the given model */
-	static FString GetPathName(const TSharedPtr<FViewModel> Item);
+	static UE_API FString GetPathName(const TSharedPtr<FViewModel> Item);
 	/** Utility method to get the full path name of the given model */
-	static FString GetPathName(const TSharedPtr<const FViewModel> Item);
+	static UE_API FString GetPathName(const TSharedPtr<const FViewModel> Item);
 	/** Utility method to get the full path name of the given model */
-	static FString GetPathName(const TWeakPtr<const FViewModel> Item);
+	static UE_API FString GetPathName(const TWeakPtr<const FViewModel> Item);
 	/** Utility method to get the full path name of the given model */
-	static void GetPathName(const FViewModel& Item, FStringBuilderBase& OutString);
+	static UE_API void GetPathName(const FViewModel& Item, FStringBuilderBase& OutString);
 };
 
 /**
  * Default implementation for some of the outliner extension API.
  */
-class SEQUENCERCORE_API FOutlinerExtensionShim : public IOutlinerExtension
+class FOutlinerExtensionShim : public IOutlinerExtension
 {
 public:
 
@@ -255,11 +257,11 @@ protected:
 /**
  * Extension for outliner items that accept drag and drop.
  */
-class SEQUENCERCORE_API IOutlinerDropTargetOutlinerExtension
+class IOutlinerDropTargetOutlinerExtension
 {
 public:
 
-	UE_SEQUENCER_DECLARE_VIEW_MODEL_TYPE_ID(IOutlinerDropTargetOutlinerExtension)
+	UE_SEQUENCER_DECLARE_VIEW_MODEL_TYPE_ID_API(UE_API, IOutlinerDropTargetOutlinerExtension)
 
 	virtual ~IOutlinerDropTargetOutlinerExtension(){}
 
@@ -270,11 +272,11 @@ public:
 /**
  * Extension for outliner models that want to compute their and their children's sizing.
  */
-class SEQUENCERCORE_API ICompoundOutlinerExtension
+class ICompoundOutlinerExtension
 {
 public:
 
-	UE_SEQUENCER_DECLARE_VIEW_MODEL_TYPE_ID(ICompoundOutlinerExtension)
+	UE_SEQUENCER_DECLARE_VIEW_MODEL_TYPE_ID_API(UE_API, ICompoundOutlinerExtension)
 
 	virtual ~ICompoundOutlinerExtension(){}
 
@@ -284,3 +286,4 @@ public:
 } // namespace Sequencer
 } // namespace UE
 
+#undef UE_API
diff --git a/Engine/Source/Editor/SequencerCore/Public/MVVM/Extensions/IPinnableExtension.h b/Engine/Source/Editor/SequencerCore/Public/MVVM/Extensions/IPinnableExtension.h
--- a/Engine/Source/Editor/SequencerCore/Public/MVVM/Extensions/IPinnableExtension.h
+++ b/Engine/Source/Editor/SequencerCore/Public/MVVM/Extensions/IPinnableExtension.h
@@ -6,6 +6,8 @@
 #include "MVVM/ViewModelTypeID.h"
 #include "Templates/SharedPointer.h"
 
+#define UE_API SEQUENCERCORE_API
+
 namespace UE
 {
 namespace Sequencer
@@ -16,11 +18,11 @@ class FViewModel;
 /**
  * An extension for outliner nodes that can be pinned at the top
  */
-class SEQUENCERCORE_API IPinnableExtension
+class IPinnableExtension
 {
 public:
 
-	UE_SEQUENCER_DECLARE_VIEW_MODEL_TYPE_ID(IPinnableExtension)
+	UE_SEQUENCER_DECLARE_VIEW_MODEL_TYPE_ID_API(UE_API, IPinnableExtension)
 
 	virtual ~IPinnableExtension(){}
 
@@ -36,13 +38,13 @@ public:
  * Default implementation for IPinnableExtension, with helper method to cache pinned
  * states for children of a pinned branch
  */
-class SEQUENCERCORE_API FPinnableExtensionShim : public IPinnableExtension
+class FPinnableExtensionShim : public IPinnableExtension
 {
 public:
 
-	virtual bool IsPinned() const override;
-	virtual void SetPinned(bool bInIsPinned) override;
-	virtual void ReportPinnedParent(bool bInIsInPinnedBranch) override;
+	UE_API virtual bool IsPinned() const override;
+	UE_API virtual void SetPinned(bool bInIsPinned) override;
+	UE_API virtual void ReportPinnedParent(bool bInIsInPinnedBranch) override;
 
 public:
 	
@@ -50,7 +52,7 @@ public:
 	 * Looks for any direct child of RootModel that is pinnable, and reports all their descendants
 	 * that they are in a pinned branch, to make that lookup faster
 	 */
-	static void UpdateCachedPinnedState(TSharedPtr<FViewModel> RootModel);
+	static UE_API void UpdateCachedPinnedState(TSharedPtr<FViewModel> RootModel);
 
 private:
 
@@ -61,3 +63,4 @@ private:
 } // namespace Sequencer
 } // namespace UE
 
+#undef UE_API
diff --git a/Engine/Source/Editor/SequencerCore/Public/MVVM/Extensions/IRecyclableExtension.h b/Engine/Source/Editor/SequencerCore/Public/MVVM/Extensions/IRecyclableExtension.h
--- a/Engine/Source/Editor/SequencerCore/Public/MVVM/Extensions/IRecyclableExtension.h
+++ b/Engine/Source/Editor/SequencerCore/Public/MVVM/Extensions/IRecyclableExtension.h
@@ -6,18 +6,20 @@
 #include "SequencerCoreFwd.h"
 #include "MVVM/ViewModelTypeID.h"
 
+#define UE_API SEQUENCERCORE_API
+
 namespace UE
 {
 namespace Sequencer
 {
 
-class SEQUENCERCORE_API IRecyclableExtension
+class IRecyclableExtension
 {
 public:
 
-	UE_SEQUENCER_DECLARE_VIEW_MODEL_TYPE_ID(IRecyclableExtension)
+	UE_SEQUENCER_DECLARE_VIEW_MODEL_TYPE_ID_API(UE_API, IRecyclableExtension)
 
-	static void CallOnRecycle(const TViewModelPtr<IRecyclableExtension>& RecyclableItem);
+	static UE_API void CallOnRecycle(const TViewModelPtr<IRecyclableExtension>& RecyclableItem);
 
 	virtual ~IRecyclableExtension(){}
 
@@ -27,3 +29,4 @@ public:
 } // namespace Sequencer
 } // namespace UE
 
+#undef UE_API
diff --git a/Engine/Source/Editor/SequencerCore/Public/MVVM/Extensions/IRenameableExtension.h b/Engine/Source/Editor/SequencerCore/Public/MVVM/Extensions/IRenameableExtension.h
--- a/Engine/Source/Editor/SequencerCore/Public/MVVM/Extensions/IRenameableExtension.h
+++ b/Engine/Source/Editor/SequencerCore/Public/MVVM/Extensions/IRenameableExtension.h
@@ -5,6 +5,8 @@
 #include "MVVM/ViewModelTypeID.h"
 #include "Delegates/Delegate.h"
 
+#define UE_API SEQUENCERCORE_API
+
 class FText;
 
 namespace UE
@@ -12,11 +14,11 @@ namespace UE
 namespace Sequencer
 {
 
-class SEQUENCERCORE_API IRenameableExtension
+class IRenameableExtension
 {
 public:
 
-	UE_SEQUENCER_DECLARE_VIEW_MODEL_TYPE_ID(IRenameableExtension)
+	UE_SEQUENCER_DECLARE_VIEW_MODEL_TYPE_ID_API(UE_API, IRenameableExtension)
 
 	virtual ~IRenameableExtension(){}
 
@@ -25,7 +27,7 @@ public:
 	virtual bool CanRename() const = 0;
 	virtual void Rename(const FText& NewName) = 0;
 
-	bool IsRenameValid(const FText& NewName, FText& OutErrorMessage) const;
+	UE_API bool IsRenameValid(const FText& NewName, FText& OutErrorMessage) const;
 
 private:
 
@@ -39,3 +41,4 @@ private:
 } // namespace Sequencer
 } // namespace UE
 
+#undef UE_API
diff --git a/Engine/Source/Editor/SequencerCore/Public/MVVM/Extensions/IResizableExtension.h b/Engine/Source/Editor/SequencerCore/Public/MVVM/Extensions/IResizableExtension.h
--- a/Engine/Source/Editor/SequencerCore/Public/MVVM/Extensions/IResizableExtension.h
+++ b/Engine/Source/Editor/SequencerCore/Public/MVVM/Extensions/IResizableExtension.h
@@ -5,16 +5,18 @@
 #include "CoreMinimal.h"
 #include "MVVM/ViewModelTypeID.h"
 
+#define UE_API SEQUENCERCORE_API
+
 namespace UE
 {
 namespace Sequencer
 {
 
-class SEQUENCERCORE_API IResizableExtension
+class IResizableExtension
 {
 public:
 
-	UE_SEQUENCER_DECLARE_VIEW_MODEL_TYPE_ID(IResizableExtension)
+	UE_SEQUENCER_DECLARE_VIEW_MODEL_TYPE_ID_API(UE_API, IResizableExtension)
 
 	virtual ~IResizableExtension(){}
 
@@ -25,3 +27,4 @@ public:
 } // namespace Sequencer
 } // namespace UE
 
+#undef UE_API
diff --git a/Engine/Source/Editor/SequencerCore/Public/MVVM/Extensions/ISelectableExtension.h b/Engine/Source/Editor/SequencerCore/Public/MVVM/Extensions/ISelectableExtension.h
--- a/Engine/Source/Editor/SequencerCore/Public/MVVM/Extensions/ISelectableExtension.h
+++ b/Engine/Source/Editor/SequencerCore/Public/MVVM/Extensions/ISelectableExtension.h
@@ -6,6 +6,8 @@
 #include "Misc/EnumClassFlags.h"
 #include "MVVM/ViewModelTypeID.h"
 
+#define UE_API SEQUENCERCORE_API
+
 struct FGuid;
 
 namespace UE
@@ -36,11 +38,11 @@ enum class ESelectionIntent
 };
 ENUM_CLASS_FLAGS(ESelectionIntent);
 
-class SEQUENCERCORE_API ISelectableExtension
+class ISelectableExtension
 {
 public:
 
-	UE_SEQUENCER_DECLARE_VIEW_MODEL_TYPE_ID(ISelectableExtension)
+	UE_SEQUENCER_DECLARE_VIEW_MODEL_TYPE_ID_API(UE_API, ISelectableExtension)
 
 	virtual ~ISelectableExtension(){}
 
@@ -51,3 +53,4 @@ public:
 } // namespace Sequencer
 } // namespace UE
 
+#undef UE_API
diff --git a/Engine/Source/Editor/SequencerCore/Public/MVVM/Extensions/ISoloableExtension.h b/Engine/Source/Editor/SequencerCore/Public/MVVM/Extensions/ISoloableExtension.h
--- a/Engine/Source/Editor/SequencerCore/Public/MVVM/Extensions/ISoloableExtension.h
+++ b/Engine/Source/Editor/SequencerCore/Public/MVVM/Extensions/ISoloableExtension.h
@@ -6,6 +6,8 @@
 #include "MVVM/ViewModelTypeID.h"
 #include "MVVM/Extensions/HierarchicalCacheExtension.h"
 
+#define UE_API SEQUENCERCORE_API
+
 namespace UE
 {
 namespace Sequencer
@@ -14,11 +16,11 @@ namespace Sequencer
 /**
  * An extension for outliner nodes that can be made 'solo'
  */
-class SEQUENCERCORE_API ISoloableExtension
+class ISoloableExtension
 {
 public:
 
-	UE_SEQUENCER_DECLARE_VIEW_MODEL_TYPE_ID(ISoloableExtension)
+	UE_SEQUENCER_DECLARE_VIEW_MODEL_TYPE_ID_API(UE_API, ISoloableExtension)
 
 	virtual ~ISoloableExtension(){}
 
@@ -43,12 +45,12 @@ enum class ECachedSoloState
 };
 ENUM_CLASS_FLAGS(ECachedSoloState)
 
-class SEQUENCERCORE_API FSoloStateCacheExtension
+class FSoloStateCacheExtension
 	: public TFlagStateCacheExtension<ECachedSoloState>
 {
 public:
 
-	UE_SEQUENCER_DECLARE_VIEW_MODEL_TYPE_ID(FSoloStateCacheExtension)
+	UE_SEQUENCER_DECLARE_VIEW_MODEL_TYPE_ID_API(UE_API, FSoloStateCacheExtension)
 
 	using Implements = TImplements<IDynamicExtension, IHierarchicalCache>;
 
@@ -56,10 +58,11 @@ public:
 
 private:
 
-	ECachedSoloState ComputeFlagsForModel(const FViewModelPtr& ViewModel) override;
-	void PostComputeChildrenFlags(const FViewModelPtr& ViewModel, ECachedSoloState& OutThisModelFlags, ECachedSoloState& OutPropagateToParentFlags) override;
+	UE_API ECachedSoloState ComputeFlagsForModel(const FViewModelPtr& ViewModel) override;
+	UE_API void PostComputeChildrenFlags(const FViewModelPtr& ViewModel, ECachedSoloState& OutThisModelFlags, ECachedSoloState& OutPropagateToParentFlags) override;
 };
 
 } // namespace Sequencer
 } // namespace UE
 
+#undef UE_API
diff --git a/Engine/Source/Editor/SequencerCore/Public/MVVM/Extensions/ISortableExtension.h b/Engine/Source/Editor/SequencerCore/Public/MVVM/Extensions/ISortableExtension.h
--- a/Engine/Source/Editor/SequencerCore/Public/MVVM/Extensions/ISortableExtension.h
+++ b/Engine/Source/Editor/SequencerCore/Public/MVVM/Extensions/ISortableExtension.h
@@ -8,6 +8,8 @@
 #include "MVVM/ViewModelTypeID.h"
 #include "Templates/SharedPointer.h"
 
+#define UE_API SEQUENCERCORE_API
+
 namespace UE
 {
 namespace Sequencer
@@ -23,7 +25,7 @@ struct FViewModelListHead;
  * Lower values have more precedence than higher values, i.e. items associated with lower
  * values are ordered before items with higher values.
  */
-struct SEQUENCERCORE_API FSortingKey
+struct FSortingKey
 {
 	/** Priority order */
 	int8 Priority = 0;
@@ -48,9 +50,9 @@ struct SEQUENCERCORE_API FSortingKey
 	}
 
 	/** Compare priorities, identifiers, and custom orders */
-	static bool ComparePriorityFirst(const FSortingKey& A, const FSortingKey& B);
+	static UE_API bool ComparePriorityFirst(const FSortingKey& A, const FSortingKey& B);
 	/** Compare custom orders (if both set), priorities, and identifiers */
-	static bool CompareCustomOrderFirst(const FSortingKey& A, const FSortingKey& B);
+	static UE_API bool CompareCustomOrderFirst(const FSortingKey& A, const FSortingKey& B);
 
 	friend bool operator==(const FSortingKey& A, const FSortingKey& B)
 	{
@@ -70,11 +72,11 @@ enum class ESortingMode
 /**
  * Extension for models that can sort their children, and be sorted themselves among their siblings
  */
-class SEQUENCERCORE_API ISortableExtension
+class ISortableExtension
 {
 public:
 
-	UE_SEQUENCER_DECLARE_VIEW_MODEL_TYPE_ID(ISortableExtension)
+	UE_SEQUENCER_DECLARE_VIEW_MODEL_TYPE_ID_API(UE_API, ISortableExtension)
 
 	virtual ~ISortableExtension(){}
 
@@ -86,11 +88,12 @@ public:
 	virtual void SetCustomOrder(int32 InCustomOrder) = 0;
 
 	/** Utility function for sorting children based on their sorting key */
-	static void SortChildren(TSharedPtr<FViewModel> ParentModel, ESortingMode SortingMode);
+	static UE_API void SortChildren(TSharedPtr<FViewModel> ParentModel, ESortingMode SortingMode);
 	/** Utility function for sorting children based on their sorting key */
-	static void SortChildren(FViewModelChildren& Children, ESortingMode SortingMode);
+	static UE_API void SortChildren(FViewModelChildren& Children, ESortingMode SortingMode);
 };
 
 } // namespace Sequencer
 } // namespace UE
 
+#undef UE_API
diff --git a/Engine/Source/Editor/SequencerCore/Public/MVVM/Extensions/IStretchableExtension.h b/Engine/Source/Editor/SequencerCore/Public/MVVM/Extensions/IStretchableExtension.h
--- a/Engine/Source/Editor/SequencerCore/Public/MVVM/Extensions/IStretchableExtension.h
+++ b/Engine/Source/Editor/SequencerCore/Public/MVVM/Extensions/IStretchableExtension.h
@@ -8,6 +8,8 @@
 #include "Math/Vector2D.h"
 #include "MVVM/ViewModelTypeID.h"
 
+#define UE_API SEQUENCERCORE_API
+
 struct FGuid;
 struct FPointerEvent;
 
@@ -30,7 +32,7 @@ enum class EStretchResult
 	Success, Failure
 };
 
-struct SEQUENCERCORE_API FStretchParameters
+struct FStretchParameters
 {
 	double Anchor = 0.0;
 	double Handle = 0.0;
@@ -41,7 +43,7 @@ struct SEQUENCERCORE_API FStretchParameters
 	}
 };
 
-struct SEQUENCERCORE_API FStretchScreenParameters
+struct FStretchScreenParameters
 {
 	FFrameTime DragStartPosition;
 	FFrameTime CurrentDragPosition;
@@ -52,11 +54,11 @@ struct SEQUENCERCORE_API FStretchScreenParameters
 	const FVirtualTrackArea* VirtualTrackArea = nullptr;
 };
 
-class SEQUENCERCORE_API IStretchableExtension
+class IStretchableExtension
 {
 public:
 
-	UE_SEQUENCER_DECLARE_VIEW_MODEL_TYPE_ID(IStretchableExtension)
+	UE_SEQUENCER_DECLARE_VIEW_MODEL_TYPE_ID_API(UE_API, IStretchableExtension)
 
 	virtual ~IStretchableExtension(){}
 
@@ -67,7 +69,7 @@ public:
 	virtual void OnEndStretch(const IStretchOperation& StretchOperation, const FStretchScreenParameters& ScreenParametersm, FStretchParameters* InOutParameters) {}
 };
 
-class SEQUENCERCORE_API IStretchOperation
+class IStretchOperation
 {
 public:
 	virtual ~IStretchOperation() {}
@@ -79,3 +81,4 @@ public:
 } // namespace Sequencer
 } // namespace UE
 
+#undef UE_API
diff --git a/Engine/Source/Editor/SequencerCore/Public/MVVM/Extensions/ITrackAreaExtension.h b/Engine/Source/Editor/SequencerCore/Public/MVVM/Extensions/ITrackAreaExtension.h
--- a/Engine/Source/Editor/SequencerCore/Public/MVVM/Extensions/ITrackAreaExtension.h
+++ b/Engine/Source/Editor/SequencerCore/Public/MVVM/Extensions/ITrackAreaExtension.h
@@ -6,6 +6,8 @@
 #include "MVVM/ViewModels/ViewModelIterators.h"
 #include "Templates/SharedPointer.h"
 
+#define UE_API SEQUENCERCORE_API
+
 namespace UE
 {
 namespace Sequencer
@@ -33,11 +35,11 @@ struct FTrackAreaParameters
 /**
  * This extension can be added to any model in order to provide a list of models that are relevent to the track area
  */
-class SEQUENCERCORE_API ITrackAreaExtension
+class ITrackAreaExtension
 {
 public:
 
-	UE_SEQUENCER_DECLARE_VIEW_MODEL_TYPE_ID(ITrackAreaExtension)
+	UE_SEQUENCER_DECLARE_VIEW_MODEL_TYPE_ID_API(UE_API, ITrackAreaExtension)
 
 	virtual ~ITrackAreaExtension(){}
 
@@ -55,3 +57,4 @@ public:
 } // namespace Sequencer
 } // namespace UE
 
+#undef UE_API
diff --git a/Engine/Source/Editor/SequencerCore/Public/MVVM/Extensions/ITrackAreaViewSpaceProviderExtension.h b/Engine/Source/Editor/SequencerCore/Public/MVVM/Extensions/ITrackAreaViewSpaceProviderExtension.h
--- a/Engine/Source/Editor/SequencerCore/Public/MVVM/Extensions/ITrackAreaViewSpaceProviderExtension.h
+++ b/Engine/Source/Editor/SequencerCore/Public/MVVM/Extensions/ITrackAreaViewSpaceProviderExtension.h
@@ -5,6 +5,8 @@
 #include "CoreTypes.h"
 #include "MVVM/ViewModelTypeID.h"
 
+#define UE_API SEQUENCERCORE_API
+
 namespace UE::Sequencer
 {
 
@@ -14,15 +16,17 @@ class FTrackAreaViewModel;
  * An extension that supplies view spaces to a track area.
  * Must exist as an immediate child of the root model to be effective.
  */
-class SEQUENCERCORE_API ITrackAreaViewSpaceProviderExtension
+class ITrackAreaViewSpaceProviderExtension
 {
 public:
 
-	UE_SEQUENCER_DECLARE_VIEW_MODEL_TYPE_ID(ITrackAreaViewSpaceProviderExtension)
+	UE_SEQUENCER_DECLARE_VIEW_MODEL_TYPE_ID_API(UE_API, ITrackAreaViewSpaceProviderExtension)
 
 	virtual ~ITrackAreaViewSpaceProviderExtension(){}
 
 	virtual void UpdateViewSpaces(FTrackAreaViewModel& TrackAreaViewModel) = 0;
 };
 
-} // namespace UE::Sequencer
\ No newline at end of file
+} // namespace UE::Sequencer
+
+#undef UE_API
