Sequencer: Remove Snap Keys and Sections to Play Range

Snapping to the play range is already handled with Snap Keys to Elements and Snap Sections to Elements. This option was restricing movement to within the play range completely and it's been deemed as undesirable/unnecessary functionality.

diff --git a/Engine/Source/Editor/Sequencer/Private/SSequencer.cpp b/Engine/Source/Editor/Sequencer/Private/SSequencer.cpp
--- a/Engine/Source/Editor/Sequencer/Private/SSequencer.cpp
+++ b/Engine/Source/Editor/Sequencer/Private/SSequencer.cpp
@@ -2515,7 +2515,6 @@ TSharedRef<SWidget> SSequencer::MakeSnapMenu()
 			FToolMenuSection& Section = Toolbar->AddSection( "KeyAndSectionsSnapping", LOCTEXT( "SnappingMenuKeyAndSectionsHeader", "Key and Sections Snapping" ) );
 			Section.AddMenuEntry( FSequencerCommands::Get().ToggleSnapKeyTimesToElements );
 			Section.AddMenuEntry( FSequencerCommands::Get().ToggleSnapSectionTimesToElements);
-			Section.AddMenuEntry( FSequencerCommands::Get().ToggleSnapKeysAndSectionsToPlayRange );
 		}
 
 		{
diff --git a/Engine/Source/Editor/Sequencer/Private/Sequencer.cpp b/Engine/Source/Editor/Sequencer/Private/Sequencer.cpp
--- a/Engine/Source/Editor/Sequencer/Private/Sequencer.cpp
+++ b/Engine/Source/Editor/Sequencer/Private/Sequencer.cpp
@@ -11468,12 +11468,6 @@ void FSequencer::BindCommands()
 		FCanExecuteAction::CreateLambda( []{ return true; } ),
 		FIsActionChecked::CreateLambda( [this]{ return Settings->GetSnapSectionTimesToElements(); } ) );
 
-	SequencerCommandBindings->MapAction(
-		Commands.ToggleSnapKeysAndSectionsToPlayRange,
-		FExecuteAction::CreateLambda([this] { Settings->SetSnapKeysAndSectionsToPlayRange(!Settings->GetSnapKeysAndSectionsToPlayRange()); }),
-		FCanExecuteAction::CreateLambda([] { return true; }),
-		FIsActionChecked::CreateLambda([this] { return Settings->GetSnapKeysAndSectionsToPlayRange(); }));
-
 	SequencerCommandBindings->MapAction(
 		Commands.ToggleSnapPlayTimeToKeys,
 		FExecuteAction::CreateLambda( [this]{ Settings->SetSnapPlayTimeToKeys( !Settings->GetSnapPlayTimeToKeys() ); } ),
diff --git a/Engine/Source/Editor/Sequencer/Private/SequencerCommands.cpp b/Engine/Source/Editor/Sequencer/Private/SequencerCommands.cpp
--- a/Engine/Source/Editor/Sequencer/Private/SequencerCommands.cpp
+++ b/Engine/Source/Editor/Sequencer/Private/SequencerCommands.cpp
@@ -138,7 +138,6 @@ void FSequencerCommands::RegisterCommands()
 
 	UI_COMMAND( ToggleSnapKeyTimesToElements, "Snap Keys to Elements", "Snap keys to other keys, sections, markers, playback range, selection range, playhead, and previous position", EUserInterfaceActionType::ToggleButton, FInputChord() );
 	UI_COMMAND( ToggleSnapSectionTimesToElements, "Snap Sections to Elements", "Snap sections to other sections, keys, markers, playback range, selection range, and playhead", EUserInterfaceActionType::ToggleButton, FInputChord() );
-	UI_COMMAND( ToggleSnapKeysAndSectionsToPlayRange, "Snap Keys and Sections to the Playback Range", "When checked, keys and sections will be snapped to the playback range bounds", EUserInterfaceActionType::ToggleButton, FInputChord() );
 
 	UI_COMMAND( ToggleSnapPlayTimeToKeys, "Snap to Keys While Scrubbing", "Snap the playhead to keys of the selected track while scrubbing", EUserInterfaceActionType::ToggleButton, FInputChord() );
 	UI_COMMAND( ToggleSnapPlayTimeToSections, "Snap to Sections While Scrubbing", "Snap the playhead to section bounds while scrubbing", EUserInterfaceActionType::ToggleButton, FInputChord() );
diff --git a/Engine/Source/Editor/Sequencer/Private/SequencerSettings.cpp b/Engine/Source/Editor/Sequencer/Private/SequencerSettings.cpp
--- a/Engine/Source/Editor/Sequencer/Private/SequencerSettings.cpp
+++ b/Engine/Source/Editor/Sequencer/Private/SequencerSettings.cpp
@@ -39,7 +39,6 @@ USequencerSettings::USequencerSettings( const FObjectInitializer& ObjectInitiali
 	bIsolateCurveEditorToSelection = true;
 	bCurveEditorVisible = false;
 	LoopMode = ESequencerLoopMode::SLM_NoLoop;
-	bSnapKeysAndSectionsToPlayRange = false;
 	bResetPlayheadWhenNavigating = false;
 	bKeepCursorInPlayRangeWhileScrubbing = false;
 	bKeepPlayRangeInSectionBounds = true;
@@ -257,20 +256,6 @@ void USequencerSettings::SetSnapSectionTimesToElements( bool InbSnapSectionTimes
 	}
 }
 
-bool USequencerSettings::GetSnapKeysAndSectionsToPlayRange() const
-{
-	return bSnapKeysAndSectionsToPlayRange;
-}
-
-void USequencerSettings::SetSnapKeysAndSectionsToPlayRange(bool bInSnapKeysAndSectionsToPlayRange)
-{
-	if (bSnapKeysAndSectionsToPlayRange != bInSnapKeysAndSectionsToPlayRange)
-	{
-		bSnapKeysAndSectionsToPlayRange = bInSnapKeysAndSectionsToPlayRange;
-		SaveConfig();
-	}
-}
-
 bool USequencerSettings::GetSnapPlayTimeToKeys() const
 {
 	return bSnapPlayTimeToKeys;
diff --git a/Engine/Source/Editor/Sequencer/Private/Tools/EditToolDragOperations.cpp b/Engine/Source/Editor/Sequencer/Private/Tools/EditToolDragOperations.cpp
--- a/Engine/Source/Editor/Sequencer/Private/Tools/EditToolDragOperations.cpp
+++ b/Engine/Source/Editor/Sequencer/Private/Tools/EditToolDragOperations.cpp
@@ -327,19 +327,6 @@ void FResizeSection::OnDrag(const FPointerEvent& MouseEvent, FVector2D LocalMous
 			LeftMovementMaximum = UE::MovieScene::DiscreteInclusiveLower(SectionBoundaries);
 			RightMovementMaximum = UE::MovieScene::DiscreteExclusiveUpper(SectionBoundaries);
 		}
-
-		if ((Settings->GetIsSnapEnabled() || MouseEvent.IsShiftDown()) && Settings->GetSnapKeysAndSectionsToPlayRange() && !Settings->ShouldKeepPlayRangeInSectionBounds())
-		{
-			if (!LeftMovementMaximum.IsSet() || LeftMovementMaximum.GetValue() < Sequencer.GetPlaybackRange().GetLowerBoundValue())
-			{
-				LeftMovementMaximum = Sequencer.GetPlaybackRange().GetLowerBoundValue();
-			}
-
-			if (!RightMovementMaximum.IsSet() || RightMovementMaximum.GetValue() > Sequencer.GetPlaybackRange().GetUpperBoundValue())
-			{
-				RightMovementMaximum = Sequencer.GetPlaybackRange().GetUpperBoundValue();
-			}
-		}
 	};
 	
 	/********************************************************************/
@@ -879,11 +866,6 @@ void FMoveKeysAndSections::OnDrag(const FPointerEvent& MouseEvent, FVector2D Loc
 		MouseTime += (SnappedTime->SnappedTime - SnappedTime->OriginalTime);
 	}
 
-	if ((Settings->GetIsSnapEnabled() || MouseEvent.IsShiftDown()) && Settings->GetSnapKeysAndSectionsToPlayRange() && !Settings->ShouldKeepPlayRangeInSectionBounds())
-	{
-		MouseTime = UE::MovieScene::ClampToDiscreteRange(MouseTime, Sequencer.GetPlaybackRange());
-	}
-
 	// We'll calculate a DeltaX based on limits on movement (snapping, section collision) and then use them on keys and sections below.
 	TOptional<FFrameNumber> MaxDeltaX = GetMovementDeltaX(MouseEvent, MouseTime);
 
@@ -1040,19 +1022,6 @@ TOptional<FFrameNumber> FMoveKeysAndSections::GetMovementDeltaX(const FPointerEv
 			RightMovementMaximum = UE::MovieScene::DiscreteExclusiveUpper(SectionBoundaries);
 		}
 		
-		if ((Settings->GetIsSnapEnabled() || MouseEvent.IsShiftDown()) && Settings->GetSnapKeysAndSectionsToPlayRange() && !Settings->ShouldKeepPlayRangeInSectionBounds())
-		{
-			if (!LeftMovementMaximum.IsSet() || LeftMovementMaximum.GetValue() < Sequencer.GetPlaybackRange().GetLowerBoundValue())
-			{
-				LeftMovementMaximum = Sequencer.GetPlaybackRange().GetLowerBoundValue();
-			}
-
-			if (!RightMovementMaximum.IsSet() || RightMovementMaximum.GetValue() > Sequencer.GetPlaybackRange().GetUpperBoundValue())
-			{
-				RightMovementMaximum = Sequencer.GetPlaybackRange().GetUpperBoundValue();
-			}
-		}
-
 		if (LeftMovementMaximum.IsSet())
 		{
 			if (Section->HasStartFrame())
@@ -1086,43 +1055,6 @@ TOptional<FFrameNumber> FMoveKeysAndSections::GetMovementDeltaX(const FPointerEv
 		}
 	}
 
-	if ((Settings->GetIsSnapEnabled() || MouseEvent.IsShiftDown()) && Settings->GetSnapKeysAndSectionsToPlayRange() && !Settings->ShouldKeepPlayRangeInSectionBounds())
-	{
-		TArray<FFrameNumber> CurrentKeyTimes;
-		CurrentKeyTimes.SetNum(KeysAsArray.Num());
-		GetKeyTimes(KeysAsArray, CurrentKeyTimes);
-
-		for (int32 Index = 0; Index < CurrentKeyTimes.Num(); ++Index)
-		{
-			FSequencerSelectedKey& SelectedKey = KeysAsArray[Index];
-			const bool bOwningSectionIsSelected = Sections.Contains(SelectedKey.Section);
-
-			// We don't want to apply delta if we have the key's section selected as well, otherwise they get double
-			// transformed (moving the section moves the keys + we add the delta to the key positions).
-			if (!bOwningSectionIsSelected)
-			{
-				FFrameNumber NewKeyTime = CurrentKeyTimes[Index] + MouseDeltaTime;
-				if (NewKeyTime < Sequencer.GetPlaybackRange().GetLowerBoundValue())
-				{
-					FFrameNumber ClampedDeltaTime = CurrentKeyTimes[Index] - Sequencer.GetPlaybackRange().GetLowerBoundValue();
-					if (!DeltaX.IsSet() || DeltaX.GetValue() > ClampedDeltaTime)
-					{
-						DeltaX = ClampedDeltaTime;
-					}
-				}
-
-				if (NewKeyTime > Sequencer.GetPlaybackRange().GetUpperBoundValue())
-				{
-					FFrameNumber ClampedDeltaTime = Sequencer.GetPlaybackRange().GetUpperBoundValue() - CurrentKeyTimes[Index];
-					if (!DeltaX.IsSet() || DeltaX.GetValue() > ClampedDeltaTime)
-					{
-						DeltaX = ClampedDeltaTime;
-					}
-				}
-			}
-		}
-	}
-
 	return DeltaX;
 }
 
diff --git a/Engine/Source/Editor/Sequencer/Public/SequencerCommands.h b/Engine/Source/Editor/Sequencer/Public/SequencerCommands.h
--- a/Engine/Source/Editor/Sequencer/Public/SequencerCommands.h
+++ b/Engine/Source/Editor/Sequencer/Public/SequencerCommands.h
@@ -351,9 +351,6 @@ public:
 	/** Toggles whether or not sections should snap to other elements. */
 	TSharedPtr< FUICommandInfo > ToggleSnapSectionTimesToElements;
 
-	/** Toggle constraining keys and sections in the play range */
-	TSharedPtr< FUICommandInfo > ToggleSnapKeysAndSectionsToPlayRange;
-
 	/** Toggles whether or not snap to key times while scrubbing. */
 	TSharedPtr< FUICommandInfo > ToggleSnapPlayTimeToKeys;
 
diff --git a/Engine/Source/Editor/Sequencer/Public/SequencerSettings.h b/Engine/Source/Editor/Sequencer/Public/SequencerSettings.h
--- a/Engine/Source/Editor/Sequencer/Public/SequencerSettings.h
+++ b/Engine/Source/Editor/Sequencer/Public/SequencerSettings.h
@@ -196,11 +196,6 @@ public:
 	/** sets whether or not to snap sections to other elements. */
 	UE_API void SetSnapSectionTimesToElements( bool InbSnapSectionTimesToElements );
 
-	/** @return true if keys and sections should be kept within the playback range when moving them */
-	UE_API bool GetSnapKeysAndSectionsToPlayRange() const;
-	/** Set whether or not keys and sections should be kept within the playback range when moving them */
-	UE_API void SetSnapKeysAndSectionsToPlayRange(bool bInSnapKeysAndSectionsToPlayRange);
-
 	/** Gets whether or not to snap the play time to keys while scrubbing. */
 	UE_API bool GetSnapPlayTimeToKeys() const;
 	/** Sets whether or not to snap the play time to keys while scrubbing. */
@@ -617,10 +612,6 @@ protected:
 	UPROPERTY( config, EditAnywhere, Category=Snapping )
 	bool bSnapSectionTimesToElements;
 
-	/** Enable or disable keeping keys and sections in the playback range. */
-	UPROPERTY(config, EditAnywhere, Category = Timeline)
-	bool bSnapKeysAndSectionsToPlayRange;
-
 	/** Enable or disable snapping the playhead to keys while scrubbing. */
 	UPROPERTY( config, EditAnywhere, Category=Snapping, meta = (DisplayName = "Snap Playhead to Keys"))
 	bool bSnapPlayTimeToKeys;
