Execute Snapping when Snapping is enabled OR Snap To Interval is forced, and not just when "Snap to Like Types" OR Snap to Interval is forced. The Snap Field can perform custom clock snapping which isn't captured by the latter two conditions.

For the SnapField, don't check the sequencer settings for whether snapping is enabled. Let the user of the SnapField do that externally.

#rnx

diff --git a/Engine/Source/Editor/Sequencer/Private/Tools/EditToolDragOperations.cpp b/Engine/Source/Editor/Sequencer/Private/Tools/EditToolDragOperations.cpp
--- a/Engine/Source/Editor/Sequencer/Private/Tools/EditToolDragOperations.cpp
+++ b/Engine/Source/Editor/Sequencer/Private/Tools/EditToolDragOperations.cpp
@@ -291,13 +291,14 @@ void FResizeSection::OnDrag(const FPointerEvent& MouseEvent, FVector2D LocalMous
 	FFrameNumber DeltaTime = (VirtualTrackArea.PixelToFrame(LocalMousePos.X) - MouseDownTime).RoundToFrame();
 
 	// Snapping
+	const bool bSnapEnabled = (Settings->GetIsSnapEnabled() || MouseEvent.IsShiftDown());
 	const bool bSnapToInterval = Settings->GetForceWholeFrames();
-	const bool bSnapToLikeTypes = ((Settings->GetIsSnapEnabled() || MouseEvent.IsShiftDown()) && Settings->GetSnapSectionTimesToElements());
+	const bool bSnapToLikeTypes = bSnapEnabled && Settings->GetSnapSectionTimesToElements();
 	
 	SnapField.GetValue().SetSnapToInterval(bSnapToInterval);
 	SnapField.GetValue().SetSnapToLikeTypes(bSnapToLikeTypes);
 
-	if (bSnapToInterval || bSnapToLikeTypes)
+	if (bSnapEnabled || bSnapToInterval)
 	{
 		TArray<FFrameTime> SectionTimes;
 		for (UMovieSceneSection* Section : Sections)
@@ -621,13 +622,14 @@ void FManipulateSectionEasing::OnDrag(const FPointerEvent& MouseEvent, FVector2D
 	}
 
 	// Snapping
+	const bool bSnapEnabled = Settings->GetIsSnapEnabled() || MouseEvent.IsShiftDown();
 	const bool bSnapToInterval = Settings->GetForceWholeFrames();
-	const bool bSnapToLikeTypes = ((Settings->GetIsSnapEnabled() || MouseEvent.IsShiftDown()) && Settings->GetSnapSectionTimesToElements());
+	const bool bSnapToLikeTypes = bSnapEnabled && Settings->GetSnapSectionTimesToElements();
 
 	SnapField.GetValue().SetSnapToInterval(bSnapToInterval);
 	SnapField.GetValue().SetSnapToLikeTypes(bSnapToLikeTypes);
 
-	if (bSnapToInterval || bSnapToLikeTypes)
+	if (bSnapEnabled || bSnapToInterval)
 	{
 		TArray<FFrameTime> SnapTimes;
 		if (bEaseIn)
@@ -833,14 +835,15 @@ void FMoveKeysAndSections::OnDrag(const FPointerEvent& MouseEvent, FVector2D Loc
 	TArray<FFrameTime> ValidSnapMarkers;
 
 	// If they have both keys and settings selected then we snap to the interval if either one of them is enabled, otherwise respect the individual setting.
-	const bool bSnapToInterval = Settings->GetForceWholeFrames() && (KeysAsArray.Num() > 0 || Sections.Num());
-	const bool bSnapToLikeTypes = (Settings->GetIsSnapEnabled() || MouseEvent.IsShiftDown()) && ( (KeysAsArray.Num() > 0 && Settings->GetSnapKeyTimesToElements()) || (Sections.Num() > 0 && Settings->GetSnapSectionTimesToElements()) );
+	const bool bSnapEnabled = (Settings->GetIsSnapEnabled() || MouseEvent.IsShiftDown());
+	const bool bSnapToInterval = (Settings->GetForceWholeFrames() && (KeysAsArray.Num() > 0 || Sections.Num()));
+	const bool bSnapToLikeTypes = bSnapEnabled && ((KeysAsArray.Num() > 0 && Settings->GetSnapKeyTimesToElements()) || (Sections.Num() > 0 && Settings->GetSnapSectionTimesToElements()));
 
 	SnapField.GetValue().SetSnapToInterval(bSnapToInterval);
 	SnapField.GetValue().SetSnapToLikeTypes(bSnapToLikeTypes);
 
 	// RelativeSnapOffsets contains both our sections and our keys, and we add them all as potential things that can snap to stuff.
-	if (bSnapToLikeTypes || bSnapToInterval)
+	if (bSnapEnabled || bSnapToInterval)
 	{
 		ValidSnapMarkers.SetNumUninitialized(RelativeSnapOffsets.Num());
 		for (int32 Index = 0; Index < RelativeSnapOffsets.Num(); ++Index)
@@ -852,7 +855,7 @@ void FMoveKeysAndSections::OnDrag(const FPointerEvent& MouseEvent, FVector2D Loc
 	// Now we'll try and snap all of these points to the closest valid snap marker (which may be a section or interval)
 	TOptional<FSequencerSnapField::FSnapResult> SnappedTime;
 
-	if (bSnapToLikeTypes || bSnapToInterval)
+	if (bSnapEnabled || bSnapToInterval)
 	{
 		// This may or may not set the SnappedTime depending on if there are any sections within the threshold.
 		SnappedTime = SnapField->Snap(ValidSnapMarkers, SnapThreshold);
diff --git a/Engine/Source/Editor/Sequencer/Private/Tools/SequencerSnapField.cpp b/Engine/Source/Editor/Sequencer/Private/Tools/SequencerSnapField.cpp
--- a/Engine/Source/Editor/Sequencer/Private/Tools/SequencerSnapField.cpp
+++ b/Engine/Source/Editor/Sequencer/Private/Tools/SequencerSnapField.cpp
@@ -114,13 +114,10 @@ void FSequencerSnapField::Initialize(const FSequencer& InSequencer, UE::Sequence
 {
 	using namespace UE::Sequencer;
 
-	if (InSequencer.GetSequencerSettings()->GetIsSnapEnabled())
+	TViewModelPtr<IClockExtension> Extension = InSequencer.GetViewModel()->GetRootSequenceModel().ImplicitCast();
+	if (Extension.IsValid() && Extension->SupportsSnapping())
 	{
-		TViewModelPtr<IClockExtension> Extension = InSequencer.GetViewModel()->GetRootSequenceModel().ImplicitCast();
-		if (Extension.IsValid() && Extension->SupportsSnapping() && Extension->ShouldSnapFrameTime())
-		{
-			ClockExtension = Extension;
-		}
+		ClockExtension = Extension;
 	}
 
 	TSharedPtr<SOutlinerView> TreeView = StaticCastSharedRef<SSequencer>(InSequencer.GetSequencerWidget())->GetTreeView();
@@ -301,7 +298,7 @@ TOptional<FSequencerSnapField::FSnapResult> FSequencerSnapField::Snap(const TArr
 			ThisSnap = Snap(Time, FMath::Max(IntervalSnapThreshold, Threshold));
 		}
 
-		if (ClockExtension.IsValid())
+		if (ClockExtension.IsValid() && ClockExtension->ShouldSnapFrameTime())
 		{
 			FFrameNumber ClockSnapTime = ClockExtension->SnapFrameTime(Time).FrameNumber;
 			if (ClockSnapTime != Time)
