Sequencer: Fix key selection on multiple rows

Visit track row models

diff --git a/Engine/Source/Editor/Sequencer/Private/Tools/SequencerEntityVisitor.cpp b/Engine/Source/Editor/Sequencer/Private/Tools/SequencerEntityVisitor.cpp
--- a/Engine/Source/Editor/Sequencer/Private/Tools/SequencerEntityVisitor.cpp
+++ b/Engine/Source/Editor/Sequencer/Private/Tools/SequencerEntityVisitor.cpp
@@ -14,6 +14,7 @@
 #include "MVVM/ViewModels/FolderModel.h"
 #include "MVVM/ViewModels/ObjectBindingModel.h"
 #include "MVVM/ViewModels/TrackModel.h"
+#include "MVVM/ViewModels/TrackRowModel.h"
 #include "MVVM/ViewModels/ViewModel.h"
 #include "MVVM/ViewModels/ViewModelHierarchy.h"
 #include "MVVM/ViewModels/ViewModelIterators.h"
@@ -171,6 +172,14 @@ void FSequencerEntityWalker::VisitAnyChannels(const ISequencerEntityVisitor& Vis
 			ChannelGroup = TopLevelChannel.Get();
 		}
 	}
+	else if (FTrackRowModel* TrackRowModel = InNode->CastThis<FTrackRowModel>())
+	{
+		if (auto It = TrackRowModel->GetTopLevelChannels().IterateSubList<FChannelGroupModel>(); It)
+		{
+			const TViewModelPtr<FChannelGroupModel>& TopLevelChannel = *It;
+			ChannelGroup = TopLevelChannel.Get();
+		}
+	}
 
 	if (ChannelGroup)
 	{
