Sequencer: Fix warnings, dereferencing BindingReferences, MP4Encoder, EntityField

Found with new version of PVS

#rnx

diff --git a/Engine/Plugins/MovieScene/MovieRenderPipeline/Source/MovieRenderPipelineMP4Encoder/Private/Windows/MoviePipelineMP4Encoder.cpp b/Engine/Plugins/MovieScene/MovieRenderPipeline/Source/MovieRenderPipelineMP4Encoder/Private/Windows/MoviePipelineMP4Encoder.cpp
--- a/Engine/Plugins/MovieScene/MovieRenderPipeline/Source/MovieRenderPipelineMP4Encoder/Private/Windows/MoviePipelineMP4Encoder.cpp
+++ b/Engine/Plugins/MovieScene/MovieRenderPipeline/Source/MovieRenderPipelineMP4Encoder/Private/Windows/MoviePipelineMP4Encoder.cpp
@@ -394,6 +394,11 @@ static HRESULT CreateVideoMediaOutputStream(IMFMediaType** OutVideoMediaTypeOutp
 		Result = MFCreateMediaType(OutVideoMediaTypeOutput);
 	}
 
+	if (*OutVideoMediaTypeOutput == nullptr)
+	{
+		return Result;
+	}
+
 	// Set the major media type for the stream.
  	if (SUCCEEDED(Result))
 	{
@@ -506,6 +511,11 @@ static HRESULT CreateAudioMediaOutputStream(IMFMediaType** OutVideoMediaTypeOutp
 		Result = MFCreateMediaType(OutVideoMediaTypeOutput);
 	}
 
+	if (*OutVideoMediaTypeOutput == nullptr)
+	{
+		return Result;
+	}
+
 	if (SUCCEEDED(Result))
 	{
 		Result = (*OutVideoMediaTypeOutput)->SetGUID(MF_MT_MAJOR_TYPE, MFMediaType_Audio);
@@ -554,6 +564,11 @@ static HRESULT CreateVideoMediaTypeIn(IMFMediaType** OutVideoMediaTypeIn, const
 		Result = MFCreateMediaType(OutVideoMediaTypeIn);
 	}
 
+	if (*OutVideoMediaTypeIn == nullptr)
+	{
+		return Result;
+	}
+
 	// Set the major media type for the stream.
 	if (SUCCEEDED(Result))
 	{
@@ -622,6 +637,11 @@ static HRESULT CreateAudioMediaTypeIn(IMFMediaType** OutAudioMediaTypeIn, const
 		Result = MFCreateMediaType(OutAudioMediaTypeIn);
 	}
 
+	if (*OutAudioMediaTypeIn == nullptr)
+	{
+		return Result;
+	}
+
 	if (SUCCEEDED(Result))
 	{
 		Result = (*OutAudioMediaTypeIn)->SetGUID(MF_MT_MAJOR_TYPE, MFMediaType_Audio);
@@ -679,6 +699,11 @@ static HRESULT GetVideoStreamEncoderAttributes(IMFAttributes** OutEncoderAttribu
 {
 	HRESULT Result = MFCreateAttributes(OutEncoderAttributes, 1);
 
+	if (*OutEncoderAttributes == nullptr)
+	{
+		return Result;
+	}
+
 	// These are set for H264 encoding and may be different for other formats in the future.
 	{
 		if (InOptions.EncodingRateControl == EMoviePipelineMP4EncodeRateControlMode::ConstantBitRate)
diff --git a/Engine/Source/Editor/Sequencer/Private/SequencerUtilities.cpp b/Engine/Source/Editor/Sequencer/Private/SequencerUtilities.cpp
--- a/Engine/Source/Editor/Sequencer/Private/SequencerUtilities.cpp
+++ b/Engine/Source/Editor/Sequencer/Private/SequencerUtilities.cpp
@@ -1279,10 +1279,12 @@ FMovieScenePossessable* FSequencerUtilities::ConvertToPossessable(TSharedRef<ISe
 		return CreatedPossessable;
 	}
 
+	const FMovieSceneBindingReferences* BindingReferences = Sequence->GetBindingReferences();
+
 	FMovieScenePossessable* ExistingPossessable = MovieScene->FindPossessable(BindingGuid);
-	if (ExistingPossessable)
+	if (ExistingPossessable && BindingReferences)
 	{
-		if (const FMovieSceneBindingReference* ExistingReference = Sequence->GetBindingReferences()->GetReference(BindingGuid, BindingIndex))
+		if (const FMovieSceneBindingReference* ExistingReference = BindingReferences->GetReference(BindingGuid, BindingIndex))
 		{
 			if (!ExistingReference->CustomBinding)
 			{
@@ -1359,7 +1361,7 @@ FMovieScenePossessable* FSequencerUtilities::ConvertToPossessable(TSharedRef<ISe
 		}
 		else if (ExistingPossessable)
 		{
-			if (const FMovieSceneBindingReferences* BindingReferences = Sequence->GetBindingReferences())
+			if (BindingReferences)
 			{
 				// If we don't have multiple bound objects, use the Possessable name instead of the template label
 				if (BindingReferences->GetReferences(BindingGuid).Num() == 1)
@@ -1392,8 +1394,8 @@ FMovieScenePossessable* FSequencerUtilities::ConvertToPossessable(TSharedRef<ISe
 	// Create or replace the binding
 	FGuid NewPossessableGuid = CreateOrReplaceBinding(Sequencer, Sequence, PossessedObject, CreateBindingParams);
 
-	TArrayView<const FMovieSceneBindingReference> BindingReferences = Sequence->GetBindingReferences()->GetReferences(BindingGuid);
-	bool bAnySpawnablesLeft = Algo::AnyOf(BindingReferences, [](const FMovieSceneBindingReference& BindingReference) { return BindingReference.CustomBinding && BindingReference.CustomBinding->IsA<UMovieSceneSpawnableBindingBase>(); });
+	TArrayView<const FMovieSceneBindingReference> References = BindingReferences ? BindingReferences->GetReferences(BindingGuid) : TArrayView<const FMovieSceneBindingReference>();
+	bool bAnySpawnablesLeft = Algo::AnyOf(References, [](const FMovieSceneBindingReference& BindingReference) { return BindingReference.CustomBinding && BindingReference.CustomBinding->IsA<UMovieSceneSpawnableBindingBase>(); });
 
 	// If we're converting from a spawnable and none of the other bindings on the guid are spawnable, we'll need to remove the spawn track
 	if (bConvertFromSpawnable && !bAnySpawnablesLeft)
@@ -2701,7 +2703,7 @@ FGuid TryCreateCustomBinding(TSharedPtr<ISequencer> Sequencer, UObject* CustomBi
 
 	if (InParams.CustomBinding)
 	{
-		const FMovieSceneBindingReference* PreviousBindingReference = BindingReferences->GetReference(InParams.ReplacementGuid, InParams.BindingIndex);
+		const FMovieSceneBindingReference* PreviousBindingReference = BindingReferences ? BindingReferences->GetReference(InParams.ReplacementGuid, InParams.BindingIndex) : nullptr;
 		// We've been provided a custom binding pre-created. Ensure it supports the object given
 		if (CustomBindingObject == nullptr 
 			|| InParams.CustomBinding->SupportsBindingCreationFromObject(CustomBindingObject)
@@ -2789,7 +2791,10 @@ FGuid TryCreateCustomBinding(TSharedPtr<ISequencer> Sequencer, UObject* CustomBi
 		}
 
 		// Add the custom binding
-		BindingReferences->AddOrReplaceBinding(NewID, NewCustomBinding, InParams.BindingIndex);
+		if (BindingReferences)
+		{
+			BindingReferences->AddOrReplaceBinding(NewID, NewCustomBinding, InParams.BindingIndex);
+		}
 
 		UObject* SpawnedObject = nullptr;
 
@@ -3025,7 +3030,10 @@ FGuid CreateGenericBinding(TSharedPtr<ISequencer> Sequencer, UMovieSceneSequence
 
 			if (InParams.ReplacementGuid.IsValid() && !ParentID.IsValid())
 			{
-				BindingReferences->AddOrReplaceBinding(NewID, MoveTemp(Locator), InParams.BindingIndex);
+				if (BindingReferences)
+				{
+					BindingReferences->AddOrReplaceBinding(NewID, MoveTemp(Locator), InParams.BindingIndex);
+				}
 				if (Sequencer)
 				{
 					Sequencer->GetEvaluationState()->Invalidate(NewID, Sequencer->GetFocusedTemplateID());
@@ -3033,7 +3041,10 @@ FGuid CreateGenericBinding(TSharedPtr<ISequencer> Sequencer, UMovieSceneSequence
 			}
 			else
 			{
-				BindingReferences->AddBinding(NewID, MoveTemp(Locator));
+				if (BindingReferences)
+				{
+					BindingReferences->AddBinding(NewID, MoveTemp(Locator));
+				}
 			}
 		}
 
diff --git a/Engine/Source/Runtime/MovieScene/Private/EntitySystem/MovieSceneEntityLedger.cpp b/Engine/Source/Runtime/MovieScene/Private/EntitySystem/MovieSceneEntityLedger.cpp
--- a/Engine/Source/Runtime/MovieScene/Private/EntitySystem/MovieSceneEntityLedger.cpp
+++ b/Engine/Source/Runtime/MovieScene/Private/EntitySystem/MovieSceneEntityLedger.cpp
@@ -101,8 +101,12 @@ void FEntityLedger::UpdateOneShotEntities(UMovieSceneEntitySystemLinker* Linker,
 		}
 
 		Params.EntityID = Query.Entity.Key.EntityID;
-		Params.EntityMetaData = EntityField->FindMetaData(Query);
-		Params.SharedMetaData = EntityField->FindSharedMetaData(Query);
+
+		if (EntityField)
+		{
+			Params.EntityMetaData = EntityField->FindMetaData(Query);
+			Params.SharedMetaData = EntityField->FindSharedMetaData(Query);
+		}
 
 		if (ImportParams.bPreRoll && (Params.EntityMetaData == nullptr || Params.EntityMetaData->bEvaluateInSequencePreRoll == false))
 		{
@@ -161,8 +165,12 @@ void FEntityLedger::UpdateConditionalEntities(UMovieSceneEntitySystemLinker* Lin
 	for (const FMovieSceneEvaluationFieldEntityQuery& Query : ConditionalEntities)
 	{
 		Params.EntityID = Query.Entity.Key.EntityID;
-		Params.EntityMetaData = EntityField->FindMetaData(Query);
-		Params.SharedMetaData = EntityField->FindSharedMetaData(Query);
+
+		if (EntityField)
+		{
+			Params.EntityMetaData = EntityField->FindMetaData(Query);
+			Params.SharedMetaData = EntityField->FindSharedMetaData(Query);
+		}
 
 		// We cache all results here in the temp cache we've made so at least we won't re-run the same condition multiple times each tick
 		bool bConditionPassed = CanImportEntity(Linker, ImportParams, EntityField, Query, DummyEntitySet, ConditionResultCache, true);
@@ -250,7 +258,7 @@ void FEntityLedger::FindImportedEntities(TWeakObjectPtr<UObject> EntityOwner, TA
 bool FEntityLedger::CanImportEntity(UMovieSceneEntitySystemLinker* Linker, const FEntityImportSequenceParams& ImportParams, const FMovieSceneEntityComponentField* EntityField, const FMovieSceneEvaluationFieldEntityQuery& Query, FMovieSceneEvaluationFieldEntitySet& OutPerTickConditionalEntities, TMap<uint32, bool>& ConditionResultCache, bool bUpdatingPerTickEntities)
 {
 	// If we don't have a condition, just return true
-	const FMovieSceneEvaluationFieldEntityMetaData* EntityMetadata = EntityField->FindMetaData(Query);
+	const FMovieSceneEvaluationFieldEntityMetaData* EntityMetadata = EntityField ? EntityField->FindMetaData(Query) : nullptr;
 	if (!EntityMetadata || !EntityMetadata->Condition)
 	{
 		return true;
@@ -269,9 +277,12 @@ bool FEntityLedger::CanImportEntity(UMovieSceneEntitySystemLinker* Linker, const
 
 	FGuid BindingID;
 
-	if (const FMovieSceneEvaluationFieldSharedEntityMetaData* SharedMetadata = EntityField->FindSharedMetaData(Query))
+	if (EntityField)
 	{
-		BindingID = SharedMetadata->ObjectBindingID;
+		if (const FMovieSceneEvaluationFieldSharedEntityMetaData* SharedMetadata = EntityField->FindSharedMetaData(Query))
+		{
+			BindingID = SharedMetadata->ObjectBindingID;
+		}
 	}
 
 	// If we have a valid binding ID, and the condition depends on object binding, then we must ensure the object binding is resolved
@@ -323,8 +334,12 @@ void FEntityLedger::ImportEntity(UMovieSceneEntitySystemLinker* Linker, const FE
 	FEntityImportParams Params;
 	Params.Sequence = ImportParams;
 	Params.EntityID = Query.Entity.Key.EntityID;
-	Params.EntityMetaData = EntityField->FindMetaData(Query);
-	Params.SharedMetaData = EntityField->FindSharedMetaData(Query);
+	
+	if (EntityField)
+	{
+		Params.EntityMetaData = EntityField->FindMetaData(Query);
+		Params.SharedMetaData = EntityField->FindSharedMetaData(Query);
+	}
 
 	bool bUnsupportedPrePostRoll = false;
 	if (ImportParams.bPreRoll && (Params.EntityMetaData == nullptr || Params.EntityMetaData->bEvaluateInSequencePreRoll == false))
