Sequencer: Fix delete selected keys so that captures the proper Sequencer to delete keys from

This is a regression from UE-237110 which removed the DeleteKeys command because users were getting confused when binding it. That captured the Sequencer when the menu was first generated, so subsequence FSequencers would no longer be able to delete keys from the menu.

diff --git a/Engine/Source/Editor/Sequencer/Private/SequencerContextMenus.cpp b/Engine/Source/Editor/Sequencer/Private/SequencerContextMenus.cpp
--- a/Engine/Source/Editor/Sequencer/Private/SequencerContextMenus.cpp
+++ b/Engine/Source/Editor/Sequencer/Private/SequencerContextMenus.cpp
@@ -222,24 +222,34 @@ void FKeyContextMenu::PopulateMenu(FMenuBuilder& MenuBuilder, TSharedPtr<FExtend
 				Section.AddMenuEntry(FGenericCommands::Get().Copy);
 				Section.AddMenuEntry(FGenericCommands::Get().Duplicate);
 			}
+
+			if (Sequencer && ContextObject)
+			{
+				FToolMenuSection& KeysRemovalSection = InMenu->AddSection("KeyRemoval", FText::GetEmpty(), FToolMenuInsert(NAME_None, EToolMenuInsertType::Last));
+				KeysRemovalSection.AddSeparator(NAME_None);
+
+				KeysRemovalSection.AddMenuEntry(
+					NAME_None,
+					LOCTEXT("DeleteKeys.Label", "Delete Keys"),
+					LOCTEXT("DeleteKeys.Description", "Deletes the selected keys"),
+					FSlateIcon(),
+					FUIAction(
+						FExecuteAction::CreateLambda([WeakSequencer = ContextObject->WeakSequencer]
+						{
+							TSharedPtr<FSequencer> Sequencer = StaticCastSharedPtr<FSequencer>(WeakSequencer.Pin());
+							if (Sequencer)
+							{
+								Sequencer->DeleteSelectedKeys();
+							}
+						}))
+				);
+			}
 		}));
 
 		FToolMenuSection& KeysSection = KeyContextMenu->AddSection("SequencerKeys", LOCTEXT("KeysMenu", "Keys"));
 		KeysSection.AddMenuEntry(FSequencerCommands::Get().SetKeyTime);
 		KeysSection.AddMenuEntry(FSequencerCommands::Get().Rekey);
 		KeysSection.AddMenuEntry(FSequencerCommands::Get().SnapToFrame);
-
-		FToolMenuSection& KeysRemovalSection = KeyContextMenu->AddSection("KeyRemoval");
-		KeysRemovalSection.AddSeparator(NAME_None);
-
-		// This used to be a command in FSequencerCommands; removed because people were rebinding it and confused that delete, bound to FGenericCommands::Delete, still deleted keys.
-		KeysRemovalSection.AddMenuEntry(
-			NAME_None,
-			LOCTEXT("DeleteKeys.Label", "Delete Keys"),
-			LOCTEXT("DeleteKeys.Description", "Deletes the selected keys"),
-			FSlateIcon(),
-			FUIAction(FExecuteAction::CreateSP(WeakSequencer.Pin().Get(), &FSequencer::DeleteSelectedKeys))
-			);
 	}
 
 	USequencerToolMenuContext* ContextObject = NewObject<USequencerToolMenuContext>();
