Fixing a bug where the keyable indicator would not update for the selected object in the details panel after opening a sequence. It required adding to the sequence or selecting new objects.

diff --git a/Engine/Source/Editor/Sequencer/Private/LevelEditorSequencerIntegration.cpp b/Engine/Source/Editor/Sequencer/Private/LevelEditorSequencerIntegration.cpp
--- a/Engine/Source/Editor/Sequencer/Private/LevelEditorSequencerIntegration.cpp
+++ b/Engine/Source/Editor/Sequencer/Private/LevelEditorSequencerIntegration.cpp
@@ -954,6 +954,9 @@ void FLevelEditorSequencerIntegration::BindDetailHandler(const FLevelEditorSeque
 		ActivateDetailHandler();
 	});
 
+	// run the handler for any open details panels.
+	ActivateDetailHandler();
+	
 	// unbind keyframe handler
 	auto DeactivateDetailKeyframeHandler = [this, OnPropertyEditorOpenedHandle]()
 	{		
