[Sequencer] Filter Search Bar Tooltip update
- Removed the "(pressing enter selects the results)"

#rnx

diff --git a/Engine/Source/Editor/Sequencer/Private/SSequencer.cpp b/Engine/Source/Editor/Sequencer/Private/SSequencer.cpp
--- a/Engine/Source/Editor/Sequencer/Private/SSequencer.cpp
+++ b/Engine/Source/Editor/Sequencer/Private/SSequencer.cpp
@@ -1360,7 +1360,7 @@ TSharedRef<SWidget> SSequencer::ConstructSearchAndFilterRow()
 		SearchBox = SNew(SSequencerSearchBox, FilterBar)
 			.AddMetaData<FTagMetaData>(FTagMetaData(TEXT("SequencerFilterSearch")))
 			.HintText(LOCTEXT("FilterSearch", "Search..."))
-			.ToolTipText(LOCTEXT("FilterSearchHint", "Type here to search (pressing enter selects the results)"))
+			.ToolTipText(LOCTEXT("FilterSearchHint", "Type here to search"))
 			.OnTextChanged(this, &SSequencer::OnOutlinerSearchChanged)
 			.OnTextCommitted(this, &SSequencer::OnOutlinerSearchCommitted)
 			.OnSaveSearchClicked(this, &SSequencer::OnOutlinerSearchSaved);
