Sequencer: Put the mark frame properties in a submenu so that the menus gets sized appropriately.

Menu size calculation doesn't work well when a menu has a details panel. So, for now, move it to a submenu, otherwise the menu gets cutoff.

#rnx

diff --git a/Engine/Source/Editor/Sequencer/Private/SequencerTimeSliderController.cpp b/Engine/Source/Editor/Sequencer/Private/SequencerTimeSliderController.cpp
--- a/Engine/Source/Editor/Sequencer/Private/SequencerTimeSliderController.cpp
+++ b/Engine/Source/Editor/Sequencer/Private/SequencerTimeSliderController.cpp
@@ -2140,11 +2140,17 @@ TSharedRef<SWidget> FSequencerTimeSliderController::OpenSetPlaybackRangeMenu(con
 						SequencerSelection.MarkedFrames.Select(MarkedIndex);
 					}
 
-					const bool bLockedMarkedFrames = TimeSliderController->TimeSliderArgs.AreMarkedFramesLocked.Get();
+					TWeakPtr<FSequencer> WeakSequencer = Sequencer.ToWeakPtr();
 
-					TSharedRef<SMarkedFramePropertyWidget> Widget = SNew(SMarkedFramePropertyWidget, MovieScene, MarkedIndex, Sequencer.ToWeakPtr());
-					Widget->SetEnabled(!bLockedMarkedFrames);
-					Section.AddEntry(FToolMenuEntry::InitWidget(NAME_None, Widget, FText::GetEmpty(), false));
+					Section.AddSubMenu(TEXT("MarkProperties"), LOCTEXT("MarkPropertiesLabel", "Properties"), FText::GetEmpty(),
+						FNewToolMenuDelegate::CreateLambda([TimeSliderController, MovieScene, MarkedIndex, WeakSequencer](UToolMenu* InSubMenu)
+						{
+							const bool bLockedMarkedFrames = TimeSliderController->TimeSliderArgs.AreMarkedFramesLocked.Get();
+							TSharedRef<SMarkedFramePropertyWidget> Widget = SNew(SMarkedFramePropertyWidget, MovieScene, MarkedIndex, WeakSequencer);
+							Widget->SetEnabled(!bLockedMarkedFrames);
+							InSubMenu->AddMenuEntry("MarkPropertiesSection", FToolMenuEntry::InitWidget(NAME_None, Widget, FText::GetEmpty(), false));
+						})
+					);
 				}
 
 				Section.AddMenuEntry( 
