Sequencer: Disable throttling with scrub time mode

Same already happens in FCurveEditorDragOperation_ScrubTime

diff --git a/Engine/Source/Editor/Sequencer/Private/Tools/SequencerEditTool_Selection.cpp b/Engine/Source/Editor/Sequencer/Private/Tools/SequencerEditTool_Selection.cpp
--- a/Engine/Source/Editor/Sequencer/Private/Tools/SequencerEditTool_Selection.cpp
+++ b/Engine/Source/Editor/Sequencer/Private/Tools/SequencerEditTool_Selection.cpp
@@ -1,6 +1,7 @@
 // Copyright Epic Games, Inc. All Rights Reserved.
 
 #include "Tools/SequencerEditTool_Selection.h"
+#include "Application/ThrottleManager.h"
 #include "Styling/AppStyle.h"
 #include "SequencerCommonHelpers.h"
 #include "SSequencer.h"
@@ -123,6 +124,8 @@ public:
 
 	virtual void OnBeginDrag(const FPointerEvent& MouseEvent, FVector2D LocalMousePos, const UE::Sequencer::FVirtualTrackArea& VirtualTrackArea) override
 	{
+		FSlateThrottleManager::Get().DisableThrottle(true);
+
 		// Start a new marquee selection
 		InitialPosition = CurrentPosition = VirtualTrackArea.PhysicalToVirtual(LocalMousePos);
 		CurrentMousePos = LocalMousePos;
@@ -144,6 +147,7 @@ public:
 	virtual void OnEndDrag(const FPointerEvent& MouseEvent, FVector2D LocalMousePos, const UE::Sequencer::FVirtualTrackArea& VirtualTrackArea) override
 	{
 		Sequencer.SetPlaybackStatus(EMovieScenePlayerStatus::Stopped);
+		FSlateThrottleManager::Get().DisableThrottle(false);
 	}
 
 	virtual int32 OnPaint(const FGeometry& AllottedGeometry, const FSlateRect& MyCullingRect, FSlateWindowElementList& OutDrawElements, int32 LayerId) const override
