Sequencer: move GetKeyablePropertyPaths to a helper class for re-use outside of Sequencer

diff --git a/Engine/Source/Editor/Sequencer/Private/MVVM/ViewModels/ObjectBindingModel.cpp b/Engine/Source/Editor/Sequencer/Private/MVVM/ViewModels/ObjectBindingModel.cpp
--- a/Engine/Source/Editor/Sequencer/Private/MVVM/ViewModels/ObjectBindingModel.cpp
+++ b/Engine/Source/Editor/Sequencer/Private/MVVM/ViewModels/ObjectBindingModel.cpp
@@ -48,6 +48,7 @@
 #include "SequencerSettings.h"
 #include "StructUtils/PropertyBag.h"
 #include "MVVM/Views/ViewUtilities.h"
+#include "Misc/SequencerObjectBindingHelper.h"
 #include "Styling/AppStyle.h"
 #include "Styling/SlateIconFinder.h"
 #include "Widgets/SSequencerBindingLifetimeOverlay.h"
@@ -65,82 +66,6 @@ FAutoConsoleVariableRef CVarSequencerObjectBindingShowNestedProperties(
 	ECVF_Default
 );
 
-namespace
-{
-
-void GetKeyablePropertyPaths(UClass* Class, void* ValuePtr, UStruct* PropertySource, FPropertyPath PropertyPath, FSequencer& Sequencer, TArray<FPropertyPath>& KeyablePropertyPaths)
-{
-	for (TFieldIterator<FProperty> PropertyIterator(PropertySource); PropertyIterator; ++PropertyIterator)
-	{
-		FProperty* Property = *PropertyIterator;
-
-		if (Property && !Property->HasAnyPropertyFlags(CPF_Deprecated) && !Property->GetBoolMetaData(TEXT("SequencerHideProperty")))
-		{
-			PropertyPath.AddProperty(FPropertyInfo(Property));
-
-			if (FArrayProperty* ArrayProperty = CastField<FArrayProperty>(Property))
-			{
-				// If this is an array property, add property paths for each item in the array. If we don't know how to key that item,
-				// like for instance it's a custom struct for which we don't have a track editor, then we recurse into it and add
-				// property paths for its own properties.
-				FScriptArrayHelper ArrayHelper(ArrayProperty, ArrayProperty->ContainerPtrToValuePtr<void>(ValuePtr));
-				for (int32 Index = 0; Index < ArrayHelper.Num(); ++Index)
-				{
-					PropertyPath.AddProperty(FPropertyInfo(ArrayProperty->Inner, Index));
-
-					if (Sequencer.CanKeyProperty(FCanKeyPropertyParams(Class, PropertyPath)))
-					{
-						KeyablePropertyPaths.Add(PropertyPath);
-					}
-					else if (FStructProperty* StructProperty = CastField<FStructProperty>(ArrayProperty->Inner))
-					{
-						GetKeyablePropertyPaths(Class, ArrayHelper.GetRawPtr(Index), StructProperty->Struct, PropertyPath, Sequencer, KeyablePropertyPaths);
-					}
-
-					PropertyPath = *PropertyPath.TrimPath(1);
-				}
-			}
-			else if (Sequencer.CanKeyProperty(FCanKeyPropertyParams(Class, PropertyPath)))
-			{
-				// This is a property that we can key directly. That is: we have a track editor specifically for that property type,
-				// such as FVector or FLinearColor or other well known structs. This also includes custom/system-specific structs
-				// like FMargin, for which UMG registers a custom track.
-				KeyablePropertyPaths.Add(PropertyPath);
-			}
-			else if (FStructProperty* StructProperty = CastField<FStructProperty>(Property))
-			{
-				if (StructProperty->Struct == FInstancedPropertyBag::StaticStruct() && StructProperty->GetBoolMetaData(TEXT("InterpBagProperties")))
-				{
-					// It's a property bag. Show the properties inside if we have been explicitly allowed to do so. Note that we are
-					// using a custom metadata ("InterpBagProperties") because the built-in "Interp" tag only works on BlueprintTypes.
-					FInstancedPropertyBag* PropertyBag = StructProperty->ContainerPtrToValuePtr<FInstancedPropertyBag>(ValuePtr);
-					FStructProperty* PropertyBagValueProperty = CastFieldChecked<FStructProperty>(StructProperty->Struct->FindPropertyByName(TEXT("Value")));
-					PropertyPath.AddProperty(FPropertyInfo(PropertyBagValueProperty));
-					const UPropertyBag* PropertyBagStruct = PropertyBag->GetPropertyBagStruct();
-					GetKeyablePropertyPaths(Class, PropertyBag, const_cast<UPropertyBag*>(PropertyBagStruct), PropertyPath, Sequencer, KeyablePropertyPaths);
-				}
-				else if (StructProperty->Struct == FInstancedStruct::StaticStruct() && StructProperty->GetBoolMetaData(TEXT("InterpStructProperties")))
-				{
-					// As above, but for an instanced struct.
-					FInstancedStruct* InstancedStruct = StructProperty->ContainerPtrToValuePtr<FInstancedStruct>(ValuePtr);
-					const UScriptStruct* InstancedStructType = InstancedStruct->GetScriptStruct();
-					GetKeyablePropertyPaths(Class, InstancedStruct, const_cast<UScriptStruct*>(InstancedStructType), PropertyPath, Sequencer, KeyablePropertyPaths);
-				}
-				else
-				{
-					// It's a struct property that we don't know how to key directly, so add property paths for its own properties.
-					// The user will have to key them individually.
-					GetKeyablePropertyPaths(Class, StructProperty->ContainerPtrToValuePtr<void>(ValuePtr), StructProperty->Struct, PropertyPath, Sequencer, KeyablePropertyPaths);
-				}
-			}
-
-			PropertyPath = *PropertyPath.TrimPath(1);
-		}
-	}
-}
-
-} // anon-namespace
-
 FObjectBindingModel::FObjectBindingModel(FSequenceModel* InOwnerModel, const FMovieSceneBinding& InBinding)
 	: ObjectBindingID(InBinding.GetObjectGuid())
 	, TrackAreaList(EViewModelListType::TrackArea)
@@ -655,8 +580,8 @@ TSharedRef<SWidget> FObjectBindingModel::GetAddTrackMenuContent()
 
 	if (BoundObject != nullptr)
 	{
-		FPropertyPath PropertyPath;
-		GetKeyablePropertyPaths(BoundObject->GetClass(), BoundObject, BoundObject->GetClass(), PropertyPath, *Sequencer, KeyablePropertyPaths);
+		TSharedRef<ISequencer> SequencerInterface = Sequencer.ToSharedRef();
+		FSequencerObjectBindingHelper::GetKeyablePropertyPaths(BoundObject, SequencerInterface, KeyablePropertyPaths);
 	}
 
 	AddPropertyMenuItems(AddTrackMenuBuilder, NumStartingBlocks, KeyablePropertyPaths, 0);
diff --git a/Engine/Source/Editor/Sequencer/Private/Misc/SequencerObjectBindingHelper.cpp b/Engine/Source/Editor/Sequencer/Private/Misc/SequencerObjectBindingHelper.cpp
new file mode 100644
--- /dev/null
+++ b/Engine/Source/Editor/Sequencer/Private/Misc/SequencerObjectBindingHelper.cpp
@@ -0,0 +1,84 @@
+// Copyright Epic Games, Inc. All Rights Reserved.
+
+#include "Misc/SequencerObjectBindingHelper.h"
+
+#include "KeyPropertyParams.h"
+#include "ISequencer.h"
+#include "StructUtils/PropertyBag.h"
+
+void FSequencerObjectBindingHelper::GetKeyablePropertyPaths(const UObject* Object, TSharedRef<ISequencer> Sequencer, TArray<FPropertyPath>& KeyablePropertyPaths)
+{
+	FPropertyPath PropertyPath;
+	GetKeyablePropertyPaths(Object->GetClass(), Object, Object->GetClass(), PropertyPath, Sequencer, KeyablePropertyPaths);
+}
+
+void FSequencerObjectBindingHelper::GetKeyablePropertyPaths(const UClass* Class, const void* ValuePtr, const UStruct* PropertySource, FPropertyPath PropertyPath, TSharedRef<ISequencer> Sequencer, TArray<FPropertyPath>& KeyablePropertyPaths)
+{
+	for (TFieldIterator<FProperty> PropertyIterator(PropertySource); PropertyIterator; ++PropertyIterator)
+	{
+		FProperty* Property = *PropertyIterator;
+
+		if (Property && !Property->HasAnyPropertyFlags(CPF_Deprecated) && !Property->GetBoolMetaData(TEXT("SequencerHideProperty")))
+		{
+			PropertyPath.AddProperty(FPropertyInfo(Property));
+
+			if (FArrayProperty* ArrayProperty = CastField<FArrayProperty>(Property))
+			{
+				// If this is an array property, add property paths for each item in the array. If we don't know how to key that item,
+				// like for instance it's a custom struct for which we don't have a track editor, then we recurse into it and add
+				// property paths for its own properties.
+				FScriptArrayHelper ArrayHelper(ArrayProperty, ArrayProperty->ContainerPtrToValuePtr<void>(ValuePtr));
+				for (int32 Index = 0; Index < ArrayHelper.Num(); ++Index)
+				{
+					PropertyPath.AddProperty(FPropertyInfo(ArrayProperty->Inner, Index));
+
+					if (Sequencer->CanKeyProperty(FCanKeyPropertyParams(Class, PropertyPath)))
+					{
+						KeyablePropertyPaths.Add(PropertyPath);
+					}
+					else if (FStructProperty* StructProperty = CastField<FStructProperty>(ArrayProperty->Inner))
+					{
+						GetKeyablePropertyPaths(Class, ArrayHelper.GetRawPtr(Index), StructProperty->Struct, PropertyPath, Sequencer, KeyablePropertyPaths);
+					}
+
+					PropertyPath = *PropertyPath.TrimPath(1);
+				}
+			}
+			else if (Sequencer->CanKeyProperty(FCanKeyPropertyParams(Class, PropertyPath)))
+			{
+				// This is a property that we can key directly. That is: we have a track editor specifically for that property type,
+				// such as FVector or FLinearColor or other well known structs. This also includes custom/system-specific structs
+				// like FMargin, for which UMG registers a custom track.
+				KeyablePropertyPaths.Add(PropertyPath);
+			}
+			else if (FStructProperty* StructProperty = CastField<FStructProperty>(Property))
+			{
+				if (StructProperty->Struct == FInstancedPropertyBag::StaticStruct() && StructProperty->GetBoolMetaData(TEXT("InterpBagProperties")))
+				{
+					// It's a property bag. Show the properties inside if we have been explicitly allowed to do so. Note that we are
+					// using a custom metadata ("InterpBagProperties") because the built-in "Interp" tag only works on BlueprintTypes.
+					const FInstancedPropertyBag* PropertyBag = StructProperty->ContainerPtrToValuePtr<FInstancedPropertyBag>(ValuePtr);
+					FStructProperty* PropertyBagValueProperty = CastFieldChecked<FStructProperty>(StructProperty->Struct->FindPropertyByName(TEXT("Value")));
+					PropertyPath.AddProperty(FPropertyInfo(PropertyBagValueProperty));
+					const UPropertyBag* PropertyBagStruct = PropertyBag->GetPropertyBagStruct();
+					GetKeyablePropertyPaths(Class, PropertyBag, const_cast<UPropertyBag*>(PropertyBagStruct), PropertyPath, Sequencer, KeyablePropertyPaths);
+				}
+				else if (StructProperty->Struct == FInstancedStruct::StaticStruct() && StructProperty->GetBoolMetaData(TEXT("InterpStructProperties")))
+				{
+					// As above, but for an instanced struct.
+					const FInstancedStruct* InstancedStruct = StructProperty->ContainerPtrToValuePtr<FInstancedStruct>(ValuePtr);
+					const UScriptStruct* InstancedStructType = InstancedStruct->GetScriptStruct();
+					GetKeyablePropertyPaths(Class, InstancedStruct, const_cast<UScriptStruct*>(InstancedStructType), PropertyPath, Sequencer, KeyablePropertyPaths);
+				}
+				else
+				{
+					// It's a struct property that we don't know how to key directly, so add property paths for its own properties.
+					// The user will have to key them individually.
+					GetKeyablePropertyPaths(Class, StructProperty->ContainerPtrToValuePtr<void>(ValuePtr), StructProperty->Struct, PropertyPath, Sequencer, KeyablePropertyPaths);
+				}
+			}
+
+			PropertyPath = *PropertyPath.TrimPath(1);
+		}
+	}
+}
diff --git a/Engine/Source/Editor/Sequencer/Public/Misc/SequencerObjectBindingHelper.h b/Engine/Source/Editor/Sequencer/Public/Misc/SequencerObjectBindingHelper.h
new file mode 100644
--- /dev/null
+++ b/Engine/Source/Editor/Sequencer/Public/Misc/SequencerObjectBindingHelper.h
@@ -0,0 +1,41 @@
+// Copyright Epic Games, Inc. All Rights Reserved.
+
+#pragma once
+
+#include "Containers/Array.h"
+#include "PropertyPath.h"
+
+class ISequencer;
+class UClass;
+class UStruct;
+
+/**
+ * Helper class for using object bindings in the Sequencer editor.
+ */
+class FSequencerObjectBindingHelper
+{
+public:
+
+	/**
+	 * Gets the keyable properties on the given object, for the given Sequencer editor.
+	 *
+	 * @param Object            The object to key.
+	 * @param Sequencer         The Sequencer that will be responsible for animating the properties.
+	 * @param KeyablePropertyPaths  The list of properties that can be keyed.
+	 */
+	SEQUENCER_API static void GetKeyablePropertyPaths(const UObject* Object, TSharedRef<ISequencer> Sequencer, TArray<FPropertyPath>& KeyablePropertyPaths);
+
+	/**
+	 * Gets the keyable properties on the given object, for the given Sequencer editor.
+	 *
+	 * @param Class             The class of the UObject to animate.
+	 * @param ValuePtr          A pointer to object, or one of its struct members.
+	 * @param PropertySource    The type of the ValuePtr.
+	 * @param PropertyPath      The path to ValuePtr (should be empty if ValuePtr is the object itself).
+	 * @param Sequencer         The Sequencer that will be responsible for animating the properties.
+	 * @param KeyablePropertyPaths  The list of properties that can be keyed.
+	 */
+	SEQUENCER_API static void GetKeyablePropertyPaths(const UClass* Class, const void* ValuePtr, const UStruct* PropertySource, FPropertyPath PropertyPath, TSharedRef<ISequencer> Sequencer, TArray<FPropertyPath>& KeyablePropertyPaths);
+
+};
+
