Sequencer: remove the "Blocking Evaluation" checkbox from the advanced menu, since we can already set it from the playback options

diff --git a/Engine/Source/Editor/Sequencer/Private/SSequencer.cpp b/Engine/Source/Editor/Sequencer/Private/SSequencer.cpp
--- a/Engine/Source/Editor/Sequencer/Private/SSequencer.cpp
+++ b/Engine/Source/Editor/Sequencer/Private/SSequencer.cpp
@@ -1888,18 +1888,6 @@ void SSequencer::FillAdvancedMenu(FMenuBuilder& MenuBuilder)
 			),
 			NAME_None,
 			EUserInterfaceActionType::ToggleButton);
-
-		MenuBuilder.AddMenuEntry(
-			LOCTEXT("VolatilityBlockingEvaluationLabel", "Blocking Evaluation"),
-			LOCTEXT("VolatilityBlockingEvaluationTooltip", "Indicates that a sequence must fully evaluate and apply its state every time it is updated, blocking until complete. Should be used sparingly as it will severely affect performance."),
-			FSlateIcon(),
-			FUIAction(
-				FExecuteAction::CreateLambda(ToggleVolatility, EMovieSceneSequenceFlags::BlockingEvaluation),
-				FCanExecuteAction(),
-				FIsActionChecked::CreateLambda(IsVolatilityChecked, EMovieSceneSequenceFlags::BlockingEvaluation)
-			),
-			NAME_None,
-			EUserInterfaceActionType::ToggleButton);
 	}
 	MenuBuilder.EndSection();
 }
