Sequencer:  New ScrubTimeStartFromCursor Settings for Sequencer and Curve Editor so it will scrub from strating position, will default to true changing existing behavior. If Playing when scrubbing or going to start/end frames it will continue playing. Needed to add OnMouseButtonDown to ICurveEditorDragOperation so it could set start time without the drag and SetStoppedPosition to ITimeSlider so we can correctly set a start time from which it will play frame, scrub time doesn't work in this case.

diff --git a/Engine/Source/Editor/CurveEditor/Private/CurveEditorSettings.cpp b/Engine/Source/Editor/CurveEditor/Private/CurveEditorSettings.cpp
--- a/Engine/Source/Editor/CurveEditor/Private/CurveEditorSettings.cpp
+++ b/Engine/Source/Editor/CurveEditor/Private/CurveEditorSettings.cpp
@@ -19,6 +19,7 @@ UCurveEditorSettings::UCurveEditorSettings()
 	ZoomPosition = ECurveEditorZoomPosition::CurrentTime;
 	SnapAxis = ECurveEditorSnapAxis::CESA_None;
 	bSnapTimeToSelection = false;
+	bScrubTimeStartFromCursor = true;
 
 	SelectionColor = FLinearColor::White;
 	ParentSpaceCustomColor = FLinearColor(.93, .31, .19); //pastel orange
@@ -337,6 +338,20 @@ void UCurveEditorSettings::DeleteSpaceSwitchColor(const FString& InControlName)
 	}
 }
 
+bool UCurveEditorSettings::GetScrubTimeStartFromCursor() const
+{
+	return bScrubTimeStartFromCursor;
+}
+
+void UCurveEditorSettings::SetScrubTimeStartFromCursor(bool bInValue)
+{
+	if (bScrubTimeStartFromCursor != bInValue)
+	{
+		bScrubTimeStartFromCursor = bInValue;
+		SaveConfig();
+	}
+}
+
 FLinearColor UCurveEditorSettings::GetNextRandomColor()
 {
 	static TArray<FLinearColor> IndexedColor;
diff --git a/Engine/Source/Editor/CurveEditor/Private/DragOperations/CurveEditorDragOperation_ScrubTime.cpp b/Engine/Source/Editor/CurveEditor/Private/DragOperations/CurveEditorDragOperation_ScrubTime.cpp
--- a/Engine/Source/Editor/CurveEditor/Private/DragOperations/CurveEditorDragOperation_ScrubTime.cpp
+++ b/Engine/Source/Editor/CurveEditor/Private/DragOperations/CurveEditorDragOperation_ScrubTime.cpp
@@ -13,21 +13,57 @@
 
 
 FCurveEditorDragOperation_ScrubTime::FCurveEditorDragOperation_ScrubTime(FCurveEditor* InCurveEditor)
-	: CurveEditor(InCurveEditor) 
+	: CurveEditor(InCurveEditor) , NextPlayerStatus(ETimeSliderPlaybackStatus::Stopped)
 {}
 
+void FCurveEditorDragOperation_ScrubTime::OnMouseButtonDown(FVector2D InitialPosition, const FGeometry& MyGeometry, const FPointerEvent& MouseEvent)
+{
+	TSharedPtr<ITimeSliderController> TimeSliderController = CurveEditor->GetTimeSliderController();
+	if (TimeSliderController)
+	{
+		const UCurveEditorSettings* Settings = GetDefault<UCurveEditorSettings>();
+		if (Settings && Settings->GetScrubTimeStartFromCursor())
+		{
+			const FFrameRate TickResolution = TimeSliderController->GetTickResolution();
+			FCurveEditorScreenSpaceH InputSpace = CurveEditor->GetPanelInputSpace();
+			const double InitialMouseTime = InputSpace.ScreenToSeconds(InitialPosition.X);
+
+			if (TimeSliderController->GetPlaybackStatus() == ETimeSliderPlaybackStatus::Playing)
+			{
+				TimeSliderController->SetPlaybackStatus(ETimeSliderPlaybackStatus::Paused);
+				TimeSliderController->SetStoppedPosition((InitialMouseTime)*TickResolution);
+				TimeSliderController->SetPlaybackStatus(ETimeSliderPlaybackStatus::Playing);
+			}
+			else
+			{
+				TimeSliderController->SetScrubPosition((InitialMouseTime)*TickResolution, /*bEvaluate*/ true);
+			}
+		}
+
+	}
+}
+
 void FCurveEditorDragOperation_ScrubTime::OnBeginDrag(FVector2D InitialPosition, FVector2D CurrentPosition, const FPointerEvent& MouseEvent)
 {
 	FSlateThrottleManager::Get().DisableThrottle(true);
+	NextPlayerStatus = ETimeSliderPlaybackStatus::Stopped;
 	TSharedPtr<ITimeSliderController> TimeSliderController = CurveEditor->GetTimeSliderController();
 	if (TimeSliderController)
-	{
+	{	
 		SnappingState.Reset();
 		LastMousePosition = CurrentPosition;
 		const FFrameRate TickResolution = TimeSliderController->GetTickResolution();
-
-		InitialTime = TickResolution.AsSeconds(TimeSliderController->GetScrubPosition());
+		if (TimeSliderController->GetPlaybackStatus() == ETimeSliderPlaybackStatus::Playing)
+		{
+			NextPlayerStatus = ETimeSliderPlaybackStatus::Playing;
+		}
+		else
+		{
+ 			NextPlayerStatus = ETimeSliderPlaybackStatus::Stopped;
+		}		
 		TimeSliderController->SetPlaybackStatus(ETimeSliderPlaybackStatus::Scrubbing);
+		
+		InitialTime = TickResolution.AsSeconds(TimeSliderController->GetScrubPosition());
 	}
 }
 
@@ -36,6 +72,7 @@ void FCurveEditorDragOperation_ScrubTime::OnDrag(FVector2D InitialPosition, FVec
 	TSharedPtr<ITimeSliderController> TimeSliderController = CurveEditor->GetTimeSliderController();
 	if (TimeSliderController)
 	{
+
 		FVector2D PixelDelta = CurveEditor->GetAxisSnap().GetSnappedPosition(InitialPosition, LastMousePosition, CurrentPosition, MouseEvent, SnappingState, true) - InitialPosition;
 		LastMousePosition = CurrentPosition;
 
@@ -61,7 +98,7 @@ void FCurveEditorDragOperation_ScrubTime::OnEndDrag(FVector2D InitialPosition, F
 	TSharedPtr<ITimeSliderController> TimeSliderController = CurveEditor->GetTimeSliderController();
 	if (TimeSliderController)
 	{
-		TimeSliderController->SetPlaybackStatus(ETimeSliderPlaybackStatus::Stopped);
+		TimeSliderController->SetPlaybackStatus(NextPlayerStatus);
 	}
 	FSlateThrottleManager::Get().DisableThrottle(false);
 }
@@ -71,6 +108,6 @@ void FCurveEditorDragOperation_ScrubTime::OnCancelDrag()
 	TSharedPtr<ITimeSliderController> TimeSliderController = CurveEditor->GetTimeSliderController();
 	if (TimeSliderController)
 	{
-		TimeSliderController->SetPlaybackStatus(ETimeSliderPlaybackStatus::Stopped);
+		TimeSliderController->SetPlaybackStatus(NextPlayerStatus);
 	}
 }
\ No newline at end of file
diff --git a/Engine/Source/Editor/CurveEditor/Private/DragOperations/CurveEditorDragOperation_ScrubTime.h b/Engine/Source/Editor/CurveEditor/Private/DragOperations/CurveEditorDragOperation_ScrubTime.h
--- a/Engine/Source/Editor/CurveEditor/Private/DragOperations/CurveEditorDragOperation_ScrubTime.h
+++ b/Engine/Source/Editor/CurveEditor/Private/DragOperations/CurveEditorDragOperation_ScrubTime.h
@@ -7,6 +7,7 @@
 #include "Input/Reply.h"
 #include "Math/Vector2D.h"
 #include "Templates/SharedPointer.h"
+#include "ITimeSlider.h"
 
 class FCurveEditor;
 struct FPointerEvent;
@@ -17,6 +18,7 @@ class FCurveEditorDragOperation_ScrubTime : public ICurveEditorDragOperation
 public:
 	FCurveEditorDragOperation_ScrubTime(FCurveEditor* CurveEditor);
 
+	virtual void OnMouseButtonDown(FVector2D InitialPosition, const FGeometry& MyGeometry, const FPointerEvent& MouseEvent) override;
 	virtual void OnBeginDrag(FVector2D InitialPosition, FVector2D CurrentPosition, const FPointerEvent& MouseEvent) override;
 	virtual void OnDrag(FVector2D InitialPosition, FVector2D CurrentPosition, const FPointerEvent& MouseEvent) override;
 	virtual void OnEndDrag(FVector2D InitialPosition, FVector2D CurrentPosition, const FPointerEvent& MouseEvent) override;
@@ -29,4 +31,5 @@ private:
 	FCurveEditorAxisSnap::FSnapState SnappingState;
 	FVector2D LastMousePosition;
 	double InitialTime;
+	ETimeSliderPlaybackStatus NextPlayerStatus;
 };
\ No newline at end of file
diff --git a/Engine/Source/Editor/CurveEditor/Private/DragOperations/ICurveEditorDragOperation.cpp b/Engine/Source/Editor/CurveEditor/Private/DragOperations/ICurveEditorDragOperation.cpp
--- a/Engine/Source/Editor/CurveEditor/Private/DragOperations/ICurveEditorDragOperation.cpp
+++ b/Engine/Source/Editor/CurveEditor/Private/DragOperations/ICurveEditorDragOperation.cpp
@@ -7,6 +7,11 @@ struct FCurvePointHandle;
 struct FGeometry;
 struct FPointerEvent;
 
+void ICurveEditorDragOperation::MouseButtonDown(FVector2D InitialPosition, const FGeometry& MyGeometry, const FPointerEvent& MouseEvent)
+{
+	OnMouseButtonDown(InitialPosition, MyGeometry, MouseEvent);
+}
+
 void ICurveEditorDragOperation::BeginDrag(FVector2D InitialPosition, FVector2D CurrentPosition, const FPointerEvent& MouseEvent)
 {
 	OnBeginDrag(InitialPosition, CurrentPosition, MouseEvent);
diff --git a/Engine/Source/Editor/CurveEditor/Private/SCurveEditorViewContainer.cpp b/Engine/Source/Editor/CurveEditor/Private/SCurveEditorViewContainer.cpp
--- a/Engine/Source/Editor/CurveEditor/Private/SCurveEditorViewContainer.cpp
+++ b/Engine/Source/Editor/CurveEditor/Private/SCurveEditorViewContainer.cpp
@@ -260,6 +260,8 @@ FReply SCurveEditorViewContainer::OnMouseButtonDown(const FGeometry& MyGeometry,
 			DragOperation = FCurveEditorDelayedDrag(MousePixel, MouseEvent.GetEffectingButton());
 			DragOperation->DragImpl = MakeUnique<FCurveEditorDragOperation_Marquee>(CurveEditor.Get());
 		}
+		const FVector2D InitialPosition = DragOperation->GetInitialPosition();
+		DragOperation->DragImpl->MouseButtonDown(InitialPosition, MyGeometry, MouseEvent);
 		return FReply::Handled();
 	}
 	// Middle Click + Alt Pan
@@ -269,6 +271,8 @@ FReply SCurveEditorViewContainer::OnMouseButtonDown(const FGeometry& MyGeometry,
 		{
 			DragOperation = FCurveEditorDelayedDrag(MousePixel, MouseEvent.GetEffectingButton());
 			DragOperation->DragImpl = MakeUnique<FCurveEditorDragOperation_PanInput>(CurveEditor.Get());
+			const FVector2D InitialPosition = DragOperation->GetInitialPosition();
+			DragOperation->DragImpl->MouseButtonDown(InitialPosition, MyGeometry, MouseEvent);
 			return FReply::Handled();
 		}
 	}
@@ -279,6 +283,8 @@ FReply SCurveEditorViewContainer::OnMouseButtonDown(const FGeometry& MyGeometry,
 		{
 			DragOperation = FCurveEditorDelayedDrag(MousePixel, MouseEvent.GetEffectingButton());
 			DragOperation->DragImpl = MakeUnique<FCurveEditorDragOperation_Zoom>(CurveEditor.Get(), nullptr);
+			const FVector2D InitialPosition = DragOperation->GetInitialPosition();
+			DragOperation->DragImpl->MouseButtonDown(InitialPosition, MyGeometry, MouseEvent);
 			return FReply::Handled();
 		}
 		// Pan Timeline
@@ -286,6 +292,8 @@ FReply SCurveEditorViewContainer::OnMouseButtonDown(const FGeometry& MyGeometry,
 		{
 			DragOperation = FCurveEditorDelayedDrag(MousePixel, MouseEvent.GetEffectingButton());
 			DragOperation->DragImpl = MakeUnique<FCurveEditorDragOperation_PanInput>(CurveEditor.Get());
+			const FVector2D InitialPosition = DragOperation->GetInitialPosition();
+			DragOperation->DragImpl->MouseButtonDown(InitialPosition, MyGeometry, MouseEvent);
 			return FReply::Handled();
 		}
 	}
diff --git a/Engine/Source/Editor/CurveEditor/Public/CurveEditorSettings.h b/Engine/Source/Editor/CurveEditor/Public/CurveEditorSettings.h
--- a/Engine/Source/Editor/CurveEditor/Public/CurveEditorSettings.h
+++ b/Engine/Source/Editor/CurveEditor/Public/CurveEditorSettings.h
@@ -212,6 +212,11 @@ public:
 	/** Sets whether to draw a value indicator line for the minimum and maximum key in the selected key range. */
 	UE_API void SetShowValueIndicators(bool bValue);
 
+	/** Gets whether or not scrubbing time hot key starts from cursor. */
+	UE_API bool GetScrubTimeStartFromCursor() const;
+	/** Sets whether or not scrubbing time hot key starts from cursor. */
+	UE_API void SetScrubTimeStartFromCursor(bool bInValue);
+
 protected:
 
 	UE_API virtual void PostEditChangeProperty(FPropertyChangedEvent& PropertyChangedEvent) override;
@@ -269,6 +274,10 @@ protected:
 	UPROPERTY(config, EditAnywhere, Category = "Curve Editor")
 	float TreeViewWidth;
 
+	/** When enabled, scrubbing time hotkey will move time from initial cursor position  */
+	UPROPERTY(config, EditAnywhere, Category = "Curve Editor")
+	bool bScrubTimeStartFromCursor;
+
 	/**
 	 * Determines how close you must move the selection marque to the center of a point in order to select it.
 	 * 
diff --git a/Engine/Source/Editor/CurveEditor/Public/ICurveEditorDragOperation.h b/Engine/Source/Editor/CurveEditor/Public/ICurveEditorDragOperation.h
--- a/Engine/Source/Editor/CurveEditor/Public/ICurveEditorDragOperation.h
+++ b/Engine/Source/Editor/CurveEditor/Public/ICurveEditorDragOperation.h
@@ -31,6 +31,11 @@ public:
 
 	virtual ~ICurveEditorDragOperation() {}
 
+	/**
+	 * Start of the drag operation
+	 */
+	void MouseButtonDown(FVector2D InitialPosition, const FGeometry& MyGeometry, const FPointerEvent& MouseEvent);
+
 	/**
 	 * Begin this drag operation with the specified initial and current positions
 	 */
@@ -66,6 +71,10 @@ public:
 
 protected:
 
+	/** Implementation method for derived types to handle mouse down*/
+	virtual void OnMouseButtonDown(FVector2D InitialPosition, const FGeometry& MyGeometry, const FPointerEvent& MouseEvent)
+	{}
+
 	/** Implementation method for derived types to begin a drag */
 	virtual void OnBeginDrag(FVector2D InitialPosition, FVector2D CurrentPosition, const FPointerEvent& MouseEvent)
 	{}
diff --git a/Engine/Source/Editor/Sequencer/Private/Sequencer.cpp b/Engine/Source/Editor/Sequencer/Private/Sequencer.cpp
--- a/Engine/Source/Editor/Sequencer/Private/Sequencer.cpp
+++ b/Engine/Source/Editor/Sequencer/Private/Sequencer.cpp
@@ -5153,14 +5153,23 @@ FReply FSequencer::OnStepBackward(FFrameNumber Increment)
 
 FReply FSequencer::OnJumpToStart()
 {
+	const bool bWasPlaying = (GetPlaybackStatus() == EMovieScenePlayerStatus::Playing);
+
 	SetPlaybackStatus(EMovieScenePlayerStatus::Stepping);
 	SetLocalTime(UE::MovieScene::DiscreteInclusiveLower(GetTimeBounds()), ESnapTimeMode::STM_None);
+	
+	if (bWasPlaying)
+	{
+		SetPlaybackStatus(EMovieScenePlayerStatus::Playing);
+	}
 	return FReply::Handled();
 }
 
 
 FReply FSequencer::OnJumpToEnd()
 {
+	const bool bWasPlaying = (GetPlaybackStatus() == EMovieScenePlayerStatus::Playing);
+
 	SetPlaybackStatus(EMovieScenePlayerStatus::Stepping);
 	const bool bInsetDisplayFrame = ScrubStyle == ESequencerScrubberStyle::FrameBlock && Settings->GetForceWholeFrames();
 
@@ -5173,6 +5182,11 @@ FReply FSequencer::OnJumpToEnd()
 	FFrameTime NewTime = UE::MovieScene::DiscreteExclusiveUpper(GetTimeBounds()) - OneFrame;
 
 	SetLocalTime(NewTime, ESnapTimeMode::STM_None);
+
+	if (bWasPlaying)
+	{
+		SetPlaybackStatus(EMovieScenePlayerStatus::Playing);
+	}
 	return FReply::Handled();
 }
 
diff --git a/Engine/Source/Editor/Sequencer/Private/SequencerEdMode.cpp b/Engine/Source/Editor/Sequencer/Private/SequencerEdMode.cpp
--- a/Engine/Source/Editor/Sequencer/Private/SequencerEdMode.cpp
+++ b/Engine/Source/Editor/Sequencer/Private/SequencerEdMode.cpp
@@ -216,7 +216,8 @@ FSequencerEdMode::~FSequencerEdMode()
 
 void FSequencerEdMode::Enter()
 {
-	bIsTracking = false;;
+	bIsTracking = false;
+	NextPlayerStatus = EMovieScenePlayerStatus::Stopped;
 	StartXValue.Reset();
 	FEdMode::Enter();
 }
@@ -339,6 +340,8 @@ bool FSequencerEdMode::StartTracking(FEditorViewportClient* InViewportClient, FV
 		}
 		if (ActiveSequencer.IsValid())
 		{
+			NextPlayerStatus = (ActiveSequencer->GetPlaybackStatus() == EMovieScenePlayerStatus::Playing) ?EMovieScenePlayerStatus::Playing :
+				EMovieScenePlayerStatus::Stopped;
 			ActiveSequencer->SetPlaybackStatus(EMovieScenePlayerStatus::Scrubbing);
 			ActiveSequencer->OnBeginScrubbingEvent().Broadcast();
 		}
@@ -422,6 +425,7 @@ bool FSequencerEdMode::ProcessCapturedMouseMoves(FEditorViewportClient* InViewpo
 {
 	const TSharedPtr<ISequencer> ActiveSequencer = GetFirstActiveSequencer();
 	const bool bTimeChange = IsPressingMoveTimeSlider(InViewport);
+
 	if (CapturedMouseMoves.Num() > 0)
 	{
 		if (ActiveSequencer.IsValid())
@@ -458,6 +462,8 @@ bool FSequencerEdMode::ProcessCapturedMouseMoves(FEditorViewportClient* InViewpo
 					}
 					if (bIsTracking == false)
 					{
+						NextPlayerStatus = (ActiveSequencer->GetPlaybackStatus() == EMovieScenePlayerStatus::Playing) ?  EMovieScenePlayerStatus::Playing :
+							EMovieScenePlayerStatus::Stopped;
 						ActiveSequencer->SetPlaybackStatus(EMovieScenePlayerStatus::Scrubbing);
 						ActiveSequencer->OnBeginScrubbingEvent().Broadcast();
 						bIsTracking = true;
@@ -467,7 +473,7 @@ bool FSequencerEdMode::ProcessCapturedMouseMoves(FEditorViewportClient* InViewpo
 			else if(bIsTracking)
 			{
 				bIsTracking = false;
-				ActiveSequencer->SetPlaybackStatus(EMovieScenePlayerStatus::Stopped);
+				ActiveSequencer->SetPlaybackStatus(NextPlayerStatus);
 				ActiveSequencer->OnEndScrubbingEvent().Broadcast();
 				StartXValue.Reset();
 			}
@@ -484,7 +490,7 @@ bool FSequencerEdMode::ProcessCapturedMouseMoves(FEditorViewportClient* InViewpo
 		if (ActiveSequencer.IsValid())
 		{
 			bIsTracking = false;
-			ActiveSequencer->SetPlaybackStatus(EMovieScenePlayerStatus::Stopped);
+			ActiveSequencer->SetPlaybackStatus(NextPlayerStatus);
 			ActiveSequencer->OnEndScrubbingEvent().Broadcast();
 			StartXValue.Reset();
 		}
diff --git a/Engine/Source/Editor/Sequencer/Private/SequencerEdMode.h b/Engine/Source/Editor/Sequencer/Private/SequencerEdMode.h
--- a/Engine/Source/Editor/Sequencer/Private/SequencerEdMode.h
+++ b/Engine/Source/Editor/Sequencer/Private/SequencerEdMode.h
@@ -9,6 +9,7 @@
 #include "EdMode.h"
 #include "Misc/FrameTime.h"
 #include "Engine/Texture2D.h"
+#include "MovieSceneFwd.h"
 
 class FCanvas;
 class FEditorViewportClient;
@@ -152,6 +153,7 @@ private:
 	//params to handle mouse move for changing time
 	/** If we are tracking */
 	bool bIsTracking = false;
+
 	/** Starting X Value*/
 	TOptional<int32> StartXValue;
 	/** Starting Time Value*/
@@ -162,6 +164,9 @@ private:
 	/** If the pivot location needs to be updated */
 	bool bUpdatePivot = false;
 
+	/**if dragging time need to reset to what it was */
+	EMovieScenePlayerStatus::Type NextPlayerStatus; 
+
 	TSharedPtr<class FSequencerEdModeTool> DefaultTool;
 	TSharedPtr<FSequencerSelectabilityTool> SelectabilityTool;
 };
diff --git a/Engine/Source/Editor/Sequencer/Private/SequencerSettings.cpp b/Engine/Source/Editor/Sequencer/Private/SequencerSettings.cpp
--- a/Engine/Source/Editor/Sequencer/Private/SequencerSettings.cpp
+++ b/Engine/Source/Editor/Sequencer/Private/SequencerSettings.cpp
@@ -65,6 +65,7 @@ USequencerSettings::USequencerSettings( const FObjectInitializer& ObjectInitiali
 	bActivateRealtimeViewports = true;
 	bEvaluateSubSequencesInIsolation = false;
 	bRerunConstructionScripts = true;
+	bScrubTimeStartFromCursor = true;
 	bVisualizePreAndPostRoll = true;
 	TrajectoryPathCap = 250;
 	FrameNumberDisplayFormat = EFrameNumberDisplayFormats::Seconds;
@@ -428,6 +429,20 @@ void USequencerSettings::SetAutoScrollEnabled(bool bInAutoScrollEnabled)
 	}
 }
 
+bool USequencerSettings::GetScrubTimeStartFromCursor() const
+{
+	return bScrubTimeStartFromCursor;
+}
+
+void USequencerSettings::SetScrubTimeStartFromCursor(bool bInValue)
+{
+	if (bScrubTimeStartFromCursor != bInValue)
+	{
+		bScrubTimeStartFromCursor = bInValue;
+		SaveConfig();
+	}
+}
+
 void USequencerSettings::SetCurveEditorZoomScaling(const FCurveEditorZoomScaleConfig& Value)
 {
 	CurveEditorZoomScaling = Value;
diff --git a/Engine/Source/Editor/Sequencer/Private/SequencerTimeSliderController.cpp b/Engine/Source/Editor/Sequencer/Private/SequencerTimeSliderController.cpp
--- a/Engine/Source/Editor/Sequencer/Private/SequencerTimeSliderController.cpp
+++ b/Engine/Source/Editor/Sequencer/Private/SequencerTimeSliderController.cpp
@@ -1760,6 +1760,14 @@ int32 FSequencerTimeSliderController::OnPaintViewArea( const FGeometry& Allotted
 	return LayerId;
 }
 
+void FSequencerTimeSliderController::SetStoppedPosition(FFrameTime InTime)
+{
+	if (TSharedPtr<ISequencer> SequencerPtr = WeakSequencer.Pin())
+	{
+		SequencerPtr->SetLocalTimeDirectly(InTime);
+	}
+}
+
 void FSequencerTimeSliderController::SetPlaybackStatus(ETimeSliderPlaybackStatus InStatus)
 {
 	using namespace UE::Sequencer;
diff --git a/Engine/Source/Editor/Sequencer/Private/Tools/SequencerEditTool_Selection.cpp b/Engine/Source/Editor/Sequencer/Private/Tools/SequencerEditTool_Selection.cpp
--- a/Engine/Source/Editor/Sequencer/Private/Tools/SequencerEditTool_Selection.cpp
+++ b/Engine/Source/Editor/Sequencer/Private/Tools/SequencerEditTool_Selection.cpp
@@ -4,6 +4,7 @@
 #include "Application/ThrottleManager.h"
 #include "Styling/AppStyle.h"
 #include "SequencerCommonHelpers.h"
+#include "SequencerTimeSliderController.h"
 #include "SSequencer.h"
 #include "ISequencerEditTool.h"
 #include "ISequencerSection.h"
@@ -110,7 +111,7 @@ class FScrubTimeDragOperation
 public:
 
 	FScrubTimeDragOperation(FSequencer& InSequencer, UE::Sequencer::STrackAreaView& InTrackArea)
-		: Sequencer(InSequencer)
+		: Sequencer(InSequencer), NextPlayerStatus(EMovieScenePlayerStatus::Stopped)
 	{}
 
 public:
@@ -129,7 +130,9 @@ public:
 		// Start a new marquee selection
 		InitialPosition = CurrentPosition = VirtualTrackArea.PhysicalToVirtual(LocalMousePos);
 		CurrentMousePos = LocalMousePos;
-
+		NextPlayerStatus = (Sequencer.GetPlaybackStatus() == EMovieScenePlayerStatus::Playing) ? EMovieScenePlayerStatus::Playing :
+			 EMovieScenePlayerStatus::Stopped;
+		
 		Sequencer.SetPlaybackStatus(EMovieScenePlayerStatus::Scrubbing);
 		SequencerStartTime = Sequencer.GetLocalTime();
 		InitialTime = VirtualTrackArea.PixelToSeconds(LocalMousePos.X);
@@ -146,8 +149,8 @@ public:
 
 	virtual void OnEndDrag(const FPointerEvent& MouseEvent, FVector2D LocalMousePos, const UE::Sequencer::FVirtualTrackArea& VirtualTrackArea) override
 	{
-		Sequencer.SetPlaybackStatus(EMovieScenePlayerStatus::Stopped);
 		FSlateThrottleManager::Get().DisableThrottle(false);
+		Sequencer.SetPlaybackStatus(NextPlayerStatus);
 	}
 
 	virtual int32 OnPaint(const FGeometry& AllottedGeometry, const FSlateRect& MyCullingRect, FSlateWindowElementList& OutDrawElements, int32 LayerId) const override
@@ -175,6 +178,8 @@ private:
 	FQualifiedFrameTime SequencerStartTime;
 	double InitialTime;
 	double CurrentTime;
+
+	EMovieScenePlayerStatus::Type NextPlayerStatus; //reset on end drag, make sure to keep playing if it was
 };
 
 class FMarqueeDragOperation
@@ -441,6 +446,22 @@ int32 FSequencerEditTool_Selection::OnPaint(const FGeometry& AllottedGeometry, c
 	return LayerId;
 }
 
+static FFrameTime ComputeFrameTimeFromMouse(const FSequencer& Sequencer, const FGeometry& Geometry, const FVector2D& ScreenSpacePosition, const FSequencerTimeSliderController::FScrubRangeToScreen& RangeToScreen, bool  bCheckSnapping)
+{
+	FVector2D CursorPos = Geometry.AbsoluteToLocal(ScreenSpacePosition);
+	double    MouseValue = RangeToScreen.LocalXToInput(CursorPos.X);
+
+	if (bCheckSnapping && Sequencer.GetSequencerSettings()->GetForceWholeFrames())
+	{
+		FFrameNumber        SnappedFrameNumber = (MouseValue * Sequencer.GetFocusedDisplayRate()).FloorToFrame();
+		FQualifiedFrameTime RoundedPlayFrame = FQualifiedFrameTime(SnappedFrameNumber, Sequencer.GetFocusedDisplayRate());
+		return RoundedPlayFrame.ConvertTo(Sequencer.GetFocusedTickResolution());
+	}
+	else
+	{
+		return MouseValue * Sequencer.GetFocusedTickResolution();
+	}
+}
 
 FReply FSequencerEditTool_Selection::OnMouseButtonDown(SWidget& OwnerWidget, const FGeometry& MyGeometry, const FPointerEvent& MouseEvent)
 {
@@ -452,6 +473,22 @@ FReply FSequencerEditTool_Selection::OnMouseButtonDown(SWidget& OwnerWidget, con
 
 	if (MouseEvent.GetEffectingButton() == EKeys::LeftMouseButton)
 	{
+		if (bIsScrubbingTime && Sequencer.GetSequencerSettings()->GetScrubTimeStartFromCursor())
+		{
+			FSequencerTimeSliderController::FScrubRangeToScreen RangeToScreen = FSequencerTimeSliderController::FScrubRangeToScreen(Sequencer.GetViewRange(), MyGeometry.Size);
+			FFrameTime FrameTime = ComputeFrameTimeFromMouse(Sequencer, MyGeometry, MouseEvent.GetScreenSpacePosition(), RangeToScreen, true);
+
+			if (Sequencer.GetPlaybackStatus() == EMovieScenePlayerStatus::Playing)
+			{
+				Sequencer.SetPlaybackStatus(EMovieScenePlayerStatus::Paused);
+				Sequencer.SetLocalTimeLooped(FrameTime);
+				Sequencer.SetPlaybackStatus(EMovieScenePlayerStatus::Playing);
+			}
+			else
+			{
+				Sequencer.SetLocalTimeLooped(FrameTime);
+			}
+		}
 		TSharedPtr<FSequencerEditorViewModel> SequencerViewModel = Sequencer.GetViewModel()->CastThisShared<FSequencerEditorViewModel>();
 		TSharedPtr<ITrackAreaHotspot> Hotspot = SequencerViewModel->GetHotspot();
 		DelayedDrag = FDelayedDrag_Hotspot(MyGeometry.AbsoluteToLocal(MouseEvent.GetScreenSpacePosition()), EKeys::LeftMouseButton, Hotspot);
@@ -516,7 +553,6 @@ FReply FSequencerEditTool_Selection::OnMouseMove(SWidget& OwnerWidget, const FGe
 	return FReply::Unhandled();
 }
 
-
 FReply FSequencerEditTool_Selection::OnMouseButtonUp(SWidget& OwnerWidget, const FGeometry& MyGeometry, const FPointerEvent& MouseEvent)
 {
 	UpdateCursor(MyGeometry, MouseEvent);
diff --git a/Engine/Source/Editor/Sequencer/Public/SequencerSettings.h b/Engine/Source/Editor/Sequencer/Public/SequencerSettings.h
--- a/Engine/Source/Editor/Sequencer/Public/SequencerSettings.h
+++ b/Engine/Source/Editor/Sequencer/Public/SequencerSettings.h
@@ -262,6 +262,11 @@ public:
 	/** Sets whether or not auto-scroll is enabled when playing. */
 	UE_API void SetAutoScrollEnabled(bool bInAutoScrollEnabled);
 	
+	/** Gets whether or not scrubbing time hot key starts from cursor. */
+	UE_API bool GetScrubTimeStartFromCursor() const;
+	/** Sets whether or not scrubing time hot key starts from cursor. */
+	UE_API void SetScrubTimeStartFromCursor(bool bInValue);
+
 	/** Gets whether or not to link the curve editor time range. */
 	UE_API bool GetLinkCurveEditorTimeRange() const;
 	/** Sets whether or not to link the curve editor time range. */
@@ -666,6 +671,10 @@ protected:
 	UPROPERTY( config, EditAnywhere, Category=Timeline )
 	bool bAutoScrollEnabled;
 
+	/** When enabled, scrubbing time hotkey will move time from initial cursor position  */
+	UPROPERTY(config, EditAnywhere, Category = Timeline)
+	bool bScrubTimeStartFromCursor;
+
 	/** Enable or disable linking the curve editor time range to the sequencer timeline's time range. */
 	UPROPERTY( config, EditAnywhere, Category=CurveEditor )
 	bool bLinkCurveEditorTimeRange;
diff --git a/Engine/Source/Editor/Sequencer/Public/SequencerTimeSliderController.h b/Engine/Source/Editor/Sequencer/Public/SequencerTimeSliderController.h
--- a/Engine/Source/Editor/Sequencer/Public/SequencerTimeSliderController.h
+++ b/Engine/Source/Editor/Sequencer/Public/SequencerTimeSliderController.h
@@ -74,6 +74,9 @@ public:
 	/** Get the current time for the Scrub handle which indicates what range is being evaluated. */
 	virtual FFrameTime GetScrubPosition() const override { return TimeSliderArgs.ScrubPosition.Get(FFrameTime()); }
 
+	/** Set the current time for the handle when playback is stopped*/
+	virtual void SetStoppedPosition(FFrameTime InTime) override;
+
 	/** Get the current time for the Scrub handle which indicates what range is being evaluated. */
 	virtual void SetScrubPosition(FFrameTime InTime, bool bEvaluate) override { CommitScrubPosition(InTime, GetPlaybackStatus() == ETimeSliderPlaybackStatus::Scrubbing, bEvaluate); }
 
diff --git a/Engine/Source/Editor/SequencerWidgets/Private/ITimeSlider.cpp b/Engine/Source/Editor/SequencerWidgets/Private/ITimeSlider.cpp
--- a/Engine/Source/Editor/SequencerWidgets/Private/ITimeSlider.cpp
+++ b/Engine/Source/Editor/SequencerWidgets/Private/ITimeSlider.cpp
@@ -38,6 +38,10 @@ void ITimeSliderController::SetScrubPosition(FFrameTime InTime, bool bEvaluate)
 {
 }
 
+void  ITimeSliderController::SetStoppedPosition(FFrameTime InTime)
+{
+}
+
 void ITimeSliderController::SetViewRange(double NewRangeMin, double NewRangeMax, EViewRangeInterpolation Interpolation)
 {
 }
diff --git a/Engine/Source/Editor/SequencerWidgets/Public/ITimeSlider.h b/Engine/Source/Editor/SequencerWidgets/Public/ITimeSlider.h
--- a/Engine/Source/Editor/SequencerWidgets/Public/ITimeSlider.h
+++ b/Engine/Source/Editor/SequencerWidgets/Public/ITimeSlider.h
@@ -123,6 +123,9 @@ public:
 
 	/** Set the current time for the Scrub handle which indicates what range is being evaluated. */
 	SEQUENCERWIDGETS_API virtual void SetScrubPosition(FFrameTime InTime, bool bEvaluate);
+
+	/** Set the current time for the handle when playback is stopped*/
+	SEQUENCERWIDGETS_API virtual void SetStoppedPosition(FFrameTime InTime);
 	
 	/** Set the playback status for the controller*/
 	SEQUENCERWIDGETS_API virtual void SetPlaybackStatus(ETimeSliderPlaybackStatus InStatus);
