[Sequencer] Marquee Select Layer Bars

diff --git a/Engine/Source/Editor/Sequencer/Private/Tools/SequencerEditTool_Selection.cpp b/Engine/Source/Editor/Sequencer/Private/Tools/SequencerEditTool_Selection.cpp
--- a/Engine/Source/Editor/Sequencer/Private/Tools/SequencerEditTool_Selection.cpp
+++ b/Engine/Source/Editor/Sequencer/Private/Tools/SequencerEditTool_Selection.cpp
@@ -36,7 +36,7 @@ struct FSelectionPreviewVisitor final
 		, SetStateTo(InSetStateTo)
 		, bPinned(bInPinned)
 	{
-		bIsControlDown = FSlateApplication::Get().GetModifierKeys().IsControlDown();
+		ModifierKeys = FSlateApplication::Get().GetModifierKeys();
 	}
 
 	virtual void VisitKeys(const UE::Sequencer::TViewModelPtr<UE::Sequencer::FChannelModel>& Channel, const TRange<FFrameNumber>& VisitRangeFrames) const override
@@ -59,15 +59,22 @@ struct FSelectionPreviewVisitor final
 		{
 			// Under default behavior keys have priority, so if a key is changing selection state then we remove any sections from the selection. The user can bypass this
 			// by holding down the control key which will allow selecting both keys and sections.
-			bool bKeySelectionHasPriority = !bIsControlDown;
-			bool bKeyIsSelected = ExistingSelection.Contains(KeyHandlesScratch[Index]);
+			const bool bPrioritizeLayerBars = ModifierKeys.IsControlDown() && ModifierKeys.IsAltDown();
+			const bool bPrioritizeKeys = !bPrioritizeLayerBars && !ModifierKeys.IsControlDown();
+			const bool bKeyIsSelected = ExistingSelection.Contains(KeyHandlesScratch[Index]);
 
-			if (bKeySelectionHasPriority && 
-				((bKeyIsSelected && SetStateTo == ESelectionPreviewState::NotSelected) ||
+			// Clear selected models
+			if (((bKeyIsSelected && SetStateTo == ESelectionPreviewState::NotSelected) ||
 				(!bKeyIsSelected && SetStateTo == ESelectionPreviewState::Selected)))
 			{
-				// Clear selected models
-				SelectionPreview.EmptyDefinedModelStates();
+				if (bPrioritizeKeys)
+				{
+					SelectionPreview.EmptyDefinedModelStates();
+				}
+				else if (bPrioritizeLayerBars)
+				{
+					SelectionPreview.EmptyDefinedKeyStates();
+				}
 			}
 
 			SelectionPreview.SetSelectionState(Channel, KeyHandlesScratch[Index], SetStateTo);
@@ -86,8 +93,9 @@ struct FSelectionPreviewVisitor final
 
 		// If key selection has priority then we check to see if there are any keys selected. If there are key selected, we don't add this section.
 		// Otherwise, we bypass this check
-		bool bKeySelectionHasPriority = !FSlateApplication::Get().GetModifierKeys().IsControlDown();
-		bool bKeyStateCheck = bKeySelectionHasPriority ? SelectionPreview.GetDefinedKeyStates().Num() == 0 : true;
+		const bool bPrioritizeLayerBars = ModifierKeys.IsControlDown() && ModifierKeys.IsAltDown();
+		const bool bPrioritizeKeys = !bPrioritizeLayerBars && !ModifierKeys.IsControlDown();
+		const bool bKeyStateCheck = bPrioritizeKeys ? SelectionPreview.GetDefinedKeyStates().Num() == 0 : true;
 
 		if (bKeyStateCheck)
 		{
@@ -102,7 +110,7 @@ private:
 	mutable TArray<FKeyHandle> KeyHandlesScratch;
 	ESelectionPreviewState SetStateTo;
 	bool bPinned;
-	bool bIsControlDown;
+	FModifierKeysState ModifierKeys;
 };
 
 class FScrubTimeDragOperation
@@ -215,7 +223,7 @@ public:
 		{
 			PreviewState = ESelectionPreviewState::Selected;
 		}
-		else if (MouseEvent.IsAltDown())
+		else if (!MouseEvent.IsControlDown() && MouseEvent.IsAltDown())
 		{
 			PreviewState = ESelectionPreviewState::NotSelected;
 		}
diff --git a/Engine/Source/Editor/Sequencer/Private/Tools/SequencerEntityVisitor.cpp b/Engine/Source/Editor/Sequencer/Private/Tools/SequencerEntityVisitor.cpp
--- a/Engine/Source/Editor/Sequencer/Private/Tools/SequencerEntityVisitor.cpp
+++ b/Engine/Source/Editor/Sequencer/Private/Tools/SequencerEntityVisitor.cpp
@@ -21,6 +21,7 @@
 #include "Misc/FrameTime.h"
 #include "MovieSceneSection.h"
 #include "SequencerCoreFwd.h"
+#include "MVVM/ViewModels/LayerBarModel.h"
 
 FSequencerEntityRange::FSequencerEntityRange(const TRange<double>& InRange, FFrameRate InTickResolution)
 	: TickResolution(InTickResolution)
@@ -94,34 +95,48 @@ void FSequencerEntityWalker::ConditionallyIntersectModel(const ISequencerEntityV
 			? VirtualGeometry.NestedBottom - VirtualGeometry.Top
 			: VirtualGeometry.Height;
 
-		if (Range.IntersectVertical(VirtualGeometry.Top, VirtualGeometry.Top + Height))
+		auto VisitIfOverlaps = [this, &Visitor, &VirtualGeometry, Height]
+			(const TViewModelPtr<ITrackLaneExtension>& InTrackLaneExtension)
 		{
-			for (const TViewModelPtr<ITrackLaneExtension>& TrackLane
-				: DataModel->GetChildrenOfType<ITrackLaneExtension>(EViewModelListType::TrackArea))
+			const FTrackLaneVirtualAlignment Alignment = InTrackLaneExtension->ArrangeVirtualTrackLaneView();
+			if (!Alignment.IsVisible())
 			{
-				FTrackLaneVirtualAlignment Alignment = TrackLane->ArrangeVirtualTrackLaneView();
-				if (!Alignment.IsVisible())
-				{
-					continue;
-				}
+				return;
+			}
 
-				FTrackLaneVerticalArrangement VerticalArrange = Alignment.VerticalAlignment.ArrangeWithin(Height);
+			const FTrackLaneVerticalArrangement VerticalArrange = Alignment.VerticalAlignment.ArrangeWithin(Height);
 
-				if (Range.IntersectVertical(
-							VirtualGeometry.Top + VerticalArrange.Offset,
-							VirtualGeometry.Top + VerticalArrange.Offset + VerticalArrange.Height) )
+			if (Range.IntersectVertical(
+				VirtualGeometry.Top + VerticalArrange.Offset,
+				VirtualGeometry.Top + VerticalArrange.Offset + VerticalArrange.Height))
+			{
+				if (Range.Range.Overlaps(Alignment.Range / Range.TickResolution))
 				{
-					if (Range.Range.Overlaps(Alignment.Range / Range.TickResolution))
-					{
-						Visitor.VisitDataModel(TrackLane.AsModel().Get());
+					Visitor.VisitDataModel(InTrackLaneExtension.AsModel().Get());
 
-						if (TViewModelPtr<FChannelModel> ChannelModel = TrackLane.ImplicitCast())
-						{
-							VisitChannel(Visitor, ChannelModel);
-						}
+					if (TViewModelPtr<FChannelModel> ChannelModel = InTrackLaneExtension.ImplicitCast())
+					{
+						VisitChannel(Visitor, ChannelModel);
 					}
 				}
 			}
+		};
+
+		if (Range.IntersectVertical(VirtualGeometry.Top, VirtualGeometry.Top + Height))
+		{
+			for (const FViewModelPtr& TopLevelChild : TrackArea->GetTopLevelChildTrackAreaModels())
+			{
+				if (const TViewModelPtr<FLayerBarModel> LayerBarModel = TopLevelChild.ImplicitCast())
+				{
+					VisitIfOverlaps(LayerBarModel);
+				}
+			}
+
+			for (const TViewModelPtr<ITrackLaneExtension>& TrackLane
+				: DataModel->GetChildrenOfType<ITrackLaneExtension>(EViewModelListType::TrackArea))
+			{
+				VisitIfOverlaps(TrackLane);
+			}
 
 			if (Range.IntersectKeyArea(DataModel, VirtualKeySize.Y))
 			{
@@ -216,17 +231,24 @@ void FSequencerEntityWalker::VisitAnyChannels(const ISequencerEntityVisitor& Vis
 
 void FSequencerEntityWalker::VisitChannel(const ISequencerEntityVisitor& Visitor, const UE::Sequencer::TViewModelPtr<UE::Sequencer::FChannelModel>& Channel)
 {
-	using namespace UE::Sequencer;
-
-	UMovieSceneSection* Section = Channel->GetSection();
-	if (Section)
+	const FModifierKeysState ModifierKeys = FSlateApplication::Get().GetModifierKeys();
+	const bool bLayerBarsOnly = ModifierKeys.IsControlDown() && ModifierKeys.IsAltDown();
+	if (bLayerBarsOnly)
 	{
-		const FFrameTime HalfKeySizeFrames   = (VirtualKeySize.X*.5f) * Range.TickResolution;
-		const FFrameTime RangeStartFrame     = Range.Range.GetLowerBoundValue() * Range.TickResolution;
-		const FFrameTime RangeEndFrame       = Range.Range.GetUpperBoundValue() * Range.TickResolution;
+		return;
+	}
 
-		TRange<FFrameNumber> VisitRangeFrames( (RangeStartFrame-HalfKeySizeFrames).CeilToFrame(), (RangeEndFrame+HalfKeySizeFrames).FloorToFrame() );
-		VisitRangeFrames = TRange<FFrameNumber>::Intersection(Section->GetRange(), VisitRangeFrames);
-		Visitor.VisitKeys(Channel, VisitRangeFrames);
+	const UMovieSceneSection* const Section = Channel.IsValid() ? Channel->GetSection() : nullptr;
+	if (!Section)
+	{
+		return;
 	}
+
+	const FFrameTime HalfKeySizeFrames = (VirtualKeySize.X*.5f) * Range.TickResolution;
+	const FFrameTime RangeStartFrame   = Range.Range.GetLowerBoundValue() * Range.TickResolution;
+	const FFrameTime RangeEndFrame     = Range.Range.GetUpperBoundValue() * Range.TickResolution;
+
+	TRange<FFrameNumber> VisitRangeFrames( (RangeStartFrame-HalfKeySizeFrames).CeilToFrame(), (RangeEndFrame+HalfKeySizeFrames).FloorToFrame() );
+	VisitRangeFrames = TRange<FFrameNumber>::Intersection(Section->GetRange(), VisitRangeFrames);
+	Visitor.VisitKeys(Channel, VisitRangeFrames);
 }
