Sequencer: Change FSequencerSelection::GetBoundObjectGuids to return unique ids

#rnx

diff --git a/Engine/Source/Editor/Sequencer/Private/MVVM/Selection.cpp b/Engine/Source/Editor/Sequencer/Private/MVVM/Selection.cpp
--- a/Engine/Source/Editor/Sequencer/Private/MVVM/Selection.cpp
+++ b/Engine/Source/Editor/Sequencer/Private/MVVM/Selection.cpp
@@ -253,9 +253,9 @@ FIndirectOutlinerSelectionIterator FSequencerSelection::IterateIndirectOutlinerS
 	return FIndirectOutlinerSelectionIterator{ &NodesWithKeysOrSections };
 }
 
-TArray<FGuid> FSequencerSelection::GetBoundObjectsGuids()
+TSet<FGuid> FSequencerSelection::GetBoundObjectsGuids()
 {
-	TArray<FGuid> OutGuids;
+	TSet<FGuid> OutGuids;
 
 	for (const TWeakViewModelPtr<IOutlinerExtension>& WeakModel : NodesWithKeysOrSections)
 	{
@@ -270,7 +270,7 @@ TArray<FGuid> FSequencerSelection::GetBoundObjectsGuids()
 		}
 	}
 
-	if (OutGuids.Num() == 0)
+	if (OutGuids.IsEmpty())
 	{
 		for (FViewModelPtr Model : Outliner)
 		{
diff --git a/Engine/Source/Editor/Sequencer/Private/Sequencer.cpp b/Engine/Source/Editor/Sequencer/Private/Sequencer.cpp
--- a/Engine/Source/Editor/Sequencer/Private/Sequencer.cpp
+++ b/Engine/Source/Editor/Sequencer/Private/Sequencer.cpp
@@ -6555,7 +6555,7 @@ void FSequencer::HandleSelectedOutlinerNodesChanged()
 		}
 	}
 
-	OnSelectionChangedObjectGuidsDelegate.Broadcast(ViewModel->GetSelection()->GetBoundObjectsGuids());
+	OnSelectionChangedObjectGuidsDelegate.Broadcast(ViewModel->GetSelection()->GetBoundObjectsGuids().Array());
 	OnSelectionChangedTracksDelegate.Broadcast(SelectedTracks.Array());
 	OnSelectionChangedSectionsDelegate.Broadcast(SelectedSections);
 }
@@ -7337,7 +7337,7 @@ void FSequencer::GetSelectedFolders(TArray<UMovieSceneFolder*>& OutSelectedFolde
 
 void FSequencer::GetSelectedObjects(TArray<FGuid>& Objects)
 {
-	Objects = GetSelection().GetBoundObjectsGuids();
+	Objects = GetSelection().GetBoundObjectsGuids().Array();
 }
 
 void FSequencer::GetSelectedKeyAreas(TArray<const IKeyArea*>& OutSelectedKeyAreas, bool bIncludeSelectedKeys)
diff --git a/Engine/Source/Editor/Sequencer/Public/MVVM/Selection/Selection.h b/Engine/Source/Editor/Sequencer/Public/MVVM/Selection/Selection.h
--- a/Engine/Source/Editor/Sequencer/Public/MVVM/Selection/Selection.h
+++ b/Engine/Source/Editor/Sequencer/Public/MVVM/Selection/Selection.h
@@ -150,7 +150,7 @@ public:
 	void GetSelectedOutlinerItems(TArray<TSharedPtr<FViewModel>>& OutItems) const;
 
 	/** DEPRECATED: Gets the outliner nodes that have selected keys or sections */
-	TArray<FGuid> GetBoundObjectsGuids();
+	TSet<FGuid> GetBoundObjectsGuids();
 
 	/** DEPRECATED */
 	void EmptySelectedOutlinerNodesWithoutSections(const TArray<UMovieSceneSection*>& Sections);
