[Sequencer] Remove limited selectability notice from viewport in favor of toolbar indicator

diff --git a/Engine/Source/Editor/Sequencer/Private/SequencerSelectabilityTool.cpp b/Engine/Source/Editor/Sequencer/Private/SequencerSelectabilityTool.cpp
--- a/Engine/Source/Editor/Sequencer/Private/SequencerSelectabilityTool.cpp
+++ b/Engine/Source/Editor/Sequencer/Private/SequencerSelectabilityTool.cpp
@@ -25,7 +25,9 @@ bool FSequencerSelectabilityTool::BoxSelect(FBox& InBox, const bool InSelect)
 	return BoxSelectWorldActors(InBox, GCurrentLevelEditingViewportClient, InSelect);
 }
 
-bool FSequencerSelectabilityTool::FrustumSelect(const FConvexVolume& InFrustum, FEditorViewportClient* const InEditorViewportClient, const bool InSelect)
+bool FSequencerSelectabilityTool::FrustumSelect(const FConvexVolume& InFrustum
+	, FEditorViewportClient* const InEditorViewportClient
+	, const bool InSelect)
 {
 	if (!bSelectionLimited
 		|| !InEditorViewportClient
@@ -45,17 +47,4 @@ bool FSequencerSelectabilityTool::FrustumSelect(const FConvexVolume& InFrustum,
 	return FrustumSelectWorldActors(InFrustum, InEditorViewportClient, InSelect);
 }
 
-void FSequencerSelectabilityTool::DrawHUD(FEditorViewportClient* const InEditorViewportClient, FViewport* const InViewport, const FSceneView* const InView, FCanvas* InCanvas)
-{
-	if (!bSelectionLimited)
-	{
-		return;
-	}
-
-	const FSequencerCommands& SequencerCommands = FSequencerCommands::Get();
-	const FText HelpText = GetLimitedSelectionText(SequencerCommands.ToggleLimitViewportSelection
-		, LOCTEXT("SequencerSelectionLimitHelpText", "Sequencer Selection Limited"));
-	DrawEnabledTextNotice(InCanvas, HelpText);
-}
-
 #undef LOCTEXT_NAMESPACE
diff --git a/Engine/Source/Editor/Sequencer/Private/SequencerSelectabilityTool.h b/Engine/Source/Editor/Sequencer/Private/SequencerSelectabilityTool.h
--- a/Engine/Source/Editor/Sequencer/Private/SequencerSelectabilityTool.h
+++ b/Engine/Source/Editor/Sequencer/Private/SequencerSelectabilityTool.h
@@ -17,7 +17,8 @@ public:
 		return TEXT("Sequencer Selectability");
 	}
 	virtual bool BoxSelect(FBox& InBox, const bool InSelect) override;
-	virtual bool FrustumSelect(const FConvexVolume& InFrustum, FEditorViewportClient* const InEditorViewportClient, const bool InSelect) override;
-	virtual void DrawHUD(FEditorViewportClient* const InEditorViewportClient, FViewport* const InViewport, const FSceneView* InView, FCanvas* const InCanvas) override;
+	virtual bool FrustumSelect(const FConvexVolume& InFrustum
+		, FEditorViewportClient* const InEditorViewportClient
+		, const bool InSelect) override;
 	//~ End FModeTool
 };
