Sequencer: Nullptr checks for ExportToOutputDevice

This fixes a bug trying to copy/export a binding that doesn't have a valid spawnable template

diff --git a/Engine/Source/Editor/Sequencer/Private/SequencerUtilities.cpp b/Engine/Source/Editor/Sequencer/Private/SequencerUtilities.cpp
--- a/Engine/Source/Editor/Sequencer/Private/SequencerUtilities.cpp
+++ b/Engine/Source/Editor/Sequencer/Private/SequencerUtilities.cpp
@@ -1667,16 +1667,19 @@ void ExportObjectsToText(const TArray<UObject*>& ObjectsToExport, FString& Expor
 
 	for (UObject* ObjectToExport : ObjectsToExport)
 	{
-		// The nodes should all be from the same scope
-		UObject* ThisOuter = ObjectToExport->GetOuter();
-		if (LastOuter != nullptr && ThisOuter != LastOuter)
+		if (ObjectToExport)
 		{
-			UE_LOG(LogSequencer, Error, TEXT("Cannot copy objects from different outers. Only copying from %s"), *LastOuter->GetName());
-			continue;
-		}
-		LastOuter = ThisOuter;
+			// The nodes should all be from the same scope
+			UObject* ThisOuter = ObjectToExport->GetOuter();
+			if (LastOuter != nullptr && ThisOuter != LastOuter)
+			{
+				UE_LOG(LogSequencer, Error, TEXT("Cannot copy objects from different outers. Only copying from %s"), *LastOuter->GetName());
+				continue;
+			}
+			LastOuter = ThisOuter;
 
-		UExporter::ExportToOutputDevice(&Context, ObjectToExport, nullptr, Archive, TEXT("copy"), 0, PPF_ExportsNotFullyQualified | PPF_Copy | PPF_Delimited, false, ThisOuter);
+			UExporter::ExportToOutputDevice(&Context, ObjectToExport, nullptr, Archive, TEXT("copy"), 0, PPF_ExportsNotFullyQualified | PPF_Copy | PPF_Delimited, false, ThisOuter);
+		}
 	}
 
 	ExportedText = Archive;
@@ -2367,6 +2370,11 @@ void ExportObjectBindingsToText(const TArray<UMovieSceneCopyableBinding*>& Objec
 
 	for (UMovieSceneCopyableBinding* ObjectToExport : ObjectsToExport)
 	{
+		if (!ObjectToExport)
+		{
+			continue;
+		}
+
 		// The nodes should all be from the same scope
 		UObject* ThisOuter = ObjectToExport->GetOuter();
 		check((LastOuter == ThisOuter) || (LastOuter == nullptr));
@@ -2427,7 +2435,10 @@ void ExportObjectBindingsToText(const TArray<UMovieSceneCopyableBinding*>& Objec
 
 		for (UMovieSceneCustomBinding* CustomBinding : ObjectToExport->CustomBindings)
 		{
-			UExporter::ExportToOutputDevice(&Context, CustomBinding, nullptr, Archive, TEXT("copy"), 0, PPF_ExportsNotFullyQualified | PPF_Copy | PPF_Delimited);
+			if (CustomBinding)
+			{
+				UExporter::ExportToOutputDevice(&Context, CustomBinding, nullptr, Archive, TEXT("copy"), 0, PPF_ExportsNotFullyQualified | PPF_Copy | PPF_Delimited);
+			}
 		}
 
 		// Restore replaceable references now that we've exported them
@@ -2462,7 +2473,10 @@ void ExportObjectBindingsToText(const TArray<UMovieSceneCopyableBinding*>& Objec
 		// We manually export the object templates for the same private-ownership reason as above. Templates need to be re-created anyways as each Spawnable contains its own copy of the template.
 		for (UObject* ObjectTemplate : ObjectToExport->SpawnableObjectTemplates)
 		{
-			UExporter::ExportToOutputDevice(&Context, ObjectTemplate, nullptr, Archive, TEXT("copy"), 0, PPF_ExportsNotFullyQualified | PPF_Copy | PPF_Delimited);
+			if (ObjectTemplate)
+			{
+				UExporter::ExportToOutputDevice(&Context, ObjectTemplate, nullptr, Archive, TEXT("copy"), 0, PPF_ExportsNotFullyQualified | PPF_Copy | PPF_Delimited);
+			}
 		}
 	}
 }
