Sequencer: Allow sequencer customization for derived sequence classes by looking for an exact match first, and then if not found, use the closest matching class. It only ever returns one customization at the moment.

This new behavior can be turned off with the cvar: Sequencer.ExactClassCustomizationsOnly 0

This fixes an issue where cine assembly assets (which are level sequences), don't have the assign actor, convert bindings menus

diff --git a/Engine/Source/Editor/Sequencer/Private/SequencerCustomizationManager.cpp b/Engine/Source/Editor/Sequencer/Private/SequencerCustomizationManager.cpp
--- a/Engine/Source/Editor/Sequencer/Private/SequencerCustomizationManager.cpp
+++ b/Engine/Source/Editor/Sequencer/Private/SequencerCustomizationManager.cpp
@@ -47,6 +47,11 @@ void FSequencerCustomizationManager::UnregisterInstancedSequencerCustomization(c
     ensureMsgf(false, TEXT("No customization found for sequence class: %s"), *SequenceClass->GetName());
 }
 
+static TAutoConsoleVariable<bool> CVarExactClassCustomizationsOnly(
+	TEXT("Sequencer.ExactClassCustomizationsOnly"),
+	false,
+	TEXT("Enable/disable sequencer customizations with exact class matches only"));
+
 void FSequencerCustomizationManager::GetSequencerCustomizations(UMovieSceneSequence* FocusedSequence, TArray<TUniquePtr<ISequencerCustomization>>& OutCustomizations)
 {
 	if (UNLIKELY(!FocusedSequence))
@@ -54,18 +59,44 @@ void FSequencerCustomizationManager::GetSequencerCustomizations(UMovieSceneSeque
 		return;
 	}
 
-    const UClass* SequenceClass = FocusedSequence->GetClass();
+    UClass* SequenceClass = FocusedSequence->GetClass();
 
-    const FCustomizationRegistryEntry* FoundEntry = CustomizationRegistryEntries.FindByPredicate(
-            [=](const FCustomizationRegistryEntry& Entry)
-            {
-                return Entry.SequenceClass == SequenceClass;
-            });
-    if (FoundEntry != nullptr)
-    {
-        TUniquePtr<ISequencerCustomization> Instance(FoundEntry->Factory.Execute());
-        OutCustomizations.Add(MoveTemp(Instance));
-    }
+	const FCustomizationRegistryEntry* FoundEntry = nullptr;
+
+	if (CVarExactClassCustomizationsOnly->GetBool())
+	{
+		FoundEntry = CustomizationRegistryEntries.FindByPredicate(
+			[=](const FCustomizationRegistryEntry& Entry)
+			{
+				return Entry.SequenceClass == SequenceClass;
+			});
+	}
+	else
+	{
+		UClass* CurrentClass = SequenceClass;
+		while (CurrentClass != UObject::StaticClass())
+		{
+			FoundEntry = CustomizationRegistryEntries.FindByPredicate(
+				[=](const FCustomizationRegistryEntry& Entry)
+				{
+					return Entry.SequenceClass == CurrentClass;
+				});
+			if (FoundEntry == nullptr)
+			{
+				CurrentClass = CurrentClass->GetSuperClass();
+			}
+			else
+			{
+				break;
+			}
+		}
+	}
+
+	if (FoundEntry != nullptr)
+	{
+		TUniquePtr<ISequencerCustomization> Instance(FoundEntry->Factory.Execute());
+		OutCustomizations.Add(MoveTemp(Instance));
+	}
 }
 
 bool FSequencerCustomizationManager::NeedsCustomizationChange(const UMovieSceneSequence* OldFocusedSequence, const UMovieSceneSequence* NewFocusedSequence) const
