Sequencer: Added nullptr checks for MovieScene

#rnx

diff --git a/Engine/Source/Editor/Sequencer/Private/MVVM/ViewModels/OutlinerItemModel.cpp b/Engine/Source/Editor/Sequencer/Private/MVVM/ViewModels/OutlinerItemModel.cpp
--- a/Engine/Source/Editor/Sequencer/Private/MVVM/ViewModels/OutlinerItemModel.cpp
+++ b/Engine/Source/Editor/Sequencer/Private/MVVM/ViewModels/OutlinerItemModel.cpp
@@ -210,11 +210,11 @@ bool FOutlinerItemModelMixin::IsExpanded() const
 
 		TSharedPtr<FSequenceModel> SequenceModel = ViewModel->FindAncestorOfType<FSequenceModel>();
 		UMovieSceneSequence*       Sequence      = SequenceModel ? SequenceModel->GetSequence() : nullptr;
+		UMovieScene* MovieScene = Sequence ? Sequence->GetMovieScene() : nullptr;
 
-		if (Sequence)
+		if (MovieScene)
 		{
 			FStringView StringView = StringBuilder.ToView();
-			UMovieScene* MovieScene = Sequence->GetMovieScene();
 			FMovieSceneEditorData& EditorData = MovieScene->GetEditorData();
 			if (const FMovieSceneExpansionState* Expansion = EditorData.ExpansionStates.FindByHash(GetTypeHash(StringView), StringView))
 			{
@@ -265,7 +265,10 @@ void FOutlinerItemModelMixin::SetExpansion(bool bInIsExpanded)
 		if (SequenceModel)
 		{
 			TSharedPtr<FSequencer> Sequencer = SequenceModel->GetSequencerImpl();
-			Sequencer->GetNodeTree()->SaveExpansionState(*ViewModel, bInIsExpanded);
+			if (Sequencer)
+			{
+				Sequencer->GetNodeTree()->SaveExpansionState(*ViewModel, bInIsExpanded);
+			}
 		}
 	}
 }
@@ -300,8 +303,11 @@ bool FOutlinerItemModelMixin::IsPinned() const
 	{
 		TSharedPtr<FSequenceModel> SequenceModel = AsViewModel()->FindAncestorOfType<FSequenceModel>();
 		TSharedPtr<FSequencer> Sequencer = SequenceModel->GetSequencerImpl();
-		const bool bWasPinned = Sequencer->GetNodeTree()->GetSavedPinnedState(*AsViewModel());
-		const_cast<FOutlinerItemModelMixin*>(this)->FPinnableExtensionShim::SetPinned(bWasPinned);
+		if (Sequencer)
+		{
+			const bool bWasPinned = Sequencer->GetNodeTree()->GetSavedPinnedState(*AsViewModel());		
+			const_cast<FOutlinerItemModelMixin*>(this)->FPinnableExtensionShim::SetPinned(bWasPinned);
+		}
 	}
 
 	return FPinnableExtensionShim::IsPinned();
@@ -356,8 +362,11 @@ void FOutlinerItemModelMixin::ToggleRootModelPinned()
 		PinnableParent->SetPinned(bShouldPin);
 
 		TSharedPtr<FSequencer> Sequencer = RootModel->GetSequencerImpl();
-		Sequencer->GetNodeTree()->SavePinnedState(*AsViewModel(), bShouldPin);
-		Sequencer->RefreshTree();
+		if (Sequencer)
+		{
+			Sequencer->GetNodeTree()->SavePinnedState(*AsViewModel(), bShouldPin);
+			Sequencer->RefreshTree();
+		}
 	}
 }
 
@@ -1158,9 +1167,11 @@ bool FEvaluableOutlinerItemModel::IsSolo() const
 {
 	TSharedPtr<FSequencerEditorViewModel> EditorViewModel = GetEditor();
 	TSharedPtr<FSequencer> Sequencer = EditorViewModel ? EditorViewModel->GetSequencerImpl() : nullptr;
-	if (Sequencer)
+	UMovieSceneSequence* Sequence = Sequencer ? Sequencer->GetFocusedMovieSceneSequence() : nullptr;
+	UMovieScene* MovieScene = Sequence ? Sequence->GetMovieScene() : nullptr;
+	if (MovieScene)
 	{
-		const TArray<FString>& SoloNodes = Sequencer->GetFocusedMovieSceneSequence()->GetMovieScene()->GetSoloNodes();
+		const TArray<FString>& SoloNodes = MovieScene->GetSoloNodes();
 
 		// Should be always called on the main thread, but thread_local for safety
 		static thread_local TStringBuilder<256> StaticBuffer;
@@ -1180,15 +1191,10 @@ void FEvaluableOutlinerItemModel::SetIsSoloed(bool bIsSoloed)
 {
 	FViewModel* ViewModel = AsViewModel();
 	TSharedPtr<FSequenceModel> SequenceModel = ViewModel->FindAncestorOfType<FSequenceModel>();
+	UMovieScene* MovieScene = SequenceModel ? SequenceModel->GetMovieScene() : nullptr;
 
-	if (SequenceModel)
+	if (MovieScene && !MovieScene->IsReadOnly())
 	{
-		UMovieScene* MovieScene = SequenceModel->GetMovieScene();
-		if (MovieScene->IsReadOnly())
-		{
-			return;
-		}
-
 		// This shouldn't modify because it's not saved with the asset. But it still needs to mark the movie scene as changed, so that the UI will refresh
 		MovieScene->MarkAsChanged();
 
@@ -1219,9 +1225,11 @@ bool FEvaluableOutlinerItemModel::IsMuted() const
 {
 	TSharedPtr<FSequencerEditorViewModel> EditorViewModel = GetEditor();
 	TSharedPtr<FSequencer> Sequencer = EditorViewModel ? EditorViewModel->GetSequencerImpl() : nullptr;
-	if (Sequencer)
+	UMovieSceneSequence* Sequence = Sequencer ? Sequencer->GetFocusedMovieSceneSequence() : nullptr;
+	UMovieScene* MovieScene = Sequence ? Sequence->GetMovieScene() : nullptr;
+	if (MovieScene)
 	{
-		const TArray<FString>& MuteNodes = Sequencer->GetFocusedMovieSceneSequence()->GetMovieScene()->GetMuteNodes();
+		const TArray<FString>& MuteNodes = MovieScene->GetMuteNodes();
 
 		// Should be always called on the main thread, but thread_local for safety
 		static thread_local TStringBuilder<256> StaticBuffer;
@@ -1241,15 +1249,10 @@ void FEvaluableOutlinerItemModel::SetIsMuted(bool bIsMuted)
 {
 	FViewModel* ViewModel = AsViewModel();
 	TSharedPtr<FSequenceModel> SequenceModel = ViewModel->FindAncestorOfType<FSequenceModel>();
+	UMovieScene* MovieScene = SequenceModel ? SequenceModel->GetMovieScene() : nullptr;
 
-	if (SequenceModel)
+	if (MovieScene && !MovieScene->IsReadOnly())
 	{
-		UMovieScene* MovieScene = SequenceModel->GetMovieScene();
-		if (MovieScene->IsReadOnly())
-		{
-			return;
-		}
-
 		// This shouldn't modify because it's not saved with the asset. But it still needs to mark the movie scene as changed, so that the UI will refresh
 		MovieScene->MarkAsChanged();
 
