Sequencer: Gather all child bindings when copying folders

This is a regression from 39154758

diff --git a/Engine/Source/Editor/Sequencer/Private/SequencerUtilities.cpp b/Engine/Source/Editor/Sequencer/Private/SequencerUtilities.cpp
--- a/Engine/Source/Editor/Sequencer/Private/SequencerUtilities.cpp
+++ b/Engine/Source/Editor/Sequencer/Private/SequencerUtilities.cpp
@@ -1717,6 +1717,22 @@ void FSequencerUtilities::CopyFolders(const TArray<UMovieSceneFolder*>& Folders,
 	ExportObjectsToText(Objects, /*out*/ ExportedText);
 }
 
+void GatherChildBindings(UMovieSceneSequence* Sequence, const FGuid& ObjectBinding, TArray<FMovieSceneBindingProxy>& Bindings)
+{
+	UMovieScene* MovieScene = Sequence->GetMovieScene();
+
+	for (int32 Index = 0; Index < MovieScene->GetPossessableCount(); ++Index)
+	{
+		const FMovieScenePossessable& Possessable = MovieScene->GetPossessable(Index);
+		if (Possessable.GetParent() == ObjectBinding)
+		{
+			Bindings.AddUnique(FMovieSceneBindingProxy(Possessable.GetGuid(), Sequence));
+
+			GatherChildBindings(Sequence, Possessable.GetGuid(), Bindings);
+		}
+	}
+}
+
 void GatherFolderContents(UMovieSceneFolder* Folder, TArray<UMovieSceneFolder*>& Folders, TArray<UMovieSceneTrack*>& Tracks, TArray<FMovieSceneBindingProxy>& Bindings)
 {
 	if (!Folder)
@@ -1733,6 +1749,8 @@ void GatherFolderContents(UMovieSceneFolder* Folder, TArray<UMovieSceneFolder*>&
 	for (const FGuid& ObjectBinding : Folder->GetChildObjectBindings())
 	{
 		Bindings.AddUnique(FMovieSceneBindingProxy(ObjectBinding, Sequence));
+
+		GatherChildBindings(Sequence, ObjectBinding, Bindings);
 	}
 
 	for (UMovieSceneTrack* ChildTrack : Folder->GetChildTracks())
