Anim In Engine: Turn off throttle when mouse button is pressed not just on drag so scrub time doesn't throttle in sequencer/curve editor

diff --git a/Engine/Source/Editor/CurveEditor/Private/DragOperations/CurveEditorDragOperation_ScrubTime.cpp b/Engine/Source/Editor/CurveEditor/Private/DragOperations/CurveEditorDragOperation_ScrubTime.cpp
--- a/Engine/Source/Editor/CurveEditor/Private/DragOperations/CurveEditorDragOperation_ScrubTime.cpp
+++ b/Engine/Source/Editor/CurveEditor/Private/DragOperations/CurveEditorDragOperation_ScrubTime.cpp
@@ -18,6 +18,7 @@ FCurveEditorDragOperation_ScrubTime::FCurveEditorDragOperation_ScrubTime(FCurveE
 
 void FCurveEditorDragOperation_ScrubTime::OnMouseButtonDown(FVector2D InitialPosition, const FGeometry& MyGeometry, const FPointerEvent& MouseEvent)
 {
+	FSlateThrottleManager::Get().DisableThrottle(true);
 	TSharedPtr<ITimeSliderController> TimeSliderController = CurveEditor->GetTimeSliderController();
 	if (TimeSliderController)
 	{
@@ -43,6 +44,11 @@ void FCurveEditorDragOperation_ScrubTime::OnMouseButtonDown(FVector2D InitialPos
 	}
 }
 
+void FCurveEditorDragOperation_ScrubTime::OnMouseButtonUp(const FGeometry& MyGeometry, const FPointerEvent& MouseEvent)
+{
+	FSlateThrottleManager::Get().DisableThrottle(false);
+}
+
 void FCurveEditorDragOperation_ScrubTime::OnBeginDrag(FVector2D InitialPosition, FVector2D CurrentPosition, const FPointerEvent& MouseEvent)
 {
 	FSlateThrottleManager::Get().DisableThrottle(true);
diff --git a/Engine/Source/Editor/CurveEditor/Private/DragOperations/CurveEditorDragOperation_ScrubTime.h b/Engine/Source/Editor/CurveEditor/Private/DragOperations/CurveEditorDragOperation_ScrubTime.h
--- a/Engine/Source/Editor/CurveEditor/Private/DragOperations/CurveEditorDragOperation_ScrubTime.h
+++ b/Engine/Source/Editor/CurveEditor/Private/DragOperations/CurveEditorDragOperation_ScrubTime.h
@@ -19,6 +19,7 @@ public:
 	FCurveEditorDragOperation_ScrubTime(FCurveEditor* CurveEditor);
 
 	virtual void OnMouseButtonDown(FVector2D InitialPosition, const FGeometry& MyGeometry, const FPointerEvent& MouseEvent) override;
+	virtual void OnMouseButtonUp(const FGeometry& MyGeometry, const FPointerEvent& MouseEvent) override;
 	virtual void OnBeginDrag(FVector2D InitialPosition, FVector2D CurrentPosition, const FPointerEvent& MouseEvent) override;
 	virtual void OnDrag(FVector2D InitialPosition, FVector2D CurrentPosition, const FPointerEvent& MouseEvent) override;
 	virtual void OnEndDrag(FVector2D InitialPosition, FVector2D CurrentPosition, const FPointerEvent& MouseEvent) override;
diff --git a/Engine/Source/Editor/CurveEditor/Private/DragOperations/ICurveEditorDragOperation.cpp b/Engine/Source/Editor/CurveEditor/Private/DragOperations/ICurveEditorDragOperation.cpp
--- a/Engine/Source/Editor/CurveEditor/Private/DragOperations/ICurveEditorDragOperation.cpp
+++ b/Engine/Source/Editor/CurveEditor/Private/DragOperations/ICurveEditorDragOperation.cpp
@@ -12,6 +12,11 @@ void ICurveEditorDragOperation::MouseButtonDown(FVector2D InitialPosition, const
 	OnMouseButtonDown(InitialPosition, MyGeometry, MouseEvent);
 }
 
+void ICurveEditorDragOperation::MouseButtonUp(const FGeometry& MyGeometry, const FPointerEvent& MouseEvent)
+{
+	OnMouseButtonUp(MyGeometry, MouseEvent);
+}
+
 void ICurveEditorDragOperation::BeginDrag(FVector2D InitialPosition, FVector2D CurrentPosition, const FPointerEvent& MouseEvent)
 {
 	OnBeginDrag(InitialPosition, CurrentPosition, MouseEvent);
diff --git a/Engine/Source/Editor/CurveEditor/Private/SCurveEditorViewContainer.cpp b/Engine/Source/Editor/CurveEditor/Private/SCurveEditorViewContainer.cpp
--- a/Engine/Source/Editor/CurveEditor/Private/SCurveEditorViewContainer.cpp
+++ b/Engine/Source/Editor/CurveEditor/Private/SCurveEditorViewContainer.cpp
@@ -373,7 +373,10 @@ FReply SCurveEditorViewContainer::OnMouseButtonUp(const FGeometry& MyGeometry, c
 	FVector2D MousePosition = MyGeometry.AbsoluteToLocal(MouseEvent.GetScreenSpacePosition());
 
 	FReply Reply = FReply::Unhandled();
-
+	if (DragOperation.IsSet())
+	{
+		DragOperation->DragImpl->MouseButtonUp(MyGeometry, MouseEvent);
+	}
 	if (DragOperation.IsSet() && DragOperation->IsDragging())
 	{
 		FVector2D InitialPosition = DragOperation->GetInitialPosition();
diff --git a/Engine/Source/Editor/CurveEditor/Public/ICurveEditorDragOperation.h b/Engine/Source/Editor/CurveEditor/Public/ICurveEditorDragOperation.h
--- a/Engine/Source/Editor/CurveEditor/Public/ICurveEditorDragOperation.h
+++ b/Engine/Source/Editor/CurveEditor/Public/ICurveEditorDragOperation.h
@@ -32,10 +32,16 @@ public:
 	virtual ~ICurveEditorDragOperation() {}
 
 	/**
-	 * Start of the drag operation
+	 * Mouse button down
 	 */
 	void MouseButtonDown(FVector2D InitialPosition, const FGeometry& MyGeometry, const FPointerEvent& MouseEvent);
 
+
+	/**
+	 * Mouse button up
+	 */
+	void MouseButtonUp(const FGeometry& MyGeometry, const FPointerEvent& MouseEvent);
+
 	/**
 	 * Begin this drag operation with the specified initial and current positions
 	 */
@@ -75,6 +81,10 @@ protected:
 	virtual void OnMouseButtonDown(FVector2D InitialPosition, const FGeometry& MyGeometry, const FPointerEvent& MouseEvent)
 	{}
 
+	/** Implementation method for derived types to handle mouse up*/
+	virtual void OnMouseButtonUp(const FGeometry& MyGeometry, const FPointerEvent& MouseEvent)
+	{}
+
 	/** Implementation method for derived types to begin a drag */
 	virtual void OnBeginDrag(FVector2D InitialPosition, FVector2D CurrentPosition, const FPointerEvent& MouseEvent)
 	{}
diff --git a/Engine/Source/Editor/Sequencer/Private/Tools/SequencerEditTool_Selection.cpp b/Engine/Source/Editor/Sequencer/Private/Tools/SequencerEditTool_Selection.cpp
--- a/Engine/Source/Editor/Sequencer/Private/Tools/SequencerEditTool_Selection.cpp
+++ b/Engine/Source/Editor/Sequencer/Private/Tools/SequencerEditTool_Selection.cpp
@@ -475,6 +475,8 @@ FReply FSequencerEditTool_Selection::OnMouseButtonDown(SWidget& OwnerWidget, con
 {
 	using namespace UE::Sequencer;
 
+	FSlateThrottleManager::Get().DisableThrottle(true);
+
 	UpdateCursor(MyGeometry, MouseEvent);
 
 	DelayedDrag.Reset();
@@ -483,6 +485,7 @@ FReply FSequencerEditTool_Selection::OnMouseButtonDown(SWidget& OwnerWidget, con
 	{
 		if (bIsScrubbingTime && Sequencer.GetSequencerSettings()->GetScrubTimeStartFromCursor())
 		{
+
 			FSequencerTimeSliderController::FScrubRangeToScreen RangeToScreen = FSequencerTimeSliderController::FScrubRangeToScreen(Sequencer.GetViewRange(), MyGeometry.Size);
 			FFrameTime FrameTime = ComputeFrameTimeFromMouse(Sequencer, MyGeometry, MouseEvent.GetScreenSpacePosition(), RangeToScreen, true);
 
@@ -563,6 +566,8 @@ FReply FSequencerEditTool_Selection::OnMouseMove(SWidget& OwnerWidget, const FGe
 
 FReply FSequencerEditTool_Selection::OnMouseButtonUp(SWidget& OwnerWidget, const FGeometry& MyGeometry, const FPointerEvent& MouseEvent)
 {
+	FSlateThrottleManager::Get().DisableThrottle(false);
+
 	UpdateCursor(MyGeometry, MouseEvent);
 
 	DelayedDrag.Reset();
