Sequencer: Fix SelectInRange which wasn't getting all the channel models from the view model

This was causing operations like Select Forward / Backward to fail

diff --git a/Engine/Source/Editor/Sequencer/Private/Sequencer.cpp b/Engine/Source/Editor/Sequencer/Private/Sequencer.cpp
--- a/Engine/Source/Editor/Sequencer/Private/Sequencer.cpp
+++ b/Engine/Source/Editor/Sequencer/Private/Sequencer.cpp
@@ -2961,7 +2961,10 @@ void FSequencer::SelectInSelectionRange(const TSharedPtr<UE::Sequencer::FViewMod
 	{
 		TArray<FKeyHandle> HandlesScratch;
 
-		for (TSharedPtr<FChannelModel> Channel : Item->GetDescendantsOfType<FChannelModel>())
+		TSet<TSharedPtr<FChannelModel>> Channels;
+		SequencerHelpers::GetAllChannels(Item, Channels);
+
+		for (TSharedPtr<FChannelModel> Channel : Channels)
 		{
 			TSharedPtr<IKeyArea> KeyArea = Channel->GetKeyArea();
 			UMovieSceneSection*  Section = KeyArea->GetOwningSection();
@@ -6760,7 +6763,7 @@ void FSequencer::SynchronizeExternalSelectionWithSequencerSelection()
 			: OutlinerItem->FindAncestorOfType<IObjectBindingExtension>(true);
 
 		// If the closest node is an object node, try to get the actor/component nodes from it.
-	if (ObjectBinding)
+		if (ObjectBinding)
 		{
 			for (auto RuntimeObject : FindBoundObjects(ObjectBinding->GetObjectGuid(), ActiveTemplateIDs.Top()) )
 			{
