Sequencer: Notify when marks are changed/added/removed so that the navigation tool can update

diff --git a/Engine/Source/Editor/Sequencer/Private/Sequencer.cpp b/Engine/Source/Editor/Sequencer/Private/Sequencer.cpp
--- a/Engine/Source/Editor/Sequencer/Private/Sequencer.cpp
+++ b/Engine/Source/Editor/Sequencer/Private/Sequencer.cpp
@@ -10334,6 +10334,8 @@ void FSequencer::ToggleMarkAtPlayPosition()
 
 		FocusedMovieScene->Modify();
 		FocusedMovieScene->DeleteMarkedFrame(MarkedFrameIndex);
+
+		NotifyMovieSceneDataChanged(EMovieSceneDataChangeType::MovieSceneStructureItemRemoved);
 	}
 	else
 	{
@@ -10341,6 +10343,8 @@ void FSequencer::ToggleMarkAtPlayPosition()
 
 		FocusedMovieScene->Modify();
 		FocusedMovieScene->AddMarkedFrame(FMovieSceneMarkedFrame(TickFrameNumber));
+
+		NotifyMovieSceneDataChanged(EMovieSceneDataChangeType::MovieSceneStructureItemAdded);
 	}
 
 	ViewModel->GetSelection()->MarkedFrames.Empty();
@@ -10368,6 +10372,8 @@ void FSequencer::SetMarkedFrame(int32 InMarkIndex, FFrameNumber InFrameNumber)
 			
 	FocusedMovieScene->Modify();
 	FocusedMovieScene->SetMarkedFrame(InMarkIndex, InFrameNumber);
+
+	NotifyMovieSceneDataChanged(EMovieSceneDataChangeType::MovieSceneStructureItemsChanged);
 }
 
 void FSequencer::AddMarkedFrame(FFrameNumber FrameNumber)
@@ -10394,6 +10400,8 @@ void FSequencer::AddMarkedFrame(FFrameNumber FrameNumber)
 
 	FocusedMovieScene->Modify();
 	FocusedMovieScene->AddMarkedFrame(FMovieSceneMarkedFrame(FrameNumber));
+
+	NotifyMovieSceneDataChanged(EMovieSceneDataChangeType::MovieSceneStructureItemAdded);
 }
 
 void FSequencer::DeleteMarkedFrame(int32 InMarkIndex)
@@ -10424,6 +10432,8 @@ void FSequencer::DeleteMarkedFrame(int32 InMarkIndex)
 		FocusedMovieScene->DeleteMarkedFrame(InMarkIndex);
 
 		ViewModel->GetSelection()->MarkedFrames.Empty();
+
+		NotifyMovieSceneDataChanged(EMovieSceneDataChangeType::MovieSceneStructureItemRemoved);
 	}
 }
 
@@ -10453,6 +10463,8 @@ void FSequencer::DeleteAllMarkedFrames()
 	FocusedMovieScene->DeleteMarkedFrames();
 
 	ViewModel->GetSelection()->MarkedFrames.Empty();
+
+	NotifyMovieSceneDataChanged(EMovieSceneDataChangeType::MovieSceneStructureItemRemoved);
 }
 
 void FSequencer::StepToNextMark()
