Sequencer: Fix folder right-click menu showing entries for adding tracks in read only sequence.

diff --git a/Engine/Source/Editor/Sequencer/Private/MVVM/ViewModels/FolderModel.cpp b/Engine/Source/Editor/Sequencer/Private/MVVM/ViewModels/FolderModel.cpp
--- a/Engine/Source/Editor/Sequencer/Private/MVVM/ViewModels/FolderModel.cpp
+++ b/Engine/Source/Editor/Sequencer/Private/MVVM/ViewModels/FolderModel.cpp
@@ -296,14 +296,21 @@ bool FFolderModel::CanDrag() const
 
 void FFolderModel::BuildContextMenu(FMenuBuilder& MenuBuilder)
 {
-	TSharedPtr<FSequencer> Sequencer = StaticCastSharedPtr<FSequencer>(GetEditor()->GetSequencer());
-
-	MenuBuilder.BeginSection("ObjectBindings");
-	Sequencer->BuildAddObjectBindingsMenu(MenuBuilder);
+	const TSharedPtr<FSequencer> Sequencer = StaticCastSharedPtr<FSequencer>(GetEditor()->GetSequencer());
+	const bool bIsReadOnly = Sequencer->GetRootMovieSceneSequence()->GetMovieScene()->IsReadOnly();
+	
+	MenuBuilder.BeginSection("ObjectBindings"); // Create the section to allow extension
+	if (!bIsReadOnly)
+	{
+		Sequencer->BuildAddObjectBindingsMenu(MenuBuilder);
+	}
 	MenuBuilder.EndSection();
 
-	MenuBuilder.BeginSection("AddTracks");
-	Sequencer->BuildAddTrackMenu(MenuBuilder);
+	MenuBuilder.BeginSection("AddTracks"); // Create the section to allow extension
+	if (!bIsReadOnly)
+	{
+		Sequencer->BuildAddTrackMenu(MenuBuilder);
+	}
 	MenuBuilder.EndSection();
 
 	FOutlinerItemModel::BuildContextMenu(MenuBuilder);
