[Sequencer] Request filter update on movie scene data changed

diff --git a/Engine/Source/Editor/Sequencer/Private/Sequencer.cpp b/Engine/Source/Editor/Sequencer/Private/Sequencer.cpp
--- a/Engine/Source/Editor/Sequencer/Private/Sequencer.cpp
+++ b/Engine/Source/Editor/Sequencer/Private/Sequencer.cpp
@@ -2758,8 +2758,14 @@ void FSequencer::NotifyMovieSceneDataChanged( EMovieSceneDataChangeType DataChan
 		{
 			bNeedTreeRefresh = true;
 		}
-		else if (NodeTree->UpdateFiltersOnTrackValueChanged())
+		else
 		{
+			// NodeTree->UpdateFiltersOnTrackValueChanged() returns true if the filter data has changed.
+			// It first compares the filter data node count and then compares the node list itself.
+			// Since adding/removing keys does not affect the nodes, we must always refresh to be sure
+			// the list is up to date for any filters that use keys.
+			NodeTree->UpdateFiltersOnTrackValueChanged();
+
 			bNeedTreeRefresh = true;
 		}
 	}
