Sequencer:  Add counter to keep track how many times we throttle since we may not always get a mouse down event

diff --git a/Engine/Source/Editor/Sequencer/Private/Tools/SequencerEditTool_Selection.cpp b/Engine/Source/Editor/Sequencer/Private/Tools/SequencerEditTool_Selection.cpp
--- a/Engine/Source/Editor/Sequencer/Private/Tools/SequencerEditTool_Selection.cpp
+++ b/Engine/Source/Editor/Sequencer/Private/Tools/SequencerEditTool_Selection.cpp
@@ -475,6 +475,7 @@ FReply FSequencerEditTool_Selection::OnMouseButtonDown(SWidget& OwnerWidget, con
 {
 	using namespace UE::Sequencer;
 
+	++ThrottleSlateCounter;
 	FSlateThrottleManager::Get().DisableThrottle(true);
 
 	UpdateCursor(MyGeometry, MouseEvent);
@@ -566,7 +567,11 @@ FReply FSequencerEditTool_Selection::OnMouseMove(SWidget& OwnerWidget, const FGe
 
 FReply FSequencerEditTool_Selection::OnMouseButtonUp(SWidget& OwnerWidget, const FGeometry& MyGeometry, const FPointerEvent& MouseEvent)
 {
-	FSlateThrottleManager::Get().DisableThrottle(false);
+	if (ThrottleSlateCounter > 0)
+	{
+		--ThrottleSlateCounter;
+		FSlateThrottleManager::Get().DisableThrottle(false);
+	}
 
 	UpdateCursor(MyGeometry, MouseEvent);
 
diff --git a/Engine/Source/Editor/Sequencer/Private/Tools/SequencerEditTool_Selection.h b/Engine/Source/Editor/Sequencer/Private/Tools/SequencerEditTool_Selection.h
--- a/Engine/Source/Editor/Sequencer/Private/Tools/SequencerEditTool_Selection.h
+++ b/Engine/Source/Editor/Sequencer/Private/Tools/SequencerEditTool_Selection.h
@@ -71,4 +71,7 @@ private:
 
 	/** Whether or not we are scrubbing time*/
 	bool bIsScrubbingTime = false;
+
+	/** Throttle slate counter on mouse down*/
+	int32 ThrottleSlateCounter = 0;
 };
