Sequencer: Find the camera actor from the component

diff --git a/Engine/Source/Editor/Sequencer/Private/SequencerUtilities.cpp b/Engine/Source/Editor/Sequencer/Private/SequencerUtilities.cpp
--- a/Engine/Source/Editor/Sequencer/Private/SequencerUtilities.cpp
+++ b/Engine/Source/Editor/Sequencer/Private/SequencerUtilities.cpp
@@ -3909,9 +3909,12 @@ FGuid FSequencerUtilities::CreateOrReplaceBinding(TSharedPtr<ISequencer> Sequenc
 
 	if (Sequencer)
 	{
-		if (ACameraActor* NewCamera = Cast<ACameraActor>(InObject))
+		if (UCameraComponent* CameraComponent = MovieSceneHelpers::CameraComponentFromRuntimeObject(InObject))
 		{
-			NewCameraAdded(Sequencer.ToSharedRef(), NewCamera, BindingGuid);
+			if (ACameraActor* CameraActor = Cast<ACameraActor>(CameraComponent->GetOwner()))
+			{
+				NewCameraAdded(Sequencer.ToSharedRef(), CameraActor, BindingGuid);
+			}
 		}
 
 		Sequencer->OnAddBinding(BindingGuid, OwnerMovieScene);
