Sequencer: No need to NotifyMovieSceneDataChanged when setting marks. Follow up to 44322126

#rnx

diff --git a/Engine/Source/Editor/Sequencer/Private/Sequencer.cpp b/Engine/Source/Editor/Sequencer/Private/Sequencer.cpp
--- a/Engine/Source/Editor/Sequencer/Private/Sequencer.cpp
+++ b/Engine/Source/Editor/Sequencer/Private/Sequencer.cpp
@@ -10335,7 +10335,7 @@ void FSequencer::ToggleMarkAtPlayPosition()
 		FocusedMovieScene->Modify();
 		FocusedMovieScene->DeleteMarkedFrame(MarkedFrameIndex);
 
-		NotifyMovieSceneDataChanged(EMovieSceneDataChangeType::MovieSceneStructureItemRemoved);
+		NotifyMovieSceneDataChanged(EMovieSceneDataChangeType::MovieSceneStructureItemsChanged);
 	}
 	else
 	{
@@ -10344,7 +10344,7 @@ void FSequencer::ToggleMarkAtPlayPosition()
 		FocusedMovieScene->Modify();
 		FocusedMovieScene->AddMarkedFrame(FMovieSceneMarkedFrame(TickFrameNumber));
 
-		NotifyMovieSceneDataChanged(EMovieSceneDataChangeType::MovieSceneStructureItemAdded);
+		NotifyMovieSceneDataChanged(EMovieSceneDataChangeType::MovieSceneStructureItemsChanged);
 	}
 
 	ViewModel->GetSelection()->MarkedFrames.Empty();
@@ -10372,8 +10372,6 @@ void FSequencer::SetMarkedFrame(int32 InMarkIndex, FFrameNumber InFrameNumber)
 			
 	FocusedMovieScene->Modify();
 	FocusedMovieScene->SetMarkedFrame(InMarkIndex, InFrameNumber);
-
-	NotifyMovieSceneDataChanged(EMovieSceneDataChangeType::MovieSceneStructureItemsChanged);
 }
 
 void FSequencer::AddMarkedFrame(FFrameNumber FrameNumber)
@@ -10401,7 +10399,7 @@ void FSequencer::AddMarkedFrame(FFrameNumber FrameNumber)
 	FocusedMovieScene->Modify();
 	FocusedMovieScene->AddMarkedFrame(FMovieSceneMarkedFrame(FrameNumber));
 
-	NotifyMovieSceneDataChanged(EMovieSceneDataChangeType::MovieSceneStructureItemAdded);
+	NotifyMovieSceneDataChanged(EMovieSceneDataChangeType::MovieSceneStructureItemsChanged);
 }
 
 void FSequencer::DeleteMarkedFrame(int32 InMarkIndex)
@@ -10433,7 +10431,7 @@ void FSequencer::DeleteMarkedFrame(int32 InMarkIndex)
 
 		ViewModel->GetSelection()->MarkedFrames.Empty();
 
-		NotifyMovieSceneDataChanged(EMovieSceneDataChangeType::MovieSceneStructureItemRemoved);
+		NotifyMovieSceneDataChanged(EMovieSceneDataChangeType::MovieSceneStructureItemsChanged);
 	}
 }
 
@@ -10464,7 +10462,7 @@ void FSequencer::DeleteAllMarkedFrames()
 
 	ViewModel->GetSelection()->MarkedFrames.Empty();
 
-	NotifyMovieSceneDataChanged(EMovieSceneDataChangeType::MovieSceneStructureItemRemoved);
+	NotifyMovieSceneDataChanged(EMovieSceneDataChangeType::MovieSceneStructureItemsChanged);
 }
 
 void FSequencer::StepToNextMark()
