[Backout] - CL44461667
#rnx
Original CL Desc
-----------------------------------------------------------------
Sequencer: Find the camera actor from the component

diff --git a/Engine/Source/Editor/Sequencer/Private/SequencerUtilities.cpp b/Engine/Source/Editor/Sequencer/Private/SequencerUtilities.cpp
--- a/Engine/Source/Editor/Sequencer/Private/SequencerUtilities.cpp
+++ b/Engine/Source/Editor/Sequencer/Private/SequencerUtilities.cpp
@@ -3909,12 +3909,9 @@ FGuid FSequencerUtilities::CreateOrReplaceBinding(TSharedPtr<ISequencer> Sequenc
 
 	if (Sequencer)
 	{
-		if (UCameraComponent* CameraComponent = MovieSceneHelpers::CameraComponentFromRuntimeObject(InObject))
+		if (ACameraActor* NewCamera = Cast<ACameraActor>(InObject))
 		{
-			if (ACameraActor* CameraActor = Cast<ACameraActor>(CameraComponent->GetOwner()))
-			{
-				NewCameraAdded(Sequencer.ToSharedRef(), CameraActor, BindingGuid);
-			}
+			NewCameraAdded(Sequencer.ToSharedRef(), NewCamera, BindingGuid);
 		}
 
 		Sequencer->OnAddBinding(BindingGuid, OwnerMovieScene);
