Sequencer: Deprecate SetObjectGuid and GetBindings and FMovieSceneBinding constructors. Deprecate AddSpawnable and AddPossessable with a given FMovieSceneBinding since you can no longer construct one

MovieScene is a friend class to MovieSceneBinding which allows it to call SetObjectBinding_Private in special circumstances (ie. replacing an existing binding)

diff --git a/Engine/Plugins/Animation/ControlRig/Source/ControlRig/Private/Sequencer/ControlRigSequencerHelpers.cpp b/Engine/Plugins/Animation/ControlRig/Source/ControlRig/Private/Sequencer/ControlRigSequencerHelpers.cpp
--- a/Engine/Plugins/Animation/ControlRig/Source/ControlRig/Private/Sequencer/ControlRigSequencerHelpers.cpp
+++ b/Engine/Plugins/Animation/ControlRig/Source/ControlRig/Private/Sequencer/ControlRigSequencerHelpers.cpp
@@ -174,7 +174,7 @@ UMovieSceneControlRigParameterTrack* FControlRigSequencerHelpers::FindControlRig
 	{
 		return Track;
 	}
-	UMovieScene* MovieScene = MovieSceneSequence->GetMovieScene();
+	const UMovieScene* MovieScene = MovieSceneSequence->GetMovieScene();
 	const TArray<FMovieSceneBinding>& Bindings = MovieScene->GetBindings();
 	bool bRecreateCurves = false;
 	TArray<TPair<UControlRig*, FName>> ControlRigPairsToReselect;
diff --git a/Engine/Plugins/Animation/ControlRig/Source/ControlRigEditor/Private/ControlRigSequencerEditorLibrary.cpp b/Engine/Plugins/Animation/ControlRig/Source/ControlRigEditor/Private/ControlRigSequencerEditorLibrary.cpp
--- a/Engine/Plugins/Animation/ControlRig/Source/ControlRigEditor/Private/ControlRigSequencerEditorLibrary.cpp
+++ b/Engine/Plugins/Animation/ControlRig/Source/ControlRigEditor/Private/ControlRigSequencerEditorLibrary.cpp
@@ -76,7 +76,7 @@ TArray<FControlRigSequencerBindingProxy> UControlRigSequencerEditorLibrary::GetC
 	TArray<FControlRigSequencerBindingProxy> ControlRigBindingProxies;
 	if (LevelSequence)
 	{
-		UMovieScene* MovieScene = LevelSequence->GetMovieScene();
+		const UMovieScene* MovieScene = LevelSequence->GetMovieScene();
 		if (MovieScene)
 		{
 			const TArray<FMovieSceneBinding>& Bindings = MovieScene->GetBindings();
@@ -3234,7 +3234,7 @@ bool UControlRigSequencerEditorLibrary::RenameControlRigControlChannels(ULevelSe
 	bool bValid = false;
 
 
-	UMovieScene* MovieScene = LevelSequence->GetMovieScene();
+	const UMovieScene* MovieScene = LevelSequence->GetMovieScene();
 	if (MovieScene)
 	{
 		const TArray<FMovieSceneBinding>& Bindings = MovieScene->GetBindings();
diff --git a/Engine/Plugins/Animation/ControlRig/Source/ControlRigEditor/Private/Sequencer/AnimLayers/AnimLayers.cpp b/Engine/Plugins/Animation/ControlRig/Source/ControlRigEditor/Private/Sequencer/AnimLayers/AnimLayers.cpp
--- a/Engine/Plugins/Animation/ControlRig/Source/ControlRigEditor/Private/Sequencer/AnimLayers/AnimLayers.cpp
+++ b/Engine/Plugins/Animation/ControlRig/Source/ControlRigEditor/Private/Sequencer/AnimLayers/AnimLayers.cpp
@@ -91,7 +91,7 @@ static void GetSelectedControlRigsAndBoundObjects(ISequencer* SequencerPtr, TArr
 	using namespace UE::Sequencer;
 	TArray<FGuid> SelectedObjects = GetSelectedOutlinerGuids(SequencerPtr);;// = SequencerPtr->GetViewModel()->GetSelection()->GetBoundObjectsGuids();
 
-	UMovieScene* MovieScene = LevelSequence->GetMovieScene();
+	const UMovieScene* MovieScene = LevelSequence->GetMovieScene();
 	if (MovieScene)
 	{
 		const TArray<FMovieSceneBinding>& Bindings = MovieScene->GetBindings();
diff --git a/Engine/Plugins/Animation/ControlRig/Source/ControlRigEditor/Private/Sequencer/ControlRigSpaceChannelEditors.cpp b/Engine/Plugins/Animation/ControlRig/Source/ControlRigEditor/Private/Sequencer/ControlRigSpaceChannelEditors.cpp
--- a/Engine/Plugins/Animation/ControlRig/Source/ControlRigEditor/Private/Sequencer/ControlRigSpaceChannelEditors.cpp
+++ b/Engine/Plugins/Animation/ControlRig/Source/ControlRigEditor/Private/Sequencer/ControlRigSpaceChannelEditors.cpp
@@ -202,7 +202,7 @@ FSpaceChannelAndSection FControlRigSpaceChannelHelpers::FindSpaceChannelAndSecti
 
 	bool bFoundTrack = false;
 	
-	const TArray<FMovieSceneBinding>& Bindings = MovieScene->GetBindings();
+	const TArray<FMovieSceneBinding>& Bindings = ((const UMovieScene*)MovieScene)->GetBindings();
 	bool bRecreateCurves = false;
 	TArray<TPair<UControlRig*, FName>> ControlRigPairsToReselect;
 	for (const FMovieSceneBinding& Binding : Bindings)
@@ -1323,7 +1323,7 @@ UMovieSceneControlRigParameterSection* FControlRigSpaceChannelHelpers::GetContro
 	{
 		return nullptr;
 	}
-	const TArray<FMovieSceneBinding>& Bindings = MovieScene->GetBindings();
+	const TArray<FMovieSceneBinding>& Bindings = ((const UMovieScene*)MovieScene)->GetBindings();
 	for (const FMovieSceneBinding& Binding : Bindings)
 	{
 		UMovieSceneControlRigParameterTrack* ControlRigParameterTrack = Cast<UMovieSceneControlRigParameterTrack>(MovieScene->FindTrack(UMovieSceneControlRigParameterTrack::StaticClass(), Binding.GetObjectGuid(), NAME_None));
diff --git a/Engine/Plugins/Animation/ControlRig/Source/ControlRigEditor/Private/Tools/ControlRigSnapper.cpp b/Engine/Plugins/Animation/ControlRig/Source/ControlRigEditor/Private/Tools/ControlRigSnapper.cpp
--- a/Engine/Plugins/Animation/ControlRig/Source/ControlRigEditor/Private/Tools/ControlRigSnapper.cpp
+++ b/Engine/Plugins/Animation/ControlRig/Source/ControlRigEditor/Private/Tools/ControlRigSnapper.cpp
@@ -433,7 +433,7 @@ static UMovieSceneControlRigParameterSection* GetControlRigSection(ISequencer* S
 	{
 		return nullptr;
 	}
-	const TArray<FMovieSceneBinding>& Bindings = MovieScene->GetBindings();
+	const TArray<FMovieSceneBinding>& Bindings = ((const UMovieScene*)MovieScene)->GetBindings();
 	for (const FMovieSceneBinding& Binding : Bindings)
 	{
 		UMovieSceneControlRigParameterTrack* ControlRigParameterTrack = Cast<UMovieSceneControlRigParameterTrack>(MovieScene->FindTrack(UMovieSceneControlRigParameterTrack::StaticClass(), Binding.GetObjectGuid(), NAME_None));
diff --git a/Engine/Plugins/Animation/ControlRig/Source/ControlRigEditor/Private/Tools/EvaluateSequencerTools.cpp b/Engine/Plugins/Animation/ControlRig/Source/ControlRigEditor/Private/Tools/EvaluateSequencerTools.cpp
--- a/Engine/Plugins/Animation/ControlRig/Source/ControlRigEditor/Private/Tools/EvaluateSequencerTools.cpp
+++ b/Engine/Plugins/Animation/ControlRig/Source/ControlRigEditor/Private/Tools/EvaluateSequencerTools.cpp
@@ -514,7 +514,7 @@ namespace AIE
 
 	FMovieSceneBinding FSequencerTransformDependencies::GetBindingFromTrack(UMovieScene* InMovieScene, UMovieSceneTrack* Track)
 	{
-		const TArray<FMovieSceneBinding>& Bindings = InMovieScene->GetBindings();
+		const TArray<FMovieSceneBinding>& Bindings = ((const UMovieScene*)InMovieScene)->GetBindings();
 		for (const FMovieSceneBinding& Binding : Bindings)
 		{
 			for (UMovieSceneTrack* PossibleTrack : Binding.GetTracks())
diff --git a/Engine/Plugins/Editor/SequencerAnimTools/Source/SequencerAnimTools/Private/Sequencer/SequencerTrailHierarchy.cpp b/Engine/Plugins/Editor/SequencerAnimTools/Source/SequencerAnimTools/Private/Sequencer/SequencerTrailHierarchy.cpp
--- a/Engine/Plugins/Editor/SequencerAnimTools/Source/SequencerAnimTools/Private/Sequencer/SequencerTrailHierarchy.cpp
+++ b/Engine/Plugins/Editor/SequencerAnimTools/Source/SequencerAnimTools/Private/Sequencer/SequencerTrailHierarchy.cpp
@@ -826,7 +826,7 @@ FGuid FSequencerTrailHierarchy::AddControlRigTrail(USkeletalMeshComponent* Compo
 		{
 			//find binding id
 			TArray<UMovieSceneControlRigParameterTrack*> Tracks;
-			UMovieScene* MovieScene = Sequencer->GetFocusedMovieSceneSequence()->GetMovieScene();
+			const UMovieScene* MovieScene = Sequencer->GetFocusedMovieSceneSequence()->GetMovieScene();
 			const TArray<FMovieSceneBinding>& Bindings = MovieScene->GetBindings();
 			FGuid BindingID;
 			for (const FMovieSceneBinding& Binding : Bindings)
diff --git a/Engine/Plugins/Enterprise/GLTFExporter/Source/GLTFExporter/Private/Tasks/GLTFDelayedAnimationTasks.cpp b/Engine/Plugins/Enterprise/GLTFExporter/Source/GLTFExporter/Private/Tasks/GLTFDelayedAnimationTasks.cpp
--- a/Engine/Plugins/Enterprise/GLTFExporter/Source/GLTFExporter/Private/Tasks/GLTFDelayedAnimationTasks.cpp
+++ b/Engine/Plugins/Enterprise/GLTFExporter/Source/GLTFExporter/Private/Tasks/GLTFDelayedAnimationTasks.cpp
@@ -145,7 +145,7 @@ void FGLTFDelayedAnimSequenceTask::Process()
 void FGLTFDelayedLevelSequenceTask::Process()
 {
 	ULevelSequence* Sequence = const_cast<ULevelSequence*>(LevelSequence);
-	UMovieScene* MovieScene = Sequence->GetMovieScene();
+	const UMovieScene* MovieScene = Sequence->GetMovieScene();
 
 	ALevelSequenceActor* LevelSequenceActor;
 	UWorld* World = Level->GetWorld();
diff --git a/Engine/Plugins/Experimental/Animation/RigMapper/Source/RigMapperEditor/Private/RigMapperEditorSubsystem.cpp b/Engine/Plugins/Experimental/Animation/RigMapper/Source/RigMapperEditor/Private/RigMapperEditorSubsystem.cpp
--- a/Engine/Plugins/Experimental/Animation/RigMapper/Source/RigMapperEditor/Private/RigMapperEditorSubsystem.cpp
+++ b/Engine/Plugins/Experimental/Animation/RigMapper/Source/RigMapperEditor/Private/RigMapperEditorSubsystem.cpp
@@ -149,9 +149,10 @@ TArray<UMovieSceneControlRigParameterSection*> URigMapperEditorSubsystem::GetSec
 {
 	TArray<UMovieSceneControlRigParameterSection*> Sections;
 	
-	if (Sequence && Sequence->GetMovieScene())
+	const UMovieScene* MovieScene = Sequence ? Sequence->GetMovieScene() : nullptr;
+	if (MovieScene)
 	{
-		for (const FMovieSceneBinding& Binding : Sequence->GetMovieScene()->GetBindings())
+		for (const FMovieSceneBinding& Binding : MovieScene->GetBindings())
 		{
 			for (const UMovieSceneTrack* Track : Binding.GetTracks())
 			{
diff --git a/Engine/Plugins/Experimental/CineCameraRigs/Source/CineCameraRigs/Private/CineCameraRigRail.cpp b/Engine/Plugins/Experimental/CineCameraRigs/Source/CineCameraRigs/Private/CineCameraRigRail.cpp
--- a/Engine/Plugins/Experimental/CineCameraRigs/Source/CineCameraRigs/Private/CineCameraRigRail.cpp
+++ b/Engine/Plugins/Experimental/CineCameraRigs/Source/CineCameraRigs/Private/CineCameraRigRail.cpp
@@ -591,7 +591,7 @@ UMovieSceneFloatTrack* ACineCameraRigRail::FindPositionTrack(const UMovieSceneSe
 		return nullptr;
 	}
 	
-	UMovieScene* MovieScene = InSequence->GetMovieScene();
+	const UMovieScene* MovieScene = InSequence->GetMovieScene();
 
 	const TArray<FMovieSceneBinding>& Bindings = MovieScene->GetBindings();
 	const FName PropertyName = bUseAbsolutePosition ? FName(TEXT("AbsolutePositionOnRail")) : FName(TEXT("CurrentPositionOnRail"));
diff --git a/Engine/Plugins/Experimental/DaySequence/Source/DaySequence/Private/DaySequence.cpp b/Engine/Plugins/Experimental/DaySequence/Source/DaySequence/Private/DaySequence.cpp
--- a/Engine/Plugins/Experimental/DaySequence/Source/DaySequence/Private/DaySequence.cpp
+++ b/Engine/Plugins/Experimental/DaySequence/Source/DaySequence/Private/DaySequence.cpp
@@ -161,11 +161,8 @@ void UDaySequence::AddSpecializedBinding(EDaySequenceBindingReferenceSpecializat
 	}
 	
 	// Add a default binding
-	FMovieScenePossessable Possessable(GetBindingName(), GetClass());
-	FMovieSceneBinding     Binding(Possessable.GetGuid(), GetBindingName());
-
-	BindingReferences.AddSpecializedBinding(Possessable.GetGuid(), Specialization);
-	MovieScene->AddPossessable(Possessable, Binding);
+	FGuid PossessableGuid = MovieScene->AddPossessable(GetBindingName(), GetClass());
+	BindingReferences.AddSpecializedBinding(PossessableGuid, Specialization);
 }
 
 FGuid UDaySequence::GetSpecializedBinding(EDaySequenceBindingReferenceSpecialization Specialization) const
diff --git a/Engine/Plugins/Experimental/DaySequence/Source/DaySequence/Private/ProceduralDaySequenceBuilder.cpp b/Engine/Plugins/Experimental/DaySequence/Source/DaySequence/Private/ProceduralDaySequenceBuilder.cpp
--- a/Engine/Plugins/Experimental/DaySequence/Source/DaySequence/Private/ProceduralDaySequenceBuilder.cpp
+++ b/Engine/Plugins/Experimental/DaySequence/Source/DaySequence/Private/ProceduralDaySequenceBuilder.cpp
@@ -591,7 +591,7 @@ void UProceduralDaySequenceBuilder::ClearKeys()
 	
 	if (UMovieScene* MovieScene = ProceduralDaySequence->GetMovieScene())
 	{
-		for (const FMovieSceneBinding& Binding : MovieScene->GetBindings())
+		for (const FMovieSceneBinding& Binding : ((const UMovieScene*)MovieScene)->GetBindings())
 		{
 			// Inconvenient we have to do this but at least FindBinding is doing a binary search and we do this once per binding.
 			if (FMovieSceneBinding* MutableBinding = MovieScene->FindBinding(Binding.GetObjectGuid()))
@@ -636,16 +636,11 @@ FGuid UProceduralDaySequenceBuilder::GetOrCreateProceduralBinding(UObject* Objec
 	FGuid RootGuid = ProceduralDaySequence->FindBindingFromObject(TargetActor, SharedPlaybackState);
 	if (!RootGuid.IsValid())
 	{
-		FString RootName = TargetActor->GetName();
-		FMovieScenePossessable Possessable(RootName, TargetActor->GetClass());
-		FMovieSceneBinding     Binding(Possessable.GetGuid(), RootName);
-
-		RootGuid = Possessable.GetGuid();
-
 		// Explicitly invoke MarkAsChanged to ensure proper notification at runtime.
 		// The Modify that AddPossessable invokes only works in editor.
 		MovieScene->MarkAsChanged();
-		MovieScene->AddPossessable(Possessable, Binding);
+
+		RootGuid = MovieScene->AddPossessable(TargetActor->GetName(), TargetActor->GetClass());
 		ProceduralDaySequence->BindPossessableObject(RootGuid, *TargetActor, TargetActor);
 	}
 
@@ -659,17 +654,14 @@ FGuid UProceduralDaySequenceBuilder::GetOrCreateProceduralBinding(UObject* Objec
 	FGuid ComponentGuid = ProceduralDaySequence->FindBindingFromObject(Component, SharedPlaybackState);
 	if (!ComponentGuid.IsValid() && Component)
 	{
-		FString Name = Component->GetName();
-		FMovieScenePossessable Possessable(Name, Component->GetClass());
-		FMovieSceneBinding     Binding(Possessable.GetGuid(), Name);
-
-		Possessable.SetParent(RootGuid, MovieScene);
-		ComponentGuid = Possessable.GetGuid();
-
 		// Explicitly invoke MarkAsChanged to ensure proper notification at runtime.
 		// The Modify that AddPossessable invokes only works in editor.
 		MovieScene->MarkAsChanged();
-		MovieScene->AddPossessable(Possessable, Binding);
+
+		ComponentGuid = MovieScene->AddPossessable(Component->GetName(), Component->GetClass());
+		FMovieScenePossessable* NewPossessable = MovieScene->FindPossessable(ComponentGuid);
+		NewPossessable->SetParent(RootGuid, MovieScene);
+
 		ProceduralDaySequence->BindPossessableObject(ComponentGuid, *Component, TargetActor);
 	}
 
diff --git a/Engine/Plugins/Experimental/DaySequence/Source/DaySequenceEditor/Private/DaySequenceFactoryNew.cpp b/Engine/Plugins/Experimental/DaySequence/Source/DaySequenceEditor/Private/DaySequenceFactoryNew.cpp
--- a/Engine/Plugins/Experimental/DaySequence/Source/DaySequenceEditor/Private/DaySequenceFactoryNew.cpp
+++ b/Engine/Plugins/Experimental/DaySequence/Source/DaySequenceEditor/Private/DaySequenceFactoryNew.cpp
@@ -48,11 +48,8 @@ void UDaySequenceFactoryNew::AddDefaultBindings(UDaySequence* NewDaySequence)
 	UMovieScene* MovieScene = NewDaySequence->GetMovieScene();
 
 	// Add a default binding
-	FMovieScenePossessable Possessable(DefaultBindingName, ADaySequenceActor::StaticClass());
-	FMovieSceneBinding     Binding(Possessable.GetGuid(), DefaultBindingName);
-
-	NewDaySequence->AddDefaultBinding(Possessable.GetGuid());
-	MovieScene->AddPossessable(Possessable, Binding);
+	FGuid PossessableGuid = MovieScene->AddPossessable(DefaultBindingName, ADaySequenceActor::StaticClass());
+	NewDaySequence->AddDefaultBinding(PossessableGuid);
 }
 
 
diff --git a/Engine/Plugins/Experimental/GeometryCacheLevelSequenceBaker/Source/GeometryCacheLevelSequenceBaker/Private/GeometryCacheLevelSequenceBaker.cpp b/Engine/Plugins/Experimental/GeometryCacheLevelSequenceBaker/Source/GeometryCacheLevelSequenceBaker/Private/GeometryCacheLevelSequenceBaker.cpp
--- a/Engine/Plugins/Experimental/GeometryCacheLevelSequenceBaker/Source/GeometryCacheLevelSequenceBaker/Private/GeometryCacheLevelSequenceBaker.cpp
+++ b/Engine/Plugins/Experimental/GeometryCacheLevelSequenceBaker/Source/GeometryCacheLevelSequenceBaker/Private/GeometryCacheLevelSequenceBaker.cpp
@@ -239,7 +239,7 @@ TArray<FGuid> FGeometryCacheLevelSequenceBaker::GetBindingsToBake(TSharedRef<ISe
 	}
 
 	TArray<FGuid> CandidateBindings;
-	for (const FMovieSceneBinding& MovieSceneBinding : MovieScene->GetBindings())
+	for (const FMovieSceneBinding& MovieSceneBinding : ((const UMovieScene*)MovieScene)->GetBindings())
 	{
 		// If there are specific bindings to export, export those only
 		if (SelectedBindings.Num() != 0 && !SelectedBindings.Contains(MovieSceneBinding.GetObjectGuid()))
diff --git a/Engine/Plugins/Experimental/SequenceNavigator/Source/SequenceNavigator/Private/Items/NavigationToolSequence.cpp b/Engine/Plugins/Experimental/SequenceNavigator/Source/SequenceNavigator/Private/Items/NavigationToolSequence.cpp
--- a/Engine/Plugins/Experimental/SequenceNavigator/Source/SequenceNavigator/Private/Items/NavigationToolSequence.cpp
+++ b/Engine/Plugins/Experimental/SequenceNavigator/Source/SequenceNavigator/Private/Items/NavigationToolSequence.cpp
@@ -679,7 +679,7 @@ UMovieScene* FNavigationToolSequence::GetSequenceMovieScene() const
 
 TArray<FMovieSceneBinding> FNavigationToolSequence::GetSortedBindings() const
 {
-	UMovieScene* const MovieScene = GetSequenceMovieScene();
+	const UMovieScene* MovieScene = GetSequenceMovieScene();
 	if (!MovieScene)
 	{
 		return {};
diff --git a/Engine/Plugins/FX/NiagaraSimCaching/Source/NiagaraSimCachingEditor/Private/Sequencer/MovieSceneNiagaraTrackRecorder.cpp b/Engine/Plugins/FX/NiagaraSimCaching/Source/NiagaraSimCachingEditor/Private/Sequencer/MovieSceneNiagaraTrackRecorder.cpp
--- a/Engine/Plugins/FX/NiagaraSimCaching/Source/NiagaraSimCachingEditor/Private/Sequencer/MovieSceneNiagaraTrackRecorder.cpp
+++ b/Engine/Plugins/FX/NiagaraSimCaching/Source/NiagaraSimCachingEditor/Private/Sequencer/MovieSceneNiagaraTrackRecorder.cpp
@@ -35,7 +35,7 @@ UMovieSceneTrackRecorder* FMovieSceneNiagaraTrackRecorderFactory::CreateTrackRec
 		UMovieSceneNiagaraTrackRecorder* TrackRecorder = NewObject<UMovieSceneNiagaraTrackRecorder>();
 		
 		UMovieScene* MovieScene = Sequence->GetMovieScene();
-		const TArray<FMovieSceneBinding>& SceneBindings = MovieScene->GetBindings();
+		const TArray<FMovieSceneBinding>& SceneBindings = ((const UMovieScene*)MovieScene)->GetBindings();
 
 		for (const FMovieSceneBinding& Binding : SceneBindings)
 		{
diff --git a/Engine/Plugins/FX/NiagaraSimCaching/Source/NiagaraSimCachingEditor/Private/Sequencer/NiagaraCacheTrackEditor.cpp b/Engine/Plugins/FX/NiagaraSimCaching/Source/NiagaraSimCachingEditor/Private/Sequencer/NiagaraCacheTrackEditor.cpp
--- a/Engine/Plugins/FX/NiagaraSimCaching/Source/NiagaraSimCachingEditor/Private/Sequencer/NiagaraCacheTrackEditor.cpp
+++ b/Engine/Plugins/FX/NiagaraSimCaching/Source/NiagaraSimCachingEditor/Private/Sequencer/NiagaraCacheTrackEditor.cpp
@@ -458,7 +458,7 @@ bool FNiagaraCacheTrackEditor::IsCacheTrackOutOfDate(UMovieSceneTrack* Track)
 
 bool FNiagaraCacheTrackEditor::HasConflictingLifecycleTrack(UMovieSceneTrack* CacheTrack) const
 {
-	UMovieScene* MovieScene = GetMovieSceneSequence()->GetMovieScene();
+	const UMovieScene* MovieScene = GetMovieSceneSequence()->GetMovieScene();
 	FFrameNumber PlaybackStartFrame = MovieScene->GetPlaybackRange().GetLowerBoundValue();
 	const TArray<FMovieSceneBinding>& SceneBindings = MovieScene->GetBindings();
 
diff --git a/Engine/Plugins/Importers/USDImporter/Source/USDStageImporter/Private/USDStageImporter.cpp b/Engine/Plugins/Importers/USDImporter/Source/USDStageImporter/Private/USDStageImporter.cpp
--- a/Engine/Plugins/Importers/USDImporter/Source/USDStageImporter/Private/USDStageImporter.cpp
+++ b/Engine/Plugins/Importers/USDImporter/Source/USDStageImporter/Private/USDStageImporter.cpp
@@ -869,7 +869,7 @@ namespace UE::USDStageImporter::Private
 
 				// TODO: Remove empty bindings too?
 				bool bHasBindingsWithTracks = false;
-				TArray<FMovieSceneBinding> Bindings = MovieScene->GetBindings();
+				TArray<FMovieSceneBinding> Bindings = ((const UMovieScene*)MovieScene)->GetBindings();
 				for (FMovieSceneBinding& Binding : Bindings)
 				{
 					if (Binding.GetTracks().Num() != 0)
diff --git a/Engine/Plugins/Media/MediaCompositing/Source/MediaCompositing/Private/MovieScene/MovieSceneMediaTrack.cpp b/Engine/Plugins/Media/MediaCompositing/Source/MediaCompositing/Private/MovieScene/MovieSceneMediaTrack.cpp
--- a/Engine/Plugins/Media/MediaCompositing/Source/MediaCompositing/Private/MovieScene/MovieSceneMediaTrack.cpp
+++ b/Engine/Plugins/Media/MediaCompositing/Source/MediaCompositing/Private/MovieScene/MovieSceneMediaTrack.cpp
@@ -267,7 +267,7 @@ void UMovieSceneMediaTrack::UpdateSectionTextureIndices()
 {
 #if WITH_EDITOR
 	// Do we have an object binding?
-	UMovieScene* MovieScene = Cast<UMovieScene>(GetOuter());
+	const UMovieScene* MovieScene = Cast<UMovieScene>(GetOuter());
 	if (MovieScene != nullptr)
 	{
 		const TArray<FMovieSceneBinding>& Bindings = MovieScene->GetBindings();
diff --git a/Engine/Plugins/Media/MediaCompositing/Source/MediaCompositingEditor/Private/Sequencer/MediaTrackEditor.cpp b/Engine/Plugins/Media/MediaCompositing/Source/MediaCompositingEditor/Private/Sequencer/MediaTrackEditor.cpp
--- a/Engine/Plugins/Media/MediaCompositing/Source/MediaCompositingEditor/Private/Sequencer/MediaTrackEditor.cpp
+++ b/Engine/Plugins/Media/MediaCompositing/Source/MediaCompositingEditor/Private/Sequencer/MediaTrackEditor.cpp
@@ -373,7 +373,7 @@ void FMediaTrackEditor::OnSequencerDataChanged(EMovieSceneDataChangeType DataCha
 			UMovieScene* MovieScene = MovieSequence->GetMovieScene();
 
 			// Loop through all bindings.
-			const TArray<FMovieSceneBinding>& Bindings = MovieScene->GetBindings();
+			const TArray<FMovieSceneBinding>& Bindings = ((const UMovieScene*)MovieScene)->GetBindings();
 			for (const FMovieSceneBinding& Binding : Bindings)
 			{
 				UMovieSceneMediaTrack* Track = Cast<UMovieSceneMediaTrack>(MovieScene->FindTrack(UMovieSceneMediaTrack::StaticClass(), Binding.GetObjectGuid(), NAME_None));
diff --git a/Engine/Plugins/MovieScene/MovieRenderPipeline/Source/MovieRenderPipelineCore/Private/MoviePipelineUtils.cpp b/Engine/Plugins/MovieScene/MovieRenderPipeline/Source/MovieRenderPipelineCore/Private/MoviePipelineUtils.cpp
--- a/Engine/Plugins/MovieScene/MovieRenderPipeline/Source/MovieRenderPipelineCore/Private/MoviePipelineUtils.cpp
+++ b/Engine/Plugins/MovieScene/MovieRenderPipeline/Source/MovieRenderPipelineCore/Private/MoviePipelineUtils.cpp
@@ -642,7 +642,7 @@ namespace MoviePipeline
 						FString BindingName = TEXT("None");
 
 						// Try to find which binding it belongs to as the names will be more useful than section types.
-						for (const FMovieSceneBinding& Binding : Node->MovieScene->GetBindings())
+						for (const FMovieSceneBinding& Binding : ((const UMovieScene*)(Node->MovieScene.Get()))->GetBindings())
 						{
 							for (const UMovieSceneTrack* BindingTrack : Binding.GetTracks())
 							{
diff --git a/Engine/Plugins/MovieScene/SequencerScripting/Source/SequencerScripting/Private/ExtensionLibraries/MovieSceneSequenceExtensions.cpp b/Engine/Plugins/MovieScene/SequencerScripting/Source/SequencerScripting/Private/ExtensionLibraries/MovieSceneSequenceExtensions.cpp
--- a/Engine/Plugins/MovieScene/SequencerScripting/Source/SequencerScripting/Private/ExtensionLibraries/MovieSceneSequenceExtensions.cpp
+++ b/Engine/Plugins/MovieScene/SequencerScripting/Source/SequencerScripting/Private/ExtensionLibraries/MovieSceneSequenceExtensions.cpp
@@ -676,7 +676,7 @@ FMovieSceneBindingProxy UMovieSceneSequenceExtensions::FindBindingByName(UMovieS
 		return FMovieSceneBindingProxy();
 	}
 
-	UMovieScene* MovieScene = GetMovieScene(Sequence);
+	const UMovieScene* MovieScene = GetMovieScene(Sequence);
 	if (MovieScene)
 	{
 		const FMovieSceneBinding* Binding = Algo::FindBy(MovieScene->GetBindings(), Name, &FMovieSceneBinding::GetName);
@@ -718,7 +718,7 @@ TArray<FMovieSceneBindingProxy> UMovieSceneSequenceExtensions::GetBindings(UMovi
 
 	TArray<FMovieSceneBindingProxy> AllBindings;
 
-	UMovieScene* MovieScene = GetMovieScene(Sequence);
+	const UMovieScene* MovieScene = GetMovieScene(Sequence);
 	if (MovieScene)
 	{
 		for (const FMovieSceneBinding& Binding : MovieScene->GetBindings())
diff --git a/Engine/Plugins/MovieScene/SequencerScripting/Source/SequencerScriptingEditor/Private/SequencerTools.cpp b/Engine/Plugins/MovieScene/SequencerScripting/Source/SequencerScriptingEditor/Private/SequencerTools.cpp
--- a/Engine/Plugins/MovieScene/SequencerScripting/Source/SequencerScriptingEditor/Private/SequencerTools.cpp
+++ b/Engine/Plugins/MovieScene/SequencerScripting/Source/SequencerScriptingEditor/Private/SequencerTools.cpp
@@ -834,7 +834,7 @@ bool USequencerToolsFunctionLibrary::ImportFBXToControlRig(UWorld* World, ULevel
 
 	bool bValid = false;
 
-	const TArray<FMovieSceneBinding>& Bindings =  MovieScene->GetBindings();
+	const TArray<FMovieSceneBinding>& Bindings =  ((const UMovieScene*)MovieScene)->GetBindings();
 	for (const FMovieSceneBinding& Binding : Bindings)
 	{
 		if (Binding.GetName() == ControlRigTrackName)
@@ -895,7 +895,7 @@ bool USequencerToolsFunctionLibrary::ExportFBXFromControlRig(ULevelSequence* Seq
 		return false;
 	}
 
-	const TArray<FMovieSceneBinding>& Bindings =  MovieScene->GetBindings();
+	const TArray<FMovieSceneBinding>& Bindings = ((const UMovieScene*)MovieScene)->GetBindings();
 	for (const FMovieSceneBinding& Binding : Bindings)
 	{
 		if (Binding.GetName() != ActorWithControlRigTrack)
diff --git a/Engine/Plugins/MovieScene/TemplateSequence/Source/TemplateSequence/Private/TemplateSequence.cpp b/Engine/Plugins/MovieScene/TemplateSequence/Source/TemplateSequence/Private/TemplateSequence.cpp
--- a/Engine/Plugins/MovieScene/TemplateSequence/Source/TemplateSequence/Private/TemplateSequence.cpp
+++ b/Engine/Plugins/MovieScene/TemplateSequence/Source/TemplateSequence/Private/TemplateSequence.cpp
@@ -170,15 +170,15 @@ FGuid UTemplateSequence::FindOrAddBinding(UObject* InObject)
 		Player->PlayTo(FMovieSceneSequencePlaybackParams(MovieScene->GetPlaybackRange().GetLowerBoundValue(), EUpdatePositionMethod::Play), FMovieSceneSequencePlayToParams());
 
 		FGuid FoundBinding;
-		for (int32 BindingIndex = 0; !FoundBinding.IsValid() && BindingIndex < GetMovieScene()->GetBindings().Num(); ++BindingIndex)
+		for (int32 BindingIndex = 0; !FoundBinding.IsValid() && BindingIndex < ((const UMovieScene*)MovieScene)->GetBindings().Num(); ++BindingIndex)
 		{
-			for (TWeakObjectPtr<UObject> BoundObject : Player->FindBoundObjects(GetMovieScene()->GetBindings()[BindingIndex].GetObjectGuid(), MovieSceneSequenceID::Root))
+			for (TWeakObjectPtr<UObject> BoundObject : Player->FindBoundObjects(((const UMovieScene*)MovieScene)->GetBindings()[BindingIndex].GetObjectGuid(), MovieSceneSequenceID::Root))
 			{
 				if (BoundObject.IsValid() && BoundObject.Get()->GetClass() == InObject->GetClass())
 				{
 					if (!ParentObject || BoundObject.Get()->GetName() == InObject->GetName())
 					{
-						FoundBinding = GetMovieScene()->GetBindings()[BindingIndex].GetObjectGuid();
+						FoundBinding = ((const UMovieScene*)MovieScene)->GetBindings()[BindingIndex].GetObjectGuid();
 						break;
 					}
 				}
diff --git a/Engine/Plugins/MovieScene/TemplateSequence/Source/TemplateSequenceEditor/Private/TrackEditors/TemplateSequenceTrackEditor.cpp b/Engine/Plugins/MovieScene/TemplateSequence/Source/TemplateSequenceEditor/Private/TrackEditors/TemplateSequenceTrackEditor.cpp
--- a/Engine/Plugins/MovieScene/TemplateSequence/Source/TemplateSequenceEditor/Private/TrackEditors/TemplateSequenceTrackEditor.cpp
+++ b/Engine/Plugins/MovieScene/TemplateSequence/Source/TemplateSequenceEditor/Private/TrackEditors/TemplateSequenceTrackEditor.cpp
@@ -611,7 +611,7 @@ TArray<FTemplateSequenceSection::FScalablePropertyInfo> FTemplateSequenceSection
 		return AllAnimatedProperties;
 	}
 
-	UMovieScene* SubMovieScene = SubSequence->GetMovieScene();
+	const UMovieScene* SubMovieScene = SubSequence->GetMovieScene();
 	if (!SubMovieScene)
 	{
 		return AllAnimatedProperties;
diff --git a/Engine/Plugins/VirtualProduction/Avalanche/Source/AvalancheSequence/Private/AvaSequence.cpp b/Engine/Plugins/VirtualProduction/Avalanche/Source/AvalancheSequence/Private/AvaSequence.cpp
--- a/Engine/Plugins/VirtualProduction/Avalanche/Source/AvalancheSequence/Private/AvaSequence.cpp
+++ b/Engine/Plugins/VirtualProduction/Avalanche/Source/AvalancheSequence/Private/AvaSequence.cpp
@@ -438,7 +438,7 @@ TArray<UObject*> UAvaSequence::GetBoundObjects(UObject* InPlaybackContext) const
 		return OutObjects;
 	}
 
-	for (const FMovieSceneBinding& Binding : MovieScene->GetBindings())
+	for (const FMovieSceneBinding& Binding : ((const UMovieScene*)MovieScene)->GetBindings())
 	{
 		TArray<UObject*, TInlineAllocator<1>> BoundObject;
 		LocateBoundObjects(Binding.GetObjectGuid(), UE::UniversalObjectLocator::FResolveParams(InPlaybackContext), MovieSceneHelpers::CreateTransientSharedPlaybackState(InPlaybackContext, MovieScene->GetTypedOuter<UMovieSceneSequence>()), BoundObject);
diff --git a/Engine/Plugins/VirtualProduction/Takes/Source/TakeRecorder/Private/Recorder/TakeRecorderSubsystemImplementation.cpp b/Engine/Plugins/VirtualProduction/Takes/Source/TakeRecorder/Private/Recorder/TakeRecorderSubsystemImplementation.cpp
--- a/Engine/Plugins/VirtualProduction/Takes/Source/TakeRecorder/Private/Recorder/TakeRecorderSubsystemImplementation.cpp
+++ b/Engine/Plugins/VirtualProduction/Takes/Source/TakeRecorder/Private/Recorder/TakeRecorderSubsystemImplementation.cpp
@@ -826,7 +826,7 @@ UTakeRecorderNamingTokensData* UTakeRecorderSubsystemImplementation::GetNamingTo
 
 bool UTakeRecorderSubsystemImplementation::HasPendingChanges() const
 {
-	UMovieScene* MovieScene = GetLevelSequence() ? GetLevelSequence()->GetMovieScene() : nullptr;
+	const UMovieScene* MovieScene = GetLevelSequence() ? GetLevelSequence()->GetMovieScene() : nullptr;
 	const bool bHasUserMadeChanges = MovieScene
 		&& (!MovieScene->GetTracks().IsEmpty()
 			|| !MovieScene->GetBindings().IsEmpty()
diff --git a/Engine/Plugins/VirtualProduction/Takes/Source/TakeRecorder/Private/SerializedRecorder.cpp b/Engine/Plugins/VirtualProduction/Takes/Source/TakeRecorder/Private/SerializedRecorder.cpp
--- a/Engine/Plugins/VirtualProduction/Takes/Source/TakeRecorder/Private/SerializedRecorder.cpp
+++ b/Engine/Plugins/VirtualProduction/Takes/Source/TakeRecorder/Private/SerializedRecorder.cpp
@@ -376,6 +376,7 @@ AActor* FSerializedRecorder::SetActorPossesableOrSpawnable(UMovieSceneSequence*
 	UMovieScene* MovieScene = InMovieSceneSequence->GetMovieScene();
 
 	AActor* Actor = nullptr;
+	FGuid PossessableGuid;
 
 	if (ActorHeader.bRecordToPossessable)
 	{		
@@ -394,13 +395,8 @@ AActor* FSerializedRecorder::SetActorPossesableOrSpawnable(UMovieSceneSequence*
 		if (Actor)
 		{
 			UClass* InitPossessedObjectClass = UClass::TryFindTypeSlow<UClass>(ActorHeader.ClassName);
-			FMovieScenePossessable Possessable(ActorHeader.Label, InitPossessedObjectClass);
-
-			Possessable.SetGuid(ActorHeader.Guid);
-			FMovieSceneBinding NewBinding(ActorHeader.Guid, ActorHeader.Label);
-			MovieScene->AddPossessable(Possessable, NewBinding);
-
-			InMovieSceneSequence->BindPossessableObject(ActorHeader.Guid, *Actor, Actor->GetWorld());
+			PossessableGuid = MovieScene->AddPossessable(ActorHeader.Label, InitPossessedObjectClass);
+			InMovieSceneSequence->BindPossessableObject(PossessableGuid, *Actor, Actor->GetWorld());
 		}
 	}
 	else
@@ -423,21 +419,16 @@ AActor* FSerializedRecorder::SetActorPossesableOrSpawnable(UMovieSceneSequence*
 		//end of MakeSpawnableTemplateFromInstance hack
 
 		Actor->SetActorLabel(ActorHeader.Label); //spawnable has samen Name used for matching.
-
-		FMovieScenePossessable Possessable(ActorHeader.Label, SpawnableClass);
-
-		Possessable.SetGuid(ActorHeader.Guid);
-		FMovieSceneBinding NewBinding(ActorHeader.Guid, ActorHeader.Label);
-		MovieScene->AddPossessable(Possessable, NewBinding);
+		PossessableGuid = MovieScene->AddPossessable(ActorHeader.Label, SpawnableClass);
 
 		UE::Sequencer::FCreateBindingParams CreateBindingParams;
 		CreateBindingParams.bAllowCustomBinding = true;
 		CreateBindingParams.bSpawnable = true;
 		CreateBindingParams.BindingNameOverride = ActorHeader.Label;
-		CreateBindingParams.ReplacementGuid = ActorHeader.Guid;
+		CreateBindingParams.ReplacementGuid = PossessableGuid;
 		FSequencerUtilities::CreateOrReplaceBinding(nullptr, InMovieSceneSequence, Actor, CreateBindingParams);
 
-		if (ActorHeader.Guid.IsValid())
+		if (PossessableGuid.IsValid())
 		{
 			// Override the Skeletal Mesh components animation modes so that they can play back the recorded
 			// animation asset instead of their original animation source (such as Animation Blueprint)
@@ -471,10 +462,10 @@ AActor* FSerializedRecorder::SetActorPossesableOrSpawnable(UMovieSceneSequence*
 			}
 		}
 	}
-	ActorGuidToActorMap.Add(ActorHeader.Guid, Actor);
+	ActorGuidToActorMap.Add(PossessableGuid, Actor);
 
 	// Tag the possessable/spawnable with the original actor label so we can find it later
-	if (FMovieScenePossessable* Possessable = MovieScene->FindPossessable(ActorHeader.Guid))
+	if (FMovieScenePossessable* Possessable = MovieScene->FindPossessable(PossessableGuid))
 	{
 		if (!Possessable->Tags.Contains(*(ActorHeader.Label)))
 		{
@@ -502,7 +493,7 @@ AActor* FSerializedRecorder::SetActorPossesableOrSpawnable(UMovieSceneSequence*
 			MovieScene->AddRootFolder(FolderToUse);
 		}
 
-		FolderToUse->AddChildObjectBinding(ActorHeader.Guid);
+		FolderToUse->AddChildObjectBinding(PossessableGuid);
 	}
 	return Actor;
 }
@@ -515,22 +506,17 @@ void FSerializedRecorder::SetComponentPossessable(UMovieSceneSequence* InMovieSc
 		UMovieScene* InMovieScene = InMovieSceneSequence->GetMovieScene();
 
 		UClass* InitPossessedObjectClass = UClass::TryFindTypeSlow<UClass>(ActorProperty.ClassName);
-		FMovieScenePossessable ChildPossessable(ActorProperty.UObjectName, InitPossessedObjectClass);
-
-		ChildPossessable.SetGuid(ActorProperty.Guid);
-		FMovieSceneBinding NewBinding(ActorProperty.Guid, ActorHeader.Label);
-
-		InMovieScene->AddPossessable(ChildPossessable, NewBinding);
+		FGuid PossessableGuid = InMovieScene->AddPossessable(ActorProperty.UObjectName, InitPossessedObjectClass);
 
 		// Set up parent/child guids for possessables within spawnables
-		FMovieScenePossessable* ChildPossessablePtr = InMovieScene->FindPossessable(ActorProperty.Guid);
-		ChildPossessablePtr->SetParent(ActorHeader.Guid, InMovieScene);
+		FMovieScenePossessable* ChildPossessablePtr = InMovieScene->FindPossessable(PossessableGuid);
+		ChildPossessablePtr->SetParent(PossessableGuid, InMovieScene);
 
 		for (UActorComponent* Component : Actor->GetComponents())
 		{
 			if (Component->GetName() == ActorProperty.UObjectName)
 			{
-				InMovieSceneSequence->BindPossessableObject(ActorProperty.Guid, *Component, Actor);
+				InMovieSceneSequence->BindPossessableObject(PossessableGuid, *Component, Actor);
 				break;
 			}
 		}
diff --git a/Engine/Plugins/VirtualProduction/Takes/Source/TakeRecorder/Private/TakeRecorderModule.cpp b/Engine/Plugins/VirtualProduction/Takes/Source/TakeRecorder/Private/TakeRecorderModule.cpp
--- a/Engine/Plugins/VirtualProduction/Takes/Source/TakeRecorder/Private/TakeRecorderModule.cpp
+++ b/Engine/Plugins/VirtualProduction/Takes/Source/TakeRecorder/Private/TakeRecorderModule.cpp
@@ -149,7 +149,7 @@ namespace UE::TakeRecorder::Private
 			return bDidRenderSection;
 		};
 
-		const TArray<FMovieSceneBinding>& Bindings = MovieScene->GetBindings();
+		const TArray<FMovieSceneBinding>& Bindings = ((const UMovieScene*)MovieScene)->GetBindings();
 		bool bDidRenderForMovieScene = false;
 		for (const FMovieSceneBinding& Binding : Bindings)
 		{
diff --git a/Engine/Plugins/VirtualProduction/Takes/Source/TakesCore/Private/TakeRecorderSources.cpp b/Engine/Plugins/VirtualProduction/Takes/Source/TakesCore/Private/TakeRecorderSources.cpp
--- a/Engine/Plugins/VirtualProduction/Takes/Source/TakesCore/Private/TakeRecorderSources.cpp
+++ b/Engine/Plugins/VirtualProduction/Takes/Source/TakesCore/Private/TakeRecorderSources.cpp
@@ -712,7 +712,7 @@ void UTakeRecorderSources::RemoveRedundantTracks()
 		}
 
 		TArray<FGuid> ParentBindings;
-		for (const FMovieSceneBinding& Binding : MovieScene->GetBindings())
+		for (const FMovieSceneBinding& Binding : ((const UMovieScene*)MovieScene)->GetBindings())
 		{
 			FMovieScenePossessable* Possessable = MovieScene->FindPossessable(Binding.GetObjectGuid());
 			if (Possessable)
@@ -722,7 +722,7 @@ void UTakeRecorderSources::RemoveRedundantTracks()
 		}
 
 		TArray<FGuid> BindingsToRemove;
-		for (const FMovieSceneBinding& Binding : MovieScene->GetBindings())
+		for (const FMovieSceneBinding& Binding : ((const UMovieScene*)MovieScene)->GetBindings())
 		{
 			if (Binding.GetTracks().Num() == 0 && !ReferencedBindings.Contains(Binding.GetObjectGuid()) && !ParentBindings.Contains(Binding.GetObjectGuid()))
 			{
diff --git a/Engine/Source/Editor/MovieSceneTools/Private/EditModes/SkeletalAnimationTrackEditMode.cpp b/Engine/Source/Editor/MovieSceneTools/Private/EditModes/SkeletalAnimationTrackEditMode.cpp
--- a/Engine/Source/Editor/MovieSceneTools/Private/EditModes/SkeletalAnimationTrackEditMode.cpp
+++ b/Engine/Source/Editor/MovieSceneTools/Private/EditModes/SkeletalAnimationTrackEditMode.cpp
@@ -277,7 +277,7 @@ void FSkeletalAnimationTrackEditMode::Render(const FSceneView* View, FViewport*
 		return;
 	}
 
-	UMovieScene* MovieScene = Sequence->GetMovieScene();
+	const UMovieScene* MovieScene = Sequence->GetMovieScene();
 	if (!MovieScene)
 	{
 		return;
@@ -285,7 +285,7 @@ void FSkeletalAnimationTrackEditMode::Render(const FSceneView* View, FViewport*
 	// Gather a map of object bindings to their implict selection state
 	TMap<const FMovieSceneBinding*, bool> ObjectBindingNodesSelectionMap;
 
-	for (const FMovieSceneBinding& Binding : Sequence->GetMovieScene()->GetBindings())
+	for (const FMovieSceneBinding& Binding : MovieScene->GetBindings())
 	{
 		TArrayView<TWeakObjectPtr<UObject>> BoundObjects = Sequencer->FindObjectsInCurrentSequence(Binding.GetObjectGuid());
 		for (UMovieSceneTrack* Track : Binding.GetTracks())
diff --git a/Engine/Source/Editor/MovieSceneTools/Private/TrackEditors/ParticleTrackEditor.cpp b/Engine/Source/Editor/MovieSceneTools/Private/TrackEditors/ParticleTrackEditor.cpp
--- a/Engine/Source/Editor/MovieSceneTools/Private/TrackEditors/ParticleTrackEditor.cpp
+++ b/Engine/Source/Editor/MovieSceneTools/Private/TrackEditors/ParticleTrackEditor.cpp
@@ -78,7 +78,7 @@ int32 FParticleSection::OnPaintSection( FSequencerSectionPainter& InPainter ) co
 			TrackColor = ParentTrack->GetColorTint();
 
 			FGuid ObjectHandle;
-			for ( const FMovieSceneBinding& Binding : OwningSequencer->GetFocusedMovieSceneSequence()->GetMovieScene()->GetBindings() )
+			for ( const FMovieSceneBinding& Binding : ((const UMovieScene*)OwningSequencer->GetFocusedMovieSceneSequence()->GetMovieScene())->GetBindings() )
 			{
 				if ( Binding.GetTracks().Contains( ParentTrack ) )
 				{
diff --git a/Engine/Source/Editor/MovieSceneTools/Public/Conditions/MovieSceneDirectorBlueprintConditionUtils.h b/Engine/Source/Editor/MovieSceneTools/Public/Conditions/MovieSceneDirectorBlueprintConditionUtils.h
--- a/Engine/Source/Editor/MovieSceneTools/Public/Conditions/MovieSceneDirectorBlueprintConditionUtils.h
+++ b/Engine/Source/Editor/MovieSceneTools/Public/Conditions/MovieSceneDirectorBlueprintConditionUtils.h
@@ -77,7 +77,7 @@ struct FMovieSceneDirectorBlueprintConditionUtils
 		}
 
 		// Add all object binding sections
-		for (FMovieSceneBinding& Binding : InMovieScene->GetBindings())
+		for (const FMovieSceneBinding& Binding : ((const UMovieScene*)InMovieScene)->GetBindings())
 		{
 			for (UMovieSceneTrack* Track : Binding.GetTracks())
 			{
diff --git a/Engine/Source/Editor/Sequencer/Private/LevelEditorSequencerIntegration.cpp b/Engine/Source/Editor/Sequencer/Private/LevelEditorSequencerIntegration.cpp
--- a/Engine/Source/Editor/Sequencer/Private/LevelEditorSequencerIntegration.cpp
+++ b/Engine/Source/Editor/Sequencer/Private/LevelEditorSequencerIntegration.cpp
@@ -1606,7 +1606,7 @@ void AddActorsToBindingsMapRecursive(FSequencer& Sequencer, UMovieSceneSequence*
 
 void AddPropertiesToBindingsMap(TWeakPtr<FSequencer> Sequencer, UMovieSceneSequence* Sequence, FMovieSceneSequenceIDRef SequenceID, TMap<FObjectKey, TArray<FString> >& PropertyBindingsMap)
 {
-	UMovieScene* MovieScene = Sequence->GetMovieScene();
+	const UMovieScene* MovieScene = Sequence->GetMovieScene();
 
 	for (FMovieSceneBinding Binding : MovieScene->GetBindings())
 	{
diff --git a/Engine/Source/Editor/Sequencer/Private/MVVM/ObjectBindingModelStorageExtension.cpp b/Engine/Source/Editor/Sequencer/Private/MVVM/ObjectBindingModelStorageExtension.cpp
--- a/Engine/Source/Editor/Sequencer/Private/MVVM/ObjectBindingModelStorageExtension.cpp
+++ b/Engine/Source/Editor/Sequencer/Private/MVVM/ObjectBindingModelStorageExtension.cpp
@@ -131,7 +131,7 @@ void FObjectBindingModelStorageExtension::OnReinitialize()
 
 		FViewModelChildren RootChildren = OwnerModel->GetChildList(EViewModelListType::Outliner);
 
-		for (const FMovieSceneBinding& Binding : MovieScene->GetBindings())
+		for (const FMovieSceneBinding& Binding : ((const UMovieScene*)MovieScene)->GetBindings())
 		{
 			// Need to create a new one for this object binding
 			TSharedPtr<FViewModel> Model = GetOrCreateModelForBinding(Binding.GetObjectGuid());
diff --git a/Engine/Source/Editor/Sequencer/Private/Sequencer.cpp b/Engine/Source/Editor/Sequencer/Private/Sequencer.cpp
--- a/Engine/Source/Editor/Sequencer/Private/Sequencer.cpp
+++ b/Engine/Source/Editor/Sequencer/Private/Sequencer.cpp
@@ -6938,7 +6938,8 @@ void FSequencer::SynchronizeSequencerSelectionWithExternalSelection()
 		return;
 	}
 
-	if (!Sequence || !Sequence->GetMovieScene())
+	const UMovieScene* MovieScene = Sequence ? Sequence->GetMovieScene() : nullptr;
+	if (!MovieScene)
 	{
 		return;
 	}
@@ -6962,7 +6963,7 @@ void FSequencer::SynchronizeSequencerSelectionWithExternalSelection()
 	check(ObjectModelStorage);
 
 	TSet<TSharedPtr<FObjectBindingModel>> NodesToSelect;
-	for (const FMovieSceneBinding& Binding : Sequence->GetMovieScene()->GetBindings())
+	for (const FMovieSceneBinding& Binding : MovieScene->GetBindings())
 	{
 		TSharedPtr<FObjectBindingModel> ObjectModel = ObjectModelStorage->FindModelForObjectBinding(Binding.GetObjectGuid());
 		if (!ObjectModel)
@@ -9436,7 +9437,7 @@ void FixSortingOrders(UMovieSceneTrack* InTrack, const TArray<UMovieSceneTrack*>
 	}
 }
 
-void FixSortingOrders(FMovieSceneBinding* InBinding, UMovieScene* MovieScene)
+void FixSortingOrders(FMovieSceneBinding* InBinding, const UMovieScene* MovieScene)
 {
 	// Fix the sorting orders if InBinding is going to conflict with any of the existing bindings (ie. same class and same name)
 	TOptional<int32> MaxSortingOrder;
@@ -11070,7 +11071,7 @@ void FSequencer::TrimOrExtendSection(bool bTrimOrExtendLeft)
 		{
 			MovieSceneToolHelpers::TrimOrExtendSection(Track, TOptional<int32>(), GetLocalTime(), bTrimOrExtendLeft, Settings->GetDeleteKeysWhenTrimming());
 		}
-		for (const FMovieSceneBinding& Binding : MovieScene->GetBindings())
+		for (const FMovieSceneBinding& Binding : ((const UMovieScene*)MovieScene)->GetBindings())
 		{
 			for (UMovieSceneTrack* Track : Binding.GetTracks())
 			{
diff --git a/Engine/Source/Editor/Sequencer/Private/SequencerEdMode.cpp b/Engine/Source/Editor/Sequencer/Private/SequencerEdMode.cpp
--- a/Engine/Source/Editor/Sequencer/Private/SequencerEdMode.cpp
+++ b/Engine/Source/Editor/Sequencer/Private/SequencerEdMode.cpp
@@ -928,6 +928,12 @@ void FSequencerEdMode::DrawTracks3D(FPrimitiveDrawInterface* PDI)
 			continue;
 		}
 
+		const UMovieScene* MovieScene = Sequence->GetMovieScene();
+		if (!MovieScene)
+		{
+			continue;
+		}
+
 		// Gather a map of object bindings to their implict selection state
 		TMap<const FMovieSceneBinding*, bool> ObjectBindingNodesSelectionMap;
 
@@ -936,7 +942,7 @@ void FSequencerEdMode::DrawTracks3D(FPrimitiveDrawInterface* PDI)
 
 		const FSequencerSelection& Selection = *Sequencer->GetViewModel()->GetSelection();
 		const TSharedRef<FSequencerNodeTree>& NodeTree = Sequencer->GetNodeTree();
-		for (const FMovieSceneBinding& Binding : Sequence->GetMovieScene()->GetBindings())
+		for (const FMovieSceneBinding& Binding : MovieScene->GetBindings())
 		{
 			TSharedPtr<FObjectBindingModel> ObjectBindingNode = ObjectStorage->FindModelForObjectBinding(Binding.GetObjectGuid());
 			if (!ObjectBindingNode)
diff --git a/Engine/Source/Editor/Sequencer/Private/SequencerUtilities.cpp b/Engine/Source/Editor/Sequencer/Private/SequencerUtilities.cpp
--- a/Engine/Source/Editor/Sequencer/Private/SequencerUtilities.cpp
+++ b/Engine/Source/Editor/Sequencer/Private/SequencerUtilities.cpp
@@ -3153,7 +3153,6 @@ bool FSequencerUtilities::PasteBindings(const FString& TextToImport, TSharedRef<
 
 	TMap<FGuid, FGuid> OldToNewGuidMap;
 	TArray<FGuid> PossessableGuids;
-	TArray<TArray<FString> > PossessableObjectNames;
 	TArray<FGuid> SpawnableGuids;
 	TMap<FGuid, UMovieSceneFolder*> GuidToFolderMap;
 
@@ -3254,27 +3253,27 @@ bool FSequencerUtilities::PasteBindings(const FString& TextToImport, TSharedRef<
 				}
 				else
 				{
-					FMovieScenePossessable NewPossessable = CopyableBinding->Possessable;
-					NewPossessable.SetGuid(FGuid::NewGuid());
-					MovieScene->AddPossessable(NewPossessable, FMovieSceneBinding(NewPossessable.GetGuid(), NewPossessable.GetName()));
-					NewGuid = NewPossessable.GetGuid();
+					NewGuid = MovieScene->AddPossessable(CopyableBinding->Possessable.GetName(), (UClass*)CopyableBinding->Possessable.GetLoadedPossessedObjectClass());
+					FMovieScenePossessable* Possessable = MovieScene->FindPossessable(NewGuid);
+					Possessable->SetParent(CopyableBinding->Possessable.GetParent(), MovieScene);
 				}
-					
-				FMovieSceneBinding NewBinding(NewGuid, CopyableBinding->Binding.GetName(), CopyableBinding->Tracks);
+
+				FMovieSceneBinding* NewBinding = MovieScene->FindBinding(NewGuid);
+				NewBinding->SetTracks(TArray<UMovieSceneTrack*>(CopyableBinding->Tracks), nullptr);
 				FMovieScenePossessable* Possessable = MovieScene->FindPossessable(NewGuid);
 
 				// Clear the transient flags on the copyable binding before assigning to the new possessable
-				for (UMovieSceneTrack* Track : NewBinding.GetTracks())
+				for (UMovieSceneTrack* Track : NewBinding->GetTracks())
 				{
 					ResetCopiedTracksFlags(Track);
 				}
 
 				// Replace the auto-generated binding with our deserialized bindings (which has our tracks)
-				MovieScene->ReplaceBinding(NewGuid, NewBinding);
+				MovieScene->ReplaceBinding(NewGuid, *NewBinding);
 
 				OldToNewGuidMap.Add(CopyableBinding->Possessable.GetGuid(), NewGuid);
 
-				BindingsPasted.Add(NewBinding);
+				BindingsPasted.Add(*NewBinding);
 
 				PossessableGuids.Add(NewGuid);
 
@@ -3387,10 +3386,8 @@ bool FSequencerUtilities::PasteBindings(const FString& TextToImport, TSharedRef<
 						AddObjectsToBinding(Sequencer, ObjectsToBind, FMovieSceneBindingProxy(NewGuid, Sequence), ResolutionContext);
 					}
 				}
-
-				PossessableObjectNames.Add(CopyableBinding->BoundObjectNames);
 			}
-			else if (CopyableBinding->Spawnable.GetGuid().IsValid())
+			else if (CopyableBinding->Spawnable.GetGuid().IsValid() && !CopyableBinding->SpawnableObjectTemplates.IsEmpty())
 			{
 				// We need to let the sequence create the spawnable so that it has everything set up properly internally.
 				// This is required to get spawnables with the correct references to object templates, object templates with
@@ -3399,40 +3396,26 @@ bool FSequencerUtilities::PasteBindings(const FString& TextToImport, TSharedRef<
 				// an object template based off of our (transient package owned) template, then find the newly created binding
 				// and update it.
 
-				FGuid NewGuid;
-				if (!CopyableBinding->SpawnableObjectTemplates.IsEmpty())
-				{
-					NewGuid = MakeNewSpawnable(Sequencer, *CopyableBinding->SpawnableObjectTemplates[0], nullptr, false, FName(*CopyableBinding->Spawnable.GetName()));
-				}
-				else
-				{
-					FMovieSceneSpawnable NewSpawnable{};
-					NewSpawnable.SetGuid(FGuid::NewGuid());
-					NewSpawnable.SetName(CopyableBinding->Spawnable.GetName());
-
-					MovieScene->AddSpawnable(NewSpawnable, FMovieSceneBinding(NewSpawnable.GetGuid(), NewSpawnable.GetName()));
-
-					NewGuid = NewSpawnable.GetGuid();
-				}
-
-				FMovieSceneBinding NewBinding(NewGuid, CopyableBinding->Binding.GetName(), CopyableBinding->Tracks);
+				FGuid NewGuid = MakeNewSpawnable(Sequencer, *CopyableBinding->SpawnableObjectTemplates[0], nullptr, false, FName(*CopyableBinding->Spawnable.GetName()));
+				FMovieSceneBinding* NewBinding = MovieScene->FindBinding(NewGuid);
+				NewBinding->SetTracks(TArray<UMovieSceneTrack*>(CopyableBinding->Tracks), nullptr);
 				FMovieSceneSpawnable* Spawnable = MovieScene->FindSpawnable(NewGuid);
 
 				// Copy the name of the original spawnable too.
 				Spawnable->SetName(CopyableBinding->Spawnable.GetName());
 
 				// Clear the transient flags on the copyable binding before assigning to the new spawnable
-				for (auto Track : NewBinding.GetTracks())
+				for (auto Track : NewBinding->GetTracks())
 				{
 					ResetCopiedTracksFlags(Track);
 				}
 
 				// Replace the auto-generated binding with our deserialized bindings (which has our tracks)
-				MovieScene->ReplaceBinding(NewGuid, NewBinding);
+				MovieScene->ReplaceBinding(NewGuid, *NewBinding);
 
 				OldToNewGuidMap.Add(CopyableBinding->Spawnable.GetGuid(), NewGuid);
 
-				BindingsPasted.Add(NewBinding);
+				BindingsPasted.Add(*NewBinding);
 
 				SpawnableGuids.Add(NewGuid);
 
@@ -3454,88 +3437,6 @@ bool FSequencerUtilities::PasteBindings(const FString& TextToImport, TSharedRef<
 		}
 	}
 
-	// Fix possessable actor bindings
-	for (int32 PossessableGuidIndex = 0; PossessableGuidIndex < PossessableGuids.Num(); ++PossessableGuidIndex)
-	{
-		if (BindingReferences && Algo::AnyOf(BindingReferences->GetReferences(PossessableGuids[PossessableGuidIndex]), [](const FMovieSceneBindingReference& Reference) {return Reference.CustomBinding; }))
-		{
-			continue;
-		}
-		FMovieScenePossessable* Possessable = MovieScene->FindPossessable(PossessableGuids[PossessableGuidIndex]);
-		if (Possessable && World)
-		{
-			auto AddActor = [&](AActor* Actor)
-			{
-				FGuid ExistingGuid = Sequencer->FindObjectId(*Actor, Sequencer->GetFocusedTemplateID());
-
-				if (!ExistingGuid.IsValid())
-				{
-					FGuid NewGuid = AssignActor(Sequencer, Actor, Possessable->GetGuid());
-
-					// If assigning produces a new guid, update the possessable guids and the bindings pasted data
-					if (NewGuid.IsValid())
-					{
-						for (auto BindingPasted : BindingsPasted)
-						{
-							if (BindingPasted.GetObjectGuid() == PossessableGuids[PossessableGuidIndex])
-							{
-								BindingPasted.SetObjectGuid(NewGuid);
-							}
-						}
-
-						if (GuidToFolderMap.Contains(PossessableGuids[PossessableGuidIndex]))
-						{
-							GuidToFolderMap.Add(NewGuid, GuidToFolderMap[PossessableGuids[PossessableGuidIndex]]);
-							GuidToFolderMap.Remove(PossessableGuids[PossessableGuidIndex]);
-						}
-
-						for (TPair<FGuid, FGuid>& OldToNewGuid : OldToNewGuidMap)
-						{
-							if (OldToNewGuid.Value == PossessableGuids[PossessableGuidIndex])
-							{
-								OldToNewGuid.Value = NewGuid;
-							}
-						}
-
-						PossessableGuids[PossessableGuidIndex] = NewGuid;
-					}
-				}
-			};
-			
-			for (TActorIterator<AActor> ActorItr(World); ActorItr; ++ActorItr)
-			{
-				AActor* Actor = *ActorItr;
-				FString ActorName = Actor ? Actor->GetName() : TEXT("");
-				if (Actor && PossessableGuidIndex < PossessableObjectNames.Num() && PossessableObjectNames[PossessableGuidIndex].Contains(Actor->GetName()))
-				{
-					AddActor(Actor);
-				}
-			}
-
-			// If pasted actors have been provided, go through those as well
-			for (const TPair<FName, TObjectPtr<AActor>>& PastedActorPair : PasteBindingsParams.PastedActors)
-			{
-				if (PastedActorPair.Value && PossessableGuidIndex < PossessableObjectNames.Num())
-				{
-					FString PastedActorPairName = ObjectTools::SanitizeObjectPath(PastedActorPair.Key.ToString());
-					if (Algo::AnyOf(PossessableObjectNames[PossessableGuidIndex], [&PastedActorPairName] (const FString& PathName) 
-					{ 
-						FString SanitizedObjectPath = ObjectTools::SanitizeObjectPath(PathName);
-						int32 LastPeriod = INDEX_NONE;
-						if (SanitizedObjectPath.FindLastChar('.', LastPeriod))
-						{
-							return SanitizedObjectPath.RightChop(LastPeriod+1) == PastedActorPairName;
-						}
-						return SanitizedObjectPath == PastedActorPairName;
-					}))
-					{
-						AddActor(PastedActorPair.Value.Get());
-					}
-				}
-			}
-		}
-	}
-	
 	// Fix up parent guids
 	for (auto PossessableGuid : PossessableGuids)
 	{
@@ -3602,24 +3503,7 @@ bool FSequencerUtilities::PasteBindings(const FString& TextToImport, TSharedRef<
 							}
 						}
 					}
-
-					if (!bBound)
-					{
-						if (PossessableGuidIndex < PossessableObjectNames.Num())
-						{
-							for (const FString& BoundObjectPath : PossessableObjectNames[PossessableGuidIndex])
-							{
-								if (UObject* FoundObject = FindObject<UObject>(WeakObject.Get(), *BoundObjectPath))
-								{
-									Sequence->BindPossessableObject(PossessableGuid, *FoundObject, WeakObject.Get());
-									bBound = true;
-									break;
-								}
-							}
-						}
-					}
 				}
-				
 
 				// If the parent doesn't actually exist, clear it.
 				FMovieScenePossessable* PossessableParent = MovieScene->FindPossessable(ParentGuid);
diff --git a/Engine/Source/Editor/Sequencer/Private/SkeletalMeshComponentSchema.cpp b/Engine/Source/Editor/Sequencer/Private/SkeletalMeshComponentSchema.cpp
--- a/Engine/Source/Editor/Sequencer/Private/SkeletalMeshComponentSchema.cpp
+++ b/Engine/Source/Editor/Sequencer/Private/SkeletalMeshComponentSchema.cpp
@@ -240,13 +240,10 @@ void FSkeletalMeshComponentSchema::BindAnimationInstance(TWeakPtr<ISequencer> We
 			}
 		}
 
-		FMovieScenePossessable Possessable(NameString, AnimClass ? AnimClass : UAnimInstance::StaticClass());
-		FMovieSceneBinding     Binding(Possessable.GetGuid(), NameString);
+		FGuid PossessableGuid = MovieScene->AddPossessable(NameString, AnimClass ? AnimClass : UAnimInstance::StaticClass());
+		BindingReferences->AddBinding(PossessableGuid, CopyTemp(AnimInstanceLocator));
 
-		BindingReferences->AddBinding(Possessable.GetGuid(), CopyTemp(AnimInstanceLocator));
-		MovieScene->AddPossessable(Possessable, Binding);
-
-		FMovieScenePossessable* NewPossessable = MovieScene->FindPossessable(Possessable.GetGuid());
+		FMovieScenePossessable* NewPossessable = MovieScene->FindPossessable(PossessableGuid);
 		if (ensure(NewPossessable))
 		{
 			NewPossessable->SetParent(Pair.Key, MovieScene);
diff --git a/Engine/Source/Editor/UnrealEd/Private/Fbx/FbxMainExport.cpp b/Engine/Source/Editor/UnrealEd/Private/Fbx/FbxMainExport.cpp
--- a/Engine/Source/Editor/UnrealEd/Private/Fbx/FbxMainExport.cpp
+++ b/Engine/Source/Editor/UnrealEd/Private/Fbx/FbxMainExport.cpp
@@ -1749,7 +1749,7 @@ FString FFbxExporter::FLevelSequenceNodeNameAdapter::GetActorNodeName(const AAct
 {
 	FString NodeName = Actor->GetName();
 
-	for ( const FMovieSceneBinding& MovieSceneBinding : MovieScene->GetBindings() )
+	for ( const FMovieSceneBinding& MovieSceneBinding : ((const UMovieScene*)MovieScene)->GetBindings() )
 	{
 		for ( TWeakObjectPtr<UObject> RuntimeObject : MovieScenePlayer->FindBoundObjects(MovieSceneBinding.GetObjectGuid(), SequenceID) )
 		{
@@ -2009,7 +2009,7 @@ bool FFbxExporter::ExportLevelSequenceTracks(UMovieSceneSequence* MovieSceneSequ
 		SkeletalMeshComp = Actor ? Cast<USkeletalMeshComponent>(Actor->GetComponentByClass(USkeletalMeshComponent::StaticClass())) : nullptr;	
 	}
 
-	UMovieScene* MovieScene = MovieSceneSequence->GetMovieScene();
+	const UMovieScene* MovieScene = MovieSceneSequence->GetMovieScene();
 	FFrameRate DisplayRate = MovieScene->GetDisplayRate();
 
 	bool bSkip3DTransformTrack = SkeletalMeshComp && GetExportOptions()->MapSkeletalMotionToRoot;
@@ -2205,7 +2205,7 @@ bool FFbxExporter::ExportLevelSequence(UMovieSceneSequence* MovieSceneSequence,
 	{
 		return false;
 	}
-	UMovieScene* MovieScene = MovieSceneSequence->GetMovieScene();
+	const UMovieScene* MovieScene = MovieSceneSequence->GetMovieScene();
 	double FrameRate = MovieScene->GetDisplayRate().AsDecimal();
 	FbxTime::EMode TimeMode = FbxTime::ConvertFrameRateToTimeMode(FrameRate);
 
diff --git a/Engine/Source/Runtime/MovieScene/Private/Compilation/MovieSceneCompiledDataManager.cpp b/Engine/Source/Runtime/MovieScene/Private/Compilation/MovieSceneCompiledDataManager.cpp
--- a/Engine/Source/Runtime/MovieScene/Private/Compilation/MovieSceneCompiledDataManager.cpp
+++ b/Engine/Source/Runtime/MovieScene/Private/Compilation/MovieSceneCompiledDataManager.cpp
@@ -954,7 +954,7 @@ void UMovieSceneCompiledDataManager::Compile(FMovieSceneCompiledDataID DataID, U
 				CompileTrack(&Entry, nullptr, Track, Params, &GatheredSignatures, &GatheredData);
 			}
 
-			for (const FMovieSceneBinding& ObjectBinding : MovieScene->GetBindings())
+			for (const FMovieSceneBinding& ObjectBinding : ((const UMovieScene*)MovieScene)->GetBindings())
 			{
 				for (UMovieSceneTrack* Track : ObjectBinding.GetTracks())
 				{
@@ -1097,7 +1097,7 @@ void UMovieSceneCompiledDataManager::Gather(const FMovieSceneCompiledDataEntry&
 			GatherTrack(nullptr, Track, Params, TrackTemplate, OutCompilerData);
 		}
 
-		for (const FMovieSceneBinding& ObjectBinding : MovieScene->GetBindings())
+		for (const FMovieSceneBinding& ObjectBinding : ((const UMovieScene*)MovieScene)->GetBindings())
 		{
 			for (UMovieSceneTrack* Track : ObjectBinding.GetTracks())
 			{
@@ -1707,7 +1707,7 @@ bool UMovieSceneCompiledDataManager::GenerateSubSequenceData(UMovieSceneSequence
 		}
 	}
 
-	for (const FMovieSceneBinding& ObjectBinding : MovieScene->GetBindings())
+	for (const FMovieSceneBinding& ObjectBinding : ((const UMovieScene*)MovieScene)->GetBindings())
 	{
 		for (UMovieSceneTrack* Track : ObjectBinding.GetTracks())
 		{
@@ -1824,7 +1824,7 @@ void UMovieSceneCompiledDataManager::PopulateSubSequenceTree(UMovieSceneSequence
 		}
 	}
 
-	for (const FMovieSceneBinding& ObjectBinding : MovieScene->GetBindings())
+	for (const FMovieSceneBinding& ObjectBinding : ((const UMovieScene*)MovieScene)->GetBindings())
 	{
 		for (UMovieSceneTrack* Track : ObjectBinding.GetTracks())
 		{
diff --git a/Engine/Source/Runtime/MovieScene/Private/MovieScene.cpp b/Engine/Source/Runtime/MovieScene/Private/MovieScene.cpp
--- a/Engine/Source/Runtime/MovieScene/Private/MovieScene.cpp
+++ b/Engine/Source/Runtime/MovieScene/Private/MovieScene.cpp
@@ -1644,7 +1644,7 @@ void UMovieScene::ReplaceBinding(const FGuid& OldGuid, const FGuid& NewGuid, con
 {
 	if (FMovieSceneBinding* Binding = FindBinding(OldGuid))
 	{
-		Binding->SetObjectGuid(NewGuid);
+		Binding->SetObjectGuid_Private(NewGuid);
 		Binding->SetName(Name);
 
 		// Replace directly changes the guid, so force a sort here
diff --git a/Engine/Source/Runtime/MovieScene/Private/MovieSceneSequence.cpp b/Engine/Source/Runtime/MovieScene/Private/MovieSceneSequence.cpp
--- a/Engine/Source/Runtime/MovieScene/Private/MovieSceneSequence.cpp
+++ b/Engine/Source/Runtime/MovieScene/Private/MovieSceneSequence.cpp
@@ -263,14 +263,15 @@ void UMovieSceneSequence::OptimizeForCook()
 	}
 
 	// Go through object bindings.
-	for (int32 ObjectBindingIndex = 0; ObjectBindingIndex < MovieScene->GetBindings().Num(); )
+	const TArray<FMovieSceneBinding>& Bindings = ((const UMovieScene*)MovieScene)->GetBindings();
+	for (int32 ObjectBindingIndex = 0; ObjectBindingIndex < Bindings.Num(); )
 	{
 		bool bRemoveObject = false;
 
 		// First, see if we need to remove the object.
-		for (int32 TrackIndex = 0; TrackIndex < MovieScene->GetBindings()[ObjectBindingIndex].GetTracks().Num(); ++TrackIndex)
+		for (int32 TrackIndex = 0; TrackIndex < Bindings[ObjectBindingIndex].GetTracks().Num(); ++TrackIndex)
 		{
-			UMovieSceneTrack* Track = MovieScene->GetBindings()[ObjectBindingIndex].GetTracks()[TrackIndex];
+			UMovieSceneTrack* Track = Bindings[ObjectBindingIndex].GetTracks()[TrackIndex];
 			if (Track && Track->GetCookOptimizationFlags() == ECookOptimizationFlags::RemoveObject)
 			{
 				bRemoveObject = true;
@@ -279,9 +280,9 @@ void UMovieSceneSequence::OptimizeForCook()
 		}
 
 		// Then, remove any appropriate tracks, or all tracks if we decided to remove the object altogether.
-		for (int32 TrackIndex = 0; TrackIndex < MovieScene->GetBindings()[ObjectBindingIndex].GetTracks().Num(); )
+		for (int32 TrackIndex = 0; TrackIndex < Bindings[ObjectBindingIndex].GetTracks().Num(); )
 		{
-			UMovieSceneTrack* Track = MovieScene->GetBindings()[ObjectBindingIndex].GetTracks()[TrackIndex];
+			UMovieSceneTrack* Track = Bindings[ObjectBindingIndex].GetTracks()[TrackIndex];
 			if (Track && (Track->GetCookOptimizationFlags() == ECookOptimizationFlags::RemoveTrack || bRemoveObject))
 			{
 				Track->RemoveForCook();
@@ -294,9 +295,9 @@ void UMovieSceneSequence::OptimizeForCook()
 
 		// Go through the tracks again and look at sections.
 		// Once again, we recurse into sub-sequences if needed (see previous comment).
-		for (int32 TrackIndex = 0; TrackIndex < MovieScene->GetBindings()[ObjectBindingIndex].GetTracks().Num(); ++TrackIndex)
+		for (int32 TrackIndex = 0; TrackIndex < Bindings[ObjectBindingIndex].GetTracks().Num(); ++TrackIndex)
 		{
-			UMovieSceneTrack* Track = MovieScene->GetBindings()[ObjectBindingIndex].GetTracks()[TrackIndex];
+			UMovieSceneTrack* Track = Bindings[ObjectBindingIndex].GetTracks()[TrackIndex];
 			if (Track)
 			{
 				for (int32 SectionIndex = 0; SectionIndex < Track->GetAllSections().Num(); )
@@ -323,8 +324,8 @@ void UMovieSceneSequence::OptimizeForCook()
 		
 		if (bRemoveObject)
 		{
-			UE_LOG(LogMovieScene, Display, TEXT("Removing muted object: %s from: %s"), *MovieScene->GetBindings()[ObjectBindingIndex].GetName(), *GetPathName());
-			FGuid GuidToRemove = MovieScene->GetBindings()[ObjectBindingIndex].GetObjectGuid();
+			UE_LOG(LogMovieScene, Display, TEXT("Removing muted object: %s from: %s"), *Bindings[ObjectBindingIndex].GetName(), *GetPathName());
+			FGuid GuidToRemove = Bindings[ObjectBindingIndex].GetObjectGuid();
 			MovieScene->RemoveSpawnable(GuidToRemove);
 			MovieScene->RemovePossessable(GuidToRemove);
 		}
diff --git a/Engine/Source/Runtime/MovieScene/Private/MovieSceneSequenceVisitor.cpp b/Engine/Source/Runtime/MovieScene/Private/MovieSceneSequenceVisitor.cpp
--- a/Engine/Source/Runtime/MovieScene/Private/MovieSceneSequenceVisitor.cpp
+++ b/Engine/Source/Runtime/MovieScene/Private/MovieSceneSequenceVisitor.cpp
@@ -149,7 +149,7 @@ void VisitSequenceImpl(UMovieSceneSequence* Sequence, const FSequenceVisitParams
 
 	if (InParams.bVisitObjectBindings)
 	{
-		for (const FMovieSceneBinding& ObjectBinding : MovieScene->GetBindings())
+		for (const FMovieSceneBinding& ObjectBinding : ((const UMovieScene*)MovieScene)->GetBindings())
 		{
 			InVisitor.VisitObjectBinding(ObjectBinding, *SubSequenceSpace);
 
diff --git a/Engine/Source/Runtime/MovieScene/Private/MovieSceneTimeHelpers.cpp b/Engine/Source/Runtime/MovieScene/Private/MovieSceneTimeHelpers.cpp
--- a/Engine/Source/Runtime/MovieScene/Private/MovieSceneTimeHelpers.cpp
+++ b/Engine/Source/Runtime/MovieScene/Private/MovieSceneTimeHelpers.cpp
@@ -131,7 +131,7 @@ void TimeHelpers::MigrateFrameTimes(const IRetimingInterface& Retimer, UMovieSce
 	Retimer.Begin(MovieScene);
 
 	int32 TotalNumTracks = MovieScene->GetTracks().Num() + (MovieScene->GetCameraCutTrack() ? 1 : 0);
-	for (const FMovieSceneBinding& Binding : MovieScene->GetBindings())
+	for (const FMovieSceneBinding& Binding : ((const UMovieScene*)MovieScene)->GetBindings())
 	{
 		TotalNumTracks += Binding.GetTracks().Num();
 	}
@@ -181,7 +181,7 @@ void TimeHelpers::MigrateFrameTimes(const IRetimingInterface& Retimer, UMovieSce
 		UE::MovieScene::MigrateFrameTimes(Retimer, Track);
 	}
 
-	for (const FMovieSceneBinding& Binding : MovieScene->GetBindings())
+	for (const FMovieSceneBinding& Binding : ((const UMovieScene*)MovieScene)->GetBindings())
 	{
 		for (UMovieSceneTrack* Track : Binding.GetTracks())
 		{
diff --git a/Engine/Source/Runtime/MovieScene/Private/Sections/MovieSceneSubSection.cpp b/Engine/Source/Runtime/MovieScene/Private/Sections/MovieSceneSubSection.cpp
--- a/Engine/Source/Runtime/MovieScene/Private/Sections/MovieSceneSubSection.cpp
+++ b/Engine/Source/Runtime/MovieScene/Private/Sections/MovieSceneSubSection.cpp
@@ -560,7 +560,7 @@ void UMovieSceneSubSection::PostEditChangeProperty(FPropertyChangedEvent& Proper
 	if (PropertyChangedEvent.Property && PropertyChangedEvent.Property->GetFName() == GET_MEMBER_NAME_CHECKED(UMovieSceneSubSection, SubSequence))
 	{
 		// Check whether the subsequence that was just set has tracks that contain the sequence that this subsection is in.
-		UMovieScene* SubSequenceMovieScene = SubSequence ? SubSequence->GetMovieScene() : nullptr;
+		const UMovieScene* SubSequenceMovieScene = SubSequence ? SubSequence->GetMovieScene() : nullptr;
 
 		UMovieSceneSubTrack* TrackOuter = Cast<UMovieSceneSubTrack>(GetOuter());
 
diff --git a/Engine/Source/Runtime/MovieScene/Public/MovieScene.h b/Engine/Source/Runtime/MovieScene/Public/MovieScene.h
--- a/Engine/Source/Runtime/MovieScene/Public/MovieScene.h
+++ b/Engine/Source/Runtime/MovieScene/Public/MovieScene.h
@@ -775,9 +775,7 @@ public:
 		return ObjectBindings;
 	}
 
-	/**
-	* @return All object bindings.
-	*/
+	UE_DEPRECATED(5.7, "Getting non-const access to GetBindings is no longer allowed. Please use const GetBindings()")
 	TArray<FMovieSceneBinding>& GetBindings()
 	{
 		return ObjectBindings;
diff --git a/Engine/Source/Runtime/MovieScene/Public/MovieSceneBinding.h b/Engine/Source/Runtime/MovieScene/Public/MovieSceneBinding.h
--- a/Engine/Source/Runtime/MovieScene/Public/MovieSceneBinding.h
+++ b/Engine/Source/Runtime/MovieScene/Public/MovieSceneBinding.h
@@ -30,6 +30,8 @@ struct FMovieSceneBinding
 #endif
 	{ }
 
+private:
+
 	/**
 	 * Creates and initializes a new instance.
 	 *
@@ -61,10 +63,18 @@ struct FMovieSceneBinding
 	{ }
 
 	/**
-	* Set the object guid
-	*
-	* @param InObjectGuid
-	*/
+	 * Set the object guid
+	 *
+	 * @param InObjectGuid
+	 */
+	void SetObjectGuid_Private(const FGuid& InObjectGuid)
+	{
+		ObjectGuid = InObjectGuid;
+	}
+
+public:
+
+	UE_DEPRECATED(5.7, "SetObjectGuid is deprecated because guids can no longer be set directly.")
 	void SetObjectGuid(const FGuid& InObjectGuid)
 	{
 		ObjectGuid = InObjectGuid;
@@ -157,6 +167,8 @@ struct FMovieSceneBinding
 	}
 #endif
 
+	friend class UMovieScene;
+
 private:
 
 	/** Object binding guid for runtime objects */
