Sequencer: Create API for painting under top time slider and track area.

Modules can add a view model under the FEditorViewModel hierarchy, implemeting ITopTimeSliderOverlayExtension or ITrackAreaOverlayExtension, and then be painted by SAreaOverlay.

diff --git a/Engine/Source/Editor/Sequencer/Private/SSequencer.cpp b/Engine/Source/Editor/Sequencer/Private/SSequencer.cpp
--- a/Engine/Source/Editor/Sequencer/Private/SSequencer.cpp
+++ b/Engine/Source/Editor/Sequencer/Private/SSequencer.cpp
@@ -56,6 +56,8 @@
 #include "Widgets/Input/SSpinBox.h"
 #include "Widgets/Input/SCheckBox.h"
 #include "Widgets/Docking/SDockTab.h"
+#include "Widgets/Overlays/STopSliderAreaOverlay.h"
+#include "Widgets/Overlays/STrackAreaOverlay.h"
 #include "SPositiveActionButton.h"
 #include "Styling/AppStyle.h"
 #include "Kismet2/KismetEditorUtilities.h"
@@ -144,6 +146,8 @@
 #include "IContentBrowserSingleton.h"
 #include "ContentBrowserModule.h"
 #include "Menus/SequencerToolbarUtils.h"
+#include "MVVM/Extensions/ITopTimeSliderOverlayExtension.h"
+#include "MVVM/Extensions/ITrackAreaOverlayExtension.h"
 #include "Settings/EditorExperimentalSettings.h"
 
 #define LOCTEXT_NAMESPACE "Sequencer"
@@ -729,7 +733,7 @@ TSharedRef<SGridPanel> SSequencer::ConstructTrackAreaGridPanel(const FArguments&
 
 	USequencerToolMenuContext* ContextObject = NewObject<USequencerToolMenuContext>();
 	ContextObject->WeakSequencer = SequencerRef;
-
+	
 	return SNew(SGridPanel)
 		.FillRow(2, 1.f)
 		.FillColumn(0, FillCoefficient_0)
@@ -971,10 +975,12 @@ TSharedRef<SGridPanel> SSequencer::ConstructTrackAreaGridPanel(const FArguments&
 			.Padding(0)
 			.Clipping(EWidgetClipping::ClipToBounds)
 			[
-				TopTimeSlider.ToSharedRef()
+				SNew(SOverlay)
+				+ SOverlay::Slot().ZOrder(1) [ SNew(STopSliderAreaOverlay, SequencerRef->GetViewModel()) ]
+				+ SOverlay::Slot().ZOrder(2) [ TopTimeSlider.ToSharedRef() ]
 			]
 		]
-
+		
 		// Overlay that draws the tick lines
 		+ SGridPanel::Slot(Column1, Row2, SGridPanel::Layer(10))
 		.Padding(ResizeBarPadding)
@@ -985,6 +991,13 @@ TSharedRef<SGridPanel> SSequencer::ConstructTrackAreaGridPanel(const FArguments&
 			.DisplayTickLines(true)
 			.Clipping(EWidgetClipping::ClipToBounds)
 		]
+		
+		// Overlay that allows drawing over the (unpinned) track area
+		+ SGridPanel::Slot(Column1, Row2, SGridPanel::Layer(15))
+		.Padding(ResizeBarPadding)
+		[
+			SNew(STrackAreaOverlay, SequencerRef->GetViewModel())
+		]
 
 		// Overlay that draws the scrub position
 		+ SGridPanel::Slot(Column1, Row2, SGridPanel::Layer(20))
diff --git a/Engine/Source/Editor/Sequencer/Private/Widgets/Overlays/STopSliderAreaOverlay.cpp b/Engine/Source/Editor/Sequencer/Private/Widgets/Overlays/STopSliderAreaOverlay.cpp
new file mode 100644
--- /dev/null
+++ b/Engine/Source/Editor/Sequencer/Private/Widgets/Overlays/STopSliderAreaOverlay.cpp
@@ -0,0 +1,36 @@
+﻿// Copyright Epic Games, Inc. All Rights Reserved.
+
+#include "STopSliderAreaOverlay.h"
+
+#include "MVVM/Extensions/ITopTimeSliderOverlayExtension.h"
+#include "MVVM/ViewModels/ViewModel.h"
+#include "MVVM/ViewModels/ViewModelIterators.h"
+
+namespace UE::Sequencer
+{
+void STopSliderAreaOverlay::Construct(const FArguments& InArgs, TWeakViewModelPtr<FViewModel> InRootModel)
+{
+	WeakRootModel = MoveTemp(InRootModel);
+	SetVisibility(EVisibility::HitTestInvisible);
+}
+
+int32 STopSliderAreaOverlay::OnPaint(
+	const FPaintArgs& Args, const FGeometry& AllottedGeometry, const FSlateRect& MyCullingRect,
+	FSlateWindowElementList& OutDrawElements, int32 LayerId, const FWidgetStyle& InWidgetStyle, bool bParentEnabled
+	) const
+{
+	const TViewModelPtr<FViewModel> RootModelPin = WeakRootModel.Pin();
+	if (!RootModelPin)
+	{
+		return LayerId;
+	}
+
+	for (TViewModelPtr<ITopTimeSliderOverlayExtension> Child : RootModelPin->GetDescendantsOfType<ITopTimeSliderOverlayExtension>())
+	{
+		LayerId = Child->OnPaintTimeSliderOverlay(
+			Args, AllottedGeometry, MyCullingRect, OutDrawElements, LayerId, InWidgetStyle
+			);
+	}
+	return LayerId;
+}
+}
diff --git a/Engine/Source/Editor/Sequencer/Private/Widgets/Overlays/STopSliderAreaOverlay.h b/Engine/Source/Editor/Sequencer/Private/Widgets/Overlays/STopSliderAreaOverlay.h
new file mode 100644
--- /dev/null
+++ b/Engine/Source/Editor/Sequencer/Private/Widgets/Overlays/STopSliderAreaOverlay.h
@@ -0,0 +1,31 @@
+﻿// Copyright Epic Games, Inc. All Rights Reserved.
+
+#pragma once
+
+#include "MVVM/ViewModelPtr.h"
+#include "Widgets/DeclarativeSyntaxSupport.h"
+#include "Widgets/SCompoundWidget.h"
+
+namespace UE::Sequencer
+{
+/** Widget that overlay the track area. It overlays descendant FViewModels implementing ITopTimeSliderOverlayExtension to draw itself. */
+class STopSliderAreaOverlay : public SCompoundWidget
+{
+public:
+
+	SLATE_BEGIN_ARGS(STopSliderAreaOverlay){}
+	SLATE_END_ARGS()
+
+	/** @param InRootModel We search for models implementing the TModelInterface under InRootModel's hierarchy.  */
+	void Construct(const FArguments& InArgs, TWeakViewModelPtr<FViewModel> InRootModel);
+
+	//~ Begin SWidget Interface
+	virtual int32 OnPaint(const FPaintArgs& Args, const FGeometry& AllottedGeometry, const FSlateRect& MyCullingRect, FSlateWindowElementList& OutDrawElements, int32 LayerId, const FWidgetStyle& InWidgetStyle, bool bParentEnabled) const override;
+	//~ End SWidget Interface
+
+private:
+
+	/** On this model, we search for IOverlayExtension sub-interfaces. */
+	TWeakViewModelPtr<FViewModel> WeakRootModel;
+};
+}
\ No newline at end of file
diff --git a/Engine/Source/Editor/Sequencer/Private/Widgets/Overlays/STrackAreaOverlay.cpp b/Engine/Source/Editor/Sequencer/Private/Widgets/Overlays/STrackAreaOverlay.cpp
new file mode 100644
--- /dev/null
+++ b/Engine/Source/Editor/Sequencer/Private/Widgets/Overlays/STrackAreaOverlay.cpp
@@ -0,0 +1,36 @@
+﻿// Copyright Epic Games, Inc. All Rights Reserved.
+
+#include "STrackAreaOverlay.h"
+
+#include "MVVM/Extensions/ITrackAreaOverlayExtension.h"
+#include "MVVM/ViewModels/ViewModel.h"
+#include "MVVM/ViewModels/ViewModelIterators.h"
+
+namespace UE::Sequencer
+{
+void STrackAreaOverlay::Construct(const FArguments& InArgs, TWeakViewModelPtr<FViewModel> InRootModel)
+{
+	WeakRootModel = MoveTemp(InRootModel);
+	SetVisibility(EVisibility::HitTestInvisible);
+}
+
+int32 STrackAreaOverlay::OnPaint(
+	const FPaintArgs& Args, const FGeometry& AllottedGeometry, const FSlateRect& MyCullingRect,
+	FSlateWindowElementList& OutDrawElements, int32 LayerId, const FWidgetStyle& InWidgetStyle, bool bParentEnabled
+	) const
+{
+	const TViewModelPtr<FViewModel> RootModelPin = WeakRootModel.Pin();
+	if (!RootModelPin)
+	{
+		return LayerId;
+	}
+
+	for (TViewModelPtr<ITrackAreaOverlayExtension> Child : RootModelPin->GetDescendantsOfType<ITrackAreaOverlayExtension>())
+	{
+		LayerId = Child->OnPaintTrackAreaOverlay(
+			Args, AllottedGeometry, MyCullingRect, OutDrawElements, LayerId, InWidgetStyle
+			);
+	}
+	return LayerId;
+}
+}
diff --git a/Engine/Source/Editor/Sequencer/Private/Widgets/Overlays/STrackAreaOverlay.h b/Engine/Source/Editor/Sequencer/Private/Widgets/Overlays/STrackAreaOverlay.h
new file mode 100644
--- /dev/null
+++ b/Engine/Source/Editor/Sequencer/Private/Widgets/Overlays/STrackAreaOverlay.h
@@ -0,0 +1,31 @@
+﻿// Copyright Epic Games, Inc. All Rights Reserved.
+
+#pragma once
+
+#include "MVVM/ViewModelPtr.h"
+#include "Widgets/DeclarativeSyntaxSupport.h"
+#include "Widgets/SCompoundWidget.h"
+
+namespace UE::Sequencer
+{
+/** Widget that overlay the track area. It overlays descendant FViewModels implementing ITrackAreaOverlayExtension to draw itself. */
+class STrackAreaOverlay : public SCompoundWidget
+{
+public:
+
+	SLATE_BEGIN_ARGS(STrackAreaOverlay){}
+	SLATE_END_ARGS()
+
+	/** @param InRootModel We search for models implementing the TModelInterface under InRootModel's hierarchy.  */
+	void Construct(const FArguments& InArgs, TWeakViewModelPtr<FViewModel> InRootModel);
+
+	//~ Begin SWidget Interface
+	virtual int32 OnPaint(const FPaintArgs& Args, const FGeometry& AllottedGeometry, const FSlateRect& MyCullingRect, FSlateWindowElementList& OutDrawElements, int32 LayerId, const FWidgetStyle& InWidgetStyle, bool bParentEnabled) const override;
+	//~ End SWidget Interface
+
+private:
+
+	/** On this model, we search for IOverlayExtension sub-interfaces. */
+	TWeakViewModelPtr<FViewModel> WeakRootModel;
+};
+}
\ No newline at end of file
diff --git a/Engine/Source/Editor/Sequencer/Public/SequencerTimeSliderController.h b/Engine/Source/Editor/Sequencer/Public/SequencerTimeSliderController.h
--- a/Engine/Source/Editor/Sequencer/Public/SequencerTimeSliderController.h
+++ b/Engine/Source/Editor/Sequencer/Public/SequencerTimeSliderController.h
@@ -144,7 +144,7 @@ public:
 	/**
 	 * Draws major tick lines in the section view                                                              
 	 */
-	int32 OnPaintViewArea( const FGeometry& AllottedGeometry, const FSlateRect& MyCullingRect, FSlateWindowElementList& OutDrawElements, int32 LayerId, bool bEnabled, const FPaintViewAreaArgs& Args ) const;
+	virtual int32 OnPaintViewArea( const FGeometry& AllottedGeometry, const FSlateRect& MyCullingRect, FSlateWindowElementList& OutDrawElements, int32 LayerId, bool bEnabled, const FPaintViewAreaArgs& Args ) const override;
 
 
 
diff --git a/Engine/Source/Editor/SequencerCore/Private/MVVM/ViewModelTypeIDs.cpp b/Engine/Source/Editor/SequencerCore/Private/MVVM/ViewModelTypeIDs.cpp
--- a/Engine/Source/Editor/SequencerCore/Private/MVVM/ViewModelTypeIDs.cpp
+++ b/Engine/Source/Editor/SequencerCore/Private/MVVM/ViewModelTypeIDs.cpp
@@ -25,7 +25,9 @@
 #include "MVVM/Extensions/ILockableExtension.h"
 #include "MVVM/Extensions/ISortableExtension.h"
 #include "MVVM/Extensions/IStretchableExtension.h"
+#include "MVVM/Extensions/ITopTimeSliderOverlayExtension.h"
 #include "MVVM/Extensions/ITrackAreaExtension.h"
+#include "MVVM/Extensions/ITrackAreaOverlayExtension.h"
 #include "MVVM/Extensions/ITrackLaneExtension.h"
 #include "MVVM/Extensions/LinkedOutlinerExtension.h"
 #include "MVVM/Extensions/ITrackAreaViewSpaceProviderExtension.h"
@@ -75,7 +77,9 @@ UE_SEQUENCER_DEFINE_VIEW_MODEL_TYPE_ID(ISelectableExtension);
 UE_SEQUENCER_DEFINE_VIEW_MODEL_TYPE_ID(ISoloableExtension);
 UE_SEQUENCER_DEFINE_VIEW_MODEL_TYPE_ID(ISortableExtension);
 UE_SEQUENCER_DEFINE_VIEW_MODEL_TYPE_ID(IStretchableExtension);
+UE_SEQUENCER_DEFINE_VIEW_MODEL_TYPE_ID(ITopTimeSliderOverlayExtension);
 UE_SEQUENCER_DEFINE_VIEW_MODEL_TYPE_ID(ITrackAreaExtension);
+UE_SEQUENCER_DEFINE_VIEW_MODEL_TYPE_ID(ITrackAreaOverlayExtension);
 UE_SEQUENCER_DEFINE_VIEW_MODEL_TYPE_ID(ITrackLaneExtension);
 UE_SEQUENCER_DEFINE_VIEW_MODEL_TYPE_ID(IDeletableExtension);
 UE_SEQUENCER_DEFINE_VIEW_MODEL_TYPE_ID(IKeyExtension);
diff --git a/Engine/Source/Editor/SequencerCore/Public/MVVM/Extensions/ITopTimeSliderOverlayExtension.h b/Engine/Source/Editor/SequencerCore/Public/MVVM/Extensions/ITopTimeSliderOverlayExtension.h
new file mode 100644
--- /dev/null
+++ b/Engine/Source/Editor/SequencerCore/Public/MVVM/Extensions/ITopTimeSliderOverlayExtension.h
@@ -0,0 +1,34 @@
+﻿// Copyright Epic Games, Inc. All Rights Reserved.
+
+#pragma once
+
+#include "MVVM/ViewModelTypeID.h"
+
+#define UE_API SEQUENCERCORE_API
+
+class FPaintArgs;
+class FSlateRect;
+class FSlateWindowElementList;
+class FWidgetStyle;
+struct FGeometry;
+
+namespace UE::Sequencer
+{
+/** This extension allows you to paint custom geometry in the time slider area. */
+class ITopTimeSliderOverlayExtension
+{
+public:
+	
+	UE_SEQUENCER_DECLARE_VIEW_MODEL_TYPE_ID_API(UE_API, ITopTimeSliderOverlayExtension)
+
+	/**
+	 * Draws the geometry.
+	 * @return The highest layer that was painted.
+	 */
+	virtual int32 OnPaintTimeSliderOverlay(const FPaintArgs& Args, const FGeometry& AllottedGeometry, const FSlateRect& MyCullingRect, FSlateWindowElementList& OutDrawElements, int32 LayerId, const FWidgetStyle& InWidgetStyle) = 0;
+
+	virtual ~ITopTimeSliderOverlayExtension() = default;
+};
+}
+
+#undef UE_API
diff --git a/Engine/Source/Editor/SequencerCore/Public/MVVM/Extensions/ITrackAreaOverlayExtension.h b/Engine/Source/Editor/SequencerCore/Public/MVVM/Extensions/ITrackAreaOverlayExtension.h
new file mode 100644
--- /dev/null
+++ b/Engine/Source/Editor/SequencerCore/Public/MVVM/Extensions/ITrackAreaOverlayExtension.h
@@ -0,0 +1,34 @@
+﻿// Copyright Epic Games, Inc. All Rights Reserved.
+
+#pragma once
+
+#include "MVVM/ViewModelTypeID.h"
+
+#define UE_API SEQUENCERCORE_API
+
+class FPaintArgs;
+class FSlateRect;
+class FSlateWindowElementList;
+class FWidgetStyle;
+struct FGeometry;
+
+namespace UE::Sequencer
+{
+/** This extension allows you to paint custom geometry in the track area. */
+class ITrackAreaOverlayExtension
+{
+public:
+	
+	UE_SEQUENCER_DECLARE_VIEW_MODEL_TYPE_ID_API(UE_API, ITrackAreaOverlayExtension)
+
+	/**
+	 * Draws the geometry.
+	 * @return The highest layer that was painted.
+	 */
+	virtual int32 OnPaintTrackAreaOverlay(const FPaintArgs& Args, const FGeometry& AllottedGeometry, const FSlateRect& MyCullingRect, FSlateWindowElementList& OutDrawElements, int32 LayerId, const FWidgetStyle& InWidgetStyle) = 0;
+
+	virtual ~ITrackAreaOverlayExtension() = default;
+};
+}
+
+#undef UE_API
