Sequencer: Improvements to time space display

- Fix menu name
- Fix time space display not appearing in the proper space
- Draw both the desired time space and the local time space
- Don't clear the parent time space when switching subsequences, show the parent time space if it exists in the current parent chain

diff --git a/Engine/Source/Editor/Sequencer/Private/Sequencer.cpp b/Engine/Source/Editor/Sequencer/Private/Sequencer.cpp
--- a/Engine/Source/Editor/Sequencer/Private/Sequencer.cpp
+++ b/Engine/Source/Editor/Sequencer/Private/Sequencer.cpp
@@ -5886,6 +5886,8 @@ FString FSequencer::GetFrameTimeText() const
 		? &RootToWarpedLocalTransform
 		: &RootToUnwarpedLocalTransform;
 
+	FFrameTime LocalTime = bWarpedTime ? GlobalPlaybackWarpTransform.TransformTime(GetGlobalTime().Time) : GetGlobalTime().Time;
+
 	if (ScrubPositionParent.IsSet() && ScrubPositionParent.GetValue() != MovieSceneSequenceID::Root)
 	{
 		if (const FMovieSceneSequenceHierarchy* Hierarchy = CompiledDataManager->FindHierarchy(RootTemplateInstance.GetCompiledDataID()))
@@ -5901,9 +5903,9 @@ FString FSequencer::GetFrameTimeText() const
 				}
 			}
 		}
+		LocalTime = RootToParentChainTransform->TransformTime(GetGlobalTime().Time, FTransformTimeParams().IgnoreClamps());
 	}
 
-	FFrameTime LocalTime = RootToParentChainTransform->TransformTime(GetGlobalTime().Time, FTransformTimeParams().IgnoreClamps());
 	return GetNumericTypeInterface()->ToString(LocalTime.GetFrame().Value);
 }
 
@@ -5929,7 +5931,14 @@ TArray<FMovieSceneSequenceID> FSequencer::GetScrubPositionParentChain() const
 
 void FSequencer::OnScrubPositionParentChanged(FMovieSceneSequenceID InScrubPositionParent)
 {
-	ScrubPositionParent = InScrubPositionParent;
+	if (InScrubPositionParent == GetFocusedTemplateID())
+	{
+		ScrubPositionParent.Reset();
+	}
+	else
+	{
+		ScrubPositionParent = InScrubPositionParent;
+	}
 }
 
 void FSequencer::StartAutoscroll(float UnitsPerS)
diff --git a/Engine/Source/Editor/Sequencer/Private/SequencerTimeSliderController.cpp b/Engine/Source/Editor/Sequencer/Private/SequencerTimeSliderController.cpp
--- a/Engine/Source/Editor/Sequencer/Private/SequencerTimeSliderController.cpp
+++ b/Engine/Source/Editor/Sequencer/Private/SequencerTimeSliderController.cpp
@@ -659,40 +659,49 @@ int32 FSequencerTimeSliderController::OnPaintTimeSlider( bool bMirrorLabels, con
 		LayerId = DrawMarkedFrames(AllottedGeometry, RangeToScreen, OutDrawElements, LayerId, DrawEffects, InWidgetStyle, true);
 
 		{
-			// Draw the current time next to the scrub handle
-			FString FrameString;
-			FLinearColor TextColor = Args.TickColor;
-			if (TimeSliderArgs.ScrubPositionText.IsSet())
-			{
-				FrameString = TimeSliderArgs.ScrubPositionText.Get();
-			}
-			else
-			{
-				FrameString = Sequencer->GetNumericTypeInterface()->ToString(ScrubPosition.Time.GetFrame().Value);
-			}
+			// Draw the scrub position string and the local frame string as ScrubPosition : LocalFrame where the ScrubPosition might be time in a parent sequence space
+			TArray<FMovieSceneSequenceID> ParentChain = TimeSliderArgs.ScrubPositionParentChain.Get();
+			const bool bValidScrubPosition = TimeSliderArgs.ScrubPositionParent.Get() != MovieSceneSequenceID::Invalid && ParentChain.Contains(TimeSliderArgs.ScrubPositionParent.Get()) && TimeSliderArgs.ScrubPositionParent.Get() != Sequencer->GetFocusedTemplateID();
+			const FString ScrubPositionString = bValidScrubPosition ? FString::Format(TEXT("{0} : "), { TimeSliderArgs.ScrubPositionText.Get() }) : FString();
+			const FString LocalFrameString = Sequencer->GetNumericTypeInterface()->ToString(ScrubPosition.Time.GetFrame().Value);
+			const FString FullString = !ScrubPositionString.IsEmpty() ? FString::Join(TArray<FString>() = { ScrubPositionString, LocalFrameString }, TEXT(" ")) : LocalFrameString;
 
-			if (TimeSliderArgs.ScrubPositionParent.Get() != MovieSceneSequenceID::Invalid)
-			{
-				TextColor = FLinearColor::Yellow;
-			}
-
-			FVector2D TextSize = FontMeasureService->Measure(FrameString, SmallLayoutFont);
+			const FVector2D ScrubPositionTextSize = FontMeasureService->Measure(ScrubPositionString, SmallLayoutFont);
+			const FVector2D LocalFrameTextSize = FontMeasureService->Measure(LocalFrameString, SmallLayoutFont);
+			const FVector2D FullTextSize = FontMeasureService->Measure(FullString, SmallLayoutFont);
 
 			// Flip the text position if getting near the end of the view range
-			static const float TextOffsetPx  = 2.f;
-			bool  bDrawLeft    = (AllottedGeometry.Size.X - HandleEnd) < (TextSize.X + 14.f) - TextOffsetPx;
-			float TextPosition = bDrawLeft ? HandleStart - TextSize.X - TextOffsetPx : HandleEnd + TextOffsetPx;
+			static const float TextOffsetPx = 2.f;
+			const bool bDrawLeft = (AllottedGeometry.Size.X - HandleEnd) < (FullTextSize.X + 14.f) - TextOffsetPx;
+
+			float TextPosition = bDrawLeft ? HandleStart - FullTextSize.X - TextOffsetPx : HandleEnd + TextOffsetPx;
+			FVector2D TextOffset = FVector2D(TextPosition, Args.bMirrorLabels ? Args.AllottedGeometry.Size.Y - FullTextSize.Y : 0.f);
+			if (!ScrubPositionString.IsEmpty())
+			{
+				FLinearColor TextColor = FLinearColor::Yellow; // customizable
+
+				FSlateDrawElement::MakeText(
+					OutDrawElements,
+					Args.StartLayer + 1,
+					Args.AllottedGeometry.ToPaintGeometry(FullTextSize, FSlateLayoutTransform(TextOffset)),
+					ScrubPositionString,
+					SmallLayoutFont,
+					Args.DrawEffects,
+					TextColor
+				);
 
-			FVector2D TextOffset( TextPosition, Args.bMirrorLabels ? Args.AllottedGeometry.Size.Y - TextSize.Y : 0.f );
+				TextPosition = TextPosition + ScrubPositionTextSize.X;
+				TextOffset = FVector2D(TextPosition, Args.bMirrorLabels ? Args.AllottedGeometry.Size.Y - FullTextSize.Y : 0.f);
+			}
 
 			FSlateDrawElement::MakeText(
 				OutDrawElements,
 				Args.StartLayer+1, 
-				Args.AllottedGeometry.ToPaintGeometry( TextSize, FSlateLayoutTransform(TextOffset) ), 
-				FrameString, 
+				Args.AllottedGeometry.ToPaintGeometry(LocalFrameTextSize, FSlateLayoutTransform(TextOffset) ),
+				LocalFrameString, 
 				SmallLayoutFont,
 				Args.DrawEffects,
-				TextColor 
+				Args.TickColor
 			);
 		}
 		
@@ -2022,10 +2031,10 @@ TSharedRef<SWidget> FSequencerTimeSliderController::OpenSetPlaybackRangeMenu(con
 					);
 			}
 
-			// SequencerPlaybackRangeMenu
+			// SequencerSelectionRangeMenu
 			{
 				TRange<FFrameNumber> SelectionRange = TimeSliderController->TimeSliderArgs.SelectionRange.Get();
-				FToolMenuSection& Section = InMenu->AddSection("SequencerSelectionRangeMenu", FText::Format(LOCTEXT("SelectionRangeTextFormat", "Playback Range ({0}):"), CurrentTimeText));
+				FToolMenuSection& Section = InMenu->AddSection("SequencerSelectionRangeMenu", FText::Format(LOCTEXT("SelectionRangeTextFormat", "Selection Range ({0}):"), CurrentTimeText));
 
 				Section.AddMenuEntry(
 					"SetSelectionStart",
@@ -2086,7 +2095,7 @@ TSharedRef<SWidget> FSequencerTimeSliderController::OpenSetPlaybackRangeMenu(con
 
 				if (Hierarchy)
 				{
-					FToolMenuSection& Section = InMenu->AddSection("SequencerParentChainMenu");
+					FToolMenuSection& Section = InMenu->AddSection("SequencerParentChainMenu", LOCTEXT("TimeSpace", "Time Space"));
 
 					TArray<FMovieSceneSequenceID> ParentChain = TimeSliderController->TimeSliderArgs.ScrubPositionParentChain.Get();
 					for (FMovieSceneSequenceID ParentID : ParentChain)
@@ -2103,7 +2112,7 @@ TSharedRef<SWidget> FSequencerTimeSliderController::OpenSetPlaybackRangeMenu(con
 						}
 
 						Section.AddMenuEntry(
-							"DisplayTimeSpace",
+							FName(*ParentText.ToString()),
 							ParentText,
 							FText::Format(LOCTEXT("DisplayTimeSpace", "Display time in the space of {0}"), ParentText),
 							FSlateIcon(),
