Sequencer: Initialize playback mode to loop selection range if the selection range is set

diff --git a/Engine/Source/Editor/Sequencer/Private/Sequencer.cpp b/Engine/Source/Editor/Sequencer/Private/Sequencer.cpp
--- a/Engine/Source/Editor/Sequencer/Private/Sequencer.cpp
+++ b/Engine/Source/Editor/Sequencer/Private/Sequencer.cpp
@@ -638,6 +638,11 @@ void FSequencer::InitSequencer(const FSequencerInitParams& InitParams
 	UpdateTimeBases();
 	PlayPosition.Reset(ConvertFrameTime(GetPlaybackRange().GetLowerBoundValue(), GetRootTickResolution(), PlayPosition.GetInputRate()));
 
+	if (!GetSelectionRange().IsEmpty() && GetLoopMode() != SLM_LoopSelectionRange)
+	{
+		Settings->SetLoopMode(SLM_LoopSelectionRange);
+	}
+
 	FilterBar = MakeShared<FSequencerFilterBar>(*this);
 
 	// Make internal widgets
