Sequencer: Add SetStarPlaybackRange/SetEndPlaybackRange to shared command bindings for Sequencer/Curve Editor

diff --git a/Engine/Source/Editor/Sequencer/Private/Sequencer.cpp b/Engine/Source/Editor/Sequencer/Private/Sequencer.cpp
--- a/Engine/Source/Editor/Sequencer/Private/Sequencer.cpp
+++ b/Engine/Source/Editor/Sequencer/Private/Sequencer.cpp
@@ -12107,6 +12107,8 @@ void FSequencer::BindCommands()
 		CurveEditorSharedBindings->MapAction(Commands.StepToNextMark, *SequencerCommandBindings->GetActionForCommand(Commands.StepToNextMark));
 		CurveEditorSharedBindings->MapAction(Commands.StepToPreviousMark, *SequencerCommandBindings->GetActionForCommand(Commands.StepToPreviousMark));
 		CurveEditorSharedBindings->MapAction(Commands.ToggleMarkAtPlayPosition, *SequencerCommandBindings->GetActionForCommand(Commands.ToggleMarkAtPlayPosition));
+		CurveEditorSharedBindings->MapAction(Commands.SetStartPlaybackRange, *SequencerCommandBindings->GetActionForCommand(Commands.SetStartPlaybackRange));
+		CurveEditorSharedBindings->MapAction(Commands.SetEndPlaybackRange, *SequencerCommandBindings->GetActionForCommand(Commands.SetEndPlaybackRange));
 
 		CurveEditorSharedBindings->MapAction(Commands.AddTransformKey, *SequencerCommandBindings->GetActionForCommand(Commands.AddTransformKey));
 		CurveEditorSharedBindings->MapAction(Commands.AddTranslationKey, *SequencerCommandBindings->GetActionForCommand(Commands.AddTranslationKey));
