Sequencer: Fix root tracks not getting copied when copying folders as well

Copy tracks in folder separately from root tracks and then combine them afterwards

diff --git a/Engine/Source/Editor/Sequencer/Private/Sequencer.cpp b/Engine/Source/Editor/Sequencer/Private/Sequencer.cpp
--- a/Engine/Source/Editor/Sequencer/Private/Sequencer.cpp
+++ b/Engine/Source/Editor/Sequencer/Private/Sequencer.cpp
@@ -10683,6 +10683,7 @@ void FSequencer::CopySelection()
 
 		FString ObjectsExportedText;
 		FString TracksExportedText;
+		FString TracksInFoldersExportedText;
 		FString FoldersExportedText;
 
 		if (ObjectsToCopy.Num())
@@ -10697,13 +10698,14 @@ void FSequencer::CopySelection()
 
 		if (FoldersToCopy.Num())
 		{
-			CopySelectedFolders(FoldersToCopy, FoldersExportedText, ObjectsExportedText, TracksExportedText);
+			CopySelectedFolders(FoldersToCopy, FoldersExportedText, ObjectsExportedText, TracksInFoldersExportedText);
 		}
 
 		FString ExportedText;
 		ExportedText += ObjectsExportedText;
 		ExportedText += TracksExportedText;
 		ExportedText += FoldersExportedText;
+		ExportedText += TracksInFoldersExportedText;
 
 		FPlatformApplicationMisc::ClipboardCopy(*ExportedText);
 	}
