Sequencer: Add setting to always stop when going to playback start or end
#trivial

diff --git a/Engine/Source/Editor/Sequencer/Private/Sequencer.cpp b/Engine/Source/Editor/Sequencer/Private/Sequencer.cpp
--- a/Engine/Source/Editor/Sequencer/Private/Sequencer.cpp
+++ b/Engine/Source/Editor/Sequencer/Private/Sequencer.cpp
@@ -5174,7 +5174,8 @@ FReply FSequencer::OnStepBackward(FFrameNumber Increment)
 
 FReply FSequencer::OnJumpToStart()
 {
-	const bool bWasPlaying = (GetPlaybackStatus() == EMovieScenePlayerStatus::Playing);
+	
+	const bool bWasPlaying = Settings->GetStopPlayingWhenJumpingToStartOrEnd() == false && (GetPlaybackStatus() == EMovieScenePlayerStatus::Playing);
 
 	SetPlaybackStatus(EMovieScenePlayerStatus::Stepping);
 	SetLocalTime(UE::MovieScene::DiscreteInclusiveLower(GetTimeBounds()), ESnapTimeMode::STM_None);
@@ -5189,7 +5190,7 @@ FReply FSequencer::OnJumpToStart()
 
 FReply FSequencer::OnJumpToEnd()
 {
-	const bool bWasPlaying = (GetPlaybackStatus() == EMovieScenePlayerStatus::Playing);
+	const bool bWasPlaying = Settings->GetStopPlayingWhenJumpingToStartOrEnd() == false && (GetPlaybackStatus() == EMovieScenePlayerStatus::Playing);
 
 	SetPlaybackStatus(EMovieScenePlayerStatus::Stepping);
 	const bool bInsetDisplayFrame = ScrubStyle == ESequencerScrubberStyle::FrameBlock && Settings->GetForceWholeFrames();
diff --git a/Engine/Source/Editor/Sequencer/Private/SequencerSettings.cpp b/Engine/Source/Editor/Sequencer/Private/SequencerSettings.cpp
--- a/Engine/Source/Editor/Sequencer/Private/SequencerSettings.cpp
+++ b/Engine/Source/Editor/Sequencer/Private/SequencerSettings.cpp
@@ -66,6 +66,7 @@ USequencerSettings::USequencerSettings( const FObjectInitializer& ObjectInitiali
 	bEvaluateSubSequencesInIsolation = false;
 	bRerunConstructionScripts = true;
 	bScrubTimeStartFromCursor = true;
+	bStopPlayingWhenJumpingToStartOrEnd = false;
 	bVisualizePreAndPostRoll = true;
 	TrajectoryPathCap = 250;
 	FrameNumberDisplayFormat = EFrameNumberDisplayFormats::Seconds;
@@ -443,6 +444,20 @@ void USequencerSettings::SetScrubTimeStartFromCursor(bool bInValue)
 	}
 }
 
+bool USequencerSettings::GetStopPlayingWhenJumpingToStartOrEnd() const
+{
+	return bStopPlayingWhenJumpingToStartOrEnd;
+}
+
+void USequencerSettings::SetStopPlayingWhenJumpingToStartOrEnd(bool bInValue)
+{
+	if (bStopPlayingWhenJumpingToStartOrEnd != bInValue)
+	{
+		bStopPlayingWhenJumpingToStartOrEnd = bInValue;
+		SaveConfig();
+	}
+}
+
 void USequencerSettings::SetCurveEditorZoomScaling(const FCurveEditorZoomScaleConfig& Value)
 {
 	CurveEditorZoomScaling = Value;
diff --git a/Engine/Source/Editor/Sequencer/Public/SequencerSettings.h b/Engine/Source/Editor/Sequencer/Public/SequencerSettings.h
--- a/Engine/Source/Editor/Sequencer/Public/SequencerSettings.h
+++ b/Engine/Source/Editor/Sequencer/Public/SequencerSettings.h
@@ -267,6 +267,12 @@ public:
 	/** Sets whether or not scrubing time hot key starts from cursor. */
 	UE_API void SetScrubTimeStartFromCursor(bool bInValue);
 
+	/** Gets whether or not we stop playing when jumping to start or end. */
+	bool GetStopPlayingWhenJumpingToStartOrEnd() const;
+
+	/** Sets whether or not we stop playing when jumping to start or end. */
+	void SetStopPlayingWhenJumpingToStartOrEnd(bool bInValue);
+
 	/** Gets whether or not to link the curve editor time range. */
 	UE_API bool GetLinkCurveEditorTimeRange() const;
 	/** Sets whether or not to link the curve editor time range. */
@@ -671,6 +677,10 @@ protected:
 	UPROPERTY( config, EditAnywhere, Category=Timeline )
 	bool bAutoScrollEnabled;
 
+	/** When enabled, playing will stop when you jump to  start or end of range*/
+	UPROPERTY(config, EditAnywhere, Category = Timeline)
+	bool bStopPlayingWhenJumpingToStartOrEnd;
+
 	/** When enabled, scrubbing time hotkey will move time from initial cursor position  */
 	UPROPERTY(config, EditAnywhere, Category = Timeline)
 	bool bScrubTimeStartFromCursor;
