Sequencer: Fixed properties being exposed to Sequencer menus that should not have been

Sequencer was inadvertently showing properties that were deprecated or not marked as Interp if they related to the new registered property track editor. This is now corrected by unifying the logic that checks these property attributes to both codepaths

diff --git a/Engine/Source/Editor/Sequencer/Private/SequencerObjectChangeListener.cpp b/Engine/Source/Editor/Sequencer/Private/SequencerObjectChangeListener.cpp
--- a/Engine/Source/Editor/Sequencer/Private/SequencerObjectChangeListener.cpp
+++ b/Engine/Source/Editor/Sequencer/Private/SequencerObjectChangeListener.cpp
@@ -14,6 +14,33 @@
 
 DEFINE_LOG_CATEGORY_STATIC(LogSequencerTools, Log, All);
 
+namespace UE::MovieScene
+{
+	enum class EPropertyAccessibility : uint8
+	{
+		Inaccessible,
+		ExplicitlyAccessible,
+		Unknown,
+	};
+	static EPropertyAccessibility GetPropertyAccessibility(const FProperty& Property)
+	{
+		// Skip deprecated properties
+		static const FName NAME_DeprecatedProperty = TEXT("DeprecatedProperty");
+		if (Property.HasAnyPropertyFlags(CPF_Deprecated) || Property.HasMetaData(NAME_DeprecatedProperty))
+		{
+			return EPropertyAccessibility::Inaccessible;
+		}
+
+		// Interp properties are always keyable
+		if (Property.HasAnyPropertyFlags(CPF_Interp))
+		{
+			return EPropertyAccessibility::ExplicitlyAccessible;
+		}
+
+		return EPropertyAccessibility::Unknown;
+	}
+}
+
 FSequencerObjectChangeListener::FSequencerObjectChangeListener( TSharedRef<ISequencer> InSequencer )
 	: Sequencer( InSequencer )
 {
@@ -202,6 +229,8 @@ bool IsHiddenFunction(const UStruct& PropertyStructure, FAnimatedPropertyKey Pro
 
 const FOnAnimatablePropertyChanged* FSequencerObjectChangeListener::FindPropertySetter(const UStruct& PropertyStructure, FAnimatedPropertyKey PropertyKey, const FProperty& Property) const
 {
+	using namespace UE::MovieScene;
+
 	const FArrayProperty* ArrayOwner = Property.GetTypedOwner<FArrayProperty>();
 	const FProperty* PropertyOrContainer = ArrayOwner ? ArrayOwner : &Property;
 
@@ -217,11 +246,21 @@ const FOnAnimatablePropertyChanged* FSequencerObjectChangeListener::FindProperty
 
 	if (const FOnAnimatablePropertyChanged* DelegatePtr = PropertyChangedEventMap.Find(PropertyKey))
 	{
-		// Skip deprecated properties
-		static const FName NAME_DeprecatedProperty = TEXT("DeprecatedProperty");
-		if (PropertyOrContainer->HasAnyPropertyFlags(CPF_Deprecated) || PropertyOrContainer->HasMetaData(NAME_DeprecatedProperty))
+			// If the function is hidden, we cannot use it
+		const bool bHiddenFunction = bCanApplyFunction && IsHiddenFunction(PropertyStructure, FAnimatedPropertyKey::FromProperty(&Property), Property.GetName());
+		if (bHiddenFunction)
+		{
+			return nullptr;
+		}
+
+		switch (GetPropertyAccessibility(*PropertyOrContainer))
 		{
+		case EPropertyAccessibility::Inaccessible:
 			return nullptr;
+		case EPropertyAccessibility::ExplicitlyAccessible:
+			return DelegatePtr;
+		case EPropertyAccessibility::Unknown:
+			break;
 		}
 
 		if (bCanApplyFunction)
@@ -231,16 +270,7 @@ const FOnAnimatablePropertyChanged* FSequencerObjectChangeListener::FindProperty
 			{
 				return nullptr;
 			}
-		}
-
-		// Interp properties are always keyable
-		if (PropertyOrContainer->HasAnyPropertyFlags(CPF_Interp))
-		{
-			return DelegatePtr;
-		}
 
-		if (bCanApplyFunction)
-		{
 			// If there is a native setter we can always animate the property
 			if (Property.HasSetter())
 			{
@@ -298,6 +328,8 @@ bool FSequencerObjectChangeListener::CanKeyProperty(FCanKeyPropertyParams KeyPro
 
 bool FSequencerObjectChangeListener::CanKeyProperty_Internal(FCanKeyPropertyParams CanKeyPropertyParams, FOnAnimatablePropertyChanged& InOutDelegate, const FProperty*& InOutProperty, FPropertyPath& InOutPropertyPath) const
 {
+	using namespace UE::MovieScene;
+
 	if (CanKeyPropertyParams.PropertyPath.GetNumProperties() == 0)
 	{
 		return false;
@@ -364,14 +396,18 @@ bool FSequencerObjectChangeListener::CanKeyProperty_Internal(FCanKeyPropertyPara
 				return true;
 			}
 
-			for (const FFindAnimatablePropertyChangedHandler& FindHandler : FindPropertyChangedHandlers)
+
+			if (GetPropertyAccessibility(*Property) == EPropertyAccessibility::ExplicitlyAccessible)
 			{
-				FOnAnimatablePropertyChanged PropertyChanged = FindHandler.Execute(Property);
-				if (PropertyChanged.IsBound())
+				for (const FFindAnimatablePropertyChangedHandler& FindHandler : FindPropertyChangedHandlers)
 				{
-					InOutProperty = Property;
-					InOutDelegate = MoveTemp(PropertyChanged);
-					return true;
+					FOnAnimatablePropertyChanged PropertyChanged = FindHandler.Execute(Property);
+					if (PropertyChanged.IsBound())
+					{
+						InOutProperty = Property;
+						InOutDelegate = MoveTemp(PropertyChanged);
+						return true;
+					}
 				}
 			}
 
