- Disable auto play when going to start / end inside niagara editor by default

#tests Open a NS make sure it's playing hit the go to start button and observe it doesn't continue playing (you may need to adjust the editor setting if you modified it already)

diff --git a/Engine/Plugins/FX/Niagara/Source/NiagaraEditor/Private/NiagaraEditorModule.cpp b/Engine/Plugins/FX/Niagara/Source/NiagaraEditor/Private/NiagaraEditorModule.cpp
--- a/Engine/Plugins/FX/Niagara/Source/NiagaraEditor/Private/NiagaraEditorModule.cpp
+++ b/Engine/Plugins/FX/Niagara/Source/NiagaraEditor/Private/NiagaraEditorModule.cpp
@@ -2320,6 +2320,7 @@ void FNiagaraEditorModule::RegisterSettings()
 	if (SettingsModule != nullptr)
 	{
 		SequencerSettings = USequencerSettingsContainer::GetOrCreate<USequencerSettings>(TEXT("NiagaraSequenceEditor"));
+		SequencerSettings->SetStopPlayingWhenJumpingToStartOrEnd(true);
 
 		SettingsModule->RegisterSettings("Editor", "ContentEditors", "NiagaraSequenceEditor",
 			LOCTEXT("NiagaraSequenceEditorSettingsName", "Niagara Sequence Editor"),
diff --git a/Engine/Source/Editor/Sequencer/Public/SequencerSettings.h b/Engine/Source/Editor/Sequencer/Public/SequencerSettings.h
--- a/Engine/Source/Editor/Sequencer/Public/SequencerSettings.h
+++ b/Engine/Source/Editor/Sequencer/Public/SequencerSettings.h
@@ -268,10 +268,10 @@ public:
 	UE_API void SetScrubTimeStartFromCursor(bool bInValue);
 
 	/** Gets whether or not we stop playing when jumping to start or end. */
-	bool GetStopPlayingWhenJumpingToStartOrEnd() const;
+	UE_API bool GetStopPlayingWhenJumpingToStartOrEnd() const;
 
 	/** Sets whether or not we stop playing when jumping to start or end. */
-	void SetStopPlayingWhenJumpingToStartOrEnd(bool bInValue);
+	UE_API void SetStopPlayingWhenJumpingToStartOrEnd(bool bInValue);
 
 	/** Gets whether or not to link the curve editor time range. */
 	UE_API bool GetLinkCurveEditorTimeRange() const;
