Sequencer: Fix OnObjectPostEditChange so that it processes changes regardless of the EventChangeType

diff --git a/Engine/Source/Editor/Sequencer/Private/SequencerObjectChangeListener.cpp b/Engine/Source/Editor/Sequencer/Private/SequencerObjectChangeListener.cpp
--- a/Engine/Source/Editor/Sequencer/Private/SequencerObjectChangeListener.cpp
+++ b/Engine/Source/Editor/Sequencer/Private/SequencerObjectChangeListener.cpp
@@ -521,7 +521,7 @@ void FSequencerObjectChangeListener::OnObjectPreEditChange( UObject* Object, con
 
 void FSequencerObjectChangeListener::OnObjectPostEditChange( UObject* Object, FPropertyChangedEvent& PropertyChangedEvent )
 {
-	if( Object && PropertyChangedEvent.ChangeType != EPropertyChangeType::Interactive )
+	if( Object )
 	{
 		bool bIsObjectValid = IsObjectValidForListening( Object );
 
