Sequencer: Fix tracks in folders getting copied multiple times

#rnx

diff --git a/Engine/Source/Editor/Sequencer/Private/Sequencer.cpp b/Engine/Source/Editor/Sequencer/Private/Sequencer.cpp
--- a/Engine/Source/Editor/Sequencer/Private/Sequencer.cpp
+++ b/Engine/Source/Editor/Sequencer/Private/Sequencer.cpp
@@ -8476,7 +8476,7 @@ void FSequencer::CopySelectedTracks(TArray<TSharedPtr<FViewModel>>& TrackNodes,
 	{
 		bool bIsParentSelected = false;
 		TSharedPtr<FViewModel> ParentNode = TrackNode->GetParent();
-		while (ParentNode.IsValid() && !ParentNode->IsA<FFolderModel>())
+		while (ParentNode.IsValid())
 		{
 			if (ViewModel->GetSelection()->Outliner.IsSelected(CastViewModel<IOutlinerExtension>(ParentNode)))
 			{
