Sequencer: Fixed properties with custom accessors not being keyable with registered property track editors

- Also fixed keying of first property when adding a new object binding at the same time due to missing call to OnAddBinding

diff --git a/Engine/Source/Editor/Sequencer/Private/SequencerObjectChangeListener.cpp b/Engine/Source/Editor/Sequencer/Private/SequencerObjectChangeListener.cpp
--- a/Engine/Source/Editor/Sequencer/Private/SequencerObjectChangeListener.cpp
+++ b/Engine/Source/Editor/Sequencer/Private/SequencerObjectChangeListener.cpp
@@ -36,7 +36,11 @@ namespace UE::MovieScene
 		{
 			return EPropertyAccessibility::ExplicitlyAccessible;
 		}
-
+		// Properties with a setter function are always keyable
+		if (Property.HasSetter())
+		{
+			return EPropertyAccessibility::ExplicitlyAccessible;
+		}
 		return EPropertyAccessibility::Unknown;
 	}
 }
@@ -386,6 +390,16 @@ bool FSequencerObjectChangeListener::CanKeyProperty_Internal(FCanKeyPropertyPara
 					InOutDelegate = *DelegatePtr;
 					return true;
 				}
+				for (const FFindAnimatablePropertyChangedHandler& FindHandler : FindPropertyChangedHandlers)
+				{
+					FOnAnimatablePropertyChanged PropertyChanged = FindHandler.Execute(Property);
+					if (PropertyChanged.IsBound())
+					{
+						InOutProperty = Property;
+						InOutDelegate = MoveTemp(PropertyChanged);
+						return true;
+					}
+				}
 			}
 
 			// Otherwise we check for magic named functions using the default logic
diff --git a/Engine/Source/Editor/Sequencer/Private/SequencerUtilities.cpp b/Engine/Source/Editor/Sequencer/Private/SequencerUtilities.cpp
--- a/Engine/Source/Editor/Sequencer/Private/SequencerUtilities.cpp
+++ b/Engine/Source/Editor/Sequencer/Private/SequencerUtilities.cpp
@@ -3145,7 +3145,7 @@ FGuid CreateGenericBinding(TSharedPtr<ISequencer> Sequencer, UMovieSceneSequence
 		}
 
 		// If this is the last one
-		if (Index == 0)
+		if (Index == ObjectsToPossess.Num()-1)
 		{
 			if (Sequencer)
 			{
@@ -3160,7 +3160,7 @@ FGuid CreateGenericBinding(TSharedPtr<ISequencer> Sequencer, UMovieSceneSequence
 		}
 	}
 
-	// Should never get here - we should always hit the Index == 0 condition inside the loop
+	// Should never get here - we should always hit the Index == ObjectsToPossess.Num()-1 condition inside the loop
 	return FGuid();
 }
 
