Sequencer: Additional changes for allow returning multiple objects with FMovieSceneBindingResolveResult

ResolveBindingInternal now returns an array as well

Deprecated MovieSceneHelpers::GetSingleBoundObject in favor of GetBoundObjects
Deprecated ResolveSingleBinding in favor of ResolveSingleBinding with multiple resolved objects

diff --git a/Engine/Plugins/MovieScene/LevelSequenceEditor/Source/LevelSequenceEditor/Private/LevelSequenceEditorSubsystem.cpp b/Engine/Plugins/MovieScene/LevelSequenceEditor/Source/LevelSequenceEditor/Private/LevelSequenceEditorSubsystem.cpp
--- a/Engine/Plugins/MovieScene/LevelSequenceEditor/Source/LevelSequenceEditor/Private/LevelSequenceEditorSubsystem.cpp
+++ b/Engine/Plugins/MovieScene/LevelSequenceEditor/Source/LevelSequenceEditor/Private/LevelSequenceEditorSubsystem.cpp
@@ -270,7 +270,16 @@ void FMovieSceneBindingPropertyInfoListCustomization::CustomizeDetails(IDetailLa
 							bShowConvert = false;
 						}
 					}
-					bool bHasBoundObject = MovieSceneHelpers::GetSingleBoundObject(Sequence, BindingGuid, Sequencer->GetSharedPlaybackState(), 0) != nullptr;
+					bool bHasBoundObject = false;
+					TArray<UObject*> BoundObjects = MovieSceneHelpers::GetBoundObjects(Sequence, BindingGuid, Sequencer->GetSharedPlaybackState(), 0);
+					for (UObject* BoundObject : BoundObjects)
+					{
+						if (BoundObject)
+						{
+							bHasBoundObject = true;
+							break;
+						}
+					}
 
 					FDetailWidgetRow& BindingTypeRow = SectionCategory.AddCustomRow(FText::GetEmpty());
 					BindingTypeRow.NameContent()
diff --git a/Engine/Source/Editor/Sequencer/Private/SequencerUtilities.cpp b/Engine/Source/Editor/Sequencer/Private/SequencerUtilities.cpp
--- a/Engine/Source/Editor/Sequencer/Private/SequencerUtilities.cpp
+++ b/Engine/Source/Editor/Sequencer/Private/SequencerUtilities.cpp
@@ -1344,8 +1344,9 @@ FMovieScenePossessable* FSequencerUtilities::ConvertToPossessable(TSharedRef<ISe
 		}
 	}
 
-	UObject* BoundObject = MovieSceneHelpers::GetSingleBoundObject(Sequence, BindingGuid, Sequencer->GetSharedPlaybackState(), BindingIndex);
-	UObject* ObjectToConvert = BoundObject;
+	// For now, just get the first bound object to convert
+	TArray<UObject*> BoundObjects = MovieSceneHelpers::GetBoundObjects(Sequence, BindingGuid, Sequencer->GetSharedPlaybackState(), BindingIndex);
+	UObject* ObjectToConvert = BoundObjects.Num() > 0 ? BoundObjects[0] : nullptr;
 
 	// If we have an old-style spawnable, use the template as the object to convert instead.
 	bool bConvertFromSpawnable = MovieSceneHelpers::IsBoundToSpawnable(Sequence, BindingGuid, Sequencer->GetSharedPlaybackState(), BindingIndex);
@@ -1360,7 +1361,7 @@ FMovieScenePossessable* FSequencerUtilities::ConvertToPossessable(TSharedRef<ISe
 	FTransform DefaultTransform = SpawnableActorTemplate ? SpawnableActorTemplate->GetActorTransform() : FTransform();
 	// Prefer the transform at the current time over the spawnable actor template's transform because that's most likely 0. 
 	// This makes it so that the object will return to the current position on restore state.
-	AActor* Actor = Cast<AActor>(BoundObject);
+	AActor* Actor = Cast<AActor>(ObjectToConvert);
 	if (Actor)
 	{
 		if (Actor->GetRootComponent())
@@ -1397,7 +1398,7 @@ FMovieScenePossessable* FSequencerUtilities::ConvertToPossessable(TSharedRef<ISe
 		SpawnInfo.Template = SpawnableActorTemplate;
 
 		UWorld* World = GCurrentLevelEditingViewportClient ? GCurrentLevelEditingViewportClient->GetWorld() : nullptr;
-		AActor* PossessedActor = World->SpawnActor(ObjectToConvert->GetClass(), &DefaultTransform, SpawnInfo);
+		AActor* PossessedActor = World->SpawnActor(Actor->GetClass(), &DefaultTransform, SpawnInfo);
 
 		if (!PossessedActor)
 		{
@@ -1555,7 +1556,8 @@ FMovieScenePossessable* FSequencerUtilities::ConvertToCustomBinding(TSharedRef<I
 		return CreatedPossessable;
 	}
 
-	UObject* ObjectToConvert = MovieSceneHelpers::GetSingleBoundObject(Sequence, BindingGuid, Sequencer->GetSharedPlaybackState(), BindingIndex);
+	TArray<UObject*> BoundObjects = MovieSceneHelpers::GetBoundObjects(Sequence, BindingGuid, Sequencer->GetSharedPlaybackState(), BindingIndex);
+	UObject* ObjectToConvert = BoundObjects.Num() > 0 ? BoundObjects[0] : nullptr;
 
 	bool bConvertFromSpawnable = false;
 
@@ -2601,9 +2603,14 @@ void FSequencerUtilities::CopyBindings(TSharedRef<ISequencer> Sequencer, const T
 							}
 						}
 					}
-					else if (UObject* RuntimeObject = MovieSceneHelpers::GetSingleBoundObject(ObjectBinding.Sequence, ObjectBinding.BindingID, Sequencer->GetSharedPlaybackState(), BindingIndex))
+					else
 					{
-						CopyableBinding->BoundObjectNames.Add(RuntimeObject->GetPathName(World));
+						TArray<UObject*> BoundObjects = MovieSceneHelpers::GetBoundObjects(ObjectBinding.Sequence, ObjectBinding.BindingID, Sequencer->GetSharedPlaybackState(), BindingIndex);
+						if (BoundObjects.Num() > 0)
+						{
+							CopyableBinding->BoundObjectNames.Add(BoundObjects[0]->GetPathName(World));
+						}
+
 					}
 					BindingIndex++;
 				}
@@ -5017,13 +5024,18 @@ bool FSequencerUtilities::CanConvertToCustomBinding(TSharedRef<ISequencer> Seque
 			{
 				UE::UniversalObjectLocator::FResolveParams LocatorResolveParams(ResolutionContext);
 				FMovieSceneBindingResolveParams BindingResolveParams{ Sequence, BindingGuid, Sequencer->GetFocusedTemplateID(), ResolutionContext };
-				UObject* CurrentBoundObject = BindingReferences->ResolveSingleBinding(BindingResolveParams, BindingIndex, LocatorResolveParams, Sequencer->GetSharedPlaybackState());
-				if (CustomBindingType
-					&& (!CurrentBindingReference->CustomBinding
-						|| CurrentBindingReference->CustomBinding->GetClass() != CustomBindingType)
-					&& CustomBindingType->GetDefaultObject<UMovieSceneCustomBinding>()->SupportsConversionFromBinding(*CurrentBindingReference, CurrentBoundObject))
+				TArray<UObject*, TInlineAllocator<1>> CurrentBoundObjects;
+				BindingReferences->ResolveSingleBinding(BindingResolveParams, BindingIndex, LocatorResolveParams, Sequencer->GetSharedPlaybackState(), CurrentBoundObjects);
+
+				for (UObject* CurrentBoundObject : CurrentBoundObjects)
 				{
-					return true;
+					if (CustomBindingType
+						&& (!CurrentBindingReference->CustomBinding
+							|| CurrentBindingReference->CustomBinding->GetClass() != CustomBindingType)
+						&& CustomBindingType->GetDefaultObject<UMovieSceneCustomBinding>()->SupportsConversionFromBinding(*CurrentBindingReference, CurrentBoundObject))
+					{
+						return true;
+					}
 				}
 			}
 		}
diff --git a/Engine/Source/Runtime/MovieScene/Private/MovieSceneBindingReferences.cpp b/Engine/Source/Runtime/MovieScene/Private/MovieSceneBindingReferences.cpp
--- a/Engine/Source/Runtime/MovieScene/Private/MovieSceneBindingReferences.cpp
+++ b/Engine/Source/Runtime/MovieScene/Private/MovieSceneBindingReferences.cpp
@@ -174,25 +174,26 @@ UObject* FMovieSceneBindingReferences::ResolveBindingFromLocator(int32 Index, co
 	return UE::MovieScene::FindBoundObjectProxy(ResolvedObject);
 }
 
-UObject* FMovieSceneBindingReferences::ResolveBindingInternal(const FMovieSceneBindingResolveParams& BindingResolveParams, const UE::UniversalObjectLocator::FResolveParams& LocatorResolveParams, int32 BindingIndex, int32 InternalIndex, TSharedPtr<const UE::MovieScene::FSharedPlaybackState> SharedPlaybackState) const
+void FMovieSceneBindingReferences::ResolveBindingInternal(const FMovieSceneBindingResolveParams& BindingResolveParams, const UE::UniversalObjectLocator::FResolveParams& LocatorResolveParams, int32 BindingIndex, int32 InternalIndex, TSharedPtr<const UE::MovieScene::FSharedPlaybackState> SharedPlaybackState, TArray<UObject*, TInlineAllocator<1>>& OutObjects) const
 {
 	// If a custom binding is present and we have valid shared playback state, resolve the custom binding
 	if (SortedReferences[InternalIndex].CustomBinding && SharedPlaybackState.IsValid())
 	{
+		TArray<UObject*> ResolvedObjects;
+
 		FMovieSceneBindingResolveResult Result = SortedReferences[InternalIndex].CustomBinding->ResolveBinding(BindingResolveParams, BindingIndex, SharedPlaybackState.ToSharedRef());
 		for (TObjectPtr<UObject> ResolvedObject : Result.Objects)
 		{
 			if (ResolvedObject.Get())
 			{
-				return UE::MovieScene::FindBoundObjectProxy(ResolvedObject.Get());
+				OutObjects.Add(UE::MovieScene::FindBoundObjectProxy(ResolvedObject.Get()));
 			}
 		}
-		return nullptr;
 	}
 	else
 	{
 		// Otherwise, attempt to resolve via the locator
-		return ResolveBindingFromLocator(InternalIndex, LocatorResolveParams);
+		OutObjects.Add(ResolveBindingFromLocator(InternalIndex, LocatorResolveParams));
 	}
 }
 
@@ -227,20 +228,26 @@ void FMovieSceneBindingReferences::ResolveBinding(const FMovieSceneBindingResolv
 
 	for (int32 Index = StartIndex; Index < Num && SortedReferences[Index].ID == BindingResolveParams.ObjectBindingID; ++Index)
 	{
-		UObject* ResolvedObject = ResolveBindingInternal(BindingResolveParams, LocatorResolveParams, Index - StartIndex, Index, SharedPlaybackState);
-		if (ResolvedObject)
-		{
-			OutObjects.Add(ResolvedObject);
-		}
+		ResolveBindingInternal(BindingResolveParams, LocatorResolveParams, Index - StartIndex, Index, SharedPlaybackState, OutObjects);
 	}
 }
 
-UObject* FMovieSceneBindingReferences::ResolveSingleBinding(const FMovieSceneBindingResolveParams& BindingResolveParams, int32 BindingIndex, const UE::UniversalObjectLocator::FResolveParams& LocatorResolveParams, TSharedPtr<const UE::MovieScene::FSharedPlaybackState> SharedPlaybackState) const
+void FMovieSceneBindingReferences::ResolveSingleBinding(const FMovieSceneBindingResolveParams& BindingResolveParams, int32 BindingIndex, const UE::UniversalObjectLocator::FResolveParams& LocatorResolveParams, TSharedPtr<const UE::MovieScene::FSharedPlaybackState> SharedPlaybackState, TArray<UObject*, TInlineAllocator<1>>& OutObjects) const
 {
 	const int32 Index = Algo::LowerBoundBy(SortedReferences, BindingResolveParams.ObjectBindingID, &FMovieSceneBindingReference::ID) + BindingIndex;
 	if (SortedReferences.IsValidIndex(Index) && SortedReferences[Index].ID == BindingResolveParams.ObjectBindingID)
 	{
-		return ResolveBindingInternal(BindingResolveParams, LocatorResolveParams, BindingIndex, Index, SharedPlaybackState);
+		ResolveBindingInternal(BindingResolveParams, LocatorResolveParams, BindingIndex, Index, SharedPlaybackState, OutObjects);
+	}
+}
+
+UObject* FMovieSceneBindingReferences::ResolveSingleBinding(const FMovieSceneBindingResolveParams& BindingResolveParams, int32 BindingIndex, const UE::UniversalObjectLocator::FResolveParams& LocatorResolveParams, TSharedPtr<const UE::MovieScene::FSharedPlaybackState> SharedPlaybackState) const
+{
+	TArray<UObject*, TInlineAllocator<1>> ResolvedObjects;
+	ResolveSingleBinding(BindingResolveParams, BindingIndex, LocatorResolveParams, SharedPlaybackState, ResolvedObjects);
+	if (ResolvedObjects.Num() > 0)
+	{
+		return ResolvedObjects[0];
 	}
 	return nullptr;
 }
diff --git a/Engine/Source/Runtime/MovieScene/Private/MovieSceneCommonHelpers.cpp b/Engine/Source/Runtime/MovieScene/Private/MovieSceneCommonHelpers.cpp
--- a/Engine/Source/Runtime/MovieScene/Private/MovieSceneCommonHelpers.cpp
+++ b/Engine/Source/Runtime/MovieScene/Private/MovieSceneCommonHelpers.cpp
@@ -838,9 +838,21 @@ FGuid MovieSceneHelpers::TryCreateCustomSpawnableBinding(UMovieSceneSequence* Se
 	return NewID;
 }
 
-
 UObject* MovieSceneHelpers::GetSingleBoundObject(UMovieSceneSequence* Sequence, const FGuid& ObjectId, TSharedRef<const UE::MovieScene::FSharedPlaybackState> SharedPlaybackState, int32 BindingIndex)
 {
+	for (UObject* BoundObject : GetBoundObjects(Sequence, ObjectId, SharedPlaybackState, BindingIndex))
+	{
+		if (BoundObject)
+		{
+			return BoundObject;
+		}
+	}
+	return nullptr;
+}
+
+TArray<UObject*> MovieSceneHelpers::GetBoundObjects(UMovieSceneSequence* Sequence, const FGuid& ObjectId, TSharedRef<const UE::MovieScene::FSharedPlaybackState> SharedPlaybackState, int32 BindingIndex)
+{
+	TArray<UObject*> BoundObjects;
 	if (Sequence)
 	{
 		if (UMovieScene* MovieScene = Sequence->GetMovieScene())
@@ -851,10 +863,12 @@ UObject* MovieSceneHelpers::GetSingleBoundObject(UMovieSceneSequence* Sequence,
 
 				if (MovieScene->FindSpawnable(ObjectId))
 				{
-					TArrayView<TWeakObjectPtr<>> BoundObjects = EvaluationState->FindBoundObjects(FMovieSceneEvaluationOperand(SequenceID, ObjectId), SharedPlaybackState);
-					if (BoundObjects.Num() > 0)
+					for (TWeakObjectPtr<> BoundObject : EvaluationState->FindBoundObjects(FMovieSceneEvaluationOperand(SequenceID, ObjectId), SharedPlaybackState))
 					{
-						return BoundObjects[0].Get();
+						if (BoundObject.Get())
+						{
+							BoundObjects.Add(BoundObject.Get());
+						}
 					}
 				}
 				else if (const FMovieScenePossessable* Possessable = MovieScene->FindPossessable(ObjectId))
@@ -877,15 +891,21 @@ UObject* MovieSceneHelpers::GetSingleBoundObject(UMovieSceneSequence* Sequence,
 							}
 						}
 
+						TArray<UObject*, TInlineAllocator<1>> SingleBoundObjects;
 						UE::UniversalObjectLocator::FResolveParams LocatorResolveParams(ResolutionContext);
 						FMovieSceneBindingResolveParams BindingResolveParams{ Sequence, ObjectId, SequenceID, ResolutionContext };
-						return Refs->ResolveSingleBinding(BindingResolveParams, BindingIndex, LocatorResolveParams, SharedPlaybackState);
+						Refs->ResolveSingleBinding(BindingResolveParams, BindingIndex, LocatorResolveParams, SharedPlaybackState, SingleBoundObjects);
+
+						for (UObject* BoundObject : SingleBoundObjects)
+						{
+							BoundObjects.Add(BoundObject);
+						}
 					}
 				}
 			}
 		}
 	}
-	return nullptr;
+	return BoundObjects;
 }
 
 UObject* MovieSceneHelpers::GetObjectTemplate(UMovieSceneSequence* Sequence, const FGuid& ObjectId, TSharedRef<const UE::MovieScene::FSharedPlaybackState> SharedPlaybackState, int32 BindingIndex)
diff --git a/Engine/Source/Runtime/MovieScene/Public/MovieSceneBindingReferences.h b/Engine/Source/Runtime/MovieScene/Public/MovieSceneBindingReferences.h
--- a/Engine/Source/Runtime/MovieScene/Public/MovieSceneBindingReferences.h
+++ b/Engine/Source/Runtime/MovieScene/Public/MovieSceneBindingReferences.h
@@ -192,8 +192,11 @@ struct FMovieSceneBindingReferences
 	 * @param BindingResolveParams		Resolve parameters
 	 * @param LocatorResolveParams      Locator-specific resolve parameters
 	 * @param SharedPlaybackState       Shared Playback State from the Sequencer
-	 * @return The single bound objects if it exists
+	 * @param OutObjects				Array to populate with resolved objectss
 	 * */
+	MOVIESCENE_API void ResolveSingleBinding(const FMovieSceneBindingResolveParams& BindingResolveParams, int32 BindingIndex, const UE::UniversalObjectLocator::FResolveParams& LocatorResolveParams, TSharedPtr<const UE::MovieScene::FSharedPlaybackState> SharedPlaybackState, TArray<UObject*, TInlineAllocator<1>>& OutObjects) const;
+
+	UE_DEPRECATED(5.7, "This method has been deprecated. Please use ResolveSingleBinding which returns an array of bound objects.")
 	MOVIESCENE_API UObject* ResolveSingleBinding(const FMovieSceneBindingResolveParams& BindingResolveParams, int32 BindingIndex, const UE::UniversalObjectLocator::FResolveParams& LocatorResolveParams, TSharedPtr<const UE::MovieScene::FSharedPlaybackState> SharedPlaybackState) const;
 
 	/**
@@ -225,7 +228,7 @@ struct FMovieSceneBindingReferences
 
 private:
 
-	UObject* ResolveBindingInternal(const FMovieSceneBindingResolveParams& BindingResolveParams, const UE::UniversalObjectLocator::FResolveParams& LocatorResolveParams, int32 BindingIndex, int32 InternalIndex, TSharedPtr<const UE::MovieScene::FSharedPlaybackState> SharedPlaybackState) const;
+	void ResolveBindingInternal(const FMovieSceneBindingResolveParams& BindingResolveParams, const UE::UniversalObjectLocator::FResolveParams& LocatorResolveParams, int32 BindingIndex, int32 InternalIndex, TSharedPtr<const UE::MovieScene::FSharedPlaybackState> SharedPlaybackState, TArray<UObject*, TInlineAllocator<1>>& OutObjects) const;
 	UObject* ResolveBindingFromLocator(int32 Index, const UE::UniversalObjectLocator::FResolveParams& LocatorResolveParams) const;
 
 	/** The map from object binding ID to an array of references that pertain to that ID */
diff --git a/Engine/Source/Runtime/MovieScene/Public/MovieSceneCommonHelpers.h b/Engine/Source/Runtime/MovieScene/Public/MovieSceneCommonHelpers.h
--- a/Engine/Source/Runtime/MovieScene/Public/MovieSceneCommonHelpers.h
+++ b/Engine/Source/Runtime/MovieScene/Public/MovieSceneCommonHelpers.h
@@ -247,8 +247,11 @@ public:
 	
 
 	/*
-	* Returns the single bound object currently bound to the given objectid and binding index (optional).
+	* Returns the objects currently bound to the given objectid and binding index (optional).
 	*/
+	static MOVIESCENE_API TArray<UObject*> GetBoundObjects(UMovieSceneSequence* Sequence, const FGuid& ObjectId, TSharedRef<const UE::MovieScene::FSharedPlaybackState> SharedPlaybackState, int32 BindingIndex = 0);
+
+	UE_DEPRECATED(5.7, "This method has been deprecated. Please use GetBoundObjects which returns an array of bound objects.")
 	static MOVIESCENE_API UObject* GetSingleBoundObject(UMovieSceneSequence* Sequence, const FGuid& ObjectId, TSharedRef<const UE::MovieScene::FSharedPlaybackState> SharedPlaybackState, int32 BindingIndex = 0);
 
 	/*
