Sequencer: View options in toolbar now used UToolMenus for extension. You can use "Sequencer.MainToolBar.ViewOptions" to extend the menu.

diff --git a/Engine/Plugins/VirtualProduction/Takes/Source/TakeRecorder/Private/Timecode/Visualization/HitchViewOptionUtils.cpp b/Engine/Plugins/VirtualProduction/Takes/Source/TakeRecorder/Private/Timecode/Visualization/HitchViewOptionUtils.cpp
--- a/Engine/Plugins/VirtualProduction/Takes/Source/TakeRecorder/Private/Timecode/Visualization/HitchViewOptionUtils.cpp
+++ b/Engine/Plugins/VirtualProduction/Takes/Source/TakeRecorder/Private/Timecode/Visualization/HitchViewOptionUtils.cpp
@@ -4,7 +4,7 @@
 #include "HitchViewOptionUtils.h"
 
 #include "Framework/MultiBox/MultiBoxBuilder.h"
-#include "Framework/MultiBox/MultiBoxExtender.h"
+#include "HitchAnalysis.h"
 #include "Hitching/FrameHitchSceneDecoration.h"
 #include "ISequencer.h"
 #include "MovieScene.h"
@@ -14,7 +14,6 @@
 #include "TakeRecorderCommands.h"
 #include "TakeRecorderHitchVisualizationSettings.h"
 #include "Templates/SharedPointer.h"
-#include "Timecode/Visualization/Sequencer/HitchViewModel_AnalyzedData.h"
 
 #define LOCTEXT_NAMESPACE "FHitchVisualizerInSequencer"
 
@@ -28,61 +27,66 @@ static FText GetRecordingIntegrityText()
 }
 }
 	
-TSharedRef<FExtender> ExtendSequencerViewOptionsWithHitchVisualization()
+void ExtendSequencerViewOptionsWithHitchVisualization()
 {
-	const TSharedRef<FExtender> Extender = MakeShared<FExtender>();
-	Extender->AddMenuExtension(TEXT("MarkedFrames"), EExtensionHook::After, nullptr,
-		FMenuExtensionDelegate::CreateLambda([](FMenuBuilder& MenuBuilder)
+	UToolMenu* ToolMenu = UToolMenus::Get()->ExtendMenu("Sequencer.MainToolBar.ViewOptions");
+	ToolMenu->AddDynamicSection("HitchVisualization", FNewToolMenuDelegate::CreateLambda([](UToolMenu* ToolMenu)
+	{
+		USequencerToolMenuContext* Context = ToolMenu->Context.FindContext<USequencerToolMenuContext>();
+		const TSharedPtr<ISequencer> SequencerPin = ensure(Context) ? Context->WeakSequencer.Pin() : nullptr;
+		if (!SequencerPin || !HasHitchData(*SequencerPin))
 		{
-			MenuBuilder.BeginSection(TEXT("HitchVisualization"), MenuExtensionDetail::GetRecordingIntegrityText());
-				
-			MenuBuilder.AddMenuEntry(
-				LOCTEXT("SkipMarkers.Title", "Show Frame Drop Markers"),
-				LOCTEXT("SkipMarkers.Description", "Whether to place markers on frames where a timecode frame was skipped."),
-				FSlateIcon(),
-				FUIAction(
-					FExecuteAction::CreateLambda([]{ UTakeRecorderHitchVisualizationSettings::Get()->ToggleShowFrameDropMarkers(); }),
-					FCanExecuteAction(),
-					FIsActionChecked::CreateLambda([]{ return UTakeRecorderHitchVisualizationSettings::Get()->GetShowFrameDropMarkers(); })
-					),
-					NAME_None, EUserInterfaceActionType::ToggleButton
-				);
+			return;
+		}
 
-			MenuBuilder.AddMenuEntry(
-				LOCTEXT("RepeatMarkers.Title", "Show Frame Repeat Markers"),
-				LOCTEXT("RepeatMarkers.Description", "Whether to place markers on frames where a timecode frame was repeated."),
-				FSlateIcon(),
-				FUIAction(
-					FExecuteAction::CreateLambda([]{ UTakeRecorderHitchVisualizationSettings::Get()->ToggleShowFrameRepeatMarkers(); }),
-					FCanExecuteAction(),
-					FIsActionChecked::CreateLambda([]{ return UTakeRecorderHitchVisualizationSettings::Get()->GetShowFrameRepeatMarkers(); })
-					),
-					NAME_None, EUserInterfaceActionType::ToggleButton
-				);
+		FToolMenuSection& Section = ToolMenu->AddSection(TEXT("HitchVisualization"), MenuExtensionDetail::GetRecordingIntegrityText());
+			
+		Section.AddMenuEntry(
+			"SkipMarkers",
+			LOCTEXT("SkipMarkers.Title", "Show Frame Drop Markers"),
+			LOCTEXT("SkipMarkers.Description", "Whether to place markers on frames where a timecode frame was skipped."),
+			FSlateIcon(),
+			FUIAction(
+				FExecuteAction::CreateLambda([]{ UTakeRecorderHitchVisualizationSettings::Get()->ToggleShowFrameDropMarkers(); }),
+				FCanExecuteAction(),
+				FIsActionChecked::CreateLambda([]{ return UTakeRecorderHitchVisualizationSettings::Get()->GetShowFrameDropMarkers(); })
+				),
+			EUserInterfaceActionType::ToggleButton
+		);
 
-			MenuBuilder.AddMenuEntry(
-				LOCTEXT("ShowCatchupRanges.Title", "Show Hitch Recovery Ranges"),
-				LOCTEXT("ShowCatchupRanges.Description", "Whether to show areas in which the engine could not keep up, i.e. was running behind."),
-				FSlateIcon(),
-				FUIAction(
-					FExecuteAction::CreateLambda([]{ UTakeRecorderHitchVisualizationSettings::Get()->ToggleShowCatchupRanges(); }),
-					FCanExecuteAction(),
-					FIsActionChecked::CreateLambda([]{ return UTakeRecorderHitchVisualizationSettings::Get()->GetShowCatchupRanges(); })
-					),
-					NAME_None, EUserInterfaceActionType::ToggleButton
-				);
+		Section.AddMenuEntry(
+			"RepeatMarkers",
+			LOCTEXT("RepeatMarkers.Title", "Show Frame Repeat Markers"),
+			LOCTEXT("RepeatMarkers.Description", "Whether to place markers on frames where a timecode frame was repeated."),
+			FSlateIcon(),
+			FUIAction(
+				FExecuteAction::CreateLambda([]{ UTakeRecorderHitchVisualizationSettings::Get()->ToggleShowFrameRepeatMarkers(); }),
+				FCanExecuteAction(),
+				FIsActionChecked::CreateLambda([]{ return UTakeRecorderHitchVisualizationSettings::Get()->GetShowFrameRepeatMarkers(); })
+				),
+			EUserInterfaceActionType::ToggleButton
+		);
 
-			MenuBuilder.AddMenuEntry(FTakeRecorderCommands::Get().ClearRecordingIntegrityData);
+		Section.AddMenuEntry(
+			"ShowCatchupRanges",
+			LOCTEXT("ShowCatchupRanges.Title", "Show Hitch Recovery Ranges"),
+			LOCTEXT("ShowCatchupRanges.Description", "Whether to show areas in which the engine could not keep up, i.e. was running behind."),
+			FSlateIcon(),
+			FUIAction(
+				FExecuteAction::CreateLambda([]{ UTakeRecorderHitchVisualizationSettings::Get()->ToggleShowCatchupRanges(); }),
+				FCanExecuteAction(),
+				FIsActionChecked::CreateLambda([]{ return UTakeRecorderHitchVisualizationSettings::Get()->GetShowCatchupRanges(); })
+				),
+				EUserInterfaceActionType::ToggleButton
+		);
 
-			MenuBuilder.EndSection();
-		}));
-	return Extender;
+		Section.AddMenuEntry("ClearRecordingIntegrityData", FTakeRecorderCommands::Get().ClearRecordingIntegrityData);
+	}), FToolMenuInsert(NAME_None, EToolMenuInsertType::Last));
 }
 
 void ExtendSequencerTimeSliderContextMenu()
 {
 	UToolMenu* ToolMenu = UToolMenus::Get()->ExtendMenu("Sequencer.TimelineMenu");
-
 	ToolMenu->AddDynamicSection("HitchVisualization", FNewToolMenuDelegate::CreateLambda([](UToolMenu* InMenu)
 	{
 		const USequencerToolMenuContext* SequencerContext = InMenu->FindContext<USequencerToolMenuContext>();
@@ -91,7 +95,6 @@ void ExtendSequencerTimeSliderContextMenu()
 		{
 			return;
 		}
-
 		FToolMenuSection& Section = InMenu->AddSection(TEXT("HitchVisualization"), MenuExtensionDetail::GetRecordingIntegrityText());
 		if (HasHitchData(*SequencerPin))
 		{
diff --git a/Engine/Plugins/VirtualProduction/Takes/Source/TakeRecorder/Private/Timecode/Visualization/HitchViewOptionUtils.h b/Engine/Plugins/VirtualProduction/Takes/Source/TakeRecorder/Private/Timecode/Visualization/HitchViewOptionUtils.h
--- a/Engine/Plugins/VirtualProduction/Takes/Source/TakeRecorder/Private/Timecode/Visualization/HitchViewOptionUtils.h
+++ b/Engine/Plugins/VirtualProduction/Takes/Source/TakeRecorder/Private/Timecode/Visualization/HitchViewOptionUtils.h
@@ -14,13 +14,17 @@ namespace UE::TakeRecorder
 class FHitchViewModel_AnalyzedData;
 	
 /**
- * Builds an extender for the Sequencer view options adding:
+ * Extends Sequencer view options adding:
  * - Show Frame Drop Markers
  * - Show Catchup Ranges
  * - Clear Recording Integrity Markers and Ranges 
  */
-TSharedRef<FExtender> ExtendSequencerViewOptionsWithHitchVisualization();
+void ExtendSequencerViewOptionsWithHitchVisualization();
 
+/**
+ * Extends the top-time slider context menu with:
+ * - Clear Recording Integrity Markers and Ranges
+ */
 void ExtendSequencerTimeSliderContextMenu();
 
 /** Handles the FTakeRecorderCommands::Get().ClearRecordingIntegrityData. */
diff --git a/Engine/Plugins/VirtualProduction/Takes/Source/TakeRecorder/Private/Timecode/Visualization/Sequencer/SequencerHitchVisualizer.cpp b/Engine/Plugins/VirtualProduction/Takes/Source/TakeRecorder/Private/Timecode/Visualization/Sequencer/SequencerHitchVisualizer.cpp
--- a/Engine/Plugins/VirtualProduction/Takes/Source/TakeRecorder/Private/Timecode/Visualization/Sequencer/SequencerHitchVisualizer.cpp
+++ b/Engine/Plugins/VirtualProduction/Takes/Source/TakeRecorder/Private/Timecode/Visualization/Sequencer/SequencerHitchVisualizer.cpp
@@ -35,14 +35,13 @@ static void UnmapInjectedCommands(const TArray<TWeakPtr<ISequencer>>& InCustomiz
 }
 	
 FSequencerHitchVisualizer::FSequencerHitchVisualizer()
-	: ViewOptionsExtender(TakeRecorder::ExtendSequencerViewOptionsWithHitchVisualization())
 {
 	ISequencerModule& SequencerModule = FModuleManager::Get().LoadModuleChecked<ISequencerModule>("Sequencer");
 	OnSequencerCreatedHandle = SequencerModule.RegisterOnSequencerCreated(
 		FOnSequencerCreated::FDelegate::CreateRaw(this, &FSequencerHitchVisualizer::OnSequencerCreated)
 		);
-	SequencerModule.GetViewMenuExtensibilityManager()->AddExtender(ViewOptionsExtender);
 
+	ExtendSequencerViewOptionsWithHitchVisualization();
 	ExtendSequencerTimeSliderContextMenu();
 }
 
@@ -51,8 +50,6 @@ FSequencerHitchVisualizer::~FSequencerHitchVisualizer()
 	if (ISequencerModule* SequencerModule = FModuleManager::Get().GetModulePtr<ISequencerModule>("Sequencer"))
 	{
 		SequencerModule->UnregisterOnSequencerCreated(OnSequencerCreatedHandle);
-		SequencerModule->GetViewMenuExtensibilityManager()->RemoveExtender(ViewOptionsExtender);
-
 		Private::UnmapInjectedCommands(CustomizedSequencers);
 	}
 }
diff --git a/Engine/Source/Editor/Sequencer/Private/SSequencer.cpp b/Engine/Source/Editor/Sequencer/Private/SSequencer.cpp
--- a/Engine/Source/Editor/Sequencer/Private/SSequencer.cpp
+++ b/Engine/Source/Editor/Sequencer/Private/SSequencer.cpp
@@ -1461,15 +1461,22 @@ TSharedRef<SWidget> SSequencer::ConstructSearchAndFilterRow()
 		];
 }
 
+namespace UE::Sequencer::ToolMenus
+{
+const FLazyName ToolbarMenuName("Sequencer.MainToolBar");
+const FLazyName ViewOptionsMenuName("Sequencer.MainToolBar.ViewOptions");
+}
+
 TSharedRef<SWidget> SSequencer::MakeToolBar()
 {
-	static const FName MenuName("Sequencer.MainToolBar");
+	using namespace UE::Sequencer::ToolMenus;
 	
-	if (!UToolMenus::Get()->IsMenuRegistered(MenuName))
+	if (!UToolMenus::Get()->IsMenuRegistered(ToolbarMenuName))
 	{
-		UToolMenu* Toolbar = UToolMenus::Get()->RegisterMenu(MenuName, NAME_None, EMultiBoxType::ToolBar);
+		UToolMenu* Toolbar = UToolMenus::Get()->RegisterMenu(ToolbarMenuName, NAME_None, EMultiBoxType::ToolBar);
 		Toolbar->AddDynamicSection("PopulateToolBar", FNewToolMenuDelegate::CreateStatic(&SSequencer::PopulateToolBar));
 	}
+	RegisterViewOptionsToolMenu();
 
 	TArray<TSharedPtr<FExtender>> AllExtenders;
 	ISequencerModule& SequencerModule = FModuleManager::GetModuleChecked<ISequencerModule>("Sequencer");
@@ -1488,7 +1495,7 @@ TSharedRef<SWidget> SSequencer::MakeToolBar()
 	// Allow any toolkits to initialize their menu context
 	OnInitToolMenuContext.ExecuteIfBound(Context);
 
-	return UToolMenus::Get()->GenerateWidget(MenuName, Context);
+	return UToolMenus::Get()->GenerateWidget(ToolbarMenuName, Context);
 }
 
 void SSequencer::PopulateToolBar(UToolMenu* InMenu)
@@ -1600,7 +1607,7 @@ void SSequencer::PopulateToolBar(UToolMenu* InMenu)
 		FToolMenuEntry ViewOptionsEntry = FToolMenuEntry::InitComboButton(
 			"ViewOptions",
 			FUIAction(),
-			FOnGetContent::CreateSP(SequencerWidget, &SSequencer::MakeViewMenu),
+			FNewToolMenuWidget::CreateSP(SequencerWidget, &SSequencer::MakeViewMenu),
 			LOCTEXT("ViewOptions", "View Options"),
 			LOCTEXT("ViewOptionsToolTip", "View Options"),
 			FSlateIcon(FAppStyle::Get().GetStyleSetName(), "Icons.Visibility")
@@ -1693,6 +1700,26 @@ void SSequencer::PopulateToolBar(UToolMenu* InMenu)
 	}
 }
 
+void SSequencer::RegisterViewOptionsToolMenu()
+{
+	using namespace UE::Sequencer::ToolMenus;
+	
+	if (!UToolMenus::Get()->IsMenuRegistered(ViewOptionsMenuName))
+	{
+		UToolMenu* ViewOptions = UToolMenus::Get()->RegisterMenu(ViewOptionsMenuName, NAME_None, EMultiBoxType::Menu);
+		ViewOptions->AddDynamicSection("Sequencer", FNewToolMenuDelegate::CreateLambda([](UToolMenu* Menu)
+		{
+			USequencerToolMenuContext* Context = Menu->FindContext<USequencerToolMenuContext>();
+			const TSharedPtr<ISequencer> SequencerPin = Context ? Context->WeakSequencer.Pin() : nullptr;
+			if (SequencerPin)
+			{
+				const TSharedRef<SSequencer> SequencerWidget = StaticCastSharedRef<SSequencer>(SequencerPin->GetSequencerWidget());
+				SequencerWidget->PopulateViewOptionMenu(Menu);
+			}
+		}));
+	}	
+}
+
 TSharedRef<SWidget> SSequencer::MakeAddMenu()
 {
 	using namespace UE::Sequencer;
@@ -1917,58 +1944,66 @@ void SSequencer::FillAdvancedMenu(FMenuBuilder& MenuBuilder)
 	MenuBuilder.EndSection();
 }
 
-TSharedRef<SWidget> SSequencer::MakeViewMenu()
+TSharedRef<SWidget> SSequencer::MakeViewMenu(const FToolMenuContext& InContext)
 {
-	TSharedPtr<FSequencer> Sequencer = SequencerPtr.Pin();
+	// This was created by the toolbar builder...
+	USequencerToolMenuContext* RootContextObject = InContext.FindContext<USequencerToolMenuContext>();
+
+	// ... and passing along RootContextObject allows for external modules to add view options that depend on state of ISequencer.
+	const TSharedPtr<FSequencer> Sequencer = SequencerPtr.Pin();
+	FToolMenuContext Context(Sequencer->GetCommandBindings());
+	Context.AddObject(RootContextObject);
 
+	// We want to support the legacy FExtender API for now
 	TArray<TSharedPtr<FExtender>> AllExtenders;
 	ISequencerModule& SequencerModule = FModuleManager::GetModuleChecked<ISequencerModule>("Sequencer");
 	AllExtenders.Add(SequencerModule.GetViewMenuExtensibilityManager()->GetAllExtenders());
 	AllExtenders.Append(ViewMenuExtenders);
+	Context.AddExtender(FExtender::Combine(AllExtenders));
+	
+	return UToolMenus::Get()->GenerateWidget(UE::Sequencer::ToolMenus::ViewOptionsMenuName, Context);
+}
 
-	TSharedPtr<FExtender> Extender = FExtender::Combine(AllExtenders);
-	FMenuBuilder MenuBuilder(true, Sequencer->GetCommandBindings(), Extender);
+void SSequencer::PopulateViewOptionMenu(UToolMenu* InToolMenu)
+{
+	TSharedPtr<FSequencer> Sequencer = SequencerPtr.Pin();
 	
 	if (Sequencer->GetHostCapabilities().bSupportsSidebar)
 	{
-		MenuBuilder.AddMenuEntry( FSequencerCommands::Get().ToggleSidebarVisible );
+		InToolMenu->AddMenuEntry("Sidebar", FToolMenuEntry::InitMenuEntry(FSequencerCommands::Get().ToggleSidebarVisible));
 	}
 
 	if (Sequencer->IsLevelEditorSequencer())
 	{
-		MenuBuilder.BeginSection("PilotCamera", LOCTEXT("PilotCamera", "Pilot Camera"));
-		{
-			MenuBuilder.AddMenuEntry(FSequencerCommands::Get().TogglePilotCamera);
-			MenuBuilder.AddMenuEntry(FSequencerCommands::Get().ToggleRestoreOriginalViewportOnCameraCutUnlock);
-			MenuBuilder.AddMenuEntry(FSequencerCommands::Get().TogglePreviewCameraCutsInSimulate);
-		}
-		MenuBuilder.EndSection();
+		FToolMenuSection& PilotCameraSection = InToolMenu->AddSection("PilotCamera", LOCTEXT("PilotCamera", "Pilot Camera"));
+		PilotCameraSection.AddMenuEntry(FSequencerCommands::Get().TogglePilotCamera);
+		PilotCameraSection.AddMenuEntry(FSequencerCommands::Get().ToggleRestoreOriginalViewportOnCameraCutUnlock);
+		PilotCameraSection.AddMenuEntry(FSequencerCommands::Get().TogglePreviewCameraCutsInSimulate);
 	}
 
-	MenuBuilder.BeginSection("SequencerSettings", LOCTEXT("SequencerSettings", "Sequencer Settings"));
+	FToolMenuSection& SequencerSettingsSection = InToolMenu->AddSection("SequencerSettings", LOCTEXT("SequencerSettings", "Sequencer Settings"));
 	{
-		MenuBuilder.AddMenuEntry( FSequencerCommands::Get().ToggleAutoScroll );
-		MenuBuilder.AddMenuEntry( FSequencerCommands::Get().ToggleShowRangeSlider );
-		MenuBuilder.AddMenuEntry( FSequencerCommands::Get().ToggleLayerBars );
-		MenuBuilder.AddMenuEntry( FSequencerCommands::Get().ToggleKeyBars );
-		MenuBuilder.AddMenuEntry( FSequencerCommands::Get().ToggleChannelColors );
-		MenuBuilder.AddMenuEntry( FSequencerCommands::Get().ToggleShowInfoButton );
-		MenuBuilder.AddMenuEntry( FSequencerCommands::Get().ToggleShowPreAndPostRoll );
+		SequencerSettingsSection.AddMenuEntry( FSequencerCommands::Get().ToggleAutoScroll );
+		SequencerSettingsSection.AddMenuEntry( FSequencerCommands::Get().ToggleShowRangeSlider );
+		SequencerSettingsSection.AddMenuEntry( FSequencerCommands::Get().ToggleLayerBars );
+		SequencerSettingsSection.AddMenuEntry( FSequencerCommands::Get().ToggleKeyBars );
+		SequencerSettingsSection.AddMenuEntry( FSequencerCommands::Get().ToggleChannelColors );
+		SequencerSettingsSection.AddMenuEntry( FSequencerCommands::Get().ToggleShowInfoButton );
+		SequencerSettingsSection.AddMenuEntry( FSequencerCommands::Get().ToggleShowPreAndPostRoll );
 
-		MenuBuilder.AddSubMenu(LOCTEXT("ViewDensityMenuLabel", "View Density"), FText::GetEmpty(), FNewMenuDelegate::CreateRaw(this, &SSequencer::FillViewDensityMenu));
+		SequencerSettingsSection.AddSubMenu("ViewDensityMenu", LOCTEXT("ViewDensityMenuLabel", "View Density"), FText::GetEmpty(), FNewMenuDelegate::CreateRaw(this, &SSequencer::FillViewDensityMenu));
 
 		// Menu Entry for Outliner Column Visibilities
 		if (OutlinerColumnVisibilities.Num() > 0)
 		{
-			MenuBuilder.AddSubMenu(LOCTEXT("ColumnVisibilityHeader", "Columns"), FText::GetEmpty(), FNewMenuDelegate::CreateRaw(this, &SSequencer::FillColumnVisibilityMenu));
+			SequencerSettingsSection.AddSubMenu("ColumnVisibility", LOCTEXT("ColumnVisibilityHeader", "Columns"), FText::GetEmpty(), FNewMenuDelegate::CreateRaw(this, &SSequencer::FillColumnVisibilityMenu));
 		}
 
 		// Menu entry for zero padding
-		auto OnZeroPadChanged = [this](uint8 NewValue) {
+		const auto OnZeroPadChanged = [this](uint8 NewValue) {
 			GetSequencerSettings()->SetZeroPadFrames(NewValue);
 		};
-
-		MenuBuilder.AddWidget(
+		const TSharedRef<SWidget> ZeroPaddingWidget = 
 			SNew(SHorizontalBox)
 			+ SHorizontalBox::Slot()
 				[
@@ -1986,20 +2021,18 @@ TSharedRef<SWidget> SSequencer::MakeViewMenu()
 					.Value_Lambda([this]() -> uint8 {
 						return GetSequencerSettings()->GetZeroPadFrames();
 					})
-				],
-			LOCTEXT("ZeroPaddingText", "Zero Pad Frame Numbers"));
+				];
+		SequencerSettingsSection.AddEntry(
+			FToolMenuEntry::InitWidget("ZeroPadding", ZeroPaddingWidget, LOCTEXT("ZeroPaddingText", "Zero Pad Frame Numbers"))
+			);
 	}
-	MenuBuilder.EndSection();
 
-	MenuBuilder.BeginSection("MarkedFrames", LOCTEXT("MarkedFramesHeader", "Marked Frames"));
+	FToolMenuSection& MarkedFramesSection = InToolMenu->AddSection("MarkedFrames", LOCTEXT("MarkedFramesHeader", "Marked Frames"));
 	{
-		MenuBuilder.AddMenuEntry(FSequencerCommands::Get().ToggleShowMarkedFrames);
-		MenuBuilder.AddMenuEntry(FSequencerCommands::Get().ToggleShowMarkedFramesGlobally);
-		MenuBuilder.AddMenuEntry(FSequencerCommands::Get().ClearGlobalMarkedFrames);
+		MarkedFramesSection.AddMenuEntry(FSequencerCommands::Get().ToggleShowMarkedFrames);
+		MarkedFramesSection.AddMenuEntry(FSequencerCommands::Get().ToggleShowMarkedFramesGlobally);
+		MarkedFramesSection.AddMenuEntry(FSequencerCommands::Get().ClearGlobalMarkedFrames);
 	}
-	MenuBuilder.EndSection();
-
-	return MenuBuilder.MakeWidget();
 }
 
 void SSequencer::OpenTaggedBindingManager()
diff --git a/Engine/Source/Editor/Sequencer/Private/SSequencer.h b/Engine/Source/Editor/Sequencer/Private/SSequencer.h
--- a/Engine/Source/Editor/Sequencer/Private/SSequencer.h
+++ b/Engine/Source/Editor/Sequencer/Private/SSequencer.h
@@ -400,8 +400,12 @@ private:
 
 	TWeakPtr<FSequencer> GetSequencer() { return SequencerPtr; }
 
+	/** Builds the base UToolMenu for Sequencer's toolbar. */
 	static void PopulateToolBar(UToolMenu* InMenu);
 
+	/** Registers the tool menu for Sequencer.MainToolBar.ViewOptions. */
+	static void RegisterViewOptionsToolMenu();
+
 	/** Makes the toolbar. */
 	TSharedRef<SWidget> MakeToolBar();
 
@@ -412,7 +416,10 @@ private:
 	TSharedRef<SWidget> MakeActionsMenu();
 
 	/** Makes the view menu for the toolbar. */
-	TSharedRef<SWidget> MakeViewMenu();
+	TSharedRef<SWidget> MakeViewMenu(const FToolMenuContext& InContext);
+	
+	/** Populates the view options menu with the default sequencer entries. */
+	void PopulateViewOptionMenu(UToolMenu* ToolMenu);
 
 	/** Makes the plabacky menu for the toolbar. */
 	TSharedRef<SWidget> MakePlaybackMenu();
