Sequencer: Fix not being able to create an Empty Binding

In the previous revision, CurrentObject can be null so that a possessable with an empty binding can be made

#rnx

diff --git a/Engine/Source/Editor/Sequencer/Private/SequencerUtilities.cpp b/Engine/Source/Editor/Sequencer/Private/SequencerUtilities.cpp
--- a/Engine/Source/Editor/Sequencer/Private/SequencerUtilities.cpp
+++ b/Engine/Source/Editor/Sequencer/Private/SequencerUtilities.cpp
@@ -3060,10 +3060,8 @@ FGuid CreateGenericBinding(TSharedPtr<ISequencer> Sequencer, UMovieSceneSequence
 		UObject*             CurrentObject = ObjectsToPossess[Index].Key;
 		const IObjectSchema* Schema        = ObjectsToPossess[Index].Value;
 
-		check(CurrentObject);
-
 		// If we're not purposefully replacing a binding, then check to see if we already have one, and use that
-		if (Sequencer && !InParams.ReplacementGuid.IsValid())
+		if (Sequencer && !InParams.ReplacementGuid.IsValid() && CurrentObject)
 		{
 			FGuid    ObjectGuid = Sequencer->FindCachedObjectId(*CurrentObject, Sequencer->GetFocusedTemplateID());
 
@@ -3092,17 +3090,18 @@ FGuid CreateGenericBinding(TSharedPtr<ISequencer> Sequencer, UMovieSceneSequence
 		{
 			FString CurrentName = (Index == ObjectsToPossess.Num()-1 && InParams.BindingNameOverride.Len() != 0)
 				? InParams.BindingNameOverride
-				: Schema->GetPrettyName(CurrentObject).ToString();
+				: CurrentObject ? Schema->GetPrettyName(CurrentObject).ToString() : TEXT("EmptyBinding");
+			UClass* CurrentClass = CurrentObject ? CurrentObject->GetClass() : UObject::StaticClass();
 
-			NewID = OwnerMovieScene->AddPossessable(CurrentName, CurrentObject->GetClass());
+			NewID = OwnerMovieScene->AddPossessable(CurrentName, CurrentClass);
 			NewPossessable = OwnerMovieScene->FindPossessable(NewID);
 		}
 
 		// If we're not trying to replace a binding, and the object is a spawnable, try and bind to that first
-		if (InParams.ReplacementGuid.IsValid() || (Sequencer && !NewPossessable->BindSpawnableObject(Sequencer->GetFocusedTemplateID(), CurrentObject, Sequencer->GetSharedPlaybackState())))
+		if (InParams.ReplacementGuid.IsValid() || (Sequencer && (!CurrentObject || !NewPossessable->BindSpawnableObject(Sequencer->GetFocusedTemplateID(), CurrentObject, Sequencer->GetSharedPlaybackState()))))
 		{
 			FUniversalObjectLocator Locator;
-			if (!OwnerSequence->MakeLocatorForObject(CurrentObject, Context, Locator) || Locator.IsEmpty())
+			if (CurrentObject && (!OwnerSequence->MakeLocatorForObject(CurrentObject, Context, Locator) || Locator.IsEmpty()))
 			{
 				// Unable to possess this object
 				return FGuid();
