ShotManagement: Fixing issue where Create Camera button would not show up for Cine Assemblies in the Sequencer editor.

#rnx

diff --git a/Engine/Source/Editor/Sequencer/Private/Sequencer.cpp b/Engine/Source/Editor/Sequencer/Private/Sequencer.cpp
--- a/Engine/Source/Editor/Sequencer/Private/Sequencer.cpp
+++ b/Engine/Source/Editor/Sequencer/Private/Sequencer.cpp
@@ -11738,7 +11738,7 @@ void FSequencer::BindCommands()
 		} ),
 		FCanExecuteAction(),
 		FIsActionChecked(),
-		FIsActionButtonVisible::CreateLambda([this] { return ExactCast<ULevelSequence>(GetFocusedMovieSceneSequence()) != nullptr && IVREditorModule::Get().IsVREditorModeActive() == false; }) //@todo VREditor: Creating a camera while in VR mode disrupts the hmd. This is a temporary fix by hiding the button when in VR mode.
+		FIsActionButtonVisible::CreateLambda([this] { return Cast<ULevelSequence>(GetFocusedMovieSceneSequence()) != nullptr; })
 	);
 
 	SequencerCommandBindings->MapAction(
