[Sequencer] Change Sequencer filter help dialog title from "Sequence Navigator" to "Sequencer"

diff --git a/Engine/Source/Editor/Sequencer/Private/Filters/SequencerFilterBar.cpp b/Engine/Source/Editor/Sequencer/Private/Filters/SequencerFilterBar.cpp
--- a/Engine/Source/Editor/Sequencer/Private/Filters/SequencerFilterBar.cpp
+++ b/Engine/Source/Editor/Sequencer/Private/Filters/SequencerFilterBar.cpp
@@ -1859,8 +1859,8 @@ bool FSequencerFilterBar::IsFilterSupported(const FString& InFilterName) const
 void FSequencerFilterBar::OpenTextExpressionHelp()
 {
 	FFilterExpressionHelpDialogConfig Config;
-	Config.IdentifierName = TEXT("NavigationToolCustomTextFilterHelp");
-	Config.DialogTitle = LOCTEXT("NavigationToolCustomTextFilterHelp", "Navigation Tool Custom Text Filter Help");
+	Config.IdentifierName = TEXT("SequencerCustomTextFilterHelp");
+	Config.DialogTitle = LOCTEXT("CustomTextFilterHelp", "Sequencer Custom Text Filter Help");
 	Config.TextFilterExpressionContexts = GetTextFilterExpressionContexts();
 	SFilterExpressionHelpDialog::Open(MoveTemp(Config));
 }
@@ -1868,14 +1868,14 @@ void FSequencerFilterBar::OpenTextExpressionHelp()
 void FSequencerFilterBar::SaveCurrentFilterSetAsCustomTextFilter()
 {
 	const FText NewFilterString = FText::FromString(GenerateTextFilterStringFromEnabledFilters());
-	// @TODO: update to the new way of handling this (similar to Navigation Tool)
+	// @TODO: update to the new way of handling this (similar to Sequence Navigator)
 	//CreateWindow_AddCustomTextFilter(DefaultNewCustomTextFilterData(NewFilterString));
 }
 
 void FSequencerFilterBar::CreateNewTextFilter()
 {
 	const FText NewFilterString = FText::FromString(GetTextFilterString());
-	// @TODO: update to the new way of handling this (similar to Navigation Tool)
+	// @TODO: update to the new way of handling this (similar to Sequence Navigator)
 	//CreateWindow_AddCustomTextFilter(DefaultNewCustomTextFilterData(NewFilterString));
 }
 
