Sequencer: Remove duplicate keys on end drag

This still allows you to move keys past one another since the deletion will happen on mouse release.
Nudging is umchanged, you can nudge keys on top of each other and duplicates will be removed when the keys are deselected.

diff --git a/Engine/Source/Editor/Sequencer/Private/MVVM/Selection.cpp b/Engine/Source/Editor/Sequencer/Private/MVVM/Selection.cpp
--- a/Engine/Source/Editor/Sequencer/Private/MVVM/Selection.cpp
+++ b/Engine/Source/Editor/Sequencer/Private/MVVM/Selection.cpp
@@ -12,57 +12,21 @@
 #include "Channels/MovieSceneChannel.h"
 #include "MovieSceneSection.h"
 #include "MovieSceneTrack.h"
+#include "SequencerCommonHelpers.h"
 
 namespace UE::Sequencer
 {
 
-void FKeySelection::RemoveDuplicateKeys(TArrayView<FKeyHandle> KeyHandles)
-{
-	TMap<TViewModelPtr<FChannelModel>, TArray<FKeyHandle> > ChannelToKeyHandles;
-
-	for (const FKeyHandle& KeyHandle : KeyHandles)
-	{
-		TViewModelPtr<FChannelModel> ChannelModel = GetModelForKey(KeyHandle);
-		if (ChannelModel)
-		{
-			if (FMovieSceneChannel* Channel = ChannelModel->GetChannel())
-			{
-				if (Channel->GetIndex(KeyHandle) == INDEX_NONE)
-				{
-					continue;
-				}
-
-				FFrameNumber KeyTime;
-				Channel->GetKeyTime(KeyHandle, KeyTime);
-
-				TArray<FFrameNumber> OutKeyTimes;
-				TArray<FKeyHandle> OutKeyHandles;
-				Channel->GetKeys(TRange<FFrameNumber>(KeyTime, KeyTime), &OutKeyTimes, &OutKeyHandles);
-				OutKeyHandles.Remove(KeyHandle);
-
-				ChannelToKeyHandles.FindOrAdd(ChannelModel).Append(OutKeyHandles);
-			}
-		}
-	}
-
-	for (TPair<TViewModelPtr<FChannelModel>, TArray<FKeyHandle> >& ChannelToKey : ChannelToKeyHandles)
-	{
-		TViewModelPtr<FChannelModel> ChannelModel = ChannelToKey.Key;
-		ChannelModel->GetSection()->Modify();
-		ChannelModel->GetChannel()->DeleteKeys(ChannelToKey.Value);
-	}
-}
-
 void FKeySelection::Deselect(FKeyHandle InKey)
 {
-	RemoveDuplicateKeys(MakeArrayView(&InKey, 1));
+	SequencerHelpers::RemoveDuplicateKeys(*this, MakeArrayView(&InKey, 1));
 
 	TUniqueFragmentSelectionSet<FKeyHandle, FChannelModel>::Deselect(InKey);
 }
 
 void FKeySelection::Empty()
 {
-	RemoveDuplicateKeys(MakeArrayView(GetSelected().Array()));
+	SequencerHelpers::RemoveDuplicateKeys(*this, MakeArrayView(GetSelected().Array()));
 
 	TUniqueFragmentSelectionSet<FKeyHandle, FChannelModel>::Empty();
 }
diff --git a/Engine/Source/Editor/Sequencer/Private/SequencerCommonHelpers.cpp b/Engine/Source/Editor/Sequencer/Private/SequencerCommonHelpers.cpp
--- a/Engine/Source/Editor/Sequencer/Private/SequencerCommonHelpers.cpp
+++ b/Engine/Source/Editor/Sequencer/Private/SequencerCommonHelpers.cpp
@@ -23,6 +23,7 @@
 #include "MVVM/ViewModels/VirtualTrackArea.h"
 #include "MVVM/Views/ITrackAreaHotspot.h"
 #include "MVVM/Selection/Selection.h"
+#include "MVVM/Selection/SequencerCoreSelectionTypes.h"
 #include "Modules/ModuleManager.h"
 #include "MovieScene.h"
 #include "MovieSceneSectionDetailsCustomization.h"
@@ -241,6 +242,45 @@ void SequencerHelpers::PerformDefaultSelection(FSequencer& Sequencer, const FPoi
 	}
 }
 
+void SequencerHelpers::RemoveDuplicateKeys(const UE::Sequencer::FKeySelection& KeySelection, TArrayView<FKeyHandle> KeyHandles)
+{
+	using namespace UE::Sequencer;
+
+	TMap<TViewModelPtr<FChannelModel>, TArray<FKeyHandle> > ChannelToKeyHandles;
+
+	for (const FKeyHandle& KeyHandle : KeyHandles)
+	{
+		TViewModelPtr<FChannelModel> ChannelModel = KeySelection.GetModelForKey(KeyHandle);
+		if (ChannelModel)
+		{
+			if (FMovieSceneChannel* Channel = ChannelModel->GetChannel())
+			{
+				if (Channel->GetIndex(KeyHandle) == INDEX_NONE)
+				{
+					continue;
+				}
+
+				FFrameNumber KeyTime;
+				Channel->GetKeyTime(KeyHandle, KeyTime);
+
+				TArray<FFrameNumber> OutKeyTimes;
+				TArray<FKeyHandle> OutKeyHandles;
+				Channel->GetKeys(TRange<FFrameNumber>(KeyTime, KeyTime), &OutKeyTimes, &OutKeyHandles);
+				OutKeyHandles.Remove(KeyHandle);
+
+				ChannelToKeyHandles.FindOrAdd(ChannelModel).Append(OutKeyHandles);
+			}
+		}
+	}
+
+	for (TPair<TViewModelPtr<FChannelModel>, TArray<FKeyHandle> >& ChannelToKey : ChannelToKeyHandles)
+	{
+		TViewModelPtr<FChannelModel> ChannelModel = ChannelToKey.Key;
+		ChannelModel->GetSection()->Modify();
+		ChannelModel->GetChannel()->DeleteKeys(ChannelToKey.Value);
+	}
+}
+
 TSharedPtr<SWidget> SequencerHelpers::SummonContextMenu(FSequencer& Sequencer, const FGeometry& MyGeometry, const FPointerEvent& MouseEvent)
 {
 	using namespace UE::Sequencer;
diff --git a/Engine/Source/Editor/Sequencer/Private/SequencerCommonHelpers.h b/Engine/Source/Editor/Sequencer/Private/SequencerCommonHelpers.h
--- a/Engine/Source/Editor/Sequencer/Private/SequencerCommonHelpers.h
+++ b/Engine/Source/Editor/Sequencer/Private/SequencerCommonHelpers.h
@@ -15,6 +15,7 @@ class FSequencer;
 class IKeyArea;
 class UMovieSceneSection;
 class UMovieSceneTrack;
+struct FKeyHandle;
 
 namespace UE
 {
@@ -22,6 +23,7 @@ namespace Sequencer
 {
 	class FChannelModel;
 	class IOutlinerExtension;
+	struct FKeySelection;
 	struct FViewModelVariantIterator;
 
 } // namespace Sequencer
@@ -64,6 +66,11 @@ public:
 	 */
 	static void PerformDefaultSelection(FSequencer& Sequencer, const FPointerEvent& MouseEvent);
 	
+	/*
+	 * Remove duplicate keys 
+	 */
+	static void RemoveDuplicateKeys(const UE::Sequencer::FKeySelection& KeySelection, TArrayView<FKeyHandle> KeyHandles);
+
 	/**
 	 * Attempt to summon a context menu for the current hotspot
 	 */
diff --git a/Engine/Source/Editor/Sequencer/Private/Tools/EditToolDragOperations.cpp b/Engine/Source/Editor/Sequencer/Private/Tools/EditToolDragOperations.cpp
--- a/Engine/Source/Editor/Sequencer/Private/Tools/EditToolDragOperations.cpp
+++ b/Engine/Source/Editor/Sequencer/Private/Tools/EditToolDragOperations.cpp
@@ -923,6 +923,14 @@ void FMoveKeysAndSections::OnEndDrag(const FPointerEvent& MouseEvent, FVector2D
 		return;
 	}
 
+	TArray<FKeyHandle> KeyHandles;
+	for (const FSequencerSelectedKey& Key : KeysAsArray)
+	{
+		KeyHandles.Add(Key.KeyHandle);
+	}
+	const FKeySelection& KeySelection = Sequencer.GetViewModel()->GetSelection()->KeySelection;
+	SequencerHelpers::RemoveDuplicateKeys(KeySelection, MakeArrayView(KeyHandles));
+
 	InitialSectionRowIndicies.Empty();
 	ModifiedNonSelectedObjects.Empty();
 
diff --git a/Engine/Source/Editor/Sequencer/Public/MVVM/Selection/Selection.h b/Engine/Source/Editor/Sequencer/Public/MVVM/Selection/Selection.h
--- a/Engine/Source/Editor/Sequencer/Public/MVVM/Selection/Selection.h
+++ b/Engine/Source/Editor/Sequencer/Public/MVVM/Selection/Selection.h
@@ -51,9 +51,6 @@ struct FKeySelection : TUniqueFragmentSelectionSet<FKeyHandle, FChannelModel>
 	 * Overridden Empty function to remove duplicate keys on deselection
 	 */
 	void Empty();
-
-private:
-	void RemoveDuplicateKeys(TArrayView<FKeyHandle> KeyHandles);
 };
 
 
