Sequencer: Refactor FSequencerObjectChangeListener::CanKeyProperty_Internal to be more robust by encorporating registered property track editor changed events

diff --git a/Engine/Source/Editor/Sequencer/Private/SequencerObjectChangeListener.cpp b/Engine/Source/Editor/Sequencer/Private/SequencerObjectChangeListener.cpp
--- a/Engine/Source/Editor/Sequencer/Private/SequencerObjectChangeListener.cpp
+++ b/Engine/Source/Editor/Sequencer/Private/SequencerObjectChangeListener.cpp
@@ -231,7 +231,7 @@ bool IsHiddenFunction(const UStruct& PropertyStructure, FAnimatedPropertyKey Pro
 	return HideFunctions.Contains(FunctionName.ToString());
 }
 
-const FOnAnimatablePropertyChanged* FSequencerObjectChangeListener::FindPropertySetter(const UStruct& PropertyStructure, FAnimatedPropertyKey PropertyKey, const FProperty& Property) const
+TOptional<FOnAnimatablePropertyChanged> FSequencerObjectChangeListener::FindPropertySetter(const UStruct& PropertyStructure, FAnimatedPropertyKey PropertyKey, const FProperty& Property) const
 {
 	using namespace UE::MovieScene;
 
@@ -245,24 +245,40 @@ const FOnAnimatablePropertyChanged* FSequencerObjectChangeListener::FindProperty
 	// Early return if explicitly supported
 	if (const FOnAnimatablePropertyChanged* DelegatePtr = PropertyPathChangedEventMap.Find(&Property))
 	{
-		return DelegatePtr;
+		return *DelegatePtr;
 	}
 
+	// Else find a generic property changed event
+	TOptional<FOnAnimatablePropertyChanged> Result;
 	if (const FOnAnimatablePropertyChanged* DelegatePtr = PropertyChangedEventMap.Find(PropertyKey))
 	{
-			// If the function is hidden, we cannot use it
+		Result = *DelegatePtr;
+	}
+	else for (const FFindAnimatablePropertyChangedHandler& FindHandler : FindPropertyChangedHandlers)
+	{
+		FOnAnimatablePropertyChanged PropertyChanged = FindHandler.Execute(&Property);
+		if (PropertyChanged.IsBound())
+		{
+			Result = MoveTemp(PropertyChanged);
+			break;
+		}
+	}
+
+	if (Result.IsSet())
+	{
+		// If the function is hidden, we cannot use it
 		const bool bHiddenFunction = bCanApplyFunction && IsHiddenFunction(PropertyStructure, FAnimatedPropertyKey::FromProperty(&Property), Property.GetName());
 		if (bHiddenFunction)
 		{
-			return nullptr;
+			return {};
 		}
 
 		switch (GetPropertyAccessibility(*PropertyOrContainer))
 		{
 		case EPropertyAccessibility::Inaccessible:
-			return nullptr;
+			return {};
 		case EPropertyAccessibility::ExplicitlyAccessible:
-			return DelegatePtr;
+			return Result;
 		case EPropertyAccessibility::Unknown:
 			break;
 		}
@@ -272,13 +288,13 @@ const FOnAnimatablePropertyChanged* FSequencerObjectChangeListener::FindProperty
 			// If the function is hidden, we cannot use it
 			if (IsHiddenFunction(PropertyStructure, FAnimatedPropertyKey::FromProperty(&Property), Property.GetName()))
 			{
-				return nullptr;
+				return {};
 			}
 
 			// If there is a native setter we can always animate the property
 			if (Property.HasSetter())
 			{
-				return DelegatePtr;
+				return Result;
 			}
 	
 			// Check to see if we have a function of the form Set<PropertyName> that we can use
@@ -306,13 +322,13 @@ const FOnAnimatablePropertyChanged* FSequencerObjectChangeListener::FindProperty
 				// Valid if there's a setter function and the property is editable.
 				if (bFoundValidFunction && bEditable)
 				{
-					return DelegatePtr;
+					return Result;
 				}
 			}
 		}
 	}
 
-	return nullptr;
+	return {};
 }
 
 bool FSequencerObjectChangeListener::CanKeyProperty(FCanKeyPropertyParams CanKeyPropertyParams) const
@@ -403,28 +419,13 @@ bool FSequencerObjectChangeListener::CanKeyProperty_Internal(FCanKeyPropertyPara
 			}
 
 			// Otherwise we check for magic named functions using the default logic
-			if (const FOnAnimatablePropertyChanged* DelegatePtr = FindPropertySetter(*PropertyContainer, PropertyKey, *Property))
+			if (TOptional<FOnAnimatablePropertyChanged> Delegate = FindPropertySetter(*PropertyContainer, PropertyKey, *Property))
 			{
 				InOutProperty = Property;
-				InOutDelegate = *DelegatePtr;
+				InOutDelegate = MoveTemp(Delegate.GetValue());
 				return true;
 			}
 
-
-			if (GetPropertyAccessibility(*Property) == EPropertyAccessibility::ExplicitlyAccessible)
-			{
-				for (const FFindAnimatablePropertyChangedHandler& FindHandler : FindPropertyChangedHandlers)
-				{
-					FOnAnimatablePropertyChanged PropertyChanged = FindHandler.Execute(Property);
-					if (PropertyChanged.IsBound())
-					{
-						InOutProperty = Property;
-						InOutDelegate = MoveTemp(PropertyChanged);
-						return true;
-					}
-				}
-			}
-
 			const FObjectPropertyBase* ObjectProperty = CastField<FObjectPropertyBase>(Property);
 
 			// Check each level of the property hierarchy
@@ -439,10 +440,10 @@ bool FSequencerObjectChangeListener::CanKeyProperty_Internal(FCanKeyPropertyPara
 				{
 					Key.ObjectTypeName = ClassType->GetFName();
 
-					if (const FOnAnimatablePropertyChanged* DelegatePtr = FindPropertySetter(*PropertyContainer, Key, *Property))
+					if (TOptional<FOnAnimatablePropertyChanged> Delegate = FindPropertySetter(*PropertyContainer, Key, *Property))
 					{
 						InOutProperty = Property;
-						InOutDelegate = *DelegatePtr;
+						InOutDelegate = MoveTemp(Delegate.GetValue());
 						return true;
 					}
 
@@ -450,10 +451,10 @@ bool FSequencerObjectChangeListener::CanKeyProperty_Internal(FCanKeyPropertyPara
 				}
 
 				Key.ObjectTypeName = NAME_None;
-				if (const FOnAnimatablePropertyChanged* DelegatePtr = FindPropertySetter(*PropertyContainer, Key, *Property))
+				if (TOptional<FOnAnimatablePropertyChanged> Delegate = FindPropertySetter(*PropertyContainer, Key, *Property))
 				{
 					InOutProperty = Property;
-					InOutDelegate = *DelegatePtr;
+					InOutDelegate = MoveTemp(Delegate.GetValue());
 					return true;
 				}
 
diff --git a/Engine/Source/Editor/Sequencer/Private/SequencerObjectChangeListener.h b/Engine/Source/Editor/Sequencer/Private/SequencerObjectChangeListener.h
--- a/Engine/Source/Editor/Sequencer/Private/SequencerObjectChangeListener.h
+++ b/Engine/Source/Editor/Sequencer/Private/SequencerObjectChangeListener.h
@@ -76,7 +76,7 @@ private:
 	bool IsObjectValidForListening( UObject* Object ) const;
 
 	/** @return A valid delegate if property a setter could be found for a property on a class */
-	const FOnAnimatablePropertyChanged* FindPropertySetter(const UStruct& PropertyStructure, FAnimatedPropertyKey PropertyKey, const FProperty& Property) const;
+	TOptional<FOnAnimatablePropertyChanged> FindPropertySetter(const UStruct& PropertyStructure, FAnimatedPropertyKey PropertyKey, const FProperty& Property) const;
 
 	/** 
 	 * Internal implementation of CanKeyProperty that allows us to also retrieve a delegate that can be used to broadcast keyable properties.
