Sequencer: Remove VR Editor mesh trail support.

#rnx

diff --git a/Engine/Source/Editor/Sequencer/Private/Sequencer.cpp b/Engine/Source/Editor/Sequencer/Private/Sequencer.cpp
--- a/Engine/Source/Editor/Sequencer/Private/Sequencer.cpp
+++ b/Engine/Source/Editor/Sequencer/Private/Sequencer.cpp
@@ -136,7 +136,6 @@
 #include "LevelSequenceActor.h"
 #include "IVREditorModule.h"
 #include "HAL/PlatformApplicationMisc.h"
-#include "SequencerKeyActor.h"
 #include "ISequencerChannelInterface.h"
 #include "IMovieRendererInterface.h"
 #include "MVVM/ViewModels/OutlinerColumns/IOutlinerColumn.h"
@@ -2795,18 +2794,6 @@ void FSequencer::NotifyMovieSceneDataChanged( EMovieSceneDataChangeType DataChan
 		}
 	}
 
-	if (DataChangeType == EMovieSceneDataChangeType::TrackValueChanged || 
-		DataChangeType == EMovieSceneDataChangeType::TrackValueChangedRefreshImmediately || 
-		DataChangeType == EMovieSceneDataChangeType::Unknown ||
-		DataChangeType == EMovieSceneDataChangeType::MovieSceneStructureItemRemoved)
-	{
-		FSequencerEdMode* SequencerEdMode = (FSequencerEdMode*)(GLevelEditorModeTools().GetActiveMode(FSequencerEdMode::EM_SequencerMode));
-		if (SequencerEdMode != nullptr)
-		{
-			SequencerEdMode->CleanUpMeshTrails();
-		}
-	}
-
 	bGlobalMarkedFramesCached = false;
 	bNeedsEvaluate = true;
 	PRAGMA_DISABLE_DEPRECATION_WARNINGS
@@ -6582,17 +6569,6 @@ void FSequencer::HandleSelectedOutlinerNodesChanged()
 
 	SynchronizeExternalSelectionWithSequencerSelection();
 
-	FSequencerEdMode* SequencerEdMode = (FSequencerEdMode*)(GLevelEditorModeTools().GetActiveMode(FSequencerEdMode::EM_SequencerMode));
-	if (SequencerEdMode != nullptr)
-	{
-		AActor* NewlySelectedActor = GEditor->GetSelectedActors()->GetTop<AActor>();
-		// If we selected an Actor or a node for an Actor that is a potential autokey candidate, clean up any existing mesh trails
-		if (NewlySelectedActor && !NewlySelectedActor->IsEditorOnly())
-		{
-			SequencerEdMode->CleanUpMeshTrails();
-		}
-	}
-
 	TSet<UMovieSceneTrack*> SelectedTracks;
 	for (TViewModelPtr<ITrackExtension> TrackModel : ViewModel->GetSelection()->Outliner.Filter<ITrackExtension>())
 	{
@@ -7007,9 +6983,6 @@ void FSequencer::SynchronizeSequencerSelectionWithExternalSelection()
 	USelection* ActorSelection = GEditor->GetSelectedActors();
 	
 	// Get the selected sequencer keys for viewport interaction
-	TArray<ASequencerKeyActor*> SelectedSequencerKeyActors;
-	ActorSelection->GetSelectedObjects<ASequencerKeyActor>(SelectedSequencerKeyActors);
-
 	TSharedPtr<FSequencerSelection> Selection = ViewModel->GetSelection();
 
 	FObjectBindingModelStorageExtension* ObjectModelStorage = ViewModel->GetRootModel()->CastDynamic<FObjectBindingModelStorageExtension>();
@@ -7032,20 +7005,6 @@ void FSequencer::SynchronizeSequencerSelectionWithExternalSelection()
 				continue;
 			}
 
-			for (ASequencerKeyActor* KeyActor : SelectedSequencerKeyActors)
-			{
-				if (KeyActor->IsEditorOnly())
-				{
-					AActor* TrailActor = KeyActor->GetAssociatedActor();
-					if (TrailActor != nullptr && RuntimeObject == TrailActor)
-					{
-						NodesToSelect.Add(ObjectModel);
-						bAllAlreadySelected = false;
-						break;
-					}
-				}
-			}
-
 			AActor* Actor = Cast<AActor>(RuntimeObject);
 			const bool bActorSelected = Actor ? ActorSelection->IsSelected(Actor) : false;
 						
diff --git a/Engine/Source/Editor/Sequencer/Private/SequencerEdMode.cpp b/Engine/Source/Editor/Sequencer/Private/SequencerEdMode.cpp
--- a/Engine/Source/Editor/Sequencer/Private/SequencerEdMode.cpp
+++ b/Engine/Source/Editor/Sequencer/Private/SequencerEdMode.cpp
@@ -28,8 +28,6 @@
 #include "Sections/MovieSceneAudioSection.h"
 #include "Compilation/MovieSceneCompiledDataManager.h"
 #include "SubtitleManager.h"
-#include "SequencerMeshTrail.h"
-#include "SequencerKeyActor.h"
 #include "EditorWorldExtension.h"
 #include "ViewportWorldInteraction.h"
 #include "SSequencer.h"
@@ -54,11 +52,6 @@
 
 const FEditorModeID FSequencerEdMode::EM_SequencerMode(TEXT("EM_SequencerMode"));
 
-static TAutoConsoleVariable<bool> CVarDrawMeshTrails(
-	TEXT("Sequencer.DrawMeshTrails"),
-	true,
-	TEXT("Toggle to show or hide Level Sequence VR Editor trails"));
-
 //DEPRECRATED
 TAutoConsoleVariable<bool> CVarUseOldSequencerMotionTrails(
 	TEXT("Sequencer.UseOldSequencerTrails"),
@@ -199,19 +192,12 @@ FSequencerEdMode::FSequencerEdMode()
 	Tools.Add(SelectabilityTool.Get());
 	SetCurrentTool(DefaultTool.Get());
 
-	bDrawMeshTrails = CVarDrawMeshTrails->GetBool();
-	CVarDrawMeshTrails.AsVariable()->SetOnChangedCallback(FConsoleVariableDelegate::CreateLambda([this](IConsoleVariable* Var) 
-	{
-		bDrawMeshTrails = Var->GetBool();
-	}));
-
 	AudioTexture = LoadObject<UTexture2D>(NULL, TEXT("/Engine/EditorResources/AudioIcons/S_AudioComponent.S_AudioComponent"));
 	check(AudioTexture);
 }
 
 FSequencerEdMode::~FSequencerEdMode()
 {
-	CVarDrawMeshTrails.AsVariable()->SetOnChangedCallback(FConsoleVariableDelegate());
 }
 
 void FSequencerEdMode::Enter()
@@ -224,8 +210,6 @@ void FSequencerEdMode::Enter()
 
 void FSequencerEdMode::Exit()
 {
-	CleanUpMeshTrails();
-
 	Sequencers.Reset();
 
 	FEdMode::Exit();
@@ -576,12 +560,6 @@ void FSequencerEdMode::Render(const FSceneView* View, FViewport* Viewport, FPrim
 	{
 		DrawTracks3D(PDI);
 	}
-	// Draw mesh trails (doesn't use the PDI)
-	else if (bDrawMeshTrails)
-	{
-		PDI = nullptr;
-		DrawTracks3D(PDI);
-	}
 #endif
 }
 
@@ -617,15 +595,6 @@ void FSequencerEdMode::DrawHUD(FEditorViewportClient* ViewportClient, FViewport*
 	}
 }
 
-void FSequencerEdMode::AddReferencedObjects(FReferenceCollector& Collector)
-{
-	for (FMeshTrailData& MeshTrail : MeshTrails)
-	{
-		Collector.AddReferencedObject(MeshTrail.Track);
-		Collector.AddReferencedObject(MeshTrail.Trail);
-	}
-}
-
 void FSequencerEdMode::OnKeySelected(FViewport* Viewport, HMovieSceneKeyProxy* KeyProxy)
 {
 	using namespace UE::Sequencer;
@@ -684,50 +653,6 @@ void FSequencerEdMode::OnKeySelected(FViewport* Viewport, HMovieSceneKeyProxy* K
 	}
 }
 
-void FSequencerEdMode::DrawMeshTransformTrailFromKey(const class ASequencerKeyActor* KeyActor)
-{
-	ASequencerMeshTrail* Trail = Cast<ASequencerMeshTrail>(KeyActor->GetOwner());
-	if(Trail != nullptr)
-	{
-		FMeshTrailData* TrailPtr = MeshTrails.FindByPredicate([Trail](const FMeshTrailData InTrail)
-		{
-			return Trail == InTrail.Trail;
-		});
-		if(TrailPtr != nullptr)
-		{
-			// From the key, get the mesh trail, and then the track associated with that mesh trail
-			UMovieScene3DTransformTrack* Track = TrailPtr->Track;
-			// Draw a mesh trail for the key's associated actor
-			TArray<TWeakObjectPtr<UObject>> KeyObjects;
-			AActor* TrailActor = KeyActor->GetAssociatedActor();
-			KeyObjects.Add(TrailActor);
-			FPrimitiveDrawInterface* PDI = nullptr;
-
-			for (TWeakPtr<ISequencer> WeakSequencer : Sequencers)
-			{
-				TSharedPtr<ISequencer> Sequencer = WeakSequencer.Pin();
-				if (Sequencer.IsValid())
-				{
-					DrawTransformTrack(Sequencer, PDI, Track, KeyObjects, true);
-				}
-			}
-		}
-	}
-}
-
-void FSequencerEdMode::CleanUpMeshTrails()
-{
-	// Clean up any existing trails
-	for (FMeshTrailData& MeshTrail : MeshTrails)
-	{
-		if (MeshTrail.Trail)
-		{
-			MeshTrail.Trail->Cleanup();
-		}
-	}
-	MeshTrails.Empty();
-}
-
 void FSequencerEdMode::DrawTransformTrack(const TSharedPtr<ISequencer>& Sequencer, FPrimitiveDrawInterface* PDI,
 											UMovieScene3DTransformTrack* TransformTrack, TArrayView<const TWeakObjectPtr<>> BoundObjects, const bool bIsSelected)
 {
@@ -736,20 +661,6 @@ void FSequencerEdMode::DrawTransformTrack(const TSharedPtr<ISequencer>& Sequence
 
 	const bool bHitTesting = PDI && PDI->IsHitTesting();
 
-	ASequencerMeshTrail* TrailActor = nullptr;
-	// Get the Trail Actor associated with this track if we are drawing mesh trails
-	if (bDrawMeshTrails)
-	{
-		FMeshTrailData* TrailPtr = MeshTrails.FindByPredicate([TransformTrack](const FMeshTrailData InTrail)
-		{
-			return InTrail.Track == TransformTrack;
-		});
-		if (TrailPtr != nullptr)
-		{
-			TrailActor = TrailPtr->Trail;
-		}
-	}
-
 	bool bShowTrajectory = TransformTrack->GetAllSections().ContainsByPredicate(
 		[bIsSelected](UMovieSceneSection* Section)
 		{
@@ -835,14 +746,6 @@ void FSequencerEdMode::DrawTransformTrack(const TSharedPtr<ISequencer>& Sequence
 						{
 							PDI->DrawLine(StartPosition, EndTransform.GetTranslation(), TrackColor, SDPG_Foreground);
 						}
-						else if (TrailActor)
-						{
-							FTransform FrameTransform = EndTransform;
-							FrameTransform.SetScale3D(FVector(3.0f));
-
-							FFrameTime FrameTime = TrackTransforms.Times[TransformIndex];
-							TrailActor->AddFrameMeshComponent(FrameTime / TickResolution, FrameTransform);
-						}
 
 						StartPosition = EndTransform.GetTranslation();
 					}
@@ -882,22 +785,6 @@ void FSequencerEdMode::DrawTransformTrack(const TSharedPtr<ISequencer>& Sequence
 					PDI->SetHitProxy(nullptr);
 				}
 			}
-			else if (TrailActor != nullptr)
-			{
-				AllSectionsScratch.Reset();
-				for (const FTrajectoryKey::FData& Value : ThisKey.KeyData)
-				{
-					UMovieScene3DTransformSection* Section = Value.Section.Get();
-					if (Section && !AllSectionsScratch.Contains(Section))
-					{
-						FTransform MeshTransform = ThisTransform;
-						MeshTransform.SetScale3D(FVector(3.0f));
-
-						TrailActor->AddKeyMeshActor(ThisKey.Time / TickResolution, MeshTransform, Section);
-						AllSectionsScratch.Add(Section);
-					}
-				}
-			}
 		}
 	}
 }
@@ -988,22 +875,6 @@ void FSequencerEdMode::DrawTracks3D(FPrimitiveDrawInterface* PDI)
 					continue;
 				}
 
-				// Ensure that we've got a mesh trail for this track
-				if (bDrawMeshTrails)
-				{
-					const bool bHasMeshTrail = Algo::FindBy(MeshTrails, TransformTrack, &FMeshTrailData::Track) != nullptr;
-					if (!bHasMeshTrail)
-					{
-						UViewportWorldInteraction* WorldInteraction = Cast<UViewportWorldInteraction>( GEditor->GetEditorWorldExtensionsManager()->GetEditorWorldExtensions( GetWorld() )->FindExtension( UViewportWorldInteraction::StaticClass() ) );
-						if( WorldInteraction != nullptr )
-						{
-							ASequencerMeshTrail* TrailActor = WorldInteraction->SpawnTransientSceneActor<ASequencerMeshTrail>(TEXT("SequencerMeshTrail"), true);
-							FMeshTrailData MeshTrail = FMeshTrailData(TransformTrack, TrailActor);
-							MeshTrails.Add(MeshTrail);
-						}
-					}
-				}
-
 				const bool bIsSelected = Pair.Value;
 				DrawTransformTrack(Sequencer, PDI, TransformTrack, Sequencer->FindObjectsInCurrentSequence(Pair.Key->GetObjectGuid()), bIsSelected);
 			}
diff --git a/Engine/Source/Editor/Sequencer/Private/SequencerEdMode.h b/Engine/Source/Editor/Sequencer/Private/SequencerEdMode.h
--- a/Engine/Source/Editor/Sequencer/Private/SequencerEdMode.h
+++ b/Engine/Source/Editor/Sequencer/Private/SequencerEdMode.h
@@ -59,19 +59,6 @@ private:
 };
 
 
-/** Stores the transform track and associated mesh trail for each drawn track */
-struct FMeshTrailData
-{
-	TObjectPtr<class UMovieScene3DTransformTrack> Track;
-	TObjectPtr<class ASequencerMeshTrail> Trail;
-
-	FMeshTrailData(class UMovieScene3DTransformTrack* InTrack, class ASequencerMeshTrail* InTrail) :
-		Track(InTrack),
-		Trail(InTrail)
-	{
-	}
-};
-
 /**
  * FSequencerEdMode is the editor mode for additional drawing and handling sequencer hotkeys in the editor
  */
@@ -93,7 +80,6 @@ public:
 	virtual void DrawHUD(FEditorViewportClient* ViewportClient, FViewport* Viewport, const FSceneView* View, FCanvas* Canvas) override;
 	virtual bool UsesTransformWidget() const override { return false; }
 	virtual bool UsesTransformWidget(UE::Widget::EWidgetMode CheckMode) const override { return false; }
-	virtual void AddReferencedObjects(FReferenceCollector& Collector) override;
 
 	virtual bool StartTracking(FEditorViewportClient* InViewportClient, FViewport* InViewport) override;
 	virtual bool EndTracking(FEditorViewportClient* InViewportClient, FViewport* InViewport) override;
@@ -116,12 +102,6 @@ public:
 
 	void OnKeySelected(FViewport* Viewport, HMovieSceneKeyProxy* KeyProxy);
 
-	/** Draw a single mesh transform track, given a key that is on that track */
-	void DrawMeshTransformTrailFromKey(const class ASequencerKeyActor* KeyActor);
-
-	/** Clean up any mesh trails and their associated key actors */
-	void CleanUpMeshTrails();
-
 	bool IsViewportSelectionLimited() const;
 
 	void EnableSelectabilityTool(const bool bInEnabled);
@@ -138,15 +118,9 @@ protected:
 private:
 	TArray<TWeakPtr<FSequencer>> Sequencers;
 
-	/**Array of the transform tracks and their associated mesh trails */
-	TArray<FMeshTrailData> MeshTrails;
-
 	/** Interrogation data for extracting transforms */
 	TSharedPtr<FMovieSceneInterrogationData> InterrogationData;
 
-	/** If true, draw mesh trails instead of debug lines*/
-	bool bDrawMeshTrails;
-
 	/** The audio texture used for drawing the audio spatialization points */
 	UTexture2D* AudioTexture;
 
diff --git a/Engine/Source/Editor/Sequencer/Private/SequencerKeyActor.cpp b/Engine/Source/Editor/Sequencer/Private/SequencerKeyActor.cpp
deleted file mode 100644
diff --git a/Engine/Source/Editor/Sequencer/Private/SequencerKeyActor.h b/Engine/Source/Editor/Sequencer/Private/SequencerKeyActor.h
deleted file mode 100644
diff --git a/Engine/Source/Editor/Sequencer/Private/SequencerMeshTrail.cpp b/Engine/Source/Editor/Sequencer/Private/SequencerMeshTrail.cpp
deleted file mode 100644
diff --git a/Engine/Source/Editor/Sequencer/Private/SequencerMeshTrail.h b/Engine/Source/Editor/Sequencer/Private/SequencerMeshTrail.h
deleted file mode 100644
