Sequencer: Fix rare crash when drawing channels

#rnx

diff --git a/Engine/Source/Editor/Sequencer/Private/MVVM/Views/SChannelView.cpp b/Engine/Source/Editor/Sequencer/Private/MVVM/Views/SChannelView.cpp
--- a/Engine/Source/Editor/Sequencer/Private/MVVM/Views/SChannelView.cpp
+++ b/Engine/Source/Editor/Sequencer/Private/MVVM/Views/SChannelView.cpp
@@ -463,6 +463,11 @@ int32 SChannelView::OnPaint(const FPaintArgs& Args, const FGeometry& AllottedGeo
 
 
 	TSharedPtr<FSequencer> Sequencer = LegacyGetSequencer();
+	if (!Sequencer.IsValid())
+	{
+		return LayerId;
+	}
+
 	FViewModelPtr Model = WeakModel.Pin();
 	if (!Model)
 	{
@@ -528,7 +533,7 @@ int32 SChannelView::OnPaint(const FPaintArgs& Args, const FGeometry& AllottedGeo
 	FSequencerKeyRendererInterface KeyRenderInterface(Sequencer.Get(), TrackArea->GetHotspot());
 	FChannelViewKeyCachedState NewCachedState(VisibleRange, TrackArea->GetHotspot(), Model, Sequencer.Get());
 	NewCachedState.bIsChannelHovered = IsHovered();
-	NewCachedState.KeySizePx = SectionInterface->GetKeySize();
+	NewCachedState.KeySizePx = SectionInterface ? SectionInterface->GetKeySize() : FVector2D(1.0, 1.0);
 	NewCachedState.NonLinearTransform = RelativeTimeToPixel.NonLinearTransform;
 
 	FKeyBatchParameters Params(RelativeTimeToPixel);
@@ -551,7 +556,8 @@ int32 SChannelView::OnPaint(const FPaintArgs& Args, const FGeometry& AllottedGeo
 		PaintArgs.KeyBarColor = KeyBarColor.Get();
 	}
 
-	if (Sequencer->GetSequencerSettings()->GetShowChannelColors() && Channel)
+	USequencerSettings* SequencerSettings = Sequencer->GetSequencerSettings();
+	if (SequencerSettings && SequencerSettings->GetShowChannelColors() && Channel)
 	{
 		if (TOptional<FLinearColor> ChannelColor = Channel->GetKeyArea()->GetColor())
 		{
