Sequencer: Fix regression in keying with filtered out channels

#rnx

#FailedValidations Preflight

diff --git a/Engine/Source/Editor/Sequencer/Private/SequencerAddKeyOperation.cpp b/Engine/Source/Editor/Sequencer/Private/SequencerAddKeyOperation.cpp
--- a/Engine/Source/Editor/Sequencer/Private/SequencerAddKeyOperation.cpp
+++ b/Engine/Source/Editor/Sequencer/Private/SequencerAddKeyOperation.cpp
@@ -126,7 +126,10 @@ bool FAddKeyOperation::ConsiderKeyableAreas(TSharedPtr<ITrackExtension> InTrackM
 			ISequencerTrackEditor* TrackEditor = InTrackModel->GetTrackEditor().Get();
 			for (TParentFirstChildIterator<FChannelModel> ChannelModelIt(SectionModel->AsShared(), bIncludeThis); ChannelModelIt; ++ChannelModelIt)
 			{	
-				bKeyedAnything |= ProcessKeyArea(TrackEditor, ChannelModelIt->GetKeyArea());
+				if (!ChannelModelIt->GetLinkedOutlinerItem() || ChannelModelIt->GetLinkedOutlinerItem()->IsFilteredOut() == false)
+				{
+					bKeyedAnything |= ProcessKeyArea(TrackEditor, ChannelModelIt->GetKeyArea());
+				}
 			}
 		}
 	}
