Sequencer: Update tooltip to "The interpolation type for the initial keyframe" to clarify that this is only for the first keyframe on the track/section

#FailedValidations Preflight

diff --git a/Engine/Source/Editor/Sequencer/Private/Menus/SequencerToolbarUtils.cpp b/Engine/Source/Editor/Sequencer/Private/Menus/SequencerToolbarUtils.cpp
--- a/Engine/Source/Editor/Sequencer/Private/Menus/SequencerToolbarUtils.cpp
+++ b/Engine/Source/Editor/Sequencer/Private/Menus/SequencerToolbarUtils.cpp
@@ -33,7 +33,7 @@ static TSharedRef<SWidget> MakeKeyGroupMenu(const TWeakPtr<FSequencer> InWeakSeq
 	}
 
 	// Interpolation
-	MenuBuilder.BeginSection("SequencerInterpolation", LOCTEXT("KeyInterpolationMenu", "Default Key Interpolation"));
+	MenuBuilder.BeginSection("SequencerInterpolation", LOCTEXT("KeyInterpolationMenu", "Initial Key Interpolation"));
 	{
 		MenuBuilder.AddMenuEntry(
 			LOCTEXT("SetKeyInterpolationSmartAuto", "Cubic (Smart Auto)"),
diff --git a/Engine/Source/Editor/Sequencer/Public/ISequencer.h b/Engine/Source/Editor/Sequencer/Public/ISequencer.h
--- a/Engine/Source/Editor/Sequencer/Public/ISequencer.h
+++ b/Engine/Source/Editor/Sequencer/Public/ISequencer.h
@@ -405,10 +405,10 @@ public:
 	*/
 	virtual void SetKeyGroupMode(EKeyGroupMode Mode) = 0;
 
-	/** @return Returns default key interpolation */
+	/** @return Returns initial key interpolation */
 	virtual EMovieSceneKeyInterpolation GetKeyInterpolation() const = 0;
 
-	/** Set default key interpolation */
+	/** Set initial key interpolation */
 	virtual void SetKeyInterpolation(EMovieSceneKeyInterpolation) = 0;
 
 	/** @return Returns whether key sections are infinite by default when created */
diff --git a/Engine/Source/Editor/Sequencer/Public/SequencerSettings.h b/Engine/Source/Editor/Sequencer/Public/SequencerSettings.h
--- a/Engine/Source/Editor/Sequencer/Public/SequencerSettings.h
+++ b/Engine/Source/Editor/Sequencer/Public/SequencerSettings.h
@@ -161,9 +161,9 @@ public:
 	/** Sets which channels are keyed when a channel is keyed */
 	UE_API void SetKeyGroupMode(EKeyGroupMode);
 
-	/** Get the default Interpolation type for newly created keyframes if the channel does not already have keyframes */
+	/** Get the initial Interpolation type for newly created keyframes if the channel does not already have keyframes */
 	UE_API EMovieSceneKeyInterpolation GetKeyInterpolation() const;
-	/** Sets default key interpolation for creating new keys on empty channels */
+	/** Sets initial key interpolation for creating new keys on empty channels */
 	UE_API void SetKeyInterpolation(EMovieSceneKeyInterpolation InKeyInterpolation);
 
 	/** Get initial spawn position. */
@@ -591,7 +591,7 @@ protected:
 	UPROPERTY(config, EditAnywhere, Category = Keyframing)
 	EKeyGroupMode KeyGroupMode;
 
-	/** The interpolation type for newly created keyframes */
+	/** The interpolation type for the initial keyframe */
 	UPROPERTY( config, EditAnywhere, Category=Keyframing )
 	EMovieSceneKeyInterpolation KeyInterpolation;
 
