Sequencer: Snap play time to pressed/dragged key only if snapping is enabled

#FailedValidations Preflight

diff --git a/Engine/Source/Editor/Sequencer/Private/Tools/EditToolDragOperations.cpp b/Engine/Source/Editor/Sequencer/Private/Tools/EditToolDragOperations.cpp
--- a/Engine/Source/Editor/Sequencer/Private/Tools/EditToolDragOperations.cpp
+++ b/Engine/Source/Editor/Sequencer/Private/Tools/EditToolDragOperations.cpp
@@ -879,7 +879,7 @@ void FMoveKeysAndSections::OnDrag(const FPointerEvent& MouseEvent, FVector2D Loc
 	bool bSectionMovementModifiedStructure = HandleSectionMovement(MouseTime, VirtualMousePos, LocalMousePos, MaxDeltaX, MouseDeltaTime);
 
 	// Update our key times by moving them by our delta.
-	HandleKeyMovement(MaxDeltaX, MouseDeltaTime);
+	HandleKeyMovement(MouseEvent, MaxDeltaX, MouseDeltaTime);
 
 	// Update our marked frames by moving them by our delta.
 	HandleMarkedFrameMovement(MaxDeltaX, MouseDeltaTime);
@@ -1378,7 +1378,7 @@ bool FMoveKeysAndSections::HandleSectionMovement(FFrameTime MouseTime, FVector2D
 	return bRowIndexChanged;
 }
 
-void FMoveKeysAndSections::HandleKeyMovement(TOptional<FFrameNumber> MaxDeltaX, FFrameNumber DesiredDeltaX)
+void FMoveKeysAndSections::HandleKeyMovement(const FPointerEvent& MouseEvent, TOptional<FFrameNumber> MaxDeltaX, FFrameNumber DesiredDeltaX)
 {
 	if (KeysAsArray.Num() == 0)
 	{
@@ -1430,7 +1430,10 @@ void FMoveKeysAndSections::HandleKeyMovement(TOptional<FFrameNumber> MaxDeltaX,
 
 
 	// Snap the play time to the new dragged key time if all the keyframes were dragged to the same time
-	if (Settings->GetSnapPlayTimeToDraggedKey() && CurrentKeyTimes.Num())
+	const bool bSnapPlayTimeToDraggedKey = 
+		((Settings->GetIsSnapEnabled() && Settings->GetSnapPlayTimeToDraggedKey()) || (MouseEvent.IsShiftDown() && MouseEvent.IsMouseButtonDown(EKeys::LeftMouseButton))) &&
+		CurrentKeyTimes.Num();
+	if (bSnapPlayTimeToDraggedKey)
 	{
 		FFrameNumber FirstFrame = CurrentKeyTimes[0];
 		auto         EqualsFirstFrame = [=](FFrameNumber In)
diff --git a/Engine/Source/Editor/Sequencer/Private/Tools/EditToolDragOperations.h b/Engine/Source/Editor/Sequencer/Private/Tools/EditToolDragOperations.h
--- a/Engine/Source/Editor/Sequencer/Private/Tools/EditToolDragOperations.h
+++ b/Engine/Source/Editor/Sequencer/Private/Tools/EditToolDragOperations.h
@@ -183,7 +183,7 @@ protected:
 	/** Move selected sections, if any. */
 	bool HandleSectionMovement(FFrameTime MouseTime, FVector2D VirtualMousePos, FVector2D LocalMousePos, TOptional<FFrameNumber> MaxDeltaX, FFrameNumber DesiredDeltaX);
 	/** Move selected keys, if any. */
-	void HandleKeyMovement(TOptional<FFrameNumber> MaxDeltaX, FFrameNumber DesiredDeltaX);
+	void HandleKeyMovement(const FPointerEvent& MouseEvent, TOptional<FFrameNumber> MaxDeltaX, FFrameNumber DesiredDeltaX);
 	/** Move selected marked frames, if any. */
 	void HandleMarkedFrameMovement(TOptional<FFrameNumber> MaxDeltaX, FFrameNumber DesiredDeltaX);
 
diff --git a/Engine/Source/Editor/Sequencer/Private/Tools/SequencerEditTool_Movement.cpp b/Engine/Source/Editor/Sequencer/Private/Tools/SequencerEditTool_Movement.cpp
--- a/Engine/Source/Editor/Sequencer/Private/Tools/SequencerEditTool_Movement.cpp
+++ b/Engine/Source/Editor/Sequencer/Private/Tools/SequencerEditTool_Movement.cpp
@@ -51,7 +51,12 @@ FReply FSequencerEditTool_Movement::OnMouseButtonDown(SWidget& OwnerWidget, cons
 
 		DelayedDrag = FDelayedDrag_Hotspot(VirtualTrackArea.CachedTrackAreaGeometry().AbsoluteToLocal(MouseEvent.GetScreenSpacePosition()), MouseEvent.GetEffectingButton(), Hotspot);
 
-		if (Sequencer.GetSequencerSettings()->GetSnapPlayTimeToPressedKey() || (MouseEvent.IsShiftDown() && MouseEvent.GetEffectingButton() == EKeys::LeftMouseButton) )
+		USequencerSettings* Settings = Sequencer.GetSequencerSettings();
+
+		const bool bSnapPlayTimeToPressedKey = 
+			(Settings->GetIsSnapEnabled() && Settings->GetSnapPlayTimeToPressedKey()) || (MouseEvent.IsShiftDown() && MouseEvent.GetEffectingButton() == EKeys::LeftMouseButton);
+
+		if (bSnapPlayTimeToPressedKey)
 		{
 			if (TSharedPtr<FKeyHotspot> KeyHotspot = HotspotCast<FKeyHotspot>(DelayedDrag->Hotspot))
 			{
@@ -571,7 +576,7 @@ FString FSequencerEditTool_Movement::TimeToString(FFrameTime Time, bool IsDelta)
 			// Convert from sequence resolution into display rate frames.
 			FFrameTime DisplayTime = FFrameRate::TransformTime(Time, TickResolution, DisplayRate);
 			FString SubframeIndicator = FMath::IsNearlyZero(DisplayTime.GetSubFrame()) ? TEXT("") : TEXT("*");
-			int32 ZeroPadFrames = Sequencer.GetSequencerSettings()->GetZeroPadFrames();
+			int32 ZeroPadFrames = Settings->GetZeroPadFrames();
 			return IsDelta ? FString::Printf(TEXT("[%+0*d%s]"), ZeroPadFrames, DisplayTime.GetFrame().Value, *SubframeIndicator) : FString::Printf(TEXT("%0*d%s"), ZeroPadFrames, DisplayTime.GetFrame().Value, *SubframeIndicator);
 		}
 		case EFrameNumberDisplayFormats::NonDropFrameTimecode:
