Sequencer: Only ResetCache from NotifyMovieSceneDataChanged for certain events, never for some like changing track values.

diff --git a/Engine/Source/Editor/Sequencer/Private/Sequencer.cpp b/Engine/Source/Editor/Sequencer/Private/Sequencer.cpp
--- a/Engine/Source/Editor/Sequencer/Private/Sequencer.cpp
+++ b/Engine/Source/Editor/Sequencer/Private/Sequencer.cpp
@@ -2751,7 +2751,7 @@ void FSequencer::NotifyMovieSceneDataChanged( EMovieSceneDataChangeType DataChan
 			UE_LOG(LogSequencer, Error, TEXT("Fatal error, focused movie scene no longer valid and there is no root sequence to default to."));
 		}
 	}
-
+	bool bResetCache = true;
 	if (DataChangeType == EMovieSceneDataChangeType::RefreshTree)
 	{
 		bNeedTreeRefresh = true;
@@ -2774,6 +2774,7 @@ void FSequencer::NotifyMovieSceneDataChanged( EMovieSceneDataChangeType DataChan
 	{
 		// Evaluate now
 		EvaluateInternal(PlayPosition.GetCurrentPositionAsRange());
+		bResetCache = false;
 	}
 	else if (DataChangeType == EMovieSceneDataChangeType::RefreshAllImmediately)
 	{
@@ -2781,6 +2782,7 @@ void FSequencer::NotifyMovieSceneDataChanged( EMovieSceneDataChangeType DataChan
 
 		// Evaluate now
 		EvaluateInternal(PlayPosition.GetCurrentPositionAsRange());
+		//not sure ? bResetCache = false;
 	}
 	else
 	{
@@ -2792,12 +2794,19 @@ void FSequencer::NotifyMovieSceneDataChanged( EMovieSceneDataChangeType DataChan
 		{
 			bNeedTreeRefresh = true;
 		}
+		if (DataChangeType == EMovieSceneDataChangeType::TrackValueChanged)
+		{
+			bResetCache = false;
+		}
 	}
 
 	bGlobalMarkedFramesCached = false;
 	bNeedsEvaluate = true;
 	PRAGMA_DISABLE_DEPRECATION_WARNINGS
-	State.ClearObjectCaches(*this);
+	if(bResetCache == true)
+	{ 
+		State.ClearObjectCaches(*this);
+	}
 	PRAGMA_ENABLE_DEPRECATION_WARNINGS
 
 	UpdatePlaybackRange();
