Sequencer: better handle looping for camera cut tracks

- Add a new "HasLooped" flag on the evaluation context
- Assume by default that looping sequences are "seamless" and do not create a straight/hard camera cut on loops
- Add an advanced evaluation flag to opt out of this assumption

diff --git a/Engine/Source/Editor/Sequencer/Private/SSequencer.cpp b/Engine/Source/Editor/Sequencer/Private/SSequencer.cpp
--- a/Engine/Source/Editor/Sequencer/Private/SSequencer.cpp
+++ b/Engine/Source/Editor/Sequencer/Private/SSequencer.cpp
@@ -2450,6 +2450,7 @@ TSharedRef<SWidget> SSequencer::MakePlaybackMenu()
 
 		MenuBuilder.AddMenuEntry( FSequencerCommands::Get().ToggleAsyncEvaluation );
 		MenuBuilder.AddMenuEntry( FSequencerCommands::Get().ToggleDynamicWeighting );
+		MenuBuilder.AddMenuEntry( FSequencerCommands::Get().ToggleLoopCuts );
 	}
 	MenuBuilder.EndSection();
 
diff --git a/Engine/Source/Editor/Sequencer/Private/Sequencer.cpp b/Engine/Source/Editor/Sequencer/Private/Sequencer.cpp
--- a/Engine/Source/Editor/Sequencer/Private/Sequencer.cpp
+++ b/Engine/Source/Editor/Sequencer/Private/Sequencer.cpp
@@ -11665,6 +11665,12 @@ void FSequencer::BindCommands()
 		FCanExecuteAction::CreateLambda( []{ return true; } ),
 		FIsActionChecked::CreateLambda( [this]{ return this->UsesDynamicWeighting(); } ) );
 
+	SequencerCommandBindings->MapAction(
+		Commands.ToggleLoopCuts,
+		FExecuteAction::CreateLambda( [this]{ this->ToggleLoopCuts(); } ),
+		FCanExecuteAction::CreateLambda( []{ return true; } ),
+		FIsActionChecked::CreateLambda( [this]{ return this->UsesLoopCuts(); } ) );
+
 	SequencerCommandBindings->MapAction(
 		Commands.ToggleKeepCursorInPlaybackRangeWhileScrubbing,
 		FExecuteAction::CreateLambda([this] { Settings->SetKeepCursorInPlayRangeWhileScrubbing(!Settings->ShouldKeepCursorInPlayRangeWhileScrubbing()); }),
@@ -12500,7 +12506,7 @@ void FSequencer::ToggleDynamicWeighting()
 	EMovieSceneSequenceFlags NewFlags = Sequence->GetFlags();
 	NewFlags ^= EMovieSceneSequenceFlags::DynamicWeighting;
 
-	FScopedTransaction Transaction(EnumHasAnyFlags(NewFlags, EMovieSceneSequenceFlags::DynamicWeighting) ? LOCTEXT("DisableDynamicWeighting", "Disable Dynamic Weighting") : LOCTEXT("EnableDynamicWeighting", "Enable Dynamic Weighting"));
+	FScopedTransaction Transaction(EnumHasAnyFlags(NewFlags, EMovieSceneSequenceFlags::DynamicWeighting) ? LOCTEXT("EnableDynamicWeighting", "Enable Dynamic Weighting") : LOCTEXT("DisableDynamicWeighting", "Disable Dynamic Weighting"));
 
 	Sequence->Modify();
 	Sequence->SetSequenceFlags(NewFlags);
@@ -12511,6 +12517,24 @@ bool FSequencer::UsesDynamicWeighting()
 	return EnumHasAnyFlags(GetFocusedMovieSceneSequence()->GetFlags(), EMovieSceneSequenceFlags::DynamicWeighting);
 }
 
+void FSequencer::ToggleLoopCuts()
+{
+	UMovieSceneSequence* Sequence = GetFocusedMovieSceneSequence();
+
+	EMovieSceneSequenceFlags NewFlags = Sequence->GetFlags();
+	NewFlags ^= EMovieSceneSequenceFlags::LoopCuts;
+
+	FScopedTransaction Transaction(EnumHasAnyFlags(NewFlags, EMovieSceneSequenceFlags::LoopCuts) ? LOCTEXT("EnableLoopCuts", "Enable Loop Cuts") : LOCTEXT("DisableLoopCuts", "Disable Loop Cuts"));
+
+	Sequence->Modify();
+	Sequence->SetSequenceFlags(NewFlags);
+}
+
+bool FSequencer::UsesLoopCuts()
+{
+	return EnumHasAnyFlags(GetFocusedMovieSceneSequence()->GetFlags(), EMovieSceneSequenceFlags::LoopCuts);
+}
+
 UE::Sequencer::FSequencerSelection& FSequencer::GetSelection()
 {
 	return *ViewModel->GetSelection();
diff --git a/Engine/Source/Editor/Sequencer/Private/Sequencer.h b/Engine/Source/Editor/Sequencer/Private/Sequencer.h
--- a/Engine/Source/Editor/Sequencer/Private/Sequencer.h
+++ b/Engine/Source/Editor/Sequencer/Private/Sequencer.h
@@ -1176,6 +1176,9 @@ private:
 	void ToggleDynamicWeighting();
 	bool UsesDynamicWeighting();
 
+	void ToggleLoopCuts();
+	bool UsesLoopCuts();
+
 	void UpdateCachedPlaybackContextAndClient();
 
 	void UpdateCachedCameraActors();
diff --git a/Engine/Source/Editor/Sequencer/Private/SequencerCommands.cpp b/Engine/Source/Editor/Sequencer/Private/SequencerCommands.cpp
--- a/Engine/Source/Editor/Sequencer/Private/SequencerCommands.cpp
+++ b/Engine/Source/Editor/Sequencer/Private/SequencerCommands.cpp
@@ -54,6 +54,7 @@ void FSequencerCommands::RegisterCommands()
 	UI_COMMAND( ToggleRerunConstructionScripts, "Allow Construction Scripts to Rerun Every Frame", "Blueprint Actors need to have Run Construction Script in Sequencer active in their Class Settings.", EUserInterfaceActionType::ToggleButton, FInputChord() );
 	UI_COMMAND( ToggleAsyncEvaluation, "Async Evaluation", "When enabled, enables a single asynchronous evaluation once per-frame. When disabled, forces a full blocking evaluation every time this sequence is evaluated (should be avoided for real-time content).", EUserInterfaceActionType::ToggleButton, FInputChord() );
 	UI_COMMAND( ToggleDynamicWeighting, "Dynamic Weighting", "When enabled, all blendable tracks will cache their initial values to ensure that they are able to correctly blend in/out when dynamic weights are being used.", EUserInterfaceActionType::ToggleButton, FInputChord() );
+	UI_COMMAND( ToggleLoopCuts, "Loop Cuts", "When enabled, the sequence considers looping as a jump/cut. When disabled, the sequence considers looping to be seamless.", EUserInterfaceActionType::ToggleButton, FInputChord());
 
 	UI_COMMAND( ToggleResetPlayheadWhenNavigating, "Reset Playhead When Navigating", "When checked, if the playhead is outside the playback range, it will be reset to the beginning when navigating in and out of subsequences", EUserInterfaceActionType::ToggleButton, FInputChord());
 	UI_COMMAND( ToggleKeepCursorInPlaybackRangeWhileScrubbing, "Keep Playhead in Playback Range While Scrubbing", "When checked, the playhead will be constrained to the current playback (or subsequence/shot) range while scrubbing", EUserInterfaceActionType::ToggleButton, FInputChord() );
diff --git a/Engine/Source/Editor/Sequencer/Public/SequencerCommands.h b/Engine/Source/Editor/Sequencer/Public/SequencerCommands.h
--- a/Engine/Source/Editor/Sequencer/Public/SequencerCommands.h
+++ b/Engine/Source/Editor/Sequencer/Public/SequencerCommands.h
@@ -144,6 +144,9 @@ public:
 	/** When enabled, all blendable tracks will always cache their initial values to ensure that they are able to correctly blend in/out when dynamic weights are being used. */
 	TSharedPtr< FUICommandInfo > ToggleDynamicWeighting;
 
+	/** When enabled, the sequence considers looping as a jump/cut. When disabled, the sequence considers looping to be seamless. */
+	TSharedPtr< FUICommandInfo > ToggleLoopCuts;
+
 	/** Toggle resetting the playhead when navigating in and out of subsequences */
 	TSharedPtr< FUICommandInfo > ToggleResetPlayheadWhenNavigating;
 
diff --git a/Engine/Source/Runtime/MovieScene/Private/Evaluation/MovieScenePlaybackManager.cpp b/Engine/Source/Runtime/MovieScene/Private/Evaluation/MovieScenePlaybackManager.cpp
--- a/Engine/Source/Runtime/MovieScene/Private/Evaluation/MovieScenePlaybackManager.cpp
+++ b/Engine/Source/Runtime/MovieScene/Private/Evaluation/MovieScenePlaybackManager.cpp
@@ -192,6 +192,7 @@ void FMovieScenePlaybackManager::InternalUpdateToTick(const FFrameNumber NextTic
 					OutContexts.Add(
 							FMovieSceneContext(PlaybackPosition.PlayTo(LoopLastTick), PlaybackStatus)
 							.SetHasJumped(true)
+							.SetHasLooped(true)
 							);
 				}
 			}
@@ -207,6 +208,7 @@ void FMovieScenePlaybackManager::InternalUpdateToTick(const FFrameNumber NextTic
 						OutContexts.Add(
 								FMovieSceneContext(PlaybackPosition.PlayTo(LastValidTick), PlaybackStatus)
 								.SetHasJumped(true)
+								.SetHasLooped(true)
 								);
 					}
 					else
@@ -215,6 +217,7 @@ void FMovieScenePlaybackManager::InternalUpdateToTick(const FFrameNumber NextTic
 						OutContexts.Add(
 								FMovieSceneContext(PlaybackPosition.PlayTo(EffectiveStartTick), PlaybackStatus)
 								.SetHasJumped(true)
+								.SetHasLooped(true)
 								);
 					}
 
@@ -272,6 +275,7 @@ void FMovieScenePlaybackManager::InternalUpdateToTick(const FFrameNumber NextTic
 					OutContexts.Add(
 							FMovieSceneContext(PlaybackPosition.PlayTo(LastLoopLastTick), PlaybackStatus)
 							.SetHasJumped(true)
+							.SetHasLooped(true)
 							);
 				}
 			}
@@ -283,6 +287,7 @@ void FMovieScenePlaybackManager::InternalUpdateToTick(const FFrameNumber NextTic
 				OutContexts.Add(
 						FMovieSceneContext(PlaybackPosition.PlayTo(LastLoopLastTick), PlaybackStatus)
 						.SetHasJumped(true)
+						.SetHasLooped(true)
 						);
 			}
 
@@ -318,6 +323,7 @@ void FMovieScenePlaybackManager::InternalUpdateToTick(const FFrameNumber NextTic
 			OutContexts.Add(
 					FMovieSceneContext(PlaybackPosition.PlayTo(EffectiveOverplayTick), PlaybackStatus)
 					.SetHasJumped(true)
+					.SetHasLooped(true)
 					);
 
 			// If the overplay leads us at/beyond the last tick of the last loop, let's count that
diff --git a/Engine/Source/Runtime/MovieScene/Private/MovieSceneSequencePlayer.cpp b/Engine/Source/Runtime/MovieScene/Private/MovieSceneSequencePlayer.cpp
--- a/Engine/Source/Runtime/MovieScene/Private/MovieSceneSequencePlayer.cpp
+++ b/Engine/Source/Runtime/MovieScene/Private/MovieSceneSequencePlayer.cpp
@@ -1322,7 +1322,8 @@ void UMovieSceneSequencePlayer::UpdateTimeCursorPosition_Internal(FFrameTime New
 			FMovieSceneEvaluationRange Range = PlayPosition.PlayTo(StartTimeWithReversed + NewFrameOffset);
 
 			const bool bHasJumped = true;
-			UpdateMovieSceneInstance(Range, StatusOverride, bHasJumped);
+			const bool bHasLooped = true;
+			UpdateMovieSceneInstance(Range, StatusOverride, bHasJumped, bHasLooped);
 
 			// Use the exact time here rather than a frame locked time to ensure we don't skip the amount that was overplayed in the time controller
 			FQualifiedFrameTime ExactCurrentTime(StartTimeWithReversed + NewFrameOffset, PlayPosition.GetInputRate());
@@ -1422,10 +1423,11 @@ void UMovieSceneSequencePlayer::UpdateTimeCursorPosition_Internal(FFrameTime New
 	// block... see `UpdateNetworkSyncProperties` as an example.
 }
 
-void UMovieSceneSequencePlayer::UpdateMovieSceneInstance(FMovieSceneEvaluationRange InRange, EMovieScenePlayerStatus::Type PlayerStatus, bool bHasJumped)
+void UMovieSceneSequencePlayer::UpdateMovieSceneInstance(FMovieSceneEvaluationRange InRange, EMovieScenePlayerStatus::Type PlayerStatus, bool bHasJumped, bool bHasLooped)
 {
 	FMovieSceneUpdateArgs Args;
 	Args.bHasJumped = bHasJumped;
+	Args.bHasLooped = bHasLooped;
 	UpdateMovieSceneInstance(InRange, PlayerStatus, Args);
 }
 
@@ -1467,6 +1469,7 @@ void UMovieSceneSequencePlayer::UpdateMovieSceneInstance(FMovieSceneEvaluationRa
 
 	FMovieSceneContext Context(InRange, PlayerStatus);
 	Context.SetHasJumped(Args.bHasJumped);
+	Context.SetHasLooped(Args.bHasLooped);
 
 	TSharedPtr<FMovieSceneEntitySystemRunner> Runner = RootTemplateInstance.GetRunner();
 	if (Runner)
diff --git a/Engine/Source/Runtime/MovieScene/Public/Evaluation/MovieScenePlayback.h b/Engine/Source/Runtime/MovieScene/Public/Evaluation/MovieScenePlayback.h
--- a/Engine/Source/Runtime/MovieScene/Public/Evaluation/MovieScenePlayback.h
+++ b/Engine/Source/Runtime/MovieScene/Public/Evaluation/MovieScenePlayback.h
@@ -191,6 +191,7 @@ struct FMovieSceneContext : FMovieSceneEvaluationRange
 		, PrePostRollStartEndTime(TNumericLimits<int32>::Lowest())
 		, HierarchicalBias(0)
 		, bHasJumped(false)
+		, bHasLooped(false)
 		, bSilent(false)
 		, bSectionPreRoll(false)
 		, bSectionPostRoll(false)
@@ -207,6 +208,7 @@ struct FMovieSceneContext : FMovieSceneEvaluationRange
 		, PrePostRollStartEndTime(TNumericLimits<int32>::Lowest())
 		, HierarchicalBias(0)
 		, bHasJumped(false)
+		, bHasLooped(false)
 		, bSilent(false)
 		, bSectionPreRoll(false)
 		, bSectionPostRoll(false)
@@ -246,6 +248,15 @@ struct FMovieSceneContext : FMovieSceneEvaluationRange
 		return bHasJumped;
 	}
 
+	/**
+	 * Check whether we've just looped around the sequence
+	 * If this returns true, HasJumped() usually also returns true
+	 */
+	inline bool HasLooped() const
+	{
+		return bHasLooped;
+	}
+
 	/**
 	 * Check whether we're evaluating in silent mode (no audio or mutating eval)
 	 */
@@ -312,6 +323,16 @@ public:
 		return *this;
 	}
 
+	/**
+	 * Indicate that we've just looped around the sequence
+	 * The caller should also probably call SetHasJumped(true)
+	 */
+	FMovieSceneContext& SetHasLooped(bool bInHasLooped)
+	{
+		bHasLooped = bInHasLooped;
+		return *this;
+	}
+
 	/**
 	 * Set the context to silent mode
 	 */
@@ -464,6 +485,9 @@ protected:
 	/** Whether this evaluation frame is happening as part of a large jump */
 	bool bHasJumped : 1;
 
+	/** Whether this evaluation frame is happening as part of looping around the sequence */
+	bool bHasLooped : 1;
+
 	/** Whether this evaluation should happen silently */
 	bool bSilent : 1;
 
diff --git a/Engine/Source/Runtime/MovieScene/Public/MovieSceneFwd.h b/Engine/Source/Runtime/MovieScene/Public/MovieSceneFwd.h
--- a/Engine/Source/Runtime/MovieScene/Public/MovieSceneFwd.h
+++ b/Engine/Source/Runtime/MovieScene/Public/MovieSceneFwd.h
@@ -116,6 +116,11 @@ enum class EMovieSceneSequenceFlags : uint8
 	 */
 	DynamicWeighting = 1 << 2,
 
+	/**
+	 * When enabled, the sequence considers looping as a jump/cut. When disabled, the sequence considers looping to be seamless.
+	 */
+	LoopCuts = 1 << 3,
+
 	/** Symbolic entry for all flags that should be inherited by parent sequences when present on a sub sequence */
 	InheritedFlags = Volatile UMETA(Hidden),
 };
@@ -163,4 +168,4 @@ MOVIESCENE_API FFrameNumber UpgradeLegacyMovieSceneTime(UObject* ErrorContext, F
 	#define UE_MOVIESCENE_TODO(MSG) UE_MOVIESCENE_TODO_IMPL(message("" __FILE__ "(" PREPROCESSOR_TO_STRING(__LINE__) "): warning TODO: " # MSG))
 #else
 	#define UE_MOVIESCENE_TODO(MSG)
-#endif
\ No newline at end of file
+#endif
diff --git a/Engine/Source/Runtime/MovieScene/Public/MovieSceneSequencePlayer.h b/Engine/Source/Runtime/MovieScene/Public/MovieSceneSequencePlayer.h
--- a/Engine/Source/Runtime/MovieScene/Public/MovieSceneSequencePlayer.h
+++ b/Engine/Source/Runtime/MovieScene/Public/MovieSceneSequencePlayer.h
@@ -529,10 +529,11 @@ protected:
 	struct FMovieSceneUpdateArgs
 	{
 		bool bHasJumped = false;
+		bool bHasLooped = false;
 		bool bIsAsync = false;
 	};
 
-	MOVIESCENE_API void UpdateMovieSceneInstance(FMovieSceneEvaluationRange InRange, EMovieScenePlayerStatus::Type PlayerStatus, bool bHasJumped = false);
+	MOVIESCENE_API void UpdateMovieSceneInstance(FMovieSceneEvaluationRange InRange, EMovieScenePlayerStatus::Type PlayerStatus, bool bHasJumped = false, bool bHasLooped = false);
 	MOVIESCENE_API virtual void UpdateMovieSceneInstance(FMovieSceneEvaluationRange InRange, EMovieScenePlayerStatus::Type PlayerStatus, const FMovieSceneUpdateArgs& Args);
 
 	MOVIESCENE_API void UpdateTimeCursorPosition(FFrameTime NewPosition, EUpdatePositionMethod Method, bool bHasJumpedOverride = false);
diff --git a/Engine/Source/Runtime/MovieSceneTracks/Private/TrackInstances/MovieSceneCameraCutTrackInstance.cpp b/Engine/Source/Runtime/MovieSceneTracks/Private/TrackInstances/MovieSceneCameraCutTrackInstance.cpp
--- a/Engine/Source/Runtime/MovieSceneTracks/Private/TrackInstances/MovieSceneCameraCutTrackInstance.cpp
+++ b/Engine/Source/Runtime/MovieSceneTracks/Private/TrackInstances/MovieSceneCameraCutTrackInstance.cpp
@@ -238,8 +238,13 @@ public:
 		FCameraCutEditorHandler::CachePreAnimatedValue(Linker, SequenceInstance);
 #endif
 
+		// We do a straight cut if we have jumped in time and we either (1) always cut on loops or (2) it's not a loop.
+		UMovieSceneSequence* RootSequence = SequenceInstance.GetSharedPlaybackState()->GetRootSequence();
+		const bool bDoLoopCuts = RootSequence ? EnumHasAnyFlags(RootSequence->GetFlags(), EMovieSceneSequenceFlags::LoopCuts) : false;
+		const bool bIsJumpCut = Context.HasJumped() && (bDoLoopCuts || !Context.HasLooped());
+
 		FMovieSceneCameraCutParams CameraCutParams;
-		CameraCutParams.bJumpCut = Context.HasJumped();
+		CameraCutParams.bJumpCut = bIsJumpCut;
 		CameraCutParams.BlendTime = Params.EaseIn.RootBlendTime;
 		CameraCutParams.BlendType = Params.EaseIn.BlendType;
 		CameraCutParams.bLockPreviousCamera = Params.bLockPreviousCamera;
