Sequencer: Fix FPS drop if mouse is moved in viewport during playback, e.g. during Take Recording.

The cause was that every frame, hit proxies were invalidated due to the FEditorViewportClient::Invalidate call.
Invalidating hit proxies is expensive because it causes the game thread to call FlushRenderingCommands (FViewport::GetRawHitProxyData).
Hit proxies are still invalidated automatically by FSlateApplication when it processes a mouse event.

diff --git a/Engine/Source/Editor/Sequencer/Private/LevelEditorSequencerIntegration.cpp b/Engine/Source/Editor/Sequencer/Private/LevelEditorSequencerIntegration.cpp
--- a/Engine/Source/Editor/Sequencer/Private/LevelEditorSequencerIntegration.cpp
+++ b/Engine/Source/Editor/Sequencer/Private/LevelEditorSequencerIntegration.cpp
@@ -538,6 +538,12 @@ void FLevelEditorSequencerIntegration::OnNewActorsDropped(const TArray<UObject*>
 	);
 }
 
+static TAutoConsoleVariable<bool> CVarEnableHitProxyInvalidationDuringEvaluation(
+	TEXT("LevelSequence.EnableHitProxyInvalidationDuringEvaluation"),
+	false,
+	TEXT("Whether Sequencer evaluation should cause the viewport's hit proxies to be invalidated. If enabled, this causes FPS drops while moving the mouse in the viewport while playing back Sequencer.")
+	);
+
 void FLevelEditorSequencerIntegration::OnSequencerEvaluated()
 {
 	RestoreSpawnablePilotData();
@@ -560,7 +566,16 @@ void FLevelEditorSequencerIntegration::OnSequencerEvaluated()
 			{
 				LevelVC->RequestRealTimeFrames(1);
 			}
-			LevelVC->Invalidate();
+
+			// Summary: When Sequencer evaluates moving the mouse in the viewport caused a drop in FPS.
+			// This fix is for 5.7, related to UE-320241.
+			// This CVar exists only for safety, since this change is coming in late to 5.7. In 5.8, it should be removed (replace "bInvalidateHitProxies" with "false").
+			// The cause WAS that every frame, hit proxies were invalidated due to the below FEditorViewportClient::Invalidate call.
+			// Invalidating hit proxies is expensive because it causes the game thread to call FlushRenderingCommands (FViewport::GetRawHitProxyData).
+			// After evaluation, it should be enough just to redraw the viewport (the rendered content) and skip the hit proxies. 
+			// Hit proxies are still invalidated automatically by FSlateApplication when it processes a mouse event.
+			const bool bInvalidateHitProxies = CVarEnableHitProxyInvalidationDuringEvaluation.GetValueOnAnyThread();
+			LevelVC->Invalidate(true, bInvalidateHitProxies);
 		}
 	}
 
