Sequencer: Don't reset transient flags for decorator objects that are Transient

#rnx

diff --git a/Engine/Source/Editor/Sequencer/Private/SequencerUtilities.cpp b/Engine/Source/Editor/Sequencer/Private/SequencerUtilities.cpp
--- a/Engine/Source/Editor/Sequencer/Private/SequencerUtilities.cpp
+++ b/Engine/Source/Editor/Sequencer/Private/SequencerUtilities.cpp
@@ -101,7 +101,11 @@ static void ResetCopiedTracksFlags(UMovieSceneTrack* Track)
 	Track->ClearFlags(RF_Transient);
 
 	ForEachObjectWithOuter(Track, [](UObject* InObject) {
-		InObject->ClearFlags(RF_Transient);
+
+		if (!InObject->GetClass() || !InObject->GetClass()->HasAnyClassFlags(CLASS_Transient))
+		{
+			InObject->ClearFlags(RF_Transient);
+		}
 	});
 
 	for (UMovieSceneSection* Section : Track->GetAllSections())
@@ -110,7 +114,6 @@ static void ResetCopiedTracksFlags(UMovieSceneTrack* Track)
 	}
 }
 
-
 namespace UE::Sequencer
 {
 
diff --git a/Engine/Source/Runtime/MovieScene/Private/MovieSceneSection.cpp b/Engine/Source/Runtime/MovieScene/Private/MovieSceneSection.cpp
--- a/Engine/Source/Runtime/MovieScene/Private/MovieSceneSection.cpp
+++ b/Engine/Source/Runtime/MovieScene/Private/MovieSceneSection.cpp
@@ -1128,7 +1128,10 @@ void UMovieSceneSection::PostPaste()
 
 	for (UObject* Decoration : GetDecorations())
 	{
-		Decoration->ClearFlags(RF_Transient);
+		if (!Decoration->GetClass() || !Decoration->GetClass()->HasAnyClassFlags(CLASS_Transient))
+		{
+			Decoration->ClearFlags(RF_Transient);
+		}
 	}
 }
 
