Sequencer: Fix crash when closing Sequencer during deferred modification flush

We do this by suppressing the deferred modification flush events during the destructor.

Fixes crash: Case 35510

diff --git a/Engine/Source/Editor/Sequencer/Private/Sequencer.cpp b/Engine/Source/Editor/Sequencer/Private/Sequencer.cpp
--- a/Engine/Source/Editor/Sequencer/Private/Sequencer.cpp
+++ b/Engine/Source/Editor/Sequencer/Private/Sequencer.cpp
@@ -854,6 +854,11 @@ FSequencer::FSequencer()
 
 FSequencer::~FSequencer()
 {
+	// Suppress any deferred modification flushes during destruction.
+	// This prevents view model updates from trying to access the dying Sequencer
+	// when the deferred flush is triggered during teardown.
+
+	UE::MovieScene::FScopedSignedObjectModifySuppress SuppressFlush;
 	if (Runner)
 	{
 		Runner->QueueFinalUpdate(RootTemplateInstance.GetRootInstanceHandle());
