Sequencer: Check for movie scene before dereferencing for tick resolution

diff --git a/Engine/Source/Editor/Sequencer/Private/SSequencerSection.cpp b/Engine/Source/Editor/Sequencer/Private/SSequencerSection.cpp
--- a/Engine/Source/Editor/Sequencer/Private/SSequencerSection.cpp
+++ b/Engine/Source/Editor/Sequencer/Private/SSequencerSection.cpp
@@ -1890,8 +1890,9 @@ void SSequencerSection::Tick( const FGeometry& AllottedGeometry, const double In
 	UMovieSceneSection*       Section      = SectionInterface->GetSectionObject();
 	TSharedPtr<FSequencer>    SequencerPtr = Sequencer.Pin();
 	TSharedPtr<FSectionModel> SectionModel = WeakSectionModel.Pin();
+	UMovieScene*              MovieScene   = Section ? Section->GetTypedOuter<UMovieScene>() : nullptr;
 
-	if (!Section || !SequencerPtr || !SectionModel)
+	if (!Section || !SequencerPtr || !SectionModel || !MovieScene)
 	{
 		return;
 	}
@@ -1902,7 +1903,7 @@ void SSequencerSection::Tick( const FGeometry& AllottedGeometry, const double In
 		{
 			constexpr float MinSectionWidth = 1.f;
 
-			FTimeToPixel TimeToPixelConverter(ParentGeometry, GetSequencer().GetViewRange(), Section->GetTypedOuter<UMovieScene>()->GetTickResolution());
+			FTimeToPixel TimeToPixelConverter(ParentGeometry, GetSequencer().GetViewRange(), MovieScene->GetTickResolution());
 			const FFrameNumber SectionLength = UE::MovieScene::DiscreteSize(Section->ComputeEffectiveRange());
 			const float SectionWidth = TimeToPixelConverter.FrameDeltaToPixel(SectionLength);
 
