Sequencer: Marker Improvements

- Deselect markers when keys and sections are also deselected
- Don't move marked frames when they are locked

diff --git a/Engine/Source/Editor/Sequencer/Private/SequencerHotspots.cpp b/Engine/Source/Editor/Sequencer/Private/SequencerHotspots.cpp
--- a/Engine/Source/Editor/Sequencer/Private/SequencerHotspots.cpp
+++ b/Engine/Source/Editor/Sequencer/Private/SequencerHotspots.cpp
@@ -73,6 +73,7 @@ void FHotspotSelectionManager::ConditionallyClearSelection()
 	{
 		Selection->TrackArea.Empty();
 		Selection->KeySelection.Empty();
+		Selection->MarkedFrames.Empty();
 
 		bAddingToSelection = true;
 	}
diff --git a/Engine/Source/Editor/Sequencer/Private/Tools/EditToolDragOperations.cpp b/Engine/Source/Editor/Sequencer/Private/Tools/EditToolDragOperations.cpp
--- a/Engine/Source/Editor/Sequencer/Private/Tools/EditToolDragOperations.cpp
+++ b/Engine/Source/Editor/Sequencer/Private/Tools/EditToolDragOperations.cpp
@@ -1467,6 +1467,11 @@ void FMoveKeysAndSections::HandleMarkedFrameMovement(TOptional<FFrameNumber> Max
 
 	const FFrameNumber EffectiveDelta = MaxDeltaX.Get(DesiredDeltaX);
 	UMovieScene* FocusedMovieScene = Sequencer.GetFocusedMovieSceneSequence()->GetMovieScene();
+	if (!FocusedMovieScene || FocusedMovieScene->AreMarkedFramesLocked())
+	{
+		return;
+	}
+
 	const TArray<FMovieSceneMarkedFrame>& AllMarkedFrames = FocusedMovieScene->GetMarkedFrames();
 
 	for (TSet<int32>::TConstIterator It = MarkedFrames.CreateConstIterator(); It; ++It)
