Anim In Engine: Call Pause directly when stopping hot key scrubbing so that Audio stops playing and all of the proper delegates/evaluations are fired.

diff --git a/Engine/Source/Editor/Sequencer/Private/SequencerEdMode.cpp b/Engine/Source/Editor/Sequencer/Private/SequencerEdMode.cpp
--- a/Engine/Source/Editor/Sequencer/Private/SequencerEdMode.cpp
+++ b/Engine/Source/Editor/Sequencer/Private/SequencerEdMode.cpp
@@ -409,7 +409,22 @@ bool FSequencerEdMode::ProcessCapturedMouseMoves(FEditorViewportClient* InViewpo
 {
 	const TSharedPtr<ISequencer> ActiveSequencer = GetFirstActiveSequencer();
 	const bool bTimeChange = IsPressingMoveTimeSlider(InViewport);
-
+	auto SetNextPlayerStatus = [&ActiveSequencer](EMovieScenePlayerStatus::Type InNextPlayerStatus)
+	{
+		if (ActiveSequencer.IsValid())
+		{
+			if (InNextPlayerStatus == EMovieScenePlayerStatus::Stopped)
+			{
+				//Call Pause() instead of SetPlaybackStatus(EMovieScenePlayerStatus::Stopped)
+				//so the correct delegates/evals can happen so all systems like audio are probably stopped (Pause set's Stopped status).
+				ActiveSequencer->Pause();
+			}
+			else
+			{
+				ActiveSequencer->SetPlaybackStatus(InNextPlayerStatus);
+			}
+		}
+	};
 	if (CapturedMouseMoves.Num() > 0)
 	{
 		if (ActiveSequencer.IsValid())
@@ -457,7 +472,7 @@ bool FSequencerEdMode::ProcessCapturedMouseMoves(FEditorViewportClient* InViewpo
 			else if(bIsTracking)
 			{
 				bIsTracking = false;
-				ActiveSequencer->SetPlaybackStatus(NextPlayerStatus);
+				SetNextPlayerStatus(NextPlayerStatus);
 				ActiveSequencer->OnEndScrubbingEvent().Broadcast();
 				StartXValue.Reset();
 			}
@@ -474,7 +489,7 @@ bool FSequencerEdMode::ProcessCapturedMouseMoves(FEditorViewportClient* InViewpo
 		if (ActiveSequencer.IsValid())
 		{
 			bIsTracking = false;
-			ActiveSequencer->SetPlaybackStatus(NextPlayerStatus);
+			SetNextPlayerStatus(NextPlayerStatus);
 			ActiveSequencer->OnEndScrubbingEvent().Broadcast();
 			StartXValue.Reset();
 		}
diff --git a/Engine/Source/Editor/Sequencer/Private/SequencerTimeSliderController.cpp b/Engine/Source/Editor/Sequencer/Private/SequencerTimeSliderController.cpp
--- a/Engine/Source/Editor/Sequencer/Private/SequencerTimeSliderController.cpp
+++ b/Engine/Source/Editor/Sequencer/Private/SequencerTimeSliderController.cpp
@@ -1814,7 +1814,9 @@ void FSequencerTimeSliderController::SetPlaybackStatus(ETimeSliderPlaybackStatus
 		default:
 		case ETimeSliderPlaybackStatus::Stopped:
 		{
-			WeakSequencer.Pin()->SetPlaybackStatus(EMovieScenePlayerStatus::Stopped);
+			//Call Pause() instead of SetPlaybackStatus(EMovieScenePlayerStatus::Stopped)
+			//so the correct delegates/evals can happen so all systems like audio are probably stopped (Pause set's Stopped status).
+			WeakSequencer.Pin()->Pause();
 			break;
 		}
 	}
diff --git a/Engine/Source/Editor/Sequencer/Private/Tools/SequencerEditTool_Selection.cpp b/Engine/Source/Editor/Sequencer/Private/Tools/SequencerEditTool_Selection.cpp
--- a/Engine/Source/Editor/Sequencer/Private/Tools/SequencerEditTool_Selection.cpp
+++ b/Engine/Source/Editor/Sequencer/Private/Tools/SequencerEditTool_Selection.cpp
@@ -158,7 +158,16 @@ public:
 	virtual void OnEndDrag(const FPointerEvent& MouseEvent, FVector2D LocalMousePos, const UE::Sequencer::FVirtualTrackArea& VirtualTrackArea) override
 	{
 		FSlateThrottleManager::Get().DisableThrottle(false);
-		Sequencer.SetPlaybackStatus(NextPlayerStatus);
+		if (NextPlayerStatus == EMovieScenePlayerStatus::Stopped)
+		{
+			//Call Pause() instead of SetPlaybackStatus(EMovieScenePlayerStatus::Stopped)
+			//so the correct delegates/evals can happen so all systems like audio are probably stopped (Pause set's Stopped status).
+			Sequencer.Pause();
+		}
+		else
+		{
+			Sequencer.SetPlaybackStatus(NextPlayerStatus);
+		}
 	}
 
 	virtual int32 OnPaint(const FGeometry& AllottedGeometry, const FSlateRect& MyCullingRect, FSlateWindowElementList& OutDrawElements, int32 LayerId) const override
