[Sequencer] Filtering:
- Add deactive text expression filter
- ensure() there are no duplicate keys for text expression filters

diff --git a/Engine/Source/Editor/Sequencer/Private/Filters/Filters/SequencerTrackFilter_Text.cpp b/Engine/Source/Editor/Sequencer/Private/Filters/Filters/SequencerTrackFilter_Text.cpp
--- a/Engine/Source/Editor/Sequencer/Private/Filters/Filters/SequencerTrackFilter_Text.cpp
+++ b/Engine/Source/Editor/Sequencer/Private/Filters/Filters/SequencerTrackFilter_Text.cpp
@@ -13,6 +13,7 @@
 #include "Filters/TextExpressions/SequencerTextFilterExpression_ConditionFunc.h"
 #include "Filters/TextExpressions/SequencerTextFilterExpression_ConditionPasses.h"
 #include "Filters/TextExpressions/SequencerTextFilterExpression_CustomBinding.h"
+#include "Filters/TextExpressions/SequencerTextFilterExpression_Deactivated.h"
 #include "Filters/TextExpressions/SequencerTextFilterExpression_EmptyBinding.h"
 #include "Filters/TextExpressions/SequencerTextFilterExpression_Group.h"
 #include "Filters/TextExpressions/SequencerTextFilterExpression_Keyed.h"
@@ -59,6 +60,7 @@ FSequencerTrackFilter_Text::FSequencerTrackFilter_Text(ISequencerTrackFilters& I
 	TextFilterExpressionContexts.Add(MakeShared<FSequencerTextFilterExpression_Modified>(InFilterInterface));
 	TextFilterExpressionContexts.Add(MakeShared<FSequencerTextFilterExpression_Time>(InFilterInterface));
 	TextFilterExpressionContexts.Add(MakeShared<FSequencerTextFilterExpression_Locked>(InFilterInterface));
+	TextFilterExpressionContexts.Add(MakeShared<FSequencerTextFilterExpression_Deactivated>(InFilterInterface));
 	TextFilterExpressionContexts.Add(MakeShared<FSequencerTextFilterExpression_Muted>(InFilterInterface));
 	TextFilterExpressionContexts.Add(MakeShared<FSequencerTextFilterExpression_Soloed>(InFilterInterface));
 	TextFilterExpressionContexts.Add(MakeShared<FSequencerTextFilterExpression_Tag>(InFilterInterface));
@@ -67,6 +69,8 @@ FSequencerTrackFilter_Text::FSequencerTrackFilter_Text(ISequencerTrackFilters& I
 	TextFilterExpressionContexts.Add(MakeShared<FSequencerTextFilterExpression_CustomBinding>(InFilterInterface));
 	TextFilterExpressionContexts.Add(MakeShared<FSequencerTextFilterExpression_EmptyBinding>(InFilterInterface));
 
+	TSet<FName> SeenKeys;
+
 	// Add global user-defined text expressions
 	for (TObjectIterator<USequencerTrackFilterTextExpressionExtension> ExtensionIt(RF_NoFlags); ExtensionIt; ++ExtensionIt)
 	{
@@ -81,6 +85,14 @@ FSequencerTrackFilter_Text::FSequencerTrackFilter_Text(ISequencerTrackFilters& I
 			for (const TSharedRef<FSequencerTextFilterExpressionContext>& TextExpression : ExtendedTextExpressions)
 			{
 				TextFilterExpressionContexts.Add(TextExpression);
+
+				// Make sure there are no duplicate keys
+				const TSet<FName> TextExpressionKeys = TextExpression->GetKeys();
+				for (const FName Key : TextExpressionKeys)
+				{
+					ensureMsgf(!SeenKeys.Contains(Key), TEXT("Sequencer text expression filter key already exists: %s"), *Key.ToString());
+				}
+				SeenKeys.Append(TextExpressionKeys);
 			}
 		}
 	}
diff --git a/Engine/Source/Editor/Sequencer/Private/Filters/TextExpressions/SequencerTextFilterExpression_Deactivated.cpp b/Engine/Source/Editor/Sequencer/Private/Filters/TextExpressions/SequencerTextFilterExpression_Deactivated.cpp
new file mode 100644
--- /dev/null
+++ b/Engine/Source/Editor/Sequencer/Private/Filters/TextExpressions/SequencerTextFilterExpression_Deactivated.cpp
@@ -0,0 +1,54 @@
+﻿// Copyright Epic Games, Inc. All Rights Reserved.
+
+#include "SequencerTextFilterExpression_Deactivated.h"
+
+#include "MVVM/Extensions/IDeactivatableExtension.h"
+#include "MVVM/Extensions/IMutableExtension.h"
+#include "MVVM/ViewModels/ViewModelIterators.h"
+
+using namespace UE::Sequencer;
+
+#define LOCTEXT_NAMESPACE "SequencerTextFilterExpression_Deactivated"
+
+FSequencerTextFilterExpression_Deactivated::FSequencerTextFilterExpression_Deactivated(ISequencerTrackFilters& InFilterInterface)
+	: FSequencerTextFilterExpressionContext(InFilterInterface)
+{
+}
+
+TSet<FName> FSequencerTextFilterExpression_Deactivated::GetKeys() const
+{
+	return { TEXT("Deactive"), TEXT("Deactived"), TEXT("Deactivate"), TEXT("Deactivated"), "Mute" };
+}
+
+ESequencerTextFilterValueType FSequencerTextFilterExpression_Deactivated::GetValueType() const
+{
+	return ESequencerTextFilterValueType::Boolean;
+}
+
+FText FSequencerTextFilterExpression_Deactivated::GetDescription() const
+{
+	return LOCTEXT("ExpressionDescription_Deactivated", "Filter by track deactivated state");
+}
+
+bool FSequencerTextFilterExpression_Deactivated::TestComplexExpression(const FName& InKey
+	, const FTextFilterString& InValue
+	, const ETextFilterComparisonOperation InComparisonOperation
+	, const ETextFilterTextComparisonMode InTextComparisonMode) const
+{
+	if (!FSequencerTextFilterExpressionContext::TestComplexExpression(InKey, InValue, InComparisonOperation, InTextComparisonMode))
+	{
+		return true;
+	}
+
+	const TViewModelPtr<IDeactivatableExtension> DeactivatableExtension
+		= FilterItem->FindAncestorOfType<IDeactivatableExtension>(/*bIncludeThis=*/true);
+	if (!DeactivatableExtension.IsValid())
+	{
+		return false;
+	}
+
+	const bool bPassed = DeactivatableExtension->IsDeactivated();
+	return CompareFStringForExactBool(InValue, InComparisonOperation, bPassed);
+}
+
+#undef LOCTEXT_NAMESPACE
diff --git a/Engine/Source/Editor/Sequencer/Private/Filters/TextExpressions/SequencerTextFilterExpression_Deactivated.h b/Engine/Source/Editor/Sequencer/Private/Filters/TextExpressions/SequencerTextFilterExpression_Deactivated.h
new file mode 100644
--- /dev/null
+++ b/Engine/Source/Editor/Sequencer/Private/Filters/TextExpressions/SequencerTextFilterExpression_Deactivated.h
@@ -0,0 +1,24 @@
+﻿// Copyright Epic Games, Inc. All Rights Reserved.
+
+#pragma once
+
+#include "Filters/SequencerTextFilterExpressionContext.h"
+
+class FSequencerTextFilterExpression_Deactivated : public FSequencerTextFilterExpressionContext
+{
+public:
+	FSequencerTextFilterExpression_Deactivated(ISequencerTrackFilters& InFilterInterface);
+
+	//~ Begin FSequencerTextFilterExpressionContext
+	virtual TSet<FName> GetKeys() const override;
+	virtual ESequencerTextFilterValueType GetValueType() const override;
+	virtual FText GetDescription() const override;
+	//~ End FSequencerTextFilterExpressionContext
+
+	//~ Begin ITextFilterExpressionContext
+	virtual bool TestComplexExpression(const FName& InKey
+		, const FTextFilterString& InValue
+		, const ETextFilterComparisonOperation InComparisonOperation
+		, const ETextFilterTextComparisonMode InTextComparisonMode) const override;
+	//~ End ITextFilterExpressionContext
+};
