[Sequencer] Filtering: Fix ensured unique seen key items to not just include extended text expression filters, but check all text expression filters.

diff --git a/Engine/Source/Editor/Sequencer/Private/Filters/Filters/SequencerTrackFilter_Text.cpp b/Engine/Source/Editor/Sequencer/Private/Filters/Filters/SequencerTrackFilter_Text.cpp
--- a/Engine/Source/Editor/Sequencer/Private/Filters/Filters/SequencerTrackFilter_Text.cpp
+++ b/Engine/Source/Editor/Sequencer/Private/Filters/Filters/SequencerTrackFilter_Text.cpp
@@ -69,8 +69,6 @@ FSequencerTrackFilter_Text::FSequencerTrackFilter_Text(ISequencerTrackFilters& I
 	TextFilterExpressionContexts.Add(MakeShared<FSequencerTextFilterExpression_CustomBinding>(InFilterInterface));
 	TextFilterExpressionContexts.Add(MakeShared<FSequencerTextFilterExpression_EmptyBinding>(InFilterInterface));
 
-	TSet<FName> SeenKeys;
-
 	// Add global user-defined text expressions
 	for (TObjectIterator<USequencerTrackFilterTextExpressionExtension> ExtensionIt(RF_NoFlags); ExtensionIt; ++ExtensionIt)
 	{
@@ -82,20 +80,11 @@ FSequencerTrackFilter_Text::FSequencerTrackFilter_Text(ISequencerTrackFilters& I
 			TArray<TSharedRef<FSequencerTextFilterExpressionContext>> ExtendedTextExpressions;
 			PotentialExtension->AddTrackFilterTextExpressionExtensions(static_cast<ISequencerTrackFilters&>(FilterInterface), ExtendedTextExpressions);
 
-			for (const TSharedRef<FSequencerTextFilterExpressionContext>& TextExpression : ExtendedTextExpressions)
-			{
-				TextFilterExpressionContexts.Add(TextExpression);
-
-				// Make sure there are no duplicate keys
-				const TSet<FName> TextExpressionKeys = TextExpression->GetKeys();
-				for (const FName Key : TextExpressionKeys)
-				{
-					ensureMsgf(!SeenKeys.Contains(Key), TEXT("Sequencer text expression filter key already exists: %s"), *Key.ToString());
-				}
-				SeenKeys.Append(TextExpressionKeys);
-			}
+			TextFilterExpressionContexts.Append(ExtendedTextExpressions);
 		}
 	}
+
+	VerifyNoDuplicateKeys();
 }
 
 FText FSequencerTrackFilter_Text::GetDisplayName() const
@@ -254,4 +243,20 @@ bool FSequencerTrackFilter_Text::IsTokenValueValid(const FExpressionToken& InTok
 	return true;
 }
 
+void FSequencerTrackFilter_Text::VerifyNoDuplicateKeys() const
+{
+	TSet<FName> SeenKeys;
+
+	for (const TSharedRef<FSequencerTextFilterExpressionContext>& TextExpression : TextFilterExpressionContexts)
+	{
+		const TSet<FName> TextExpressionKeys = TextExpression->GetKeys();
+		for (const FName Key : TextExpressionKeys)
+		{
+			ensureMsgf(!SeenKeys.Contains(Key), TEXT("Sequencer text expression filter key already exists: %s"), *Key.ToString());
+		}
+
+		SeenKeys.Append(TextExpressionKeys);
+	}
+}
+
 #undef LOCTEXT_NAMESPACE
diff --git a/Engine/Source/Editor/Sequencer/Private/Filters/Filters/SequencerTrackFilter_Text.h b/Engine/Source/Editor/Sequencer/Private/Filters/Filters/SequencerTrackFilter_Text.h
--- a/Engine/Source/Editor/Sequencer/Private/Filters/Filters/SequencerTrackFilter_Text.h
+++ b/Engine/Source/Editor/Sequencer/Private/Filters/Filters/SequencerTrackFilter_Text.h
@@ -39,6 +39,8 @@ protected:
 	static bool IsTokenOperator(const FExpressionToken& InToken, const ESequencerTextFilterValueType InValueType);
 	static bool IsTokenValueValid(const FExpressionToken& InToken, const ESequencerTextFilterValueType InValueType);
 
+	void VerifyNoDuplicateKeys() const;
+
 	/** Expression evaluator that can be used to perform complex text filter queries */
 	FTextFilterExpressionEvaluator TextFilterExpressionEvaluator;
 
diff --git a/Engine/Source/Editor/Sequencer/Private/Filters/TextExpressions/SequencerTextFilterExpression_Deactivated.cpp b/Engine/Source/Editor/Sequencer/Private/Filters/TextExpressions/SequencerTextFilterExpression_Deactivated.cpp
--- a/Engine/Source/Editor/Sequencer/Private/Filters/TextExpressions/SequencerTextFilterExpression_Deactivated.cpp
+++ b/Engine/Source/Editor/Sequencer/Private/Filters/TextExpressions/SequencerTextFilterExpression_Deactivated.cpp
@@ -17,7 +17,7 @@ FSequencerTextFilterExpression_Deactivated::FSequencerTextFilterExpression_Deact
 
 TSet<FName> FSequencerTextFilterExpression_Deactivated::GetKeys() const
 {
-	return { TEXT("Deactive"), TEXT("Deactived"), TEXT("Deactivate"), TEXT("Deactivated"), "Mute" };
+	return { TEXT("Deactive"), TEXT("Deactived"), TEXT("Deactivate"), TEXT("Deactivated") };
 }
 
 ESequencerTextFilterValueType FSequencerTextFilterExpression_Deactivated::GetValueType() const
diff --git a/Engine/Source/Editor/Sequencer/Public/Filters/Widgets/SSequencerFilter.h b/Engine/Source/Editor/Sequencer/Public/Filters/Widgets/SSequencerFilter.h
--- a/Engine/Source/Editor/Sequencer/Public/Filters/Widgets/SSequencerFilter.h
+++ b/Engine/Source/Editor/Sequencer/Public/Filters/Widgets/SSequencerFilter.h
@@ -7,7 +7,7 @@
 class SSequencerFilterCheckBox;
 
 /**
- * Generic Sequencer filter class used by Sequencer and Navigation Tool filter bars.
+ * Generic Sequencer filter class used by Sequencer and Sequence Navigator filter bars.
  */
 class SSequencerFilter : public SCompoundWidget
 {
