[ITF Viewport Interactions] Recreate the "hold B and drag mouse to shift sequencer time" interaction in sequncer mode using ITF.

diff --git a/Engine/Source/Editor/Experimental/EditorInteractiveToolsFramework/Private/ViewportInteractions/ViewportDragInteraction.cpp b/Engine/Source/Editor/Experimental/EditorInteractiveToolsFramework/Private/ViewportInteractions/ViewportDragInteraction.cpp
--- a/Engine/Source/Editor/Experimental/EditorInteractiveToolsFramework/Private/ViewportInteractions/ViewportDragInteraction.cpp
+++ b/Engine/Source/Editor/Experimental/EditorInteractiveToolsFramework/Private/ViewportInteractions/ViewportDragInteraction.cpp
@@ -56,7 +56,10 @@ void UViewportDragInteraction::OnDragStart(const FInputDeviceRay& InDragStartPos
 
 	bIsDragging = true;
 
-	GetViewportInteractionsBehaviorSource()->SetMouseCursorOverride(EMouseCursor::None);
+	if (UViewportInteractionsBehaviorSource* BehaviorSource = GetViewportInteractionsBehaviorSource())
+	{
+		BehaviorSource->SetMouseCursorOverride(EMouseCursor::None);
+	}
 }
 
 void UViewportDragInteraction::OnDrag(const FInputDeviceRay& InDragPos)
@@ -84,6 +87,11 @@ void UViewportDragInteraction::OnDragEnd(const FInputDeviceRay& InDragEndPos)
 	MouseY = static_cast<int32>(InDragEndPos.ScreenPosition.Y);
 
 	bIsDragging = false;
+	
+	if (UViewportInteractionsBehaviorSource* BehaviorSource = GetViewportInteractionsBehaviorSource())
+	{
+		BehaviorSource->ClearMouseCursorOverride();
+	}
 }
 
 void UViewportDragInteraction::OnForceEndCapture()
diff --git a/Engine/Source/Editor/Sequencer/Private/SequencerEdMode.cpp b/Engine/Source/Editor/Sequencer/Private/SequencerEdMode.cpp
--- a/Engine/Source/Editor/Sequencer/Private/SequencerEdMode.cpp
+++ b/Engine/Source/Editor/Sequencer/Private/SequencerEdMode.cpp
@@ -49,6 +49,9 @@
 #include "IMovieScenePlaybackClient.h"
 #include "Components/PrimitiveComponent.h"
 #include "SequencerSelectabilityTool.h"
+#include "Tools/EdModeInteractiveToolsContext.h"
+#include "ViewportInteractions/MoveSequencerTimeSliderViewportInteraction.h"
+#include "ViewportInteractions/ViewportInteractionsBehaviorSource.h"
 
 const FEditorModeID FSequencerEdMode::EM_SequencerMode(TEXT("EM_SequencerMode"));
 
@@ -205,6 +208,13 @@ void FSequencerEdMode::Enter()
 	bIsTracking = false;
 	NextPlayerStatus = EMovieScenePlayerStatus::Stopped;
 	StartXValue.Reset();
+	
+	if (UViewportInteractionsBehaviorSource* InteractionSource = GetModeManager()->GetInteractiveToolsContext()->GetViewportInteractionsBehaviorSource())
+	{
+		UMoveSequencerTimeSliderViewportInteraction* MoveInteraction = InteractionSource->AddInteraction<UMoveSequencerTimeSliderViewportInteraction>(true);
+		MoveInteraction->EdMode = this;
+	}
+	
 	FEdMode::Enter();
 }
 
diff --git a/Engine/Source/Editor/Sequencer/Private/ViewportInteractions/MoveSequencerTimeSliderViewportInteraction.cpp b/Engine/Source/Editor/Sequencer/Private/ViewportInteractions/MoveSequencerTimeSliderViewportInteraction.cpp
new file mode 100644
--- /dev/null
+++ b/Engine/Source/Editor/Sequencer/Private/ViewportInteractions/MoveSequencerTimeSliderViewportInteraction.cpp
@@ -0,0 +1,138 @@
+﻿// Copyright Epic Games, Inc. All Rights Reserved.
+
+
+#include "MoveSequencerTimeSliderViewportInteraction.h"
+
+#include "AnimatedRange.h"
+#include "ISequencer.h"
+#include "SequencerCommands.h"
+#include "SequencerEdMode.h"
+#include "EditorViewportClient.h"
+
+UMoveSequencerTimeSliderViewportInteraction::UMoveSequencerTimeSliderViewportInteraction()
+{
+	if (ensure(ViewportClickDragBehavior.IsValid()))
+	{
+		ViewportClickDragBehavior->SetDefaultPriority(FInputCapturePriority::DEFAULT_TOOL_PRIORITY);
+		ViewportClickDragBehavior->SetDragConfirmedPriority(FInputCapturePriority::DEFAULT_TOOL_PRIORITY);
+	
+		ViewportClickDragBehavior->SetTriggerMode(UViewportClickDragBehavior::ETriggerMode::Down);
+		ViewportClickDragBehavior->SetModifierMode(UViewportClickDragBehavior::EModifierMode::Continuous);
+		ViewportClickDragBehavior->SetRequireAllButtons(false);
+	
+		ViewportClickDragBehavior->EnableButton(EKeys::LeftMouseButton);
+		ViewportClickDragBehavior->EnableButton(EKeys::RightMouseButton);
+		ViewportClickDragBehavior->EnableButton(EKeys::MiddleMouseButton);
+	}
+}
+
+
+bool UMoveSequencerTimeSliderViewportInteraction::CanBeActivated(const FInputDeviceState& InInputDeviceState) const
+{
+	if (EdMode)
+	{
+		if (FEditorViewportClient* Client = GetEditorViewportClient())
+		{
+			if (EdMode->IsPressingMoveTimeSlider(Client->Viewport))
+			{
+				return true;
+			}
+		}
+	}
+	return false;
+}
+
+FInputRayHit UMoveSequencerTimeSliderViewportInteraction::CanBeginClickDragSequence(const FInputDeviceRay& PressPos)
+{
+	return FInputRayHit(TNumericLimits<double>::Max());
+}
+
+void UMoveSequencerTimeSliderViewportInteraction::OnDragStart(const FInputDeviceRay& InDragStartPos)
+{
+	Super::OnDragStart(InDragStartPos);
+	
+	if (!EdMode)
+	{
+		return;
+	}
+
+	if (TSharedPtr<ISequencer> Sequencer = EdMode->GetFirstActiveSequencer())
+	{
+		NextPlayerStatus = Sequencer->GetPlaybackStatus() == EMovieScenePlayerStatus::Playing
+			? EMovieScenePlayerStatus::Playing
+			: EMovieScenePlayerStatus::Stopped;
+	
+		Sequencer->SetPlaybackStatus(EMovieScenePlayerStatus::Scrubbing);
+		Sequencer->OnBeginScrubbingEvent().Broadcast();
+		
+		AccumulatedChange = 0.0f;
+		StartingFrame = Sequencer->GetLocalTime().Time;
+	}
+}
+
+void UMoveSequencerTimeSliderViewportInteraction::OnDragDelta(float InMouseDeltaX, float InMouseDeltaY)
+{
+	FEditorViewportClient* Client = GetEditorViewportClient();
+	if (!Client)
+	{
+		return;
+	}
+	
+	if (!EdMode)
+	{
+		return;
+	}
+
+	const TSharedPtr<ISequencer> ActiveSequencer = EdMode->GetFirstActiveSequencer();
+	if (!ActiveSequencer)
+	{
+		return;
+	}
+	
+	FIntPoint Origin, Size;
+	Client->GetViewportDimensions(Origin, Size);
+	if (Size.X == 0.0f)
+	{
+		return;
+	}
+	
+	const float ViewPortSize = static_cast<float>(Size.X);
+	
+	AccumulatedChange += InMouseDeltaX;
+	
+	const float FloatViewDiff = AccumulatedChange / ViewPortSize;
+
+	const FFrameRate TickResolution = ActiveSequencer->GetFocusedTickResolution();
+	const FAnimatedRange Range = ActiveSequencer->GetViewRange();
+	const TRange<FFrameNumber> ViewRange(TickResolution.AsFrameNumber(Range.GetLowerBoundValue()), TickResolution.AsFrameNumber(Range.GetUpperBoundValue()));
+	FFrameNumber FrameDiff = ViewRange.Size<FFrameNumber>() * FloatViewDiff;
+	
+	FFrameTime SnappedTime = StartingFrame + FrameDiff;
+	ActiveSequencer->SnapSequencerTime(SnappedTime);
+	ActiveSequencer->SetLocalTime(SnappedTime);
+}
+
+void UMoveSequencerTimeSliderViewportInteraction::OnDragEnd(const FInputDeviceRay& InDragEndPos)
+{
+	Super::OnDragEnd(InDragEndPos);
+	
+	if (!EdMode)
+	{
+		return;
+	}
+
+	if (const TSharedPtr<ISequencer> ActiveSequencer = EdMode->GetFirstActiveSequencer())
+	{
+		if (NextPlayerStatus == EMovieScenePlayerStatus::Stopped)
+		{
+			//Call Pause() instead of SetPlaybackStatus(EMovieScenePlayerStatus::Stopped)
+			//so the correct delegates/evals can happen so all systems like audio are probably stopped (Pause set's Stopped status).
+			ActiveSequencer->Pause();
+		}
+		else
+		{
+			ActiveSequencer->SetPlaybackStatus(NextPlayerStatus);
+		}
+		ActiveSequencer->OnEndScrubbingEvent().Broadcast();
+	}
+}
diff --git a/Engine/Source/Editor/Sequencer/Private/ViewportInteractions/MoveSequencerTimeSliderViewportInteraction.h b/Engine/Source/Editor/Sequencer/Private/ViewportInteractions/MoveSequencerTimeSliderViewportInteraction.h
new file mode 100644
--- /dev/null
+++ b/Engine/Source/Editor/Sequencer/Private/ViewportInteractions/MoveSequencerTimeSliderViewportInteraction.h
@@ -0,0 +1,35 @@
+﻿// Copyright Epic Games, Inc. All Rights Reserved.
+
+#pragma once
+
+#include "CoreMinimal.h"
+#include "UObject/Object.h"
+#include "ViewportInteractions/ViewportDragInteraction.h"
+#include "MovieSceneFwd.h"
+#include "MoveSequencerTimeSliderViewportInteraction.generated.h"
+
+class FSequencerEdMode;
+/**
+ * 
+ */
+UCLASS()
+class SEQUENCER_API UMoveSequencerTimeSliderViewportInteraction : public UViewportDragInteraction
+{
+	GENERATED_BODY()
+	
+public:
+	UMoveSequencerTimeSliderViewportInteraction();
+	
+	virtual FInputRayHit CanBeginClickDragSequence(const FInputDeviceRay& PressPos) override;
+	virtual bool CanBeActivated(const FInputDeviceState& InInputDeviceState = FInputDeviceState()) const override;
+	virtual void OnDragStart(const FInputDeviceRay& InDragStartPos) override;
+	virtual void OnDragDelta(float InMouseDeltaX, float InMouseDeltaY) override;
+	virtual void OnDragEnd(const FInputDeviceRay& InDragEndPos) override;
+	
+	FSequencerEdMode* EdMode = nullptr;
+	
+private:
+	float AccumulatedChange = 0.0f;
+	FFrameTime StartingFrame;
+	EMovieScenePlayerStatus::Type NextPlayerStatus = EMovieScenePlayerStatus::Stopped;
+};
diff --git a/Engine/Source/Editor/Sequencer/Sequencer.Build.cs b/Engine/Source/Editor/Sequencer/Sequencer.Build.cs
--- a/Engine/Source/Editor/Sequencer/Sequencer.Build.cs
+++ b/Engine/Source/Editor/Sequencer/Sequencer.Build.cs
@@ -25,6 +25,7 @@ public class Sequencer : ModuleRules
 				"CurveEditor",
 				"InputCore",
 				"Engine", 
+				"EditorInteractiveToolsFramework",
 				"Slate", 
 				"SlateCore",
 				"SceneOutliner",
