Sequencer: Add Select Markers in Selection Range

Also, select markers in Select All in Selection Range, Select Forwards, Select Backwards

diff --git a/Engine/Source/Editor/Sequencer/Private/SSequencer.cpp b/Engine/Source/Editor/Sequencer/Private/SSequencer.cpp
--- a/Engine/Source/Editor/Sequencer/Private/SSequencer.cpp
+++ b/Engine/Source/Editor/Sequencer/Private/SSequencer.cpp
@@ -1800,6 +1800,7 @@ TSharedRef<SWidget> SSequencer::MakeActionsMenu()
 		MenuBuilder.AddMenuEntry(FSequencerCommands::Get().ClearSelectionRange);
 		MenuBuilder.AddMenuEntry(FSequencerCommands::Get().SelectKeysInSelectionRange);
 		MenuBuilder.AddMenuEntry(FSequencerCommands::Get().SelectSectionsInSelectionRange);
+		MenuBuilder.AddMenuEntry(FSequencerCommands::Get().SelectMarkersInSelectionRange);
 		MenuBuilder.AddMenuEntry(FSequencerCommands::Get().SelectAllInSelectionRange);
 	}
 	MenuBuilder.EndSection();
diff --git a/Engine/Source/Editor/Sequencer/Private/Sequencer.cpp b/Engine/Source/Editor/Sequencer/Private/Sequencer.cpp
--- a/Engine/Source/Editor/Sequencer/Private/Sequencer.cpp
+++ b/Engine/Source/Editor/Sequencer/Private/Sequencer.cpp
@@ -2511,7 +2511,7 @@ void FSequencer::StretchTime(FFrameTime InDeltaTime)
 	}
 
 	FocusedMovieScene->SetSelectionRange(SelectionRange);
-	SelectInSelectionRange(true, true);
+	SelectInSelectionRange(true, true, true);
 	TransformSelectedKeysAndSections(InDeltaTime, 1.f);
 
 	// Return state
@@ -2556,7 +2556,7 @@ void FSequencer::ShrinkTime(FFrameTime InDeltaTime)
 	}
 
 	FocusedMovieScene->SetSelectionRange(CheckRange);
-	SelectInSelectionRange(true, true);
+	SelectInSelectionRange(true, true, true);
 
 	TSharedPtr<FSequencerSelection> Selection = ViewModel->GetSelection();
 	if (Selection->KeySelection.Num() > 0)
@@ -2597,7 +2597,7 @@ void FSequencer::ShrinkTime(FFrameTime InDeltaTime)
 	}
 
 	FocusedMovieScene->SetSelectionRange(SelectionRange);
-	SelectInSelectionRange(true, true);
+	SelectInSelectionRange(true, true, true);
 	TransformSelectedKeysAndSections(-InDeltaTime, 1.f);
 
 	// Return state
@@ -3041,7 +3041,7 @@ void FSequencer::ClearSelectionRange()
 	SetSelectionRange(TRange<FFrameNumber>::Empty());
 }
 
-void FSequencer::SelectInSelectionRange(bool bSelectKeys, bool bSelectSections)
+void FSequencer::SelectInSelectionRange(bool bSelectKeys, bool bSelectSections, bool bSelectMarkers)
 {
 	using namespace UE::Sequencer;
 
@@ -3056,11 +3056,35 @@ void FSequencer::SelectInSelectionRange(bool bSelectKeys, bool bSelectSections)
 	// Don't empty all selection, just keys and sections
 	Selection->KeySelection.Empty();
 	Selection->TrackArea.Empty();
+	Selection->MarkedFrames.Empty();
 
 	for (const TViewModelPtr<IOutlinerExtension>& DisplayNode : NodeTree->GetRootNodes())
 	{
 		SelectInSelectionRange(DisplayNode, SelectionRange, bSelectKeys, bSelectSections);
 	}
+
+	if (bSelectMarkers)
+	{
+		SelectMarkersInRange(SelectionRange);
+	}
+}
+
+void FSequencer::SelectMarkersInRange(const TRange<FFrameNumber>& Range)
+{
+	using namespace UE::Sequencer;
+
+	UMovieSceneSequence* Sequence = GetFocusedMovieSceneSequence();
+	UMovieScene* MovieScene = Sequence->GetMovieScene();
+
+	TSharedPtr<FSequencerSelection> Selection = ViewModel->GetSelection();
+
+	for (int32 Index = 0; Index < MovieScene->GetMarkedFrames().Num(); ++Index)
+	{
+		if (Range.Contains(MovieScene->GetMarkedFrames()[Index].FrameNumber))
+		{
+			Selection->MarkedFrames.Select(Index);
+		}
+	}
 }
 
 void FSequencer::SelectForward()
@@ -3077,6 +3101,7 @@ void FSequencer::SelectForward()
 
 	Selection->KeySelection.Empty();
 	Selection->TrackArea.Empty();
+	Selection->MarkedFrames.Empty();
 
 	if (Selection->GetNodesWithSelectedKeysOrSections().Num() != 0)
 	{
@@ -3096,6 +3121,8 @@ void FSequencer::SelectForward()
 	{
 		SelectInSelectionRange(ViewModel->GetRootModel(), SelectionRange, true, true);
 	}
+
+	SelectMarkersInRange(SelectionRange);
 }
 
 
@@ -3113,6 +3140,7 @@ void FSequencer::SelectBackward()
 
 	Selection->KeySelection.Empty();
 	Selection->TrackArea.Empty();
+	Selection->MarkedFrames.Empty();
 
 	if (Selection->GetNodesWithSelectedKeysOrSections().Num() != 0)
 	{
@@ -3132,6 +3160,8 @@ void FSequencer::SelectBackward()
 	{
 		SelectInSelectionRange(ViewModel->GetRootModel(), SelectionRange, true, true);
 	}
+
+	SelectMarkersInRange(SelectionRange);
 }
 
 
@@ -11899,17 +11929,22 @@ void FSequencer::BindCommands()
 
 	SequencerCommandBindings->MapAction(
 		Commands.SelectKeysInSelectionRange,
-		FExecuteAction::CreateSP(this, &FSequencer::SelectInSelectionRange, true, false),
+		FExecuteAction::CreateSP(this, &FSequencer::SelectInSelectionRange, true, false, false),
 		FCanExecuteAction::CreateLambda(IsSelectionRangeNonEmpty));
 
 	SequencerCommandBindings->MapAction(
 		Commands.SelectSectionsInSelectionRange,
-		FExecuteAction::CreateSP(this, &FSequencer::SelectInSelectionRange, false, true),
+		FExecuteAction::CreateSP(this, &FSequencer::SelectInSelectionRange, false, true, false),
+		FCanExecuteAction::CreateLambda(IsSelectionRangeNonEmpty));
+
+	SequencerCommandBindings->MapAction(
+		Commands.SelectMarkersInSelectionRange,
+		FExecuteAction::CreateSP(this, &FSequencer::SelectInSelectionRange, false, false, true),
 		FCanExecuteAction::CreateLambda(IsSelectionRangeNonEmpty));
 
 	SequencerCommandBindings->MapAction(
 		Commands.SelectAllInSelectionRange,
-		FExecuteAction::CreateSP(this, &FSequencer::SelectInSelectionRange, true, true),
+		FExecuteAction::CreateSP(this, &FSequencer::SelectInSelectionRange, true, true, true),
 		FCanExecuteAction::CreateLambda(IsSelectionRangeNonEmpty));
 
 	SequencerCommandBindings->MapAction(
diff --git a/Engine/Source/Editor/Sequencer/Private/Sequencer.h b/Engine/Source/Editor/Sequencer/Private/Sequencer.h
--- a/Engine/Source/Editor/Sequencer/Private/Sequencer.h
+++ b/Engine/Source/Editor/Sequencer/Private/Sequencer.h
@@ -185,13 +185,14 @@ public:
 	/** Clear and reset the selection range. */
 	void ClearSelectionRange();
 
-	/** Select all keys that fall into the current selection range. */
-	void SelectInSelectionRange(bool bSelectKeys, bool bSelectSections);
+	/** Select all that fall into the current selection range. */
+	void SelectInSelectionRange(bool bSelectKeys, bool bSelectSections, bool bSelectMarkers);
+	void SelectMarkersInRange(const TRange<FFrameNumber>& InRange);
 
-	/** Select all keys and sections forward from the current time */
+	/** Select all keys, sections and markers forward from the current time */
 	void SelectForward();
 
-	/** Select all keys and sections backward from the current time */
+	/** Select all keys, sections and markers backward from the current time */
 	void SelectBackward();
 
 	/**
diff --git a/Engine/Source/Editor/Sequencer/Private/SequencerCommands.cpp b/Engine/Source/Editor/Sequencer/Private/SequencerCommands.cpp
--- a/Engine/Source/Editor/Sequencer/Private/SequencerCommands.cpp
+++ b/Engine/Source/Editor/Sequencer/Private/SequencerCommands.cpp
@@ -74,10 +74,11 @@ void FSequencerCommands::RegisterCommands()
 	UI_COMMAND( ClearSelectionRange, "Clear Selection Range", "Clear the selection range", EUserInterfaceActionType::Button, FInputChord(EModifierKey::Control|EModifierKey::Shift, EKeys::X) );
 	UI_COMMAND( SelectKeysInSelectionRange, "Select Keys in Selection Range", "Select all keys that fall into the selection range", EUserInterfaceActionType::Button, FInputChord() );
 	UI_COMMAND( SelectSectionsInSelectionRange, "Select Sections in Selection Range", "Select all sections that fall into the selection range", EUserInterfaceActionType::Button, FInputChord() );
+	UI_COMMAND( SelectMarkersInSelectionRange, "Select Markers in Selection Range", "Select all markers that fall into the selection range", EUserInterfaceActionType::Button, FInputChord() );
 	UI_COMMAND( SelectAllInSelectionRange, "Select All in Selection Range", "Select all keys and section that fall into the selection range", EUserInterfaceActionType::Button, FInputChord() );
 
-	UI_COMMAND( SelectForward, "Select All Keys and Sections Forward", "Select all keys and sections forward from the current time", EUserInterfaceActionType::Button, FInputChord(EModifierKey::Control, EKeys::RightBracket) );
-	UI_COMMAND( SelectBackward, "Select All Keys and Sections Backward", "Select all keys and sections backward from the current time", EUserInterfaceActionType::Button, FInputChord(EModifierKey::Control, EKeys::LeftBracket) );
+	UI_COMMAND( SelectForward, "Select All Keys, Sections and Markers Forward", "Select all keys, sections and markers forward from the current time", EUserInterfaceActionType::Button, FInputChord(EModifierKey::Control, EKeys::RightBracket) );
+	UI_COMMAND( SelectBackward, "Select All Keys, Sections and Markers Backward", "Select all keys, sections and markers backward from the current time", EUserInterfaceActionType::Button, FInputChord(EModifierKey::Control, EKeys::LeftBracket) );
 	UI_COMMAND( SelectNone, "Select None", "Select none", EUserInterfaceActionType::Button, FInputChord(EKeys::Escape));
 
 	UI_COMMAND( AddActorsToSequencer, "Add Actors", "Add actors to sequencer", EUserInterfaceActionType::Button, FInputChord(EModifierKey::Control, EKeys::A) );
diff --git a/Engine/Source/Editor/Sequencer/Public/SequencerCommands.h b/Engine/Source/Editor/Sequencer/Public/SequencerCommands.h
--- a/Engine/Source/Editor/Sequencer/Public/SequencerCommands.h
+++ b/Engine/Source/Editor/Sequencer/Public/SequencerCommands.h
@@ -198,13 +198,16 @@ public:
 	/** Select all sections that fall into the selection range*/
 	TSharedPtr< FUICommandInfo > SelectSectionsInSelectionRange;
 
-	/** Select all keys and sections that fall into the selection range*/
+	/** Select all markers that fall into the selection range*/
+	TSharedPtr< FUICommandInfo > SelectMarkersInSelectionRange;
+
+	/** Select all keys, sections and markers that fall into the selection range*/
 	TSharedPtr< FUICommandInfo > SelectAllInSelectionRange;
 
-	/** Select all keys and sections forward from the current time */
+	/** Select all keys, sections and markers forward from the current time */
 	TSharedPtr< FUICommandInfo > SelectForward;
 
-	/** Select all keys and sections backward from the current time */
+	/** Select all keys, sections and markers backward from the current time */
 	TSharedPtr< FUICommandInfo > SelectBackward;
 
 	/** Select none */
