Sequencer: Selection range and start/end range are now considered markers, so snap to them also when Snap to Markers

diff --git a/Engine/Source/Editor/Sequencer/Private/Sequencer.cpp b/Engine/Source/Editor/Sequencer/Private/Sequencer.cpp
--- a/Engine/Source/Editor/Sequencer/Private/Sequencer.cpp
+++ b/Engine/Source/Editor/Sequencer/Private/Sequencer.cpp
@@ -5810,29 +5810,45 @@ FFrameNumber FSequencer::OnGetNearestKey(FFrameTime InTime, ENearestKeyOption Ne
 	
 	if (EnumHasAnyFlags(NearestKeyOption, ENearestKeyOption::NKO_SearchMarkers))
 	{
+		TArray<FFrameNumber> Markers;
+
 		TArray<FMovieSceneMarkedFrame> MarkedFrames = GetMarkedFrames();
 		for (const FMovieSceneMarkedFrame& MarkedFrame : MarkedFrames)
 		{
-			if (!NearestTime.IsSet())
-			{
-				NearestTime = MarkedFrame.FrameNumber;
-			}
-			else if (FMath::Abs(MarkedFrame.FrameNumber - CurrentTime) < FMath::Abs(NearestTime.GetValue() - CurrentTime))
-			{
-				NearestTime = MarkedFrame.FrameNumber;
-			}
+			Markers.Add(MarkedFrame.FrameNumber);
 		}
 
 		TArray<FMovieSceneMarkedFrame> GlobalMarkedFrames = GetGlobalMarkedFrames();
 		for (const FMovieSceneMarkedFrame& GlobalMarkedFrame : GlobalMarkedFrames)
+		{
+			Markers.Add(GlobalMarkedFrame.FrameNumber);
+		}
+
+		if (UMovieSceneSequence* Sequence = GetFocusedMovieSceneSequence())
+		{
+			const bool bInsetDisplayFrame = ScrubStyle == ESequencerScrubberStyle::FrameBlock && Settings->GetForceWholeFrames();
+			FFrameRate LocalResolution = GetFocusedTickResolution();
+			FFrameRate DisplayRate = GetFocusedDisplayRate();
+			FFrameTime OneFrame = bInsetDisplayFrame ? FFrameRate::TransformTime(FFrameTime(1), DisplayRate, LocalResolution) : FFrameTime(1);
+
+			TRange<FFrameNumber> PlaybackRange = Sequence->GetMovieScene()->GetPlaybackRange();
+			Markers.Add(UE::MovieScene::DiscreteInclusiveLower(PlaybackRange));
+			Markers.Add(UE::MovieScene::DiscreteExclusiveUpper(PlaybackRange) - OneFrame.FrameNumber);
+
+			TRange<FFrameNumber> SelectionRange = Sequence->GetMovieScene()->GetSelectionRange();
+			Markers.Add(UE::MovieScene::DiscreteInclusiveLower(SelectionRange));
+			Markers.Add(UE::MovieScene::DiscreteExclusiveUpper(SelectionRange) - OneFrame.FrameNumber);
+		}
+
+		for (const FFrameNumber& FrameNumber : Markers)
 		{
 			if (!NearestTime.IsSet())
 			{
-				NearestTime = GlobalMarkedFrame.FrameNumber;
+				NearestTime = FrameNumber;
 			}
-			else if (FMath::Abs(GlobalMarkedFrame.FrameNumber - CurrentTime) < FMath::Abs(NearestTime.GetValue() - CurrentTime))
+			else if (FMath::Abs(FrameNumber - CurrentTime) < FMath::Abs(NearestTime.GetValue() - CurrentTime))
 			{
-				NearestTime = GlobalMarkedFrame.FrameNumber;
+				NearestTime = FrameNumber;
 			}
 		}
 	}
