Sequencer: Add mouse button shortcuts for navigate forwards/backwards

The same shortcuts exist for the blueprint editor to navigate between tabs

diff --git a/Engine/Source/Editor/Sequencer/Private/SSequencer.cpp b/Engine/Source/Editor/Sequencer/Private/SSequencer.cpp
--- a/Engine/Source/Editor/Sequencer/Private/SSequencer.cpp
+++ b/Engine/Source/Editor/Sequencer/Private/SSequencer.cpp
@@ -3056,6 +3056,21 @@ void SSequencer::OnFocusChanging( const FWeakWidgetPath& PreviousFocusPath, cons
 	}
 }
 
+FReply SSequencer::OnMouseButtonDown(const FGeometry& MyGeometry, const FPointerEvent& MouseEvent)
+{
+	// A toolkit tab is active, so direct all command processing to it
+	TSharedPtr<FSequencer> SequencerPin = SequencerPtr.Pin();
+	if (SequencerPin.IsValid())
+	{
+		if (SequencerPin->GetCommandBindings()->ProcessCommandBindings(MouseEvent))
+		{
+			return FReply::Handled();
+		}
+	}
+
+	return FReply::Unhandled();
+}
+
 void SSequencer::OnMouseEnter(const FGeometry& MyGeometry, const FPointerEvent& MouseEvent)
 {
 	SCompoundWidget::OnMouseEnter(MyGeometry, MouseEvent);
diff --git a/Engine/Source/Editor/Sequencer/Private/SSequencer.h b/Engine/Source/Editor/Sequencer/Private/SSequencer.h
--- a/Engine/Source/Editor/Sequencer/Private/SSequencer.h
+++ b/Engine/Source/Editor/Sequencer/Private/SSequencer.h
@@ -368,6 +368,7 @@ protected:
 	virtual FReply OnDrop( const FGeometry& MyGeometry, const FDragDropEvent& DragDropEvent ) override;
 	virtual FReply OnKeyDown( const FGeometry& MyGeometry, const FKeyEvent& InKeyEvent ) override;
 	virtual void OnFocusChanging( const FWeakWidgetPath& PreviousFocusPath, const FWidgetPath& NewWidgetPath, const FFocusEvent& InFocusEvent ) override;
+	virtual FReply OnMouseButtonDown(const FGeometry& MyGeometry, const FPointerEvent& MouseEvent) override;
 	virtual void OnMouseEnter(const FGeometry& MyGeometry, const FPointerEvent& MouseEvent) override;
 	virtual void OnMouseLeave(const FPointerEvent& MouseEvent) override;
 
diff --git a/Engine/Source/Editor/Sequencer/Private/SequencerCommands.cpp b/Engine/Source/Editor/Sequencer/Private/SequencerCommands.cpp
--- a/Engine/Source/Editor/Sequencer/Private/SequencerCommands.cpp
+++ b/Engine/Source/Editor/Sequencer/Private/SequencerCommands.cpp
@@ -42,8 +42,8 @@ void FSequencerCommands::RegisterCommands()
 	UI_COMMAND( ZoomToFit, "Zoom to Fit", "Zoom to Fit", EUserInterfaceActionType::Button, FInputChord(EKeys::F) );
 	UI_COMMAND( ZoomInViewRange, "Zoom into the View Range", "Zoom into the view range", EUserInterfaceActionType::Button, FInputChord(EKeys::Equals) );
 	UI_COMMAND( ZoomOutViewRange, "Zoom out of the View Range", "Zoom out of the view range", EUserInterfaceActionType::Button, FInputChord(EKeys::Hyphen) );
-	UI_COMMAND( NavigateBackward, "Navigate Backward", "Go backward to the previously viewed shot/subsequence", EUserInterfaceActionType::Button, FInputChord(EModifierKey::Alt, EKeys::Left));
-	UI_COMMAND( NavigateForward, "Navigate Forward", "Go forward to the previously viewed shot/subsequence", EUserInterfaceActionType::Button, FInputChord(EModifierKey::Alt, EKeys::Right));
+	UI_COMMAND( NavigateBackward, "Navigate Backward", "Go backward to the previously viewed shot/subsequence", EUserInterfaceActionType::Button, FInputChord(EKeys::ThumbMouseButton), FInputChord(EModifierKey::Alt, EKeys::Left));
+	UI_COMMAND( NavigateForward, "Navigate Forward", "Go forward to the previously viewed shot/subsequence", EUserInterfaceActionType::Button, FInputChord(EKeys::ThumbMouseButton2), FInputChord(EModifierKey::Alt, EKeys::Right));
 
 	UI_COMMAND( SetSelectionRangeToNextShot, "Set Selection Range to Next Shot", "Set the selection range to the next shot", EUserInterfaceActionType::Button, FInputChord(EKeys::PageUp) );
 	UI_COMMAND( SetSelectionRangeToPreviousShot, "Set Selection Range to Previous Shot", "Set the selection range to the previous shot", EUserInterfaceActionType::Button, FInputChord(EKeys::PageDown) );
