Sequencer: Refresh hit proxies when:
- you end scrubbing,
- playback stops
This is achieved by OnEndDeferUpdates being called at the end of both of these cases.

Background of the change: when Sequencer played back, moving the mouse in the viewport used to cause a drop in FPS. The reason was that every frame, hit proxies were invalidated due to the FEditorViewportClient::Invalidate call. Invalidating hit proxies is expensive because it causes the game thread to call FlushRenderingCommands (FViewport::GetRawHitProxyData).

As compromise, during playback, we will no longer update hit proxies. As consequence, the user can no longer click actors that are being animated.

Important: If a user clicks an actor moved by evaluation, it will hit whatever what behind that actor before evaluation... until the hit proxies are invalidated.

diff --git a/Engine/Source/Editor/Sequencer/Private/LevelEditorSequencerIntegration.cpp b/Engine/Source/Editor/Sequencer/Private/LevelEditorSequencerIntegration.cpp
--- a/Engine/Source/Editor/Sequencer/Private/LevelEditorSequencerIntegration.cpp
+++ b/Engine/Source/Editor/Sequencer/Private/LevelEditorSequencerIntegration.cpp
@@ -544,40 +544,50 @@ static TAutoConsoleVariable<bool> CVarEnableHitProxyInvalidationDuringEvaluation
 	TEXT("Whether Sequencer evaluation should cause the viewport's hit proxies to be invalidated. If enabled, this causes FPS drops while moving the mouse in the viewport while playing back Sequencer.")
 	);
 
-void FLevelEditorSequencerIntegration::OnSequencerEvaluated()
+namespace UE::Sequencer
+{
+static void ReRenderLevelViewports()
 {
-	RestoreSpawnablePilotData();
-
-	// Redraw if not in PIE/simulate
-	const bool bIsInPIEOrSimulate = GEditor->PlayWorld != NULL || GEditor->bIsSimulatingInEditor;
-	if (bIsInPIEOrSimulate)
-	{
-		return;
-	}
-
 	// Request a single real-time frame to be rendered to ensure that we tick the world and update the viewport
 	// We only do this on level viewports instead of GetAllViewportClients to avoid needlessly redrawing Cascade,
 	// Blueprint, and other editors that have a 3d viewport.
 	for (FEditorViewportClient* LevelVC : GEditor->GetLevelViewportClients())
 	{
-		if (LevelVC)
+		if (LevelVC && !LevelVC->IsRealtime())
 		{
-			if (!LevelVC->IsRealtime())
-			{
-				LevelVC->RequestRealTimeFrames(1);
-			}
+			LevelVC->RequestRealTimeFrames(1);
+		}
+	}
+}
 
-			// Summary: When Sequencer evaluates moving the mouse in the viewport caused a drop in FPS.
-			// This fix is for 5.7, related to UE-320241.
-			// This CVar exists only for safety, since this change is coming in late to 5.7. In 5.8, it should be removed (replace "bInvalidateHitProxies" with "false").
-			// The cause WAS that every frame, hit proxies were invalidated due to the below FEditorViewportClient::Invalidate call.
-			// Invalidating hit proxies is expensive because it causes the game thread to call FlushRenderingCommands (FViewport::GetRawHitProxyData).
-			// After evaluation, it should be enough just to redraw the viewport (the rendered content) and skip the hit proxies. 
-			// Hit proxies are still invalidated automatically by FSlateApplication when it processes a mouse event.
-			const bool bInvalidateHitProxies = CVarEnableHitProxyInvalidationDuringEvaluation.GetValueOnAnyThread();
+static void InvalidateLevelViewports(bool bInvalidateHitProxies)
+{
+	for (FEditorViewportClient* LevelVC : GEditor->GetLevelViewportClients())
+	{
+		if (LevelVC)
+		{
 			LevelVC->Invalidate(true, bInvalidateHitProxies);
 		}
 	}
+}
+}
+
+void FLevelEditorSequencerIntegration::OnSequencerEvaluated()
+{
+	RestoreSpawnablePilotData();
+
+	// Redraw if not in PIE/simulate
+	const bool bIsInPIEOrSimulate = GEditor->PlayWorld != NULL || GEditor->bIsSimulatingInEditor;
+	if (bIsInPIEOrSimulate)
+	{
+		return;
+	}
+	
+	UE::Sequencer::ReRenderLevelViewports();
+	
+	// Summary: When Sequencer plays back, moving the mouse in the viewport used to cause a drop in FPS. See CL description.
+	const bool bInvalidateHitProxies = !bDeferUpdates;
+	UE::Sequencer::InvalidateLevelViewports(bInvalidateHitProxies);
 
 	if (!bDeferUpdates)
 	{
@@ -597,6 +607,9 @@ void FLevelEditorSequencerIntegration::OnEndDeferUpdates()
 {
 	bDeferUpdates = false;
 	UpdateDetails();
+
+	constexpr bool bInvalidateHitProxies = true;
+	UE::Sequencer::InvalidateLevelViewports(bInvalidateHitProxies);
 }
 
 bool FLevelEditorSequencerIntegration::IsBindingVisible(const FMovieSceneBinding& InBinding)
