Sequencer: Translate selected keys and sections when dragging marks. This unifies all the behaviors, so whatever is selected is moved.

Clarify MouseDownPosition so that there's a MouseDownPosition and a LastMousePosition.
Fixed marked frames not undoable
Removed no longer used SetMark()

#Synced-CL 48032220

diff --git a/Engine/Source/Editor/Sequencer/Private/SSequencerTransformBox.cpp b/Engine/Source/Editor/Sequencer/Private/SSequencerTransformBox.cpp
--- a/Engine/Source/Editor/Sequencer/Private/SSequencerTransformBox.cpp
+++ b/Engine/Source/Editor/Sequencer/Private/SSequencerTransformBox.cpp
@@ -171,7 +171,7 @@ FReply SSequencerTransformBox::OnPlusButtonClicked()
 	{
 		TSharedPtr<FSequencer> Sequencer = SequencerPtr.Pin();
 
-		Sequencer->TransformSelectedKeysAndSections(DeltaTime.GetValue(), 1.f);
+		Sequencer->TransformSelection(DeltaTime.GetValue(), 1.f);
 	}
 
 	return FReply::Handled();
@@ -188,7 +188,7 @@ FReply SSequencerTransformBox::OnMinusButtonClicked()
 	{
 		TSharedPtr<FSequencer> Sequencer = SequencerPtr.Pin();
 	
-		Sequencer->TransformSelectedKeysAndSections(-DeltaTime.GetValue(), 1.f);
+		Sequencer->TransformSelection(-DeltaTime.GetValue(), 1.f);
 	}
 
 	return FReply::Handled();
@@ -200,7 +200,7 @@ FReply SSequencerTransformBox::OnMultiplyButtonClicked()
 	{
 		TSharedPtr<FSequencer> Sequencer = SequencerPtr.Pin();
 		
-		Sequencer->TransformSelectedKeysAndSections(0, ScaleFactor);
+		Sequencer->TransformSelection(0, ScaleFactor);
 	}
 
 	return FReply::Handled();
@@ -212,7 +212,7 @@ FReply SSequencerTransformBox::OnDivideButtonClicked()
 	{
 		TSharedPtr<FSequencer> Sequencer = SequencerPtr.Pin();
 
-		Sequencer->TransformSelectedKeysAndSections(0, 1.f/ScaleFactor);
+		Sequencer->TransformSelection(0, 1.f/ScaleFactor);
 	}
 	return FReply::Handled();
 }
diff --git a/Engine/Source/Editor/Sequencer/Private/Sequencer.cpp b/Engine/Source/Editor/Sequencer/Private/Sequencer.cpp
--- a/Engine/Source/Editor/Sequencer/Private/Sequencer.cpp
+++ b/Engine/Source/Editor/Sequencer/Private/Sequencer.cpp
@@ -2281,11 +2281,11 @@ bool FSequencer::CanSnapToFrame() const
 	return bKeysSelected;
 }
 
-void FSequencer::TransformSelectedKeysAndSections(FFrameTime InDeltaTime, float InScale)
+void FSequencer::TransformSelection(FFrameTime InDeltaTime, float InScale)
 {
 	using namespace UE::Sequencer;
 
-	FScopedTransaction TransformKeysAndSectionsTransaction(NSLOCTEXT("Sequencer", "TransformKeysandSections_Transaction", "Transform Keys and Sections"));
+	FScopedTransaction TransformSelectionTransaction(NSLOCTEXT("Sequencer", "TransformSelection_Transaction", "Transform Selection"));
 	bool bAnythingChanged = false;
 
 	TSet<UMovieSceneSection*> SelectedSections = ViewModel->GetSelection()->GetSelectedSections();
@@ -2428,6 +2428,8 @@ void FSequencer::TransformSelectedKeysAndSections(FFrameTime InDeltaTime, float
 					Channel->SetKeyTimes(KeyHandles, NewKeyTimes);
 				}
 			}
+
+			Section->BroadcastChanged();
 		}
 
 		// Marked frames
@@ -2465,7 +2467,7 @@ void FSequencer::TransformSelectedKeysAndSections(FFrameTime InDeltaTime, float
 	}
 }
 
-void FSequencer::TranslateSelectedKeysAndSections(bool bTranslateLeft)
+void FSequencer::TranslateSelection(bool bTranslateLeft)
 {
 	if (IsReadOnly())
 	{
@@ -2475,7 +2477,7 @@ void FSequencer::TranslateSelectedKeysAndSections(bool bTranslateLeft)
 
 	int32 Shift = bTranslateLeft ? -1 : 1;
 	FFrameTime Delta = FQualifiedFrameTime(Shift, GetFocusedDisplayRate()).ConvertTo(GetFocusedTickResolution());
-	TransformSelectedKeysAndSections(Delta, 1.f);
+	TransformSelection(Delta, 1.f);
 }
 
 void FSequencer::StretchTime(FFrameTime InDeltaTime)
@@ -2512,7 +2514,7 @@ void FSequencer::StretchTime(FFrameTime InDeltaTime)
 
 	FocusedMovieScene->SetSelectionRange(SelectionRange);
 	SelectInSelectionRange(true, true, true);
-	TransformSelectedKeysAndSections(InDeltaTime, 1.f);
+	TransformSelection(InDeltaTime, 1.f);
 
 	// Return state
 	FocusedMovieScene->SetSelectionRange(CachedSelectionRange);
@@ -2598,7 +2600,7 @@ void FSequencer::ShrinkTime(FFrameTime InDeltaTime)
 
 	FocusedMovieScene->SetSelectionRange(SelectionRange);
 	SelectInSelectionRange(true, true, true);
-	TransformSelectedKeysAndSections(-InDeltaTime, 1.f);
+	TransformSelection(-InDeltaTime, 1.f);
 
 	// Return state
 	FocusedMovieScene->SetSelectionRange(CachedSelectionRange);
@@ -11207,11 +11209,11 @@ void FSequencer::BindCommands()
 
 	SequencerCommandBindings->MapAction(
 		Commands.TranslateLeft,
-		FExecuteAction::CreateSP( this, &FSequencer::TranslateSelectedKeysAndSections, true) );
+		FExecuteAction::CreateSP( this, &FSequencer::TranslateSelection, true) );
 
 	SequencerCommandBindings->MapAction(
 		Commands.TranslateRight,
-		FExecuteAction::CreateSP( this, &FSequencer::TranslateSelectedKeysAndSections, false) );
+		FExecuteAction::CreateSP( this, &FSequencer::TranslateSelection, false) );
 
 	SequencerCommandBindings->MapAction(
 		Commands.TrimSectionLeft,
diff --git a/Engine/Source/Editor/Sequencer/Private/Sequencer.h b/Engine/Source/Editor/Sequencer/Private/Sequencer.h
--- a/Engine/Source/Editor/Sequencer/Private/Sequencer.h
+++ b/Engine/Source/Editor/Sequencer/Private/Sequencer.h
@@ -289,11 +289,11 @@ public:
 	/** Are there keys to snap? */
 	bool CanSnapToFrame() const;
 
-	/** Transform the selected keys and sections */
-	void TransformSelectedKeysAndSections(FFrameTime InDeltaTime, float InScale);
+	/** Transform the selection */
+	void TransformSelection(FFrameTime InDeltaTime, float InScale);
 
-	/** Translate the selected keys and section by the time snap interval */
-	void TranslateSelectedKeysAndSections(bool bTranslateLeft);
+	/** Translate the selection by the time snap interval */
+	void TranslateSelection(bool bTranslateLeft);
 
 	/** Stretch time*/
 	void StretchTime(FFrameTime InDeltaTime);
diff --git a/Engine/Source/Editor/Sequencer/Private/SequencerTimeSliderController.cpp b/Engine/Source/Editor/Sequencer/Private/SequencerTimeSliderController.cpp
--- a/Engine/Source/Editor/Sequencer/Private/SequencerTimeSliderController.cpp
+++ b/Engine/Source/Editor/Sequencer/Private/SequencerTimeSliderController.cpp
@@ -705,14 +705,14 @@ int32 FSequencerTimeSliderController::OnPaintTimeSlider( bool bMirrorLabels, con
 			);
 		}
 		
-		if (MouseDragType == DRAG_SETTING_RANGE && MouseDownPosition[0].IsSet() && MouseDownPosition[1].IsSet())
+		if (MouseDragType == DRAG_SETTING_RANGE && MouseDownPosition.IsSet() && LastMousePosition.IsSet())
 		{
 			FFrameRate Resolution = GetTickResolution();
 			FFrameTime MouseDownTime[2];
 
 			FScrubRangeToScreen MouseDownRange(GetViewRange(), MouseDownGeometry.Size);
-			MouseDownTime[0] = ComputeFrameTimeFromMouse(MouseDownGeometry, MouseDownPosition[0].GetValue(), MouseDownRange);
-			MouseDownTime[1] = ComputeFrameTimeFromMouse(MouseDownGeometry, MouseDownPosition[1].GetValue(), MouseDownRange);
+			MouseDownTime[0] = ComputeFrameTimeFromMouse(MouseDownGeometry, MouseDownPosition.GetValue(), MouseDownRange);
+			MouseDownTime[1] = ComputeFrameTimeFromMouse(MouseDownGeometry, LastMousePosition.GetValue(), MouseDownRange);
 
 			float      MouseStartPosX = RangeToScreen.InputToLocalX(MouseDownTime[0] / Resolution);
 			float      MouseEndPosX   = RangeToScreen.InputToLocalX(MouseDownTime[1] / Resolution);
@@ -960,7 +960,7 @@ FReply FSequencerTimeSliderController::OnMouseButtonDown( SWidget& WidgetOwner,
 	DistanceDragged = 0;
 	MouseDownPlaybackRange = TimeSliderArgs.PlaybackRange.Get();
 	MouseDownSelectionRange = TimeSliderArgs.SelectionRange.Get();
-	MouseDownPosition[0] = MouseDownPosition[1] = MouseEvent.GetScreenSpacePosition();
+	MouseDownPosition = LastMousePosition = MouseEvent.GetScreenSpacePosition();
 	MouseDownGeometry = MyGeometry;
 	bMouseDownInRegion = false;
 
@@ -1010,8 +1010,8 @@ FReply FSequencerTimeSliderController::OnMouseButtonUp( SWidget& WidgetOwner, co
 		
 		bPanning = false;
 		bMouseDownInRegion = false;
-		MouseDownPosition[0].Reset();
-		MouseDownPosition[1].Reset();
+		MouseDownPosition.Reset();
+		LastMousePosition.Reset();
 		
 		return FReply::Handled().ReleaseMouseCapture();
 	}
@@ -1037,10 +1037,10 @@ FReply FSequencerTimeSliderController::OnMouseButtonUp( SWidget& WidgetOwner, co
 		{
 			TimeSliderArgs.OnMarkEndDrag.ExecuteIfBound();
 		}
-		else if (MouseDragType == DRAG_SETTING_RANGE && MouseDownPosition[0].IsSet())
+		else if (MouseDragType == DRAG_SETTING_RANGE && MouseDownPosition.IsSet())
 		{
 			// Zooming
-			FFrameTime MouseDownStart = ComputeFrameTimeFromMouse(MyGeometry, MouseDownPosition[0].GetValue(), RangeToScreen);
+			FFrameTime MouseDownStart = ComputeFrameTimeFromMouse(MyGeometry, MouseDownPosition.GetValue(), RangeToScreen);
 
 			const bool bCanZoomIn  = MouseTime > MouseDownStart;
 			const bool bCanZoomOut = ViewRangeStack.Num() > 0;
@@ -1117,15 +1117,15 @@ FReply FSequencerTimeSliderController::OnMouseButtonUp( SWidget& WidgetOwner, co
 		MouseDragType = DRAG_NONE;
 		DistanceDragged = 0.f;
 		bMouseDownInRegion = false;
-		MouseDownPosition[0].Reset();
-		MouseDownPosition[1].Reset();
+		MouseDownPosition.Reset();
+		LastMousePosition.Reset();
 
 		return FReply::Handled().ReleaseMouseCapture();
 	}
 
 	bMouseDownInRegion = false;
-	MouseDownPosition[0].Reset();
-	MouseDownPosition[1].Reset();
+	MouseDownPosition.Reset();
+	LastMousePosition.Reset();
 	return FReply::Unhandled();
 }
 
@@ -1234,7 +1234,7 @@ FReply FSequencerTimeSliderController::OnMouseMoveImpl( SWidget& WidgetOwner, co
 			SetViewRange(NewViewOutputMin, NewViewOutputMax, EViewRangeInterpolation::Immediate);
 		}
 	}
-	else if ((bHandleLeftMouseButton || bHandleMiddleMouseButton) && MouseDownPosition[0].IsSet())
+	else if ((bHandleLeftMouseButton || bHandleMiddleMouseButton) && MouseDownPosition.IsSet())
 	{
 		TRange<double> LocalViewRange = GetViewRange();
 		FScrubRangeToScreen RangeToScreen(LocalViewRange, MyGeometry.Size);
@@ -1244,7 +1244,7 @@ FReply FSequencerTimeSliderController::OnMouseMoveImpl( SWidget& WidgetOwner, co
 		{
 			if ( DistanceDragged > 0.f /*FSlateApplication::Get().GetDragTriggerDistance()*/ )
 			{
-				FFrameTime MouseDownFree = ComputeFrameTimeFromMouse(MyGeometry, MouseDownPosition[0].GetValue(), RangeToScreen, false);
+				FFrameTime MouseDownFree = ComputeFrameTimeFromMouse(MyGeometry, MouseDownPosition.GetValue(), RangeToScreen, false);
 
 				const bool       bReadOnly          = Sequencer->IsReadOnly();
 				const FFrameRate TickResolution     = GetTickResolution();
@@ -1285,16 +1285,6 @@ FReply FSequencerTimeSliderController::OnMouseMoveImpl( SWidget& WidgetOwner, co
 					MouseDragType = DRAG_MARK;
 					HandleMarkSelection(DragMarkIndex);
 
-					DragMarkMap.Empty();
-					const FMarkedFrameSelection& SelectedMarkedFrames = Sequencer->GetViewModel()->GetSelection()->MarkedFrames;
-					const TArray<FMovieSceneMarkedFrame>& MarkedFrames = TimeSliderArgs.MarkedFrames.Get();
-					for (TSet<int32>::TConstIterator It = SelectedMarkedFrames.GetSelected(); It; ++It)
-					{
-						const int32 MarkIndex = *It;		
-						const FMovieSceneMarkedFrame& MarkedFrame = MarkedFrames[MarkIndex];
-						DragMarkMap.Add(MarkIndex, MarkedFrame.FrameNumber);
-					}
-
 					TimeSliderArgs.OnMarkBeginDrag.ExecuteIfBound();
 				}
 				else if (FSlateApplication::Get().GetModifierKeys().AreModifersDown(EModifierKey::Control) && bHandleMiddleMouseButton == false)
@@ -1312,8 +1302,12 @@ FReply FSequencerTimeSliderController::OnMouseMoveImpl( SWidget& WidgetOwner, co
 		{
 			FFrameTime MouseTime = ComputeFrameTimeFromMouse(MyGeometry, MouseEvent.GetScreenSpacePosition(), RangeToScreen);
 			FFrameTime ScrubTime = ComputeScrubTimeFromMouse(MyGeometry, MouseEvent, RangeToScreen);
-			FFrameTime MouseDownTime = ComputeFrameTimeFromMouse(MyGeometry, MouseDownPosition[0].GetValue(), RangeToScreen);
+			FFrameTime MouseDownTime = ComputeFrameTimeFromMouse(MyGeometry, MouseDownPosition.GetValue(), RangeToScreen);
+			FFrameTime LastMouseTime = ComputeFrameTimeFromMouse(MyGeometry, LastMousePosition.GetValue(), RangeToScreen);
 			FFrameNumber DiffFrame = MouseTime.FrameNumber - MouseDownTime.FrameNumber;
+			FFrameNumber LastDiffFrame = MouseTime.FrameNumber - LastMouseTime.FrameNumber;
+
+			LastMousePosition = MouseEvent.GetScreenSpacePosition();
 
 			// Set the start range time?
 			if (MouseDragType == DRAG_PLAYBACK_START)
@@ -1368,17 +1362,13 @@ FReply FSequencerTimeSliderController::OnMouseMoveImpl( SWidget& WidgetOwner, co
 			}
 			else if (MouseDragType == DRAG_MARK)
 			{
-				SetMark(DiffFrame);
+				Sequencer->TransformSelection(LastDiffFrame, 1.f);
 			}
 			else if (MouseDragType == DRAG_SCRUBBING_TIME)
 			{
 				// Delegate responsibility for clamping to the current viewrange to the client
 				CommitScrubPosition( ScrubTime, /*bIsScrubbing=*/true, /*bEvaluate*/ !bHandleMiddleMouseButton); //if middle mouse button down we don't evaluate on the time change
 			}
-			else if (MouseDragType == DRAG_SETTING_RANGE)
-			{
-				MouseDownPosition[1] = MouseEvent.GetScreenSpacePosition();
-			}
 		}
 	}
 	else if (bFromTimeSlider && !bLockedMarkedFrames && DragMarkIndex == INDEX_NONE)
@@ -2624,16 +2614,6 @@ void FSequencerTimeSliderController::HandleMarkSelection(int32 InMarkIndex)
 	}
 }
 
-void FSequencerTimeSliderController::SetMark(FFrameNumber DiffFrame)
-{
-	for (TMap<int32, FFrameNumber>::TConstIterator It = DragMarkMap.CreateConstIterator(); It; ++It)
-	{
-		int32 MarkIndex = It.Key();
-		FFrameNumber FrameNumber = SnapSequencerTime(It.Value() + DiffFrame).FrameNumber;
-		TimeSliderArgs.OnSetMarkedFrame.ExecuteIfBound(MarkIndex, FrameNumber);
-	}
-}
-
 void FSequencerTimeSliderController::AddMarkAtFrame(FFrameNumber FrameNumber)
 {
 	TimeSliderArgs.OnAddMarkedFrame.ExecuteIfBound(FrameNumber);
diff --git a/Engine/Source/Editor/Sequencer/Private/Tools/EditToolDragOperations.cpp b/Engine/Source/Editor/Sequencer/Private/Tools/EditToolDragOperations.cpp
--- a/Engine/Source/Editor/Sequencer/Private/Tools/EditToolDragOperations.cpp
+++ b/Engine/Source/Editor/Sequencer/Private/Tools/EditToolDragOperations.cpp
@@ -1472,6 +1472,8 @@ void FMoveKeysAndSections::HandleMarkedFrameMovement(TOptional<FFrameNumber> Max
 		return;
 	}
 
+	FocusedMovieScene->Modify();
+
 	const TArray<FMovieSceneMarkedFrame>& AllMarkedFrames = FocusedMovieScene->GetMarkedFrames();
 
 	for (TSet<int32>::TConstIterator It = MarkedFrames.CreateConstIterator(); It; ++It)
diff --git a/Engine/Source/Editor/Sequencer/Public/SequencerTimeSliderController.h b/Engine/Source/Editor/Sequencer/Public/SequencerTimeSliderController.h
--- a/Engine/Source/Editor/Sequencer/Public/SequencerTimeSliderController.h
+++ b/Engine/Source/Editor/Sequencer/Public/SequencerTimeSliderController.h
@@ -334,8 +334,6 @@ private:
 	void SetSelectionRangeStart(FFrameNumber NewStart);
 	void SetSelectionRangeEnd(FFrameNumber NewEnd);
 
-	void SetMark(FFrameNumber DiffFrame);
-
 	TSharedRef<SWidget> OpenSetPlaybackRangeMenu(const FGeometry& MyGeometry, const FPointerEvent& MouseEvent);
 	FFrameTime SnapSequencerTime(FFrameTime InTime) const;
 	FFrameTime ComputeScrubTimeFromMouse(const FGeometry& Geometry, const FPointerEvent& MouseEvent, FScrubRangeToScreen RangeToScreen) const;
@@ -403,8 +401,11 @@ private:
 	/** If we are currently panning the panel */
 	bool bPanning;
 
-	/** Mouse down position range */
-	TOptional<FVector2D> MouseDownPosition[2];
+	/** Mouse down position */
+	TOptional<FVector2D> MouseDownPosition;
+	
+	/** Last mouse position */
+	TOptional<FVector2D> LastMousePosition;
 
 	/** Geometry on mouse down */
 	FGeometry MouseDownGeometry;
@@ -421,9 +422,6 @@ private:
 	/** Index of mark being hovered */
 	int32 HoverMarkIndex;
 
-	/** Map of the indices of the marks being edited and their initial frame numbers when pressed */
-	TMap<int32, FFrameNumber> DragMarkMap;
-
 	/** When > 0, we should not show context menus */
 	int32 ContextMenuSuppression;
 	
