- Restore t.MaxFPS in destructor as Close may not be called

#Synced-CL 48229894

diff --git a/Engine/Source/Editor/Sequencer/Private/Sequencer.cpp b/Engine/Source/Editor/Sequencer/Private/Sequencer.cpp
--- a/Engine/Source/Editor/Sequencer/Private/Sequencer.cpp
+++ b/Engine/Source/Editor/Sequencer/Private/Sequencer.cpp
@@ -865,6 +865,15 @@ FSequencer::~FSequencer()
 		GEditor->UnregisterForUndo(this);
 	}
 
+	if (!GExitPurge)
+	{
+		if (OldMaxTickRate.IsSet())
+		{
+			GEngine->SetMaxFPS(OldMaxTickRate.GetValue());
+			OldMaxTickRate.Reset();
+		}
+	}
+
 	for (auto TrackEditor : TrackEditors)
 	{
 		TrackEditor->OnRelease();
