Sequencer: Step forward/backward should respect looping state

Animators would like stepping forward/backward to respect looping state just as it does for next/previous key.

Added an argument to SetLocalTimeLooped so that it doesn't rely on the internal variable PlaybackSpeed to determine whether playback is forwards or backwards

#Synced-CL 48032220

diff --git a/Engine/Source/Editor/Sequencer/Private/Sequencer.cpp b/Engine/Source/Editor/Sequencer/Private/Sequencer.cpp
--- a/Engine/Source/Editor/Sequencer/Private/Sequencer.cpp
+++ b/Engine/Source/Editor/Sequencer/Private/Sequencer.cpp
@@ -1112,7 +1112,8 @@ void FSequencer::Tick(float InDeltaTime)
 	{
 		FQualifiedFrameTime CurrentTime = GetLocalTime();
 		FFrameTime Offset = (AutoscrubOffset.GetValue() * AutoScrollFactor) * CurrentTime.Rate;
-		SetLocalTimeLooped(CurrentTime.Time + Offset, CurrentTimeBreadcrumbs);
+		const bool bForwards = PlaybackSpeed > 0;
+		SetLocalTimeLooped(CurrentTime.Time + Offset, CurrentTimeBreadcrumbs, bForwards);
 	}
 
 	if (GetSelectionRange().IsEmpty() && GetLoopMode() == SLM_LoopSelectionRange)
@@ -1129,7 +1130,8 @@ void FSequencer::Tick(float InDeltaTime)
 		FFrameTime LocalTime = RootToUnwarpedLocalTransform.TransformTime(NewGlobalTime, FTransformTimeParams().HarvestBreadcrumbs(Breadcrumbs).IgnoreClamps());
 
 		FTimeDomainOverride DomainOverride = OverrideTimeDomain(ETimeDomain::Unwarped);
-		SetLocalTimeLooped(LocalTime, Breadcrumbs);
+		const bool bForwards = PlaybackSpeed > 0;
+		SetLocalTimeLooped(LocalTime, Breadcrumbs, bForwards);
 
 		if (GetPlaybackStatus() == EMovieScenePlayerStatus::Playing)
 		{
@@ -3908,6 +3910,7 @@ void FSequencer::UpdateAutoScroll(double NewTime, float ThresholdPercentage)
 
 	TRange<double> ViewRange = GetViewRange();
 	const double Threshold = (ViewRange.GetUpperBoundValue() - ViewRange.GetLowerBoundValue()) * ThresholdPercentage;
+	const bool bForwards = PlaybackSpeed > 0;
 
 	const FQualifiedFrameTime LocalTime = GetLocalTime();
 
@@ -3916,11 +3919,11 @@ void FSequencer::UpdateAutoScroll(double NewTime, float ThresholdPercentage)
 	{
 		if (AutoscrollOffset.GetValue() < 0 && LocalTime.AsSeconds() > ViewRange.GetLowerBoundValue() + Threshold)
 		{
-			SetLocalTimeLooped( (ViewRange.GetLowerBoundValue() + Threshold) * LocalTime.Rate, CurrentTimeBreadcrumbs);
+			SetLocalTimeLooped( (ViewRange.GetLowerBoundValue() + Threshold) * LocalTime.Rate, CurrentTimeBreadcrumbs, bForwards);
 		}
 		else if (AutoscrollOffset.GetValue() > 0 && LocalTime.AsSeconds() < ViewRange.GetUpperBoundValue() - Threshold)
 		{
-			SetLocalTimeLooped( (ViewRange.GetUpperBoundValue() - Threshold) * LocalTime.Rate, CurrentTimeBreadcrumbs);
+			SetLocalTimeLooped( (ViewRange.GetUpperBoundValue() - Threshold) * LocalTime.Rate, CurrentTimeBreadcrumbs, bForwards);
 		}
 	}
 
@@ -5212,7 +5215,8 @@ FReply FSequencer::OnStepForward(FFrameNumber Increment)
 	FFrameTime NewPosition = FFrameRate::TransformTime((CurrentTime.ConvertTo(DisplayRate) + Increment).FloorToFrame(), DisplayRate, CurrentTime.Rate);
 
 	FTimeDomainOverride TimeDomain = OverrideTimeDomain(ETimeDomain::Unwarped);
-	SetLocalTime(NewPosition, ESnapTimeMode::STM_Interval);
+	const bool bForwards = true;
+	SetLocalTimeLooped(NewPosition, CurrentTimeBreadcrumbs, bForwards);
 	return FReply::Handled();
 }
 
@@ -5229,7 +5233,8 @@ FReply FSequencer::OnStepBackward(FFrameNumber Increment)
 	FFrameTime NewPosition = FFrameRate::TransformTime((CurrentTime.ConvertTo(DisplayRate) - Increment).FloorToFrame(), DisplayRate, CurrentTime.Rate);
 
 	FTimeDomainOverride TimeDomain = OverrideTimeDomain(ETimeDomain::Unwarped);
-	SetLocalTime(NewPosition, ESnapTimeMode::STM_Interval);
+	const bool bForwards = false;
+	SetLocalTimeLooped(NewPosition, CurrentTimeBreadcrumbs, bForwards);
 	return FReply::Handled();
 }
 
@@ -5392,7 +5397,7 @@ ESequencerLoopMode FSequencer::GetLoopMode() const
 	return Settings->GetLoopMode();
 }
 
-void FSequencer::SetLocalTimeLooped(FFrameTime NewLocalTime, const FMovieSceneTransformBreadcrumbs& Breadcrumbs)
+void FSequencer::SetLocalTimeLooped(FFrameTime NewLocalTime, const FMovieSceneTransformBreadcrumbs& Breadcrumbs, bool bForwards)
 {
 	using namespace UE::MovieScene;
 	using namespace UE::Sequencer;
@@ -5413,7 +5418,7 @@ void FSequencer::SetLocalTimeLooped(FFrameTime NewLocalTime, const FMovieSceneTr
 	bool bHasJumped = false;
 	bool bRestarted = false;
 
-	if (PauseOnFrame.IsSet() && ((PlaybackSpeed > 0 && NewLocalTime > PauseOnFrame.GetValue()) || (PlaybackSpeed < 0 && NewLocalTime < PauseOnFrame.GetValue())))
+	if (PauseOnFrame.IsSet() && ((bForwards && NewLocalTime > PauseOnFrame.GetValue()) || (!bForwards && NewLocalTime < PauseOnFrame.GetValue())))
 	{
 		NewGlobalTime = LocalToRootTransform.TryTransformTime(PauseOnFrame.GetValue(), BreadcrumbsToUse);
 
@@ -5438,7 +5443,7 @@ void FSequencer::SetLocalTimeLooped(FFrameTime NewLocalTime, const FMovieSceneTr
 
 		if (LoopTime < MinInclusiveTime || LoopTime > MaxInclusiveTime)
 		{
-			NewGlobalTime = LocalToRootTransform.TryTransformTime((PlaybackSpeed > 0 ? MinInclusiveTime : MaxInclusiveTime), BreadcrumbsToUse);
+			NewGlobalTime = LocalToRootTransform.TryTransformTime((bForwards ? MinInclusiveTime : MaxInclusiveTime), BreadcrumbsToUse);
 			if (NewGlobalTime)
 			{
 				bResetPosition = true;
@@ -5463,7 +5468,7 @@ void FSequencer::SetLocalTimeLooped(FFrameTime NewLocalTime, const FMovieSceneTr
 			FFrameNumber MinInclusiveTime = UE::MovieScene::DiscreteInclusiveLower(RootTimeBounds);
 			FFrameNumber MaxInclusiveTime = UE::MovieScene::DiscreteExclusiveUpper(RootTimeBounds) - 1;
 			bool bReachedEnd = false;
-			if (PlaybackSpeed > 0)
+			if (bForwards)
 			{
 				bReachedEnd = OldWarpedGlobalTime <= MaxInclusiveTime && NewWarpedGlobalTime >= MaxInclusiveTime;
 			}
@@ -5475,7 +5480,7 @@ void FSequencer::SetLocalTimeLooped(FFrameTime NewLocalTime, const FMovieSceneTr
 			// Stop if we hit the playback range end
 			if (bReachedEnd)
 			{
-				NewGlobalTime = GlobalPlaybackWarpTransform.Inverse().TryTransformTime(PlaybackSpeed > 0 ? MaxInclusiveTime : MinInclusiveTime);
+				NewGlobalTime = GlobalPlaybackWarpTransform.Inverse().TryTransformTime(bForwards ? MaxInclusiveTime : MinInclusiveTime);
 				NewPlaybackStatus = EMovieScenePlayerStatus::Stopped;
 			}
 		}
diff --git a/Engine/Source/Editor/Sequencer/Private/Sequencer.h b/Engine/Source/Editor/Sequencer/Private/Sequencer.h
--- a/Engine/Source/Editor/Sequencer/Private/Sequencer.h
+++ b/Engine/Source/Editor/Sequencer/Private/Sequencer.h
@@ -520,7 +520,7 @@ public:
 	bool IsPilotCamera() const;
 
 	/** Sets the new global time calculated from local time and the given warp counter, accounting for looping options */
-	void SetLocalTimeLooped(FFrameTime InTime, const FMovieSceneTransformBreadcrumbs& Breadcrumbs=FMovieSceneTransformBreadcrumbs());
+	void SetLocalTimeLooped(FFrameTime InTime, const FMovieSceneTransformBreadcrumbs& Breadcrumbs=FMovieSceneTransformBreadcrumbs(), bool bForwards=true);
 
 	ESequencerLoopMode GetLoopMode() const;
 
