Sequencer: Follow up to 48294308 - Step forward/backward should respect looping state

If PlayOnce, don't loop

#Synced-CL 48303336

diff --git a/Engine/Source/Editor/Sequencer/Private/Sequencer.cpp b/Engine/Source/Editor/Sequencer/Private/Sequencer.cpp
--- a/Engine/Source/Editor/Sequencer/Private/Sequencer.cpp
+++ b/Engine/Source/Editor/Sequencer/Private/Sequencer.cpp
@@ -5215,8 +5215,15 @@ FReply FSequencer::OnStepForward(FFrameNumber Increment)
 	FFrameTime NewPosition = FFrameRate::TransformTime((CurrentTime.ConvertTo(DisplayRate) + Increment).FloorToFrame(), DisplayRate, CurrentTime.Rate);
 
 	FTimeDomainOverride TimeDomain = OverrideTimeDomain(ETimeDomain::Unwarped);
-	const bool bForwards = true;
-	SetLocalTimeLooped(NewPosition, CurrentTimeBreadcrumbs, bForwards);
+	if (GetLoopMode() == ESequencerLoopMode::SLM_Loop || GetLoopMode() == ESequencerLoopMode::SLM_LoopSelectionRange)
+	{
+		const bool bForwards = true;
+		SetLocalTimeLooped(NewPosition, CurrentTimeBreadcrumbs, bForwards);
+	}
+	else
+	{
+		SetLocalTime(NewPosition, ESnapTimeMode::STM_Interval);
+	}
 	return FReply::Handled();
 }
 
@@ -5233,8 +5240,15 @@ FReply FSequencer::OnStepBackward(FFrameNumber Increment)
 	FFrameTime NewPosition = FFrameRate::TransformTime((CurrentTime.ConvertTo(DisplayRate) - Increment).FloorToFrame(), DisplayRate, CurrentTime.Rate);
 
 	FTimeDomainOverride TimeDomain = OverrideTimeDomain(ETimeDomain::Unwarped);
-	const bool bForwards = false;
-	SetLocalTimeLooped(NewPosition, CurrentTimeBreadcrumbs, bForwards);
+	if (GetLoopMode() == ESequencerLoopMode::SLM_Loop || GetLoopMode() == ESequencerLoopMode::SLM_LoopSelectionRange)
+	{
+		const bool bForwards = false;
+		SetLocalTimeLooped(NewPosition, CurrentTimeBreadcrumbs, bForwards);
+	}
+	else
+	{
+		SetLocalTime(NewPosition, ESnapTimeMode::STM_Interval);
+	}
 	return FReply::Handled();
 }
 
