Sequencer: Follow up to Translate selected keys and sections when dragging marks. 48161809

Don't sort marked frames while dragging

#Synced-CL 48303336

diff --git a/Engine/Source/Editor/Sequencer/Private/Sequencer.cpp b/Engine/Source/Editor/Sequencer/Private/Sequencer.cpp
--- a/Engine/Source/Editor/Sequencer/Private/Sequencer.cpp
+++ b/Engine/Source/Editor/Sequencer/Private/Sequencer.cpp
@@ -2460,7 +2460,10 @@ void FSequencer::TransformSelection(FFrameTime InDeltaTime, float InScale)
 				FocusedMovieScene->SetMarkedFrame(MarkIndex, FrameNumber);
 			}
 
-			FocusedMovieScene->SortMarkedFrames();
+			if (!bDraggingMarks)
+			{
+				FocusedMovieScene->SortMarkedFrames();
+			}
 		}
 	}
 	
@@ -5974,13 +5977,15 @@ void FSequencer::OnSelectionRangeEndDrag()
 
 
 void FSequencer::OnMarkBeginDrag()
-{
+{	
+	bDraggingMarks = true;
 	GEditor->BeginTransaction(LOCTEXT("SetMark_Transaction", "Set Mark"));
 }
 
 
 void FSequencer::OnMarkEndDrag()
 {
+	bDraggingMarks = false;
 	UMovieSceneSequence* Sequence = GetFocusedMovieSceneSequence();
 	UMovieScene* OwnerMovieScene = Sequence ? Sequence->GetMovieScene() : nullptr;
 	if (OwnerMovieScene)
diff --git a/Engine/Source/Editor/Sequencer/Private/Sequencer.h b/Engine/Source/Editor/Sequencer/Private/Sequencer.h
--- a/Engine/Source/Editor/Sequencer/Private/Sequencer.h
+++ b/Engine/Source/Editor/Sequencer/Private/Sequencer.h
@@ -1577,6 +1577,7 @@ private:
 
 	FAcquiredResources AcquiredResources;
 
+	bool bDraggingMarks;
 	bool bGlobalMarkedFramesCached;
 	TArray<FMovieSceneMarkedFrame> GlobalMarkedFramesCache;
 
