Sequencer: Remove unused code

#rnx
#Synced-CL 48303336

diff --git a/Engine/Source/Editor/Sequencer/Private/MVVM/Views/SSequencerTrackAreaView.cpp b/Engine/Source/Editor/Sequencer/Private/MVVM/Views/SSequencerTrackAreaView.cpp
--- a/Engine/Source/Editor/Sequencer/Private/MVVM/Views/SSequencerTrackAreaView.cpp
+++ b/Engine/Source/Editor/Sequencer/Private/MVVM/Views/SSequencerTrackAreaView.cpp
@@ -32,21 +32,6 @@ void SSequencerTrackAreaView::Construct(const FArguments& InArgs, TWeakPtr<FTrac
 
 int32 SSequencerTrackAreaView::OnPaint(const FPaintArgs& Args, const FGeometry& AllottedGeometry, const FSlateRect& MyCullingRect, FSlateWindowElementList& OutDrawElements, int32 LayerId, const FWidgetStyle& InWidgetStyle, bool bParentEnabled ) const
 {
-	TSharedPtr<FTrackAreaViewModel> ViewModel = WeakViewModel.Pin();
-	if (ViewModel)
-	{
-		FSequencerEditorViewModel* EditorViewModel = ViewModel->GetEditor()->CastThisChecked<FSequencerEditorViewModel>();
-		TSharedRef<FSequencer> Sequencer = StaticCastSharedRef<FSequencer>(EditorViewModel->GetSequencer().ToSharedRef());
-
-		// Give track editors a chance to paint
-		// @todo_sequencer_mvvm: move track editors' code to viewmodels.
-		//const TArray<TSharedPtr<ISequencerTrackEditor>>& TrackEditors = Sequencer->GetTrackEditors();
-		//for (const TSharedPtr<ISequencerTrackEditor>& TrackEditor : TrackEditors)
-		//{
-		//	LayerId = TrackEditor->PaintTrackArea(Args, AllottedGeometry, MyCullingRect, OutDrawElements, LayerId + 1, InWidgetStyle);
-		//}
-	}
-
 	return STrackAreaView::OnPaint(Args, AllottedGeometry, MyCullingRect, OutDrawElements, LayerId, InWidgetStyle, bParentEnabled);
 }
 
