Fixing Sequencer issues.
- Deprecation warning disabling/enabling should only be done with Core macros
- Fixing EngineTest deprecation warnings.

#Synced-CL 48814203

diff --git a/Engine/Source/Editor/Sequencer/Private/Sequencer.h b/Engine/Source/Editor/Sequencer/Private/Sequencer.h
--- a/Engine/Source/Editor/Sequencer/Private/Sequencer.h
+++ b/Engine/Source/Editor/Sequencer/Private/Sequencer.h
@@ -811,19 +811,10 @@ public:
 	virtual FOnEndScrubbingEvent& OnEndScrubbingEvent() override { return OnEndScrubbingDelegate; }
 	virtual FOnMovieSceneDataChanged& OnMovieSceneDataChanged() override { return OnMovieSceneDataChangedDelegate; }
 
-#if defined(__clang__)
-#pragma clang diagnostic push
-#pragma clang diagnostic ignored "-Wdeprecated-declarations"
-#elif defined(_MSC_VER)
-	__pragma(warning(disable:4996))
-#endif
+	PRAGMA_DISABLE_DEPRECATION_WARNINGS
 	UE_DEPRECATED(5.8, "FOnChannelChanged OnChannelChanged() is now deprecated. Use FMovieSceneOnChannelChangedWithTime OnChannelChangedWithTime()")
 	FOnChannelChanged& OnChannelChanged() { return OnChannelChangedDelegate_DEPRACRATED; }
-#if defined(__clang__)
-#pragma clang diagnostic pop
-#elif defined(_MSC_VER)
-	__pragma(warning(default:4996))
-#endif
+	PRAGMA_ENABLE_DEPRECATION_WARNINGS
 
 	virtual FMovieSceneOnChannelChangedWithTime& OnChannelChangedWithTime() override { return OnChannelChangedDelegate; }
 	virtual FOnMovieSceneBindingsChanged& OnMovieSceneBindingsChanged() override { return OnMovieSceneBindingsChangedDelegate; }
@@ -1509,18 +1500,9 @@ private:
 	/** A delegate which is called any time the movie scene data is changed. */
 	FOnMovieSceneDataChanged OnMovieSceneDataChangedDelegate;
 
-#if defined(__clang__)
-#pragma clang diagnostic push
-#pragma clang diagnostic ignored "-Wdeprecated-declarations"
-#elif defined(_MSC_VER)
-	__pragma(warning(disable:4996))
-#endif
+	PRAGMA_DISABLE_DEPRECATION_WARNINGS
 	FOnChannelChanged OnChannelChangedDelegate_DEPRACRATED;
-#if defined(__clang__)
-#pragma clang diagnostic pop
-#elif defined(_MSC_VER)
-	__pragma(warning(default:4996))
-#endif
+	PRAGMA_ENABLE_DEPRECATION_WARNINGS
 
 	/** A delegate which is called when the channel is changed by Sequencer. */
 	FMovieSceneOnChannelChangedWithTime OnChannelChangedDelegate;
diff --git a/Engine/Source/Editor/Sequencer/Public/ISequencer.h b/Engine/Source/Editor/Sequencer/Public/ISequencer.h
--- a/Engine/Source/Editor/Sequencer/Public/ISequencer.h
+++ b/Engine/Source/Editor/Sequencer/Public/ISequencer.h
@@ -761,20 +761,10 @@ public:
 	/** Gets a multicast delegate which is executed whenever the movie scene data is changed. */
 	virtual FOnMovieSceneDataChanged& OnMovieSceneDataChanged() = 0;
 
-#if defined(__clang__)
-#pragma clang diagnostic push
-#pragma clang diagnostic ignored "-Wdeprecated-declarations"
-#elif defined(_MSC_VER)
-	__pragma(warning(disable:4996))
-#endif
-
+	PRAGMA_DISABLE_DEPRECATION_WARNINGS
 	UE_DEPRECATED(5.8, "FOnChannelChanged is now deprecated. Use FMovieSceneOnChannelChangedWithTime")
 	virtual FOnChannelChanged& OnChannelChanged() = 0;
-#if defined(__clang__)
-#pragma clang diagnostic pop
-#elif defined(_MSC_VER)
-	__pragma(warning(default:4996))
-#endif
+	PRAGMA_ENABLE_DEPRECATION_WARNINGS
 
 	/** Gets a multicast delegate which is executed whenever a channel is changed by Sequencer. */
 	virtual FMovieSceneOnChannelChangedWithTime& OnChannelChangedWithTime() = 0;
diff --git a/Engine/Source/Runtime/MovieScene/Public/MovieScene.h b/Engine/Source/Runtime/MovieScene/Public/MovieScene.h
--- a/Engine/Source/Runtime/MovieScene/Public/MovieScene.h
+++ b/Engine/Source/Runtime/MovieScene/Public/MovieScene.h
@@ -989,21 +989,10 @@ public:
 	*
 	*/
 
-#if defined(__clang__)
-#pragma clang diagnostic push
-#pragma clang diagnostic ignored "-Wdeprecated-declarations"
-#elif defined(_MSC_VER)
-	__pragma(warning(disable:4996))
-#endif
+	PRAGMA_DISABLE_DEPRECATION_WARNINGS
 	UE_DEPRECATED(5.8, "FMovieSceneOnChannelChanged OnChannelChanged() is now deprecated. Use FMovieSceneOnChannelChangedWithTime OnChannelChangedWithTime()")
 	FMovieSceneOnChannelChanged& OnChannelChanged() { return OnChannelChangedDelegate_DEPRACRATED; }
-#if defined(__clang__)
-#pragma clang diagnostic pop
-#elif defined(_MSC_VER)
-	__pragma(warning(default:4996)) 
-#endif
-
-
+	PRAGMA_ENABLE_DEPRECATION_WARNINGS
 
 	FMovieSceneOnChannelChangedWithTime& OnChannelChangedWithTime() { return OnChannelChangedDelegate; }
 #endif
@@ -1435,21 +1424,9 @@ private:
 	//delegates
 	private:
 
-
-#if defined(__clang__)
-#pragma clang diagnostic push
-#pragma clang diagnostic ignored "-Wdeprecated-declarations"
-#elif defined(_MSC_VER)
-		__pragma(warning(disable:4996))
-#endif
-		FMovieSceneOnChannelChanged OnChannelChangedDelegate_DEPRACRATED;
-
-#if defined(__clang__)
-#pragma clang diagnostic pop
-#elif defined(_MSC_VER)
-	__pragma(warning(default:4996))
-#endif
-
+	PRAGMA_DISABLE_DEPRECATION_WARNINGS
+	FMovieSceneOnChannelChanged OnChannelChangedDelegate_DEPRACRATED;
+	PRAGMA_ENABLE_DEPRECATION_WARNINGS
 
 	FMovieSceneOnChannelChangedWithTime OnChannelChangedDelegate;
 #endif
