Sequencer: Added command for Selection from Content Browser

#Synced-CL 48857117

diff --git a/Engine/Source/Editor/Sequencer/Private/SSequencer.cpp b/Engine/Source/Editor/Sequencer/Private/SSequencer.cpp
--- a/Engine/Source/Editor/Sequencer/Private/SSequencer.cpp
+++ b/Engine/Source/Editor/Sequencer/Private/SSequencer.cpp
@@ -1083,6 +1083,12 @@ void SSequencer::BindCommands(TSharedRef<FUICommandList> SequencerCommandBinding
 		return false;
 	};
 
+	SequencerCommandBindings->MapAction(
+		FSequencerCommands::Get().AddSelectionFromContentBrowser,
+		FExecuteAction::CreateSP(this, &SSequencer::AddFromContentBrowser),
+		FCanExecuteAction::CreateSP(this, &SSequencer::CanAddFromContentBrowser)
+	);
+
 	SequencerCommandBindings->MapAction(
 		FGenericCommands::Get().Paste,
 		FExecuteAction::CreateSP(this, &SSequencer::OnPaste),
@@ -1710,22 +1716,17 @@ TSharedRef<SWidget> SSequencer::MakeAddMenu()
 {
 	using namespace UE::Sequencer;
 
+	TSharedPtr<FSequencer> Sequencer = SequencerPtr.Pin();
+
 	TSharedPtr<FExtender> Extender = FExtender::Combine(AddMenuExtenders);
-	FMenuBuilder MenuBuilder(true, nullptr, Extender);
+	FMenuBuilder MenuBuilder(true, Sequencer->GetCommandBindings(), Extender);
 
 	if (SequencerPtr.Pin()->GetHostCapabilities().bSupportsAddFromContentBrowser)
 	{
-		MenuBuilder.AddMenuEntry(
-			LOCTEXT("SelectedFromContentBrowser", "Selection from Content Browser"),
-			LOCTEXT("SelectedFromContentBrowserToolTip", "Add selected content from the content browser"),
-			FSlateIcon(FAppStyle::GetAppStyleSetName(), "Icons.Use"),
-			FUIAction(
-				FExecuteAction::CreateRaw(this, &SSequencer::AddFromContentBrowser),
-				FCanExecuteAction::CreateRaw(this, &SSequencer::CanAddFromContentBrowser)));
+		MenuBuilder.AddMenuEntry(FSequencerCommands::Get().AddSelectionFromContentBrowser, NAME_None, TAttribute<FText>(), TAttribute<FText>(), FSlateIcon(FAppStyle::GetAppStyleSetName(), "Icons.Use"));
 	}
 
 	{
-		TSharedPtr<FSequencer> Sequencer = SequencerPtr.Pin();
 		Sequencer->GetViewModel()->GetOutliner()->CastThisChecked<FSequencerOutlinerViewModel>()->BuildContextMenu(MenuBuilder);
 	}
 
diff --git a/Engine/Source/Editor/Sequencer/Private/SequencerCommands.cpp b/Engine/Source/Editor/Sequencer/Private/SequencerCommands.cpp
--- a/Engine/Source/Editor/Sequencer/Private/SequencerCommands.cpp
+++ b/Engine/Source/Editor/Sequencer/Private/SequencerCommands.cpp
@@ -82,6 +82,8 @@ void FSequencerCommands::RegisterCommands()
 	UI_COMMAND( SelectNone, "Select None", "Select none", EUserInterfaceActionType::Button, FInputChord(EKeys::Escape));
 
 	UI_COMMAND( AddActorsToSequencer, "Add Actors", "Add actors to sequencer", EUserInterfaceActionType::Button, FInputChord(EModifierKey::Control, EKeys::A) );
+	UI_COMMAND( AddSelectionFromContentBrowser, "Selection from Content Browser", "Add selected content from the content browser", EUserInterfaceActionType::Button, FInputChord());
+
 	UI_COMMAND( SetKey, "Set Key", "Sets a key on the selected tracks", EUserInterfaceActionType::Button, FInputChord(EKeys::Enter) );
 	UI_COMMAND(SetInterpolationCubicSmartAuto, "Set Key Smart Auto", "Cubic interpolation - Smart Automatic tangents", EUserInterfaceActionType::Button, FInputChord(EKeys::Zero));
 	UI_COMMAND( SetInterpolationCubicAuto, "Set Key Auto", "Cubic interpolation - Automatic tangents", EUserInterfaceActionType::Button, FInputChord(EKeys::One));
diff --git a/Engine/Source/Editor/Sequencer/Public/SequencerCommands.h b/Engine/Source/Editor/Sequencer/Public/SequencerCommands.h
--- a/Engine/Source/Editor/Sequencer/Public/SequencerCommands.h
+++ b/Engine/Source/Editor/Sequencer/Public/SequencerCommands.h
@@ -216,6 +216,9 @@ public:
 	/** Add selected actors to sequencer */
 	TSharedPtr< FUICommandInfo > AddActorsToSequencer;
 
+	/** Add selection from content browser */
+	TSharedPtr< FUICommandInfo > AddSelectionFromContentBrowser;
+
 	/** Sets a key at the current time for the selected actor */
 	TSharedPtr< FUICommandInfo > SetKey;
 
