Sequencer: Added nullptr checks for editor

#rnx

#Synced-CL 48857117

diff --git a/Engine/Source/Editor/Sequencer/Private/MVVM/ViewModels/SequenceModel.cpp b/Engine/Source/Editor/Sequencer/Private/MVVM/ViewModels/SequenceModel.cpp
--- a/Engine/Source/Editor/Sequencer/Private/MVVM/ViewModels/SequenceModel.cpp
+++ b/Engine/Source/Editor/Sequencer/Private/MVVM/ViewModels/SequenceModel.cpp
@@ -190,12 +190,24 @@ TSharedPtr<FSequencerEditorViewModel> FSequenceModel::GetEditor() const
 
 TSharedPtr<ISequencer> FSequenceModel::GetSequencer() const
 {
-	return GetEditor()->GetSequencer();
+	TSharedPtr<FSequencerEditorViewModel> SequencerEditorViewModel = WeakEditor.Pin();
+	if (SequencerEditorViewModel)
+	{	
+		return SequencerEditorViewModel->GetSequencer();
+	}
+
+	return nullptr;
 }
 
 TSharedPtr<FSequencer> FSequenceModel::GetSequencerImpl() const
 {
-	return GetEditor()->GetSequencerImpl();
+	TSharedPtr<FSequencerEditorViewModel> SequencerEditorViewModel = WeakEditor.Pin();
+	if (SequencerEditorViewModel)
+	{
+		return SequencerEditorViewModel->GetSequencerImpl();
+	}
+
+	return nullptr;
 }
 
 void FSequenceModel::SortChildren()
@@ -408,7 +420,13 @@ void FSequenceModel::PerformDrop(const FViewModelPtr& TargetModel, const FDragDr
 
 void FSequenceModel::OnModifiedIndirectly(UMovieSceneSignedObject*)
 {
-	TSharedPtr<ISequencer> Sequencer = GetEditor()->GetSequencer();
+	TSharedPtr<FSequencerEditorViewModel> SequencerEditorViewModel = WeakEditor.Pin();
+	if (!SequencerEditorViewModel)
+	{
+		return;
+	}
+
+	TSharedPtr<ISequencer> Sequencer = SequencerEditorViewModel->GetSequencer();
 	const bool bRecording = Sequencer && Sequencer->OnGetIsRecording().IsBound() && Sequencer->OnGetIsRecording().Execute();
 
 	// Disregard updating cached flags during recording
@@ -423,7 +441,13 @@ void FSequenceModel::OnModifiedIndirectly(UMovieSceneSignedObject*)
 
 void FSequenceModel::OnModifiedDirectly(UMovieSceneSignedObject*)
 {
-	TSharedPtr<ISequencer> Sequencer = GetEditor()->GetSequencer();
+	TSharedPtr<FSequencerEditorViewModel> SequencerEditorViewModel = WeakEditor.Pin();
+	if (!SequencerEditorViewModel)
+	{
+		return;
+	}
+
+	TSharedPtr<ISequencer> Sequencer = SequencerEditorViewModel->GetSequencer();
 	const bool bRecording = Sequencer && Sequencer->OnGetIsRecording().IsBound() && Sequencer->OnGetIsRecording().Execute();
 
 	// Disregard updating cached flags during recording
