[ITF Interactions] Regression: Right-Mouse+WASD navigation in the viewport sets a keyframe when you use the 'S' key for the "Set key" action

#Synced-CL 48901306

diff --git a/Engine/Source/Editor/Sequencer/Private/SequencerEdMode.cpp b/Engine/Source/Editor/Sequencer/Private/SequencerEdMode.cpp
--- a/Engine/Source/Editor/Sequencer/Private/SequencerEdMode.cpp
+++ b/Engine/Source/Editor/Sequencer/Private/SequencerEdMode.cpp
@@ -241,7 +241,7 @@ bool FSequencerEdMode::IsCompatibleWith(FEditorModeID OtherModeID) const
 
 bool FSequencerEdMode::InputKey( FEditorViewportClient* ViewportClient, FViewport* Viewport, FKey Key, EInputEvent Event )
 {
-	if (Event != IE_Released)
+	if (Event != IE_Released && !ViewportClient->IsFlightCameraActive())
 	{
 		if (const TSharedPtr<ISequencer> ActiveSequencer = GetFirstActiveSequencer())
 		{
