Sequencer: Jump to next/previous key should respect looping state

If not looping, jumping to the next key while on the last key should not jump to the first key. Vice versa for previous key.

Added an argument to GetNextKey for loop

#Synced-CL 48909628

diff --git a/Engine/Source/Editor/Sequencer/Private/MVVM/ViewModels/SequencerOutlinerViewModel.cpp b/Engine/Source/Editor/Sequencer/Private/MVVM/ViewModels/SequencerOutlinerViewModel.cpp
--- a/Engine/Source/Editor/Sequencer/Private/MVVM/ViewModels/SequencerOutlinerViewModel.cpp
+++ b/Engine/Source/Editor/Sequencer/Private/MVVM/ViewModels/SequencerOutlinerViewModel.cpp
@@ -93,17 +93,17 @@ TSharedRef<FDragDropOperation> FSequencerOutlinerViewModel::InitiateDrag(TArray<
 	return FSequencerOutlinerDragDropOp::New( MoveTemp(InDraggedModels), DefaultText, nullptr );
 }
 
-FFrameNumber FSequencerOutlinerViewModel::GetNextKey(const TArray<TSharedRef<UE::Sequencer::FViewModel>>& InNodes, FFrameNumber FrameNumber, EMovieSceneTimeUnit TimeUnit, const TRange<FFrameNumber>& Range)
+FFrameNumber FSequencerOutlinerViewModel::GetNextKey(const TArray<TSharedRef<UE::Sequencer::FViewModel>>& InNodes, FFrameNumber FrameNumber, EMovieSceneTimeUnit TimeUnit, const TRange<FFrameNumber>& Range, const bool bAllowLoop)
 {
-	return GetNextKeyInternal(InNodes, FrameNumber, TimeUnit, Range, EFindKeyDirection::Forwards);
+	return GetNextKeyInternal(InNodes, FrameNumber, TimeUnit, Range, EFindKeyDirection::Forwards, bAllowLoop);
 }
 
-FFrameNumber FSequencerOutlinerViewModel::GetPreviousKey(const TArray<TSharedRef<UE::Sequencer::FViewModel>>& InNodes, FFrameNumber FrameNumber, EMovieSceneTimeUnit TimeUnit, const TRange<FFrameNumber>& Range)
+FFrameNumber FSequencerOutlinerViewModel::GetPreviousKey(const TArray<TSharedRef<UE::Sequencer::FViewModel>>& InNodes, FFrameNumber FrameNumber, EMovieSceneTimeUnit TimeUnit, const TRange<FFrameNumber>& Range, const bool bAllowLoop)
 {
-	return GetNextKeyInternal(InNodes, FrameNumber, TimeUnit, Range, EFindKeyDirection::Backwards);
+	return GetNextKeyInternal(InNodes, FrameNumber, TimeUnit, Range, EFindKeyDirection::Backwards, bAllowLoop);
 }
 
-FFrameNumber FSequencerOutlinerViewModel::GetNextKeyInternal(const TArray<TSharedRef<UE::Sequencer::FViewModel>>& InNodes, FFrameNumber FrameNumber, EMovieSceneTimeUnit TimeUnit, const TRange<FFrameNumber>& Range, EFindKeyDirection Direction)
+FFrameNumber FSequencerOutlinerViewModel::GetNextKeyInternal(const TArray<TSharedRef<UE::Sequencer::FViewModel>>& InNodes, FFrameNumber FrameNumber, EMovieSceneTimeUnit TimeUnit, const TRange<FFrameNumber>& Range, EFindKeyDirection Direction, const bool bAllowLoop)
 {
 	FSequencerEditorViewModel* EditorViewModel = GetEditor()->CastThisChecked<FSequencerEditorViewModel>();
 	TSharedPtr<ISequencer> Sequencer = EditorViewModel->GetSequencer();
@@ -134,7 +134,7 @@ FFrameNumber FSequencerOutlinerViewModel::GetNextKeyInternal(const TArray<TShare
 
 		KeyCollection->Update(FSequencerKeyCollectionSignature::FromNodesRecursive(InNodes, ThresholdFrames));
 
-		TOptional<FFrameNumber> NextKey = KeyCollection->GetNextKey(FrameNumber, Direction, Range, EFindKeyType::FKT_All);
+		TOptional<FFrameNumber> NextKey = KeyCollection->GetNextKey(FrameNumber, Direction, Range, EFindKeyType::FKT_All, bAllowLoop);
 		if (NextKey.IsSet())
 		{
 			if (TimeUnit == EMovieSceneTimeUnit::DisplayRate)
diff --git a/Engine/Source/Editor/Sequencer/Private/Scripting/SequencerModuleOutlinerScriptingObject.cpp b/Engine/Source/Editor/Sequencer/Private/Scripting/SequencerModuleOutlinerScriptingObject.cpp
--- a/Engine/Source/Editor/Sequencer/Private/Scripting/SequencerModuleOutlinerScriptingObject.cpp
+++ b/Engine/Source/Editor/Sequencer/Private/Scripting/SequencerModuleOutlinerScriptingObject.cpp
@@ -57,7 +57,7 @@ TArray<UMovieSceneSection*> USequencerModuleOutlinerScriptingObject::GetSections
 	return Sections;
 }
 
-FFrameNumber USequencerModuleOutlinerScriptingObject::GetNextKey(const TArray<FSequencerViewModelScriptingStruct>& InNodes, FFrameNumber FrameNumber, EMovieSceneTimeUnit TimeUnit) const
+FFrameNumber USequencerModuleOutlinerScriptingObject::GetNextKey(const TArray<FSequencerViewModelScriptingStruct>& InNodes, FFrameNumber FrameNumber, EMovieSceneTimeUnit TimeUnit, const bool bAllowLoop) const
 {
 	using namespace UE::Sequencer;
 
@@ -80,10 +80,10 @@ FFrameNumber USequencerModuleOutlinerScriptingObject::GetNextKey(const TArray<FS
 		}
 	}
 
-	return Outliner->GetNextKey(Nodes, FrameNumber, TimeUnit, Range);
+	return Outliner->GetNextKey(Nodes, FrameNumber, TimeUnit, Range, bAllowLoop);
 }
 
-FFrameNumber USequencerModuleOutlinerScriptingObject::GetPreviousKey(const TArray<FSequencerViewModelScriptingStruct>& InNodes, FFrameNumber FrameNumber, EMovieSceneTimeUnit TimeUnit) const
+FFrameNumber USequencerModuleOutlinerScriptingObject::GetPreviousKey(const TArray<FSequencerViewModelScriptingStruct>& InNodes, FFrameNumber FrameNumber, EMovieSceneTimeUnit TimeUnit, const bool bAllowLoop) const
 {
 	using namespace UE::Sequencer;
 
@@ -106,7 +106,7 @@ FFrameNumber USequencerModuleOutlinerScriptingObject::GetPreviousKey(const TArra
 		}
 	}
 
-	return Outliner->GetPreviousKey(Nodes, FrameNumber, TimeUnit, Range);
+	return Outliner->GetPreviousKey(Nodes, FrameNumber, TimeUnit, Range, bAllowLoop);
 }
 
 void USequencerModuleOutlinerScriptingObject::AddBindingTags(const TArray<FSequencerViewModelScriptingStruct>& InNodes, const TArray<FName>& TagNames)
diff --git a/Engine/Source/Editor/Sequencer/Private/Scripting/SequencerModuleOutlinerScriptingObject.h b/Engine/Source/Editor/Sequencer/Private/Scripting/SequencerModuleOutlinerScriptingObject.h
--- a/Engine/Source/Editor/Sequencer/Private/Scripting/SequencerModuleOutlinerScriptingObject.h
+++ b/Engine/Source/Editor/Sequencer/Private/Scripting/SequencerModuleOutlinerScriptingObject.h
@@ -26,11 +26,10 @@ public:
 	TArray<UMovieSceneSection*> GetSections(const TArray<FSequencerViewModelScriptingStruct>& InNodes) const;
 
 	UFUNCTION(BlueprintCallable, Category = "Sequencer Editor")
-	FFrameNumber GetNextKey(const TArray<FSequencerViewModelScriptingStruct>& InNodes, FFrameNumber FrameNumber, EMovieSceneTimeUnit TimeUnit = EMovieSceneTimeUnit::DisplayRate) const;
+	FFrameNumber GetNextKey(const TArray<FSequencerViewModelScriptingStruct>& InNodes, FFrameNumber FrameNumber, EMovieSceneTimeUnit TimeUnit = EMovieSceneTimeUnit::DisplayRate, bool bAllowLoop = true) const;
 
 	UFUNCTION(BlueprintCallable, Category = "Sequencer Editor")
-	FFrameNumber GetPreviousKey(const TArray<FSequencerViewModelScriptingStruct>& InNodes, FFrameNumber FrameNumber, EMovieSceneTimeUnit TimeUnit = EMovieSceneTimeUnit::DisplayRate) const;
-
+	FFrameNumber GetPreviousKey(const TArray<FSequencerViewModelScriptingStruct>& InNodes, FFrameNumber FrameNumber, EMovieSceneTimeUnit TimeUnit = EMovieSceneTimeUnit::DisplayRate, bool bAllowLoop = true) const;
 
 	// Tags
 	UFUNCTION(BlueprintCallable, Category = "Sequencer Editor")
diff --git a/Engine/Source/Editor/Sequencer/Private/Sequencer.cpp b/Engine/Source/Editor/Sequencer/Private/Sequencer.cpp
--- a/Engine/Source/Editor/Sequencer/Private/Sequencer.cpp
+++ b/Engine/Source/Editor/Sequencer/Private/Sequencer.cpp
@@ -5362,7 +5362,8 @@ FReply FSequencer::JumpToPreviousKey()
 	{
 		TRange<FFrameNumber> Range = GetTimeBounds();
 		FFrameNumber FrameNumber = GetLocalTime().Time.FloorToFrame();
-		TOptional<FFrameNumber> NewTime = SelectedKeyCollection->GetNextKey(FrameNumber, EFindKeyDirection::Backwards, Range, EFindKeyType::FKT_All);
+		const bool bAllowLoop = GetLoopMode() != SLM_NoLoop;
+		TOptional<FFrameNumber> NewTime = SelectedKeyCollection->GetNextKey(FrameNumber, EFindKeyDirection::Backwards, Range, EFindKeyType::FKT_All, bAllowLoop);
 		if (NewTime.IsSet())
 		{
 			SetPlaybackStatus(EMovieScenePlayerStatus::Stepping);
@@ -5392,7 +5393,8 @@ FReply FSequencer::JumpToNextKey()
 	{
 		TRange<FFrameNumber> Range = GetTimeBounds();
 		FFrameNumber FrameNumber = GetLocalTime().Time.FloorToFrame();
-		TOptional<FFrameNumber> NewTime = SelectedKeyCollection->GetNextKey(FrameNumber, EFindKeyDirection::Forwards, Range, EFindKeyType::FKT_All);
+		const bool bAllowLoop = GetLoopMode() != SLM_NoLoop;
+		TOptional<FFrameNumber> NewTime = SelectedKeyCollection->GetNextKey(FrameNumber, EFindKeyDirection::Forwards, Range, EFindKeyType::FKT_All, bAllowLoop);
 		if (NewTime.IsSet())
 		{
 			SetPlaybackStatus(EMovieScenePlayerStatus::Stepping);
diff --git a/Engine/Source/Editor/Sequencer/Private/SequencerKeyCollection.cpp b/Engine/Source/Editor/Sequencer/Private/SequencerKeyCollection.cpp
--- a/Engine/Source/Editor/Sequencer/Private/SequencerKeyCollection.cpp
+++ b/Engine/Source/Editor/Sequencer/Private/SequencerKeyCollection.cpp
@@ -322,7 +322,7 @@ TArrayView<const FFrameNumber> GetKeysInRangeInternal(const TArray<FFrameNumber>
 	return TArrayView<const FFrameNumber>();
 }
 
-TOptional<FFrameNumber> GetNextKeyInternal(const TArray<FFrameNumber>& Times, FFrameNumber FrameNumber, EFindKeyDirection Direction, const TRange<FFrameNumber>& Range)
+TOptional<FFrameNumber> GetNextKeyInternal(const TArray<FFrameNumber>& Times, FFrameNumber FrameNumber, EFindKeyDirection Direction, const TRange<FFrameNumber>& Range, const bool bAllowLoop)
 {
 	TArrayView<const FFrameNumber> KeysInRange = GetKeysInRangeInternal(Times, Range);
 
@@ -340,7 +340,7 @@ TOptional<FFrameNumber> GetNextKeyInternal(const TArray<FFrameNumber>& Times, FF
 	{
 		return KeysInRange[Index];
 	}
-	else if (KeysInRange.Num() > 0)
+	else if (KeysInRange.Num() > 0 && bAllowLoop)
 	{
 		if (Direction == EFindKeyDirection::Forwards)
 		{
@@ -400,19 +400,19 @@ TArrayView<const FFrameNumber> FSequencerKeyCollection::GetKeysInRange(const TRa
 	return TArrayView<const FFrameNumber>();
 }
 
-TOptional<FFrameNumber> FSequencerKeyCollection::GetNextKey(FFrameNumber FrameNumber, EFindKeyDirection Direction, const TRange<FFrameNumber>& Range, EFindKeyType FindKeyType) const
+TOptional<FFrameNumber> FSequencerKeyCollection::GetNextKey(FFrameNumber FrameNumber, EFindKeyDirection Direction, const TRange<FFrameNumber>& Range, EFindKeyType FindKeyType, const bool bAllowLoop) const
 {
 	if (FindKeyType == EFindKeyType::FKT_Keys)
 	{
-		return GetNextKeyInternal(GroupedTimes, FrameNumber, Direction, Range);
+		return GetNextKeyInternal(GroupedTimes, FrameNumber, Direction, Range, bAllowLoop);
 	}
 	else if (FindKeyType == EFindKeyType::FKT_Sections)
 	{
-		return GetNextKeyInternal(GroupedSectionTimes, FrameNumber, Direction, Range);
+		return GetNextKeyInternal(GroupedSectionTimes, FrameNumber, Direction, Range, bAllowLoop);
 	}
 	else if (FindKeyType == EFindKeyType::FKT_All)
 	{
-		return GetNextKeyInternal(AllGroupedTimes, FrameNumber, Direction, Range);
+		return GetNextKeyInternal(AllGroupedTimes, FrameNumber, Direction, Range, bAllowLoop);
 	}
 	return TOptional<FFrameNumber>();
 }
diff --git a/Engine/Source/Editor/Sequencer/Public/MVVM/ViewModels/SequencerOutlinerViewModel.h b/Engine/Source/Editor/Sequencer/Public/MVVM/ViewModels/SequencerOutlinerViewModel.h
--- a/Engine/Source/Editor/Sequencer/Public/MVVM/ViewModels/SequencerOutlinerViewModel.h
+++ b/Engine/Source/Editor/Sequencer/Public/MVVM/ViewModels/SequencerOutlinerViewModel.h
@@ -34,12 +34,12 @@ public:
 	TSharedRef<FDragDropOperation> InitiateDrag(TArray<TWeakViewModelPtr<IOutlinerExtension>>&& InDraggedModels) override;
 	void RequestUpdate() override;
 
-	FFrameNumber GetNextKey(const TArray<TSharedRef<UE::Sequencer::FViewModel>>& InNodes, FFrameNumber FrameNumber, EMovieSceneTimeUnit TimeUnit, const TRange<FFrameNumber>& Range);
-	FFrameNumber GetPreviousKey(const TArray<TSharedRef<UE::Sequencer::FViewModel>>& InNodes, FFrameNumber FrameNumber, EMovieSceneTimeUnit TimeUnit, const TRange<FFrameNumber>& Range);
+	FFrameNumber GetNextKey(const TArray<TSharedRef<UE::Sequencer::FViewModel>>& InNodes, FFrameNumber FrameNumber, EMovieSceneTimeUnit TimeUnit, const TRange<FFrameNumber>& Range, const bool bAllowLoop);
+	FFrameNumber GetPreviousKey(const TArray<TSharedRef<UE::Sequencer::FViewModel>>& InNodes, FFrameNumber FrameNumber, EMovieSceneTimeUnit TimeUnit, const TRange<FFrameNumber>& Range, const bool bAllowLoop);
 
 private:
 
-	FFrameNumber GetNextKeyInternal(const TArray<TSharedRef<UE::Sequencer::FViewModel>>& InNodes, FFrameNumber FrameNumber, EMovieSceneTimeUnit TimeUnit, const TRange<FFrameNumber>& Range, EFindKeyDirection Direction);
+	FFrameNumber GetNextKeyInternal(const TArray<TSharedRef<UE::Sequencer::FViewModel>>& InNodes, FFrameNumber FrameNumber, EMovieSceneTimeUnit TimeUnit, const TRange<FFrameNumber>& Range, EFindKeyDirection Direction, const bool bAllowLoop);
 
 public:
 
diff --git a/Engine/Source/Editor/Sequencer/Public/SequencerKeyCollection.h b/Engine/Source/Editor/Sequencer/Public/SequencerKeyCollection.h
--- a/Engine/Source/Editor/Sequencer/Public/SequencerKeyCollection.h
+++ b/Engine/Source/Editor/Sequencer/Public/SequencerKeyCollection.h
@@ -105,8 +105,9 @@ public:
 	/**
 	 * Search forwards or backwards for the first key within the specified range
 	 *
-	 * @param Range      The range to search within
-	 * @param Direction  Whether to return the first or last key that reside in the given range
+	 * @param Range       The range to search within
+	 * @param Direction   Whether to return the first or last key that reside in the given range
+	 * @param FindKeyType Whether to find all keys, sections or all
 	 * @return (Optional) the time of the key that matched the range
 	 */
 	SEQUENCER_API TOptional<FFrameNumber> FindFirstKeyInRange(const TRange<FFrameNumber>& Range, EFindKeyDirection Direction, EFindKeyType FindKeyType = EFindKeyType::FKT_All) const;
@@ -116,7 +117,8 @@ public:
 	/**
 	 * Get a view of all key times that reside within the specified range
 	 *
-	 * @param Range      The range to search within
+	 * @param Range       The range to search within
+	 * @param FindKeyType Whether to find all keys, sections or all
 	 * @return A (possibly empty) array view of all the times that lie within the range
 	 */
 	SEQUENCER_API TArrayView<const FFrameNumber> GetKeysInRange(const TRange<FFrameNumber>& Range, EFindKeyType FindKeyType = EFindKeyType::FKT_All) const;
@@ -126,10 +128,13 @@ public:
 	/**
 	* Search forwards or backwards for the next key from the specified frame number
 	* @param FrameNumber The frame number to search from
-	* @param Direction  Whether to return the next key or previous key from that time
+	* @param Direction   Whether to return the next key or previous key from that time
+	* @param Range       The range to search within
+	* @param FindKeyType Whether to find all keys, sections or all
+	* @param bAllowLoop  Whether to loop around if past the last key when retrieving the next key (or past the first key when retrieving the previous key)
 	* @return (Optional)  Frame number of the key that's next or previous from that time 
 	*/
-	SEQUENCER_API TOptional<FFrameNumber> GetNextKey(FFrameNumber FrameNumber, EFindKeyDirection Direction, const TRange<FFrameNumber>& Range, EFindKeyType FindKeyType = EFindKeyType::FKT_All) const;
+	SEQUENCER_API TOptional<FFrameNumber> GetNextKey(FFrameNumber FrameNumber, EFindKeyDirection Direction, const TRange<FFrameNumber>& Range, EFindKeyType FindKeyType = EFindKeyType::FKT_All, const bool bAllowLoop = true) const;
 	UE_DEPRECATED(5.4, "Please use GetNextKey which takes EFindKeyType, GetNextSectionKey is no longer supported")
 	TOptional<FFrameNumber> GetNextSectionKey(FFrameNumber FrameNumber, EFindKeyDirection Direction, const TRange<FFrameNumber>& Range) const { return GetNextKey(FrameNumber, Direction, Range, EFindKeyType::FKT_Sections); }
 
