Sequencer: Large anim mixer update and Sequencer MVVM refactor.

diff --git a/Engine/Plugins/Experimental/MovieScenePoseSearchTracks/Source/MovieScenePoseSearchTracksEditor/Private/TrackEditors/StitchAnimTrackEditor.h b/Engine/Plugins/Experimental/MovieScenePoseSearchTracks/Source/MovieScenePoseSearchTracksEditor/Private/TrackEditors/StitchAnimTrackEditor.h
--- a/Engine/Plugins/Experimental/MovieScenePoseSearchTracks/Source/MovieScenePoseSearchTracksEditor/Private/TrackEditors/StitchAnimTrackEditor.h
+++ b/Engine/Plugins/Experimental/MovieScenePoseSearchTracks/Source/MovieScenePoseSearchTracksEditor/Private/TrackEditors/StitchAnimTrackEditor.h
@@ -13,6 +13,7 @@
 #include "ISequencerTrackEditor.h"
 #include "MovieSceneTrackEditor.h"
 #include "SequencerCoreFwd.h"
+#include "IMovieSceneAnimMixerItemMenuProvider.h"
 
 namespace UE::Sequencer
 {
@@ -30,6 +31,7 @@ class UPoseSearchDatabase;
  * Tools for stitch anim tracks
  */
 class FStitchAnimTrackEditor : public FMovieSceneTrackEditor
+							 , public IMovieSceneAnimMixerItemMenuProvider
 {
 public:
 	/**
@@ -53,7 +55,8 @@ public:
 public:
 
 	// ISequencerTrackEditor interface
-
+	virtual void OnInitialize() override;
+	virtual void OnRelease() override;
 	virtual void BuildObjectBindingTrackMenu(FMenuBuilder& MenuBuilder, const TArray<FGuid>& ObjectBindings, const UClass* ObjectClass) override;
 	virtual bool HandleAssetAdded(UObject* Asset, const FGuid& TargetObjectGuid) override;
 	virtual TSharedRef<ISequencerSection> MakeSectionInterface(UMovieSceneSection& SectionObject, UMovieSceneTrack& Track, FGuid ObjectBinding) override;
@@ -66,18 +69,23 @@ public:
 private:
 
 	/** Animation sub menu */
-	TSharedRef<SWidget> BuildAddAnimationSubMenu(FGuid ObjectBinding, USkeleton* Skeleton, UE::Sequencer::TWeakViewModelPtr<UE::Sequencer::ITrackExtension> TrackModel);
-	TSharedRef<SWidget> BuildAnimationSubMenu(FGuid ObjectBinding, USkeleton* Skeleton, UMovieSceneTrack* Track);
-	void AddAnimationSubMenu(FMenuBuilder& MenuBuilder, TArray<FGuid> ObjectBindings, USkeleton* Skeleton, UMovieSceneTrack* Track);
+	TSharedRef<SWidget> BuildAddAnimationSubMenu(FGuid ObjectBinding, USkeleton* Skeleton, UE::Sequencer::TWeakViewModelPtr<UE::Sequencer::FViewModel> WeakViewModel, TFunction<void(UMovieSceneSection*)> InSectionCallback = nullptr);
+	TSharedRef<SWidget> BuildAnimationSubMenu(FGuid ObjectBinding, USkeleton* Skeleton, UMovieSceneTrack* Track, TFunction<void(UMovieSceneSection*)> InSectionCallback = nullptr);
+	void AddAnimationSubMenu(FMenuBuilder& MenuBuilder, TArray<FGuid> ObjectBindings, USkeleton* Skeleton, UMovieSceneTrack* Track, TFunction<void(UMovieSceneSection*)> InSectionCallback = nullptr);
 
 	/** Animation asset selected */
-	void OnAnimationDatabaseAssetSelected(const FAssetData& AssetData, TArray<FGuid> ObjectBindings, UMovieSceneTrack* Track);
+	void OnAnimationDatabaseAssetSelected(const FAssetData& AssetData, TArray<FGuid> ObjectBindings, UMovieSceneTrack* Track, TFunction<void(UMovieSceneSection*)> InSectionCallback);
 
 	/** Animation asset enter pressed */
-	void OnAnimationDatabaseAssetEnterPressed(const TArray<FAssetData>& AssetData, TArray<FGuid> ObjectBindings, UMovieSceneTrack* Track);
+	void OnAnimationDatabaseAssetEnterPressed(const TArray<FAssetData>& AssetData, TArray<FGuid> ObjectBindings, UMovieSceneTrack* Track, TFunction<void(UMovieSceneSection*)> InSectionCallback);
 
 	/** Delegate for AnimatablePropertyChanged in AddKey */
-	FKeyPropertyResult AddKeyInternal(FFrameNumber KeyTime, UObject* Object, UPoseSearchDatabase* PoseSearchDatabase, UMovieSceneTrack* Track, int32 RowIndex);
+	FKeyPropertyResult AddKeyInternal(FFrameNumber KeyTime, UObject* Object, UPoseSearchDatabase* PoseSearchDatabase, UMovieSceneTrack* Track, int32 RowIndex, TFunction<void(UMovieSceneSection*)> InSectionCallback);
+
+	// IMovieSceneAnimMixerItemMenuProvider interface
+	const UClass* GetHandledMixerItemClass() const override;
+	void PopulateAddMixerItemMenu(FMenuBuilder& MenuBuilder, TArray<FGuid> ObjectBindings, UMovieSceneTrack* Track, TSharedPtr<ISequencer> Sequencer, int32 RowIndex) override;
+	void PopulateObjectBindingAnimationMenu(FMenuBuilder& MenuBuilder, const TArray<FGuid>& ObjectBindings, const UClass* ObjectClass, bool bIsInsideSubmenu) override;
 };
 
 
diff --git a/Engine/Plugins/Experimental/SequenceNavigator/Source/SequenceNavigator/Private/NavigationTool.cpp b/Engine/Plugins/Experimental/SequenceNavigator/Source/SequenceNavigator/Private/NavigationTool.cpp
--- a/Engine/Plugins/Experimental/SequenceNavigator/Source/SequenceNavigator/Private/NavigationTool.cpp
+++ b/Engine/Plugins/Experimental/SequenceNavigator/Source/SequenceNavigator/Private/NavigationTool.cpp
@@ -498,14 +498,14 @@ FNavigationToolViewModelPtr FNavigationTool::FindItem(const FNavigationToolItemI
 
 TArray<FNavigationToolViewModelPtr> FNavigationTool::TryFindItems(const FViewModelPtr& InSequencerOutlinerViewModel) const
 {
-	if (const TViewModelPtr<FTrackModel> TrackModel = InSequencerOutlinerViewModel.ImplicitCast())
+	if (const TViewModelPtr<ITrackExtension> TrackExtension = InSequencerOutlinerViewModel.ImplicitCast())
 	{
-		return FindItemsFromMovieSceneObject(TrackModel->GetTrack());
+		return FindItemsFromMovieSceneObject(TrackExtension->GetTrack());
 	}
 
-	if (const TViewModelPtr<FTrackRowModel> TrackRowModel = InSequencerOutlinerViewModel.ImplicitCast())
+	if (const TViewModelPtr<ITrackRowExtension> TrackRowExtension = InSequencerOutlinerViewModel.ImplicitCast())
 	{
-		return FindItemsFromMovieSceneObject(TrackRowModel->GetTrack());
+		return FindItemsFromMovieSceneObject(TrackRowExtension->GetParentTrack());
 	}
 
 	if (const TViewModelPtr<FObjectBindingModel> ObjectBindingModel = InSequencerOutlinerViewModel.ImplicitCast())
diff --git a/Engine/Plugins/FX/Niagara/Source/NiagaraEditor/Private/Sequencer/LevelSequence/NiagaraSystemTrackEditor.cpp b/Engine/Plugins/FX/Niagara/Source/NiagaraEditor/Private/Sequencer/LevelSequence/NiagaraSystemTrackEditor.cpp
--- a/Engine/Plugins/FX/Niagara/Source/NiagaraEditor/Private/Sequencer/LevelSequence/NiagaraSystemTrackEditor.cpp
+++ b/Engine/Plugins/FX/Niagara/Source/NiagaraEditor/Private/Sequencer/LevelSequence/NiagaraSystemTrackEditor.cpp
@@ -18,6 +18,7 @@
 
 #include "MVVM/ViewModels/OutlinerColumns/OutlinerColumnTypes.h"
 #include "MVVM/Extensions/ITrackExtension.h"
+#include "MVVM/ViewModels/ViewModelIterators.h"
 
 #define LOCTEXT_NAMESPACE "NiagaraSystemTrackEditor"
 
@@ -144,7 +145,13 @@ TSharedPtr<SWidget> FNiagaraSystemTrackEditor::BuildOutlinerColumnWidget(const F
 {
 	using namespace UE::Sequencer;
 
-	UMovieSceneNiagaraSystemTrack* NiagaraSystemTrack = Cast<UMovieSceneNiagaraSystemTrack>(Params.TrackModel->GetTrack());
+	UMovieSceneTrack* Track = nullptr;
+	if (TViewModelPtr<ITrackExtension> TrackModel = Params.ViewModel->FindAncestorOfType<ITrackExtension>(true))
+	{
+		Track = TrackModel->GetTrack();
+	}
+
+	UMovieSceneNiagaraSystemTrack* NiagaraSystemTrack = Cast<UMovieSceneNiagaraSystemTrack>(Track);
 	if (!NiagaraSystemTrack)
 	{
 		return nullptr;
diff --git a/Engine/Plugins/FX/Niagara/Source/NiagaraEditor/Private/Sequencer/NiagaraSequence/NiagaraEmitterTrackEditor.cpp b/Engine/Plugins/FX/Niagara/Source/NiagaraEditor/Private/Sequencer/NiagaraSequence/NiagaraEmitterTrackEditor.cpp
--- a/Engine/Plugins/FX/Niagara/Source/NiagaraEditor/Private/Sequencer/NiagaraSequence/NiagaraEmitterTrackEditor.cpp
+++ b/Engine/Plugins/FX/Niagara/Source/NiagaraEditor/Private/Sequencer/NiagaraSequence/NiagaraEmitterTrackEditor.cpp
@@ -27,6 +27,7 @@
 #include "MVVM/ViewModels/OutlinerColumns/OutlinerColumnTypes.h"
 #include "MVVM/Views/SOutlinerItemViewBase.h"
 #include "MVVM/Views/SOutlinerTrackColorPicker.h"
+#include "MVVM/ViewModels/ViewModelIterators.h"
 
 #define LOCTEXT_NAMESPACE "NiagaraEmitterTrackEditor"
 
@@ -334,8 +335,13 @@ TSharedPtr<SWidget> FNiagaraEmitterTrackEditor::BuildOutlinerColumnWidget(const
 
 	if (ColumnName == FCommonOutlinerNames::Edit)
 	{
-		UMovieSceneNiagaraEmitterTrack* Track = CastChecked<UMovieSceneNiagaraEmitterTrack>(Params.TrackModel->GetTrack());
-		return SNew(SEmitterTrackWidget, *Track);
+		UMovieSceneNiagaraEmitterTrack* Track = nullptr;
+		if (TViewModelPtr<ITrackExtension> TrackModel = Params.ViewModel->FindAncestorOfType<ITrackExtension>(true))
+		{
+			Track = CastChecked<UMovieSceneNiagaraEmitterTrack>(TrackModel->GetTrack());
+		}
+
+		return Track ? SNew(SEmitterTrackWidget, *Track) : SNullWidget::NullWidget;;
 	}
 
 	if (ColumnName == FCommonOutlinerNames::Label)
diff --git a/Engine/Plugins/FX/NiagaraSimCaching/Source/NiagaraSimCachingEditor/Private/Sequencer/NiagaraCacheTrackEditor.cpp b/Engine/Plugins/FX/NiagaraSimCaching/Source/NiagaraSimCachingEditor/Private/Sequencer/NiagaraCacheTrackEditor.cpp
--- a/Engine/Plugins/FX/NiagaraSimCaching/Source/NiagaraSimCachingEditor/Private/Sequencer/NiagaraCacheTrackEditor.cpp
+++ b/Engine/Plugins/FX/NiagaraSimCaching/Source/NiagaraSimCachingEditor/Private/Sequencer/NiagaraCacheTrackEditor.cpp
@@ -27,6 +27,7 @@
 #include "Widgets/Input/SButton.h"
 #include "MVVM/ViewModels/ViewDensity.h"
 #include "MVVM/ViewModels/OutlinerColumns/OutlinerColumnTypes.h"
+#include "MVVM/ViewModels/ViewModelIterators.h"
 #include "MVVM/Extensions/ITrackExtension.h"
 #include "MVVM/Extensions/IObjectBindingExtension.h"
 
@@ -517,7 +518,13 @@ TSharedPtr<SWidget> FNiagaraCacheTrackEditor::BuildOutlinerColumnWidget(const FB
 		TSharedPtr<ISequencer> SequencerPtr = GetSequencer();
 		FWeakObjectPtr BoundObject = SequencerPtr.IsValid() ? SequencerPtr->FindSpawnedObjectOrTemplate(ObjectBinding) : nullptr;
 		
-		UMovieSceneTrack* Track = Params.TrackModel->GetTrack();
+
+		UMovieSceneTrack* Track = nullptr;
+		if (TViewModelPtr<ITrackExtension> TrackModel = Params.ViewModel->FindAncestorOfType<ITrackExtension>(true))
+		{
+			Track = TrackModel->GetTrack();
+		}
+
 
 		return SNew(SHorizontalBox)
 		+ SHorizontalBox::Slot()
diff --git a/Engine/Plugins/MovieScene/LevelSequenceEditor/Source/LevelSequenceEditor/Private/LevelSequenceEditorBlueprintLibrary.cpp b/Engine/Plugins/MovieScene/LevelSequenceEditor/Source/LevelSequenceEditor/Private/LevelSequenceEditorBlueprintLibrary.cpp
--- a/Engine/Plugins/MovieScene/LevelSequenceEditor/Source/LevelSequenceEditor/Private/LevelSequenceEditorBlueprintLibrary.cpp
+++ b/Engine/Plugins/MovieScene/LevelSequenceEditor/Source/LevelSequenceEditor/Private/LevelSequenceEditorBlueprintLibrary.cpp
@@ -565,7 +565,7 @@ void ULevelSequenceEditorBlueprintLibrary::SelectKeys(const FSequencerChannelPro
 			TSharedPtr<FSectionModel> SectionHandle = SectionModelStorage->FindModelForSection(Section);
 			if (SectionHandle)
 			{
-				TParentFirstChildIterator<FChannelGroupModel> KeyAreaNodes = SectionHandle->GetParentTrackModel().AsModel()->GetDescendantsOfType<FChannelGroupModel>();
+				TParentFirstChildIterator<FChannelGroupModel> KeyAreaNodes = SectionHandle->GetParentSectionOwnerModel().AsModel()->GetDescendantsOfType<FChannelGroupModel>();
 				for (const TViewModelPtr<FChannelGroupModel>& KeyAreaNode : KeyAreaNodes)
 				{
 					if (KeyAreaNode->GetChannelName() == ChannelProxy.ChannelName)
diff --git a/Engine/Plugins/MovieScene/SequenceValidator/Source/SequenceValidator/Private/Widgets/SSequenceValidatorResults.cpp b/Engine/Plugins/MovieScene/SequenceValidator/Source/SequenceValidator/Private/Widgets/SSequenceValidatorResults.cpp
--- a/Engine/Plugins/MovieScene/SequenceValidator/Source/SequenceValidator/Private/Widgets/SSequenceValidatorResults.cpp
+++ b/Engine/Plugins/MovieScene/SequenceValidator/Source/SequenceValidator/Private/Widgets/SSequenceValidatorResults.cpp
@@ -345,7 +345,7 @@ void SSequenceValidatorResults::OnTreeViewMouseButtonDoubleClick(FTreeItemPtr It
 			{
 				if (const TViewModelPtr<FSectionModel> SectionModel = TrackAreaItem.ImplicitCast())
 				{
-					TParentFirstChildIterator<FChannelGroupModel> KeyAreaNodes = SectionModel->GetParentTrackModel().AsModel()->GetDescendantsOfType<FChannelGroupModel>();
+					TParentFirstChildIterator<FChannelGroupModel> KeyAreaNodes = SectionModel->GetParentSectionOwnerModel().AsModel()->GetDescendantsOfType<FChannelGroupModel>();
 					for (const TViewModelPtr<FChannelGroupModel>& KeyAreaNode : KeyAreaNodes)
 					{
 						TSharedPtr<FChannelModel> KeyAreaChannelModel = KeyAreaNode->GetChannel(SectionModel);
diff --git a/Engine/Plugins/MovieScene/SequencerScripting/Source/SequencerScriptingEditor/Private/SequencerCurveEditorObject.cpp b/Engine/Plugins/MovieScene/SequencerScripting/Source/SequencerScriptingEditor/Private/SequencerCurveEditorObject.cpp
--- a/Engine/Plugins/MovieScene/SequencerScripting/Source/SequencerScriptingEditor/Private/SequencerCurveEditorObject.cpp
+++ b/Engine/Plugins/MovieScene/SequencerScripting/Source/SequencerScriptingEditor/Private/SequencerCurveEditorObject.cpp
@@ -168,7 +168,7 @@ void USequencerCurveEditorObject::ShowCurve(const FSequencerChannelProxy& Channe
 			TParentFirstChildIterator<IOutlinerExtension> OutlinerExtenstionIt = OutlinerViewModel->GetRootItem()->GetDescendantsOfType<IOutlinerExtension>();
 			for (; OutlinerExtenstionIt; ++OutlinerExtenstionIt)
 			{
-				if (TSharedPtr<FTrackModel> TrackModel = OutlinerExtenstionIt.GetCurrentItem()->FindAncestorOfType<FTrackModel>())
+				if (TViewModelPtr<ITrackExtension> TrackModel = OutlinerExtenstionIt.GetCurrentItem()->FindAncestorOfType<ITrackExtension>())
 				{
 					if (UMovieSceneTrack* Track = TrackModel->GetTrack())
 					{
diff --git a/Engine/Source/Editor/MovieSceneTools/Private/MVVM/ViewModels/TimeWarpChannelModel.cpp b/Engine/Source/Editor/MovieSceneTools/Private/MVVM/ViewModels/TimeWarpChannelModel.cpp
--- a/Engine/Source/Editor/MovieSceneTools/Private/MVVM/ViewModels/TimeWarpChannelModel.cpp
+++ b/Engine/Source/Editor/MovieSceneTools/Private/MVVM/ViewModels/TimeWarpChannelModel.cpp
@@ -5,6 +5,7 @@
 #include "MVVM/ViewModels/SectionModel.h"
 #include "MVVM/ViewModels/OutlinerColumns/OutlinerColumnTypes.h"
 #include "MVVM/Extensions/ITrackExtension.h"
+#include "MVVM/Extensions/ISectionOwnerExtension.h"
 #include "MVVM/Views/ViewUtilities.h"
 #include "SequencerUtilities.h"
 
@@ -61,10 +62,10 @@ TSharedRef<SWidget> FTimeWarpChannelModel::BuildReplaceTimeWarpSubMenu()
 
 void FTimeWarpChannelModel::PopulateReplaceTimeWarpSubMenu(FMenuBuilder& MenuBuilder)
 {
-	TViewModelPtr<ITrackExtension> Track = GetLinkedOutlinerItem().AsModel()->FindAncestorOfType<ITrackExtension>();
-	if (Track)
+	TViewModelPtr<ISectionOwnerExtension> SectionOwner = GetLinkedOutlinerItem().AsModel()->FindAncestorOfType<ISectionOwnerExtension>();
+	if (SectionOwner)
 	{
-		FSequencerUtilities::PopulateTimeWarpChannelSubMenu(MenuBuilder, Track);
+		FSequencerUtilities::PopulateTimeWarpChannelSubMenu(MenuBuilder, SectionOwner.AsWeak());
 	}
 }
 
diff --git a/Engine/Source/Editor/MovieSceneTools/Private/TrackEditors/AudioTrackEditor.cpp b/Engine/Source/Editor/MovieSceneTools/Private/TrackEditors/AudioTrackEditor.cpp
--- a/Engine/Source/Editor/MovieSceneTools/Private/TrackEditors/AudioTrackEditor.cpp
+++ b/Engine/Source/Editor/MovieSceneTools/Private/TrackEditors/AudioTrackEditor.cpp
@@ -1402,7 +1402,8 @@ TSharedPtr<SWidget> FAudioTrackEditor::BuildOutlinerColumnWidget(const FBuildCol
 {
 	using namespace UE::Sequencer;
 
-	UMovieSceneTrack* Track = Params.TrackModel->GetTrack();
+	TViewModelPtr<ITrackExtension> TrackModel = Params.ViewModel->FindAncestorOfType<ITrackExtension>(true);
+	UMovieSceneTrack* Track = TrackModel ? TrackModel->GetTrack() : nullptr;
 	if (!Track)
 	{
 		return SNullWidget::NullWidget;
@@ -1410,7 +1411,7 @@ TSharedPtr<SWidget> FAudioTrackEditor::BuildOutlinerColumnWidget(const FBuildCol
 
 	if (ColumnName == FCommonOutlinerNames::Edit)
 	{
-		if (Params.TrackModel.AsModel()->IsA<FTrackModel>())
+		if (TrackModel.AsModel()->IsA<FTrackModel>())
 		{
 			return SNew(AudioTrack::SLanguagePreviewComboWidget, Track->GetTypedOuter<UMovieScene>(), Track);
 		}
diff --git a/Engine/Source/Editor/MovieSceneTools/Private/TrackEditors/CinematicShotTrackEditor.cpp b/Engine/Source/Editor/MovieSceneTools/Private/TrackEditors/CinematicShotTrackEditor.cpp
--- a/Engine/Source/Editor/MovieSceneTools/Private/TrackEditors/CinematicShotTrackEditor.cpp
+++ b/Engine/Source/Editor/MovieSceneTools/Private/TrackEditors/CinematicShotTrackEditor.cpp
@@ -79,11 +79,13 @@ TSharedPtr<SWidget> FCinematicShotTrackEditor::BuildOutlinerColumnWidget(const F
 	{
 		return UE::Sequencer::MakeAddButton(
 			LOCTEXT("CinematicShotText", "Shot"),
-			FOnGetContent::CreateSP(this, &FCinematicShotTrackEditor::HandleAddSubSequenceComboButtonGetMenuContent, Params.TrackModel.AsWeak()),
+			FOnGetContent::CreateSP(this, &FCinematicShotTrackEditor::HandleAddSubSequenceComboButtonGetMenuContent, Params.ViewModel.AsWeak()),
 			Params.ViewModel);
 	}
 
-	if (!Params.ViewModel->IsA<FTrackRowModel>())
+	// Only show camera lock on track-level models, not on rows or layers
+	ITrackRowExtension* TrackRowExt = Params.ViewModel->CastThis<ITrackRowExtension>();
+	if (!TrackRowExt)
 	{
 		bool bAddCameraLock = false;
 		if (ColumnName == FCommonOutlinerNames::Nav)
diff --git a/Engine/Source/Editor/MovieSceneTools/Private/TrackEditors/CommonAnimationTrackEditor.cpp b/Engine/Source/Editor/MovieSceneTools/Private/TrackEditors/CommonAnimationTrackEditor.cpp
--- a/Engine/Source/Editor/MovieSceneTools/Private/TrackEditors/CommonAnimationTrackEditor.cpp
+++ b/Engine/Source/Editor/MovieSceneTools/Private/TrackEditors/CommonAnimationTrackEditor.cpp
@@ -6,6 +6,7 @@
 #include "Rendering/DrawElements.h"
 #include "Widgets/SBoxPanel.h"
 #include "Framework/MultiBox/MultiBoxBuilder.h"
+#include "Features/IModularFeatures.h"
 #include "GameFramework/Actor.h"
 #include "AssetRegistry/AssetData.h"
 #include "Animation/AnimSequenceBase.h"
@@ -95,6 +96,7 @@
 #include "TimeToPixel.h"
 #include "SequencerUtilities.h"
 #include "SequencerAnimationOverride.h"
+#include "MVVM/ViewModels/ViewModelIterators.h"
 
 #define LOCTEXT_NAMESPACE "FCommonAnimationTrackEditor"
 
@@ -1400,7 +1402,7 @@ void FCommonAnimationTrackEditor::OnInitialize()
 
 	++FCommonAnimationTrackEditor::NumberActive;
 
-	// Activate the default mode in case FEditorModeTools::Tick isn't run before here. 
+	// Activate the default mode in case FEditorModeTools::Tick isn't run before here.
 	// This can be removed once a general fix for UE-143791 has been implemented.
 	GLevelEditorModeTools().ActivateDefaultMode();
 
@@ -1410,8 +1412,8 @@ void FCommonAnimationTrackEditor::OnInitialize()
 	{
 		EditMode->SetSequencer(GetSequencer());
 	}
-
 }
+
 void FCommonAnimationTrackEditor::OnRelease()
 {
 	--FCommonAnimationTrackEditor::NumberActive;
@@ -1433,7 +1435,6 @@ void FCommonAnimationTrackEditor::OnRelease()
 	{
 		GLevelEditorModeTools().DeactivateMode(FSkeletalAnimationTrackEditMode::ModeName);
 	}
-
 }
 
 void FCommonAnimationTrackEditor::AddReferencedObjects(FReferenceCollector& Collector)
@@ -2056,36 +2057,63 @@ void FCommonAnimationTrackEditor::ConstructObjectBindingTrackMenu(FMenuBuilder&
 
 void FCommonAnimationTrackEditor::BuildObjectBindingTrackMenu(FMenuBuilder& MenuBuilder, const TArray<FGuid>& ObjectBindings, const UClass* ObjectClass)
 {
-	if (ObjectClass->IsChildOf(USkeletalMeshComponent::StaticClass()) || ObjectClass->IsChildOf(AActor::StaticClass()) || ObjectClass->IsChildOf(UChildActorComponent::StaticClass()))
+	if (!ObjectClass || (!ObjectClass->IsChildOf(USkeletalMeshComponent::StaticClass()) && !ObjectClass->IsChildOf(AActor::StaticClass()) && !ObjectClass->IsChildOf(UChildActorComponent::StaticClass())))
 	{
-		const TSharedPtr<ISequencer> ParentSequencer = GetSequencer();
-
-		USkeleton* Skeleton = AcquireSkeletonFromObjectGuid(ObjectBindings[0], GetSequencer());
+		return;
+	}
 
-		if (Skeleton)
-		{
-			// Load the asset registry module
-			FAssetRegistryModule& AssetRegistryModule = FModuleManager::LoadModuleChecked<FAssetRegistryModule>(TEXT("AssetRegistry"));
+	// Collect all animation menu providers and sort by priority
+	TArray<IMovieSceneAnimMixerItemMenuProvider*> Providers = IModularFeatures::Get().GetModularFeatureImplementations<IMovieSceneAnimMixerItemMenuProvider>(IMovieSceneAnimMixerItemMenuProvider::GetModularFeatureName());
+	Providers.Sort([](IMovieSceneAnimMixerItemMenuProvider& A, IMovieSceneAnimMixerItemMenuProvider& B)
+	{
+		return A.GetObjectBindingAnimationMenuPriority() < B.GetObjectBindingAnimationMenuPriority();
+	});
 
-			// Collect a full list of assets with the specified class
-			TArray<FAssetData> AssetDataList;
-			AssetRegistryModule.Get().GetAssetsByClass(UAnimSequenceBase::StaticClass()->GetClassPathName(), AssetDataList, true);
+	// Check if the Animation Mixer plugin is enabled
+	const bool bAnimationMixerEnabled = FModuleManager::Get().IsModuleLoaded(TEXT("MovieSceneAnimMixerEditor"));
 
-			if (AssetDataList.Num())
+	if (bAnimationMixerEnabled)
+	{
+		// Create an "Animation" submenu containing all animation-related track options
+		MenuBuilder.AddSubMenu(
+			LOCTEXT("AddAnimationSubmenu", "Animation"),
+			NSLOCTEXT("Sequencer", "AddAnimationSubmenuTooltip", "Add animation-related tracks."),
+			FNewMenuDelegate::CreateLambda([Providers, ObjectBindings, ObjectClass](FMenuBuilder& SubMenuBuilder)
 			{
-				UMovieSceneTrack* Track = nullptr;
-
-				MenuBuilder.AddSubMenu(
-					LOCTEXT("AddAnimation", "Animation"), NSLOCTEXT("Sequencer", "AddAnimationTooltip", "Adds an animation track."),
-					FNewMenuDelegate::CreateRaw(this, &FCommonAnimationTrackEditor::AddAnimationSubMenu, ObjectBindings, Skeleton, Track)
-				);
-			}
+				for (IMovieSceneAnimMixerItemMenuProvider* Provider : Providers)
+				{
+					Provider->PopulateObjectBindingAnimationMenu(SubMenuBuilder, ObjectBindings, ObjectClass, /*bIsInsideSubmenu=*/ true);
+				}
+			})
+		);
+	}
+	else
+	{
+		// Add animation menu entries directly to the top-level menu
+		for (IMovieSceneAnimMixerItemMenuProvider* Provider : Providers)
+		{
+			Provider->PopulateObjectBindingAnimationMenu(MenuBuilder, ObjectBindings, ObjectClass, /*bIsInsideSubmenu=*/ false);
 		}
 	}
 }
 
-TSharedRef<SWidget> FCommonAnimationTrackEditor::BuildAddAnimationSubMenu(FGuid ObjectBinding, USkeleton* Skeleton, UE::Sequencer::TWeakViewModelPtr<UE::Sequencer::ITrackExtension> WeakTrackModel)
+TSharedRef<SWidget> FCommonAnimationTrackEditor::BuildAddAnimationSubMenu(FGuid ObjectBinding, USkeleton* Skeleton, UE::Sequencer::TWeakViewModelPtr<UE::Sequencer::FViewModel> WeakViewModel)
 {
+	TViewModelPtr<FViewModel> ViewModel = WeakViewModel.Pin();
+	if (!ViewModel)
+	{
+		return SNullWidget::NullWidget;
+	}
+
+	TViewModelPtr<ITrackExtension> TrackModel = ViewModel->FindAncestorOfType<ITrackExtension>(true);
+
+	if (!TrackModel)
+	{
+		return SNullWidget::NullWidget;
+	}
+
+	UMovieSceneTrack* Track = TrackModel->GetTrack();
+
 	FMenuBuilder MenuBuilder(true, nullptr);
 
 	TArray<FGuid> ObjectBindings;
@@ -2094,14 +2122,14 @@ TSharedRef<SWidget> FCommonAnimationTrackEditor::BuildAddAnimationSubMenu(FGuid
 
 	MenuBuilder.BeginSection(NAME_None, LOCTEXT("TimeWarp_Label", "Time Warp"));
 	{
-		FSequencerUtilities::MakeTimeWarpMenuEntry(MenuBuilder, WeakTrackModel);
+		FSequencerUtilities::MakeTimeWarpMenuEntry(MenuBuilder, WeakViewModel);
 	}
 	MenuBuilder.EndSection();
 
 
 	MenuBuilder.BeginSection(NAME_None, LOCTEXT("AddAnimation_Label", "Add Animation"));
 	{
-		AddAnimationSubMenu(MenuBuilder, ObjectBindings, Skeleton, WeakTrackModel.Pin()->GetTrack());
+		AddAnimationSubMenu(MenuBuilder, ObjectBindings, Skeleton, Track);
 	}
 	MenuBuilder.EndSection();
 
@@ -2138,15 +2166,15 @@ bool FCommonAnimationTrackEditor::ShouldFilterAsset(const FAssetData& AssetData)
 	return ((EAdditiveAnimationType)AdditiveTypeEnum->GetValueByName(*EnumString) == AAT_RotationOffsetMeshSpace);
 }
 
-void FCommonAnimationTrackEditor::AddAnimationSubMenu(FMenuBuilder& MenuBuilder, TArray<FGuid> ObjectBindings, USkeleton* Skeleton, UMovieSceneTrack* Track)
+void FCommonAnimationTrackEditor::AddAnimationSubMenu(FMenuBuilder& MenuBuilder, TArray<FGuid> ObjectBindings, USkeleton* Skeleton, UMovieSceneTrack* Track, int32 RowIndex)
 {
 	TSharedPtr<ISequencer> SequencerPtr = GetSequencer();
 	UMovieSceneSequence* Sequence = SequencerPtr.IsValid() ? SequencerPtr->GetFocusedMovieSceneSequence() : nullptr;
 
 	FAssetPickerConfig AssetPickerConfig;
 	{
-		AssetPickerConfig.OnAssetSelected = FOnAssetSelected::CreateRaw( this, &FCommonAnimationTrackEditor::OnAnimationAssetSelected, ObjectBindings, Track);
-		AssetPickerConfig.OnAssetEnterPressed = FOnAssetEnterPressed::CreateRaw( this, &FCommonAnimationTrackEditor::OnAnimationAssetEnterPressed, ObjectBindings, Track);
+		AssetPickerConfig.OnAssetSelected = FOnAssetSelected::CreateRaw(this, &FCommonAnimationTrackEditor::OnAnimationAssetSelected, ObjectBindings, Track, RowIndex);
+		AssetPickerConfig.OnAssetEnterPressed = FOnAssetEnterPressed::CreateRaw(this, &FCommonAnimationTrackEditor::OnAnimationAssetEnterPressed, ObjectBindings, Track, RowIndex);
 		AssetPickerConfig.bAllowNullSelection = false;
 		AssetPickerConfig.bAddFilterUI = true;
 		AssetPickerConfig.bShowTypeInColumnView = false;
@@ -2188,7 +2216,7 @@ bool FCommonAnimationTrackEditor::FilterAnimSequences(const FAssetData& AssetDat
 	return false;
 }
 
-void FCommonAnimationTrackEditor::OnAnimationAssetSelected(const FAssetData& AssetData, TArray<FGuid> ObjectBindings, UMovieSceneTrack* Track)
+void FCommonAnimationTrackEditor::OnAnimationAssetSelected(const FAssetData& AssetData, TArray<FGuid> ObjectBindings, UMovieSceneTrack* Track, int32 RowIndex)
 {
 	FSlateApplication::Get().DismissAllMenus();
 
@@ -2204,17 +2232,16 @@ void FCommonAnimationTrackEditor::OnAnimationAssetSelected(const FAssetData& Ass
 		for (FGuid ObjectBinding : ObjectBindings)
 		{
 			UObject* Object = SequencerPtr->FindSpawnedObjectOrTemplate(ObjectBinding);
-			int32 RowIndex = INDEX_NONE;
 			AnimatablePropertyChanged(FOnKeyProperty::CreateRaw(this, &FCommonAnimationTrackEditor::AddKeyInternal, Object, AnimSequence, Track, RowIndex));
 		}
 	}
 }
 
-void FCommonAnimationTrackEditor::OnAnimationAssetEnterPressed(const TArray<FAssetData>& AssetData, TArray<FGuid> ObjectBindings, UMovieSceneTrack* Track)
+void FCommonAnimationTrackEditor::OnAnimationAssetEnterPressed(const TArray<FAssetData>& AssetData, TArray<FGuid> ObjectBindings, UMovieSceneTrack* Track, int32 RowIndex)
 {
 	if (AssetData.Num() > 0)
 	{
-		OnAnimationAssetSelected(AssetData[0].GetAsset(), ObjectBindings, Track);
+		OnAnimationAssetSelected(AssetData[0].GetAsset(), ObjectBindings, Track, RowIndex);
 	}
 }
 
@@ -2285,7 +2312,7 @@ TSharedPtr<SWidget> FCommonAnimationTrackEditor::BuildOutlinerEditWidget(const F
 
 	if (Skeleton)
 	{
-		FOnGetContent HandleGetAddButtonContent = FOnGetContent::CreateSP(this, &FCommonAnimationTrackEditor::BuildAddAnimationSubMenu, ObjectBinding, Skeleton, Params.TrackModel.AsWeak());
+		FOnGetContent HandleGetAddButtonContent = FOnGetContent::CreateSP(this, &FCommonAnimationTrackEditor::BuildAddAnimationSubMenu, ObjectBinding, Skeleton, Params.ViewModel.AsWeak());
 		return UE::Sequencer::MakeAddButton(LOCTEXT("AnimationText", "Animation"), HandleGetAddButtonContent, Params.ViewModel);
 	}
 	else
diff --git a/Engine/Source/Editor/MovieSceneTools/Private/TrackEditors/SkeletalAnimationTrackEditor.cpp b/Engine/Source/Editor/MovieSceneTools/Private/TrackEditors/SkeletalAnimationTrackEditor.cpp
--- a/Engine/Source/Editor/MovieSceneTools/Private/TrackEditors/SkeletalAnimationTrackEditor.cpp
+++ b/Engine/Source/Editor/MovieSceneTools/Private/TrackEditors/SkeletalAnimationTrackEditor.cpp
@@ -3,8 +3,15 @@
 #include "TrackEditors/SkeletalAnimationTrackEditor.h"
 
 #include "AnimSequencerInstanceProxy.h"
+#include "Animation/AnimSequenceBase.h"
+#include "AssetRegistry/AssetRegistryModule.h"
+#include "Components/SkeletalMeshComponent.h"
+#include "Features/IModularFeatures.h"
 #include "Framework/MultiBox/MultiBoxBuilder.h"
+#include "GameFramework/Actor.h"
+#include "Sections/MovieSceneSkeletalAnimationSection.h"
 #include "Tracks/MovieSceneSkeletalAnimationTrack.h"
+#include "AssetRegistry/AssetData.h"
 
 #define LOCTEXT_NAMESPACE "FSkeletalAnimationTrackEditor"
 
@@ -32,14 +39,84 @@ void FSkeletalAnimationTrackEditor::OnInitialize()
 {
 	FCoreUObjectDelegates::OnObjectPropertyChanged.AddRaw(this, &FSkeletalAnimationTrackEditor::OnPostPropertyChanged);
 	FCommonAnimationTrackEditor::OnInitialize();
+
+	// Register as animation mixer item menu provider
+	IModularFeatures::Get().RegisterModularFeature(IMovieSceneAnimMixerItemMenuProvider::GetModularFeatureName(), this);
 }
 
 void FSkeletalAnimationTrackEditor::OnRelease()
 {
+	// Unregister as animation mixer item menu provider
+	IModularFeatures::Get().UnregisterModularFeature(IMovieSceneAnimMixerItemMenuProvider::GetModularFeatureName(), this);
+
 	FCoreUObjectDelegates::OnObjectPropertyChanged.RemoveAll(this);
 	FCommonAnimationTrackEditor::OnRelease();
 }
 
+const UClass* FSkeletalAnimationTrackEditor::GetHandledMixerItemClass() const
+{
+	return UMovieSceneSkeletalAnimationSection::StaticClass();
+}
+
+void FSkeletalAnimationTrackEditor::PopulateAddMixerItemMenu(FMenuBuilder& MenuBuilder, TArray<FGuid> ObjectBindings, UMovieSceneTrack* Track, TSharedPtr<ISequencer> Sequencer, int32 RowIndex)
+{
+	if (ObjectBindings.Num() > 0 && Sequencer.IsValid())
+	{
+		USkeleton* Skeleton = AcquireSkeletonFromObjectGuid(ObjectBindings[0], Sequencer);
+		if (Skeleton)
+		{
+			AddAnimationSubMenu(MenuBuilder, ObjectBindings, Skeleton, Track, RowIndex);
+		}
+	}
+}
+
+void FSkeletalAnimationTrackEditor::PopulateObjectBindingAnimationMenu(FMenuBuilder& MenuBuilder, const TArray<FGuid>& ObjectBindings, const UClass* ObjectClass, bool bIsInsideSubmenu)
+{
+	if (!ObjectClass || (!ObjectClass->IsChildOf(USkeletalMeshComponent::StaticClass()) && !ObjectClass->IsChildOf(AActor::StaticClass()) && !ObjectClass->IsChildOf(UChildActorComponent::StaticClass())))
+	{
+		return;
+	}
+
+	USkeleton* Skeleton = AcquireSkeletonFromObjectGuid(ObjectBindings[0], GetSequencer());
+
+	if (Skeleton)
+	{
+		// Load the asset registry module
+		FAssetRegistryModule& AssetRegistryModule = FModuleManager::LoadModuleChecked<FAssetRegistryModule>(TEXT("AssetRegistry"));
+
+		// Collect a full list of assets with the specified class
+		TArray<FAssetData> AssetDataList;
+		AssetRegistryModule.Get().GetAssetsByClass(UAnimSequenceBase::StaticClass()->GetClassPathName(), AssetDataList, true);
+
+		if (AssetDataList.Num())
+		{
+			UMovieSceneTrack* Track = nullptr;
+
+			if (bIsInsideSubmenu)
+			{
+				// Inside the Animation submenu - add a section header and embed the asset picker directly
+				MenuBuilder.BeginSection(NAME_None, LOCTEXT("AnimationSectionHeader", "Animation"));
+				{
+					AddAnimationSubMenu(MenuBuilder, ObjectBindings, Skeleton, Track, INDEX_NONE);
+				}
+				MenuBuilder.EndSection();
+			}
+			else
+			{
+				// At the top-level menu - wrap in a submenu
+				MenuBuilder.AddSubMenu(
+					LOCTEXT("AddAnimationSubMenu", "Animation"),
+					LOCTEXT("AddAnimationSubMenuTooltip", "Adds an animation track."),
+					FNewMenuDelegate::CreateLambda([this, ObjectBindings, Skeleton, Track](FMenuBuilder& SubMenuBuilder)
+					{
+						AddAnimationSubMenu(SubMenuBuilder, ObjectBindings, Skeleton, Track, INDEX_NONE);
+					})
+				);
+			}
+		}
+	}
+}
+
 bool FSkeletalAnimationTrackEditor::SupportsSequence(UMovieSceneSequence* InSequence) const
 {
 	ETrackSupport TrackSupported = InSequence ? InSequence->IsTrackSupported(UMovieSceneSkeletalAnimationTrack::StaticClass()) : ETrackSupport::NotSupported;
diff --git a/Engine/Source/Editor/MovieSceneTools/Private/TrackEditors/SkeletalAnimationTrackEditor.h b/Engine/Source/Editor/MovieSceneTools/Private/TrackEditors/SkeletalAnimationTrackEditor.h
--- a/Engine/Source/Editor/MovieSceneTools/Private/TrackEditors/SkeletalAnimationTrackEditor.h
+++ b/Engine/Source/Editor/MovieSceneTools/Private/TrackEditors/SkeletalAnimationTrackEditor.h
@@ -3,11 +3,14 @@
 #pragma once
 
 #include "TrackEditors/CommonAnimationTrackEditor.h"
+#include "IMovieSceneAnimMixerItemMenuProvider.h"
 
 /**
  * Tools for animation tracks
  */
-class FSkeletalAnimationTrackEditor : public UE::Sequencer::FCommonAnimationTrackEditor
+class FSkeletalAnimationTrackEditor
+	: public UE::Sequencer::FCommonAnimationTrackEditor
+	, public IMovieSceneAnimMixerItemMenuProvider
 {
 public:
 
@@ -30,6 +33,13 @@ public:
 	virtual TSharedRef<ISequencerSection> MakeSectionInterface(UMovieSceneSection& SectionObject, UMovieSceneTrack& Track, FGuid ObjectBinding) override;
 	virtual TSubclassOf<UMovieSceneCommonAnimationTrack> GetTrackClass() const override;
 
+	// IMovieSceneAnimMixerItemMenuProvider interface
+	virtual const UClass* GetHandledMixerItemClass() const override;
+	virtual void PopulateAddMixerItemMenu(FMenuBuilder& MenuBuilder, TArray<FGuid> ObjectBindings, UMovieSceneTrack* Track, TSharedPtr<ISequencer> Sequencer, int32 RowIndex) override;
+	virtual int32 GetMixerItemMenuPriority() const override { return 0; } // Highest priority in mixer '+' menu
+	virtual void PopulateObjectBindingAnimationMenu(FMenuBuilder& MenuBuilder, const TArray<FGuid>& ObjectBindings, const UClass* ObjectClass, bool bIsInsideSubmenu) override;
+	virtual int32 GetObjectBindingAnimationMenuPriority() const override { return 25; } // After Animation Mixer in Object Binding menu
+
 private:
 
 	void OnPostPropertyChanged(UObject* InObject, struct FPropertyChangedEvent& InPropertyChangedEvent);
diff --git a/Engine/Source/Editor/MovieSceneTools/Private/TrackEditors/SubTrackEditor.cpp b/Engine/Source/Editor/MovieSceneTools/Private/TrackEditors/SubTrackEditor.cpp
--- a/Engine/Source/Editor/MovieSceneTools/Private/TrackEditors/SubTrackEditor.cpp
+++ b/Engine/Source/Editor/MovieSceneTools/Private/TrackEditors/SubTrackEditor.cpp
@@ -335,7 +335,7 @@ void FSubTrackEditor::BuildAddTrackMenu(FMenuBuilder& MenuBuilder)
 
 TSharedPtr<SWidget> FSubTrackEditor::BuildOutlinerEditWidget(const FGuid& ObjectBinding, UMovieSceneTrack* Track, const FBuildEditWidgetParams& Params)
 {
-	return UE::Sequencer::MakeAddButton(GetSubTrackName(), FOnGetContent::CreateSP(this, &FSubTrackEditor::HandleAddSubSequenceComboButtonGetMenuContent, Params.TrackModel.AsWeak()), Params.ViewModel);
+	return UE::Sequencer::MakeAddButton(GetSubTrackName(), FOnGetContent::CreateSP(this, &FSubTrackEditor::HandleAddSubSequenceComboButtonGetMenuContent, Params.ViewModel.AsWeak()), Params.ViewModel);
 }
 
 
@@ -1570,11 +1570,18 @@ UMovieSceneSubTrack* FSubTrackEditor::FindOrCreateSubTrack(UMovieScene* MovieSce
 	return SubTrack;
 }
 
-TSharedRef<SWidget> FSubTrackEditor::HandleAddSubSequenceComboButtonGetMenuContent(UE::Sequencer::TWeakViewModelPtr<UE::Sequencer::ITrackExtension> WeakTrackModel)
+TSharedRef<SWidget> FSubTrackEditor::HandleAddSubSequenceComboButtonGetMenuContent(UE::Sequencer::TWeakViewModelPtr<UE::Sequencer::FViewModel> WeakViewModel)
 {
 	using namespace UE::Sequencer;
 
-	TViewModelPtr<ITrackExtension> TrackModel = WeakTrackModel.Pin();
+	TViewModelPtr<FViewModel> ViewModel = WeakViewModel.Pin();
+	if (!ViewModel)
+	{
+		return SNullWidget::NullWidget;
+	}
+
+	TViewModelPtr<ITrackExtension> TrackModel = ViewModel->FindAncestorOfType<ITrackExtension>(true);
+
 	if (!TrackModel)
 	{
 		return SNullWidget::NullWidget;
@@ -1586,7 +1593,7 @@ TSharedRef<SWidget> FSubTrackEditor::HandleAddSubSequenceComboButtonGetMenuConte
 
 	MenuBuilder.BeginSection(NAME_None, LOCTEXT("TimeWarpCategory", "Time Warp"));
 	{
-		FSequencerUtilities::MakeTimeWarpMenuEntry(MenuBuilder, WeakTrackModel);
+		FSequencerUtilities::MakeTimeWarpMenuEntry(MenuBuilder, WeakViewModel);
 	}
 
 	MenuBuilder.BeginSection(TEXT("ChooseSequence"), LOCTEXT("InsertSequence", "Insert Sequence"));
diff --git a/Engine/Source/Editor/MovieSceneTools/Public/TrackEditors/CommonAnimationTrackEditor.h b/Engine/Source/Editor/MovieSceneTools/Public/TrackEditors/CommonAnimationTrackEditor.h
--- a/Engine/Source/Editor/MovieSceneTools/Public/TrackEditors/CommonAnimationTrackEditor.h
+++ b/Engine/Source/Editor/MovieSceneTools/Public/TrackEditors/CommonAnimationTrackEditor.h
@@ -12,6 +12,7 @@
 #include "ISequencerSection.h"
 #include "MovieSceneTrackEditor.h"
 #include "SequencerCoreFwd.h"
+#include "IMovieSceneAnimMixerItemMenuProvider.h"
 
 #define UE_API MOVIESCENETOOLS_API
 
@@ -36,7 +37,9 @@ class ITrackExtension;
 /**
  * Tools for animation tracks
  */
-class FCommonAnimationTrackEditor : public FMovieSceneTrackEditor, public FGCObject
+class FCommonAnimationTrackEditor
+	: public FMovieSceneTrackEditor
+	, public FGCObject
 {
 public:
 	/**
@@ -84,9 +87,9 @@ protected:
 	virtual TSubclassOf<UMovieSceneCommonAnimationTrack> GetTrackClass() const = 0;
 
 	/** Animation sub menu */
-	UE_API TSharedRef<SWidget> BuildAddAnimationSubMenu(FGuid ObjectBinding, USkeleton* Skeleton, UE::Sequencer::TWeakViewModelPtr<UE::Sequencer::ITrackExtension> TrackModel);
+	UE_API TSharedRef<SWidget> BuildAddAnimationSubMenu(FGuid ObjectBinding, USkeleton* Skeleton, UE::Sequencer::TWeakViewModelPtr<UE::Sequencer::FViewModel> WeakViewModel);
 	UE_API TSharedRef<SWidget> BuildAnimationSubMenu(FGuid ObjectBinding, USkeleton* Skeleton, UMovieSceneTrack* Track);
-	UE_API void AddAnimationSubMenu(FMenuBuilder& MenuBuilder, TArray<FGuid> ObjectBindings, USkeleton* Skeleton, UMovieSceneTrack* Track);
+	UE_API void AddAnimationSubMenu(FMenuBuilder& MenuBuilder, TArray<FGuid> ObjectBindings, USkeleton* Skeleton, UMovieSceneTrack* Track, int32 RowIndex = INDEX_NONE);
 
 	/** Filter only compatible skeletons */
 	UE_API bool FilterAnimSequences(const FAssetData& AssetData, USkeleton* Skeleton);
@@ -95,10 +98,10 @@ protected:
 	UE_API bool ShouldFilterAsset(const FAssetData& AssetData);
 
 	/** Animation asset selected */
-	UE_API void OnAnimationAssetSelected(const FAssetData& AssetData, TArray<FGuid> ObjectBindings, UMovieSceneTrack* Track);
+	UE_API void OnAnimationAssetSelected(const FAssetData& AssetData, TArray<FGuid> ObjectBindings, UMovieSceneTrack* Track, int32 RowIndex = INDEX_NONE);
 
 	/** Animation asset enter pressed */
-	UE_API void OnAnimationAssetEnterPressed(const TArray<FAssetData>& AssetData, TArray<FGuid> ObjectBindings, UMovieSceneTrack* Track);
+	UE_API void OnAnimationAssetEnterPressed(const TArray<FAssetData>& AssetData, TArray<FGuid> ObjectBindings, UMovieSceneTrack* Track, int32 RowIndex = INDEX_NONE);
 
 	/** Delegate for AnimatablePropertyChanged in AddKey */
 	UE_API FKeyPropertyResult AddKeyInternal(FFrameNumber KeyTime, UObject* Object, UAnimSequenceBase* AnimSequence, UMovieSceneTrack* Track, int32 RowIndex);
diff --git a/Engine/Source/Editor/MovieSceneTools/Public/TrackEditors/SubTrackEditor.h b/Engine/Source/Editor/MovieSceneTools/Public/TrackEditors/SubTrackEditor.h
--- a/Engine/Source/Editor/MovieSceneTools/Public/TrackEditors/SubTrackEditor.h
+++ b/Engine/Source/Editor/MovieSceneTools/Public/TrackEditors/SubTrackEditor.h
@@ -167,7 +167,7 @@ protected:
 	UE_API UMovieSceneSubTrack* FindOrCreateSubTrack(UMovieScene* MovieScene, UMovieSceneTrack* Track) const;
 
 	/** Callback for generating the menu of the "Add Sequence" combo button. */
-	UE_API TSharedRef<SWidget> HandleAddSubSequenceComboButtonGetMenuContent(UE::Sequencer::TWeakViewModelPtr<UE::Sequencer::ITrackExtension> TrackModel);
+	UE_API TSharedRef<SWidget> HandleAddSubSequenceComboButtonGetMenuContent(UE::Sequencer::TWeakViewModelPtr<UE::Sequencer::FViewModel> ViewModel);
 
 private:
 
diff --git a/Engine/Source/Editor/Sequencer/Private/Filters/Filters/SequencerTrackFilter_HideIsolate.cpp b/Engine/Source/Editor/Sequencer/Private/Filters/Filters/SequencerTrackFilter_HideIsolate.cpp
--- a/Engine/Source/Editor/Sequencer/Private/Filters/Filters/SequencerTrackFilter_HideIsolate.cpp
+++ b/Engine/Source/Editor/Sequencer/Private/Filters/Filters/SequencerTrackFilter_HideIsolate.cpp
@@ -190,10 +190,10 @@ void FSequencerTrackFilter_HideIsolate::IsolateCategoryGroupTracks(const TSet<TW
 		const TViewModelPtr<FCategoryGroupModel> CategoryGroupModel = Track.AsModel()->FindAncestorOfType<FCategoryGroupModel>(true);
 		if (CategoryGroupModel.IsValid())
 		{
-			const TViewModelPtr<FTrackModel> ParentTrack = Track.AsModel()->FindAncestorOfType<FTrackModel>();
+			const TViewModelPtr<ITrackExtension> ParentTrack = Track.AsModel()->FindAncestorOfType<ITrackExtension>();
 			if (ParentTrack.IsValid())
 			{
-				IsolateChildCategoryGroups(ParentTrack);
+				IsolateChildCategoryGroups(ParentTrack.ImplicitCast());
 			}
 		}
 		else
diff --git a/Engine/Source/Editor/Sequencer/Private/Filters/Filters/SequencerTrackFilter_Keyed.cpp b/Engine/Source/Editor/Sequencer/Private/Filters/Filters/SequencerTrackFilter_Keyed.cpp
--- a/Engine/Source/Editor/Sequencer/Private/Filters/Filters/SequencerTrackFilter_Keyed.cpp
+++ b/Engine/Source/Editor/Sequencer/Private/Filters/Filters/SequencerTrackFilter_Keyed.cpp
@@ -1,4 +1,4 @@
-﻿// Copyright Epic Games, Inc. All Rights Reserved.
+// Copyright Epic Games, Inc. All Rights Reserved.
 
 #include "Filters/Filters/SequencerTrackFilter_Keyed.h"
 #include "Filters/SequencerFilterBar.h"
@@ -6,6 +6,7 @@
 #include "Framework/Commands/Commands.h"
 #include "MVVM/ViewModelPtr.h"
 #include "MVVM/ViewModels/CategoryModel.h"
+#include "MVVM/Extensions/ISectionOwnerExtension.h"
 #include "Sequencer.h"
 
 using namespace UE::Sequencer;
@@ -48,45 +49,9 @@ bool FSequencerTrackFilter_Keyed::PassesFilter(FSequencerTrackFilterType InItem)
 	{
 		return ChannelGroupOutlinerModel->IsAnimated();
 	}
-
-	const TWeakViewModelPtr<ITrackExtension> WeakTrack = GetFilterInterface().GetFilterData().ResolveTrack(InItem);
-	return DoesTrackExtensionHaveKeys(WeakTrack);
-}
-
-FText FSequencerTrackFilter_Keyed::GetDisplayName() const
-{
-	return LOCTEXT("SequenceTrackFilter_Keyed", "Keyed");
-}
-
-FSlateIcon FSequencerTrackFilter_Keyed::GetIcon() const
-{
-	return FSlateIcon(FAppStyle::GetAppStyleSetName(), TEXT("Sequencer.IconKeyUser"));
-}
-
-FString FSequencerTrackFilter_Keyed::GetName() const
-{
-	return StaticName();
-}
-
-bool FSequencerTrackFilter_Keyed::DoesTrackExtensionHaveKeys(const TWeakViewModelPtr<ITrackExtension>& InTrack)
-{
-	const TViewModelPtr<ITrackExtension> Track = InTrack.Pin();
-	if (!Track.IsValid())
+	else if (const TViewModelPtr<ISectionOwnerExtension> WeakSectionOwner = InItem.ImplicitCast())
 	{
-		return false;
-	}
-
-	UMovieSceneTrack* const TrackObject = Track->GetTrack();
-	if (!TrackObject)
-	{
-		return false;
-	}
-
-	const int32 RowIndex = Track->GetRowIndex();
-
-	for (const UMovieSceneSection* const Section : TrackObject->GetAllSections())
-	{
-		if (Section->GetRowIndex() == RowIndex)
+		for (const UMovieSceneSection* const Section : WeakSectionOwner->GetSections())
 		{
 			for (const FMovieSceneChannelEntry& ChannelEntry : Section->GetChannelProxy().GetAllEntries())
 			{
@@ -106,4 +71,19 @@ bool FSequencerTrackFilter_Keyed::DoesTrackExtensionHaveKeys(const TWeakViewMode
 	return false;
 }
 
+FText FSequencerTrackFilter_Keyed::GetDisplayName() const
+{
+	return LOCTEXT("SequenceTrackFilter_Keyed", "Keyed");
+}
+
+FSlateIcon FSequencerTrackFilter_Keyed::GetIcon() const
+{
+	return FSlateIcon(FAppStyle::GetAppStyleSetName(), TEXT("Sequencer.IconKeyUser"));
+}
+
+FString FSequencerTrackFilter_Keyed::GetName() const
+{
+	return StaticName();
+}
+
 #undef LOCTEXT_NAMESPACE
diff --git a/Engine/Source/Editor/Sequencer/Private/Filters/Filters/SequencerTrackFilter_TimeWarp.cpp b/Engine/Source/Editor/Sequencer/Private/Filters/Filters/SequencerTrackFilter_TimeWarp.cpp
--- a/Engine/Source/Editor/Sequencer/Private/Filters/Filters/SequencerTrackFilter_TimeWarp.cpp
+++ b/Engine/Source/Editor/Sequencer/Private/Filters/Filters/SequencerTrackFilter_TimeWarp.cpp
@@ -52,9 +52,9 @@ FString FSequencerTrackFilter_TimeWarp::GetName() const
 
 bool FSequencerTrackFilter_TimeWarp::PassesFilter(FSequencerTrackFilterType InItem) const
 {
-	if (const TViewModelPtr<ITrackExtension> TrackModel = InItem->FindAncestorOfType<ITrackExtension>(true))
+	if (const TViewModelPtr<ISectionOwnerExtension> SectionOwner = InItem->FindAncestorOfType<ISectionOwnerExtension>(true))
 	{
-		for (const TViewModelPtr<FSectionModel>& SectionModel : TrackModel->GetSectionModels().IterateSubList<FSectionModel>())
+		for (const TViewModelPtr<FSectionModel>& SectionModel : SectionOwner->GetSectionModels().IterateSubList<FSectionModel>())
 		{
 			UMovieSceneSection* Section = SectionModel->GetSection();
 			if (!Section)
diff --git a/Engine/Source/Editor/Sequencer/Private/MVVM/Extensions/IConditionableExtension.cpp b/Engine/Source/Editor/Sequencer/Private/MVVM/Extensions/IConditionableExtension.cpp
--- a/Engine/Source/Editor/Sequencer/Private/MVVM/Extensions/IConditionableExtension.cpp
+++ b/Engine/Source/Editor/Sequencer/Private/MVVM/Extensions/IConditionableExtension.cpp
@@ -35,14 +35,14 @@ namespace UE::Sequencer
 			// This allows us to surface on the track row level that a section on that row has a condition
 
 			int32 RowIndex = INDEX_NONE;
-			if (TViewModelPtr<FTrackRowModel> TrackRowModel = ViewModel.ImplicitCast())
+			if (TViewModelPtr<ITrackRowExtension> TrackRowExtension = ViewModel.ImplicitCast())
 			{
-				RowIndex = TrackRowModel->GetRowIndex();
+				RowIndex = TrackRowExtension->GetRowIndex();
 			}
 
-			if (TViewModelPtr<ITrackExtension> Track = ViewModel.ImplicitCast())
+			if (TViewModelPtr<ISectionOwnerExtension> SectionOwner = ViewModel.ImplicitCast())
 			{
-				for (UMovieSceneSection* Section : Track->GetSections())
+				for (UMovieSceneSection* Section : SectionOwner->GetSections())
 				{
 					if (Section 
 						&& (RowIndex == INDEX_NONE || Section->GetRowIndex() == RowIndex)
diff --git a/Engine/Source/Editor/Sequencer/Private/MVVM/ObjectBindingModelStorageExtension.cpp b/Engine/Source/Editor/Sequencer/Private/MVVM/ObjectBindingModelStorageExtension.cpp
--- a/Engine/Source/Editor/Sequencer/Private/MVVM/ObjectBindingModelStorageExtension.cpp
+++ b/Engine/Source/Editor/Sequencer/Private/MVVM/ObjectBindingModelStorageExtension.cpp
@@ -247,7 +247,7 @@ void FObjectBindingModelStorageExtension::OnBindingRemoved(const FGuid& ObjectBi
 void FObjectBindingModelStorageExtension::OnTrackAddedToBinding(UMovieSceneTrack* Track, const FGuid& ObjectBindingID)
 {
 	TSharedPtr<FObjectBindingModel> Model = ObjectBindingToModel.FindRef(ObjectBindingID).Pin();
-	if (Model)
+	if (Model && Track && Track->ParentTrack == nullptr)
 	{
 		Model->AddTrack(Track);
 	}
diff --git a/Engine/Source/Editor/Sequencer/Private/MVVM/Selection.cpp b/Engine/Source/Editor/Sequencer/Private/MVVM/Selection.cpp
--- a/Engine/Source/Editor/Sequencer/Private/MVVM/Selection.cpp
+++ b/Engine/Source/Editor/Sequencer/Private/MVVM/Selection.cpp
@@ -279,6 +279,14 @@ TSet<UMovieSceneTrack*> FSequencerSelection::GetSelectedTracks() const
 		}
 	}
 
+	for (TViewModelPtr<ITrackRowExtension> TrackRowExtension : Outliner.Filter<ITrackRowExtension>())
+	{
+		if (UMovieSceneTrack* Track = TrackRowExtension->GetParentTrack())
+		{
+			SelectedTracks.Add(Track);
+		}
+	}
+
 	return SelectedTracks;
 }
 
@@ -287,19 +295,23 @@ TSet<TPair<UMovieSceneTrack*, int32>> FSequencerSelection::GetSelectedTrackRows(
 	TSet<TPair<UMovieSceneTrack*, int32>> SelectedTrackRows;
 	SelectedTrackRows.Reserve(TrackArea.Num());
 
+	// Only add a 'track row' as selected if either we have a row-level model selected (FTrackRowModel or layer model),
+	// or else we have a track selected and there's only a single track row, and the track allows multiple rows.
+	for (TViewModelPtr<ITrackRowExtension> TrackRowExtension : Outliner.Filter<ITrackRowExtension>())
+	{
+		if (UMovieSceneTrack* Track = TrackRowExtension->GetParentTrack())
+		{
+			SelectedTrackRows.Add(TPair<UMovieSceneTrack*, int32>(Track, TrackRowExtension->GetRowIndex()));
+		}
+	}
+
 	for (TViewModelPtr<ITrackExtension> TrackExtension : Outliner.Filter<ITrackExtension>())
 	{
-		// Only add a 'track row' as selected if either we have an actual 'track row' selected, or else we have a track selected and there's only a single
-		// track row, and the track allows multiple rows.
 		if (UMovieSceneTrack* Track = TrackExtension->GetTrack())
 		{
-			if (TViewModelPtr<FTrackRowModel> TrackRowModel = TrackExtension.ImplicitCast())
-			{
-				SelectedTrackRows.Add(TPair<UMovieSceneTrack*, int32>(Track, TrackExtension->GetRowIndex()));
-			}
-			else if (Track->SupportsMultipleRows() && Track->GetMaxRowIndex() == 0)
+			if (Track->SupportsMultipleRows() && Track->GetMaxRowIndex() == 0)
 			{
-				SelectedTrackRows.Add(TPair<UMovieSceneTrack*, int32>(Track, TrackExtension->GetRowIndex()));
+				SelectedTrackRows.Add(TPair<UMovieSceneTrack*, int32>(Track, 0));
 			}
 		}
 	}
diff --git a/Engine/Source/Editor/Sequencer/Private/MVVM/SequencerViewModelTypeIDs.cpp b/Engine/Source/Editor/Sequencer/Private/MVVM/SequencerViewModelTypeIDs.cpp
--- a/Engine/Source/Editor/Sequencer/Private/MVVM/SequencerViewModelTypeIDs.cpp
+++ b/Engine/Source/Editor/Sequencer/Private/MVVM/SequencerViewModelTypeIDs.cpp
@@ -15,12 +15,16 @@
 
 #include "MVVM/Extensions/ICurveEditorTreeItemExtension.h"
 #include "MVVM/Extensions/ITrackExtension.h"
+#include "MVVM/Extensions/ITrackRowExtension.h"
+#include "MVVM/Extensions/ISectionOwnerExtension.h"
+#include "MVVM/Extensions/ITopLevelChannelHolderExtension.h"
 #include "MVVM/Extensions/IObjectBindingExtension.h"
 #include "MVVM/Extensions/ISnappableExtension.h"
 #include "MVVM/Extensions/IBindingLifetimeExtension.h"
 #include "MVVM/Extensions/IConditionableExtension.h"
 #include "MVVM/Extensions/ITimeDomainExtension.h"
 #include "MVVM/Extensions/IClockExtension.h"
+#include "MVVM/Extensions/ISectionDragOwnerExtension.h"
 
 #include "MVVM/ViewModels/SequencerOutlinerViewModel.h"
 #include "MVVM/ViewModels/SequencerEditorViewModel.h"
@@ -62,8 +66,12 @@ UE_SEQUENCER_DEFINE_CASTABLE(FTrackRowModel);
 
 // Interface types
 UE_SEQUENCER_DEFINE_VIEW_MODEL_TYPE_ID(ITrackExtension);
+UE_SEQUENCER_DEFINE_VIEW_MODEL_TYPE_ID(ITrackRowExtension);
+UE_SEQUENCER_DEFINE_VIEW_MODEL_TYPE_ID(ISectionOwnerExtension);
+UE_SEQUENCER_DEFINE_VIEW_MODEL_TYPE_ID(ITopLevelChannelHolderExtension);
 UE_SEQUENCER_DEFINE_VIEW_MODEL_TYPE_ID(IConditionableExtension);
 UE_SEQUENCER_DEFINE_VIEW_MODEL_TYPE_ID(IClockExtension);
+UE_SEQUENCER_DEFINE_VIEW_MODEL_TYPE_ID(ISectionDragOwnerExtension);
 
 // View model types
 UE_SEQUENCER_DEFINE_CASTABLE(FSequencerEditorViewModel);
diff --git a/Engine/Source/Editor/Sequencer/Private/MVVM/TrackModelStorageExtension.cpp b/Engine/Source/Editor/Sequencer/Private/MVVM/TrackModelStorageExtension.cpp
--- a/Engine/Source/Editor/Sequencer/Private/MVVM/TrackModelStorageExtension.cpp
+++ b/Engine/Source/Editor/Sequencer/Private/MVVM/TrackModelStorageExtension.cpp
@@ -84,7 +84,7 @@ TSharedPtr<FTrackModel> FTrackModelStorageExtension::CreateModelForTrack(UMovieS
 
 	FObjectKey TrackAsKey(InTrack);
 
-	TSharedPtr<FTrackModel> TrackModel = TrackToModel.FindRef(TrackAsKey).Pin();
+	TSharedPtr<FTrackModel> TrackModel = TrackToModel.FindRef(TrackAsKey).Pin().ImplicitCast();
 	if (TrackModel)
 	{
 		// If we already have a track model, just ensure it's added to the correct parent
@@ -124,11 +124,25 @@ TSharedPtr<FTrackModel> FTrackModelStorageExtension::CreateModelForTrack(UMovieS
 }
 
 TSharedPtr<FTrackModel> FTrackModelStorageExtension::FindModelForTrack(UMovieSceneTrack* InTrack) const
+{
+	FObjectKey TrackAsKey(InTrack);
+	TWeakViewModelPtr<ITrackExtension> WeakModel = TrackToModel.FindRef(TrackAsKey);
+	TViewModelPtr<ITrackExtension> Model = WeakModel.Pin();
+	return Model.ImplicitCast();
+}
+
+TViewModelPtr<ITrackExtension> FTrackModelStorageExtension::FindTrackExtensionForTrack(UMovieSceneTrack* InTrack) const
 {
 	FObjectKey TrackAsKey(InTrack);
 	return TrackToModel.FindRef(TrackAsKey).Pin();
 }
 
+void FTrackModelStorageExtension::RegisterTrackExtensionForTrack(UMovieSceneTrack* InTrack, TViewModelPtr<ITrackExtension> InModel)
+{
+	FObjectKey TrackAsKey(InTrack);
+	TrackToModel.Add(TrackAsKey, InModel);
+}
+
 void FTrackModelStorageExtension::OnTrackAdded(UMovieSceneTrack* InTrack)
 {
 	CreateModelForTrack(InTrack);
@@ -138,12 +152,15 @@ void FTrackModelStorageExtension::OnTrackRemoved(UMovieSceneTrack* InTrack)
 {
 	FObjectKey TrackAsKey(InTrack);
 
-	TSharedPtr<FTrackModel> Model = TrackToModel.FindRef(TrackAsKey).Pin();
-
-	if (Model)
+	if (TViewModelPtr<ITrackExtension> Model = TrackToModel.FindRef(TrackAsKey).Pin())
 	{
-		Model->RemoveFromParent();
+		if (TSharedPtr<FViewModel> ViewModel = Model.AsModel())
+		{
+			ViewModel->RemoveFromParent();
+		}
 	}
+
+	TrackToModel.Remove(TrackAsKey);
 }
 
 } // namespace Sequencer
diff --git a/Engine/Source/Editor/Sequencer/Private/MVVM/ViewModels/CategoryModel.cpp b/Engine/Source/Editor/Sequencer/Private/MVVM/ViewModels/CategoryModel.cpp
--- a/Engine/Source/Editor/Sequencer/Private/MVVM/ViewModels/CategoryModel.cpp
+++ b/Engine/Source/Editor/Sequencer/Private/MVVM/ViewModels/CategoryModel.cpp
@@ -310,10 +310,10 @@ void FCategoryGroupModel::BuildSidebarMenu(FMenuBuilder& MenuBuilder)
 		return;
 	}
 
-	if (const TViewModelPtr<FTrackModel> ParentTrack = FindAncestorOfType<FTrackModel>())
+	if (const TViewModelPtr<ISectionOwnerExtension> SectionOwner = FindAncestorOfType<ISectionOwnerExtension>())
 	{
 		TArray<TWeakObjectPtr<>> WeakSectionObjects;
-		Algo::Transform(ParentTrack->GetSections(), WeakSectionObjects, [](UMovieSceneSection* const InSection)
+		Algo::Transform(SectionOwner->GetSections(), WeakSectionObjects, [](UMovieSceneSection* const InSection)
 			{
 				return InSection;
 			});
diff --git a/Engine/Source/Editor/Sequencer/Private/MVVM/ViewModels/ObjectBindingModel.cpp b/Engine/Source/Editor/Sequencer/Private/MVVM/ViewModels/ObjectBindingModel.cpp
--- a/Engine/Source/Editor/Sequencer/Private/MVVM/ViewModels/ObjectBindingModel.cpp
+++ b/Engine/Source/Editor/Sequencer/Private/MVVM/ViewModels/ObjectBindingModel.cpp
@@ -116,6 +116,11 @@ void FObjectBindingModel::OnConstruct()
 
 	for (UMovieSceneTrack* Track : Binding->GetTracks())
 	{
+		// Don't show child tracks here- they'll be shown underneath their proper outliner parent (e.g. anim mixer)
+		if (Track && Track->ParentTrack != nullptr)
+		{
+			continue;
+		}
 		AddTrack(Track);
 	}
 }
diff --git a/Engine/Source/Editor/Sequencer/Private/MVVM/ViewModels/OutlinerIndicators/TimeWarpOutlinerIndicatorBuilder.cpp b/Engine/Source/Editor/Sequencer/Private/MVVM/ViewModels/OutlinerIndicators/TimeWarpOutlinerIndicatorBuilder.cpp
--- a/Engine/Source/Editor/Sequencer/Private/MVVM/ViewModels/OutlinerIndicators/TimeWarpOutlinerIndicatorBuilder.cpp
+++ b/Engine/Source/Editor/Sequencer/Private/MVVM/ViewModels/OutlinerIndicators/TimeWarpOutlinerIndicatorBuilder.cpp
@@ -3,7 +3,7 @@
 #include "MVVM/ViewModels/OutlinerIndicators/TimeWarpOutlinerIndicatorBuilder.h"
 
 #include "MovieSceneSection.h"
-#include "MVVM/Extensions/ITrackExtension.h"
+#include "MVVM/Extensions/ISectionOwnerExtension.h"
 #include "MVVM/SharedViewModelData.h"
 #include "MVVM/ViewModels/SectionModel.h"
 #include "MVVM/ViewModels/OutlinerColumns/IOutlinerColumn.h"
@@ -29,13 +29,10 @@ FName FTimeWarpOutlinerIndicatorBuilder::GetIndicatorName() const
 
 bool FTimeWarpOutlinerIndicatorBuilder::IsItemCompatibleWithIndicator(const FCreateOutlinerColumnParams& InParams) const
 {
-	if (InParams.OutlinerExtension.AsModel()->IsA<ITrackExtension>())
+	if (TViewModelPtr<ISectionOwnerExtension> SectionOwner = InParams.OutlinerExtension.ImplicitCast())
 	{
-		ITrackExtension* TrackModel = InParams.OutlinerExtension.AsModel()->CastThis<ITrackExtension>();
-
-		for (const TViewModelPtr<FSectionModel>& SectionModel : TrackModel->GetSectionModels().IterateSubList<FSectionModel>())
+		for (UMovieSceneSection* Section : SectionOwner->GetSections())
 		{
-			UMovieSceneSection* Section = SectionModel->GetSection();
 			if (!Section)
 			{
 				continue;
diff --git a/Engine/Source/Editor/Sequencer/Private/MVVM/ViewModels/OutlinerItemModel.cpp b/Engine/Source/Editor/Sequencer/Private/MVVM/ViewModels/OutlinerItemModel.cpp
--- a/Engine/Source/Editor/Sequencer/Private/MVVM/ViewModels/OutlinerItemModel.cpp
+++ b/Engine/Source/Editor/Sequencer/Private/MVVM/ViewModels/OutlinerItemModel.cpp
@@ -6,6 +6,7 @@
 #include "MVVM/ViewModels/FolderModel.h"
 #include "MVVM/Extensions/IDraggableOutlinerExtension.h"
 #include "MVVM/Extensions/ITrackExtension.h"
+#include "MVVM/Extensions/ITrackRowExtension.h"
 #include "MVVM/Extensions/IObjectBindingExtension.h"
 #include "MVVM/ViewModels/OutlinerViewModel.h"
 #include "MVVM/ViewModels/SequencerEditorViewModel.h"
@@ -635,7 +636,7 @@ void FOutlinerItemModelMixin::BuildOrganizeContextMenu(FMenuBuilder& MenuBuilder
 	FSequencer* const SequencerRaw = Sequencer.Get();
 	TSharedRef<FViewModel> ThisNode = AsViewModel()->AsShared();
 
-	const bool bFilterableNode = (ThisNode->IsA<ITrackExtension>() || ThisNode->IsA<IObjectBindingExtension>() || ThisNode->IsA<FFolderModel>());
+	const bool bFilterableNode = (ThisNode->IsA<ITrackExtension>() || ThisNode->IsA<ITrackRowExtension>() || ThisNode->IsA<IObjectBindingExtension>() || ThisNode->IsA<FFolderModel>());
 	const bool bIsReadOnly = Sequencer->IsReadOnly();
 	
 	TArray<UMovieSceneTrack*> AllTracks;
@@ -647,7 +648,16 @@ void FOutlinerItemModelMixin::BuildOrganizeContextMenu(FMenuBuilder& MenuBuilder
 			UMovieSceneTrack* Track = TrackExtension->GetTrack();
 			if (Track)
 			{
-				AllTracks.Add(Track);
+				AllTracks.AddUnique(Track);
+			}
+		}
+
+		if (ITrackRowExtension* TrackRowExtension = Node->CastThis<ITrackRowExtension>())
+		{
+			UMovieSceneTrack* Track = TrackRowExtension->GetParentTrack();
+			if (Track)
+			{
+				AllTracks.AddUnique(Track);
 			}
 		}
 
@@ -851,9 +861,9 @@ TArray<UMovieSceneSection*> FOutlinerItemModelMixin::GetSelectedSections() const
 
 	if (Sections.IsEmpty())
 	{
-		for (const TViewModelPtr<ITrackExtension> TrackExtension : Selection->Outliner.Filter<ITrackExtension>())
+		for (const TViewModelPtr<ISectionOwnerExtension> SectionOwner : Selection->Outliner.Filter<ISectionOwnerExtension>())
 		{
-			for (UMovieSceneSection* const Section : TrackExtension->GetSections())
+			for (UMovieSceneSection* const Section : SectionOwner->GetSections())
 			{
 				Sections.Add(Section);
 			}
@@ -904,12 +914,12 @@ TArray<TPair<UMovieSceneTrack*, int32>> FOutlinerItemModelMixin::GetSelectedTrac
 		return AllTrackRows;
 	}
 
-	for (const TViewModelPtr<ITrackExtension> TrackExtension : Selection->Outliner.Filter<ITrackExtension>())
+	for (const TViewModelPtr<ITrackRowExtension> TrackRowExtension : Selection->Outliner.Filter<ITrackRowExtension>())
 	{
-		UMovieSceneTrack* const Track = TrackExtension->GetTrack();
+		UMovieSceneTrack* const Track = TrackRowExtension->GetParentTrack();
 		if (IsValid(Track))
 		{
-			AllTrackRows.Add(TPair<UMovieSceneTrack*, int32>(Track, TrackExtension->GetRowIndex()));
+			AllTrackRows.Add(TPair<UMovieSceneTrack*, int32>(Track, TrackRowExtension->GetRowIndex()));
 		}
 	}
 
@@ -1163,15 +1173,19 @@ bool FOutlinerItemModelMixin::PassesFilter(const FCurveEditorTreeFilter* InFilte
 
 bool FEvaluableOutlinerItemModel::IsDeactivated() const
 {
-	const TParentFirstChildIterator<ITrackExtension> Descendants = GetDescendantsOfType<ITrackExtension>(true);
-	if (!Descendants)
+	const TParentFirstChildIterator<ITrackExtension> TrackDescendants = GetDescendantsOfType<ITrackExtension>(true);
+	const TParentFirstChildIterator<ITrackRowExtension> TrackRowDescendants = GetDescendantsOfType<ITrackRowExtension>(true);
+	
+	if (!TrackDescendants && !TrackRowDescendants)
 	{
 		return false;
 	}
 
+
+
 	bool bNoTrackAreaModels = true;
 
-	for (const TViewModelPtr<ITrackExtension>& TrackNode : Descendants)
+	for (const TViewModelPtr<ITrackExtension>& TrackNode : TrackDescendants)
 	{
 		const UMovieSceneTrack* const Track = TrackNode->GetTrack();
 		if (!IsValid(Track))
@@ -1187,19 +1201,31 @@ bool FEvaluableOutlinerItemModel::IsDeactivated() const
 
 		bNoTrackAreaModels = false;
 
-		if (const TViewModelPtr<FTrackRowModel> TrackRowModel = TrackNode.ImplicitCast())
+		if (!Track->IsEvalDisabled(/*bInCheckLocal=*/false))
 		{
-			if (!Track->IsRowEvalDisabled(TrackNode->GetRowIndex(), /*bInCheckLocal=*/false))
-			{
-				return false;
-			}
+			return false;
 		}
-		else
+	}
+
+	for (const TViewModelPtr<ITrackRowExtension>& TrackRowNode : TrackRowDescendants)
+	{
+		const UMovieSceneTrack* const Track = TrackRowNode->GetParentTrack();
+		if (!IsValid(Track))
 		{
-			if (!Track->IsEvalDisabled(/*bInCheckLocal=*/false))
-			{
-				return false;
-			}
+			continue;
+		}
+
+		const TViewModelPtr<ITrackAreaExtension> TrackAreaModel = TrackRowNode.ImplicitCast();
+		if (!TrackAreaModel->GetTrackAreaModelList())
+		{
+			continue;
+		}
+
+		bNoTrackAreaModels = false;
+
+		if (!Track->IsRowEvalDisabled(TrackRowNode->GetRowIndex(), /*bInCheckLocal=*/false))
+		{
+			return false;
 		}
 	}
 
@@ -1219,23 +1245,28 @@ void FEvaluableOutlinerItemModel::SetIsDeactivated(const bool bInIsDeactivated)
 		}
 
 		// Deactive state (dirtying, saved with asset, evaluation)
-		if (const TViewModelPtr<FTrackRowModel> TrackRowModel = TrackNode.ImplicitCast())
+		if (bInIsDeactivated != Track->IsEvalDisabled(/*bInCheckLocal=*/false))
 		{
-			if (bInIsDeactivated != Track->IsRowEvalDisabled(TrackNode->GetRowIndex(), /*bInCheckLocal=*/false))
-			{
-				Track->Modify();
-				Track->SetRowEvalDisabled(bInIsDeactivated, TrackNode->GetRowIndex());
-				bAnyChanged = true;
-			}
+			Track->Modify();
+			Track->SetEvalDisabled(bInIsDeactivated);
+			bAnyChanged = true;
 		}
-		else
+	}
+
+	for (const TViewModelPtr<ITrackRowExtension>& TrackRowNode : GetDescendantsOfType<ITrackRowExtension>(true))
+	{
+		UMovieSceneTrack* const Track = TrackRowNode->GetParentTrack();
+		if (!IsValid(Track))
 		{
-			if (bInIsDeactivated != Track->IsEvalDisabled(/*bInCheckLocal=*/false))
-			{
-				Track->Modify();
-				Track->SetEvalDisabled(bInIsDeactivated);
-				bAnyChanged = true;
-			}
+			continue;
+		}
+
+		// Deactive state (dirtying, saved with asset, evaluation)
+		if (bInIsDeactivated != Track->IsRowEvalDisabled(TrackRowNode->GetRowIndex(), /*bInCheckLocal=*/false))
+		{
+			Track->Modify();
+			Track->SetRowEvalDisabled(bInIsDeactivated, TrackRowNode->GetRowIndex());
+			bAnyChanged = true;
 		}
 	}
 
diff --git a/Engine/Source/Editor/Sequencer/Private/MVVM/ViewModels/SectionModel.cpp b/Engine/Source/Editor/Sequencer/Private/MVVM/ViewModels/SectionModel.cpp
--- a/Engine/Source/Editor/Sequencer/Private/MVVM/ViewModels/SectionModel.cpp
+++ b/Engine/Source/Editor/Sequencer/Private/MVVM/ViewModels/SectionModel.cpp
@@ -90,6 +90,11 @@ TSharedPtr<ISequencerSection> FSectionModel::GetSectionInterface() const
 	return SectionInterface;
 }
 
+TViewModelPtr<ISectionOwnerExtension> FSectionModel::GetParentSectionOwnerModel() const
+{
+	return FindAncestorOfType<ISectionOwnerExtension>();
+}
+
 TViewModelPtr<ITrackExtension> FSectionModel::GetParentTrackModel() const
 {
 	return FindAncestorOfType<ITrackExtension>();
diff --git a/Engine/Source/Editor/Sequencer/Private/MVVM/ViewModels/SequencerOutlinerViewModel.cpp b/Engine/Source/Editor/Sequencer/Private/MVVM/ViewModels/SequencerOutlinerViewModel.cpp
--- a/Engine/Source/Editor/Sequencer/Private/MVVM/ViewModels/SequencerOutlinerViewModel.cpp
+++ b/Engine/Source/Editor/Sequencer/Private/MVVM/ViewModels/SequencerOutlinerViewModel.cpp
@@ -7,6 +7,9 @@
 #include "Sequencer.h"
 #include "SequencerKeyCollection.h"
 #include "SequencerOutlinerItemDragDropOp.h"
+#include "MVVM/ViewModels/TrackRowReorderDragDropOp.h"
+#include "MVVM/ViewModels/TrackRowModel.h"
+#include "MVVM/ViewModels/TrackModel.h"
 
 #include "Framework/MultiBox/MultiBoxBuilder.h"
 #include "Framework/Commands/GenericCommands.h"
@@ -89,6 +92,54 @@ void FSequencerOutlinerViewModel::BuildCustomContextMenuForGuid(FMenuBuilder& Me
 
 TSharedRef<FDragDropOperation> FSequencerOutlinerViewModel::InitiateDrag(TArray<TWeakViewModelPtr<IOutlinerExtension>>&& InDraggedModels)
 {
+	// Check if all dragged items are ITrackRowExtensions from the same parent track
+	bool bAllTrackRows = true;
+	TSharedPtr<ITrackExtension> CommonParentTrack;
+
+	for (const TWeakViewModelPtr<IOutlinerExtension>& WeakModel : InDraggedModels)
+	{
+		TViewModelPtr<ITrackRowExtension> RowModel = WeakModel.Pin().ImplicitCast();
+		if (!RowModel)
+		{
+			bAllTrackRows = false;
+			break;
+		}
+
+		TViewModelPtr<ITrackExtension> ParentTrack = RowModel.AsModel()->FindAncestorOfType<ITrackExtension>();
+		if (!ParentTrack)
+		{
+			bAllTrackRows = false;
+			break;
+		}
+
+		// Ensure all rows belong to same track
+		if (!CommonParentTrack.IsValid())
+		{
+			CommonParentTrack = ParentTrack.ImplicitCast();
+		}
+		else if (TSharedPtr<ITrackExtension> OtherParent = ParentTrack.ImplicitCast())
+		{
+			if (CommonParentTrack != OtherParent)
+			{
+				bAllTrackRows = false;
+				break;
+			}
+		}
+		else
+		{
+			bAllTrackRows = false;
+			break;
+		}
+	}
+
+	// If all items are track rows from the same parent, use track row reorder operation
+	if (bAllTrackRows && CommonParentTrack.IsValid())
+	{
+		FText DefaultText = FText::Format(NSLOCTEXT("SequencerOutlinerViewModel", "ReorderTrackRowsFormat", "Reorder {0} track row(s)"), FText::AsNumber(InDraggedModels.Num()));
+		return FTrackRowReorderDragDropOp::New(MoveTemp(InDraggedModels), DefaultText, nullptr);
+	}
+
+	// Otherwise, use generic outliner drag/drop
 	FText DefaultText = FText::Format( NSLOCTEXT( "SequencerOutlinerViewModel", "DefaultDragDropFormat", "Move {0} item(s)" ), FText::AsNumber( InDraggedModels.Num() ) );
 	return FSequencerOutlinerDragDropOp::New( MoveTemp(InDraggedModels), DefaultText, nullptr );
 }
diff --git a/Engine/Source/Editor/Sequencer/Private/MVVM/ViewModels/TrackModel.cpp b/Engine/Source/Editor/Sequencer/Private/MVVM/ViewModels/TrackModel.cpp
--- a/Engine/Source/Editor/Sequencer/Private/MVVM/ViewModels/TrackModel.cpp
+++ b/Engine/Source/Editor/Sequencer/Private/MVVM/ViewModels/TrackModel.cpp
@@ -13,6 +13,7 @@
 #include "MVVM/ViewModels/SequenceModel.h"
 #include "MVVM/ViewModels/TrackModelLayoutBuilder.h"
 #include "MVVM/ViewModels/TrackRowModel.h"
+#include "MVVM/ViewModels/TrackRowReorderDragDropOp.h"
 #include "MVVM/SectionModelStorageExtension.h"
 #include "MVVM/TrackRowModelStorageExtension.h"
 #include "MVVM/Views/SOutlinerTrackView.h"
@@ -89,11 +90,6 @@ UMovieSceneTrack* FTrackModel::GetTrack() const
 	return WeakTrack.Get();
 }
 
-int32 FTrackModel::GetRowIndex() const
-{
-	return 0;
-}
-
 TSharedPtr<ISequencerTrackEditor> FTrackModel::GetTrackEditor() const
 {
 	return TrackEditor;
@@ -287,31 +283,19 @@ void FTrackModel::ForceUpdate()
 		FTrackRowModelStorageExtension* TrackRowModelStorage = SequenceModel->CastDynamic<FTrackRowModelStorageExtension>();
 		check(TrackRowModelStorage);
 
-		// We will build some info about what sections go on what row
-		// Note: the OldSections pointer is just to keep the row section models alive until we re-assign them
-		struct FRowData
-		{
-			TSharedPtr<FTrackRowModel> Row;
-			TSharedPtr<FViewModel> SectionsTail;
-			TUniquePtr<FScopedViewModelListHead> RecycledModels;
-			bool bNeedsLayout = false;
-		};
-		TArray<FRowData, TInlineAllocator<8>> RowModels;
-		RowModels.SetNum(PopulatedRows.Num());
-
 		// Create track row models for all populated rows
 		TSharedPtr<FTrackRowModel> LastTrackRowModel;
 		for (TConstSetBitIterator<> It(PopulatedRows); It; ++It)
 		{
 			const int32 RowIndex = It.GetIndex();
 
-			RowModels[RowIndex].bNeedsLayout = NumRows != PreviousLayoutNumRows;
+			bool bNeedsLayout = NumRows != PreviousLayoutNumRows;
 
 			TSharedPtr<FTrackRowModel> TrackRowModel = TrackRowModelStorage->FindModelForTrackRow(Track, RowIndex);
 			if (!TrackRowModel)
 			{
 				TrackRowModel = TrackRowModelStorage->CreateModelForTrackRow(Track, RowIndex);
-				RowModels[RowIndex].bNeedsLayout = true;
+				bNeedsLayout = true;
 			}
 
 			if (ensure(TrackRowModel))
@@ -319,55 +303,9 @@ void FTrackModel::ForceUpdate()
 				OutlinerChildren.InsertChild(TrackRowModel, LastTrackRowModel);
 				LastTrackRowModel = TrackRowModel;
 
-				RowModels[RowIndex].Row = TrackRowModel;
-
-				// Recycle sections, outliner children, and more, while keeping them alive.
-				RowModels[RowIndex].RecycledModels = MakeUnique<FScopedViewModelListHead>(TrackRowModel, EViewModelListType::Recycled);
-				FViewModelChildren RecycledRowModels = RowModels[RowIndex].RecycledModels->GetChildren();
-				TrackRowModel->GetSectionModels().MoveChildrenTo<IRecyclableExtension>(RecycledRowModels, IRecyclableExtension::CallOnRecycle);
+				TrackRowModel->RefreshLayout(bNeedsLayout);
 			}
 		}
-
-		// Add all sections to both their appropriate track rows and ourselves
-		for (UMovieSceneSection* Section : Track->GetAllSections())
-		{
-			const int32 RowIndex = Section->GetRowIndex();
-
-			TSharedPtr<FSectionModel> SectionModel = SectionModelStorage->FindModelForSection(Section);
-			if (!SectionModel && TrackEditor)
-			{
-				TSharedRef<ISequencerSection> SectionInterface = TrackEditor->MakeSectionInterface(*Section, *Track, ObjectBinding);
-				SectionModel = SectionModelStorage->CreateModelForSection(Section, SectionInterface);
-				RowModels[RowIndex].bNeedsLayout = true;
-			}
-			else if (SectionModel)
-			{
-				RowModels[RowIndex].bNeedsLayout |= SectionModel->NeedsLayout();
-			}
-
-			RowModels[RowIndex].Row->GetSectionModels().InsertChild(SectionModel, RowModels[RowIndex].SectionsTail);
-			RowModels[RowIndex].SectionsTail = SectionModel;
-		}
-
-		// Rebuild the outliner layout for each track row
-		for (FRowData& RowData : RowModels)
-		{
-			const bool bStillHasRecycledChildren = RowData.RecycledModels && !RowData.RecycledModels->GetChildren().IsEmpty();
-			if (RowData.Row && (RowData.bNeedsLayout || bStillHasRecycledChildren))
-			{
-				FViewModelChildren RecycledRowModels = RowData.RecycledModels->GetChildren();
-				RowData.Row->GetChildList(EViewModelListType::Outliner).MoveChildrenTo<IRecyclableExtension>(RecycledRowModels, IRecyclableExtension::CallOnRecycle);
-				RowData.Row->GetTopLevelChannels().MoveChildrenTo<IRecyclableExtension>(RecycledRowModels, IRecyclableExtension::CallOnRecycle);
-
-				FTrackModelLayoutBuilder LayoutBuilder(RowData.Row->AsShared());
-				for (TSharedPtr<FSectionModel> Section : RowData.Row->GetChildrenOfType<FSectionModel>(EViewModelListType::TrackArea))
-				{
-					LayoutBuilder.RefreshLayout(Section);
-				}
-			}
-			// else: unset row... it should only happen while we are dragging sections, until
-			//       we fixup row indices
-		}
 	}
 
 	PreviousLayoutNumRows = NumRows;
@@ -498,7 +436,7 @@ void FTrackModel::CreateCurveModels(TArray<TUniquePtr<FCurveModel>>& OutCurveMod
 
 bool FTrackModel::GetDefaultExpansionState() const
 {
-	TViewModelListIterator<ITrackExtension> It = GetChildrenOfType<ITrackExtension>();
+	TViewModelListIterator<ITrackRowExtension> It = GetChildrenOfType<ITrackRowExtension>();
 	const bool bHasTrackRows = (bool)It;
 	if (bHasTrackRows)
 	{
@@ -757,7 +695,7 @@ EConditionableConditionState FTrackModel::GetConditionState() const
 			// Special case. If we support multiple rows, and there is only a single row, then we must also check track row metadata for a condition here, as there will be no track row model.
 			if (Track->SupportsMultipleRows() && Track->GetMaxRowIndex() == 0)
 			{
-				if (const FMovieSceneTrackRowMetadata* TrackRowMetadata = Track->FindTrackRowMetadata(GetRowIndex()))
+				if (const FMovieSceneTrackRowMetadata* TrackRowMetadata = Track->FindTrackRowMetadata(0))
 				{
 					if (TrackRowMetadata->ConditionContainer.Condition)
 					{
@@ -786,7 +724,7 @@ void FTrackModel::SetConditionEditorForceTrue(bool bEditorForceTrue)
 	UMovieSceneTrack* const Track = GetTrack();
 	if (IsValid(Track))
 	{
-		if (Track->ConditionContainer.Condition)
+		if (Track->ConditionContainer.Condition) 
 		{
 			const FScopedTransaction Transaction(NSLOCTEXT("SequencerTrackNode", "ConditionEditorForceTrue", "Set Condition Editor Force True"));
 			Track->ConditionContainer.Condition->Modify();
@@ -802,6 +740,54 @@ bool FTrackModel::CanDrag() const
 	return ObjectBindingExtension == nullptr;
 }
 
+TOptional<EItemDropZone> FTrackModel::CanAcceptDrop(const FViewModelPtr& TargetModel, const FDragDropEvent& DragDropEvent, EItemDropZone InItemDropZone)
+{
+	// Check if this is a track row reorder operation
+	TSharedPtr<FTrackRowReorderDragDropOp> TrackRowReorderOp = DragDropEvent.GetOperationAs<FTrackRowReorderDragDropOp>();
+	if (TrackRowReorderOp.IsValid())
+	{
+		// Ensure the dragged rows belong to this track
+		TSharedPtr<FTrackModel> DraggedParent = TrackRowReorderOp->GetCommonParentTrack();
+		if (DraggedParent.Get() == this)
+		{
+			// Accept drops above or below track rows for reordering
+			if (InItemDropZone == EItemDropZone::AboveItem || InItemDropZone == EItemDropZone::BelowItem)
+			{
+				return InItemDropZone;
+			}
+		}
+	}
+
+	return TOptional<EItemDropZone>();
+}
+
+void FTrackModel::PerformDrop(const FViewModelPtr& TargetModel, const FDragDropEvent& DragDropEvent, EItemDropZone InItemDropZone)
+{
+	TSharedPtr<FTrackRowReorderDragDropOp> TrackRowReorderOp = DragDropEvent.GetOperationAs<FTrackRowReorderDragDropOp>();
+	if (TrackRowReorderOp.IsValid())
+	{
+		// Ensure the dragged rows belong to this track
+		TSharedPtr<FTrackModel> DraggedParent = TrackRowReorderOp->GetCommonParentTrack();
+		if (DraggedParent.Get() == this && TargetModel)
+		{
+			// Get the target ITrackRowExtension
+			TViewModelPtr<ITrackRowExtension> TargetTrackRowItem = TargetModel.ImplicitCast();
+			if (TargetTrackRowItem)
+			{
+				int32 TargetRowIndex = TargetTrackRowItem->GetRowIndex();
+
+				// Adjust target based on drop zone
+				if (InItemDropZone == EItemDropZone::BelowItem)
+				{
+					TargetRowIndex++;
+				}
+
+				TrackRowReorderOp->ExecuteReorder(TargetRowIndex);
+			}
+		}
+	}
+}
+
 void FTrackModel::BuildContextMenu(FMenuBuilder& MenuBuilder)
 {
 	const TSharedPtr<FSequencerEditorViewModel> EditorViewModel = GetEditor();
@@ -833,7 +819,7 @@ void FTrackModel::BuildContextMenu(FMenuBuilder& MenuBuilder)
 
 	if (Track->GetSupportedBlendTypes().Num() > 0)
 	{
-		SequencerHelpers::BuildNewSectionMenu(Sequencer, GetRowIndex() + 1, GetTrack(), MenuBuilder);
+		SequencerHelpers::BuildNewSectionMenu(Sequencer, 1, GetTrack(), MenuBuilder);
 	}
 
 	SequencerHelpers::BuildBlendingMenu(Sequencer, Track, MenuBuilder);
@@ -881,7 +867,7 @@ void FTrackModel::BuildSidebarMenu(FMenuBuilder& MenuBuilder)
 
 	if (Track->GetSupportedBlendTypes().Num() > 0)
 	{
-		SequencerHelpers::BuildNewSectionMenu(Sequencer, GetRowIndex() + 1, GetTrack(), MenuBuilder);
+		SequencerHelpers::BuildNewSectionMenu(Sequencer, 1, GetTrack(), MenuBuilder);
 	}
 
 	SequencerHelpers::BuildBlendingMenu(Sequencer, Track, MenuBuilder);
@@ -1036,6 +1022,309 @@ bool FTrackModel::FindBoundObjects(TArray<UObject*>& OutBoundObjects) const
 	return true;
 }
 
+int32 FTrackModel::GetMaxRowIndexForSectionDrag(const TSet<UMovieSceneSection*>& ExcludeSections) const
+{
+	int32 MaxRow = INDEX_NONE;
+	UMovieSceneTrack* Track = GetTrack();
+	if (!Track)
+	{
+		return MaxRow;
+	}
+
+	for (UMovieSceneSection* Section : Track->GetAllSections())
+	{
+		if (!ExcludeSections.Contains(Section))
+		{
+			MaxRow = FMath::Max(MaxRow, Section->GetRowIndex() + 1);
+		}
+	}
+	return MaxRow;
+}
+
+void FTrackModel::ShiftRowsDownFromIndex(int32 StartIndex, const TSet<UMovieSceneSection*>& ExcludeSections)
+{
+	UMovieSceneTrack* Track = GetTrack();
+	if (!Track)
+	{
+		return;
+	}
+
+	Track->Modify();
+
+	const TArray<UMovieSceneSection*>& AllSections = Track->GetAllSections();
+	for (UMovieSceneSection* Section : AllSections)
+	{
+		if (!ExcludeSections.Contains(Section) && Section->GetRowIndex() >= StartIndex)
+		{
+			Section->Modify();
+			Section->SetRowIndex(Section->GetRowIndex() + 1);
+		}
+	}
+}
+
+bool FTrackModel::OnSectionVerticalDrag(const FSectionVerticalDragContext& Context)
+{
+	if (!Context.Section || !Context.DraggedSections)
+	{
+		return false;
+	}
+
+	UMovieSceneTrack* Track = GetTrack();
+	if (!Track)
+	{
+		return false;
+	}
+
+	// Don't allow multi-row selection to be moved vertically
+	if (Context.bSectionsAreOnDifferentRows)
+	{
+		return false;
+	}
+
+	int32 CurrentRowIndex = Context.Section->GetRowIndex();
+	int32 TargetRowIndex = CurrentRowIndex;
+	int32 MaxRowIndex = GetMaxRowIndexForSectionDrag(*Context.DraggedSections);
+
+	// Check if we have TrackRowModel children - determines which calculation method to use
+	// (FTrackRowModel children mean the track has multiple rows with separate row models)
+	bool bHasTrackRowChildren = false;
+	for (TSharedPtr<FViewModel> ChildNode : GetChildren())
+	{
+		if (ChildNode->CastThis<FTrackRowModel>())
+		{
+			bHasTrackRowChildren = true;
+			break;
+		}
+	}
+
+	if (bHasTrackRowChildren)
+	{
+		// FTrackRowModel case - use child iteration to find target row based on geometry
+		// This is used when the track has TrackRowModel children (expanded multi-row track)
+
+		// First, find the TrackRowModel that contains the current section
+		// (this is needed for checking non-dragged sections in the current row)
+		TSharedPtr<FTrackRowModel> CurrentSectionRow;
+		for (TSharedPtr<FViewModel> ChildNode : GetChildren())
+		{
+			if (TSharedPtr<FTrackRowModel> TrackRow = ChildNode->CastThisShared<FTrackRowModel>())
+			{
+				// Check if this row contains the section
+				for (TSharedPtr<FSectionModel> SectionNode : TrackRow->GetTrackAreaModelListAs<FSectionModel>())
+				{
+					if (SectionNode->GetSection() == Context.Section)
+					{
+						CurrentSectionRow = TrackRow;
+						break;
+					}
+				}
+				if (CurrentSectionRow)
+				{
+					break;
+				}
+			}
+		}
+
+		int32 ChildIndex = 0;
+		for (TSharedPtr<FViewModel> ChildNode : GetChildren())
+		{
+			FVirtualGeometry ChildVirtualGeometry;
+			if (IGeometryExtension* ChildGeometryExtension = ChildNode->CastThis<IGeometryExtension>())
+			{
+				ChildVirtualGeometry = ChildGeometryExtension->GetVirtualGeometry();
+			}
+
+			float VirtualSectionTop = ChildVirtualGeometry.Top;
+			float VirtualSectionBottom = ChildVirtualGeometry.NestedBottom;
+
+			if (ChildIndex == 0 && (Context.VirtualMousePos.Y <= VirtualSectionTop || Context.LocalMousePos.Y <= 0))
+			{
+				// Mouse is above the first row
+				TargetRowIndex = 0;
+
+				// Check if the current row has non-dragged sections
+				// (if so, we need to shift others down instead of just moving to row 0)
+				if (CurrentSectionRow)
+				{
+					for (TSharedPtr<FSectionModel> SectionNode : CurrentSectionRow->GetTrackAreaModelListAs<FSectionModel>())
+					{
+						if (!Context.DraggedSections->Contains(SectionNode->GetSection()))
+						{
+							TargetRowIndex = -1;
+							break;
+						}
+					}
+				}
+				break;
+			}
+			else if (Context.VirtualMousePos.Y < VirtualSectionBottom)
+			{
+				TargetRowIndex = ChildIndex;
+				break;
+			}
+			else
+			{
+				TargetRowIndex = ChildIndex + 1;
+			}
+
+			++ChildIndex;
+		}
+	}
+	else
+	{
+		// Find the total height of the track - this is necessary because tracks may contain key areas, but they will not use sub tracks unless there is more than one row
+		const int32 NumRows = FMath::Max(CurrentRowIndex + 1, MaxRowIndex);
+		const FVirtualGeometry TrackVirtualGeometry = GetVirtualGeometry();
+		const float VirtualSectionHeight = TrackVirtualGeometry.NestedBottom - TrackVirtualGeometry.Top;
+		const float VirtualRowHeight = VirtualSectionHeight / NumRows;
+
+		const int32 NewIndex = FMath::FloorToInt((Context.VirtualMousePos.Y - TrackVirtualGeometry.Top) / VirtualRowHeight);
+		TargetRowIndex = FMath::Clamp(NewIndex, 0, MaxRowIndex);
+
+		// If close to the top of the track, signal to move everything else down
+		if (Context.VirtualMousePos.Y <= TrackVirtualGeometry.Top || Context.LocalMousePos.Y <= 0)
+		{
+			TargetRowIndex = -1;
+		}
+	}
+
+	// Check if anything changed
+	if (TargetRowIndex == CurrentRowIndex)
+	{
+		return false;
+	}
+
+	// Prevent flickering by only moving sections if the user has actually made an effort to do so
+	int32 TargetRowDelta = TargetRowIndex - CurrentRowIndex;
+	if (TargetRowDelta != 0 && Context.PrevMousePosY.IsSet())
+	{
+		float MouseDeltaY = Context.LocalMousePos.Y - Context.PrevMousePosY.GetValue();
+		// Check mouse has been moved in the direction of intended move
+		if ((TargetRowDelta < 0 && MouseDeltaY > 1.0f) ||
+			(TargetRowDelta > 0 && MouseDeltaY < 1.0f))
+		{
+			// Mouse was not moved in the direction the section wants to swap
+			// Assume offset is due to UI relayout and block moving the section
+			return false;
+		}
+	}
+
+	// Build list of non-dragged sections for overlap checking
+	TArray<UMovieSceneSection*> NonDraggedSections;
+	for (UMovieSceneSection* TrackSection : Track->GetAllSections())
+	{
+		if (!Context.DraggedSections->Contains(TrackSection))
+		{
+			NonDraggedSections.Add(TrackSection);
+		}
+	}
+
+	// Check if vertical movement is allowed (blend type allows overlap, or no overlap would occur)
+	// Note: We pass 0 for horizontal delta since we're only checking vertical movement here
+	if (!Context.Section->GetBlendType().IsValid() &&
+		Context.Section->OverlapsWithSections(NonDraggedSections, TargetRowDelta, 0))
+	{
+		// Section would overlap and doesn't support blending - don't allow move
+		return false;
+	}
+
+	bool bRowIndexChanged = false;
+
+	// Reached the top, move everything else we're not moving downwards
+	if (TargetRowIndex == -1)
+	{
+		// If the sections being moved are all at the top, and all others are below it, do nothing
+		bool bSectionsBeingMovedAreAtTop = true;
+		for (const FInitialRowIndex& InitialRowIndex : *Context.InitialSectionRowIndices)
+		{
+			if (!Context.DraggedSections->Contains(InitialRowIndex.Section))
+			{
+				if (Context.LowestDraggedRow.IsSet() && InitialRowIndex.RowIndex <= Context.LowestDraggedRow.GetValue())
+				{
+					bSectionsBeingMovedAreAtTop = false;
+					break;
+				}
+			}
+		}
+
+		if (!bSectionsBeingMovedAreAtTop)
+		{
+			for (const FInitialRowIndex& InitialRowIndex : *Context.InitialSectionRowIndices)
+			{
+				if (!Context.DraggedSections->Contains(InitialRowIndex.Section))
+				{
+					InitialRowIndex.Section->Modify();
+					InitialRowIndex.Section->SetRowIndex(InitialRowIndex.RowIndex + 1);
+					bRowIndexChanged = true;
+				}
+			}
+		}
+	}
+	else
+	{
+		// But first check if sections being moved are already at the bottom
+		// If the sections being moved are all at the bottom, and all others are above it, do nothing
+		bool bSectionsBeingMovedAreAtBottom = true;
+		for (const FInitialRowIndex& InitialRowIndex : *Context.InitialSectionRowIndices)
+		{
+			if (!Context.DraggedSections->Contains(InitialRowIndex.Section))
+			{
+				if (Context.HighestDraggedRow.IsSet() && InitialRowIndex.RowIndex >= Context.HighestDraggedRow.GetValue())
+				{
+					bSectionsBeingMovedAreAtBottom = false;
+					break;
+				}
+			}
+		}
+
+		if (!bSectionsBeingMovedAreAtBottom || TargetRowIndex < CurrentRowIndex)
+		{
+			Context.Section->Modify();
+			Context.Section->SetRowIndex(TargetRowIndex);
+			bRowIndexChanged = true;
+		}
+	}
+
+	// Handle auto-expansion for tracks with child rows (TrackRowModels)
+	// This matches the original FTrackRowModel expansion logic
+	if (bHasTrackRowChildren)
+	{
+		if (TargetRowIndex > 0)
+		{
+			// Targeting a row > 0 - expand if collapsed
+			if (!IsExpanded() && !bAutoExpandedDuringDrag)
+			{
+				SetExpansion(true);
+				bAutoExpandedDuringDrag = true;
+			}
+		}
+		else if (bAutoExpandedDuringDrag)
+		{
+			// Targeting row 0 or top insertion - collapse if we auto-expanded
+			SetExpansion(false);
+			bAutoExpandedDuringDrag = false;
+		}
+	}
+
+	return bRowIndexChanged;
+}
+
+void FTrackModel::OnBeginSectionVerticalDrag()
+{
+	// Reset auto-expansion tracking at the start of each drag operation
+	bAutoExpandedDuringDrag = false;
+}
+
+void FTrackModel::OnEndSectionVerticalDrag(const FSectionVerticalDragContext& Context)
+{
+	// If we auto-expanded this track during the drag, collapse it when drag ends
+	if (bAutoExpandedDuringDrag)
+	{
+		SetExpansion(false);
+		bAutoExpandedDuringDrag = false;
+	}
+}
+
 } // namespace UE::Sequencer
 
 #undef LOCTEXT_NAMESPACE
diff --git a/Engine/Source/Editor/Sequencer/Private/MVVM/ViewModels/TrackModelLayoutBuilder.cpp b/Engine/Source/Editor/Sequencer/Private/MVVM/ViewModels/TrackModelLayoutBuilder.cpp
--- a/Engine/Source/Editor/Sequencer/Private/MVVM/ViewModels/TrackModelLayoutBuilder.cpp
+++ b/Engine/Source/Editor/Sequencer/Private/MVVM/ViewModels/TrackModelLayoutBuilder.cpp
@@ -1,6 +1,6 @@
 // Copyright Epic Games, Inc. All Rights Reserved.
 
-#include "TrackModelLayoutBuilder.h"
+#include "MVVM/ViewModels/TrackModelLayoutBuilder.h"
 
 #include "Channels/MovieSceneChannelEditorData.h"
 #include "Channels/MovieSceneChannelHandle.h"
diff --git a/Engine/Source/Editor/Sequencer/Private/MVVM/ViewModels/TrackRowModel.cpp b/Engine/Source/Editor/Sequencer/Private/MVVM/ViewModels/TrackRowModel.cpp
--- a/Engine/Source/Editor/Sequencer/Private/MVVM/ViewModels/TrackRowModel.cpp
+++ b/Engine/Source/Editor/Sequencer/Private/MVVM/ViewModels/TrackRowModel.cpp
@@ -2,7 +2,7 @@
 
 #include "MVVM/ViewModels/TrackRowModel.h"
 
-#include "MVVM/Extensions/ITrackExtension.h"
+#include "MVVM/Extensions/ITrackRowExtension.h"
 #include "MVVM/ViewModels/SectionModel.h"
 #include "MVVM/ViewModels/SequenceModel.h"
 #include "MVVM/ViewModels/ViewModelIterators.h"
@@ -12,6 +12,7 @@
 #include "MVVM/Views/SOutlinerTrackView.h"
 #include "MVVM/ViewModels/SequencerEditorViewModel.h"
 #include "MVVM/Selection/Selection.h"
+#include "MVVM/SectionModelStorageExtension.h"
 
 #include "SequencerNodeTree.h"
 #include "ISequencer.h"
@@ -58,7 +59,7 @@ FTrackRowModel::~FTrackRowModel()
 
 void FTrackRowModel::Initialize()
 {
-	UMovieSceneTrack* Track = GetTrack();
+	UMovieSceneTrack* Track = GetParentTrack();
 	TSharedPtr<FSequenceModel> SequenceModel = FindAncestorOfType<FSequenceModel>();
 	if (Track && SequenceModel)
 	{
@@ -73,7 +74,7 @@ FViewModelChildren FTrackRowModel::GetTopLevelChannels()
 	return GetChildrenForList(&TopLevelChannelList);
 }
 
-UMovieSceneTrack* FTrackRowModel::GetTrack() const
+UMovieSceneTrack* FTrackRowModel::GetParentTrack() const
 {
 	return WeakTrack.Get();
 }
@@ -130,7 +131,7 @@ void FTrackRowModel::CreateCurveModels(TArray<TUniquePtr<FCurveModel>>& OutCurve
 
 bool FTrackRowModel::IsDimmed() const
 {
-	UMovieSceneTrack* Track = GetTrack();
+	UMovieSceneTrack* Track = GetParentTrack();
 
 	if (Track)
 	{
@@ -215,7 +216,7 @@ void FTrackRowModel::SetIsLocked(bool bInIsLocked)
 
 const UMovieSceneCondition* FTrackRowModel::GetCondition() const
 {
-	UMovieSceneTrack* Track = GetTrack();
+	UMovieSceneTrack* Track = GetParentTrack();
 
 	if (Track)
 	{
@@ -231,7 +232,7 @@ const UMovieSceneCondition* FTrackRowModel::GetCondition() const
 
 EConditionableConditionState FTrackRowModel::GetConditionState() const
 {
-	UMovieSceneTrack* Track = GetTrack();
+	UMovieSceneTrack* Track = GetParentTrack();
 
 	if (Track)
 	{
@@ -274,7 +275,7 @@ EConditionableConditionState FTrackRowModel::GetConditionState() const
 
 void FTrackRowModel::SetConditionEditorForceTrue(bool bEditorForceTrue)
 {
-	UMovieSceneTrack* Track = GetTrack();
+	UMovieSceneTrack* Track = GetParentTrack();
 
 	if (Track)
 	{
@@ -323,13 +324,13 @@ const FSlateBrush* FTrackRowModel::GetIconBrush() const
 
 FText FTrackRowModel::GetLabel() const
 {
-	UMovieSceneTrack* Track = GetTrack();
+	UMovieSceneTrack* Track = GetParentTrack();
 	return Track ? Track->GetTrackRowDisplayName(RowIndex) : FText::GetEmpty();
 }
 
 FSlateColor FTrackRowModel::GetLabelColor() const
 {
-	UMovieSceneTrack* Track = GetTrack();
+	UMovieSceneTrack* Track = GetParentTrack();
 	TSharedPtr<FSequenceModel> SequenceModel = FindAncestorOfType<FSequenceModel>();
 	if (!Track || !SequenceModel)
 	{
@@ -368,13 +369,13 @@ TSharedPtr<SWidget> FTrackRowModel::CreateOutlinerViewForColumn(const FCreateOut
 
 bool FTrackRowModel::CanRename() const
 {
-	UMovieSceneNameableTrack* NameableTrack = Cast<UMovieSceneNameableTrack>(GetTrack());
+	UMovieSceneNameableTrack* NameableTrack = Cast<UMovieSceneNameableTrack>(GetParentTrack());
 	return NameableTrack && NameableTrack->CanRename();
 }
 
 void FTrackRowModel::Rename(const FText& NewName)
 {
-	UMovieSceneNameableTrack* NameableTrack = ::Cast<UMovieSceneNameableTrack>(GetTrack());
+	UMovieSceneNameableTrack* NameableTrack = ::Cast<UMovieSceneNameableTrack>(GetParentTrack());
 
 	if (NameableTrack && !NameableTrack->GetTrackRowDisplayName(GetRowIndex()).EqualTo(NewName))
 	{
@@ -393,7 +394,7 @@ void FTrackRowModel::Rename(const FText& NewName)
 
 bool FTrackRowModel::IsRenameValidImpl(const FText& NewName, FText& OutErrorMessage) const
 {
-	UMovieSceneNameableTrack* NameableTrack = ::Cast<UMovieSceneNameableTrack>(GetTrack());
+	UMovieSceneNameableTrack* NameableTrack = ::Cast<UMovieSceneNameableTrack>(GetParentTrack());
 	if (NameableTrack)
 	{
 		return NameableTrack->ValidateDisplayName(NewName, OutErrorMessage);
@@ -403,14 +404,14 @@ bool FTrackRowModel::IsRenameValidImpl(const FText& NewName, FText& OutErrorMess
 
 bool FTrackRowModel::IsResizable() const
 {
-	UMovieSceneTrack* Track = GetTrack();
+	UMovieSceneTrack* Track = GetParentTrack();
 	TSharedPtr<ISequencerTrackEditor> TrackEditorPtr = GetTrackEditor();
 	return Track && TrackEditorPtr && TrackEditorPtr->IsResizable(Track);
 }
 
 void FTrackRowModel::Resize(float NewSize)
 {
-	UMovieSceneTrack* Track = GetTrack();
+	UMovieSceneTrack* Track = GetParentTrack();
 	TSharedPtr<ISequencerTrackEditor> TrackEditorPtr = GetTrackEditor();
 	if (Track && TrackEditorPtr && TrackEditorPtr->IsResizable(Track))
 	{
@@ -418,6 +419,68 @@ void FTrackRowModel::Resize(float NewSize)
 	}
 }
 
+void FTrackRowModel::RefreshLayout(bool bChildrenNeedLayout)
+{
+	TSharedPtr<FViewModel> SectionsTail;
+	TUniquePtr<FScopedViewModelListHead> RecycledModels;
+
+	TSharedPtr<FSequenceModel> SequenceModel = FindAncestorOfType<FSequenceModel>();
+	if (!SequenceModel)
+	{
+		// Not part of a full sequence hierarchy yet - wait for OnSetSharedData()
+		return;
+	}
+	FGuid ObjectBinding;
+	if (TSharedPtr<IObjectBindingExtension> ObjectBindingExtension = FindAncestorOfType<IObjectBindingExtension>())
+	{
+		ObjectBinding = ObjectBindingExtension->GetObjectGuid();
+	}
+
+	FSectionModelStorageExtension* SectionModelStorage = SequenceModel->CastDynamic<FSectionModelStorageExtension>();
+	// Recycle sections, outliner children, and more, while keeping them alive.
+	RecycledModels = MakeUnique<FScopedViewModelListHead>(SharedThis(this), EViewModelListType::Recycled);
+	FViewModelChildren RecycledRowModels = RecycledModels->GetChildren();
+
+	GetSectionModels().MoveChildrenTo<IRecyclableExtension>(RecycledRowModels, IRecyclableExtension::CallOnRecycle);
+
+	if (UMovieSceneTrack* Track = WeakTrack.Get())
+	{
+		for (UMovieSceneSection* Section : Track->GetAllSections())
+		{
+			if (Section && Section->GetRowIndex() == RowIndex)
+			{
+				TSharedPtr<FSectionModel> SectionModel = SectionModelStorage->FindModelForSection(Section);
+				if (!SectionModel && TrackEditor)
+				{
+					TSharedRef<ISequencerSection> SectionInterface = TrackEditor->MakeSectionInterface(*Section, *Track, ObjectBinding);
+					SectionModel = SectionModelStorage->CreateModelForSection(Section, SectionInterface);
+					bChildrenNeedLayout = true;
+				}
+				else if (SectionModel)
+				{
+					bChildrenNeedLayout |= SectionModel->NeedsLayout();
+				}
+
+				GetSectionModels().InsertChild(SectionModel, SectionsTail);
+				SectionsTail = SectionModel;
+			}
+		}
+	}
+	
+	const bool bStillHasRecycledChildren = RecycledModels && !RecycledModels->GetChildren().IsEmpty();
+	if (bChildrenNeedLayout || bStillHasRecycledChildren)
+	{
+		GetChildList(EViewModelListType::Outliner).MoveChildrenTo<IRecyclableExtension>(RecycledRowModels, IRecyclableExtension::CallOnRecycle);
+		GetTopLevelChannels().MoveChildrenTo<IRecyclableExtension>(RecycledRowModels, IRecyclableExtension::CallOnRecycle);
+
+		FTrackModelLayoutBuilder LayoutBuilder(AsShared());
+		for (TSharedPtr<FSectionModel> Section : GetChildrenOfType<FSectionModel>(EViewModelListType::TrackArea))
+		{
+			LayoutBuilder.RefreshLayout(Section);
+		}
+	}
+}
+
 void FTrackRowModel::BuildContextMenu(FMenuBuilder& MenuBuilder)
 {
 	const TSharedPtr<FSequencerEditorViewModel> EditorViewModel = GetEditor();
@@ -432,7 +495,7 @@ void FTrackRowModel::BuildContextMenu(FMenuBuilder& MenuBuilder)
 		return;
 	}
 
-	UMovieSceneTrack* const Track = GetTrack();
+	UMovieSceneTrack* const Track = GetParentTrack();
 	if (!IsValid(Track))
 	{
 		return;
@@ -445,7 +508,7 @@ void FTrackRowModel::BuildContextMenu(FMenuBuilder& MenuBuilder)
 
 	if (Track->GetSupportedBlendTypes().Num() > 0)
 	{
-		SequencerHelpers::BuildNewSectionMenu(Sequencer, GetRowIndex() + 1, GetTrack(), MenuBuilder);
+		SequencerHelpers::BuildNewSectionMenu(Sequencer, GetRowIndex() + 1, GetParentTrack(), MenuBuilder);
 	}
 
 	SequencerHelpers::BuildBlendingMenu(Sequencer, Track, MenuBuilder);
@@ -479,7 +542,7 @@ void FTrackRowModel::BuildSidebarMenu(FMenuBuilder& MenuBuilder)
 		return;
 	}
 
-	UMovieSceneTrack* const Track = GetTrack();
+	UMovieSceneTrack* const Track = GetParentTrack();
 	if (!IsValid(Track))
 	{
 		return;
@@ -496,7 +559,7 @@ void FTrackRowModel::BuildSidebarMenu(FMenuBuilder& MenuBuilder)
 
 	if (Track->GetSupportedBlendTypes().Num() > 0)
 	{
-		SequencerHelpers::BuildNewSectionMenu(Sequencer, GetRowIndex() + 1, GetTrack(), MenuBuilder);
+		SequencerHelpers::BuildNewSectionMenu(Sequencer, GetRowIndex() + 1, GetParentTrack(), MenuBuilder);
 	}
 
 	SequencerHelpers::BuildBlendingMenu(Sequencer, Track, MenuBuilder);
@@ -519,7 +582,7 @@ bool FTrackRowModel::CanDelete(FText* OutErrorMessage) const
 
 void FTrackRowModel::Delete()
 {
-	UMovieSceneTrack* Track = GetTrack();
+	UMovieSceneTrack* Track = GetParentTrack();
 	if (!Track)
 	{
 		return;
@@ -550,7 +613,7 @@ void FTrackRowModel::Delete()
 }
 bool FTrackRowModel::IsMuted() const
 {
-	UMovieSceneTrack* Track = GetTrack();
+	UMovieSceneTrack* Track = GetParentTrack();
 	if (!Track)
 	{
 		return false;
@@ -566,7 +629,7 @@ bool FTrackRowModel::IsMuted() const
 
 void FTrackRowModel::SetIsMuted(bool bIsMuted)
 {
-	UMovieSceneTrack* Track = GetTrack();
+	UMovieSceneTrack* Track = GetParentTrack();
 	if (!Track)
 	{
 		return;
@@ -583,7 +646,7 @@ void FTrackRowModel::SetIsMuted(bool bIsMuted)
 
 bool FTrackRowModel::IsSolo() const
 {
-	UMovieSceneTrack* Track = GetTrack();
+	UMovieSceneTrack* Track = GetParentTrack();
 	if (!Track)
 	{
 		return false;
@@ -599,7 +662,7 @@ bool FTrackRowModel::IsSolo() const
 
 void FTrackRowModel::SetIsSoloed(bool bIsSoloed)
 {
-	UMovieSceneTrack* Track = GetTrack();
+	UMovieSceneTrack* Track = GetParentTrack();
 	if (!Track)
 	{
 		return;
@@ -614,6 +677,50 @@ void FTrackRowModel::SetIsSoloed(bool bIsSoloed)
 	}
 }
 
+bool FTrackRowModel::SetRowIndex(int32 NewRowIndex)
+{
+	if (RowIndex == NewRowIndex)
+	{
+		return false;
+	}
+
+	UMovieSceneTrack* Track = GetParentTrack();
+	if (!Track)
+	{
+		return false;
+	}
+
+	Track->Modify();
+
+	// Update all sections on this row to the new row index
+	bool bAnyChanged = false;
+	for (TSharedPtr<FSectionModel> SectionModel : SectionList.Iterate<FSectionModel>())
+	{
+		if (UMovieSceneSection* Section = SectionModel->GetSection())
+		{
+			if (Section->GetRowIndex() == RowIndex)
+			{
+				Section->SetRowIndex(NewRowIndex);
+				bAnyChanged = true;
+			}
+		}
+	}
+
+	if (bAnyChanged)
+	{
+		// Update our cached row index
+		RowIndex = NewRowIndex;
+	}
+
+	return bAnyChanged;
+}
+
+bool FTrackRowModel::CanDrag() const
+{
+	// For now, regular track rows cannot be dragged in the outliner
+	return false;
+}
+
 } // namespace UE::Sequencer
 
 #undef LOCTEXT_NAMESPACE
diff --git a/Engine/Source/Editor/Sequencer/Private/MVVM/Views/SChannelView.cpp b/Engine/Source/Editor/Sequencer/Private/MVVM/Views/SChannelView.cpp
--- a/Engine/Source/Editor/Sequencer/Private/MVVM/Views/SChannelView.cpp
+++ b/Engine/Source/Editor/Sequencer/Private/MVVM/Views/SChannelView.cpp
@@ -753,7 +753,7 @@ void SChannelView::CreateKeysUnderMouse(const FVector2D& MousePosition, const FG
 	}
 	else
 	{
-		TSharedPtr<FTrackModel>             TrackModel             = Section->FindAncestorOfType<FTrackModel>();
+		TViewModelPtr<ITrackExtension>      TrackModel             = Section->FindAncestorOfType<ITrackExtension>();
 		TSharedPtr<IObjectBindingExtension> ObjectBindingExtension = Section->FindAncestorOfType<IObjectBindingExtension>();
 
 		FVector2D LocalSpaceMousePosition = AllottedGeometry.AbsoluteToLocal(MousePosition);
diff --git a/Engine/Source/Editor/Sequencer/Private/MVVM/Views/SSequencerTrackAreaView.cpp b/Engine/Source/Editor/Sequencer/Private/MVVM/Views/SSequencerTrackAreaView.cpp
--- a/Engine/Source/Editor/Sequencer/Private/MVVM/Views/SSequencerTrackAreaView.cpp
+++ b/Engine/Source/Editor/Sequencer/Private/MVVM/Views/SSequencerTrackAreaView.cpp
@@ -3,6 +3,7 @@
 #include "MVVM/Views/SSequencerTrackAreaView.h"
 #include "MVVM/Extensions/IObjectBindingExtension.h"
 #include "MVVM/Extensions/ITrackExtension.h"
+#include "MVVM/Extensions/ITrackRowExtension.h"
 #include "MVVM/ViewModels/OutlinerSpacer.h"
 #include "MVVM/ViewModels/SequencerTrackAreaViewModel.h"
 #include "MVVM/Views/SOutlinerView.h"
@@ -102,10 +103,12 @@ FReply SSequencerTrackAreaView::OnDragOver(const FGeometry& MyGeometry, const FD
 
 	if (Sequencer && DroppedItem)
 	{
-		if (TSharedPtr<ITrackExtension> TrackDestination = DroppedItem.ImplicitCast())
+		TSharedPtr<ITrackExtension> TrackDestination = DroppedItem.ImplicitCast();
+		TSharedPtr<ITrackRowExtension> TrackRowDestination = DroppedItem.ImplicitCast();
+		if (TrackDestination || TrackRowDestination)
 		{
-			UMovieSceneTrack* Track    = TrackDestination->GetTrack();
-			const int32       TrackRow = TrackDestination->GetRowIndex();
+			UMovieSceneTrack* Track    = TrackDestination ? TrackDestination->GetTrack() : TrackRowDestination->GetParentTrack();
+			const int32       TrackRow = TrackRowDestination ? TrackRowDestination->GetRowIndex() : 0;
 
 			TSharedPtr<IObjectBindingExtension> BindingExtension = DroppedItem.AsModel()->FindAncestorOfType<IObjectBindingExtension>();
 			FGuid ObjectBinding = BindingExtension ? BindingExtension->GetObjectGuid() : FGuid();
@@ -155,10 +158,11 @@ FReply SSequencerTrackAreaView::OnDrop(const FGeometry& MyGeometry, const FDragD
 	if (Sequencer && DroppedItem)
 	{
 		ITrackExtension* TrackDestination = DroppedItem->CastThis<ITrackExtension>();
-		if (TrackDestination)
+		ITrackRowExtension* TrackRowDestination = DroppedItem->CastThis<ITrackRowExtension>();
+		if (TrackDestination || TrackRowDestination)
 		{
-			UMovieSceneTrack* Track    = TrackDestination->GetTrack();
-			const int32       TrackRow = TrackDestination->GetRowIndex();
+			UMovieSceneTrack* Track    = TrackDestination ? TrackDestination->GetTrack() : TrackRowDestination->GetParentTrack();
+			const int32       TrackRow = TrackRowDestination ? TrackRowDestination->GetRowIndex() : 0;
 
 			TSharedPtr<IObjectBindingExtension> BindingExtension = DroppedItem->FindAncestorOfType<IObjectBindingExtension>();
 			FGuid ObjectBinding = BindingExtension ? BindingExtension->GetObjectGuid() : FGuid();
