Sequencer: Large anim mixer update and Sequencer MVVM refactor.

diff --git a/Engine/Source/Editor/Sequencer/Private/MovieSceneTrackEditor.cpp b/Engine/Source/Editor/Sequencer/Private/MovieSceneTrackEditor.cpp
--- a/Engine/Source/Editor/Sequencer/Private/MovieSceneTrackEditor.cpp
+++ b/Engine/Source/Editor/Sequencer/Private/MovieSceneTrackEditor.cpp
@@ -20,6 +20,7 @@
 #include "MVVM/Views/SSequencerKeyNavigationButtons.h"
 #include "MVVM/Views/SOutlinerTrackColorPicker.h"
 #include "MVVM/Views/SOutlinerItemViewBase.h"
+#include "MVVM/Extensions/ITrackRowExtension.h"
 #include "SKeyAreaEditorSwitcher.h"
 
 namespace UE::Sequencer
@@ -241,7 +242,10 @@ TSharedPtr<SWidget> FMovieSceneTrackEditor::BuildOutlinerColumnWidget(const FBui
 {
 	using namespace UE::Sequencer;
 
-	UMovieSceneTrack* Track = Params.TrackModel->GetTrack();
+	TViewModelPtr<ITrackExtension> TrackModel = Params.ViewModel->FindAncestorOfType<ITrackExtension>(true);
+	TViewModelPtr<ITrackRowExtension> TrackRowModel = Params.ViewModel.ImplicitCast();
+	UMovieSceneTrack* Track = TrackModel ? TrackModel->GetTrack() : nullptr;
+
 	TViewModelPtr<FSequencerEditorViewModel> Editor       = Params.Editor->CastThisShared<FSequencerEditorViewModel>();
 	TViewModelPtr<IOutlinerExtension>        OutlinerItem = Params.ViewModel.ImplicitCast();
 	if (!Track || !Editor || !OutlinerItem)
@@ -267,10 +271,10 @@ TSharedPtr<SWidget> FMovieSceneTrackEditor::BuildOutlinerColumnWidget(const FBui
 
 		FBuildEditWidgetParams EditWidgetParams(Params);
 		EditWidgetParams.NodeIsHovered = IsHovered;
-		EditWidgetParams.RowIndex = Params.TrackModel->GetRowIndex();
+		EditWidgetParams.RowIndex = TrackRowModel ? TrackRowModel->GetRowIndex() : 0;
 		if (Params.ViewModel->FindAncestorOfType<ITrackExtension>())
 		{
-			EditWidgetParams.TrackInsertRowIndex = Params.TrackModel->GetRowIndex();
+			EditWidgetParams.TrackInsertRowIndex = TrackRowModel ? TrackRowModel->GetRowIndex() : 0;
 		}
 		else if (Track && Track->SupportsMultipleRows())
 		{
@@ -287,10 +291,10 @@ TSharedPtr<SWidget> FMovieSceneTrackEditor::BuildOutlinerColumnWidget(const FBui
 		// This function is a misnomer but is called this way for backwards compatibility
 		FBuildEditWidgetParams EditWidgetParams(Params);
 		EditWidgetParams.NodeIsHovered = IsHovered;
-		EditWidgetParams.RowIndex = Params.TrackModel->GetRowIndex();
+		EditWidgetParams.RowIndex = TrackRowModel ? TrackRowModel->GetRowIndex() : 0;
 		if (Params.ViewModel->FindAncestorOfType<ITrackExtension>())
 		{
-			EditWidgetParams.TrackInsertRowIndex = Params.TrackModel->GetRowIndex();
+			EditWidgetParams.TrackInsertRowIndex = TrackRowModel ? TrackRowModel->GetRowIndex() : 0;
 		}
 		else if (Track && Track->SupportsMultipleRows())
 		{
diff --git a/Engine/Source/Editor/Sequencer/Private/SSequencer.cpp b/Engine/Source/Editor/Sequencer/Private/SSequencer.cpp
--- a/Engine/Source/Editor/Sequencer/Private/SSequencer.cpp
+++ b/Engine/Source/Editor/Sequencer/Private/SSequencer.cpp
@@ -2745,8 +2745,7 @@ void SSequencer::UpdateLayoutTree()
 
 			bool bAnyChanged = false;
 
-			for (TViewModelPtr<ITrackExtension> TrackNode
-				: Sequencer->GetNodeTree()->GetRootNode()->GetDescendantsOfType<ITrackExtension>())
+			for (TViewModelPtr<ITrackExtension> TrackNode : Sequencer->GetNodeTree()->GetRootNode()->GetDescendantsOfType<ITrackExtension>())
 			{
 				UMovieSceneTrack* const Track = TrackNode->GetTrack();
 				if (!IsValid(Track))
@@ -2762,35 +2761,47 @@ void SSequencer::UpdateLayoutTree()
 
 				const bool bLocalEvalDisabled = bIsSoloed ? false : (bIsMuted || (bAnySoloNodes && !bIsSoloed));
 
-				if (const TViewModelPtr<FTrackRowModel> TrackRowModel = TrackNode.ImplicitCast())
-				{
-					const int32 TrackRowIndex = TrackNode->GetRowIndex();
-					const bool bCurrentLocalRowEvalDisabled = Track->IsLocalRowEvalDisabled(TrackRowIndex);
+				const bool bCurrentLocalEvalDisabled = Track->IsLocalEvalDisabled();
 
-					if (bLocalEvalDisabled != bCurrentLocalRowEvalDisabled
-						|| (bIsSoloed && (bIsSoloed != bLocalEvalDisabled)))
-					{
-						Track->MarkAsChanged();
-						if (!bLocalEvalDisabled)
-						{
-							Track->SetLocalEvalDisabled(false);
-						}
-						Track->SetLocalRowEvalDisabled(bLocalEvalDisabled, TrackRowIndex);
+				if (bLocalEvalDisabled != bCurrentLocalEvalDisabled)
+				{
+					Track->MarkAsChanged();
+					Track->SetLocalEvalDisabled(bLocalEvalDisabled);
 
-						bAnyChanged = true;
-					}
+					bAnyChanged = true;
 				}
-				else
+			}
+
+			for (TViewModelPtr<ITrackRowExtension> TrackRowNode : Sequencer->GetNodeTree()->GetRootNode()->GetDescendantsOfType<ITrackRowExtension>())
+			{
+				UMovieSceneTrack* const Track = TrackRowNode->GetParentTrack();
+				if (!IsValid(Track))
 				{
-					const bool bCurrentLocalEvalDisabled = Track->IsLocalEvalDisabled();
+					continue;
+				}
 
-					if (bLocalEvalDisabled != bCurrentLocalEvalDisabled)
-					{
-						Track->MarkAsChanged();
-						Track->SetLocalEvalDisabled(bLocalEvalDisabled);
+				const ECachedMuteState MuteFlags = MuteState->GetCachedFlags(TrackRowNode);
+				const ECachedSoloState SoloFlags = SoloState->GetCachedFlags(TrackRowNode);
+
+				const bool bIsMuted = EnumHasAnyFlags(MuteFlags, ECachedMuteState::Muted | ECachedMuteState::ImplicitlyMutedByParent);
+				const bool bIsSoloed = EnumHasAnyFlags(SoloFlags, ECachedSoloState::Soloed | ECachedSoloState::ImplicitlySoloedByParent);
 
-						bAnyChanged = true;
+				const bool bLocalEvalDisabled = bIsSoloed ? false : (bIsMuted || (bAnySoloNodes && !bIsSoloed));
+
+				const int32 TrackRowIndex = TrackRowNode->GetRowIndex();
+				const bool bCurrentLocalRowEvalDisabled = Track->IsLocalRowEvalDisabled(TrackRowIndex);
+
+				if (bLocalEvalDisabled != bCurrentLocalRowEvalDisabled
+					|| (bIsSoloed && (bIsSoloed != bLocalEvalDisabled)))
+				{
+					Track->MarkAsChanged();
+					if (!bLocalEvalDisabled)
+					{
+						Track->SetLocalEvalDisabled(false);
 					}
+					Track->SetLocalRowEvalDisabled(bLocalEvalDisabled, TrackRowIndex);
+
+					bAnyChanged = true;
 				}
 			}
 
@@ -3638,7 +3649,7 @@ FPasteContextMenuArgs SSequencer::GeneratePasteArgs(FFrameNumber PasteAtTime, TS
 		{
 			if (TSharedPtr<FSectionModel> Handle = Sequencer->GetNodeTree()->GetSectionModel(Section))
 			{
-				TViewModelPtr<IOutlinerExtension> TrackModel = Handle->GetParentTrackModel().ImplicitCast();
+				TViewModelPtr<IOutlinerExtension> TrackModel = Handle->GetParentSectionOwnerModel().ImplicitCast();
 				if (TrackModel)
 				{
 					PasteIntoNodes.Add(TrackModel);
@@ -3648,7 +3659,7 @@ FPasteContextMenuArgs SSequencer::GeneratePasteArgs(FFrameNumber PasteAtTime, TS
 
 		for (FViewModelPtr SelectedNode : Sequencer->GetViewModel()->GetSelection()->Outliner)
 		{
-			if (SelectedNode->IsA<FCategoryGroupModel>() || SelectedNode->IsA<ITrackExtension>() || SelectedNode->IsA<FChannelGroupModel>())
+			if (SelectedNode->IsA<FCategoryGroupModel>() || SelectedNode->IsA<ITrackExtension>() || SelectedNode->IsA<ITrackRowExtension>() || SelectedNode->IsA<FChannelGroupModel>())
 			{
 				TViewModelPtr<IOutlinerExtension> TrackModel = CastViewModelChecked<IOutlinerExtension>(SelectedNode);
 				if (TrackModel)
diff --git a/Engine/Source/Editor/Sequencer/Private/Sequencer.cpp b/Engine/Source/Editor/Sequencer/Private/Sequencer.cpp
--- a/Engine/Source/Editor/Sequencer/Private/Sequencer.cpp
+++ b/Engine/Source/Editor/Sequencer/Private/Sequencer.cpp
@@ -6667,6 +6667,14 @@ void FSequencer::HandleSelectedOutlinerNodesChanged()
 		{
 			SelectedTracks.Add(Track);
 		}
+	}	
+	
+	for (TViewModelPtr<ITrackRowExtension> TrackRowModel : ViewModel->GetSelection()->Outliner.Filter<ITrackRowExtension>())
+	{
+		if (UMovieSceneTrack* Track = TrackRowModel->GetParentTrack())
+		{
+			SelectedTracks.Add(Track);
+		}
 	}
 
 	TArray<UMovieSceneSection*> SelectedSections;
@@ -7405,7 +7413,7 @@ void FSequencer::GetSelectedKeyAreas(TArray<const IKeyArea*>& OutSelectedKeyArea
 			NodesToKey.Add(Node.AsModel().ToSharedRef());
 
 			// No need to gather key areas from binding/tracks because they have no key areas
-			if (Node->IsA<IObjectBindingExtension>() || Node->IsA<ITrackExtension>() || Node->IsA<FFolderModel>())
+			if (Node->IsA<IObjectBindingExtension>() || Node->IsA<ITrackExtension>() || Node->IsA<ITrackRowExtension>() || Node->IsA<FFolderModel>())
 			{
 				continue;
 			}
@@ -7426,7 +7434,7 @@ void FSequencer::GetSelectedKeyAreas(TArray<const IKeyArea*>& OutSelectedKeyArea
 	for (TSharedRef<FViewModel> Node : NodesToKey)
 	{
 		//if object or track selected we don't want all of the children only if spefically selected.
-		if (!Node->IsA<ITrackExtension>() && !Node->IsA<IObjectBindingExtension>() && !Node->IsA<FFolderModel>())
+		if (!Node->IsA<ITrackExtension>() && !Node->IsA<ITrackRowExtension>() && !Node->IsA<IObjectBindingExtension>() && !Node->IsA<FFolderModel>())
 		{
 			SequencerHelpers::GetAllKeyAreas(Node, KeyAreas);
 		}
@@ -7464,7 +7472,7 @@ void FSequencer::SelectByNthCategoryNode(UMovieSceneSection* Section, int Index,
 	int32 Count = 0;
 	if (SectionHandle)
 	{
-		TSharedPtr<FViewModel> TrackNode = SectionHandle->GetParentTrackModel();
+		TSharedPtr<FViewModel> TrackNode = SectionHandle->GetParentSectionOwnerModel();
 		for (TViewModelPtr<FCategoryGroupModel> Node : TrackNode->GetChildrenOfType<FCategoryGroupModel>())
 		{
 			if (Count++ == Index)
@@ -7495,7 +7503,9 @@ void FSequencer::SelectByNthCategoryNode(UMovieSceneSection* Section, int Index,
 			for (FViewModelPtr DisplayNode : NodesToSelect)
 			{
 				FViewModelPtr Parent = DisplayNode->GetParent();
-				if (TViewModelPtr<ITrackExtension> ParentTrack = Parent.ImplicitCast())
+				TViewModelPtr<ITrackExtension> ParentTrack = Parent.ImplicitCast();
+				TViewModelPtr<ITrackRowExtension> ParentTrackRow = Parent.ImplicitCast();
+				if (ParentTrack || ParentTrackRow)
 				{
 					if (TViewModelPtr<IOutlinerExtension> ParentOutliner = Parent.ImplicitCast())
 					{
@@ -7535,7 +7545,7 @@ void FSequencer::SelectByChannels(UMovieSceneSection* Section, TArrayView<const
 	TSharedPtr<FSectionModel> SectionHandle = NodeTree->GetSectionModel(Section);
 	if (SectionHandle)
 	{
-		TSharedPtr<FViewModel> TrackNode = SectionHandle->GetParentTrackModel();
+		TSharedPtr<FViewModel> TrackNode = SectionHandle->GetParentSectionOwnerModel();
 		for (TSharedPtr<FChannelGroupModel> KeyAreaNode : TrackNode->GetDescendantsOfType<FChannelGroupModel>())
 		{
 			for (TSharedPtr<IKeyArea> KeyArea : KeyAreaNode->GetAllKeyAreas())
@@ -7607,7 +7617,7 @@ void FSequencer::SelectByChannels(UMovieSceneSection* Section, const TArray<FNam
 	TSharedPtr<FSectionModel> SectionHandle = NodeTree->GetSectionModel(Section);
 	if (SectionHandle)
 	{
-		TSharedPtr<FViewModel> TrackNode = SectionHandle->GetParentTrackModel();
+		TSharedPtr<FViewModel> TrackNode = SectionHandle->GetParentSectionOwnerModel();
 		for (TSharedPtr<FChannelGroupOutlinerModel> KeyAreaNode : TrackNode->GetDescendantsOfType<FChannelGroupOutlinerModel>())
 		{
 			TViewModelPtr<IOutlinerExtension> Parent = KeyAreaNode->FindAncestorOfType<IOutlinerExtension>();
@@ -7647,7 +7657,7 @@ void FSequencer::SelectByChannels(UMovieSceneSection* Section, const TArray<FNam
 			if (Settings->GetAutoExpandNodesOnSelection())
 			{
 				TViewModelPtr<IOutlinerExtension> Parent = DisplayNode.AsModel()->FindAncestorOfType<IOutlinerExtension>();
-				if (Parent && Parent.AsModel()->IsA<ITrackExtension>() && !Parent->IsExpanded())
+				if (Parent && (Parent.AsModel()->IsA<ITrackExtension>() || Parent.AsModel()->IsA<ITrackRowExtension>()) && !Parent->IsExpanded())
 				{
 					Parent->SetExpansion(true);
 				}
@@ -7700,14 +7710,14 @@ void FSequencer::DeselectObject(FGuid ObjectBinding)
 void FSequencer::SelectTrack(UMovieSceneTrack* Track)
 {
 	using namespace UE::Sequencer;
-	
+
 	FTrackModelStorageExtension* TrackStorage = ViewModel->GetRootModel()->CastDynamic<FTrackModelStorageExtension>();
 	check(TrackStorage);
 
-	TSharedPtr<FTrackModel> TrackModel = TrackStorage->FindModelForTrack(Track);
-	if (TrackModel)
+	TViewModelPtr<ITrackExtension> TrackExtension = TrackStorage->FindTrackExtensionForTrack(Track);
+	if (TrackExtension)
 	{
-		ViewModel->GetSelection()->Outliner.Select(TrackModel);
+		ViewModel->GetSelection()->Outliner.Select(TrackExtension.ImplicitCast());
 	}
 }
 
@@ -8055,6 +8065,16 @@ TArray<TSharedRef<UE::Sequencer::FViewModel> > FSequencer::GetSelectedNodesInFol
 					}
 				}
 			}
+			else if (ITrackRowExtension* TrackRowNode = SelectedNode->CastThis<ITrackRowExtension>())
+			{
+				if (TrackRowNode->GetParentTrack())
+				{
+					if (Folder->GetFolder()->GetChildTracks().Contains(TrackRowNode->GetParentTrack()))
+					{
+						NodesToFolders.Add(SelectedNode.AsModel().ToSharedRef());
+					}
+				}
+			}
 		}
 	}
 
@@ -8449,6 +8469,13 @@ void FSequencer::RemoveSelectedNodesFromFolders()
 					Folder->GetFolder()->RemoveChildTrack(TrackNode->GetTrack());
 				}
 			}
+			else if (ITrackRowExtension* TrackRowNode = NodeInFolder->CastThis<ITrackRowExtension>())
+			{
+				if (TrackRowNode->GetParentTrack())
+				{
+					Folder->GetFolder()->RemoveChildTrack(TrackRowNode->GetParentTrack());
+				}
+			}
 		}
 	}
 
@@ -8509,19 +8536,17 @@ void FSequencer::CopySelectedTracks(TArray<TSharedPtr<FViewModel>>& TrackNodes,
 
 		if (!bIsParentSelected)
 		{
-			// If this is a subtrack, only copy the sections that belong to this row. otherwise copying the entire track will copy all the sections across all the rows
-			if (FTrackRowModel* TrackRowNode = TrackNode->CastThis<FTrackRowModel>())
+			// If this is a track row, only copy the sections that belong to this row. otherwise copying the entire track will copy all the sections across all the rows
+			if (ITrackRowExtension* TrackRowExtension = TrackNode->CastThis<ITrackRowExtension>())
 			{
-				for (UMovieSceneSection* Section : TrackRowNode->GetTrack()->GetAllSections())
+				if (ISectionOwnerExtension* SectionOwner = TrackNode->CastThis<ISectionOwnerExtension>())
 				{
-					if (Section && Section->GetRowIndex() == TrackRowNode->GetRowIndex())
-					{
-						CopyableSections.Add(Section);
-					}
+					TArray<UMovieSceneSection*> Sections = SectionOwner->GetSections();
+					CopyableSections.Append(Sections);
 				}
 			}
 			else if (UMovieSceneTrack* Track = TrackNode->CastThisChecked<ITrackExtension>()->GetTrack())
-			{	
+			{
 				CopyableTracks.Add(Track);
 			}
 		}
@@ -8795,10 +8820,15 @@ FSequencer::EPasteResult FSequencer::PasteSections(const FString& TextToImport,
 			if (TrackForNode)
 			{
 				Tracks.Add(TrackForNode);
-				if (FTrackRowModel* TrackRowNode = Node->CastThis<FTrackRowModel>())
-				{
-					TrackRowIndices.Add(TrackRowNode->GetRowIndex());
-				}
+			}
+		}
+		else if (ITrackRowExtension* TrackRowNode = Node->CastThis<ITrackRowExtension>())
+		{
+			UMovieSceneTrack* TrackForNode = TrackRowNode->GetParentTrack();
+			if (TrackForNode)
+			{
+				Tracks.Add(TrackForNode);
+				TrackRowIndices.Add(TrackRowNode->GetRowIndex());
 			}
 		}
 	}
@@ -8813,8 +8843,8 @@ FSequencer::EPasteResult FSequencer::PasteSections(const FString& TextToImport,
 
 			for (TSharedRef<FViewModel> DescendantNode : DescendantNodes)
 			{
-				// Don't automatically paste onto subtracks because that would lead to multiple paste destinations
-				if (DescendantNode->IsA<FTrackRowModel>())
+				// Don't automatically paste onto row-level models because that would lead to multiple paste destinations
+				if (ITrackRowExtension* DescendantTrackRowNode = DescendantNode->CastThis<ITrackRowExtension>())
 				{
 					continue;
 				}
@@ -8825,10 +8855,11 @@ FSequencer::EPasteResult FSequencer::PasteSections(const FString& TextToImport,
 					if (TrackForNode)
 					{
 						Tracks.Add(TrackForNode);
-						if (FTrackRowModel* TrackRowNode = Node->CastThis<FTrackRowModel>())
-						{
-							TrackRowIndices.Add(TrackRowNode->GetRowIndex());
-						}
+					}
+
+					if (ITrackRowExtension* TrackRowNode = Node->CastThis<ITrackRowExtension>())
+					{
+						TrackRowIndices.Add(TrackRowNode->GetRowIndex());
 					}
 				}
 			}
@@ -10656,6 +10687,13 @@ void GatherTracksAndObjectsToCopy(TSharedRef<UE::Sequencer::FViewModel> Node, TA
 	using namespace UE::Sequencer;
 
 	if (ITrackExtension* TrackNode = Node->CastThis<ITrackExtension>())
+	{
+		if (!TracksToCopy.Contains(Node))
+		{
+			TracksToCopy.Add(Node);
+		}
+	}	
+	else if (ITrackRowExtension* TrackRowNode = Node->CastThis<ITrackRowExtension>())
 	{
 		if (!TracksToCopy.Contains(Node))
 		{
@@ -11102,7 +11140,15 @@ void FSequencer::TrimOrExtendSection(bool bTrimOrExtendLeft)
 				UMovieSceneTrack* Track = TrackNode->GetTrack();
 				if (Track)
 				{
-					MovieSceneToolHelpers::TrimOrExtendSection(Track, Node->IsA<FTrackRowModel>() ? TrackNode->GetRowIndex() : TOptional<int32>(), GetLocalTime(), bTrimOrExtendLeft, Settings->GetDeleteKeysWhenTrimming());
+					MovieSceneToolHelpers::TrimOrExtendSection(Track, TOptional<int32>(), GetLocalTime(), bTrimOrExtendLeft, Settings->GetDeleteKeysWhenTrimming());
+				}
+			}
+			else if (ITrackRowExtension* TrackRowNode = Node->CastThis<ITrackRowExtension>())
+			{
+				UMovieSceneTrack* Track = TrackRowNode->GetParentTrack();
+				if (Track)
+				{
+					MovieSceneToolHelpers::TrimOrExtendSection(Track, TrackRowNode->GetRowIndex(), GetLocalTime(), bTrimOrExtendLeft, Settings->GetDeleteKeysWhenTrimming());
 				}
 			}
 			else if (IObjectBindingExtension* ObjectBindingNode = Node->CastThis<IObjectBindingExtension>())
@@ -11619,13 +11665,13 @@ void FSequencer::BindCommands()
 			for (TWeakViewModelPtr<IOutlinerExtension> WeakNode : SelectedNodes)
 			{
 				FViewModelPtr Node = WeakNode.Pin();
-				if (Node->IsA<ITrackExtension>() || Node->IsA<IObjectBindingExtension>() || Node->IsA<FFolderModel>())
+				if (Node->IsA<ITrackExtension>() || Node->IsA<ITrackRowExtension>() || Node->IsA<IObjectBindingExtension>() || Node->IsA<FFolderModel>())
 				{
 					// if contains one node that can be copied we allow the action
 					// later on we will filter out the invalid nodes in CopySelection() or CutSelection()
 					return true;
 				}
-				else if (Node->GetParent().IsValid() && Node->GetParent()->IsA<ITrackExtension>() && !Node->IsA<FCategoryGroupModel>())
+				else if (Node->GetParent().IsValid() && (Node->GetParent()->IsA<ITrackExtension>() || Node->GetParent()->IsA<ITrackRowExtension>()) && !Node->IsA<FCategoryGroupModel>())
 				{
 					return true;
 				}
diff --git a/Engine/Source/Editor/Sequencer/Private/SequencerAddKeyOperation.cpp b/Engine/Source/Editor/Sequencer/Private/SequencerAddKeyOperation.cpp
--- a/Engine/Source/Editor/Sequencer/Private/SequencerAddKeyOperation.cpp
+++ b/Engine/Source/Editor/Sequencer/Private/SequencerAddKeyOperation.cpp
@@ -10,6 +10,7 @@
 #include "ISequencerTrackEditor.h"
 #include "MVVM/Extensions/IOutlinerExtension.h"
 #include "MVVM/Extensions/ITrackExtension.h"
+#include "MVVM/Extensions/ISectionOwnerExtension.h"
 #include "MVVM/ViewModelTypeID.h"
 #include "MVVM/ViewModels/ChannelModel.h"
 #include "MVVM/ViewModels/ViewModel.h"
@@ -74,7 +75,7 @@ FAddKeyOperation FAddKeyOperation::FromKeyAreas(ISequencerTrackEditor* TrackEdit
 	{
 		for (const TSharedRef<IKeyArea>& KeyArea : InKeyAreas)
 		{
-			Operation.ProcessKeyArea(TrackEditor, KeyArea);
+			Operation.ProcessKeyArea(TrackEditor, KeyArea.ToSharedPtr());
 		}
 	}
 	return Operation;
@@ -83,24 +84,24 @@ FAddKeyOperation FAddKeyOperation::FromKeyAreas(ISequencerTrackEditor* TrackEdit
 void FAddKeyOperation::AddPreFilteredNodes(TArrayView<const TWeakPtr<FViewModel>> FilteredNodes)
 {
 	TSharedPtr<FSequenceModel> SequenceModel;
-	const TArray<FViewModelTypeID> TypeIDs({ ITrackExtension::ID, IOutlinerExtension::ID });
+	const TArray<FViewModelTypeID> TypeIDs({ ISectionOwnerExtension::ID, IOutlinerExtension::ID });
 	for (const TWeakPtr<FViewModel>& FilteredNode : FilteredNodes)
 	{
 		TSharedPtr<FViewModel> Node = FilteredNode.Pin();
 		TSharedPtr<FViewModel> ParentModel = Node->FindAncestorOfTypes(MakeArrayView(TypeIDs));
 		if (ParentModel)
 		{
-			ConsiderKeyableAreas(ParentModel->CastThisShared<ITrackExtension>(), Node);
+			ConsiderKeyableAreas(ParentModel->CastThisShared<ISectionOwnerExtension>(), Node);
 		}
 		else
 		{
 			constexpr bool bIncludeThis = true;
 			for (FParentFirstChildIterator Child(FilteredNode.Pin(), bIncludeThis); Child; ++Child)
 			{
-				if (Child->IsA<ITrackExtension>() && Child->IsA<IOutlinerExtension>())
+				if (Child->IsA<ISectionOwnerExtension>() && Child->IsA<IOutlinerExtension>())
 				{
-					TSharedPtr<ITrackExtension> TrackExtension = Child->CastThisShared<ITrackExtension>();
-					ConsiderKeyableAreas(TrackExtension, *Child);
+					TViewModelPtr<ISectionOwnerExtension> SectionOwner = Child->CastThisShared<ISectionOwnerExtension>();
+					ConsiderKeyableAreas(SectionOwner, *Child);
 
 					// Stop recursively iterating everything within this Child again since we just considered all keyable areas for this and its children
 					Child.IgnoreCurrentChildren();
@@ -110,25 +111,26 @@ void FAddKeyOperation::AddPreFilteredNodes(TArrayView<const TWeakPtr<FViewModel>
 	}
 }
 
-bool FAddKeyOperation::ConsiderKeyableAreas(TSharedPtr<ITrackExtension> InTrackModel, FViewModelPtr KeyAnythingBeneath)
+bool FAddKeyOperation::ConsiderKeyableAreas(TViewModelPtr<ISectionOwnerExtension> InSectionOwner, FViewModelPtr KeyAnythingBeneath)
 {
 	bool bKeyedAnything = false;
 
 	constexpr bool bIncludeThis = true;
 	
 	// Prefer the section that is marked SectionToKey
+	TViewModelPtr<ITrackExtension> TrackModel = InSectionOwner.AsModel()->FindAncestorOfType<ITrackExtension>(bIncludeThis);
+	TSharedPtr<ISequencerTrackEditor> TrackEditor = TrackModel ? TrackModel->GetTrackEditor() : nullptr;
 	for (const TViewModelPtr<FSectionModel>& SectionModel : KeyAnythingBeneath->GetDescendantsOfType<FSectionModel>(bIncludeThis))
 	{
 		UMovieSceneSection* Section = SectionModel->GetSection();
 		UMovieSceneTrack* Track = Cast<UMovieSceneTrack>(Section->GetOuter());
 		if (Track && Track->GetSectionToKey() == Section)
 		{
-			ISequencerTrackEditor* TrackEditor = InTrackModel->GetTrackEditor().Get();
 			for (TParentFirstChildIterator<FChannelModel> ChannelModelIt(SectionModel->AsShared(), bIncludeThis); ChannelModelIt; ++ChannelModelIt)
 			{	
 				if (!ChannelModelIt->GetLinkedOutlinerItem() || ChannelModelIt->GetLinkedOutlinerItem()->IsFilteredOut() == false)
 				{
-					bKeyedAnything |= ProcessKeyArea(TrackEditor, ChannelModelIt->GetKeyArea());
+					bKeyedAnything |= ProcessKeyArea(TrackEditor.Get(), ChannelModelIt->GetKeyArea());
 				}
 			}
 		}
@@ -139,17 +141,16 @@ bool FAddKeyOperation::ConsiderKeyableAreas(TSharedPtr<ITrackExtension> InTrackM
 	{
 		for (TParentFirstChildIterator<FChannelGroupModel> ChannelGroupModelIt(KeyAnythingBeneath->AsShared(), bIncludeThis); ChannelGroupModelIt; ++ChannelGroupModelIt)
 		{
-			bKeyedAnything |= ProcessKeyArea(InTrackModel, *ChannelGroupModelIt);
+			bKeyedAnything |= ProcessKeyArea(TrackEditor.Get(), *ChannelGroupModelIt);
 		}
 	}
 
 	return bKeyedAnything;
 }
 
-bool FAddKeyOperation::ProcessKeyArea(TSharedPtr<ITrackExtension> InTrackModel, TViewModelPtr<FChannelGroupModel> InChannelGroupModel)
+bool FAddKeyOperation::ProcessKeyArea(ISequencerTrackEditor* InTrackEditor, TViewModelPtr<FChannelGroupModel> InChannelGroupModel)
 {
 	bool bKeyedAnything = false;
-	ISequencerTrackEditor* TrackEditor = InTrackModel->GetTrackEditor().Get();
 
 	IOutlinerExtension* OutlinerExtension = InChannelGroupModel->CastThis<IOutlinerExtension>();
 	if (!OutlinerExtension || OutlinerExtension->IsFilteredOut() == false)
@@ -159,7 +160,7 @@ bool FAddKeyOperation::ProcessKeyArea(TSharedPtr<ITrackExtension> InTrackModel,
 		{
 			if (TSharedPtr<FChannelModel> Channel = WeakChannel.Pin())
 			{
-				bKeyedAnything |= ProcessKeyArea(TrackEditor, Channel->GetKeyArea());
+				bKeyedAnything |= ProcessKeyArea(InTrackEditor, Channel->GetKeyArea());
 			}
 		}
 	}
diff --git a/Engine/Source/Editor/Sequencer/Private/SequencerContextMenus.cpp b/Engine/Source/Editor/Sequencer/Private/SequencerContextMenus.cpp
--- a/Engine/Source/Editor/Sequencer/Private/SequencerContextMenus.cpp
+++ b/Engine/Source/Editor/Sequencer/Private/SequencerContextMenus.cpp
@@ -1765,15 +1765,15 @@ void FPasteContextMenu::Setup()
 		{
 			for (const TViewModelPtr<IOutlinerExtension>& Node : Args.DestinationNodes)
 			{
-				TViewModelPtr<ITrackExtension> TrackNode = Node.AsModel()->FindAncestorOfType<ITrackExtension>(true);
-				if (!TrackNode)
+				TViewModelPtr<ISectionOwnerExtension> SectionOwner = Node.AsModel()->FindAncestorOfType<ISectionOwnerExtension>(true);
+				if (!SectionOwner)
 				{
 					continue;
 				}
 
 				FPasteDestination& Destination = PasteDestinations[PasteDestinations.AddDefaulted()];
 
-				for (UMovieSceneSection* Section : TrackNode->GetSections())
+				for (UMovieSceneSection* Section : SectionOwner->GetSections())
 				{
 					if (Section)
 					{
@@ -1818,13 +1818,13 @@ void FPasteContextMenu::Setup()
 		// Build a list of sections based on selected tracks
 		for (const TViewModelPtr<IOutlinerExtension>& Node : Args.DestinationNodes)
 		{
-			TViewModelPtr<ITrackExtension> TrackNode = Node.AsModel()->FindAncestorOfType<ITrackExtension>(true);
-			if (!TrackNode)
+			TViewModelPtr<ISectionOwnerExtension> SectionOwner = Node.AsModel()->FindAncestorOfType<ISectionOwnerExtension>(true);
+			if (!SectionOwner)
 			{
 				continue;
 			}
 
-			UMovieSceneSection* Section = MovieSceneHelpers::FindNearestSectionAtTime(TrackNode->GetSections(), Args.PasteAtTime);
+			UMovieSceneSection* Section = MovieSceneHelpers::FindNearestSectionAtTime(SectionOwner->GetSections(), Args.PasteAtTime);
 			TSharedPtr<FSectionModel> SectionModel = Sequencer->GetNodeTree()->GetSectionModel(Section);
 			if (SectionModel)
 			{
@@ -1859,7 +1859,7 @@ void FPasteContextMenu::Setup()
 		FPasteDestination& Destination = PasteDestinations[PasteDestinations.AddDefaulted()];
 		if (Pair.Value.Num() == 1)
 		{
-			TSharedPtr<FViewModel> Model = Pair.Value[0]->FindAncestorOfTypes({ITrackExtension::ID, IOutlinerExtension::ID});
+			TSharedPtr<FViewModel> Model = Pair.Value[0]->FindAncestorOfTypes({ ISectionOwnerExtension::ID, IOutlinerExtension::ID});
 			if (ensure(Model))
 			{
 				FString Path = IOutlinerExtension::GetPathName(Model);
@@ -1873,7 +1873,7 @@ void FPasteContextMenu::Setup()
 
 		for (TSharedPtr<FSectionModel> Section : Pair.Value)
 		{
-			FViewModelPtr Model = Section->FindAncestorOfTypes({ITrackExtension::ID, IOutlinerExtension::ID});
+			FViewModelPtr Model = Section->FindAncestorOfTypes({ISectionOwnerExtension::ID, IOutlinerExtension::ID});
 			GatherPasteDestinationsForNode(Model.ImplicitCast(), Section->GetSection(), NAME_None, Destination.Reconcilers);
 		}
 
diff --git a/Engine/Source/Editor/Sequencer/Private/SequencerEdMode.cpp b/Engine/Source/Editor/Sequencer/Private/SequencerEdMode.cpp
--- a/Engine/Source/Editor/Sequencer/Private/SequencerEdMode.cpp
+++ b/Engine/Source/Editor/Sequencer/Private/SequencerEdMode.cpp
@@ -8,6 +8,7 @@
 #include "MVVM/ViewModels/ChannelModel.h"
 #include "MVVM/ViewModels/SequencerEditorViewModel.h"
 #include "MVVM/ViewModels/ViewModelIterators.h"
+#include "MVVM/Extensions/ITrackRowExtension.h"
 #include "MVVM/Views/STrackAreaView.h"
 #include "MVVM/Selection/Selection.h"
 #include "EditorViewportClient.h"
@@ -662,7 +663,7 @@ void FSequencerEdMode::OnKeySelected(FViewport* Viewport, HMovieSceneKeyProxy* K
 				TSharedPtr<FSectionModel> SectionHandle = Sequencer->GetNodeTree()->GetSectionModel(Section);
 				if (SectionHandle && KeyData.KeyHandle.IsSet())
 				{
-					TParentFirstChildIterator<FChannelGroupModel> KeyAreaNodes = SectionHandle->GetParentTrackModel().AsModel()->GetDescendantsOfType<FChannelGroupModel>();
+					TParentFirstChildIterator<FChannelGroupModel> KeyAreaNodes = SectionHandle->GetParentSectionOwnerModel().AsModel()->GetDescendantsOfType<FChannelGroupModel>();
 					for (const TViewModelPtr<FChannelGroupModel>& KeyAreaNode : KeyAreaNodes)
 					{
 						TSharedPtr<FChannelModel> Channel = KeyAreaNode->GetChannel(Section);
@@ -935,6 +936,8 @@ void FSequencerEdMode::DrawAudioTracks(FPrimitiveDrawInterface* PDI)
 
 		FQualifiedFrameTime CurrentTime = Sequencer->GetLocalTime();
 
+		TSet<UMovieSceneAudioTrack*> SelectedAudioTracks;
+
 		for (TViewModelPtr<ITrackExtension> TrackModel : Sequencer->GetViewModel()->GetSelection()->Outliner.Filter<ITrackExtension>())
 		{
 			UMovieSceneAudioTrack* AudioTrack = Cast<UMovieSceneAudioTrack>(TrackModel->GetTrack());
@@ -943,7 +946,22 @@ void FSequencerEdMode::DrawAudioTracks(FPrimitiveDrawInterface* PDI)
 			{
 				continue;
 			}
+			SelectedAudioTracks.Add(AudioTrack);
+		}
+		
+		for (TViewModelPtr<ITrackRowExtension> TrackRowModel : Sequencer->GetViewModel()->GetSelection()->Outliner.Filter<ITrackRowExtension>())
+		{
+			UMovieSceneAudioTrack* AudioTrack = Cast<UMovieSceneAudioTrack>(TrackRowModel->GetParentTrack());
 
+			if (!AudioTrack)
+			{
+				continue;
+			}
+			SelectedAudioTracks.Add(AudioTrack);
+		}
+		
+		for (UMovieSceneAudioTrack* AudioTrack : SelectedAudioTracks)
+		{
 			for (UMovieSceneSection* Section : AudioTrack->GetAudioSections())
 			{
 				UMovieSceneAudioSection* AudioSection = Cast<UMovieSceneAudioSection>(Section);
diff --git a/Engine/Source/Editor/Sequencer/Private/SequencerUtilities.cpp b/Engine/Source/Editor/Sequencer/Private/SequencerUtilities.cpp
--- a/Engine/Source/Editor/Sequencer/Private/SequencerUtilities.cpp
+++ b/Engine/Source/Editor/Sequencer/Private/SequencerUtilities.cpp
@@ -91,6 +91,7 @@
 #include "Variants/MovieSceneTimeWarpGetter.h"
 #include "UObject/UObjectIterator.h"
 #include "ObjectTools.h"
+#include "MVVM/Extensions/ISectionOwnerExtension.h"
 
 #include UE_INLINE_GENERATED_CPP_BY_NAME(SequencerUtilities)
 
@@ -166,18 +167,18 @@ TSharedRef<SWidget> FSequencerUtilities::MakeAddButton(FText HoverText, FOnClick
 	return UE::Sequencer::MakeAddButton(HoverText, OnClicked, HoverState, IsEnabled);
 }
 
-void FSequencerUtilities::MakeTimeWarpMenuEntry(FMenuBuilder& MenuBuilder, UE::Sequencer::TWeakViewModelPtr<UE::Sequencer::ITrackExtension> WeakTrackModel)
+void FSequencerUtilities::MakeTimeWarpMenuEntry(FMenuBuilder& MenuBuilder, UE::Sequencer::TWeakViewModelPtr<UE::Sequencer::FViewModel> WeakViewModel)
 {
 	using namespace UE::Sequencer;
 
-	TViewModelPtr<ITrackExtension> TrackModel = WeakTrackModel.Pin();
-	if (!TrackModel)
+	TViewModelPtr<ISectionOwnerExtension> SectionOwner = WeakViewModel.ImplicitPin();
+	if (!SectionOwner)
 	{
 		return;
 	}
 
 	TOptional<UClass*> CommonClass;
-	for (const TViewModelPtr<FSectionModel>& SectionModel : TrackModel->GetSectionModels().IterateSubList<FSectionModel>())
+	for (const TViewModelPtr<FSectionModel>& SectionModel : SectionOwner->GetSectionModels().IterateSubList<FSectionModel>())
 	{
 		UMovieSceneSection* Section = SectionModel->GetSection();
 		if (!Section)
@@ -213,7 +214,7 @@ void FSequencerUtilities::MakeTimeWarpMenuEntry(FMenuBuilder& MenuBuilder, UE::S
 	MenuBuilder.AddSubMenu(
 		TimeWarpLabel,
 		TimeWarpToolTip,
-		FNewMenuDelegate::CreateStatic(PopulateTimeWarpChannelSubMenu, WeakTrackModel)
+		FNewMenuDelegate::CreateStatic(PopulateTimeWarpChannelSubMenu, SectionOwner.AsWeak())
 	);
 }
 
@@ -269,20 +270,27 @@ void FSequencerUtilities::PopulateTimeWarpSubMenu(FMenuBuilder& MenuBuilder, TFu
 }
 
 void FSequencerUtilities::PopulateTimeWarpChannelSubMenu(FMenuBuilder& MenuBuilder, UE::Sequencer::TWeakViewModelPtr<UE::Sequencer::ITrackExtension> WeakTrackModel)
+{
+	using namespace UE::Sequencer;
+	TViewModelPtr<ISectionOwnerExtension> SectionOwner = WeakTrackModel.ImplicitPin();
+	return PopulateTimeWarpChannelSubMenu(MenuBuilder, SectionOwner.AsWeak());
+}
+
+void FSequencerUtilities::PopulateTimeWarpChannelSubMenu(FMenuBuilder& MenuBuilder, UE::Sequencer::TWeakViewModelPtr<UE::Sequencer::ISectionOwnerExtension> WeakSectionOwner)
 {
 	using namespace UE::Sequencer;
 
-	auto HandleTimeWarpSelection = [WeakTrackModel](TSubclassOf<UMovieSceneTimeWarpGetter> Class)
+	auto HandleTimeWarpSelection = [WeakSectionOwner](TSubclassOf<UMovieSceneTimeWarpGetter> Class)
 	{
-		TViewModelPtr<ITrackExtension> TrackModel = WeakTrackModel.Pin();
-		if (!TrackModel)
+		TViewModelPtr<ISectionOwnerExtension> SectionOwner = WeakSectionOwner.Pin();
+		if (!SectionOwner)
 		{
 			return;
 		}
 
 		FScopedTransaction Transaction(LOCTEXT("ChangeTimeWarpType", "Changed Time Warp type"));
 
-		for (const TViewModelPtr<FSectionModel>& SectionModel : TrackModel->GetSectionModels().IterateSubList<FSectionModel>())
+		for (const TViewModelPtr<FSectionModel>& SectionModel : SectionOwner->GetSectionModels().IterateSubList<FSectionModel>())
 		{
 			UMovieSceneSection*         Section = SectionModel->GetSection();
 			FMovieSceneTimeWarpVariant* Variant = Section ? Section->GetTimeWarp() : nullptr;
@@ -298,7 +306,7 @@ void FSequencerUtilities::PopulateTimeWarpChannelSubMenu(FMenuBuilder& MenuBuild
 
 				Section->InvalidateChannelProxy();
 
-				TViewModelPtr<IOutlinerExtension> Outliner = TrackModel.ImplicitCast();
+				TViewModelPtr<IOutlinerExtension> Outliner = SectionOwner.ImplicitCast();
 				if (Outliner && !Outliner->IsExpanded())
 				{
 					Outliner->SetExpansion(true);
diff --git a/Engine/Source/Editor/Sequencer/Private/Tools/EditToolDragOperations.cpp b/Engine/Source/Editor/Sequencer/Private/Tools/EditToolDragOperations.cpp
--- a/Engine/Source/Editor/Sequencer/Private/Tools/EditToolDragOperations.cpp
+++ b/Engine/Source/Editor/Sequencer/Private/Tools/EditToolDragOperations.cpp
@@ -6,6 +6,7 @@
 #include "MovieSceneTrack.h"
 #include "MovieSceneSequence.h"
 #include "Sequencer.h"
+#include "SSequencer.h"
 #include "SequencerSettings.h"
 #include "SequencerCommonHelpers.h"
 #include "MVVM/Selection/Selection.h"
@@ -797,6 +798,34 @@ void FMoveKeysAndSections::OnBeginDrag(const FPointerEvent& MouseEvent, FVector2
 		}
 	}
 
+	// Notify all drag owners that drag is beginning
+	FSectionModelStorageExtension* SectionStorage = Sequencer.GetViewModel()->GetRootModel()->CastDynamic<FSectionModelStorageExtension>();
+	if (SectionStorage)
+	{
+		TSet<TSharedPtr<FViewModel>> ProcessedDragOwners;
+		for (UMovieSceneSection* Section : Sections)
+		{
+			TSharedPtr<FSectionModel> SectionModel = SectionStorage->FindModelForSection(Section);
+			TViewModelPtr<ISectionOwnerExtension> SectionOwner = SectionModel ? SectionModel->GetParentSectionOwnerModel() : nullptr;
+			if (!SectionOwner)
+			{
+				continue;
+			}
+
+			TViewModelPtr<ISectionDragOwnerExtension> DragOwner = SectionOwner.ImplicitCast();
+			if (!DragOwner)
+			{
+				DragOwner = SectionOwner.AsModel()->FindAncestorOfType<ISectionDragOwnerExtension>();
+			}
+
+			if (DragOwner && !ProcessedDragOwners.Contains(DragOwner.AsModel()))
+			{
+				DragOwner->OnBeginSectionVerticalDrag();
+				ProcessedDragOwners.Add(DragOwner.AsModel());
+			}
+		}
+	}
+
 	// Our Key Handles don't store their times so we need to convert the handles into an array of times
 	// so that we can store the relative offset to each one.
 	TArray<FFrameNumber> KeyTimes;
@@ -974,9 +1003,106 @@ void FMoveKeysAndSections::OnEndDrag(const FPointerEvent& MouseEvent, FVector2D
 		}
 	}
 
+	// Notify track models that vertical drag has ended
+	if (bAllowVerticalMovement && Sections.Num() > 0)
+	{
+		FSectionModelStorageExtension* SectionStorage = Sequencer.GetViewModel()->GetRootModel()->CastDynamic<FSectionModelStorageExtension>();
+		if (SectionStorage)
+		{
+			TSet<TSharedPtr<FViewModel>> ProcessedDragOwners;
+
+			// Notify all unique drag owners that drag ended
+			for (UMovieSceneSection* Section : Sections)
+			{
+				TSharedPtr<FSectionModel> SectionModel = SectionStorage->FindModelForSection(Section);
+				TViewModelPtr<ISectionOwnerExtension> SectionOwner = SectionModel ? SectionModel->GetParentSectionOwnerModel() : nullptr;
+				if (!SectionOwner)
+				{
+					continue;
+				}
+
+				// Find the drag owner (track model that implements ISectionDragOwnerExtension)
+				TViewModelPtr<ISectionDragOwnerExtension> DragOwner = SectionOwner.ImplicitCast();
+				if (!DragOwner)
+				{
+					DragOwner = SectionOwner.AsModel()->FindAncestorOfType<ISectionDragOwnerExtension>();
+				}
+
+				if (DragOwner && !ProcessedDragOwners.Contains(DragOwner.AsModel()))
+				{
+					// Build context for the end of drag (preview-based tracks need this to apply changes)
+					FSectionVerticalDragContext EndDragContext;
+					EndDragContext.Section = Section;
+					EndDragContext.DraggedSections = &Sections;
+					EndDragContext.InitialSectionRowIndices = &InitialSectionRowIndicies;
+
+					// Notify the drag owner that drag ended
+					DragOwner->OnEndSectionVerticalDrag(EndDragContext);
+					ProcessedDragOwners.Add(DragOwner.AsModel());
+				}
+			}
+		}
+	}
+
 	EndTransaction();
 }
 
+int32 FMoveKeysAndSections::OnPaint(const FGeometry& AllottedGeometry, const FSlateRect& MyCullingRect, FSlateWindowElementList& OutDrawElements, int32 LayerId) const
+{
+	using namespace UE::Sequencer;
+
+	// Let drag owners paint their drop indicators
+	if (bAllowVerticalMovement && Sections.Num() > 0)
+	{
+		// Get the VirtualTrackArea for coordinate conversion
+		TSharedRef<SWidget> SequencerWidget = Sequencer.GetSequencerWidget();
+		SSequencer* SSequencerPtr = static_cast<SSequencer*>(&SequencerWidget.Get());
+		FVirtualTrackArea VirtualTrackArea = SSequencerPtr->GetVirtualTrackArea();
+
+		FSectionModelStorageExtension* SectionStorage = Sequencer.GetViewModel()->GetRootModel()->CastDynamic<FSectionModelStorageExtension>();
+		if (SectionStorage)
+		{
+			TSet<TSharedPtr<FViewModel>> ProcessedDragOwners;
+
+			// Find unique drag owners
+			for (UMovieSceneSection* Section : Sections)
+			{
+				TSharedPtr<FSectionModel> SectionModel = SectionStorage->FindModelForSection(Section);
+				TViewModelPtr<ISectionOwnerExtension> SectionOwner = SectionModel ? SectionModel->GetParentSectionOwnerModel() : TViewModelPtr<ISectionOwnerExtension>();
+				if (!SectionOwner)
+				{
+					continue;
+				}
+
+				// Find the drag owner (track model that implements ISectionDragOwnerExtension)
+				TViewModelPtr<ISectionDragOwnerExtension> DragOwner = SectionOwner.ImplicitCast();
+				if (!DragOwner)
+				{
+					DragOwner = SectionOwner.AsModel()->FindAncestorOfType<ISectionDragOwnerExtension>();
+				}
+
+				if (DragOwner && !ProcessedDragOwners.Contains(DragOwner.AsModel()))
+				{
+					// Let the drag owner paint its preview indicators
+					LayerId = DragOwner->OnPaintSectionDragPreview(
+						AllottedGeometry,
+						VirtualTrackArea,
+						MyCullingRect,
+						OutDrawElements,
+						LayerId,
+						FWidgetStyle(),
+						true
+					);
+
+					ProcessedDragOwners.Add(DragOwner.AsModel());
+				}
+			}
+		}
+	}
+
+	return LayerId;
+}
+
 void FMoveKeysAndSections::ModifyNonSelectedSections()
 {
 	using namespace UE::Sequencer;
@@ -1144,204 +1270,40 @@ bool FMoveKeysAndSections::HandleSectionMovement(FFrameTime MouseTime, FVector2D
 
 		Tracks.AddUnique(TrackExtModel.AsModel());
 
-		int32 TargetRowIndex = Section->GetRowIndex();
-
-		// @todo_sequencer_mvvm: need to go through this dragging code and figure out what's going on and what it means
-		//                       now that models and views have separated concerns
-
-		// Handle vertical dragging to re-arrange tracks. We don't support vertical rearranging if you're dragging via
-		// a key, as the built in offset causes it to always jump down a row even without moving the mouse.
-		if (Track->SupportsMultipleRows() && AllSections.Num() > 1 && bAllowVerticalMovement)
+		// Get the drag owner (track model that implements ISectionDragOwnerExtension)
+		TViewModelPtr<ISectionDragOwnerExtension> DragOwner = TrackExtModel.ImplicitCast();
+		if (!DragOwner)
 		{
-			// Compute the max row index whilst disregarding the one we're dragging
-			int32 MaxRowIndex = 0;
-			for (UMovieSceneSection* NonDraggedSection : NonDraggedSections)
-			{
-				if (NonDraggedSection != Section)
-				{
-					MaxRowIndex = FMath::Max(NonDraggedSection->GetRowIndex() + 1, MaxRowIndex);
-				}
-			}
-
-			// Handle sub-track and non-sub-track dragging
-			if (TViewModelPtr<FTrackModel> TrackModel = TrackExtModel.ImplicitCast())
-			{
-				const int32 NumRows = FMath::Max(Section->GetRowIndex() + 1, MaxRowIndex);
-
-				// Find the total height of the track - this is necessary because tracks may contain key areas, but they will not use sub tracks unless there is more than one row
-				const FVirtualGeometry VirtualGeometry = TrackModel->GetVirtualGeometry();
-
-				// Assume same height rows
-				const float VirtualSectionHeight = VirtualGeometry.NestedBottom - VirtualGeometry.Top;
-				const float VirtualRowHeight = VirtualSectionHeight / NumRows;
-				const float MouseOffsetWithinRow = VirtualMousePos.Y - (VirtualGeometry.Top + (VirtualRowHeight * TargetRowIndex));
-
-				const int32 NewIndex = FMath::FloorToInt((VirtualMousePos.Y - VirtualGeometry.Top) / VirtualRowHeight);
-				TargetRowIndex = FMath::Clamp(NewIndex, 0, MaxRowIndex);
-
-				// If close to the top of the row, move else everything down
-				if (VirtualMousePos.Y <= VirtualGeometry.Top || LocalMousePos.Y <= 0)
-				{
-					TargetRowIndex = -1;
-				}
-			}
-			else if (TViewModelPtr<FTrackRowModel> TrackRow = TrackExtModel.ImplicitCast())
-			{
-				TSharedPtr<FTrackModel> ParentTrack = TrackExtModel.AsModel()->FindAncestorOfType<FTrackModel>();
-				if (ensure(ParentTrack.IsValid()))
-				{
-					int32 ChildIndex = 0;
-					for (TSharedPtr<FViewModel> ChildNode : ParentTrack->GetChildren())
-					{
-						FVirtualGeometry ChildVirtualGeometry;
-						if (IGeometryExtension* ChildGeometryExtension = ChildNode->CastThis<IGeometryExtension>())
-						{
-							ChildVirtualGeometry = ChildGeometryExtension->GetVirtualGeometry();
-						}
-
-						float VirtualSectionTop = ChildVirtualGeometry.Top;
-						float VirtualSectionBottom = ChildVirtualGeometry.NestedBottom;
-
-						if (ChildIndex == 0 && (VirtualMousePos.Y <= VirtualSectionTop || LocalMousePos.Y <= 0))
-						{
-							TargetRowIndex = 0;
-							for (TSharedPtr<FSectionModel> SectionNode : TrackRow->GetTrackAreaModelListAs<FSectionModel>())
-							{
-								if (!Sections.Contains(SectionNode->GetSection()))
-								{
-									TargetRowIndex = -1;
-									break;
-								}
-							}
-							break;
-						}
-						else if (VirtualMousePos.Y < VirtualSectionBottom)
-						{
-							TargetRowIndex = ChildIndex;
-							break;
-						}
-						else
-						{
-							TargetRowIndex = ChildIndex + 1;
-						}
-
-						++ChildIndex;
-					}
-				
-					// Track if we're expanding a parent track so we can unexpand it if we stop targeting it
-					if (TargetRowIndex > 0)
-					{
-						if (!ParentTrack->IsExpanded() && ParentTrack != ExpandedParentTrack)
-						{
-							if (TSharedPtr<FTrackModel> ExpandedParentTrackPinned = ExpandedParentTrack.Pin())
-							{
-								ExpandedParentTrackPinned->SetExpansion(false);
-								ExpandedParentTrack = nullptr;
-							}
-							ExpandedParentTrack = ParentTrack;
-							ParentTrack->SetExpansion(true);
-						}
-					}
-					else if (TSharedPtr<FTrackModel> ExpandedParentTrackPinned = ExpandedParentTrack.Pin())
-					{
-						ExpandedParentTrackPinned->SetExpansion(false);
-						ExpandedParentTrack = nullptr;
-					}
-				}
-			}
+			DragOwner = TrackExtModel.AsModel()->FindAncestorOfType<ISectionDragOwnerExtension>();
 		}
 
-		bool bDeltaX = DesiredDeltaX != 0;
-		bool bDeltaY = TargetRowIndex != Section->GetRowIndex();
-		const int32 TargetRowDelta = TargetRowIndex - Section->GetRowIndex();
-
-		// Prevent flickering by only moving sections if the user has actually made an effort to do so
-		if (bDeltaY && PrevMousePosY.IsSet())
-		{
-			// Check mouse has been moved in the direction of intended move
-			if ((TargetRowDelta < 0 && LocalMousePos.Y - PrevMousePosY.GetValue() > 1.0f) || (TargetRowDelta > 0 && LocalMousePos.Y - PrevMousePosY.GetValue() < 1.0f))
-			{
-				// Mouse was not moved in the direction the section wants to swap
-				// Assume offset is due to UI relayout and block moving the section
-				bDeltaY = false;
-			}
-		}
+		// @todo_sequencer_mvvm: need to go through this dragging code and figure out what's going on and what it means
+		//                       now that models and views have separated concerns
 
 		// Horizontal movement
+		bool bDeltaX = DesiredDeltaX != 0;
 		if (bDeltaX)
 		{
 			Section->MoveSection(MaxDeltaX.Get(DesiredDeltaX));
 		}
 
-
-		// Vertical movement
-		if (bDeltaY && !bSectionsAreOnDifferentRows &&
-			(
-				Section->GetBlendType().IsValid() ||
-				!Section->OverlapsWithSections(NonDraggedSections, TargetRowIndex - Section->GetRowIndex(), DesiredDeltaX.Value)
-				)
-			)
+		// Vertical movement - delegate to drag owner
+		if (DragOwner.IsValid())
 		{
-			// Reached the top, move everything else we're not moving downwards
-			if (TargetRowIndex == -1)
-			{
-				if (!bSectionsAreOnDifferentRows)
-				{
-					// If the sections being moved are all at the top, and all others are below it, do nothing
-					bool bSectionsBeingMovedAreAtTop = true;
-					for (const FInitialRowIndex& InitialRowIndex : InitialSectionRowIndicies)
-					{
-						if (!Sections.Contains(InitialRowIndex.Section))
-						{
-							if (LowestRowIndex.IsSet() && InitialRowIndex.RowIndex <= LowestRowIndex.GetValue())
-							{
-								bSectionsBeingMovedAreAtTop = false;
-								break;
-							}
-						}
-					}
-
-					if (!bSectionsBeingMovedAreAtTop)
-					{
-						for (const FInitialRowIndex& InitialRowIndex : InitialSectionRowIndicies)
-						{
-							if (!Sections.Contains(InitialRowIndex.Section))
-							{
-								InitialRowIndex.Section->Modify();
-								InitialRowIndex.Section->SetRowIndex(InitialRowIndex.RowIndex + 1);
-								bRowIndexChanged = true;
-							}
-						}
-					}
-				}
-			}
-			else
-			{
-				if (!bSectionsAreOnDifferentRows)
-				{
-					// If the sections being moved are all at the bottom, and all others are aove it, do nothing
-					bool bSectionsBeingMovedAreAtBottom = true;
-
-					for (const FInitialRowIndex& InitialRowIndex : InitialSectionRowIndicies)
-					{
-						if (!Sections.Contains(InitialRowIndex.Section))
-						{
-							if (HighestRowIndex.IsSet() && InitialRowIndex.RowIndex >= HighestRowIndex.GetValue())
-							{
-								bSectionsBeingMovedAreAtBottom = false;
-								break;
-							}
-						}
-					}
-
-					if (!bSectionsBeingMovedAreAtBottom || TargetRowIndex < Section->GetRowIndex())
-					{
-						Section->Modify();
-						Section->SetRowIndex(TargetRowIndex);
-						bRowIndexChanged = true;
-					}
-				}
-			}
+			// Build context for vertical drag
+			FSectionVerticalDragContext DragContext;
+			DragContext.Section = Section;
+			DragContext.VirtualMousePos = VirtualMousePos;
+			DragContext.LocalMousePos = LocalMousePos;
+			DragContext.DraggedSections = &Sections;
+			DragContext.InitialSectionRowIndices = &InitialSectionRowIndicies;
+			DragContext.bSectionsAreOnDifferentRows = bSectionsAreOnDifferentRows;
+			DragContext.LowestDraggedRow = LowestRowIndex;
+			DragContext.HighestDraggedRow = HighestRowIndex;
+			DragContext.PrevMousePosY = PrevMousePosY;
+
+			bool bSectionRowChanged = DragOwner->OnSectionVerticalDrag(DragContext);
+			bRowIndexChanged |= bSectionRowChanged;
 		}
 	}
 
diff --git a/Engine/Source/Editor/Sequencer/Private/Tools/EditToolDragOperations.h b/Engine/Source/Editor/Sequencer/Private/Tools/EditToolDragOperations.h
--- a/Engine/Source/Editor/Sequencer/Private/Tools/EditToolDragOperations.h
+++ b/Engine/Source/Editor/Sequencer/Private/Tools/EditToolDragOperations.h
@@ -12,6 +12,7 @@
 #include "Tools/SequencerSnapField.h"
 #include "Channels/MovieSceneChannelHandle.h"
 #include "MVVM/Extensions/IDraggableTrackAreaExtension.h"
+#include "MVVM/Extensions/ISectionDragOwnerExtension.h"
 
 
 class FSequencer;
@@ -170,6 +171,7 @@ public:
 	virtual void OnDrag(const FPointerEvent& MouseEvent, FVector2D LocalMousePos, const UE::Sequencer::FVirtualTrackArea& VirtualTrackArea) override;
 	virtual void OnEndDrag(const FPointerEvent& MouseEvent, FVector2D LocalMousePos, const UE::Sequencer::FVirtualTrackArea& VirtualTrackArea) override;
 	virtual FCursorReply GetCursor() const override { return FCursorReply::Cursor(EMouseCursor::CardinalCross); }
+	virtual int32 OnPaint(const FGeometry& AllottedGeometry, const FSlateRect& MyCullingRect, FSlateWindowElementList& OutDrawElements, int32 LayerId) const override;
 	// ~FEditToolDragOperation interface
 
 	/* UE::Sequencer:IDragOperation Interface */
@@ -215,14 +217,8 @@ protected:
 	/** Array of relative offsets for each selected item. Keys + Sections are both added to this array. */
 	TArray<FFrameNumber> RelativeSnapOffsets;
 
-	struct FInitialRowIndex
-	{
-		UMovieSceneSection* Section;
-		int32 RowIndex;
-	};
-
 	/** Store the row each section starts on when we start dragging. */
-	TArray<FInitialRowIndex> InitialSectionRowIndicies;
+	TArray<UE::Sequencer::FInitialRowIndex> InitialSectionRowIndicies;
 
 	/** Array of sections and owners that we called Modify on because we're editing keys that belong to these sections, but not actually moving these sections. */
 	TSet<UObject*> ModifiedNonSelectedObjects;
@@ -230,9 +226,6 @@ protected:
 	/** Optional snap field to use when dragging */
 	TOptional<FSequencerSnapField> SnapField;
 
-	/** If we expanded a parent track while dragging, track it here so we can re-collapse it if not dropping on it. */
-	TWeakPtr<UE::Sequencer::FTrackModel> ExpandedParentTrack;
-
 	/** If the user is moving them via clicking on the Section then we'll allow vertical re-arranging, otherwise not. */
 	bool bAllowVerticalMovement;
 
diff --git a/Engine/Source/Editor/Sequencer/Private/Tools/SequencerEditTool_Movement.cpp b/Engine/Source/Editor/Sequencer/Private/Tools/SequencerEditTool_Movement.cpp
--- a/Engine/Source/Editor/Sequencer/Private/Tools/SequencerEditTool_Movement.cpp
+++ b/Engine/Source/Editor/Sequencer/Private/Tools/SequencerEditTool_Movement.cpp
@@ -392,6 +392,12 @@ int32 FSequencerEditTool_Movement::OnPaint(const FGeometry& AllottedGeometry, co
 {
 	using namespace UE::Sequencer;
 
+	// Let the drag operation paint first (for preview indicators)
+	if (DragOperation.IsValid())
+	{
+		LayerId = DragOperation->OnPaint(AllottedGeometry, MyCullingRect, OutDrawElements, LayerId);
+	}
+
 	if (CursorDecorator)
 	{
 		FSlateDrawElement::MakeBox(
diff --git a/Engine/Source/Editor/Sequencer/Private/Tools/SequencerEntityVisitor.cpp b/Engine/Source/Editor/Sequencer/Private/Tools/SequencerEntityVisitor.cpp
--- a/Engine/Source/Editor/Sequencer/Private/Tools/SequencerEntityVisitor.cpp
+++ b/Engine/Source/Editor/Sequencer/Private/Tools/SequencerEntityVisitor.cpp
@@ -8,6 +8,7 @@
 #include "MVVM/Extensions/IGeometryExtension.h"
 #include "MVVM/Extensions/IOutlinerExtension.h"
 #include "MVVM/Extensions/ITrackAreaExtension.h"
+#include "MVVM/Extensions/ITrackExtension.h"
 #include "MVVM/Extensions/ITrackLaneExtension.h"
 #include "MVVM/ViewModels/CategoryModel.h"
 #include "MVVM/ViewModels/ChannelModel.h"
@@ -178,18 +179,10 @@ void FSequencerEntityWalker::VisitAnyChannels(const ISequencerEntityVisitor& Vis
 	{
 		ChannelGroup = ChannelGroupNode;
 	}
-	if (FTrackModel* TrackModel = InNode->CastThis<FTrackModel>())
+	if (ITopLevelChannelHolderExtension* TopLevelChannelHolder = InNode->CastThis<ITopLevelChannelHolderExtension>())
 	{
 		// TODO: Top-level channel group -- there used to be only one at most but this assumption isn't necessarily true anymore
-		if (auto It = TrackModel->GetTopLevelChannels().IterateSubList<FChannelGroupModel>(); It)
-		{
-			const TViewModelPtr<FChannelGroupModel>& TopLevelChannel = *It;
-			ChannelGroup = TopLevelChannel.Get();
-		}
-	}
-	else if (FTrackRowModel* TrackRowModel = InNode->CastThis<FTrackRowModel>())
-	{
-		if (auto It = TrackRowModel->GetTopLevelChannels().IterateSubList<FChannelGroupModel>(); It)
+		if (auto It = TopLevelChannelHolder->GetTopLevelChannels().IterateSubList<FChannelGroupModel>(); It)
 		{
 			const TViewModelPtr<FChannelGroupModel>& TopLevelChannel = *It;
 			ChannelGroup = TopLevelChannel.Get();
diff --git a/Engine/Source/Editor/Sequencer/Private/Widgets/Sidebar/SequencerSelectionDrawer.cpp b/Engine/Source/Editor/Sequencer/Private/Widgets/Sidebar/SequencerSelectionDrawer.cpp
--- a/Engine/Source/Editor/Sequencer/Private/Widgets/Sidebar/SequencerSelectionDrawer.cpp
+++ b/Engine/Source/Editor/Sequencer/Private/Widgets/Sidebar/SequencerSelectionDrawer.cpp
@@ -15,6 +15,7 @@
 #include "MVVM/ViewModels/CategoryModel.h"
 #include "MVVM/ViewModels/FolderModel.h"
 #include "MVVM/ViewModels/SectionModel.h"
+#include "MVVM/ViewModels/SectionOutlinerModel.h"
 #include "MVVM/ViewModels/SequencerEditorViewModel.h"
 #include "MVVM/ViewModels/TrackRowModel.h"
 #include "PropertyEditorModule.h"
@@ -338,27 +339,36 @@ void FSequencerSelectionDrawer::BuildOutlinerDetails(const TSharedRef<FSequencer
 	TSet<TViewModelPtr<FChannelGroupOutlinerModel>> ChannelGroups;
 	TArray<TPair<TWeakObjectPtr<UMovieSceneTrack>, int32>> SelectedTrackRows;
 
+	bool bIsSectionOutliner = false;
 	for (const FViewModelPtr OutlinerItem : InSelection->Outliner)
 	{
+		if (const TViewModelPtr<ISectionOwnerExtension> SectionOwner = OutlinerItem.ImplicitCast())
+		{
+			WeakSectionObjects.Append(SectionOwner->GetSections());
+		}
+
 		if (const TViewModelPtr<ITrackExtension> TrackExtension = OutlinerItem.ImplicitCast())
 		{
-			WeakSectionObjects.Append(TrackExtension->GetSections());
 			UMovieSceneTrack* Track = TrackExtension->GetTrack();
 			if (IsValid(Track))
 			{
 				WeakTrackObjects.Add(Track);
-				// Only add a 'track row' as selected if either we have an actual 'track row' selected, or else we have a track selected and there's only a single
-				// track row, and the track allows multiple rows.
-				if (TViewModelPtr<FTrackRowModel> TrackRowModel = TrackExtension.ImplicitCast())
+				// Only add a 'track row' as selected if we have a track selected and there's only a single track row, and the track allows multiple rows.
+				if (Track->SupportsMultipleRows() && Track->GetMaxRowIndex() == 0)
 				{
-					SelectedTrackRows.Add(TPair<TWeakObjectPtr<UMovieSceneTrack>, int32>(Track, TrackExtension->GetRowIndex()));
-				}
-				else if (Track->SupportsMultipleRows() && Track->GetMaxRowIndex() == 0)
-				{
-					SelectedTrackRows.Add(TPair<TWeakObjectPtr<UMovieSceneTrack>, int32>(Track, TrackExtension->GetRowIndex()));
+					SelectedTrackRows.Add(TPair<TWeakObjectPtr<UMovieSceneTrack>, int32>(Track, 0));
 				}
 			}
 		}
+		else if (const TViewModelPtr<ITrackRowExtension> TrackRowExtension = OutlinerItem.ImplicitCast())
+		{
+			UMovieSceneTrack* Track = TrackRowExtension->GetParentTrack();
+			if (IsValid(Track))
+			{
+				WeakTrackObjects.Add(Track);
+				SelectedTrackRows.Add(TPair<TWeakObjectPtr<UMovieSceneTrack>, int32>(Track, TrackRowExtension->GetRowIndex()));
+			}
+		}
 		else if (const TViewModelPtr<FObjectBindingModel> ObjectBindingModel = OutlinerItem.ImplicitCast())
 		{
 			ObjectBindings.Add(ObjectBindingModel);
@@ -373,6 +383,10 @@ void FSequencerSelectionDrawer::BuildOutlinerDetails(const TSharedRef<FSequencer
         {
         	ChannelGroups.Add(ChannelGroupOutlinerModel);
         }
+		else if (const TViewModelPtr<FSectionOutlinerModel> SectionOutlinerModel = OutlinerItem.ImplicitCast())
+		{
+			bIsSectionOutliner = true;
+		}
 	}
 
 	if (!WeakFolderObjects.IsEmpty())
@@ -414,12 +428,12 @@ void FSequencerSelectionDrawer::BuildOutlinerDetails(const TSharedRef<FSequencer
 		SequencerHelpers::BuildEditTrackMenu(WeakSequencer, WeakTrackObjects, MenuBuilder, false);
 	}
 
-	// Only show section details for a selected track/trackrow if a single infinite section on the row
+	// Only show section details for a selected track/trackrow if a single infinite section on the row or a section outliner
 	if (WeakSectionObjects.Num() == 1 && WeakSectionObjects[0].IsValid())
 	{
 		if (UMovieSceneSection* Section = Cast<UMovieSceneSection>(WeakSectionObjects[0].Get()))
 		{
-			if (!Section->GetRange().HasLowerBound() && !Section->GetRange().HasUpperBound())
+			if (bIsSectionOutliner || (!Section->GetRange().HasLowerBound() && !Section->GetRange().HasUpperBound()))
 			{
 				SequencerHelpers::BuildEditSectionMenu(WeakSequencer, WeakSectionObjects, MenuBuilder, false);
 			}
diff --git a/Engine/Source/Editor/Sequencer/Public/Filters/Filters/SequencerTrackFilter_Keyed.h b/Engine/Source/Editor/Sequencer/Public/Filters/Filters/SequencerTrackFilter_Keyed.h
--- a/Engine/Source/Editor/Sequencer/Public/Filters/Filters/SequencerTrackFilter_Keyed.h
+++ b/Engine/Source/Editor/Sequencer/Public/Filters/Filters/SequencerTrackFilter_Keyed.h
@@ -1,4 +1,4 @@
-﻿// Copyright Epic Games, Inc. All Rights Reserved.
+// Copyright Epic Games, Inc. All Rights Reserved.
 
 #pragma once
 
@@ -9,8 +9,6 @@ class FSequencerTrackFilter_Keyed : public FSequencerTrackFilter
 public:
 	static FString StaticName() { return TEXT("Keyed"); }
 
-	static bool DoesTrackExtensionHaveKeys(const UE::Sequencer::TWeakViewModelPtr<UE::Sequencer::ITrackExtension>& InTrack);
-
 	FSequencerTrackFilter_Keyed(ISequencerTrackFilters& InFilterInterface, TSharedPtr<FFilterCategory> InCategory = nullptr);
 
 	//~ Begin FSequencerTrackFilter
diff --git a/Engine/Source/Editor/Sequencer/Public/ISequencerTrackEditor.h b/Engine/Source/Editor/Sequencer/Public/ISequencerTrackEditor.h
--- a/Engine/Source/Editor/Sequencer/Public/ISequencerTrackEditor.h
+++ b/Engine/Source/Editor/Sequencer/Public/ISequencerTrackEditor.h
@@ -35,15 +35,13 @@ namespace UE::Sequencer
 
 struct FBuildColumnWidgetParams : UE::Sequencer::FCreateOutlinerViewParams
 {
-	FBuildColumnWidgetParams(UE::Sequencer::TViewModelPtr<UE::Sequencer::ITrackExtension> InTrack, const UE::Sequencer::FCreateOutlinerViewParams& InBaseParams)
+	FBuildColumnWidgetParams(UE::Sequencer::TViewModelPtr<UE::Sequencer::FViewModel> InViewModelItem, const UE::Sequencer::FCreateOutlinerViewParams& InBaseParams)
 		: UE::Sequencer::FCreateOutlinerViewParams(InBaseParams)
-		, ViewModel(InTrack.AsModel())
-		, TrackModel(InTrack)
+		, ViewModel(InViewModelItem.AsModel())
 	{}
 
 	/**  */
 	UE::Sequencer::FViewModelPtr ViewModel;
-	UE::Sequencer::TViewModelPtr<UE::Sequencer::ITrackExtension> TrackModel;
 };
 
 /** Data structure containing information required to build an edit widget */
diff --git a/Engine/Source/Editor/Sequencer/Public/MVVM/Extensions/ITrackExtension.h b/Engine/Source/Editor/Sequencer/Public/MVVM/Extensions/ITrackExtension.h
--- a/Engine/Source/Editor/Sequencer/Public/MVVM/Extensions/ITrackExtension.h
+++ b/Engine/Source/Editor/Sequencer/Public/MVVM/Extensions/ITrackExtension.h
@@ -3,22 +3,16 @@
 #pragma once
 
 #include "MVVM/ViewModelTypeID.h"
-#include "MVVM/ViewModels/SectionModel.h"
 #include "MVVM/ViewModels/ViewModel.h"
-#include "MVVM/ViewModels/ViewModelHierarchy.h"
-#include "MVVM/ViewModels/ViewModelIterators.h"
 
 #define UE_API SEQUENCER_API
 
-class UMovieSceneSection;
 class UMovieSceneTrack;
 class ISequencerTrackEditor;
 
 namespace UE::Sequencer
 {
 
-struct FViewModelChildren;
-
 class ITrackExtension
 {
 public:
@@ -28,22 +22,6 @@ public:
 	virtual ~ITrackExtension()= default;
 
 	virtual UMovieSceneTrack* GetTrack() const = 0;
-	virtual int32 GetRowIndex() const = 0;
-	virtual FViewModelChildren GetSectionModels() = 0;
-
-	TArray<UMovieSceneSection*> GetSections() 
-	{ 
-		TArray<UMovieSceneSection*> Sections; 
-		for (const TViewModelPtr<FSectionModel>& Item : GetSectionModels().IterateSubList<FSectionModel>())
-		{
-			if (UMovieSceneSection* Section = Item->GetSection())
-			{
-				Sections.Add(Section);
-			}
-		}
-
-		return Sections;
-	}
 
 	virtual TSharedPtr<ISequencerTrackEditor> GetTrackEditor() const = 0;
 };
diff --git a/Engine/Source/Editor/Sequencer/Public/MVVM/TrackModelStorageExtension.h b/Engine/Source/Editor/Sequencer/Public/MVVM/TrackModelStorageExtension.h
--- a/Engine/Source/Editor/Sequencer/Public/MVVM/TrackModelStorageExtension.h
+++ b/Engine/Source/Editor/Sequencer/Public/MVVM/TrackModelStorageExtension.h
@@ -4,6 +4,7 @@
 
 #include "UObject/ObjectKey.h"
 #include "MVVM/Extensions/DynamicExtensionContainer.h"
+#include "MVVM/ViewModelPtr.h"
 #include "EventHandlers/ISequenceDataEventHandler.h"
 #include "ISequencerModule.h"
 
@@ -17,6 +18,7 @@ namespace Sequencer
 class FTrackModel;
 class FSequenceModel;
 class FViewModel;
+class ITrackExtension;
 
 class FTrackModelStorageExtension
 	: public IDynamicExtension
@@ -33,8 +35,15 @@ public:
 
 	TSharedPtr<FTrackModel> CreateModelForTrack(UMovieSceneTrack* InTrack, TSharedPtr<FViewModel> DesiredParent = nullptr);
 
+	/** Find the model for a track - returns FTrackModel for backwards compatibility */
 	SEQUENCER_API TSharedPtr<FTrackModel> FindModelForTrack(UMovieSceneTrack* InTrack) const;
 
+	/** Find the registered track extension for a track */
+	SEQUENCER_API TViewModelPtr<ITrackExtension> FindTrackExtensionForTrack(UMovieSceneTrack* InTrack) const;
+
+	/** Register a track extension for a track */
+	SEQUENCER_API void RegisterTrackExtensionForTrack(UMovieSceneTrack* InTrack, TViewModelPtr<ITrackExtension> InModel);
+
 private:
 
 	void OnTrackAdded(UMovieSceneTrack*) override;
@@ -44,7 +53,7 @@ private:
 
 	TArray<FOnCreateTrackModel> TrackModelCreators;
 
-	TMap<FObjectKey, TWeakPtr<FTrackModel>> TrackToModel;
+	TMap<FObjectKey, TWeakViewModelPtr<ITrackExtension>> TrackToModel;
 	FSequenceModel* OwnerModel;
 };
 
diff --git a/Engine/Source/Editor/Sequencer/Public/MVVM/ViewModels/SectionModel.h b/Engine/Source/Editor/Sequencer/Public/MVVM/ViewModels/SectionModel.h
--- a/Engine/Source/Editor/Sequencer/Public/MVVM/ViewModels/SectionModel.h
+++ b/Engine/Source/Editor/Sequencer/Public/MVVM/ViewModels/SectionModel.h
@@ -32,6 +32,7 @@ namespace Sequencer
 {
 
 class ITrackExtension;
+class ISectionOwnerExtension;
 class FChannelModel;
 struct FOverlappingSections;
 
@@ -85,10 +86,13 @@ public:
 	/** Gets the section interface */
 	UE_API TSharedPtr<ISequencerSection> GetSectionInterface() const;
 
-	/** Gets the parent track or track row model */
+	/** Gets the parent section owner model (which could be a track, track row, or something else like a mixer layer) */
+	UE_API TViewModelPtr<ISectionOwnerExtension> GetParentSectionOwnerModel() const;
+
+	/** Gets the ancestor track model */
 	UE_API TViewModelPtr<ITrackExtension> GetParentTrackModel() const;
 
-	/** Gets the parent track or track row model as an ITrackExtension */
+	/** Gets the ancestor track model as an ITrackExtension */
 	UE_API TViewModelPtr<ITrackExtension> GetParentTrackExtension() const;
 
 	/** Gets an array of sections that underlap the specified section */
diff --git a/Engine/Source/Editor/Sequencer/Public/MVVM/ViewModels/TrackModel.h b/Engine/Source/Editor/Sequencer/Public/MVVM/ViewModels/TrackModel.h
--- a/Engine/Source/Editor/Sequencer/Public/MVVM/ViewModels/TrackModel.h
+++ b/Engine/Source/Editor/Sequencer/Public/MVVM/ViewModels/TrackModel.h
@@ -16,6 +16,9 @@
 #include "MVVM/Extensions/IDraggableOutlinerExtension.h"
 #include "MVVM/Extensions/IDeletableExtension.h"
 #include "MVVM/Extensions/IConditionableExtension.h"
+#include "MVVM/Extensions/ISectionDragOwnerExtension.h"
+#include "MVVM/Extensions/ISectionOwnerExtension.h"
+#include "MVVM/Extensions/ITopLevelChannelHolderExtension.h"
 #include "MVVM/ViewModels/ViewModelHierarchy.h"
 #include "EventHandlers/ISignedObjectEventHandler.h"
 
@@ -36,13 +39,13 @@ namespace UE
 namespace Sequencer
 {
 
-class FSectionModel;
-
 class FTrackModel
 	: public FEvaluableOutlinerItemModel
 	, public IRenameableExtension
 	, public IResizableExtension
 	, public ITrackExtension
+	, public ISectionOwnerExtension
+	, public ITopLevelChannelHolderExtension
 	, public ITrackAreaExtension
 	, public ILockableExtension
 	, public IGroupableExtension
@@ -50,6 +53,8 @@ class FTrackModel
 	, public IDraggableOutlinerExtension
 	, public IDeletableExtension
 	, public IConditionableExtension
+	, public IOutlinerDropTargetOutlinerExtension
+	, public ISectionDragOwnerExtension
 	, public UE::MovieScene::TIntrusiveEventHandler<UE::MovieScene::ISignedObjectEventHandler>
 	, public UE::MovieScene::IDeferredSignedObjectFlushSignal
 {
@@ -60,19 +65,21 @@ public:
 		, IRenameableExtension
 		, IResizableExtension
 		, ITrackExtension
+		, ISectionOwnerExtension
+		, ITopLevelChannelHolderExtension
 		, ITrackAreaExtension
 		, ILockableExtension
 		, IGroupableExtension
 		, ISortableExtension
 		, IDraggableOutlinerExtension
 		, IDeletableExtension
-		, IConditionableExtension);
+		, IConditionableExtension
+		, IOutlinerDropTargetOutlinerExtension
+		, ISectionDragOwnerExtension);
 
 	UE_API explicit FTrackModel(UMovieSceneTrack* Track);
 	UE_API ~FTrackModel();
 
-	UE_API FViewModelChildren GetTopLevelChannels();
-
 public:
 
 	static UE_API EViewModelListType GetTopLevelChannelType();
@@ -80,10 +87,13 @@ public:
 
 	/*~ ITrackExtension */
 	UE_API UMovieSceneTrack* GetTrack() const override;
-	UE_API int32 GetRowIndex() const override;
+	UE_API TSharedPtr<ISequencerTrackEditor> GetTrackEditor() const override;
+
+	/*~ ISectionOwnerExtension */
 	UE_API FViewModelChildren GetSectionModels() override;
 
-	UE_API TSharedPtr<ISequencerTrackEditor> GetTrackEditor() const override;
+	/*~ ITopLevelTrackHolderExtension */
+	UE_API FViewModelChildren GetTopLevelChannels() override;
 
 	/*~ ISignedObjectEventHandler */
 	UE_API void OnModifiedDirectly(UMovieSceneSignedObject*) override;
@@ -141,6 +151,10 @@ public:
 	/*~ IDraggableOutlinerExtension */
 	UE_API bool CanDrag() const override;
 
+	/*~ IOutlinerDropTargetOutlinerExtension */
+	UE_API virtual TOptional<EItemDropZone> CanAcceptDrop(const FViewModelPtr& TargetModel, const FDragDropEvent& DragDropEvent, EItemDropZone InItemDropZone) override;
+	UE_API virtual void PerformDrop(const FViewModelPtr& TargetModel, const FDragDropEvent& DragDropEvent, EItemDropZone InItemDropZone) override;
+
 	/*~ FOutlinerItemModel */
 	UE_API bool HasCurves() const override;
 	UE_API void BuildContextMenu(FMenuBuilder& MenuBuilder) override;
@@ -162,12 +176,28 @@ public:
 	/*~ FViewModel interface */
 	UE_API virtual void OnConstruct() override;
 
-private:
 
-	UE_API void ForceUpdate();
+protected:
+
+	/** Rebuilds the track's layout and child models. Override to customize track layout behavior. */
+	UE_API virtual void ForceUpdate();
+
+	/*~ ISectionDragOwnerExtension */
+	UE_API virtual void OnBeginSectionVerticalDrag() override;
+	UE_API virtual bool OnSectionVerticalDrag(const FSectionVerticalDragContext& Context) override;
+	UE_API virtual void OnEndSectionVerticalDrag(const FSectionVerticalDragContext& Context) override;
+
+protected:
+
+	/** Shift all content at StartIndex and below down by 1 row */
+	UE_API virtual void ShiftRowsDownFromIndex(int32 StartIndex, const TSet<UMovieSceneSection*>& ExcludeSections);
+
+private:
 
 	UE_API bool FindBoundObjects(TArray<UObject*>& OutBoundObjects) const;
 
+	int32 GetMaxRowIndexForSectionDrag(const TSet<UMovieSceneSection*>& ExcludeSections) const;
+
 	/** A second children list for the sections inside this track */
 	FViewModelListHead SectionList;
 	FViewModelListHead TopLevelChannelList;
@@ -181,6 +211,9 @@ private:
 	bool bNeedsUpdate;
 
 	int32 PreviousLayoutNumRows = -1;
+
+	/** Tracks if this track was auto-expanded during a section drag operation */
+	bool bAutoExpandedDuringDrag = false;
 };
 
 } // namespace Sequencer
diff --git a/Engine/Source/Editor/Sequencer/Public/MVVM/ViewModels/TrackRowModel.h b/Engine/Source/Editor/Sequencer/Public/MVVM/ViewModels/TrackRowModel.h
--- a/Engine/Source/Editor/Sequencer/Public/MVVM/ViewModels/TrackRowModel.h
+++ b/Engine/Source/Editor/Sequencer/Public/MVVM/ViewModels/TrackRowModel.h
@@ -4,13 +4,16 @@
 
 #include "MVVM/ViewModels/ViewModel.h"
 #include "MVVM/ViewModels/OutlinerItemModel.h"
-#include "MVVM/Extensions/ITrackExtension.h"
+#include "MVVM/Extensions/ISectionOwnerExtension.h"
+#include "MVVM/Extensions/ITopLevelChannelHolderExtension.h"
+#include "MVVM/Extensions/ITrackRowExtension.h"
 #include "MVVM/Extensions/ITrackAreaExtension.h"
 #include "MVVM/Extensions/ILockableExtension.h"
 #include "MVVM/Extensions/IRenameableExtension.h"
 #include "MVVM/Extensions/IResizableExtension.h"
 #include "MVVM/Extensions/IDeletableExtension.h"
 #include "MVVM/Extensions/IConditionableExtension.h"
+#include "MVVM/Extensions/IDraggableOutlinerExtension.h"
 
 #include "UObject/WeakObjectPtr.h"
 #include "UObject/WeakObjectPtrTemplates.h"
@@ -32,11 +35,14 @@ class FTrackRowModel
 	: public FEvaluableOutlinerItemModel
 	, public ITrackAreaExtension
 	, public ILockableExtension
-	, public ITrackExtension
+	, public ITrackRowExtension
+	, public ISectionOwnerExtension
+	, public ITopLevelChannelHolderExtension
 	, public IRenameableExtension
 	, public IResizableExtension
 	, public IDeletableExtension
 	, public IConditionableExtension
+	, public IDraggableOutlinerExtension
 {
 public:
 
@@ -44,18 +50,23 @@ public:
 		, FEvaluableOutlinerItemModel
 		, ITrackAreaExtension
 		, ILockableExtension
-		, ITrackExtension
+		, ITrackRowExtension
+		, ISectionOwnerExtension
+		, ITopLevelChannelHolderExtension
 		, IRenameableExtension
 		, IResizableExtension
 		, IDeletableExtension
-		, IConditionableExtension);
+		, IConditionableExtension
+		, IDraggableOutlinerExtension);
 
 	UE_API explicit FTrackRowModel(UMovieSceneTrack* InTrack, int32 InRowIndex);
 	UE_API ~FTrackRowModel();
 
 	UE_API void Initialize();
 
-	UE_API FViewModelChildren GetTopLevelChannels();
+	/* Refreshes the layout of the row model including all child sections/other outliner items.
+	 * Pass in whether to force children layout (e.g. if this is a new track row) */
+	UE_API virtual void RefreshLayout(bool bChildrenNeedLayout=true);
 
 	/*~ FOutlinerItemModel */
 	UE_API void BuildContextMenu(FMenuBuilder& MenuBuilder) override;
@@ -79,11 +90,16 @@ public:
 	UE_API bool IsResizable() const override;
 	UE_API void Resize(float NewSize) override;
 
-	/*~ ITrackExtension */
-	UE_API UMovieSceneTrack* GetTrack() const override;
+	/*~ ITrackRowExtension */
 	UE_API int32 GetRowIndex() const override;
+	UE_API bool SetRowIndex(int32 NewRowIndex) override;
+	UE_API UMovieSceneTrack* GetParentTrack() const override;
+
+	/*~ ISectionOwnerExtension */
 	UE_API FViewModelChildren GetSectionModels() override;
-	TSharedPtr<ISequencerTrackEditor> GetTrackEditor() const override { return TrackEditor; }
+
+	/*~ ITopLevelChannelHolderExtension */
+	UE_API FViewModelChildren GetTopLevelChannels() override;
 
 	/*~ ITrackAreaExtension */
 	UE_API FTrackAreaParameters GetTrackAreaParameters() const override;
@@ -114,6 +130,13 @@ public:
 	UE_API bool IsSolo() const override;
 	UE_API void SetIsSoloed(bool bIsSoloed) override;
 
+	/*~ IDraggableOutlinerExtension */
+	UE_API virtual bool CanDrag() const override;
+
+	/*~ Deprecated ITrackExtension functions */
+	UMovieSceneTrack* GetTrack() const { return GetParentTrack(); }
+	TSharedPtr<ISequencerTrackEditor> GetTrackEditor() const { return TrackEditor; }
+
 private:
 
 	FViewModelListHead SectionList;
diff --git a/Engine/Source/Editor/Sequencer/Public/SequencerAddKeyOperation.h b/Engine/Source/Editor/Sequencer/Public/SequencerAddKeyOperation.h
--- a/Engine/Source/Editor/Sequencer/Public/SequencerAddKeyOperation.h
+++ b/Engine/Source/Editor/Sequencer/Public/SequencerAddKeyOperation.h
@@ -30,6 +30,7 @@ namespace Sequencer
 class FChannelGroupModel;
 class FViewModel;
 class ITrackExtension;
+class ISectionOwnerExtension;
 class IOutlinerExtension;
 
 /**
@@ -83,7 +84,7 @@ private:
 	 * @param InTrackEditor         The track editor
 	 * @param InKeyAnythingBeneath  A node to search within for key areas
 	 */
-	UE_API bool ConsiderKeyableAreas(TSharedPtr<ITrackExtension> InTrackModel, FViewModelPtr InKeyAnythingBeneath);
+	UE_API bool ConsiderKeyableAreas(TViewModelPtr<ISectionOwnerExtension> InSectionOwner, FViewModelPtr InKeyAnythingBeneath);
 
 
 	/**
@@ -92,7 +93,7 @@ private:
 	 * @param InTrackNode         The current track node
 	 * @param InKeyAreaNode       The key area node to add key areas from
 	 */
-	UE_API bool ProcessKeyArea(TSharedPtr<ITrackExtension> InTrackModel, TViewModelPtr<FChannelGroupModel> InChannelGroupModel);
+	UE_API bool ProcessKeyArea(ISequencerTrackEditor* InTrackEditor, TViewModelPtr<FChannelGroupModel> InChannelGroupModel);
 
 
 	/**
diff --git a/Engine/Source/Editor/Sequencer/Public/SequencerUtilities.h b/Engine/Source/Editor/Sequencer/Public/SequencerUtilities.h
--- a/Engine/Source/Editor/Sequencer/Public/SequencerUtilities.h
+++ b/Engine/Source/Editor/Sequencer/Public/SequencerUtilities.h
@@ -37,6 +37,7 @@ namespace UE::Sequencer
 {
 
 class ITrackExtension;
+class ISectionOwnerExtension;
 
 struct FCreateBindingParams
 {
@@ -238,12 +239,15 @@ struct FSequencerUtilities
 	/* Creates a button (used for +Section) that fires a user-defined OnClick response with no sub-menu. */
 	static UE_API TSharedRef<SWidget> MakeAddButton(FText HoverText, FOnClicked OnClicked, const TAttribute<bool>& HoverState, TWeakPtr<ISequencer> InSequencer);
 
-	static UE_API void MakeTimeWarpMenuEntry(FMenuBuilder& MenuBuilder, UE::Sequencer::TWeakViewModelPtr<UE::Sequencer::ITrackExtension> TrackModel);
+	static UE_API void MakeTimeWarpMenuEntry(FMenuBuilder& MenuBuilder, UE::Sequencer::TWeakViewModelPtr<UE::Sequencer::FViewModel> WeakViewModel);
 
 	static UE_API void PopulateTimeWarpSubMenu(FMenuBuilder& MenuBuilder, TFunction<void(TSubclassOf<UMovieSceneTimeWarpGetter>)> OnTimeWarpPicked);
-
+	
+	UE_DEPRECATED(5.8, "Please use the version that takes an ISectionOwnerExtension")
 	static UE_API void PopulateTimeWarpChannelSubMenu(FMenuBuilder& MenuBuilder, UE::Sequencer::TWeakViewModelPtr<UE::Sequencer::ITrackExtension> TrackModel);
 
+	static UE_API void PopulateTimeWarpChannelSubMenu(FMenuBuilder& MenuBuilder, UE::Sequencer::TWeakViewModelPtr<UE::Sequencer::ISectionOwnerExtension> SectionOwner);
+
 	static UE_API void CreateNewSection(UMovieSceneTrack* InTrack, TWeakPtr<ISequencer> InSequencer, int32 InRowIndex, EMovieSceneBlendType InBlendType);
 
 	static UE_API void PopulateMenu_CreateNewSection(FMenuBuilder& MenuBuilder, int32 RowIndex, UMovieSceneTrack* Track, TWeakPtr<ISequencer> InSequencer);
diff --git a/Engine/Source/Runtime/MovieScene/Private/EntitySystem/MovieSceneEntityLedger.cpp b/Engine/Source/Runtime/MovieScene/Private/EntitySystem/MovieSceneEntityLedger.cpp
--- a/Engine/Source/Runtime/MovieScene/Private/EntitySystem/MovieSceneEntityLedger.cpp
+++ b/Engine/Source/Runtime/MovieScene/Private/EntitySystem/MovieSceneEntityLedger.cpp
@@ -11,6 +11,7 @@
 #include "Conditions/MovieSceneCondition.h"
 
 #include "MovieSceneSection.h"
+#include "MovieSceneTrack.h"
 
 namespace UE
 {
@@ -142,6 +143,17 @@ void FEntityLedger::UpdateOneShotEntities(UMovieSceneEntitySystemLinker* Linker,
 			if (UMovieSceneSection* Section = Cast<UMovieSceneSection>(EntityOwner))
 			{
 				Section->BuildDefaultComponents(Linker, Params, &ImportedEntity);
+
+				if (UMovieSceneTrack* Track = Section->GetTypedOuter<UMovieSceneTrack>())
+				{
+					for (UObject* Decoration : Track->GetDecorations())
+					{
+						if (IMovieSceneEntityDecorator* EntityDecorator = Cast<IMovieSceneEntityDecorator>(Decoration))
+						{
+							EntityDecorator->ExtendEntity(Linker, Params, &ImportedEntity);
+						}
+					}
+				}
 			}
 
 			FMovieSceneEntityID NewEntityID = ImportedEntity.Manufacture(Params, &Linker->EntityManager);
@@ -206,6 +218,17 @@ void FEntityLedger::UpdateConditionalEntities(UMovieSceneEntitySystemLinker* Lin
 				if (UMovieSceneSection* Section = Cast<UMovieSceneSection>(EntityOwner))
 				{
 					Section->BuildDefaultComponents(Linker, Params, &ImportedEntity);
+
+					if (UMovieSceneTrack* Track = Section->GetTypedOuter<UMovieSceneTrack>())
+					{
+						for (UObject* Decoration : Track->GetDecorations())
+						{
+							if (IMovieSceneEntityDecorator* EntityDecorator = Cast<IMovieSceneEntityDecorator>(Decoration))
+							{
+								EntityDecorator->ExtendEntity(Linker, Params, &ImportedEntity);
+							}
+						}
+					}
 				}
 
 				FMovieSceneEntityID NewEntityID = ImportedEntity.Manufacture(Params, &Linker->EntityManager);
@@ -394,6 +417,17 @@ void FEntityLedger::ImportEntity(UMovieSceneEntitySystemLinker* Linker, const FE
 		if (UMovieSceneSection* Section = Cast<UMovieSceneSection>(EntityOwner))
 		{
 			Section->BuildDefaultComponents(Linker, Params, &ImportedEntity);
+
+			if (UMovieSceneTrack* Track = Section->GetTypedOuter<UMovieSceneTrack>())
+			{
+				for (UObject* Decoration : Track->GetDecorations())
+				{
+					if (IMovieSceneEntityDecorator* EntityDecorator = Cast<IMovieSceneEntityDecorator>(Decoration))
+					{
+						EntityDecorator->ExtendEntity(Linker, Params, &ImportedEntity);
+					}
+				}
+			}
 		}
 
 		FMovieSceneEntityID NewEntityID = ImportedEntity.Manufacture(Params, &Linker->EntityManager);
diff --git a/Engine/Source/Runtime/MovieScene/Public/MovieSceneTrack.h b/Engine/Source/Runtime/MovieScene/Public/MovieSceneTrack.h
--- a/Engine/Source/Runtime/MovieScene/Public/MovieSceneTrack.h
+++ b/Engine/Source/Runtime/MovieScene/Public/MovieSceneTrack.h
@@ -223,6 +223,10 @@ public:
 	FMovieSceneTrackDisplayOptions DisplayOptions;
 #endif
 
+	/* Parent track for this track. If a parent track is set, the track won't show up under the object binding. */
+	UPROPERTY()
+	TObjectPtr<UMovieSceneTrack> ParentTrack = nullptr;
+
 	/** Optional dynamic condition for whether this track/any of the sections on this track evaluates at runtime. */
 	UPROPERTY(EditAnywhere, Category = "Track")
 	FMovieSceneConditionContainer ConditionContainer;
diff --git a/Engine/Source/Runtime/MovieSceneTracks/MovieSceneTracks.Build.cs b/Engine/Source/Runtime/MovieSceneTracks/MovieSceneTracks.Build.cs
--- a/Engine/Source/Runtime/MovieSceneTracks/MovieSceneTracks.Build.cs
+++ b/Engine/Source/Runtime/MovieSceneTracks/MovieSceneTracks.Build.cs
@@ -15,12 +15,12 @@ public class MovieSceneTracks : ModuleRules
 				"TimeManagement",
 				"AnimationCore",
 				"Constraints",
-				"AudioMixer"
+				"AudioMixer",
+				"SlateCore" 
 			});
 
 		PrivateDependencyModuleNames.AddRange(new string[]
 			{
-				"SlateCore",
 				"AnimGraphRuntime",
 				"PropertyPath"
 			});
diff --git a/Engine/Source/Runtime/MovieSceneTracks/Private/Tracks/MovieSceneCommonAnimationTrack.cpp b/Engine/Source/Runtime/MovieSceneTracks/Private/Tracks/MovieSceneCommonAnimationTrack.cpp
--- a/Engine/Source/Runtime/MovieSceneTracks/Private/Tracks/MovieSceneCommonAnimationTrack.cpp
+++ b/Engine/Source/Runtime/MovieSceneTracks/Private/Tracks/MovieSceneCommonAnimationTrack.cpp
@@ -95,6 +95,12 @@ bool UMovieSceneCommonAnimationTrack::IsEmpty() const
 
 UMovieSceneSection* UMovieSceneCommonAnimationTrack::AddNewAnimationOnRow(FFrameNumber KeyTime, UAnimSequenceBase* AnimSequence, int32 RowIndex)
 {
+	// If RowIndex is INDEX_NONE, compute the next available row
+	if (RowIndex == INDEX_NONE)
+	{
+		RowIndex = ComputeNextAvailableRowIndex();
+	}
+
 	UMovieSceneSkeletalAnimationSection* NewSection = Cast<UMovieSceneSkeletalAnimationSection>(CreateNewSection());
 	{
 		FFrameTime AnimationLength = AnimSequence->GetPlayLength() * GetTypedOuter<UMovieScene>()->GetTickResolution();
@@ -116,6 +122,29 @@ UMovieSceneSection* UMovieSceneCommonAnimationTrack::AddNewAnimationOnRow(FFrame
 	return NewSection;
 }
 
+void UMovieSceneCommonAnimationTrack::AddSectionWithAutoRow(UMovieSceneSection* Section)
+{
+	if (!Section)
+	{
+		return;
+	}
+
+	// Base implementation: just add the section
+	// The section's row index should already be set, or will use default behavior
+	AddSection(*Section);
+}
+
+int32 UMovieSceneCommonAnimationTrack::ComputeNextAvailableRowIndex() const
+{
+	// Find max row from existing sections
+	int32 MaxRowIndex = -1;
+	for (UMovieSceneSection* ExistingSection : AnimationSections)
+	{
+		MaxRowIndex = FMath::Max(MaxRowIndex, ExistingSection->GetRowIndex());
+	}
+	return MaxRowIndex + 1;
+}
+
 #if WITH_EDITOR
 void UMovieSceneCommonAnimationTrack::PostEditImport()
 {
diff --git a/Engine/Source/Runtime/MovieSceneTracks/Public/Tracks/MovieSceneCommonAnimationTrack.h b/Engine/Source/Runtime/MovieSceneTracks/Public/Tracks/MovieSceneCommonAnimationTrack.h
--- a/Engine/Source/Runtime/MovieSceneTracks/Public/Tracks/MovieSceneCommonAnimationTrack.h
+++ b/Engine/Source/Runtime/MovieSceneTracks/Public/Tracks/MovieSceneCommonAnimationTrack.h
@@ -53,6 +53,22 @@ public:
 	/** Adds a new animation to this track on the next available/non-overlapping row */
 	 virtual UMovieSceneSection* AddNewAnimation(FFrameNumber KeyTime, UAnimSequenceBase* AnimSequence) { return AddNewAnimationOnRow(KeyTime, AnimSequence, INDEX_NONE); }
 
+	/** Adds a section to the track with automatically computed row index. For mixer tracks, computes row based on existing sections and child tracks. */
+	UE_API virtual void AddSectionWithAutoRow(UMovieSceneSection* Section);
+
+	/** Computes the next available row index. For mixer tracks, considers both sections and child tracks. */
+	UE_API virtual int32 ComputeNextAvailableRowIndex() const;
+
+	/** Adds a child track of the specified class at the given row index. Only valid for mixer tracks. */
+	virtual UMovieSceneTrack* AddChildTrack(FGuid ObjectBinding, TSubclassOf<UMovieSceneTrack> TrackClass, int32 RowIndex)
+	{
+		checkf(false, TEXT("AddChildTrack is only valid for UMovieSceneAnimationMixerTrack"));
+		return nullptr;
+	}
+
+	/** Gets all child tracks. Only valid for mixer tracks. Base implementation returns empty array. */
+	virtual void GetAllChildTracks(TArray<TObjectPtr<UMovieSceneTrack>>& OutChildTracks) const {}
+
 	UE_API virtual const TArray<UMovieSceneSection*>& GetAllSections() const override;
 	UE_API virtual UMovieSceneSection* CreateNewSection() override;
 	UE_API virtual bool SupportsType(TSubclassOf<UMovieSceneSection> SectionClass) const override;
