Sequencer: Fix static analyzer error.

#Synced-CL 49029986

diff --git a/Engine/Source/Editor/Sequencer/Private/MVVM/ViewModels/SectionOutlinerModel.cpp b/Engine/Source/Editor/Sequencer/Private/MVVM/ViewModels/SectionOutlinerModel.cpp
--- a/Engine/Source/Editor/Sequencer/Private/MVVM/ViewModels/SectionOutlinerModel.cpp
+++ b/Engine/Source/Editor/Sequencer/Private/MVVM/ViewModels/SectionOutlinerModel.cpp
@@ -29,7 +29,7 @@ namespace Sequencer
 namespace
 {
 	// Helper function to check if a view model has any keyable areas
-	bool HasKeyableAreas(const FViewModel& ViewModel)
+	bool ViewModelHasKeyableAreas(const FViewModel& ViewModel)
 	{
 		for (TSharedPtr<FChannelGroupModel> ChannelGroup : ViewModel.GetDescendantsOfType<FChannelGroupModel>())
 		{
@@ -216,7 +216,7 @@ TSharedPtr<SWidget> FSectionOutlinerModel::CreateOutlinerViewForColumn(const FCr
 	// KeyFrame column - show add key button if there are keyable channels
 	if (InColumnName == FCommonOutlinerNames::KeyFrame)
 	{
-		if (HasKeyableAreas(*this))
+		if (ViewModelHasKeyableAreas(*this))
 		{
 			EKeyNavigationButtons Buttons = EKeyNavigationButtons::AddKey;
 			return SNew(SSequencerKeyNavigationButtons, SharedThis(this), Editor->GetSequencer())
@@ -227,7 +227,7 @@ TSharedPtr<SWidget> FSectionOutlinerModel::CreateOutlinerViewForColumn(const FCr
 	// Nav column - show navigation buttons if there are keys
 	if (InColumnName == FCommonOutlinerNames::Nav)
 	{
-		if (HasKeyableAreas(*this))
+		if (ViewModelHasKeyableAreas(*this))
 		{
 			EKeyNavigationButtons Buttons = InParams.TreeViewRow->IsColumnVisible(FCommonOutlinerNames::KeyFrame)
 				? EKeyNavigationButtons::NavOnly
