Sequencer: Modify several Sequencer utility functions to use the UWorld from the Sequencer's playback context, rather than the GCurrentLevelEditingViewportClient world.
#rnx
#Synced-CL 49044222

diff --git a/Engine/Source/Editor/Sequencer/Private/Sequencer.cpp b/Engine/Source/Editor/Sequencer/Private/Sequencer.cpp
--- a/Engine/Source/Editor/Sequencer/Private/Sequencer.cpp
+++ b/Engine/Source/Editor/Sequencer/Private/Sequencer.cpp
@@ -8677,7 +8677,8 @@ FSequencer::EPasteResult FSequencer::PasteObjectBindings(const FString& TextToIm
 		}
 	}
 
-	UWorld* World = GCurrentLevelEditingViewportClient ? GCurrentLevelEditingViewportClient->GetWorld() : nullptr;
+	UObject* PlaybackContext = GetPlaybackContext();
+	UWorld* World = PlaybackContext ? PlaybackContext->GetWorld() : nullptr;
 
 	bool bDuplicateExistingActors = false;
 	TArray<FString> ObjectNames = FSequencerUtilities::GetPasteBindingsObjectNames(AsShared(), TextToImport);
diff --git a/Engine/Source/Editor/Sequencer/Private/SequencerUtilities.cpp b/Engine/Source/Editor/Sequencer/Private/SequencerUtilities.cpp
--- a/Engine/Source/Editor/Sequencer/Private/SequencerUtilities.cpp
+++ b/Engine/Source/Editor/Sequencer/Private/SequencerUtilities.cpp
@@ -971,7 +971,8 @@ FGuid FSequencerUtilities::CreateCamera(TSharedRef<ISequencer> Sequencer, const
 		return CameraGuid;
 	}
 
-	UWorld* World = GCurrentLevelEditingViewportClient ? GCurrentLevelEditingViewportClient->GetWorld() : nullptr;
+	UObject* PlaybackContext = Sequencer->GetPlaybackContext();
+	UWorld* World = PlaybackContext ? PlaybackContext->GetWorld() : nullptr;
 	if (!World)
 	{
 		return CameraGuid;
@@ -1076,7 +1077,13 @@ FGuid FSequencerUtilities::CreateCameraWithRig(TSharedRef<ISequencer> Sequencer,
 	}
 
 	// Create a cine camera actor
-	UWorld* World = GCurrentLevelEditingViewportClient ? GCurrentLevelEditingViewportClient->GetWorld() : nullptr;
+	UObject* PlaybackContext = Sequencer->GetPlaybackContext();
+	UWorld* World = PlaybackContext ? PlaybackContext->GetWorld() : nullptr;
+	if (!World)
+	{
+		return CameraGuid;
+	}
+
 	OutActor = World->SpawnActor<ACineCameraActor>();
 
 	FString NewCameraName = MovieSceneHelpers::MakeUniqueSpawnableName(MovieScene, FName::NameToDisplayString(ACineCameraActor::StaticClass()->GetFName().ToString(), false));
@@ -1408,40 +1415,44 @@ FMovieScenePossessable* FSequencerUtilities::ConvertToPossessable(TSharedRef<ISe
 		SpawnInfo.bDeferConstruction = true;
 		SpawnInfo.Template = SpawnableActorTemplate;
 
-		UWorld* World = GCurrentLevelEditingViewportClient ? GCurrentLevelEditingViewportClient->GetWorld() : nullptr;
-		AActor* PossessedActor = World->SpawnActor(Actor->GetClass(), &DefaultTransform, SpawnInfo);
-
-		if (!PossessedActor)
+		UObject* PlaybackContext = Sequencer->GetPlaybackContext();
+		UWorld* World = PlaybackContext ? PlaybackContext->GetWorld() : nullptr;
+		if (World)
 		{
-			return nullptr;
-		}
+			AActor* PossessedActor = World->SpawnActor(Actor->GetClass(), &DefaultTransform, SpawnInfo);
 
-		FString ActorLabel = Actor->GetActorLabel();
-		if (Spawnable)
-		{
-			ActorLabel = Spawnable->GetName();
-		}
-		else if (ExistingPossessable)
-		{
-			if (BindingReferences)
+			if (!PossessedActor)
 			{
-				// If we don't have multiple bound objects, use the Possessable name instead of the template label
-				if (BindingReferences->GetReferences(BindingGuid).Num() == 1)
+				return nullptr;
+			}
+
+			FString ActorLabel = Actor->GetActorLabel();
+			if (Spawnable)
+			{
+				ActorLabel = Spawnable->GetName();
+			}
+			else if (ExistingPossessable)
+			{
+				if (BindingReferences)
 				{
-					ActorLabel = ExistingPossessable->GetName();
+					// If we don't have multiple bound objects, use the Possessable name instead of the template label
+					if (BindingReferences->GetReferences(BindingGuid).Num() == 1)
+					{
+						ActorLabel = ExistingPossessable->GetName();
+					}
 				}
 			}
-		}
 
-		PossessedActor->SetActorLabel(ActorLabel);
+			PossessedActor->SetActorLabel(ActorLabel);
 
-		const bool bIsDefaultTransform = true;
-		PossessedActor->FinishSpawning(DefaultTransform, bIsDefaultTransform);
+			const bool bIsDefaultTransform = true;
+			PossessedActor->FinishSpawning(DefaultTransform, bIsDefaultTransform);
 
-		// The transform needs to be set again for deferred construction and dynamic root components. Until the fix for: UE-67537
-		PossessedActor->SetActorTransform(DefaultTransform);
+			// The transform needs to be set again for deferred construction and dynamic root components. Until the fix for: UE-67537
+			PossessedActor->SetActorTransform(DefaultTransform);
 
-		PossessedObject = PossessedActor;
+			PossessedObject = PossessedActor;
+		}
 	}
 
 	Sequence->Modify();
@@ -2569,7 +2580,8 @@ void FSequencerUtilities::CopyBindings(TSharedRef<ISequencer> Sequencer, const T
 
 void FSequencerUtilities::CopyBindings(TSharedRef<ISequencer> Sequencer, const TArray<FMovieSceneBindingProxy>& Bindings, const TArray<UMovieSceneFolder*>& InFolders, FOutputDevice& Ar)
 {
-	UWorld* World = GCurrentLevelEditingViewportClient ? GCurrentLevelEditingViewportClient->GetWorld() : nullptr;
+	UObject* PlaybackContext = Sequencer->GetPlaybackContext();
+	UWorld* World = PlaybackContext ? PlaybackContext->GetWorld() : nullptr;
 
 	TArray<UMovieSceneCopyableBinding*> Objects;
 	for (const FMovieSceneBindingProxy& ObjectBinding : Bindings)
@@ -2959,8 +2971,10 @@ FGuid CreateGenericBinding(TSharedPtr<ISequencer> Sequencer, UMovieSceneSequence
 	if (bAllowCustom)
 	{
 		// In addition to the raw object, we also try spawning an actor from an actor factory if relevant, to give the custom binding an option to create from that as well
-		AActor* FactoryCreatedActorInstance = nullptr; 
-		UWorld* World = GCurrentLevelEditingViewportClient ? GCurrentLevelEditingViewportClient->GetWorld() : nullptr;
+		AActor* FactoryCreatedActorInstance = nullptr;
+
+		UObject* PlaybackContext = Sequencer->GetPlaybackContext();
+		UWorld* World = PlaybackContext ? PlaybackContext->GetWorld() : nullptr;
 		if (InObject && !InObject->IsA<AActor>())
 		{
 			// Workaround for a bug in UActorFactoryBlueprint- the actor factory will claim it can create an actor for a blueprint generated class, but then fail to do so
@@ -2998,7 +3012,7 @@ FGuid CreateGenericBinding(TSharedPtr<ISequencer> Sequencer, UMovieSceneSequence
 			}
 		}
 		NewBindingID = TryCreateCustomBinding(Sequencer, InObject, FactoryCreatedActorInstance, BindingReferences, InParams, OwnerMovieScene, bSpawnable, bReplaceable);
-		if (FactoryCreatedActorInstance)
+		if (FactoryCreatedActorInstance && World)
 		{
 			const bool bNetForce = false;
 			const bool bShouldModifyLevel = false;
@@ -3199,7 +3213,8 @@ bool FSequencerUtilities::PasteBindings(const FString& TextToImport, TSharedRef<
 
 	UMovieScene* RootMovieScene = Sequencer->GetRootMovieSceneSequence()->GetMovieScene();
 
-	UWorld* World = GCurrentLevelEditingViewportClient ? GCurrentLevelEditingViewportClient->GetWorld() : nullptr;
+	UObject* PlaybackContext = Sequencer->GetPlaybackContext();
+	UWorld* World = PlaybackContext ? PlaybackContext->GetWorld() : nullptr;
 
 	const FScopedTransaction Transaction(LOCTEXT("PasteBindings", "Paste Bindings"));
 
