Anim In Engine:  When auto scrolling we get default hierarchy order, then focus on the first  non-channel item.

#Synced-CL 49039648

diff --git a/Engine/Source/Editor/Sequencer/Private/MVVM/CurveEditorExtension.cpp b/Engine/Source/Editor/Sequencer/Private/MVVM/CurveEditorExtension.cpp
--- a/Engine/Source/Editor/Sequencer/Private/MVVM/CurveEditorExtension.cpp
+++ b/Engine/Source/Editor/Sequencer/Private/MVVM/CurveEditorExtension.cpp
@@ -8,7 +8,10 @@
 #include "Framework/Docking/TabManager.h"
 #include "IPropertyRowGenerator.h"
 #include "IStructureDetailsView.h"
+#include "MVVM/ViewModels/ViewModelIterators.h"
 #include "MVVM/ViewModels/SequencerEditorViewModel.h"
+#include "MVVM/ViewModels/ChannelModel.h"
+#include "MVVM/Extensions/IOutlinerExtension.h"
 #include "CurveEditorAxis.h"
 #include "ICurveEditorExtension.h"
 #include "SCurveEditorView.h"
@@ -704,7 +707,23 @@ void FCurveEditorExtension::RequestSyncSelection()
 		Transaction.Reset();
 	});
 }
+static void GetAllChildrenRecursively(TSharedPtr<SCurveEditorTree>& CurveEditorTreeView, const FCurveEditorTreeItemID Item, int32& LastAddedItem,
+	TMap<const FCurveEditorTreeItemID, int32>& ItemToSortOrder)
+{
+	ItemToSortOrder.Add(Item, LastAddedItem++);
+
+	TArray<FCurveEditorTreeItemID> Children;
+
+	CurveEditorTreeView->GetTreeItemChildren(Item, Children);
+	// Add them to the output array
+	for (const FCurveEditorTreeItemID ChildItemID : Children)
+	{
+		GetAllChildrenRecursively(CurveEditorTreeView, ChildItemID, LastAddedItem, ItemToSortOrder);
+	}
+}
 
+//we now sync to the first non-channel in the tree hierarchy, this is needed for hosted channels with control rigs
+//we need to sort the hierarchy each time since the curve id's being used previously were not consistent in order of viewing/creationg, persistent now for undo/redo
 void FCurveEditorExtension::SyncSelection()
 {
 	if (!ensure(CurveEditorModel && CurveEditorTreeView))
@@ -728,7 +747,31 @@ void FCurveEditorExtension::SyncSelection()
 
 	CurveEditorTreeView->ClearSelection();
 
+	//map of tree items id's to their actually sort order, need to use this since tree items ids are not in order anymore
+	TMap<const FCurveEditorTreeItemID, int32>  ItemToSortOrder;
+	int32 LastAddedItem = 0;
+	TArray<FCurveEditorTreeItemID> Children;
+	TArrayView<const FCurveEditorTreeItemID> RootItems = CurveEditorTreeView->GetRootItems();
+
+
+	for (const FCurveEditorTreeItemID RootItemID : RootItems)
+	{
+		GetAllChildrenRecursively(CurveEditorTreeView, RootItemID, LastAddedItem, ItemToSortOrder);
+	}
+
+	auto GetSortOrder = [&ItemToSortOrder](FCurveEditorTreeItemID InTreeItemID) -> int32
+	{
+		if (const int32* SortOrder = ItemToSortOrder.Find(InTreeItemID))
+		{
+			return *SortOrder;
+		}
+		return INDEX_NONE;
+	};
+
 	FCurveEditorTreeItemID FirstCurveEditorTreeItemID;
+	int32 FirstSortOrder = INDEX_NONE;
+	bool bSelectedIsChannel = true;
+
 	for (TViewModelPtr<IOutlinerExtension> SelectedItem : OwnerModel->GetSelection()->Outliner)
 	{
 		if (TViewModelPtr<ICurveEditorTreeItemExtension> CurveEditorItem = SelectedItem.ImplicitCast())
@@ -739,15 +782,39 @@ void FCurveEditorExtension::SyncSelection()
 				if (!CurveEditorTreeView->IsItemSelected(CurveEditorTreeItem))
 				{
 					CurveEditorTreeView->SetItemSelection(CurveEditorTreeItem, true);
+					bool bIsChannel = SelectedItem.AsModel()->IsA<FChannelModel>() || SelectedItem.AsModel()->IsA<FChannelGroupOutlinerModel>()
+					|| SelectedItem.AsModel()->IsA<FChannelGroupModel>();
 					if (!FirstCurveEditorTreeItemID.IsValid())
 					{
 						FirstCurveEditorTreeItemID = CurveEditorTreeItem;
+						FirstSortOrder = GetSortOrder(CurveEditorTreeItem);
+						bSelectedIsChannel = bIsChannel;
 					}
 					else
 					{
-						if (CurveEditorTreeItem.GetValue() < FirstCurveEditorTreeItemID.GetValue())
+						int32 SortOrder = GetSortOrder(CurveEditorTreeItem);
+						if (bIsChannel == false)
+						{
+							if (bSelectedIsChannel == true)
+							{
+								FirstCurveEditorTreeItemID = CurveEditorTreeItem;
+								FirstSortOrder = SortOrder;
+								bSelectedIsChannel = false;
+							}
+							else if (SortOrder < FirstSortOrder)
+							{
+								FirstCurveEditorTreeItemID = CurveEditorTreeItem;
+								FirstSortOrder = SortOrder;
+								bSelectedIsChannel = false;
+							}
+						}
+						else if (bSelectedIsChannel == true) //its a channel only sort if selected is a non channel
 						{
-							FirstCurveEditorTreeItemID = CurveEditorTreeItem;
+							if (SortOrder < FirstSortOrder)
+							{
+								FirstCurveEditorTreeItemID = CurveEditorTreeItem;
+								FirstSortOrder = SortOrder;
+							}
 						}
 					}
 				}
