Sequencer: Fix CreateGenericBinding function to properly check the Sequencer pointer before attempting to acquire a UWorld from its playback context.
#rnx
#Synced-CL 49044222

diff --git a/Engine/Source/Editor/Sequencer/Private/SequencerUtilities.cpp b/Engine/Source/Editor/Sequencer/Private/SequencerUtilities.cpp
--- a/Engine/Source/Editor/Sequencer/Private/SequencerUtilities.cpp
+++ b/Engine/Source/Editor/Sequencer/Private/SequencerUtilities.cpp
@@ -2973,8 +2973,17 @@ FGuid CreateGenericBinding(TSharedPtr<ISequencer> Sequencer, UMovieSceneSequence
 		// In addition to the raw object, we also try spawning an actor from an actor factory if relevant, to give the custom binding an option to create from that as well
 		AActor* FactoryCreatedActorInstance = nullptr;
 
-		UObject* PlaybackContext = Sequencer->GetPlaybackContext();
-		UWorld* World = PlaybackContext ? PlaybackContext->GetWorld() : nullptr;
+		UWorld* World = nullptr;
+		if (Sequencer)
+		{
+			UObject* PlaybackContext = Sequencer->GetPlaybackContext();
+			World = PlaybackContext ? PlaybackContext->GetWorld() : nullptr;
+		}
+		else
+		{
+			World = GCurrentLevelEditingViewportClient ? GCurrentLevelEditingViewportClient->GetWorld() : nullptr;
+		}
+
 		if (InObject && !InObject->IsA<AActor>())
 		{
 			// Workaround for a bug in UActorFactoryBlueprint- the actor factory will claim it can create an actor for a blueprint generated class, but then fail to do so
