[Sequencer] Show filter bar by default (only when there are active filters)
#Synced-CL 49128623

diff --git a/Engine/Source/Editor/Sequencer/Private/SequencerSettings.cpp b/Engine/Source/Editor/Sequencer/Private/SequencerSettings.cpp
--- a/Engine/Source/Editor/Sequencer/Private/SequencerSettings.cpp
+++ b/Engine/Source/Editor/Sequencer/Private/SequencerSettings.cpp
@@ -80,6 +80,7 @@ USequencerSettings::USequencerSettings( const FObjectInitializer& ObjectInitiali
 	ViewDensity = "Relaxed";
 	AssetBrowserWidth = 500.f;
 	AssetBrowserHeight = 300.f;
+	bFilterBarVisible = true;
 	bNavigationToolVisible = false;
 
 	SectionColorTints.Add(FColor(88, 102, 142, 255)); // blue
