[Sequencer] Allow movement of keys while "Keyed" filter is enabled and stop keys from disappearing

#Synced-CL 49128623

diff --git a/Engine/Source/Editor/Sequencer/Private/Filters/SequencerFilterBar.cpp b/Engine/Source/Editor/Sequencer/Private/Filters/SequencerFilterBar.cpp
--- a/Engine/Source/Editor/Sequencer/Private/Filters/SequencerFilterBar.cpp
+++ b/Engine/Source/Editor/Sequencer/Private/Filters/SequencerFilterBar.cpp
@@ -1519,7 +1519,10 @@ bool FSequencerFilterBar::FilterNodesRecursive(const bool bInHasActiveFilter, co
 
 	if (bAllFiltersPassed || bAnyChildPassed)
 	{
-		if (SequencerSettings->GetAutoExpandNodesOnFilterPass())
+		// Don't expand if a mouse capture is effecting TrackValueChanged events to avoid keys
+		// being hidden and expanding tracks while dragging
+		if (SequencerSettings->GetAutoExpandNodesOnFilterPass()
+			&& !FSlateApplication::Get().HasAnyMouseCaptor())
 		{
 			SetTrackParentsExpanded(InStartNode.ImplicitCast(), true);
 		}
