[Sequencer] Add FSequencerFilterBase::OnPreFilter() that fires before the items are recursively checked for passed filters to allow filters to preprocess
#Synced-CL 49128623

diff --git a/Engine/Source/Editor/Sequencer/Private/Filters/Filters/SequencerTrackFilter_Group.cpp b/Engine/Source/Editor/Sequencer/Private/Filters/Filters/SequencerTrackFilter_Group.cpp
--- a/Engine/Source/Editor/Sequencer/Private/Filters/Filters/SequencerTrackFilter_Group.cpp
+++ b/Engine/Source/Editor/Sequencer/Private/Filters/Filters/SequencerTrackFilter_Group.cpp
@@ -34,6 +34,14 @@ TSharedPtr<FUICommandInfo> FSequencerTrackFilter_Group::GetToggleCommand() const
 	return FSequencerTrackFilterCommands::Get().ToggleFilter_Groups;
 }
 
+void FSequencerTrackFilter_Group::OnPreFilter()
+{
+	if (const UMovieSceneSequence* const FocusedMovieSceneSequence = GetSequencer().GetFocusedMovieSceneSequence())
+	{
+		UpdateMovieScene(FocusedMovieSceneSequence->GetMovieScene());
+	}
+}
+
 FText FSequencerTrackFilter_Group::GetDisplayName() const
 {
 	return LOCTEXT("SequencerTrackFilter_Group", "Groups");
diff --git a/Engine/Source/Editor/Sequencer/Private/Filters/Filters/SequencerTrackFilter_Level.cpp b/Engine/Source/Editor/Sequencer/Private/Filters/Filters/SequencerTrackFilter_Level.cpp
--- a/Engine/Source/Editor/Sequencer/Private/Filters/Filters/SequencerTrackFilter_Level.cpp
+++ b/Engine/Source/Editor/Sequencer/Private/Filters/Filters/SequencerTrackFilter_Level.cpp
@@ -37,6 +37,13 @@ TSharedPtr<FUICommandInfo> FSequencerTrackFilter_Level::GetToggleCommand() const
 	return FSequencerTrackFilterCommands::Get().ToggleFilter_Levels;
 }
 
+void FSequencerTrackFilter_Level::OnPreFilter()
+{
+	const UObject* const PlaybackContext = GetSequencer().GetPlaybackContext();
+	UWorld* const World = PlaybackContext ? PlaybackContext->GetWorld() : nullptr;
+	UpdateWorld(World);
+}
+
 FText FSequencerTrackFilter_Level::GetDisplayName() const
 {
 	return LOCTEXT("SequencerTrackFilter_Level", "Levels");
diff --git a/Engine/Source/Editor/Sequencer/Private/Filters/SequencerFilterBar.cpp b/Engine/Source/Editor/Sequencer/Private/Filters/SequencerFilterBar.cpp
--- a/Engine/Source/Editor/Sequencer/Private/Filters/SequencerFilterBar.cpp
+++ b/Engine/Source/Editor/Sequencer/Private/Filters/SequencerFilterBar.cpp
@@ -93,9 +93,6 @@ FSequencerFilterBar::~FSequencerFilterBar()
     LevelFilter->OnChanged().RemoveAll(this);
     HideIsolateFilter->OnChanged().RemoveAll(this);
     SelectedFilter->OnChanged().RemoveAll(this);
-
-	CommonFilters.Reset();
-	InternalFilters.Reset();
 }
 
 TSharedPtr<ICustomTextFilter<FSequencerTrackFilterType>> FSequencerFilterBar::CreateTextFilter()
@@ -1398,31 +1395,47 @@ UWorld* FSequencerFilterBar::GetWorld() const
 	return nullptr;
 }
 
+void FSequencerFilterBar::PreFilter() const
+{
+	auto PreFilter = [this](const TSharedRef<FSequencerTrackFilter>& InFilter)
+		{
+			if (IsFilterActive(InFilter))
+			{
+				InFilter->OnPreFilter();
+			}
+		};
+
+	auto PreFilterCollection = [PreFilter](const TSharedRef<FSequencerTrackFilterCollection>& InCollection)
+		{
+			InCollection->ForEachFilter([PreFilter]
+				(const TSharedRef<FSequencerTrackFilter>& InFilter)
+				{
+					PreFilter(InFilter);
+					return true;
+				});
+		};
+
+	PreFilterCollection(CommonFilters);
+	PreFilterCollection(InternalFilters);
+
+	PreFilter(TextFilter);
+	PreFilter(HideIsolateFilter);
+	PreFilter(LevelFilter);
+	PreFilter(GroupFilter);
+	PreFilter(SelectedFilter);
+	PreFilter(ModifiedFilter);
+}
+
 const FSequencerFilterData& FSequencerFilterBar::FilterNodes()
 {
 	//SCOPE_LOG_TIME_IN_SECONDS(TEXT("FSequencerFilterBar::FilterNodes()"), nullptr);
 
-	// Update the world for the level filter
-	const UObject* const PlaybackContext = Sequencer.GetPlaybackContext();
-	UWorld* const World = PlaybackContext ? PlaybackContext->GetWorld() : nullptr;
-	LevelFilter->UpdateWorld(World);
-
-	// Update the group filter
-	if (const UMovieSceneSequence* const FocusedMovieSceneSequence = Sequencer.GetFocusedMovieSceneSequence())
-	{
-		GroupFilter->UpdateMovieScene(FocusedMovieSceneSequence->GetMovieScene());
-	}
+	// Allow all filters to do any pre-processing before filtering
+	PreFilter();
 
 	// Reset all filter data
 	FilterData.Reset();
 
-	// Always include the bottom spacer
-	TViewModelPtr<IOutlinerExtension> SpacerNode;
-	if (const FSequenceModel* const SequenceModel = Sequencer.GetNodeTree()->GetRootNode()->CastThis<FSequenceModel>())
-	{
-		SpacerNode = CastViewModelChecked<IOutlinerExtension>(SequenceModel->GetBottomSpacer());
-	}
-
 	// Loop through all nodes and filter recursively
 	const bool bHasActiveFilter = HasAnyFilterActive();
 	for (const TViewModelPtr<IOutlinerExtension>& RootNode : Sequencer.GetNodeTree()->GetRootNodes())
@@ -1431,7 +1444,13 @@ const FSequencerFilterData& FSequencerFilterBar::FilterNodes()
 	}
 
 	// Always filter in spacer node
-	SpacerNode->SetFilteredOut(false);
+	if (const FSequenceModel* const SequenceModel = Sequencer.GetNodeTree()->GetRootNode()->CastThis<FSequenceModel>())
+	{
+		if (const TViewModelPtr<IOutlinerExtension> SpacerNode = CastViewModelChecked<IOutlinerExtension>(SequenceModel->GetBottomSpacer()))
+		{
+			SpacerNode->SetFilteredOut(false);
+		}
+	}
 
 	return FilterData;
 }
diff --git a/Engine/Source/Editor/Sequencer/Private/Filters/SequencerFilterBar.h b/Engine/Source/Editor/Sequencer/Private/Filters/SequencerFilterBar.h
--- a/Engine/Source/Editor/Sequencer/Private/Filters/SequencerFilterBar.h
+++ b/Engine/Source/Editor/Sequencer/Private/Filters/SequencerFilterBar.h
@@ -266,6 +266,9 @@ protected:
 	bool PassesAllInternalFilters(const UE::Sequencer::TViewModelPtr<UE::Sequencer::IOutlinerExtension>& InNode);
 	bool PassesAllCustomTextFilters(const UE::Sequencer::TViewModelPtr<UE::Sequencer::IOutlinerExtension>& InNode);
 
+	/** Calls the event before nodes are filtered to give filters a chance to pre-process */
+	void PreFilter() const;
+
 	TArray<TSharedRef<FSequencerTrackFilter>> GetFilterList(const bool bInIncludeCustomTextFilters = false) const;
 
 	/** The Sequencer this filter bar is interacting with */
@@ -281,8 +284,8 @@ protected:
 	TSharedRef<FFilterCategory> MiscCategory;
 	TSharedRef<FFilterCategory> TransientCategory;
 
-	TSharedPtr<FSequencerTrackFilterCollection> CommonFilters;
-	TSharedPtr<FSequencerTrackFilterCollection> InternalFilters;
+	TSharedRef<FSequencerTrackFilterCollection> CommonFilters;
+	TSharedRef<FSequencerTrackFilterCollection> InternalFilters;
 
 	TSharedRef<FSequencerTrackFilter_Text> TextFilter;
 	TSharedRef<FSequencerTrackFilter_HideIsolate> HideIsolateFilter;
diff --git a/Engine/Source/Editor/Sequencer/Public/Filters/Filters/SequencerTrackFilter_Group.h b/Engine/Source/Editor/Sequencer/Public/Filters/Filters/SequencerTrackFilter_Group.h
--- a/Engine/Source/Editor/Sequencer/Public/Filters/Filters/SequencerTrackFilter_Group.h
+++ b/Engine/Source/Editor/Sequencer/Public/Filters/Filters/SequencerTrackFilter_Group.h
@@ -22,6 +22,7 @@ public:
 	//~ Begin FSequencerTrackFilter
 	virtual FText GetDefaultToolTipText() const override;
 	virtual TSharedPtr<FUICommandInfo> GetToggleCommand() const override;
+	virtual void OnPreFilter() override;
 	//~ End FSequencerTrackFilter
 
 	//~ Begin FFilterBase
diff --git a/Engine/Source/Editor/Sequencer/Public/Filters/Filters/SequencerTrackFilter_Level.h b/Engine/Source/Editor/Sequencer/Public/Filters/Filters/SequencerTrackFilter_Level.h
--- a/Engine/Source/Editor/Sequencer/Public/Filters/Filters/SequencerTrackFilter_Level.h
+++ b/Engine/Source/Editor/Sequencer/Public/Filters/Filters/SequencerTrackFilter_Level.h
@@ -15,6 +15,7 @@ public:
 	//~ Begin FSequencerTrackFilter
 	virtual FText GetDefaultToolTipText() const override;
 	virtual TSharedPtr<FUICommandInfo> GetToggleCommand() const override;
+	virtual void OnPreFilter() override;
 	//~ End FSequencerTrackFilter
 
 	//~ Begin FFilterBase
diff --git a/Engine/Source/Editor/Sequencer/Public/Filters/SequencerFilterBase.h b/Engine/Source/Editor/Sequencer/Public/Filters/SequencerFilterBase.h
--- a/Engine/Source/Editor/Sequencer/Public/Filters/SequencerFilterBase.h
+++ b/Engine/Source/Editor/Sequencer/Public/Filters/SequencerFilterBase.h
@@ -117,6 +117,9 @@ public:
 
 	virtual bool SupportsSequence(UMovieSceneSequence* const InSequence) const { return true; }
 
+	/** Event called before nodes are filtered to give the filter a chance to pre-process */
+	virtual void OnPreFilter() {}
+
 protected:
 	void MapToggleAction(const TSharedPtr<FUICommandInfo>& InCommand)
 	{
