[Sequencer] Config setting for section color alpha
#Synced-CL 49683375

diff --git a/Engine/Source/Editor/Sequencer/Private/SSequencerSection.cpp b/Engine/Source/Editor/Sequencer/Private/SSequencerSection.cpp
--- a/Engine/Source/Editor/Sequencer/Private/SSequencerSection.cpp
+++ b/Engine/Source/Editor/Sequencer/Private/SSequencerSection.cpp
@@ -111,6 +111,9 @@ struct FSequencerSectionPainterImpl : FSequencerSectionPainter
 	{
 		using namespace UE::MovieScene;
 
+		const USequencerSettings* const SequencerSettings = Sequencer.GetSequencerSettings();
+		check(SequencerSettings);
+
 		TRange<FFrameNumber> SectionRange = SectionModel->GetRange();
 
 		const ESlateDrawEffect DrawEffects = bParentEnabled
@@ -135,7 +138,7 @@ struct FSequencerSectionPainterImpl : FSequencerSectionPainter
 			FSlateLayoutTransform(FVector2D(-InfiniteLowerOffset, 0.f))
 			);
 
-		if (Sequencer.GetSequencerSettings()->ShouldShowPrePostRoll())
+		if (SequencerSettings->ShouldShowPrePostRoll())
 		{
 			TOptional<FSlateClippingState> PreviousClipState = DrawElements.GetClippingState();
 			if (PreviousClipState.IsSet())
@@ -211,7 +214,7 @@ struct FSequencerSectionPainterImpl : FSequencerSectionPainter
 		const bool bHasChildren = Outliner && Outliner.AsModel()->GetChildren(EViewModelListType::Outliner);
 		const bool bIsExpanded  = Outliner && Outliner->IsExpanded();
 
-		FLinearColor BlendedTint = BlendColor(Tint).CopyWithNewOpacity(1.f);
+		const FLinearColor BlendedTint = BlendColor(Tint).CopyWithNewOpacity(SequencerSettings->GetSectionColorAlpha());
 
 		if (!bHasChildren || !bIsExpanded)
 		{
diff --git a/Engine/Source/Editor/Sequencer/Private/SequencerSettings.cpp b/Engine/Source/Editor/Sequencer/Private/SequencerSettings.cpp
--- a/Engine/Source/Editor/Sequencer/Private/SequencerSettings.cpp
+++ b/Engine/Source/Editor/Sequencer/Private/SequencerSettings.cpp
@@ -61,6 +61,7 @@ USequencerSettings::USequencerSettings( const FObjectInitializer& ObjectInitiali
 	PlaybackRangeStartColor = FLinearColor(0.1f, 0.5f, 0.1f, 1.0f);
 	PlaybackRangeEndColor = FLinearColor(0.5f, 0.01f, 0.01f, 1.0f);
 	MarkedFrameColor = FLinearColor(0.f, 1.f, 1.f, 0.4f);
+	SectionColorAlpha = 1.0f;
 	bCleanPlaybackMode = true;
 	bActivateRealtimeViewports = true;
 	bEvaluateSubSequencesInIsolation = false;
@@ -932,6 +933,20 @@ void USequencerSettings::SetSectionColorTints(const TArray<FColor>& InSectionCol
 	}
 }
 
+float USequencerSettings::GetSectionColorAlpha() const
+{
+	return SectionColorAlpha;
+}
+
+void USequencerSettings::SetSectionColorAlpha(const float InAlpha)
+{
+	if (SectionColorAlpha != InAlpha)
+	{
+		SectionColorAlpha = InAlpha;
+		SaveConfig();
+	}
+}
+
 bool USequencerSettings::GetCleanPlaybackMode() const
 {
 	return bCleanPlaybackMode;
diff --git a/Engine/Source/Editor/Sequencer/Public/SequencerSettings.h b/Engine/Source/Editor/Sequencer/Public/SequencerSettings.h
--- a/Engine/Source/Editor/Sequencer/Public/SequencerSettings.h
+++ b/Engine/Source/Editor/Sequencer/Public/SequencerSettings.h
@@ -432,6 +432,11 @@ public:
 	/** Set the section color tints */
 	UE_API void SetSectionColorTints(const TArray<FColor>& InSectionColorTints);
 
+	/** @return The alpha opacity to use for section background colors */
+	UE_API float GetSectionColorAlpha() const;
+	/** Set the alpha opacity to use for section background colors */
+	UE_API void SetSectionColorAlpha(const float InAlpha);
+
 	/** @return Whether to playback in clean mode (game view, hide viewport UI) */
 	UE_API bool GetCleanPlaybackMode() const;
 	/** Toggle whether to playback in clean mode */
@@ -809,6 +814,10 @@ protected:
 	UPROPERTY(config, EditAnywhere, Category = General)
 	TArray<FColor> SectionColorTints;
 
+	/** Alpha transparency to use for section background colors */
+	UPROPERTY(config, EditAnywhere, Category = General, meta = (ClampMin = "0.0", ClampMax = "1.0", UIMin = "0.1", UIMax = "1.0"))
+	float SectionColorAlpha;
+
 	/** When enabled, sequencer will playback in clean mode (game view, hide viewport UI) */
 	UPROPERTY(config, EditAnywhere, Category = General)
 	bool bCleanPlaybackMode;
