[Sequencer]
- Add track area track separator horizontal lines to help animators know what track keys are on
- Toggleable on/off through editor Sequencer settings
#Synced-CL 49683375

diff --git a/Engine/Source/Editor/Sequencer/Private/SequencerSettings.cpp b/Engine/Source/Editor/Sequencer/Private/SequencerSettings.cpp
--- a/Engine/Source/Editor/Sequencer/Private/SequencerSettings.cpp
+++ b/Engine/Source/Editor/Sequencer/Private/SequencerSettings.cpp
@@ -62,6 +62,7 @@ USequencerSettings::USequencerSettings( const FObjectInitializer& ObjectInitiali
 	PlaybackRangeEndColor = FLinearColor(0.5f, 0.01f, 0.01f, 1.0f);
 	MarkedFrameColor = FLinearColor(0.f, 1.f, 1.f, 0.4f);
 	SectionColorAlpha = 1.0f;
+	bDrawHorizontalTrackSeparatorLines = true;
 	bCleanPlaybackMode = true;
 	bActivateRealtimeViewports = true;
 	bEvaluateSubSequencesInIsolation = false;
@@ -714,7 +715,7 @@ bool USequencerSettings::GetShowSequencerToolbar() const
 
 void USequencerSettings::SetShowSequencerToolbar(bool bInShowSequencerToolbar)
 {
-	if(bShowSequencerToolbar != bInShowSequencerToolbar)
+	if (bShowSequencerToolbar != bInShowSequencerToolbar)
 	{
 		bShowSequencerToolbar = bInShowSequencerToolbar;
 		SaveConfig();
@@ -947,6 +948,20 @@ void USequencerSettings::SetSectionColorAlpha(const float InAlpha)
 	}
 }
 
+bool USequencerSettings::ShouldDrawHorizontalTrackSeparatorLines() const
+{
+	return bDrawHorizontalTrackSeparatorLines;
+}
+
+void USequencerSettings::SetDrawHorizontalTrackSeparatorLines(const bool bInDraw)
+{
+	if (bDrawHorizontalTrackSeparatorLines != bInDraw)
+	{
+		bDrawHorizontalTrackSeparatorLines = bInDraw;
+		SaveConfig();
+	}
+}
+
 bool USequencerSettings::GetCleanPlaybackMode() const
 {
 	return bCleanPlaybackMode;
diff --git a/Engine/Source/Editor/Sequencer/Private/SequencerTimeSliderController.cpp b/Engine/Source/Editor/Sequencer/Private/SequencerTimeSliderController.cpp
--- a/Engine/Source/Editor/Sequencer/Private/SequencerTimeSliderController.cpp
+++ b/Engine/Source/Editor/Sequencer/Private/SequencerTimeSliderController.cpp
@@ -1,6 +1,7 @@
 // Copyright Epic Games, Inc. All Rights Reserved.
 
 #include "SequencerTimeSliderController.h"
+#include "MVVM/ViewModels/OutlinerSpacer.h"
 #include "MVVM/ViewModels/SequencerEditorViewModel.h"
 #include "MVVM/Selection/Selection.h"
 #include "MVVM/SharedViewModelData.h"
@@ -360,6 +361,74 @@ void FSequencerTimeSliderController::DrawLinearTicks(FSlateWindowElementList& Ou
 	}
 }
 
+void FSequencerTimeSliderController::DrawHorizontalTrackSeparatorLines(FSlateWindowElementList& OutDrawElements
+	, const TRange<double>& InViewRange, const FScrubRangeToScreen& InRangeToScreen, FDrawTickArgs& InArgs) const
+{
+	using namespace UE::Sequencer;
+
+	const TSharedPtr<FSequencer> Sequencer = WeakSequencer.Pin();
+	if (!Sequencer.IsValid())
+	{
+		return;
+	}
+
+	const TSharedPtr<SSequencer> SequencerWidget = StaticCastSharedRef<SSequencer>(Sequencer->GetSequencerWidget());
+	if (!SequencerWidget.IsValid())
+	{
+		return;
+	}
+
+	const TSharedPtr<SOutlinerView> TreeViewWidget = SequencerWidget->GetTreeView();
+	if (!TreeViewWidget.IsValid())
+	{
+		return;
+	}
+
+	const TSharedPtr<SOutlinerView> OutlinerViewWidget = Sequencer->GetOutlinerViewWidget();
+	if (!OutlinerViewWidget.IsValid())
+	{
+		return;
+	}
+
+	TArray<TViewModelPtr<IOutlinerExtension>> VisibleItems;
+	OutlinerViewWidget->GetVisibleItems(VisibleItems);
+
+	for (const TViewModelPtr<IOutlinerExtension>& Item : VisibleItems)
+	{
+		if (Item.AsModel()->IsA<FOutlinerSpacer>())
+		{
+			continue;
+		}
+
+		// Get the vertical position and height of this item
+		const TViewModelPtr<IGeometryExtension> GeometryExtension = Item.AsModel()->CastThisShared<IGeometryExtension>();
+		if (!GeometryExtension.IsValid())
+		{
+			continue;
+		}
+
+		// Convert virtual geometry to physical screen space
+		const FVirtualGeometry ItemGeometry = GeometryExtension->GetVirtualGeometry();
+		const float ItemBottom = TreeViewWidget->VirtualToPhysical(ItemGeometry.GetTop() + ItemGeometry.GetNestedHeight());
+
+		// Draw a horizontal line at the top of this item
+		TArray<FVector2D> LinePoints;
+		LinePoints.AddUninitialized(2);
+		LinePoints[0] = FVector2D(0.f, ItemBottom);
+		LinePoints[1] = FVector2D(InArgs.AllottedGeometry.Size.X, ItemBottom);
+
+		FSlateDrawElement::MakeLines(
+			OutDrawElements,
+			InArgs.StartLayer,
+			InArgs.AllottedGeometry.ToPaintGeometry(),
+			LinePoints,
+			InArgs.DrawEffects,
+			InArgs.TickColor,
+			false
+		);
+	}
+}
+
 int32 FSequencerTimeSliderController::DrawScalingAnchors(const FGeometry& AllottedGeometry, const FScrubRangeToScreen& RangeToScreen, FSlateWindowElementList& OutDrawElements, int32 LayerId, const ESlateDrawEffect& DrawEffects) const
 {
 	UMovieSceneScalingAnchors* Anchors = WeakSequencer.Pin()->GetFocusedMovieSceneSequence()->GetMovieScene()->FindDecoration<UMovieSceneScalingAnchors>();
@@ -1566,6 +1635,12 @@ int32 FSequencerTimeSliderController::OnPaintViewArea( const FGeometry& Allotted
 		}
 
 		DrawTicks( OutDrawElements, LocalViewRange, RangeToScreen, DrawTickArgs );
+
+		const USequencerSettings* const SequencerSettings = Sequencer->GetSequencerSettings();
+		if (SequencerSettings && SequencerSettings->ShouldDrawHorizontalTrackSeparatorLines())
+		{
+			DrawHorizontalTrackSeparatorLines(OutDrawElements, LocalViewRange, RangeToScreen, DrawTickArgs);
+		}
 	}
 
 	if (Args.bDisplayMarkedFrames)
diff --git a/Engine/Source/Editor/Sequencer/Public/SequencerSettings.h b/Engine/Source/Editor/Sequencer/Public/SequencerSettings.h
--- a/Engine/Source/Editor/Sequencer/Public/SequencerSettings.h
+++ b/Engine/Source/Editor/Sequencer/Public/SequencerSettings.h
@@ -371,7 +371,7 @@ public:
 	/** @return true if showing sequencer toolbar */
 	UE_API bool GetShowSequencerToolbar() const;
 	/** Set whether to show sequencer toolbar bar */
-	UE_API void SetShowSequencerToolbar(bool bInDrawTickLines);
+	UE_API void SetShowSequencerToolbar(bool bInShowSequencerToolbar);
 
 	/** @return true if showing marked frames */
 	UE_API bool GetShowMarkedFrames() const;
@@ -437,6 +437,11 @@ public:
 	/** Set the alpha opacity to use for section background colors */
 	UE_API void SetSectionColorAlpha(const float InAlpha);
 
+	/** @return True if horizontal track separator lines should be drawn in the track area */
+	UE_API bool ShouldDrawHorizontalTrackSeparatorLines() const;
+	/** Set if horizontal track separator lines should be drawn in the track area */
+	UE_API void SetDrawHorizontalTrackSeparatorLines(const bool bInDraw);
+
 	/** @return Whether to playback in clean mode (game view, hide viewport UI) */
 	UE_API bool GetCleanPlaybackMode() const;
 	/** Toggle whether to playback in clean mode */
@@ -818,6 +823,10 @@ protected:
 	UPROPERTY(config, EditAnywhere, Category = General, meta = (ClampMin = "0.0", ClampMax = "1.0", UIMin = "0.1", UIMax = "1.0"))
 	float SectionColorAlpha;
 
+	/** Draw horizontal track separator lines in the track area */
+	UPROPERTY(config, EditAnywhere, Category = General)
+	bool bDrawHorizontalTrackSeparatorLines;
+
 	/** When enabled, sequencer will playback in clean mode (game view, hide viewport UI) */
 	UPROPERTY(config, EditAnywhere, Category = General)
 	bool bCleanPlaybackMode;
diff --git a/Engine/Source/Editor/Sequencer/Public/SequencerTimeSliderController.h b/Engine/Source/Editor/Sequencer/Public/SequencerTimeSliderController.h
--- a/Engine/Source/Editor/Sequencer/Public/SequencerTimeSliderController.h
+++ b/Engine/Source/Editor/Sequencer/Public/SequencerTimeSliderController.h
@@ -259,6 +259,11 @@ private:
 	 */
 	void DrawLinearTicks(FSlateWindowElementList& OutDrawElements, const TRange<double>& ViewRange, const FScrubRangeToScreen& RangeToScreen, FDrawTickArgs& InArgs) const;
 
+	/**
+	 * Draw horizontal separator lines between tracks
+	 */
+	void DrawHorizontalTrackSeparatorLines(FSlateWindowElementList& OutDrawElements, const TRange<double>& InViewRange, const FScrubRangeToScreen& InRangeToScreen, FDrawTickArgs& InArgs) const;
+
 	/**
 	 * Draw the selection range.
 	 *
