[Sequencer Ed Mode] Orbit Issues with Sequencer and Animation Mode

Prevents camera and drag handling to be performed if ITF already has active mouse capture.

#Synced-CL 49725045

diff --git a/Engine/Source/Editor/Sequencer/Private/SequencerEdMode.cpp b/Engine/Source/Editor/Sequencer/Private/SequencerEdMode.cpp
--- a/Engine/Source/Editor/Sequencer/Private/SequencerEdMode.cpp
+++ b/Engine/Source/Editor/Sequencer/Private/SequencerEdMode.cpp
@@ -298,6 +298,7 @@ bool FSequencerEdMode::IsMovingCamera(FViewport* InViewport) const
 
 	return ((SequencerSettings ? SequencerSettings->GetLeftMouseDragDoesMarquee() : false) && LeftMouseButtonDown && bIsAltKeyDown);
 }
+
 bool FSequencerEdMode::IsDoingDrag(FViewport* InViewport) const
 {
 	const USequencerSettings* SequencerSettings = GetSequencerSettings();
@@ -325,6 +326,11 @@ TSharedPtr<ISequencer> FSequencerEdMode::GetFirstActiveSequencer() const
 
 bool FSequencerEdMode::StartTracking(FEditorViewportClient* InViewportClient, FViewport* InViewport)
 {
+	if (Owner && Owner->GetInteractiveToolsContext()->InputRouter->HasActiveMouseCapture())
+	{
+		return false;
+	}
+
 	if (SelectabilityTool->IsSelectionLimited())
 	{
 		SelectabilityTool->StartTracking(InViewportClient, InViewport);
@@ -367,6 +373,11 @@ bool FSequencerEdMode::StartTracking(FEditorViewportClient* InViewportClient, FV
 
 bool FSequencerEdMode::EndTracking(FEditorViewportClient* InViewportClient, FViewport* InViewport)
 {
+	if (Owner && Owner->GetInteractiveToolsContext()->InputRouter->HasActiveMouseCapture())
+	{
+		return false;
+	}
+
 	if (SelectabilityTool->IsSelectionLimited())
 	{
 		SelectabilityTool->EndTracking(InViewportClient, InViewport);
@@ -400,6 +411,11 @@ void FSequencerEdMode::Tick(FEditorViewportClient* ViewportClient,float DeltaTim
 
 bool FSequencerEdMode::InputDelta(FEditorViewportClient* InViewportClient, FViewport* InViewport, FVector& InDrag, FRotator& InRot, FVector& InScale)
 {
+	if (Owner && Owner->GetInteractiveToolsContext()->InputRouter->HasActiveMouseCapture())
+	{
+		return false;
+	}
+
 	if (IsPressingMoveTimeSlider(InViewport))
 	{
 		return true;
