[AvalancheTests] Automate 'Sequence - Validate marks can be triggered based on their direction parameter'

#rnx

#Synced-CL 45465245

diff --git a/Engine/Source/Editor/Sequencer/Private/Sequencer.h b/Engine/Source/Editor/Sequencer/Private/Sequencer.h
--- a/Engine/Source/Editor/Sequencer/Private/Sequencer.h
+++ b/Engine/Source/Editor/Sequencer/Private/Sequencer.h
@@ -501,7 +501,7 @@ public:
 
 	/** Functions to push on to the transport controls we use */
 	FReply OnRecord();
-	FReply OnPlayForward(bool bTogglePlay);
+	DLLEXPORT FReply OnPlayForward(bool bTogglePlay);
 	DLLEXPORT FReply OnPlayBackward(bool bTogglePlay);
 	FReply OnStepForward(FFrameNumber Increment = FFrameNumber(1));
 	FReply OnStepBackward(FFrameNumber Increment = FFrameNumber(1));
@@ -868,7 +868,7 @@ public:
 	virtual void NotifyBindingsChanged() override;
 	virtual void SetViewportSettings(const TMap<FViewportClient*, EMovieSceneViewportParams>& ViewportParamsMap) override;
 	virtual void GetViewportSettings(TMap<FViewportClient*, EMovieSceneViewportParams>& ViewportParamsMap) const override;
-	virtual EMovieScenePlayerStatus::Type GetPlaybackStatus() const override;
+	DLLEXPORT virtual EMovieScenePlayerStatus::Type GetPlaybackStatus() const override;
 	virtual void SetPlaybackStatus(EMovieScenePlayerStatus::Type InPlaybackStatus) override;
 	virtual FMovieSceneSpawnRegister& GetSpawnRegister() override { return *SpawnRegister; }
 	virtual bool IsPreview() const override { return SilentModeCount != 0; }
