Sequencer: Remove NotifyMovieSceneDataChanged/Evaluation always firing when PostUndo happens since this breaks undos that don't require a re-evaluation.  This is hidden behind a cvar,Sequencer.EvaluateOnUndo.   which is currently set to false which is the current behavior.

#Synced-CL 49854396

diff --git a/Engine/Source/Editor/Sequencer/Private/Sequencer.cpp b/Engine/Source/Editor/Sequencer/Private/Sequencer.cpp
--- a/Engine/Source/Editor/Sequencer/Private/Sequencer.cpp
+++ b/Engine/Source/Editor/Sequencer/Private/Sequencer.cpp
@@ -220,6 +220,12 @@ namespace UE
 	namespace Sequencer
 	{
 
+	static bool bEvaluateOnUndo = true;
+	static FAutoConsoleVariableRef CVarEvaluateOnUndo(
+		TEXT("Sequencer.EvaluateOnUndo"),
+		bEvaluateOnUndo,
+		TEXT("Enable/disable Sequencer evaluating on undo/redo"));
+
 		struct FDeferredSignedObjectChangeHandler : UE::MovieScene::IDeferredSignedObjectChangeHandler
 		{
 			FDeferredSignedObjectChangeHandler()
@@ -6302,7 +6308,10 @@ bool FSequencer::MatchesContext(const FTransactionContext& InContext, const TArr
 
 void FSequencer::PostUndo(bool bSuccess)
 {
-	NotifyMovieSceneDataChanged( EMovieSceneDataChangeType::Unknown );
+	if (UE::Sequencer::bEvaluateOnUndo)
+	{
+		NotifyMovieSceneDataChanged(EMovieSceneDataChangeType::Unknown);
+	}
 	SynchronizeSequencerSelectionWithExternalSelection();
 	OnNodeGroupsCollectionChanged();
 
