Sequencer: Remove FSequencerFilterBar::GetWidget reference from SSequencerFilterBar.

This refactor prepares a follow-up feature for creating multiple instances of FSequencerFilterBar.
There will be multiple SSequencerFilterBar for a single FSequencerFilterBar.

#Synced-CL 49720279

diff --git a/Engine/Plugins/Experimental/SequenceNavigator/Source/SequenceNavigator/Private/Filters/NavigationToolFilterBar.cpp b/Engine/Plugins/Experimental/SequenceNavigator/Source/SequenceNavigator/Private/Filters/NavigationToolFilterBar.cpp
--- a/Engine/Plugins/Experimental/SequenceNavigator/Source/SequenceNavigator/Private/Filters/NavigationToolFilterBar.cpp
+++ b/Engine/Plugins/Experimental/SequenceNavigator/Source/SequenceNavigator/Private/Filters/NavigationToolFilterBar.cpp
@@ -74,7 +74,7 @@ FNavigationToolFilterBar::~FNavigationToolFilterBar()
 void FNavigationToolFilterBar::Init()
 {
 	CommonFilters->OnChanged().AddSP(this, &FNavigationToolFilterBar::RequestFilterUpdate);
-	TextFilter->OnChanged().AddSP(this, &FNavigationToolFilterBar::RequestFilterUpdate);
+	TextFilter->OnChanged().AddSP(this, &FNavigationToolFilterBar::HandleTextFilterChanged);
 	//SelectedFilter->OnChanged().AddSP(this, &FNavigationToolFilterBar::RequestFilterUpdate);
 
 	CreateDefaultFilters();
@@ -295,6 +295,7 @@ void FNavigationToolFilterBar::MuteFilters(const bool bInMute)
 		FilterBarWidget->SetMuted(bFiltersMuted);
 	}
 
+	MuteFiltersChangedEvent.Broadcast(bFiltersMuted);
 	RequestFilterUpdate();
 }
 
@@ -1245,6 +1246,32 @@ TArray<TSharedRef<ISequencerTextFilterExpressionContext>> FNavigationToolFilterB
 	return TextFilter->GetTextFilterExpressionContexts();
 }
 
+FString FNavigationToolFilterBar::GenerateTextFilterStringFromEnabledFilters() const
+{
+	TArray<TSharedRef<FNavigationToolFilter>> FiltersToSave;
+
+	FiltersToSave.Append(GetCommonFilters());
+
+	for (const TSharedRef<FNavigationToolFilter_CustomText>& Filter : CustomTextFilters)
+	{
+		FiltersToSave.Add(Filter);
+	}
+
+	FString GeneratedFilterString = TextFilter->GetRawFilterText().ToString();
+
+	for (const TSharedRef<FNavigationToolFilter>& Filter : FiltersToSave)
+	{
+		if (IsFilterActive(Filter) && IsFilterEnabled(Filter))
+		{
+			const FString AndAddString = GeneratedFilterString.IsEmpty() ? TEXT("") : TEXT(" AND ");
+			const FString ThisFilterGeneratedString = FString::Format(TEXT("{0}{1}==TRUE"), { AndAddString, *Filter->GetName() });
+			GeneratedFilterString.Append(ThisFilterGeneratedString);
+		}
+	}
+
+	return GeneratedFilterString;
+}
+
 bool FNavigationToolFilterBar::FilterNodesRecursive(INavigationToolView& InToolView
 	, const UNavigationToolSettings& InSettings
 	, const bool bInHasActiveFilter
@@ -1309,32 +1336,6 @@ bool FNavigationToolFilterBar::FilterNodesRecursive(INavigationToolView& InToolV
 	return false;
 }
 
-FString FNavigationToolFilterBar::GenerateTextFilterStringFromEnabledFilters() const
-{
-	TArray<TSharedRef<FNavigationToolFilter>> FiltersToSave;
-
-	FiltersToSave.Append(GetCommonFilters());
-
-	for (const TSharedRef<FNavigationToolFilter_CustomText>& Filter : CustomTextFilters)
-	{
-		FiltersToSave.Add(Filter);
-	}
-
-	FString GeneratedFilterString = TextFilter->GetRawFilterText().ToString();
-
-	for (const TSharedRef<FNavigationToolFilter>& Filter : FiltersToSave)
-	{
-		if (IsFilterActive(Filter) && IsFilterEnabled(Filter))
-		{
-			const FString AndAddString = GeneratedFilterString.IsEmpty() ? TEXT("") : TEXT(" AND ");
-			const FString ThisFilterGeneratedString = FString::Format(TEXT("{0}{1}==TRUE"), { AndAddString, *Filter->GetName() });
-			GeneratedFilterString.Append(ThisFilterGeneratedString);
-		}
-	}
-
-	return GeneratedFilterString;
-}
-
 TArray<TSharedRef<FNavigationToolFilter>> FNavigationToolFilterBar::GetFilterList(const bool bInIncludeCustomTextFilters) const
 {
 	TArray<TSharedRef<FNavigationToolFilter>> AllFilters;
@@ -1641,6 +1642,12 @@ bool FNavigationToolFilterBar::TryDeleteCustomTextFilter(const FString& InFilter
 	return true;
 }
 
+void FNavigationToolFilterBar::HandleTextFilterChanged()
+{
+	TextFilterTextChangedEvent.Broadcast();
+	RequestFilterUpdate();
+}
+
 bool FNavigationToolFilterBar::TryModifyCustomTextFilter(const FCustomTextFilterData& InFilterData, const FString& InOldFilterName, FText& OutErrorText)
 {
 	UNavigationToolSettings* const ToolSettings = GetMutableDefault<UNavigationToolSettings>();
diff --git a/Engine/Plugins/Experimental/SequenceNavigator/Source/SequenceNavigator/Public/Filters/NavigationToolFilterBar.h b/Engine/Plugins/Experimental/SequenceNavigator/Source/SequenceNavigator/Public/Filters/NavigationToolFilterBar.h
--- a/Engine/Plugins/Experimental/SequenceNavigator/Source/SequenceNavigator/Public/Filters/NavigationToolFilterBar.h
+++ b/Engine/Plugins/Experimental/SequenceNavigator/Source/SequenceNavigator/Public/Filters/NavigationToolFilterBar.h
@@ -70,6 +70,8 @@ public:
 	FNavigationToolFiltersChanged& OnFiltersChanged() { return FiltersChangedEvent; }
 	virtual FOnFilterBarStateChanged& OnStateChanged() override { return StateChangedEvent; }
 	virtual FSimpleMulticastDelegate& OnRequestUpdate() override { return RequestUpdateEvent; }
+	virtual FSimpleMulticastDelegate& OnTextFilterTextChanged() override { return TextFilterTextChangedEvent; }
+	virtual FMuteStateChangedDelegate& OnMuteFiltersChanged() override { return MuteFiltersChangedEvent; }
 
 	TSharedRef<FNavigationToolFilter_Text> GetTextFilter() const;
 	FText GetFilterErrorText() const;
@@ -116,6 +118,7 @@ public:
 
 	virtual const FTextFilterExpressionEvaluator& GetTextFilterExpressionEvaluator() const override;
 	virtual TArray<TSharedRef<ISequencerTextFilterExpressionContext>> GetTextFilterExpressionContexts() const override;
+	virtual FString GenerateTextFilterStringFromEnabledFilters() const override;
 
 	//~ End ISequencerFilterBar
 
@@ -182,8 +185,6 @@ public:
 
 	const FNavigationToolFilterData& FilterNodes();
 
-	FString GenerateTextFilterStringFromEnabledFilters() const;
-
 	bool ShouldUpdateOnTrackValueChanged() const;
 
 	TSharedRef<SComboButton> MakeAddFilterButton();
@@ -255,6 +256,8 @@ protected:
 	FNavigationToolFiltersChanged FiltersChangedEvent;
 	FOnFilterBarStateChanged StateChangedEvent;
 	FSimpleMulticastDelegate RequestUpdateEvent;
+	FSimpleMulticastDelegate TextFilterTextChangedEvent;
+	FMuteStateChangedDelegate MuteFiltersChangedEvent;
 
 private:
 	static int32 InstanceCount;
@@ -270,6 +273,8 @@ private:
 	bool TryCreateCustomTextFilter(const FCustomTextFilterData& InNewFilterData, const FString& InOldFilterName, const bool bInApply, FText& OutErrorText);
 	bool TryModifyCustomTextFilter(const FCustomTextFilterData& InNewFilterData, const FString& InOldFilterName, FText& OutErrorText);
 	bool TryDeleteCustomTextFilter(const FString& InFilterName, FText& OutErrorText);
+	
+	void HandleTextFilterChanged();
 };
 
 } // namespace UE::SequenceNavigator
diff --git a/Engine/Source/Editor/Sequencer/Private/Filters/Menus/SequencerFilterBarContextMenu.cpp b/Engine/Source/Editor/Sequencer/Private/Filters/Menus/SequencerFilterBarContextMenu.cpp
--- a/Engine/Source/Editor/Sequencer/Private/Filters/Menus/SequencerFilterBarContextMenu.cpp
+++ b/Engine/Source/Editor/Sequencer/Private/Filters/Menus/SequencerFilterBarContextMenu.cpp
@@ -6,6 +6,7 @@
 #include "ToolMenus.h"
 #include "Filters/SequencerFilterBar.h"
 #include "Filters/SequencerTrackFilterCommands.h"
+#include "Filters/Utils/FilterViewUtils.h"
 #include "Filters/Widgets/SSequencerFilterBar.h"
 
 #define LOCTEXT_NAMESPACE "SequencerFilterBarContextMenu"
@@ -68,17 +69,20 @@ void FSequencerFilterBarContextMenu::PopulateOptionsSection(UToolMenu& InMenu)
 
 	Section.AddMenuEntry(TrackFilterCommands.ToggleFilterBarVisibility);
 
-	if (const TSharedPtr<SSequencerFilterBar> FilterBarWidget = FilterBar->GetWidget())
-	{
-		Section.AddSeparator(NAME_Name);
-
-		Section.AddMenuEntry(TEXT("SaveCurrentFilterSetAsCustomTextFilter"),
-			LOCTEXT("SaveCurrentFilterSetAsCustomTextFilter", "Save Current as New Filter"),
-			LOCTEXT("SaveCurrentFilterSetAsCustomTextFilterTooltip", "Saves the enabled and active set of common filters as a custom text filter"),
-			FSlateIcon(FAppStyle::GetAppStyleSetName(), TEXT("LevelEditor.SaveAs")),
-			FUIAction(FExecuteAction::CreateSP(FilterBarWidget.ToSharedRef(), &SSequencerFilterBar::SaveCurrentFilterSetAsCustomTextFilter)));
-	}
-
+	Section.AddSeparator(NAME_Name);
+	Section.AddMenuEntry(
+		TEXT("SaveCurrentFilterSetAsCustomTextFilter"),
+		LOCTEXT("SaveCurrentFilterSetAsCustomTextFilter", "Save Current as New Filter"),
+		LOCTEXT("SaveCurrentFilterSetAsCustomTextFilterTooltip", "Saves the enabled and active set of common filters as a custom text filter"),
+		FSlateIcon(FAppStyle::GetAppStyleSetName(), TEXT("LevelEditor.SaveAs")),
+		FUIAction(FExecuteAction::CreateLambda([WeakFilterBar = FilterBar.ToWeakPtr()]()
+			{
+				if (const TSharedPtr<FSequencerFilterBar> FilterBarPin = WeakFilterBar.Pin())
+				{
+					UE::Sequencer::OpenDialog_SaveCurrentFilterSetAsCustomTextFilter(FilterBarPin.ToSharedRef());
+				}
+			}))
+		);
 	Section.AddSeparator(NAME_Name);
 
 	PopulateFilterBulkOptionsSection(InMenu);
diff --git a/Engine/Source/Editor/Sequencer/Private/Filters/Menus/SequencerTrackFilterMenu.cpp b/Engine/Source/Editor/Sequencer/Private/Filters/Menus/SequencerTrackFilterMenu.cpp
--- a/Engine/Source/Editor/Sequencer/Private/Filters/Menus/SequencerTrackFilterMenu.cpp
+++ b/Engine/Source/Editor/Sequencer/Private/Filters/Menus/SequencerTrackFilterMenu.cpp
@@ -16,6 +16,7 @@
 #include "ToolMenu.h"
 #include "ToolMenus.h"
 #include "ToolMenuSection.h"
+#include "Filters/Utils/FilterViewUtils.h"
 
 #define LOCTEXT_NAMESPACE "SequencerTrackFilterMenu"
 
@@ -313,13 +314,27 @@ void FSequencerTrackFilterMenu::FillCustomTextFiltersMenu(UToolMenu* const InMen
 		LOCTEXT("TextExpressionHelp", "Text Expression Help"),
 		LOCTEXT("TextExpressionHelpToolTip", "Opens the help dialog for the advanced search syntax text expressions"),
 		FSlateIcon(FAppStyle::GetAppStyleSetName(), TEXT("Icons.Help")),
-		FUIAction(FExecuteAction::CreateSP(this, &FSequencerTrackFilterMenu::OnOpenTextExpressionHelp)));
+		FUIAction(FExecuteAction::CreateLambda([WeakFilterBar = FilterBar.ToWeakPtr()]
+			{
+				if (const TSharedPtr<FSequencerFilterBar> FilterBarPin = WeakFilterBar.Pin())
+				{
+					UE::Sequencer::OpenDialog_TextExpressionHelp(FilterBarPin.ToSharedRef());
+				}
+			}))
+		);
 
 	CustomTextOptionsSection.AddMenuEntry(TEXT("SaveCurrentAsNewTextFilter"),
 		LOCTEXT("SaveCurrentAsNewTextFilter", "Save Current as New Filter"),
 		LOCTEXT("SaveCurrentAsNewTextFilterToolTip", "Saves the enabled and active set of common filters as a custom text filter"),
 		FSlateIcon(FAppStyle::GetAppStyleSetName(), TEXT("LevelEditor.SaveAs")),
-		FUIAction(FExecuteAction::CreateSP(this, &FSequencerTrackFilterMenu::SaveCurrentFilterSetAsCustomTextFilter)));
+		FUIAction(FExecuteAction::CreateLambda([WeakFilterBar = FilterBar.ToWeakPtr()]
+			{
+				if (const TSharedPtr<FSequencerFilterBar> FilterBarPin = WeakFilterBar.Pin())
+				{
+					UE::Sequencer::OpenDialog_SaveCurrentFilterSetAsCustomTextFilter(FilterBarPin.ToSharedRef());
+				}
+			}))
+		);
 
 	CustomTextOptionsSection.AddMenuEntry(TEXT("CreateNewTextFilter"),
 		LOCTEXT("CreateNewTextFilter", "Create New Filter"),
@@ -742,40 +757,6 @@ void FSequencerTrackFilterMenu::OpenNodeGroupsManager()
 	SequencerWidget->OpenNodeGroupsManager();
 }
 
-void FSequencerTrackFilterMenu::OnOpenTextExpressionHelp()
-{
-	const TSharedPtr<FSequencerFilterBar> FilterBar = WeakFilterBar.Pin();
-	if (!FilterBar.IsValid())
-	{
-		return;
-	}
-
-	const TSharedPtr<SSequencerFilterBar> FilterBarWidget = FilterBar->GetWidget();
-	if (!FilterBarWidget.IsValid())
-	{
-		return;
-	}
-
-	FilterBarWidget->OnOpenTextExpressionHelp();
-}
-
-void FSequencerTrackFilterMenu::SaveCurrentFilterSetAsCustomTextFilter()
-{
-	const TSharedPtr<FSequencerFilterBar> FilterBar = WeakFilterBar.Pin();
-	if (!FilterBar.IsValid())
-	{
-		return;
-	}
-
-	const TSharedPtr<SSequencerFilterBar> FilterBarWidget = FilterBar->GetWidget();
-	if (!FilterBarWidget.IsValid())
-	{
-		return;
-	}
-
-	FilterBarWidget->SaveCurrentFilterSetAsCustomTextFilter();
-}
-
 void FSequencerTrackFilterMenu::CreateNewTextFilter()
 {
 	const TSharedPtr<FSequencerFilterBar> FilterBar = WeakFilterBar.Pin();
diff --git a/Engine/Source/Editor/Sequencer/Private/Filters/Menus/SequencerTrackFilterMenu.h b/Engine/Source/Editor/Sequencer/Private/Filters/Menus/SequencerTrackFilterMenu.h
--- a/Engine/Source/Editor/Sequencer/Private/Filters/Menus/SequencerTrackFilterMenu.h
+++ b/Engine/Source/Editor/Sequencer/Private/Filters/Menus/SequencerTrackFilterMenu.h
@@ -65,9 +65,6 @@ protected:
 
 	void OpenNodeGroupsManager();
 
-	void OnOpenTextExpressionHelp();
-
-	void SaveCurrentFilterSetAsCustomTextFilter();
 	void CreateNewTextFilter();
 
 private:
diff --git a/Engine/Source/Editor/Sequencer/Private/Filters/SequencerFilterBar.cpp b/Engine/Source/Editor/Sequencer/Private/Filters/SequencerFilterBar.cpp
--- a/Engine/Source/Editor/Sequencer/Private/Filters/SequencerFilterBar.cpp
+++ b/Engine/Source/Editor/Sequencer/Private/Filters/SequencerFilterBar.cpp
@@ -68,7 +68,7 @@ FSequencerFilterBar::FSequencerFilterBar(FSequencer& InSequencer)
 
 	CommonFilters->OnChanged().AddRaw(this, &FSequencerFilterBar::RequestFilterUpdate);
 	InternalFilters->OnChanged().AddRaw(this, &FSequencerFilterBar::RequestFilterUpdate);
-	TextFilter->OnChanged().AddRaw(this, &FSequencerFilterBar::RequestFilterUpdate);
+	TextFilter->OnChanged().AddRaw(this, &FSequencerFilterBar::HandleTextFilterChanged);
 	LevelFilter->OnChanged().AddRaw(this, &FSequencerFilterBar::RequestFilterUpdate);
 	HideIsolateFilter->OnChanged().AddRaw(this, &FSequencerFilterBar::RequestFilterUpdate);
 	SelectedFilter->OnChanged().AddRaw(this, &FSequencerFilterBar::RequestFilterUpdate);
@@ -311,14 +311,6 @@ FName FSequencerFilterBar::GetIdentifier() const
 	return TEXT("SequencerMain");
 }
 
-TSharedRef<SSequencerFilterBar> FSequencerFilterBar::GenerateWidget(const TSharedPtr<SSequencerSearchBox>& InSearchBox, const EFilterBarLayout InLayout)
-{
-	return SNew(SSequencerFilterBar, SharedThis(this))
-		.FilterBarLayout(InLayout)
-		.AddMetaData<FTagMetaData>(FTagMetaData(TEXT("SequencerTrackFilters")))
-		.FilterSearchBox(InSearchBox);
-}
-
 bool FSequencerFilterBar::AreFiltersMuted() const
 {
 	return bFiltersMuted;
@@ -326,13 +318,14 @@ bool FSequencerFilterBar::AreFiltersMuted() const
 
 void FSequencerFilterBar::MuteFilters(const bool bInMute)
 {
-	bFiltersMuted = bInMute;
-
-	if (const TSharedPtr<SSequencerFilterBar> FilterBarWidget = GetWidget())
+	if (bFiltersMuted == bInMute)
 	{
-		FilterBarWidget->SetMuted(bFiltersMuted);
+		return;
 	}
-
+	
+	bFiltersMuted = bInMute;
+	
+	OnMuteFiltersChangedEvent.Broadcast(bFiltersMuted);
 	RequestFilterUpdate();
 }
 
@@ -362,11 +355,6 @@ FString FSequencerFilterBar::GetTextFilterString() const
 void FSequencerFilterBar::SetTextFilterString(const FString& InText)
 {
 	TextFilter->SetRawFilterText(FText::FromString(InText));
-
-	if (const TSharedPtr<SSequencerFilterBar> FilterBarWidget = GetWidget())
-	{
-		FilterBarWidget->SetTextFilterString(InText);
-	}
 }
 
 bool FSequencerFilterBar::DoesTextFilterStringContainExpressionPair(const ISequencerTextFilterExpressionContext& InExpression) const
@@ -458,7 +446,7 @@ void FSequencerFilterBar::EmptyIsolatedTracks()
 void FSequencerFilterBar::RequestFilterUpdate()
 {
 	Sequencer.GetNodeTree()->RequestFilterUpdate();
-
+	
 	RequestUpdateEvent.Broadcast();
 }
 
@@ -1451,7 +1439,7 @@ const FSequencerFilterData& FSequencerFilterBar::FilterNodes()
 			SpacerNode->SetFilteredOut(false);
 		}
 	}
-
+	
 	return FilterData;
 }
 
@@ -1476,6 +1464,32 @@ TArray<TSharedRef<ISequencerTextFilterExpressionContext>> FSequencerFilterBar::G
 	return TextExpressions;
 }
 
+FString FSequencerFilterBar::GenerateTextFilterStringFromEnabledFilters() const
+{
+	TArray<TSharedRef<FSequencerTrackFilter>> FiltersToSave;
+
+	FiltersToSave.Append(GetCommonFilters());
+
+	for (const TSharedRef<FSequencerTrackFilter_CustomText>& Filter : CustomTextFilters)
+	{
+		FiltersToSave.Add(Filter);
+	}
+
+	FString GeneratedFilterString = TextFilter->GetRawFilterText().ToString();
+
+	for (const TSharedRef<FSequencerTrackFilter>& Filter : FiltersToSave)
+	{
+		if (IsFilterActive(Filter) && IsFilterEnabled(Filter))
+		{
+			const FString AndAddString = GeneratedFilterString.IsEmpty() ? TEXT("") : TEXT(" AND ");
+			const FString ThisFilterGeneratedString = FString::Format(TEXT("{0}{1}==TRUE"), { AndAddString, *Filter->GetName() });
+			GeneratedFilterString.Append(ThisFilterGeneratedString);
+		}
+	}
+
+	return GeneratedFilterString;
+}
+
 bool FSequencerFilterBar::FilterNodesRecursive(const bool bInHasActiveFilter, const TViewModelPtr<IOutlinerExtension>& InStartNode)
 {
 	/**
@@ -1555,6 +1569,12 @@ bool FSequencerFilterBar::FilterNodesRecursive(const bool bInHasActiveFilter, co
 	return false;
 }
 
+void FSequencerFilterBar::HandleTextFilterChanged()
+{
+	OnTextFilterTextChangedEvent.Broadcast();
+	RequestFilterUpdate();
+}
+
 TSet<TWeakViewModelPtr<IOutlinerExtension>> FSequencerFilterBar::GetSelectedTracksOrAll() const
 {
 	TSharedPtr<FSequencerEditorViewModel> SequencerViewModel = GetSequencer().GetViewModel();
@@ -1628,32 +1648,6 @@ void FSequencerFilterBar::SetTrackParentsExpanded(const TViewModelPtr<IOutlinerE
 	}
 }
 
-FString FSequencerFilterBar::GenerateTextFilterStringFromEnabledFilters() const
-{
-	TArray<TSharedRef<FSequencerTrackFilter>> FiltersToSave;
-
-	FiltersToSave.Append(GetCommonFilters());
-
-	for (const TSharedRef<FSequencerTrackFilter_CustomText>& Filter : CustomTextFilters)
-	{
-		FiltersToSave.Add(Filter);
-	}
-
-	FString GeneratedFilterString = TextFilter->GetRawFilterText().ToString();
-
-	for (const TSharedRef<FSequencerTrackFilter>& Filter : FiltersToSave)
-	{
-		if (IsFilterActive(Filter) && IsFilterEnabled(Filter))
-		{
-			const FString AndAddString = GeneratedFilterString.IsEmpty() ? TEXT("") : TEXT(" AND ");
-			const FString ThisFilterGeneratedString = FString::Format(TEXT("{0}{1}==TRUE"), { AndAddString, *Filter->GetName() });
-			GeneratedFilterString.Append(ThisFilterGeneratedString);
-		}
-	}
-
-	return GeneratedFilterString;
-}
-
 TArray<TSharedRef<FSequencerTrackFilter>> FSequencerFilterBar::GetFilterList(const bool bInIncludeCustomTextFilters) const
 {
 	TArray<TSharedRef<FSequencerTrackFilter>> AllFilters;
@@ -1843,16 +1837,6 @@ void FSequencerFilterBar::ToggleFilterBarLayout()
 	}
 }
 
-TSharedPtr<SSequencerFilterBar> FSequencerFilterBar::GetWidget() const
-{
-	const TSharedPtr<SSequencer> SequencerWidget = StaticCastSharedRef<SSequencer>(GetSequencer().GetSequencerWidget());
-	if (!SequencerWidget.IsValid())
-	{
-		return nullptr;
-	}
-	return SequencerWidget->GetFilterBarWidget();
-}
-
 bool FSequencerFilterBar::IsFilterSupported(const TSharedRef<FSequencerTrackFilter>& InFilter) const
 {
 	UMovieSceneSequence* const MovieSceneSequence = Sequencer.GetFocusedMovieSceneSequence();
diff --git a/Engine/Source/Editor/Sequencer/Private/Filters/SequencerFilterBar.h b/Engine/Source/Editor/Sequencer/Private/Filters/SequencerFilterBar.h
--- a/Engine/Source/Editor/Sequencer/Private/Filters/SequencerFilterBar.h
+++ b/Engine/Source/Editor/Sequencer/Private/Filters/SequencerFilterBar.h
@@ -27,7 +27,7 @@ enum class EFilterBarLayout : uint8;
 
 namespace UE::Sequencer
 {
-	class IOutlinerExtension;
+class IOutlinerExtension;
 }
 
 enum class ESequencerFilterChange : uint8
@@ -55,8 +55,6 @@ public:
 
 	TSharedPtr<ICustomTextFilter<FSequencerTrackFilterType>> CreateTextFilter();
 
-	TSharedRef<SSequencerFilterBar> GenerateWidget(const TSharedPtr<SSequencerSearchBox>& InSearchBox, const EFilterBarLayout InLayout);
-
 	virtual bool AreFiltersMuted() const override;
 	virtual void MuteFilters(const bool bInMute) override;
 	void ToggleMuteFilters();
@@ -144,6 +142,7 @@ public:
 
 	virtual const FTextFilterExpressionEvaluator& GetTextFilterExpressionEvaluator() const override;
 	virtual TArray<TSharedRef<ISequencerTextFilterExpressionContext>> GetTextFilterExpressionContexts() const override;
+	virtual FString GenerateTextFilterStringFromEnabledFilters() const override;
 
 	//~ End ISequencerTrackFilters
 
@@ -228,8 +227,6 @@ public:
 
 	const FSequencerFilterData& FilterNodes();
 
-	FString GenerateTextFilterStringFromEnabledFilters() const;
-
 	bool ShouldUpdateOnTrackValueChanged() const;
 
 	TSharedRef<SFilterBarIsolateHideShow> MakeIsolateHideShowPanel();
@@ -245,15 +242,14 @@ public:
 	virtual void SetToHorizontalLayout() override;
 	virtual void ToggleFilterBarLayout() override;
 
-	/** Attempts to get the filter bar widget from the Sequencer widget */
-	TSharedPtr<SSequencerFilterBar> GetWidget() const;
-
 	bool IsFilterSupported(const TSharedRef<FSequencerTrackFilter>& InFilter) const;
 	bool IsFilterSupported(const FString& InFilterName) const;
 
 	virtual FOnFilterBarStateChanged& OnStateChanged() override { return StateChangedEvent; }
 	virtual FSimpleMulticastDelegate& OnRequestUpdate() override { return RequestUpdateEvent; }
-
+	virtual FSimpleMulticastDelegate& OnTextFilterTextChanged() override { return OnTextFilterTextChangedEvent; }
+	virtual FMuteStateChangedDelegate& OnMuteFiltersChanged() override { return OnMuteFiltersChangedEvent; }
+	
 	virtual void OpenTextExpressionHelp() override;
 	virtual void SaveCurrentFilterSetAsCustomTextFilter() override;
 	virtual void CreateNewTextFilter() override;
@@ -303,10 +299,15 @@ protected:
 	FSequencerFiltersChanged FiltersChangedEvent;
 	FOnFilterBarStateChanged StateChangedEvent;
 	FSimpleMulticastDelegate RequestUpdateEvent;
+	FSimpleMulticastDelegate OnTextFilterTextChangedEvent;
+	FMuteStateChangedDelegate OnMuteFiltersChangedEvent;
 
 private:
 	static int32 InstanceCount;
 
 	/** Do not call directly! Should only be called by FilterNodes(). */
 	bool FilterNodesRecursive(const bool bInHasActiveFilter, const UE::Sequencer::TViewModelPtr<UE::Sequencer::IOutlinerExtension>& InStartNode);
+	
+	/** Invoked when the text filter changes. */
+	void HandleTextFilterChanged();
 };
diff --git a/Engine/Source/Editor/Sequencer/Private/Filters/Utils/FilterViewUtils.cpp b/Engine/Source/Editor/Sequencer/Private/Filters/Utils/FilterViewUtils.cpp
new file mode 100644
--- /dev/null
+++ b/Engine/Source/Editor/Sequencer/Private/Filters/Utils/FilterViewUtils.cpp
@@ -0,0 +1,37 @@
+﻿// Copyright Epic Games, Inc. All Rights Reserved.
+
+#include "FilterViewUtils.h"
+
+#include "Filters/ISequencerTrackFilters.h"
+#include "Filters/Widgets/SFilterExpressionHelpDialog.h"
+#include "Filters/Widgets/SSequencerCustomTextFilterDialog.h"
+#include "Templates/SharedPointer.h"
+
+#define LOCTEXT_NAMESPACE "FilterViewUtils"
+
+namespace UE::Sequencer
+{
+void OpenDialog_SaveCurrentFilterSetAsCustomTextFilter(const TSharedRef<ISequencerTrackFilters>& InFilterBar)
+{
+	FCustomTextFilterData CustomTextFilterData;
+	CustomTextFilterData.FilterString = FText::FromString(InFilterBar->GenerateTextFilterStringFromEnabledFilters());
+	if (CustomTextFilterData.FilterLabel.IsEmpty())
+	{
+		CustomTextFilterData.FilterLabel = LOCTEXT("NewFilterName", "New Filter Name");
+	}
+
+	SSequencerCustomTextFilterDialog::CreateWindow_AddCustomTextFilter(InFilterBar, MoveTemp(CustomTextFilterData));
+}
+
+void OpenDialog_TextExpressionHelp(const TSharedRef<ISequencerTrackFilters>& InFilterBar)
+{
+	FFilterExpressionHelpDialogConfig Config;
+	Config.IdentifierName = TEXT("SequencerCustomTextFilterHelp");
+	Config.DialogTitle = LOCTEXT("SequencerCustomTextFilterHelp", "Sequencer Custom Text Filter Help");
+	Config.TextFilterExpressionContexts = InFilterBar->GetTextFilterExpressionContexts();
+
+	SFilterExpressionHelpDialog::Open(MoveTemp(Config));
+}
+}
+
+#undef LOCTEXT_NAMESPACE
\ No newline at end of file
diff --git a/Engine/Source/Editor/Sequencer/Private/Filters/Utils/FilterViewUtils.h b/Engine/Source/Editor/Sequencer/Private/Filters/Utils/FilterViewUtils.h
new file mode 100644
--- /dev/null
+++ b/Engine/Source/Editor/Sequencer/Private/Filters/Utils/FilterViewUtils.h
@@ -0,0 +1,17 @@
+﻿// Copyright Epic Games, Inc. All Rights Reserved.
+
+#pragma once
+
+#include "Templates/SharedPointerFwd.h"
+
+class ISequencerTrackFilters;
+
+namespace UE::Sequencer
+{
+/** Opens a dialog for saving the current filter set as a custom text filter. */
+void OpenDialog_SaveCurrentFilterSetAsCustomTextFilter(const TSharedRef<ISequencerTrackFilters>& InFilterBar);
+
+/** Opens a dialog that helps the user for creating custom text expressions. */
+void OpenDialog_TextExpressionHelp(const TSharedRef<ISequencerTrackFilters>& InFilterBar);
+}
+
diff --git a/Engine/Source/Editor/Sequencer/Private/Filters/Widgets/SSequencerFilterBar.cpp b/Engine/Source/Editor/Sequencer/Private/Filters/Widgets/SSequencerFilterBar.cpp
--- a/Engine/Source/Editor/Sequencer/Private/Filters/Widgets/SSequencerFilterBar.cpp
+++ b/Engine/Source/Editor/Sequencer/Private/Filters/Widgets/SSequencerFilterBar.cpp
@@ -86,6 +86,8 @@ void SSequencerFilterBar::Construct(const FArguments& InArgs, const TSharedRef<F
 	CreateFilterWidgetsFromConfig();
 
 	InFilterBar->OnFiltersChanged().AddSP(this, &SSequencerFilterBar::OnFiltersChanged);
+	InFilterBar->OnMuteFiltersChanged().AddSP(this, &SSequencerFilterBar::OnMuteStateChanged);
+	InFilterBar->OnTextFilterTextChanged().AddSP(this, &SSequencerFilterBar::OnTextFilterChanged);
 }
 
 FReply SSequencerFilterBar::OnMouseButtonUp(const FGeometry& MyGeometry, const FPointerEvent& MouseEvent)
@@ -485,37 +487,23 @@ void SSequencerFilterBar::OnFiltersChanged(const ESequencerFilterChange InChange
 	};
 }
 
-void SSequencerFilterBar::CreateAddCustomTextFilterWindowFromSearch(const FText& InSearchText)
+void SSequencerFilterBar::OnMuteStateChanged(bool bNewIsMuted)
 {
-	const TSharedPtr<FSequencerFilterBar> FilterBar = WeakFilterBar.Pin();
-	if (!FilterBar.IsValid())
+	if (const TSharedPtr<FSequencerFilterBar> FilterBar = GetFilterBar())
 	{
-		return;
+		SetMuted(bNewIsMuted);
 	}
-
-	FCustomTextFilterData CustomTextFilterData;
-	CustomTextFilterData.FilterLabel = LOCTEXT("NewFilterName", "New Filter Name");
-	CustomTextFilterData.FilterString = InSearchText;
-
-	SSequencerCustomTextFilterDialog::CreateWindow_AddCustomTextFilter(FilterBar.ToSharedRef(), MoveTemp(CustomTextFilterData));
 }
 
-void SSequencerFilterBar::OnOpenTextExpressionHelp()
+void SSequencerFilterBar::OnTextFilterChanged()
 {
-	const TSharedPtr<FSequencerFilterBar> FilterBar = WeakFilterBar.Pin();
-	if (!FilterBar.IsValid())
+	if (const TSharedPtr<FSequencerFilterBar> FilterBar = GetFilterBar())
 	{
-		return;
+		SetTextFilterString(FilterBar->GetTextFilterString());
 	}
-
-	FFilterExpressionHelpDialogConfig Config;
-	Config.IdentifierName = TEXT("SequencerCustomTextFilterHelp");
-	Config.DialogTitle = LOCTEXT("SequencerCustomTextFilterHelp", "Sequencer Custom Text Filter Help");
-	Config.TextFilterExpressionContexts = FilterBar->GetTextFilterExpressionContexts();
-
-	SFilterExpressionHelpDialog::Open(MoveTemp(Config));
 }
-void SSequencerFilterBar::SaveCurrentFilterSetAsCustomTextFilter()
+
+void SSequencerFilterBar::CreateAddCustomTextFilterWindowFromSearch(const FText& InSearchText)
 {
 	const TSharedPtr<FSequencerFilterBar> FilterBar = WeakFilterBar.Pin();
 	if (!FilterBar.IsValid())
@@ -524,11 +512,8 @@ void SSequencerFilterBar::SaveCurrentFilterSetAsCustomTextFilter()
 	}
 
 	FCustomTextFilterData CustomTextFilterData;
-	CustomTextFilterData.FilterString = FText::FromString(FilterBar->GenerateTextFilterStringFromEnabledFilters());
-	if (CustomTextFilterData.FilterLabel.IsEmpty())
-	{
-		CustomTextFilterData.FilterLabel = LOCTEXT("NewFilterName", "New Filter Name");
-	}
+	CustomTextFilterData.FilterLabel = LOCTEXT("NewFilterName", "New Filter Name");
+	CustomTextFilterData.FilterString = InSearchText;
 
 	SSequencerCustomTextFilterDialog::CreateWindow_AddCustomTextFilter(FilterBar.ToSharedRef(), MoveTemp(CustomTextFilterData));
 }
diff --git a/Engine/Source/Editor/Sequencer/Private/Filters/Widgets/SSequencerFilterBar.h b/Engine/Source/Editor/Sequencer/Private/Filters/Widgets/SSequencerFilterBar.h
--- a/Engine/Source/Editor/Sequencer/Private/Filters/Widgets/SSequencerFilterBar.h
+++ b/Engine/Source/Editor/Sequencer/Private/Filters/Widgets/SSequencerFilterBar.h
@@ -68,10 +68,6 @@ public:
 
 	void CreateAddCustomTextFilterWindowFromSearch(const FText& InSearchText);
 
-	void OnOpenTextExpressionHelp();
-
-	void SaveCurrentFilterSetAsCustomTextFilter();
-
 	TWeakPtr<SSequencerSearchBox> GetSearchBox() const;
 
 	// Set the state of the filter bar. Muted means that the filters are muted, but the context menu is still enabled and accessible.
@@ -92,7 +88,12 @@ protected:
 	void OnEnableAllGroupFilters(bool bEnableAll);
 	void OnNodeGroupFilterClicked(UMovieSceneNodeGroup* NodeGroup);
 
+	/** Refreshes the displayed filters based on the filter change. */
 	void OnFiltersChanged(const ESequencerFilterChange InChangeType, const TSharedRef<FSequencerTrackFilter>& InFilter);
+	/** Updates the mute state in response to the owning FSequencerTrackFilter changing mute state. */
+	void OnMuteStateChanged(bool bNewIsMuted);
+	/** Updates the search text */
+	void OnTextFilterChanged();
 
 	void CreateFilterWidgetsFromConfig();
 
diff --git a/Engine/Source/Editor/Sequencer/Private/SSequencer.cpp b/Engine/Source/Editor/Sequencer/Private/SSequencer.cpp
--- a/Engine/Source/Editor/Sequencer/Private/SSequencer.cpp
+++ b/Engine/Source/Editor/Sequencer/Private/SSequencer.cpp
@@ -1397,7 +1397,10 @@ TSharedRef<SWidget> SSequencer::ConstructSearchAndFilterRow()
 			.OnSaveSearchClicked(this, &SSequencer::OnOutlinerSearchSaved);
 	}
 
-	FilterBarWidget = FilterBar->GenerateWidget(SearchBox, GetFilterBarLayout());
+	FilterBarWidget = SNew(SSequencerFilterBar, FilterBar.ToSharedRef())
+		.FilterBarLayout(GetFilterBarLayout())
+		.AddMetaData<FTagMetaData>(FTagMetaData(TEXT("SequencerTrackFilters")))
+		.FilterSearchBox(SearchBox);
 
 	if (!FilterComboButtonWidget.IsValid())
 	{
diff --git a/Engine/Source/Editor/Sequencer/Public/Filters/ISequencerFilterBar.h b/Engine/Source/Editor/Sequencer/Public/Filters/ISequencerFilterBar.h
--- a/Engine/Source/Editor/Sequencer/Public/Filters/ISequencerFilterBar.h
+++ b/Engine/Source/Editor/Sequencer/Public/Filters/ISequencerFilterBar.h
@@ -18,7 +18,8 @@ class ISequencerFilterBar : public TSharedFromThis<ISequencerFilterBar>
 {
 public:
 	DECLARE_EVENT_TwoParams(ISequencerFilterBar, FOnFilterBarStateChanged, const bool /*InIsVisible*/, const EFilterBarLayout /*InNewLayout*/);
-
+	DECLARE_MULTICAST_DELEGATE_OneParam(FMuteStateChangedDelegate, bool /* bNewIsMuted */);
+	
 	virtual ~ISequencerFilterBar() = default;
 
 	virtual FName GetIdentifier() const = 0;
@@ -70,6 +71,11 @@ public:
 
 	virtual const FTextFilterExpressionEvaluator& GetTextFilterExpressionEvaluator() const = 0;
 	virtual TArray<TSharedRef<ISequencerTextFilterExpressionContext>> GetTextFilterExpressionContexts() const = 0;
+	/** 
+	 * @return The current filter state exported as string, such that it can be saved in a config and reapplied again.
+	 * @see FSequencerFilterBarConfig::GetCustomTextFilters, FSequencerTrackFilter_CustomText 
+	 */
+	virtual FString GenerateTextFilterStringFromEnabledFilters() const = 0;
 
 	virtual void OpenTextExpressionHelp() = 0;
 	virtual void SaveCurrentFilterSetAsCustomTextFilter() = 0;
@@ -80,6 +86,12 @@ public:
 
 	/** Event called to request that the subscribe update its view since the filters state has changed */
 	virtual FSimpleMulticastDelegate& OnRequestUpdate() = 0;
+	
+	/** Event called when the search text has changed. */
+	virtual FSimpleMulticastDelegate& OnTextFilterTextChanged() = 0;
+	
+	/** Event called when the mute state of the filters has changed. */
+	virtual FMuteStateChangedDelegate& OnMuteFiltersChanged() = 0;
 
 	virtual bool ShouldShowFilterBarWidget() const = 0;
 
