Sequencer: improve unselection
- SequencerHelpers::RemoveDuplicateKeys (called in FKeySelection::Deselect) does not have much to do in several circumstances
- this has a significant impact when working with a large number of channels and keys

#rnx

#Synced-CL 49971374

diff --git a/Engine/Source/Editor/Sequencer/Private/SequencerCommonHelpers.cpp b/Engine/Source/Editor/Sequencer/Private/SequencerCommonHelpers.cpp
--- a/Engine/Source/Editor/Sequencer/Private/SequencerCommonHelpers.cpp
+++ b/Engine/Source/Editor/Sequencer/Private/SequencerCommonHelpers.cpp
@@ -248,14 +248,61 @@ void SequencerHelpers::RemoveDuplicateKeys(const UE::Sequencer::FKeySelection& K
 	using namespace UE::Sequencer;
 
 	TMap<TViewModelPtr<FChannelModel>, TArray<FKeyHandle> > ChannelToKeyHandles;
-
+	TMap<TViewModelPtr<FChannelModel>, bool> ChannelsToCheck;
+	auto ShouldCheckChannel = [&ChannelsToCheck](const TViewModelPtr<FChannelModel>& ChannelModel)
+	{
+		if (ChannelModel)
+		{
+			if (bool* Found = ChannelsToCheck.Find(ChannelModel))
+			{
+				return *Found;
+			}
+			
+			bool bCheck = false;
+			
+			if (FMovieSceneChannel* Channel = ChannelModel->GetChannel())
+			{
+				TArray<FFrameNumber> AllTimes;
+				TArray<FKeyHandle> AllHandles;
+				Channel->GetKeys(TRange<FFrameNumber>::All(), &AllTimes, &AllHandles);
+			
+				if (AllTimes.Num() != AllHandles.Num())
+				{
+					// should never happen btw...
+					bCheck = true;
+				}
+				else if (AllTimes.Num() > 1) // no need to check if there's only one key
+				{
+					// if there is more than one key, check is only needed if the times are not increasing 
+					for (int32 Index = 1; Index < AllTimes.Num(); Index++)
+					{
+						if (AllTimes[Index] <= AllTimes[Index-1])
+						{
+							bCheck = true;
+							break;
+						}
+					}
+				}
+			}
+   
+			ChannelsToCheck.Add(ChannelModel, bCheck);
+			return bCheck;
+		}
+		return false;
+	};
+	
 	for (const FKeyHandle& KeyHandle : KeyHandles)
 	{
 		TViewModelPtr<FChannelModel> ChannelModel = KeySelection.GetModelForKey(KeyHandle);
 		if (ChannelModel)
 		{
-			if (FMovieSceneChannel* Channel = ChannelModel->GetChannel())
+			if (!ShouldCheckChannel(ChannelModel))
 			{
+				continue;
+			}
+			
+			if (FMovieSceneChannel* Channel = ChannelModel->GetChannel())
+			{	
 				if (Channel->GetIndex(KeyHandle) == INDEX_NONE)
 				{
 					continue;
@@ -267,9 +314,13 @@ void SequencerHelpers::RemoveDuplicateKeys(const UE::Sequencer::FKeySelection& K
 				TArray<FFrameNumber> OutKeyTimes;
 				TArray<FKeyHandle> OutKeyHandles;
 				Channel->GetKeys(TRange<FFrameNumber>(KeyTime, KeyTime), &OutKeyTimes, &OutKeyHandles);
-				OutKeyHandles.Remove(KeyHandle);
-
-				ChannelToKeyHandles.FindOrAdd(ChannelModel).Append(OutKeyHandles);
+				
+				// nothing to do if there's only one handle
+				if (OutKeyHandles.Num() > 1)
+				{	
+					OutKeyHandles.Remove(KeyHandle);
+					ChannelToKeyHandles.FindOrAdd(ChannelModel).Append(OutKeyHandles);
+				}
 			}
 		}
 	}
