[Sequencer] Make sidebar selection details able to be popped out

#Synced-CL 50018186

diff --git a/Engine/Source/Editor/Sequencer/Private/SSequencer.cpp b/Engine/Source/Editor/Sequencer/Private/SSequencer.cpp
--- a/Engine/Source/Editor/Sequencer/Private/SSequencer.cpp
+++ b/Engine/Source/Editor/Sequencer/Private/SSequencer.cpp
@@ -149,6 +149,10 @@
 #include "MVVM/Extensions/ITopTimeSliderOverlayExtension.h"
 #include "MVVM/Extensions/ITrackAreaOverlayExtension.h"
 #include "Settings/EditorExperimentalSettings.h"
+#include "Widgets/Docking/SDockTab.h"
+#include "Widgets/Sidebar/SequencerSelectionDrawer.h"
+#include "WorkspaceMenuStructure.h"
+#include "WorkspaceMenuStructureModule.h"
 
 #define LOCTEXT_NAMESPACE "Sequencer"
 
@@ -504,6 +508,7 @@ TSharedRef<SWidget> SSequencer::ConstructSidebarContent()
 			SidebarContainer = SNew(SSidebarContainer);
 
 			DetailsSidebar = SNew(SSidebar, SidebarContainer.ToSharedRef())
+				.Visibility(EVisibility::SelfHitTestInvisible)
 				.TabLocation(ESidebarTabLocation::Right)
 				.InitialDrawerSize(SidebarState.GetDrawerSize())
 				.OnStateChanged(this, &SSequencer::OnSidebarStateChanged)
@@ -511,6 +516,8 @@ TSharedRef<SWidget> SSequencer::ConstructSidebarContent()
 					{
 						return FilterBarSplitterContainer.ToSharedRef();
 					}));
+
+			SelectionDetailsDrawer = MakeShared<FSequencerSelectionDrawer>(GetSequencer());
 		}
 
 		SidebarContainer->RebuildSidebar(DetailsSidebar.ToSharedRef(), SidebarState);
@@ -518,6 +525,8 @@ TSharedRef<SWidget> SSequencer::ConstructSidebarContent()
 	else
 	{
 		DetailsSidebar.Reset();
+		SidebarContainer.Reset();
+		SelectionDetailsDrawer.Reset();
 	}
 	
 	RebuildFilterBarContent();
@@ -1962,11 +1971,16 @@ TSharedRef<SWidget> SSequencer::MakeViewMenu(const FToolMenuContext& InContext)
 
 void SSequencer::PopulateViewOptionMenu(UToolMenu* InToolMenu)
 {
-	TSharedPtr<FSequencer> Sequencer = SequencerPtr.Pin();
-	
+	const TSharedPtr<FSequencer> Sequencer = SequencerPtr.Pin();
+	if (!Sequencer.IsValid())
+	{
+		return;
+	}
+
 	if (Sequencer->GetHostCapabilities().bSupportsSidebar)
 	{
-		InToolMenu->AddMenuEntry("Sidebar", FToolMenuEntry::InitMenuEntry(FSequencerCommands::Get().ToggleSidebarVisible));
+		InToolMenu->AddMenuEntry(TEXT("SidebarVisible"), FToolMenuEntry::InitMenuEntry(FSequencerCommands::Get().ToggleSidebarVisible));
+		InToolMenu->AddMenuEntry(TEXT("SelectionDetailsPoppedOut"), FToolMenuEntry::InitMenuEntry(FSequencerCommands::Get().ToggleSelectionDetailsPoppedOut));
 	}
 
 	if (Sequencer->IsLevelEditorSequencer())
@@ -2654,6 +2668,10 @@ SSequencer::~SSequencer()
 
 	// Ensure the FilterBarWidget destructor is called
 	FilterBarSplitterContainer->SetContent(SNullWidget::NullWidget);
+
+	CloseSelectionDetailsTab();
+
+	UnregisterSelectionDetailsTab();
 }
 
 void SSequencer::RegisterActiveTimerForPlayback()
@@ -4080,8 +4098,7 @@ bool SSequencer::UnregisterDrawerSection(const FName InDrawerId, const FName InS
 
 void SSequencer::OnSidebarStateChanged(const FSidebarState& InNewState)
 {
-	USequencerSettings* const SequencerSettings = GetSequencerSettings();
-	if (IsValid(SequencerSettings))
+	if (USequencerSettings* const SequencerSettings = GetSequencerSettings())
 	{
 		SequencerSettings->SetSidebarState(InNewState);
 	}
@@ -4089,8 +4106,7 @@ void SSequencer::OnSidebarStateChanged(const FSidebarState& InNewState)
 
 bool SSequencer::IsSidebarVisible() const
 {
-	USequencerSettings* const SequencerSettings = GetSequencerSettings();
-	if (IsValid(SequencerSettings))
+	if (USequencerSettings* const SequencerSettings = GetSequencerSettings())
 	{
 		return SequencerSettings->GetSidebarState().IsVisible();
 	}
@@ -4100,7 +4116,7 @@ bool SSequencer::IsSidebarVisible() const
 void SSequencer::SetSidebarVisible(const bool bInVisible)
 {
 	USequencerSettings* const SequencerSettings = GetSequencerSettings();
-	if (!IsValid(SequencerSettings))
+	if (!SequencerSettings)
 	{
 		return;
 	}
@@ -4117,6 +4133,69 @@ void SSequencer::ToggleSidebarVisible()
 	SetSidebarVisible(!IsSidebarVisible());
 }
 
+bool SSequencer::IsSelectionDetailsPoppedOut() const
+{
+	if (USequencerSettings* const SequencerSettings = GetSequencerSettings())
+	{
+		return SequencerSettings->IsSelectionDetailsPoppedOut();
+	}
+	return false;
+}
+
+void SSequencer::SetSelectionDetailsPoppedOut(const bool bInPoppedOut)
+{
+	const TSharedPtr<FTabManager> TabManager = GetTabManager();
+	if (!TabManager.IsValid())
+	{
+		return;
+	}
+
+	USequencerSettings* const SequencerSettings = GetSequencerSettings();
+	if (!SequencerSettings)
+	{
+		return;
+	}
+
+	SequencerSettings->SetSelectionDetailsPoppedOut(bInPoppedOut);
+
+	if (bInPoppedOut)
+	{
+		RegisterSelectionDetailsTab();
+
+		// Only proceed with unregistering the drawer if the tab was successfully created
+		if (TabManager->TryInvokeTab(FSequencer::SelectionDetailsTabId).IsValid())
+		{
+			SequencerSettings->SetSelectionDetailsPoppedOut(true);
+
+			if (DetailsSidebar.IsValid() && DetailsSidebar->IsDrawerDocked(FSequencer::SelectionDrawerId))
+			{
+				DetailsSidebar->SetDrawerDocked(FSequencer::SelectionDrawerId, false);
+			}
+
+			UnregisterDrawerSection(FSequencer::SelectionDrawerId, FSequencerSelectionDrawer::UniqueId);
+		}
+	}
+	else
+	{
+		if (const TSharedPtr<SDockTab> ExistingTab = TabManager->FindExistingLiveTab(FSequencer::SelectionDetailsTabId))
+		{
+			ExistingTab->RequestCloseTab();
+		}
+
+		UnregisterSelectionDetailsTab();
+
+		if (SelectionDetailsDrawer.IsValid())
+		{
+			RegisterDrawerSection(FSequencer::SelectionDrawerId, SelectionDetailsDrawer);
+		}
+	}
+}
+
+void SSequencer::ToggleSelectionDetailsPoppedOut()
+{
+	SetSelectionDetailsPoppedOut(!IsSelectionDetailsPoppedOut());
+}
+
 void SSequencer::ToggleSidebarSelectionDrawerOpen()
 {
 	if (!DetailsSidebar.IsValid())
@@ -4300,4 +4379,115 @@ void SSequencer::RebuildSearchAndFilterRow()
 	}
 }
 
+TSharedPtr<FTabManager> SSequencer::GetTabManager() const
+{
+	if (const TSharedPtr<FSequencer> Sequencer = SequencerPtr.Pin())
+	{
+		if (const TSharedPtr<IToolkitHost> ToolkitHost = Sequencer->GetToolkitHost())
+		{
+			return ToolkitHost->GetTabManager();
+		}
+	}
+	return nullptr;
+}
+
+void SSequencer::RegisterSelectionDetailsTab()
+{
+	const TSharedPtr<FTabManager> TabManager = GetTabManager();
+	if (!TabManager.IsValid() || TabManager->HasTabSpawner(FSequencer::SelectionDetailsTabId))
+	{
+		return;
+	}
+
+	TabManager->RegisterTabSpawner(FSequencer::SelectionDetailsTabId
+		, FOnSpawnTab::CreateSP(this, &SSequencer::SpawnSelectionDetailsTab))
+			.SetDisplayName(LOCTEXT("SequencerSelectionDetailsDisplayName", "Sequencer Details"))
+			.SetIcon(FSlateIcon(FAppStyle::GetAppStyleSetName(), TEXT("Icons.Details")))
+			.SetGroup(WorkspaceMenu::GetMenuStructure().GetLevelEditorCinematicsCategory());
+}
+
+void SSequencer::UnregisterSelectionDetailsTab()
+{
+	if (const TSharedPtr<FTabManager> TabManager = GetTabManager())
+	{
+		TabManager->UnregisterTabSpawner(FSequencer::SelectionDetailsTabId);
+	}
+}
+
+TSharedRef<SDockTab> SSequencer::SpawnSelectionDetailsTab(const FSpawnTabArgs& InArgs)
+{
+	const TWeakPtr<SSequencer> WeakSequencerWidget = SharedThis(this);
+
+	const TSharedRef<SDockTab> ToolTab = SNew(SDockTab)
+		.TabRole(ETabRole::NomadTab)
+		.ShouldAutosize(false)
+		.OnTabClosed_Lambda([WeakSequencerWidget](const TSharedRef<SDockTab> InDockTab)
+			{
+				const TSharedPtr<SSequencer> SequencerWidget = WeakSequencerWidget.Pin();
+				if (!SequencerWidget.IsValid())
+				{
+					return;
+				}
+
+				if (USequencerSettings* const SequencerSettings = SequencerWidget->GetSequencerSettings())
+				{
+					SequencerSettings->SetSelectionDetailsPoppedOut(false);
+				}
+
+				SequencerWidget->UnregisterSelectionDetailsTab();
+
+				if (ensure(SequencerWidget->SelectionDetailsDrawer.IsValid()))
+				{
+					SequencerWidget->RegisterDrawerSection(FSequencer::SelectionDrawerId, SequencerWidget->SelectionDetailsDrawer);
+				}
+			})
+		[
+			SelectionDetailsDrawer.IsValid() ? SelectionDetailsDrawer->CreateContentWidget() : SNullWidget::NullWidget
+		];
+
+	return ToolTab;
+}
+
+void SSequencer::CloseSelectionDetailsTab()
+{
+	if (const TSharedPtr<FTabManager> TabManager = GetTabManager())
+	{
+		if (const TSharedPtr<SDockTab> ExistingTab = TabManager->FindExistingLiveTab(FSequencer::SelectionDetailsTabId))
+		{
+			ExistingTab->RequestCloseTab();
+		}
+	}
+}
+
+FText SSequencer::GetSidebarSelectionDrawerToolTipText() const
+{
+	const TSharedRef<const FInputChord> DrawerOpenActiveChord = FSequencerCommands::Get().ToggleSidebarSelectionDrawerOpen->GetFirstValidChord();
+	const TSharedRef<const FInputChord> DrawerDockActiveChord = FSequencerCommands::Get().ToggleSidebarDrawerDock->GetFirstValidChord();
+
+	FText ToolTipText = LOCTEXT("SelectionDetailsPanelTooltip", "Open Sequencer selection details panel.");
+
+	if (DrawerOpenActiveChord->IsValidChord() || DrawerDockActiveChord->IsValidChord())
+	{
+		ToolTipText = FText::Format(LOCTEXT("ExtendedSelectionDetailsPanelTooltip", "{0}\n"), ToolTipText);
+	}
+
+	if (DrawerOpenActiveChord->IsValidChord())
+	{
+		ToolTipText = FText::Format(LOCTEXT("ExtendedSelectionDrawerOpenDetailsPanelTooltip", "{0}\n"
+			"{1} to toggle the drawer open or closed")
+			, ToolTipText
+			, DrawerOpenActiveChord->GetInputText(true));
+	}
+
+	if (DrawerDockActiveChord->IsValidChord())
+	{
+		ToolTipText = FText::Format(LOCTEXT("ExtendedSelectionDrawerDockDetailsPanelTooltip", "{0}\n"
+			"{1} to toggle the drawer docked or undocked")
+			, ToolTipText
+			, DrawerDockActiveChord->GetInputText(true));
+	}
+
+	return ToolTipText;
+}
+
 #undef LOCTEXT_NAMESPACE
diff --git a/Engine/Source/Editor/Sequencer/Private/SSequencer.h b/Engine/Source/Editor/Sequencer/Private/SSequencer.h
--- a/Engine/Source/Editor/Sequencer/Private/SSequencer.h
+++ b/Engine/Source/Editor/Sequencer/Private/SSequencer.h
@@ -29,6 +29,7 @@ class FFolderDragDropOp;
 class FAssetDragDropOp;
 class FClassDragDropOp;
 class FMovieSceneClipboard;
+class FSequencerSelectionDrawer;
 class FSequencerTimeSliderController;
 class FSequencerTrackFilter;
 class FSequencerViewOptionsMenu;
@@ -684,6 +685,15 @@ public:
 	/** Toggle the visibility of the sidebar. */
 	void ToggleSidebarVisible();
 
+	/** @return True if the selection details are popped out of the Sequencer tab to its own tab window. */
+	bool IsSelectionDetailsPoppedOut() const;
+
+	/** Set the popped out state of the selection details in the Sequencer tab. */
+	void SetSelectionDetailsPoppedOut(const bool bInPoppedOut);
+
+	/** Toggle the popped out state of the selection details in the Sequencer tab. */
+	void ToggleSelectionDetailsPoppedOut();
+
 	/** Toggles the sidebar "Selection" drawer open or closed. */
 	void ToggleSidebarSelectionDrawerOpen();
 
@@ -719,6 +729,16 @@ private:
 
 	void OnTrackFiltersChanged(const ESequencerFilterChange InChangeType, const TSharedRef<FSequencerTrackFilter>& InFilter);
 
+	TSharedPtr<FTabManager> GetTabManager() const;
+
+	void RegisterSelectionDetailsTab();
+	void UnregisterSelectionDetailsTab();
+
+	TSharedRef<SDockTab> SpawnSelectionDetailsTab(const FSpawnTabArgs& InArgs);
+	void CloseSelectionDetailsTab();
+
+	FText GetSidebarSelectionDrawerToolTipText() const;
+
 	/** This is the main container widget for the Sequencer (minus the toolbar). This is the parent of
 	 * the GridPanel widget below and is used to rebuild content when the sidebar is added/removed. */
 	TSharedPtr<SBox> MainContentContainer;
@@ -894,6 +914,8 @@ private:
 	TSharedPtr<SSidebarContainer> SidebarContainer;
 	TSharedPtr<SSidebar> DetailsSidebar;
 
+	TSharedPtr<FSequencerSelectionDrawer> SelectionDetailsDrawer;
+
 	TSharedPtr<SSequencerFilterBar> FilterBarWidget;
 	TSharedPtr<SComboButton> FilterComboButtonWidget;
 
diff --git a/Engine/Source/Editor/Sequencer/Private/Sequencer.cpp b/Engine/Source/Editor/Sequencer/Private/Sequencer.cpp
--- a/Engine/Source/Editor/Sequencer/Private/Sequencer.cpp
+++ b/Engine/Source/Editor/Sequencer/Private/Sequencer.cpp
@@ -334,6 +334,7 @@ namespace UE
 } // namespace UE
 
 const FName FSequencer::SelectionDrawerId = TEXT("SelectionDetails");
+const FName FSequencer::SelectionDetailsTabId = TEXT("SequencerSelectionDetails");
 bool FSequencer::bSelectionLimited = false;
 
 void FSequencer::InitSequencer(const FSequencerInitParams& InitParams
@@ -775,13 +776,14 @@ void FSequencer::InitSequencer(const FSequencerInitParams& InitParams
 	{
 		FSidebarDrawerConfig DetailsDrawer;
 		DetailsDrawer.UniqueId = SelectionDrawerId;
+		DetailsDrawer.SortOrder = 0.f;
 		DetailsDrawer.ButtonText = LOCTEXT("SelectionDetailsPanelLabel", "Selection");
 		DetailsDrawer.ToolTipText = TAttribute<FText>::CreateSP(this, &FSequencer::GetSidebarSelectionDrawerToolTipText);
 		DetailsDrawer.Icon = FAppStyle::GetBrush(TEXT("EditorPreferences.TabIcon"));
 		DetailsDrawer.InitialState = Settings->GetSidebarState().FindOrAddDrawerState(SelectionDrawerId);
 		RegisterDrawer(MoveTemp(DetailsDrawer));
 
-		RegisterDrawerSection(SelectionDrawerId, MakeShared<FSequencerSelectionDrawer>(SharedThisRef));
+		SequencerWidget->SetSelectionDetailsPoppedOut(Settings->IsSelectionDetailsPoppedOut());
 	}
 }
 
@@ -12167,6 +12169,12 @@ void FSequencer::BindCommands()
 			FExecuteAction::CreateSP(this, &FSequencer::ToggleSidebar),
 			FCanExecuteAction(),
 			FIsActionChecked::CreateSP(this, &FSequencer::IsSidebarVisible));
+
+		SequencerCommandBindings->MapAction(
+			Commands.ToggleSelectionDetailsPoppedOut,
+			FExecuteAction::CreateSP(this, &FSequencer::ToggleSelectionDetailsPoppedOut),
+			FCanExecuteAction(),
+			FIsActionChecked::CreateSP(this, &FSequencer::IsSelectionDetailsPoppedOut));
 	}
 
 	SequencerCommandBindings->MapAction(
@@ -12865,7 +12873,7 @@ void FSequencer::ToggleSidebar()
 {
 	if (SequencerWidget.IsValid())
 	{
-		return SequencerWidget->ToggleSidebarVisible();
+		SequencerWidget->ToggleSidebarVisible();
 	}
 }
 
@@ -12874,6 +12882,14 @@ bool FSequencer::IsSidebarVisible() const
 	return SequencerWidget.IsValid() ? SequencerWidget->IsSidebarVisible() : false;
 }
 
+void FSequencer::ToggleSelectionDetailsPoppedOut()
+{
+	if (SequencerWidget.IsValid())
+	{
+		SequencerWidget->ToggleSelectionDetailsPoppedOut();
+	}
+}
+
 void FSequencer::ToggleSidebarSelectionDrawer()
 {
 	if (SequencerWidget.IsValid())
@@ -12882,6 +12898,19 @@ void FSequencer::ToggleSidebarSelectionDrawer()
 	}
 }
 
+bool FSequencer::IsSelectionDetailsPoppedOut() const
+{
+	return SequencerWidget.IsValid() ? SequencerWidget->IsSelectionDetailsPoppedOut() : false;
+}
+
+void FSequencer::SetSelectionDetailsPoppedOut(const bool bInPoppedOut)
+{
+	if (SequencerWidget.IsValid())
+	{
+		SequencerWidget->SetSelectionDetailsPoppedOut(bInPoppedOut);
+	}
+}
+
 void FSequencer::ToggleSidebarDrawerDocked()
 {
 	if (SequencerWidget.IsValid())
diff --git a/Engine/Source/Editor/Sequencer/Private/Sequencer.h b/Engine/Source/Editor/Sequencer/Private/Sequencer.h
--- a/Engine/Source/Editor/Sequencer/Private/Sequencer.h
+++ b/Engine/Source/Editor/Sequencer/Private/Sequencer.h
@@ -125,6 +125,7 @@ class FSequencer final
 
 public:
 	SEQUENCER_API static const FName SelectionDrawerId;
+	SEQUENCER_API static const FName SelectionDetailsTabId;
 
 	/** Constructor */
 	FSequencer();
@@ -1267,6 +1268,15 @@ public:
 	/** @return True if the Sequencer sidebar is currently visible. */
 	bool IsSidebarVisible() const;
 
+	/** Toggles the state of the Sequencer selection details between docked in the Sequencer tab and being popped out to its own tab window. */
+	void ToggleSelectionDetailsPoppedOut();
+
+	/** @return True if the Sequencer selection details are currently popped out to its own tab window. */
+	bool IsSelectionDetailsPoppedOut() const;
+
+	/** Set the popped out state of the selection details in the Sequencer tab. */
+	void SetSelectionDetailsPoppedOut(const bool bInPoppedOut);
+
 	/** Toggles the sidebar "Selection" drawer open or closed. */
 	void ToggleSidebarSelectionDrawer();
 
diff --git a/Engine/Source/Editor/Sequencer/Private/SequencerCommands.cpp b/Engine/Source/Editor/Sequencer/Private/SequencerCommands.cpp
--- a/Engine/Source/Editor/Sequencer/Private/SequencerCommands.cpp
+++ b/Engine/Source/Editor/Sequencer/Private/SequencerCommands.cpp
@@ -204,6 +204,8 @@ void FSequencerCommands::RegisterCommands()
 	UI_COMMAND(BakeTransform, "Bake Transform", "Bake transform in world space, removing any existing transform and attach tracks", EUserInterfaceActionType::Button, FInputChord());
 
 	UI_COMMAND(ToggleSidebarVisible, "Sidebar Details", "Toggles the visibility of the Sequencer sidebar.", EUserInterfaceActionType::ToggleButton, FInputChord());
+	UI_COMMAND(ToggleSelectionDetailsPoppedOut, "Popout Selection Details", "Toggles the state of the Sequencer selection details between docked in the Sequencer tab and being popped out to its own tab window", EUserInterfaceActionType::ToggleButton, FInputChord());
+
 	UI_COMMAND(ToggleSidebarSelectionDrawerOpen, "Sidebar Selection Drawer Open", "Toggles the visible state of the sidebar selection drawer. Undocks and closes the sidebar drawer if docked.", EUserInterfaceActionType::Button, FInputChord(EModifierKey::Control, EKeys::E));
 	UI_COMMAND(ToggleSidebarDrawerDock, "Sidebar Drawer Docked", "Toggles the docked state of the sidebar current drawer. Undocks the docked sidebar drawer if docked or docks the sidebar drawer if there is one open and no currently docked drawer.", EUserInterfaceActionType::Button, FInputChord(EModifierKey::Alt, EKeys::D));
 
diff --git a/Engine/Source/Editor/Sequencer/Private/SequencerSettings.cpp b/Engine/Source/Editor/Sequencer/Private/SequencerSettings.cpp
--- a/Engine/Source/Editor/Sequencer/Private/SequencerSettings.cpp
+++ b/Engine/Source/Editor/Sequencer/Private/SequencerSettings.cpp
@@ -84,6 +84,7 @@ USequencerSettings::USequencerSettings( const FObjectInitializer& ObjectInitiali
 	AssetBrowserHeight = 300.f;
 	bFilterBarVisible = true;
 	bNavigationToolVisible = false;
+	bSelectionDetailsPoppedOut = false;
 
 	SectionColorTints.Add(FColor(88, 102, 142, 255)); // blue
 	SectionColorTints.Add(FColor(99, 137, 132, 255)); // blue-green
@@ -1312,3 +1313,12 @@ void USequencerSettings::SetNavigationToolVisible(const bool bInVisible)
 		SaveConfig();
 	}
 }
+
+void USequencerSettings::SetSelectionDetailsPoppedOut(const bool bInPoppedOut)
+{
+	if (bSelectionDetailsPoppedOut != bInPoppedOut)
+	{
+		bSelectionDetailsPoppedOut = bInPoppedOut;
+		SaveConfig();
+	}
+}
diff --git a/Engine/Source/Editor/Sequencer/Private/Widgets/Sidebar/SequencerSelectionDrawer.cpp b/Engine/Source/Editor/Sequencer/Private/Widgets/Sidebar/SequencerSelectionDrawer.cpp
--- a/Engine/Source/Editor/Sequencer/Private/Widgets/Sidebar/SequencerSelectionDrawer.cpp
+++ b/Engine/Source/Editor/Sequencer/Private/Widgets/Sidebar/SequencerSelectionDrawer.cpp
@@ -106,11 +106,7 @@ FSequencerSelectionDrawer::FSequencerSelectionDrawer(const TWeakPtr<FSequencer>&
 
 FSequencerSelectionDrawer::~FSequencerSelectionDrawer()
 {
-	if (const TSharedPtr<ISequencer> Sequencer = WeakSequencer.Pin())
-	{
-		Sequencer->OnActorAddedToSequencer().RemoveAll(this);
-		Sequencer->OnMovieSceneDataChanged().RemoveAll(this);
-	}
+	UnbindEvents();
 }
 
 FName FSequencerSelectionDrawer::GetUniqueId() const
@@ -130,36 +126,9 @@ FText FSequencerSelectionDrawer::GetSectionDisplayText() const
 
 TSharedRef<SWidget> FSequencerSelectionDrawer::CreateContentWidget()
 {
-	if (const TSharedPtr<ISequencer> Sequencer = WeakSequencer.Pin())
-	{
-		Sequencer->OnActorAddedToSequencer().AddLambda([this](AActor* InActor, const FGuid InGuid)
-			{
-				OnSequencerSelectionChanged();
-			});
+	BindEvents();
 
-		Sequencer->OnMovieSceneDataChanged().AddLambda([this](const EMovieSceneDataChangeType InChangeType)
-			{
-				// Update when a channel is overriden with a curve extension
-				if (InChangeType == EMovieSceneDataChangeType::MovieSceneStructureItemsChanged)
-				{
-					OnSequencerSelectionChanged();
-				}
-			});
-
-		Sequencer->OnCloseEvent().AddLambda([this](const TSharedRef<ISequencer> InSequencer)
-			{
-				ResetContent();
-			});
-
-		if (const TSharedPtr<FSequencerSelection> SequencerSelection = Private::GetSelection(*Sequencer.Get()))
-		{
-			SequencerSelection->OnChanged.AddSP(this, &FSequencerSelectionDrawer::OnSequencerSelectionChanged);
-
-			OnSequencerSelectionChanged();
-		}
-	}
-
-	return SNew(SScrollBox)
+	const TSharedRef<SScrollBox> OutWidget = SNew(SScrollBox)
 		+ SScrollBox::Slot()
 		[
 			SNew(SBorder)
@@ -174,11 +143,73 @@ TSharedRef<SWidget> FSequencerSelectionDrawer::CreateContentWidget()
 				]
 			]
 		];
+
+	OnSequencerSelectionChanged();
+
+	return OutWidget;
+}
+
+void FSequencerSelectionDrawer::BindEvents()
+{
+	const TSharedPtr<ISequencer> Sequencer = WeakSequencer.Pin();
+	if (!Sequencer.IsValid())
+	{
+		return;
+	}
+
+	UnbindEvents();
+
+	ActorAddedToSequencerDelegate = Sequencer->OnActorAddedToSequencer().AddLambda([this](AActor* InActor, const FGuid InGuid)
+		{
+			OnSequencerSelectionChanged();
+		});
+
+	MovieSceneDataChangedDelegate = Sequencer->OnMovieSceneDataChanged().AddLambda([this](const EMovieSceneDataChangeType InChangeType)
+		{
+			// Update when a channel is overriden with a curve extension
+			if (InChangeType == EMovieSceneDataChangeType::MovieSceneStructureItemsChanged)
+			{
+				OnSequencerSelectionChanged();
+			}
+		});
+
+	if (const TSharedPtr<FSequencerSelection> SequencerSelection = Private::GetSelection(*Sequencer.Get()))
+	{
+		SequencerSelection->OnChanged.AddSP(this, &FSequencerSelectionDrawer::OnSequencerSelectionChanged);
+	}
+}
+
+void FSequencerSelectionDrawer::UnbindEvents()
+{
+	const TSharedPtr<ISequencer> Sequencer = WeakSequencer.Pin();
+	if (!Sequencer.IsValid())
+	{
+		return;
+	}
+
+	if (ActorAddedToSequencerDelegate.IsValid())
+	{
+		Sequencer->OnActorAddedToSequencer().Remove(ActorAddedToSequencerDelegate);
+		ActorAddedToSequencerDelegate.Reset();
+	}
+	if (MovieSceneDataChangedDelegate.IsValid())
+	{
+		Sequencer->OnMovieSceneDataChanged().Remove(MovieSceneDataChangedDelegate);
+		MovieSceneDataChangedDelegate.Reset();
+	}
+
+	if (const TSharedPtr<FSequencerSelection> SequencerSelection = Private::GetSelection(*Sequencer.Get()))
+	{
+		SequencerSelection->OnChanged.RemoveAll(this);
+	}
 }
 
 void FSequencerSelectionDrawer::ResetContent()
 {
-	ContentBox->ClearChildren();
+	if (ContentBox.IsValid())
+	{
+		ContentBox->ClearChildren();
+	}
 
 	CurveChannelExtension.Reset();
 
@@ -187,7 +218,7 @@ void FSequencerSelectionDrawer::ResetContent()
 
 void FSequencerSelectionDrawer::OnSequencerSelectionChanged()
 {
-	if (!ContentBox.IsValid())
+	if (!ContentBox.IsValid() || !GEditor)
 	{
 		return;
 	}
@@ -199,9 +230,13 @@ void FSequencerSelectionDrawer::OnSequencerSelectionChanged()
 
 	bWaitingToHandleSelectionChanged = true;
 
-	GEditor->GetTimerManager()->SetTimerForNextTick([this]()
+	const TWeakPtr<FSequencerSelectionDrawer> ThisWeak = SharedThis(this);
+	GEditor->GetTimerManager()->SetTimerForNextTick([ThisWeak]()
 		{
-			UpdateFromSelectionNextFrame();
+			if (const TSharedPtr<FSequencerSelectionDrawer> This = ThisWeak.Pin())
+			{
+				This->UpdateFromSelectionNextFrame();
+			}
 		});
 }
 
@@ -227,54 +262,67 @@ void FSequencerSelectionDrawer::UpdateFromSelectionNextFrame()
 
 	auto AddToContent = [this](const TSharedRef<SWidget>& InWidget)
 		{
-			ContentBox->AddSlot()
-				.AutoHeight()
-				[
-					InWidget
-				];
+			if (ContentBox.IsValid())
+			{
+				ContentBox->AddSlot()
+					.AutoHeight()
+					[
+						InWidget
+					];
+			}
 		};
 
-	const ISequencerModule& SequencerModule = FModuleManager::Get().LoadModuleChecked<ISequencerModule>(TEXT("Sequencer"));
-	const TSharedPtr<FExtensibilityManager> SidebarExtensibilityManager = SequencerModule.GetSidebarExtensibilityManager();
+	if (SelectionRef->Outliner.Num() > 0
+		|| SelectionRef->TrackArea.Num() > 0
+		|| SelectionRef->KeySelection.Num() > 0
+		|| SelectionRef->MarkedFrames.Num() > 0)
+	{
+		const ISequencerModule& SequencerModule = FModuleManager::Get().LoadModuleChecked<ISequencerModule>(TEXT("Sequencer"));
+		const TSharedPtr<FExtensibilityManager> SidebarExtensibilityManager = SequencerModule.GetSidebarExtensibilityManager();
+
+		FMenuBuilder MenuBuilder(/*bInShouldCloseWindowAfterMenuSelection=*/false
+			, Sequencer->GetCommandBindings()
+			, SidebarExtensibilityManager->GetAllExtenders()
+			, /*bInCloseSelfOnly=*/true, &FCoreStyle::Get(), /*bInSearchable=*/true, TEXT("Sequencer.Sidebar"));
+
+		/**
+		 * Selection details display order preference:
+		 *  1) Key items
+		 *  2) Track area items (if no key selected)
+		 *  3) Outliner items (if no key or track area selected)
+		 *  4) Marked frames
+		 */
+
+		// 1) Key items
+		BuildKeySelectionDetails(SelectionRef, MenuBuilder);
+
+		// Early out for key selections
+		const bool bIsKeySelected = SequencerSelection->KeySelection.Num() > 0;
+		if (bIsKeySelected)
+		{
+			AddToContent(MenuBuilder.MakeWidget());
+			return;
+		}
 
-	FMenuBuilder MenuBuilder(/*bInShouldCloseWindowAfterMenuSelection=*/false
-		, Sequencer->GetCommandBindings()
-		, SidebarExtensibilityManager->GetAllExtenders()
-		, /*bInCloseSelfOnly=*/true, &FCoreStyle::Get(), /*bInSearchable=*/true, TEXT("Sequencer.Sidebar"));
+		// 2) Track area items
+		BuildTrackAreaDetails(SelectionRef, MenuBuilder);
 
-	/**
-	 * Selection details display order preference:
-	 *  1) Key items
-	 *  2) Track area items (if no key selected)
-	 *  3) Outliner items (if no key or track area selected)
-	 *  4) Marked frames
-	 */
+		// 3) Outliner items
+		const bool bIsTrackAreaSelected = SequencerSelection->TrackArea.Num() > 0;
+		if (!bIsTrackAreaSelected)
+		{
+			BuildOutlinerDetails(SelectionRef, MenuBuilder);
+		}
 
-	// 1) Key items
-	BuildKeySelectionDetails(SelectionRef, MenuBuilder);
+		// 4) Marked frames
+		BuildMarkedFrameDetails(SelectionRef, MenuBuilder);
 
-	// Early out for key selections
-	const bool bIsKeySelected = SequencerSelection->KeySelection.Num() > 0;
-	if (bIsKeySelected)
-	{
 		AddToContent(MenuBuilder.MakeWidget());
-		return;
 	}
-
-	// 2) Track area items
-	BuildTrackAreaDetails(SelectionRef, MenuBuilder);
-
-	// 3) Outliner items
-	const bool bIsTrackAreaSelected = SequencerSelection->TrackArea.Num() > 0;
-	if (!bIsTrackAreaSelected)
+	else
 	{
-		BuildOutlinerDetails(SelectionRef, MenuBuilder);
+		AddToContent(CreateNoSelectionHintText());
 	}
-
-	// 4) Marked frames
-	BuildMarkedFrameDetails(SelectionRef, MenuBuilder);
-
-	AddToContent(MenuBuilder.MakeWidget());
 }
 
 void FSequencerSelectionDrawer::BuildKeySelectionDetails(const TSharedRef<FSequencerSelection>& InSelection, FMenuBuilder& MenuBuilder)
diff --git a/Engine/Source/Editor/Sequencer/Private/Widgets/Sidebar/SequencerSelectionDrawer.h b/Engine/Source/Editor/Sequencer/Private/Widgets/Sidebar/SequencerSelectionDrawer.h
--- a/Engine/Source/Editor/Sequencer/Private/Widgets/Sidebar/SequencerSelectionDrawer.h
+++ b/Engine/Source/Editor/Sequencer/Private/Widgets/Sidebar/SequencerSelectionDrawer.h
@@ -38,6 +38,9 @@ public:
 	//~ End ISidebarDrawerContent
 
 protected:
+	void BindEvents();
+	void UnbindEvents();
+
 	void OnSequencerSelectionChanged();
 	void UpdateFromSelectionNextFrame();
 
@@ -67,4 +70,7 @@ protected:
 	TArray<TSharedPtr<ISidebarChannelExtension>> ChannelExtensions;
 
 	bool bWaitingToHandleSelectionChanged = false;
+
+	FDelegateHandle ActorAddedToSequencerDelegate;
+	FDelegateHandle MovieSceneDataChangedDelegate;
 };
diff --git a/Engine/Source/Editor/Sequencer/Public/SequencerCommands.h b/Engine/Source/Editor/Sequencer/Public/SequencerCommands.h
--- a/Engine/Source/Editor/Sequencer/Public/SequencerCommands.h
+++ b/Engine/Source/Editor/Sequencer/Public/SequencerCommands.h
@@ -486,6 +486,9 @@ public:
 	/** Toggles the visibility of the Sequencer sidebar */
 	TSharedPtr< FUICommandInfo > ToggleSidebarVisible;
 
+	/** Toggles the state of the Sequencer selection details between docked in the Sequencer tab and being popped out to its own tab window */
+	TSharedPtr< FUICommandInfo > ToggleSelectionDetailsPoppedOut;
+
 	/** Toggles the sidebar "Selection" drawer open or closed */
 	TSharedPtr< FUICommandInfo > ToggleSidebarSelectionDrawerOpen;
 
diff --git a/Engine/Source/Editor/Sequencer/Public/SequencerSettings.h b/Engine/Source/Editor/Sequencer/Public/SequencerSettings.h
--- a/Engine/Source/Editor/Sequencer/Public/SequencerSettings.h
+++ b/Engine/Source/Editor/Sequencer/Public/SequencerSettings.h
@@ -587,6 +587,11 @@ public:
 	/** Sets the visibility of the Sequence Navigator */
 	UE_API void SetNavigationToolVisible(const bool bInVisible);
 
+	/** @return True if the selection details are popped out from the Sequencer tab to their own tab */
+	UE_API bool IsSelectionDetailsPoppedOut() const { return bSelectionDetailsPoppedOut; }
+	/** Sets the popped out state of the selection details in the Sequencer tab */
+	UE_API void SetSelectionDetailsPoppedOut(const bool bInPoppedOut);
+
 protected:
 
 	/** The auto change mode (auto-key, auto-track or none). */
@@ -942,6 +947,10 @@ protected:
 	UPROPERTY(config, EditAnywhere, Category = General, meta = (DisplayName = "Sequence Navigator Visible"))
 	bool bNavigationToolVisible;
 
+	/** Last saved popped out state of the selection details */
+	UPROPERTY(config, EditAnywhere, Category = General)
+	bool bSelectionDetailsPoppedOut;
+
 	UFUNCTION()
 	static UE_API bool ShouldShowThumbnailCaptureSettings();
 
diff --git a/Engine/Source/Editor/Sequencer/Sequencer.Build.cs b/Engine/Source/Editor/Sequencer/Sequencer.Build.cs
--- a/Engine/Source/Editor/Sequencer/Sequencer.Build.cs
+++ b/Engine/Source/Editor/Sequencer/Sequencer.Build.cs
@@ -58,6 +58,11 @@ public class Sequencer : ModuleRules
 				}
 			);
 
+		if (Target.bBuildEditor)
+		{
+			PrivateDependencyModuleNames.Add("WorkspaceMenuStructure");
+		}
+
 		CircularlyReferencedDependentModules.AddRange(
 			new string[] {
 				"ViewportInteraction",
