[Sequencer] Add pin state changed delegate that will be used by the simplified timeline to update displayed keys when a tracks pin state changes

#Synced-CL 50019235

diff --git a/Engine/Source/Editor/Sequencer/Private/MVVM/ViewModels/OutlinerItemModel.cpp b/Engine/Source/Editor/Sequencer/Private/MVVM/ViewModels/OutlinerItemModel.cpp
--- a/Engine/Source/Editor/Sequencer/Private/MVVM/ViewModels/OutlinerItemModel.cpp
+++ b/Engine/Source/Editor/Sequencer/Private/MVVM/ViewModels/OutlinerItemModel.cpp
@@ -165,6 +165,8 @@ void FOutlinerItemModelMixin::AddDisplayOptionsPropertyMenuItem(FMenuBuilder& In
 	);
 }
 
+FOutlinerItemModelMixin::FOnPinStateChanged FOutlinerItemModelMixin::PinStateChangedDelegate;
+
 FOutlinerItemModelMixin::FOutlinerItemModelMixin()
 	: OutlinerChildList(EViewModelListType::Outliner)
 	, bInitializedExpansion(false)
@@ -314,6 +316,13 @@ bool FOutlinerItemModelMixin::IsPinned() const
 	return FPinnableExtensionShim::IsPinned();
 }
 
+void FOutlinerItemModelMixin::SetPinned(const bool bInIsPinned)
+{
+	FPinnableExtensionShim::SetPinned(bInIsPinned);
+
+	PinStateChangedDelegate.Broadcast(AsViewModel()->CastThisShared<IPinnableExtension>(), bInIsPinned);
+}
+
 bool FOutlinerItemModelMixin::IsDimmed() const
 {
 	const FViewModel* const ViewModel = AsViewModel();
diff --git a/Engine/Source/Editor/Sequencer/Public/MVVM/ViewModels/OutlinerItemModel.h b/Engine/Source/Editor/Sequencer/Public/MVVM/ViewModels/OutlinerItemModel.h
--- a/Engine/Source/Editor/Sequencer/Public/MVVM/ViewModels/OutlinerItemModel.h
+++ b/Engine/Source/Editor/Sequencer/Public/MVVM/ViewModels/OutlinerItemModel.h
@@ -70,6 +70,11 @@ public:
 
 	/*~ IPinnableExtension */
 	UE_API bool IsPinned() const override;
+	UE_API void SetPinned(const bool bInIsPinned) override;
+
+	DECLARE_MULTICAST_DELEGATE_TwoParams(FOnPinStateChanged, const TViewModelPtr<IPinnableExtension>& /*InPinnableExtension*/, const bool /*bInPinned*/)
+	/** Delegate that fires when the pin state of this item changes */
+	static UE_API FOnPinStateChanged& OnPinStateChanged() { return PinStateChangedDelegate; }
 
 	/*~ IDimmableExtension */
 	UE_API bool IsDimmed() const override;
@@ -129,6 +134,8 @@ private:
 	FCurveEditorTreeItemID CurveEditorItemID;
 	mutable bool bInitializedExpansion;
 	mutable bool bInitializedPinnedState;
+
+	static FOnPinStateChanged PinStateChangedDelegate;
 };
 
 //
