Sequencer: Extract filter row widgets from SSequencer to separate widget class.
This prepares the widget to be instantiated per FSequencerFilterBar instance in follow up changes.

New classes:
- SSequencerSearchAndFilterRow: Knows about all the widgets in the row above the tree view (except the filter bar pills - that stays as it is)
- SSearchAndFilterWidget: Specifically holds the filtering widgets, i.e. search, isolate/hide/show, and similar widgets. In the future, we'll add a button for switching between linked filtering here.

This makes SSequencerSearchBox update its display text when the underlying ISequencerFilterBar search text changes.
This was previously done by SSequencerFilterBar. It makes more sense for the SSequencerSearchBox to handle it so it can stand by itself.

#Synced-CL 50040876

diff --git a/Engine/Source/Editor/Sequencer/Private/Filters/Widgets/SSearchAndFilterWidget.cpp b/Engine/Source/Editor/Sequencer/Private/Filters/Widgets/SSearchAndFilterWidget.cpp
new file mode 100644
--- /dev/null
+++ b/Engine/Source/Editor/Sequencer/Private/Filters/Widgets/SSearchAndFilterWidget.cpp
@@ -0,0 +1,86 @@
+﻿// Copyright Epic Games, Inc. All Rights Reserved.
+
+#include "SSearchAndFilterWidget.h"
+
+#include "Sequencer.h"
+#include "Filters/Menus/SequencerViewOptionsMenu.h"
+#include "Filters/SequencerFilterBar.h"
+#include "Filters/Widgets/SFilterBarIsolateHideShow.h"
+#include "Filters/Widgets/SSequencerSearchBox.h"
+#include "SPositiveActionButton.h"
+#include "SSequencerFilterBar.h"
+#include "Styling/AppStyle.h"
+
+#define LOCTEXT_NAMESPACE "SSearchAndFilterRow"
+
+namespace UE::Sequencer
+{
+void SSearchAndFilterWidget::Construct(
+	const FArguments& InArgs, const TSharedRef<FSequencerFilterBar>& InFilterBar
+	)
+{
+	FilterBar = InFilterBar;
+	OnSearchTextChanged = InArgs._OnSearchTextChanged;
+	
+	ChildSlot
+	[
+		SNew(SHorizontalBox)
+
+		// Advanced Search Filter Combo Button
+		+ SHorizontalBox::Slot()
+		.AutoWidth()
+		.VAlign(VAlign_Center)
+		.Padding(2.f, 0.f, 0.f, 0.f)
+		[
+			InFilterBar->MakeAddFilterButton()
+		]
+
+		// Advanced Search Box
+		+ SHorizontalBox::Slot()
+		.VAlign(VAlign_Center)
+		.Padding(2.f, 0.f, 0.f, 0.f)
+		[
+			SNew(SSequencerSearchBox, InFilterBar)
+			.AddMetaData<FTagMetaData>(FTagMetaData(TEXT("SequencerFilterSearch")))
+			.HintText(LOCTEXT("FilterSearch", "Search..."))
+			.ToolTipText(LOCTEXT("FilterSearchHint", "Type here to search"))
+			.OnTextChanged(this, &SSearchAndFilterWidget::OnOutlinerSearchChanged)
+			.OnTextCommitted(this, &SSearchAndFilterWidget::OnOutlinerSearchCommitted)
+			.OnSaveSearchClicked(this, &SSearchAndFilterWidget::OnOutlinerSearchSaved)
+		]
+
+		// Isolate / Hide / Show
+		+ SHorizontalBox::Slot()
+		.AutoWidth()
+		.VAlign(VAlign_Center)
+		.Padding(2.f, 0.f, 0.f, 0.f)
+		[
+			InFilterBar->MakeIsolateHideShowPanel()
+		]
+	];
+}
+
+void SSearchAndFilterWidget::SetSearchText(const FText& InSearchText)
+{
+	OnSearchTextChanged.ExecuteIfBound(InSearchText);
+}
+
+void SSearchAndFilterWidget::OnOutlinerSearchChanged(const FText& InSearchText)
+{
+	SetSearchText(InSearchText);
+}
+
+void SSearchAndFilterWidget::OnOutlinerSearchCommitted(const FText& InSearchText, ETextCommit::Type Arg)
+{
+	SetSearchText(InSearchText);
+}
+
+void SSearchAndFilterWidget::OnOutlinerSearchSaved(const FText& InSearchText)
+{
+	FCustomTextFilterData CustomTextFilterData;
+	CustomTextFilterData.FilterString = InSearchText;
+	SSequencerCustomTextFilterDialog::CreateWindow_AddCustomTextFilter(FilterBar.ToSharedRef(), CustomTextFilterData);
+}
+} // namespace UE::Sequencer
+
+#undef LOCTEXT_NAMESPACE
\ No newline at end of file
diff --git a/Engine/Source/Editor/Sequencer/Private/Filters/Widgets/SSearchAndFilterWidget.h b/Engine/Source/Editor/Sequencer/Private/Filters/Widgets/SSearchAndFilterWidget.h
new file mode 100644
--- /dev/null
+++ b/Engine/Source/Editor/Sequencer/Private/Filters/Widgets/SSearchAndFilterWidget.h
@@ -0,0 +1,42 @@
+﻿// Copyright Epic Games, Inc. All Rights Reserved.
+
+#pragma once
+
+#include "SSequencerCustomTextFilterDialog.h"
+#include "Widgets/DeclarativeSyntaxSupport.h"
+#include "Widgets/SCompoundWidget.h"
+
+class FSequencerFilterBar;
+class ISequencerFilterBar;
+class SSequencerSearchBox;
+
+namespace UE::Sequencer
+{
+/** Contains widgets for searching and filtering the Sequencer tree view. */
+class SSearchAndFilterWidget : public SCompoundWidget
+{
+public:
+	
+	SLATE_BEGIN_ARGS(SSearchAndFilterWidget){}
+		/** Invoked when the search text changes. */
+		SLATE_EVENT(FOnTextChanged, OnSearchTextChanged)
+	SLATE_END_ARGS()
+
+	void Construct(const FArguments& InArgs, const TSharedRef<FSequencerFilterBar>& InFilterBar);
+
+private:
+	
+	/** The filter bar that is being displayed. */
+	TSharedPtr<FSequencerFilterBar> FilterBar;
+	
+	/** Invoked when the search text changes. */
+	FOnTextChanged OnSearchTextChanged;
+	
+	void SetSearchText(const FText& InSearchText);
+	
+	void OnOutlinerSearchChanged(const FText& InSearchText);
+	void OnOutlinerSearchCommitted(const FText& InSearchText, ETextCommit::Type Arg);
+	void OnOutlinerSearchSaved(const FText& InSearchText);
+};
+} // namespace UE::Sequencer
+
diff --git a/Engine/Source/Editor/Sequencer/Private/Filters/Widgets/SSequencerFilterBar.cpp b/Engine/Source/Editor/Sequencer/Private/Filters/Widgets/SSequencerFilterBar.cpp
--- a/Engine/Source/Editor/Sequencer/Private/Filters/Widgets/SSequencerFilterBar.cpp
+++ b/Engine/Source/Editor/Sequencer/Private/Filters/Widgets/SSequencerFilterBar.cpp
@@ -87,7 +87,6 @@ void SSequencerFilterBar::Construct(const FArguments& InArgs, const TSharedRef<F
 
 	InFilterBar->OnFiltersChanged().AddSP(this, &SSequencerFilterBar::OnFiltersChanged);
 	InFilterBar->OnMuteFiltersChanged().AddSP(this, &SSequencerFilterBar::OnMuteStateChanged);
-	InFilterBar->OnTextFilterTextChanged().AddSP(this, &SSequencerFilterBar::OnTextFilterChanged);
 }
 
 FReply SSequencerFilterBar::OnMouseButtonUp(const FGeometry& MyGeometry, const FPointerEvent& MouseEvent)
@@ -119,30 +118,6 @@ TSharedPtr<FSequencerFilterBar> SSequencerFilterBar::GetFilterBar() const
 	return WeakFilterBar.Pin();
 }
 
-void SSequencerFilterBar::SetTextFilterString(const FString& InText)
-{
-	const TSharedPtr<FSequencerFilterBar> FilterBar = WeakFilterBar.Pin();
-	if (!FilterBar.IsValid())
-	{
-		return;
-	}
-
-	if (const TSharedPtr<SSequencerSearchBox> SearchBox = WeakSearchBox.Pin())
-	{
-		if (!InText.Equals(SearchBox->GetText().ToString()))
-		{
-			SearchBox->SetText(FText::FromString(InText));
-		}
-
-		SearchBox->SetError(FilterBar->GetFilterErrorText());
-	}
-}
-
-FText SSequencerFilterBar::GetFilterErrorText() const
-{
-	return WeakFilterBar.IsValid() ? WeakFilterBar.Pin()->GetFilterErrorText() : FText::GetEmpty();
-}
-
 EFilterBarLayout SSequencerFilterBar::GetLayout() const
 {
 	return FilterBarLayout;
@@ -495,14 +470,6 @@ void SSequencerFilterBar::OnMuteStateChanged(bool bNewIsMuted)
 	}
 }
 
-void SSequencerFilterBar::OnTextFilterChanged()
-{
-	if (const TSharedPtr<FSequencerFilterBar> FilterBar = GetFilterBar())
-	{
-		SetTextFilterString(FilterBar->GetTextFilterString());
-	}
-}
-
 void SSequencerFilterBar::CreateAddCustomTextFilterWindowFromSearch(const FText& InSearchText)
 {
 	const TSharedPtr<FSequencerFilterBar> FilterBar = WeakFilterBar.Pin();
diff --git a/Engine/Source/Editor/Sequencer/Private/Filters/Widgets/SSequencerFilterBar.h b/Engine/Source/Editor/Sequencer/Private/Filters/Widgets/SSequencerFilterBar.h
--- a/Engine/Source/Editor/Sequencer/Private/Filters/Widgets/SSequencerFilterBar.h
+++ b/Engine/Source/Editor/Sequencer/Private/Filters/Widgets/SSequencerFilterBar.h
@@ -55,10 +55,6 @@ public:
 
 	TSharedPtr<FSequencerFilterBar> GetFilterBar() const;
 
-	void SetTextFilterString(const FString& InFilterString);
-
-	FText GetFilterErrorText() const;
-
 	EFilterBarLayout GetLayout() const;
 	void SetLayout(const EFilterBarLayout InFilterBarLayout);
 
@@ -92,8 +88,6 @@ protected:
 	void OnFiltersChanged(const ESequencerFilterChange InChangeType, const TSharedRef<FSequencerTrackFilter>& InFilter);
 	/** Updates the mute state in response to the owning FSequencerTrackFilter changing mute state. */
 	void OnMuteStateChanged(bool bNewIsMuted);
-	/** Updates the search text */
-	void OnTextFilterChanged();
 
 	void CreateFilterWidgetsFromConfig();
 
diff --git a/Engine/Source/Editor/Sequencer/Private/Filters/Widgets/SSequencerSearchBox.cpp b/Engine/Source/Editor/Sequencer/Private/Filters/Widgets/SSequencerSearchBox.cpp
--- a/Engine/Source/Editor/Sequencer/Private/Filters/Widgets/SSequencerSearchBox.cpp
+++ b/Engine/Source/Editor/Sequencer/Private/Filters/Widgets/SSequencerSearchBox.cpp
@@ -84,6 +84,10 @@ void SSequencerSearchBox::Construct(const FArguments& InArgs, const TWeakPtr<ISe
 			]
 			.MenuContent(GetSuggestionListMenuContent())
 		];
+	
+	const TSharedPtr<ISequencerFilterBar> FilterBarPin = WeakFilterBar.Pin();
+	check(FilterBarPin);
+	FilterBarPin->OnTextFilterTextChanged().AddSP(this, &SSequencerSearchBox::OnFilterTextChanged);
 }
 
 TSharedRef<SWidget> SSequencerSearchBox::GetSuggestionListMenuContent()
@@ -641,6 +645,15 @@ end:
 	}
 }
 
+void SSequencerSearchBox::OnFilterTextChanged()
+{
+	if (const TSharedPtr<ISequencerFilterBar> FilterBarPin = WeakFilterBar.Pin())
+	{
+		const FText NewText = FText::FromString(FilterBarPin->GetTextFilterString());
+		SetText(NewText);
+	}
+}
+
 bool SSequencerSearchBox::IsOperatorToken(const FExpressionToken& InToken)
 {
 	return InToken.Node.Cast<TextFilterExpressionParser::FEqual>()
diff --git a/Engine/Source/Editor/Sequencer/Private/SSequencer.cpp b/Engine/Source/Editor/Sequencer/Private/SSequencer.cpp
--- a/Engine/Source/Editor/Sequencer/Private/SSequencer.cpp
+++ b/Engine/Source/Editor/Sequencer/Private/SSequencer.cpp
@@ -145,6 +145,7 @@
 #include "Filters/Widgets/SSequencerSearchBox.h"
 #include "IContentBrowserSingleton.h"
 #include "ContentBrowserModule.h"
+#include "Widgets/Sequencer/SSequencerSearchAndFilterRow.h"
 #include "Menus/SequencerToolbarUtils.h"
 #include "MVVM/Extensions/ITopTimeSliderOverlayExtension.h"
 #include "MVVM/Extensions/ITrackAreaOverlayExtension.h"
@@ -282,8 +283,6 @@ void SSequencer::Construct(const FArguments& InArgs, TSharedRef<FSequencer> InSe
 
 	GridPanel = ConstructTrackAreaGridPanel(InArgs, TimeSliderArgs);
 
-	ViewOptionsMenu = MakeShared<FSequencerViewOptionsMenu>();
-
 	if (const TSharedPtr<FSequencerFilterBar> FilterBar = GetFilterBar())
 	{
 		FilterBar->OnStateChanged().AddSP(this, &SSequencer::OnFilterBarStateChanged);
@@ -748,6 +747,10 @@ TSharedRef<SGridPanel> SSequencer::ConstructTrackAreaGridPanel(const FArguments&
 	USequencerToolMenuContext* ContextObject = NewObject<USequencerToolMenuContext>();
 	ContextObject->WeakSequencer = SequencerRef;
 	
+	FilterRow = SNew(SSequencerSearchAndFilterRow)
+		.ExtenderForAddMenu_Lambda([this]{ return FExtender::Combine(AddMenuExtenders); })
+		.Sequencer(SequencerPtr);
+	
 	return SNew(SGridPanel)
 		.FillRow(2, 1.f)
 		.FillColumn(0, FillCoefficient_0)
@@ -1381,109 +1384,6 @@ void SSequencer::HandleOutlinerNodeSelectionChanged()
 	}
 }
 
-TSharedRef<SWidget> SSequencer::ConstructSearchAndFilterRow()
-{
-	const TSharedPtr<FSequencer> Sequencer = SequencerPtr.Pin();
-	if (!Sequencer.IsValid())
-	{
-		return SNullWidget::NullWidget;
-	}
-
-	const TSharedPtr<FSequencerFilterBar> FilterBar = GetFilterBar();
-	if (!FilterBar.IsValid())
-	{
-		return SNullWidget::NullWidget;
-	}
-
-	if (!SearchBox.IsValid())
-	{
-		SearchBox = SNew(SSequencerSearchBox, FilterBar)
-			.AddMetaData<FTagMetaData>(FTagMetaData(TEXT("SequencerFilterSearch")))
-			.HintText(LOCTEXT("FilterSearch", "Search..."))
-			.ToolTipText(LOCTEXT("FilterSearchHint", "Type here to search"))
-			.OnTextChanged(this, &SSequencer::OnOutlinerSearchChanged)
-			.OnTextCommitted(this, &SSequencer::OnOutlinerSearchCommitted)
-			.OnSaveSearchClicked(this, &SSequencer::OnOutlinerSearchSaved);
-	}
-
-	FilterBarWidget = SNew(SSequencerFilterBar, FilterBar.ToSharedRef())
-		.FilterBarLayout(GetFilterBarLayout())
-		.AddMetaData<FTagMetaData>(FTagMetaData(TEXT("SequencerTrackFilters")))
-		.FilterSearchBox(SearchBox);
-
-	if (!FilterComboButtonWidget.IsValid())
-	{
-		FilterComboButtonWidget = FilterBar->MakeAddFilterButton();
-	}
-
-	FilterBarWidget->SetMuted(FilterBar->AreFiltersMuted());
-
-	return SNew(SHorizontalBox)
-
-		// Add Button
-		+ SHorizontalBox::Slot()
-		.AutoWidth()
-		.VAlign(VAlign_Center)
-		[
-			SNew(SPositiveActionButton)
-			.OnGetMenuContent(this, &SSequencer::MakeAddMenu)
-			.Icon(FAppStyle::Get().GetBrush("Icons.Plus"))
-			.Text(LOCTEXT("Add", "Add"))
-			.IsEnabled_Lambda([this]() { return !SequencerPtr.Pin()->IsReadOnly(); })
-		]
-
-		// Advanced Search Filter Combo Button
-		+ SHorizontalBox::Slot()
-		.AutoWidth()
-		.VAlign(VAlign_Center)
-		.Padding(2.f, 0.f, 0.f, 0.f)
-		[
-			FilterComboButtonWidget.ToSharedRef()
-		]
-
-		// Advanced Search Box
-		+ SHorizontalBox::Slot()
-		.VAlign(VAlign_Center)
-		.Padding(2.f, 0.f, 0.f, 0.f)
-		[
-			SearchBox.ToSharedRef()
-		]
-
-		// Isolate / Hide / Show
-		+ SHorizontalBox::Slot()
-		.AutoWidth()
-		.VAlign(VAlign_Center)
-		.Padding(2.f, 0.f, 0.f, 0.f)
-		[
-			FilterBar->MakeIsolateHideShowPanel()
-		]
-
-		// View Options Combo Button
-		+ SHorizontalBox::Slot()
-		.AutoWidth()
-		.VAlign(VAlign_Center)
-		.Padding(2.f, 0.f, 0.f, 0.f)
-		[
-			SAssignNew(ViewOptionsComboButton, SComboButton)
-			.ContentPadding(2.f)
-			.HAlign(HAlign_Center)
-			.VAlign(VAlign_Center)
-			.ToolTipText(LOCTEXT("ViewOptionsToolTip", "View Options"))
-			.ComboButtonStyle(FAppStyle::Get(), TEXT("SimpleComboButtonWithIcon"))
-			.OnGetMenuContent_Lambda([this]() -> TSharedRef<SWidget>
-				{
-					return ViewOptionsMenu->CreateMenu(SequencerPtr);
-				})
-			.HasDownArrow(false)
-			.ButtonContent()
-			[
-				SNew(SImage)
-				.ColorAndOpacity(FSlateColor::UseForeground())
-				.Image(FAppStyle::Get().GetBrush(TEXT("Icons.Settings")))
-			]
-		];
-}
-
 namespace UE::Sequencer::ToolMenus
 {
 const FLazyName ToolbarMenuName("Sequencer.MainToolBar");
@@ -2892,48 +2792,6 @@ void SSequencer::PopBreadcrumb()
 	BreadcrumbTrail->PopCrumb();
 }
 
-FText SSequencer::GetSearchText() const
-{
-	if (FilterBarWidget.IsValid())
-	{
-		if (const TSharedPtr<FSequencerFilterBar> FilterBar = FilterBarWidget->GetFilterBar())
-		{
-			return FText::FromString(FilterBar->GetTextFilterString());
-		}
-	}
-	return FText::GetEmpty();
-}
-
-void SSequencer::SetSearchText(const FText& InSearchText)
-{
-	if (const TSharedPtr<FSequencer> Sequencer = SequencerPtr.Pin())
-	{
-		Sequencer->GetNodeTree()->SetTextFilterString(InSearchText.ToString());
-
-		TreeView->Refresh();
-	}
-}
-
-void SSequencer::OnOutlinerSearchChanged(const FText& InFilter)
-{
-	SetSearchText(InFilter);
-}
-
-void SSequencer::OnOutlinerSearchCommitted(const FText& InFilter, ETextCommit::Type InCommitInfo)
-{
-	SetSearchText(InFilter);
-}
-
-void SSequencer::OnOutlinerSearchSaved(const FText& InFilterText)
-{
-	if (FilterBarWidget.IsValid())
-	{
-		FCustomTextFilterData CustomTextFilterData;
-		CustomTextFilterData.FilterString = InFilterText;
-		SSequencerCustomTextFilterDialog::CreateWindow_AddCustomTextFilter(FilterBarWidget->GetFilterBar().ToSharedRef(), CustomTextFilterData);
-	}
-}
-
 void SSequencer::OnDragEnter( const FGeometry& MyGeometry, const FDragDropEvent& DragDropEvent )
 {
 	// @todo sequencer: Add drop validity cue
@@ -3066,21 +2924,16 @@ FReply SSequencer::OnDrop( const FGeometry& MyGeometry, const FDragDropEvent& Dr
 FReply SSequencer::OnKeyDown( const FGeometry& MyGeometry, const FKeyEvent& InKeyEvent ) 
 {
 	// A toolkit tab is active, so direct all command processing to it
-	TSharedPtr<FSequencer> SequencerPin = SequencerPtr.Pin();
-	if (SequencerPin.IsValid())
+	const TSharedPtr<FSequencer> SequencerPin = SequencerPtr.Pin();
+	if (SequencerPin && SequencerPin->GetCommandBindings()->ProcessCommandBindings(InKeyEvent))
 	{
-		if (SequencerPin->GetCommandBindings()->ProcessCommandBindings(InKeyEvent))
-		{
-			return FReply::Handled();
-		}
+		return FReply::Handled();
 	}
 
-	if (const TSharedPtr<FSequencerFilterBar> FilterBar = FilterBarWidget->GetFilterBar())
+	const TSharedPtr<FSequencerFilterBar> FilterBar = FilterBarWidget->GetFilterBar();
+	if (FilterBar && FilterBar->GetCommandList()->ProcessCommandBindings(InKeyEvent))
 	{
-		if (FilterBar->GetCommandList()->ProcessCommandBindings(InKeyEvent))
-		{
-			return FReply::Handled();
-		}
+		return FReply::Handled();
 	}
 
 	return FReply::Unhandled();
@@ -4277,34 +4130,21 @@ TSharedPtr<FSequencerFilterBar> SSequencer::GetFilterBar() const
 	return nullptr;
 }
 
-TSharedPtr<SSequencerFilterBar> SSequencer::GetFilterBarWidget() const
-{
-	return FilterBarWidget;
-}
-
 bool SSequencer::IsFilterBarVisible() const
 {
 	if (!FilterBarWidget.IsValid())
 	{
 		return false;
 	}
-
+	
 	USequencerSettings* const SequencerSettings = GetSequencerSettings();
-	if (ensure(SequencerSettings))
-	{
-		if (!SequencerSettings->IsFilterBarVisible())
-		{
-			return false;
-		}
-	}
-
-	const TSharedPtr<FSequencerFilterBar> FilterBar = FilterBarWidget->GetFilterBar();
-	if (!FilterBar.IsValid())
+	if (ensure(SequencerSettings) && !SequencerSettings->IsFilterBarVisible())
 	{
 		return false;
 	}
 
-	return FilterBar->HasAnyFiltersEnabled();
+	const TSharedPtr<FSequencerFilterBar> FilterBar = GetFilterBar();
+	return FilterBar && FilterBar->HasAnyFiltersEnabled();
 }
 
 void SSequencer::ToggleFilterBarVisibility()
@@ -4358,24 +4198,27 @@ void SSequencer::OnTrackFiltersChanged(const ESequencerFilterChange InChangeType
 void SSequencer::RebuildSearchAndFilterRow()
 {
 	SearchAndFilterRow->ClearChildren();
-
 	SearchAndFilterRow->AddSlot()
 		.AutoHeight()
 		[
-			ConstructSearchAndFilterRow()
+			FilterRow.ToSharedRef()
 		];
+	
+	const TSharedRef<FSequencerFilterBar> FilterBar = GetFilterBar().ToSharedRef();
+	FilterBarWidget = SNew(SSequencerFilterBar, FilterBar)
+		.FilterBarLayout(GetFilterBarLayout())
+		.AddMetaData<FTagMetaData>(FTagMetaData(TEXT("SequencerTrackFilters")))
+		.FilterSearchBox(SearchBox);
+	FilterBarWidget->SetMuted(FilterBar->AreFiltersMuted());
 
-	if (FilterBarWidget.IsValid())
+	if (IsFilterBarVisible() && GetFilterBarLayout() == EFilterBarLayout::Horizontal)
 	{
-		if (IsFilterBarVisible() && GetFilterBarLayout() == EFilterBarLayout::Horizontal)
-		{
-			SearchAndFilterRow->AddSlot()
-				.AutoHeight()
-				.Padding(0)
-				[
-					FilterBarWidget.ToSharedRef()
-				];
-		}
+		SearchAndFilterRow->AddSlot()
+			.AutoHeight()
+			.Padding(0)
+			[
+				FilterBarWidget.ToSharedRef()
+			];
 	}
 }
 
diff --git a/Engine/Source/Editor/Sequencer/Private/SSequencer.h b/Engine/Source/Editor/Sequencer/Private/SSequencer.h
--- a/Engine/Source/Editor/Sequencer/Private/SSequencer.h
+++ b/Engine/Source/Editor/Sequencer/Private/SSequencer.h
@@ -61,8 +61,9 @@ namespace UE
 {
 namespace Sequencer
 {
+class SSequencerSearchAndFilterRow;
 
-	class SOutlinerView;
+class SOutlinerView;
 	class STrackAreaView;
 	class FVirtualTrackArea;
 	class IOutlinerColumn;
@@ -346,12 +347,6 @@ public:
 	/** Sets the play time for the sequence. Will extend the working range if out of bounds. */
 	void SetPlayTime(double Frame);
 
-	/** Gets the text to search by */
-	FText GetSearchText() const;
-
-	/** Sets the text to search by */
-	void SetSearchText(const FText& InSearchText);
-
 public:
 
 	// FNotifyHook overrides
@@ -379,7 +374,6 @@ private:
 
 	TSharedRef<SWidget> ConstructSidebarContent();
 	TSharedRef<SWidget> ConstructFilterBarContent();
-	TSharedRef<SWidget> ConstructSearchAndFilterRow();
 	TSharedRef<SWidget> ConstructGridOverlayContent();
 	TSharedRef<SGridPanel> ConstructTrackAreaGridPanel(const FArguments& InArgs, const FTimeSliderArgs& InTimeSliderArgs);
 
@@ -494,15 +488,6 @@ private:
 	*/
 	void OnValueSnapIntervalChanged( float InInterval );
 
-	/**
-	 * Called when the outliner search terms change                                                              
-	 */
-	void OnOutlinerSearchChanged(const FText& InFilter);
-
-	void OnOutlinerSearchCommitted(const FText& InFilter, ETextCommit::Type InCommitInfo);
-	
-	void OnOutlinerSearchSaved(const FText& InFilterText);
-
 	/**
 	 * @return The fill percentage of the animation outliner
 	 */
@@ -704,7 +689,6 @@ public:
 	SEQUENCER_API void EnablePendingFocusOnHovering(const bool InEnabled);
 	
 	TSharedPtr<FSequencerFilterBar> GetFilterBar() const;
-	TSharedPtr<SSequencerFilterBar> GetFilterBarWidget() const;
 
 	bool IsFilterBarVisible() const;
 	void ToggleFilterBarVisibility();
@@ -916,12 +900,11 @@ private:
 
 	TSharedPtr<FSequencerSelectionDrawer> SelectionDetailsDrawer;
 
+	/** Displays the filter pills. Placed at different UI locations depending on the layout mode. */
 	TSharedPtr<SSequencerFilterBar> FilterBarWidget;
-	TSharedPtr<SComboButton> FilterComboButtonWidget;
-
-	/** The button that displays view options */
-	TSharedPtr<SComboButton> ViewOptionsComboButton;
-	TSharedPtr<FSequencerViewOptionsMenu> ViewOptionsMenu;
-
+	
+	/** Contains the filters widgets plus optionally the filter bar when in horizontal layout mode. */
 	TSharedPtr<SVerticalBox> SearchAndFilterRow;
+	/** The widgets displayed above the tree view. */
+	TSharedPtr<UE::Sequencer::SSequencerSearchAndFilterRow> FilterRow;
 };
diff --git a/Engine/Source/Editor/Sequencer/Private/Sequencer.cpp b/Engine/Source/Editor/Sequencer/Private/Sequencer.cpp
--- a/Engine/Source/Editor/Sequencer/Private/Sequencer.cpp
+++ b/Engine/Source/Editor/Sequencer/Private/Sequencer.cpp
@@ -6863,23 +6863,21 @@ void FSequencer::AddNodesToExistingNodeGroup(TArrayView<const UE::Sequencer::TWe
 
 void FSequencer::ClearFilters()
 {
-	SequencerWidget->SetSearchText(FText::GetEmpty());
+	FilterBar->SetTextFilterString(FString());
 	FilterBar->EnableFilters(false);
 	GetSequencerSettings()->SetShowSelectedNodesOnly(false);
 
 	UMovieSceneSequence* FocusedMovieSequence = GetFocusedMovieSceneSequence();
-	UMovieScene* FocusedMovieScene = nullptr;
-	if (IsValid(FocusedMovieSequence))
+	UMovieScene* FocusedMovieScene = FocusedMovieSequence ? FocusedMovieSequence->GetMovieScene() : nullptr;
+	if (FocusedMovieScene)
 	{
-		FocusedMovieScene = FocusedMovieSequence->GetMovieScene();
-		if (IsValid(FocusedMovieScene))
+		for (UMovieSceneNodeGroup* NodeGroup : FocusedMovieScene->GetNodeGroups())
 		{
-			for (UMovieSceneNodeGroup* NodeGroup : FocusedMovieScene->GetNodeGroups())
-			{
-				NodeGroup->SetEnableFilter(false);
-			}
+			NodeGroup->SetEnableFilter(false);
 		}
 	}
+	
+	NodeTree->RequestFilterUpdate();
 }
 
 void FSequencer::SynchronizeExternalSelectionWithSequencerSelection()
diff --git a/Engine/Source/Editor/Sequencer/Private/Widgets/Sequencer/SAddToSequencerButton.cpp b/Engine/Source/Editor/Sequencer/Private/Widgets/Sequencer/SAddToSequencerButton.cpp
new file mode 100644
--- /dev/null
+++ b/Engine/Source/Editor/Sequencer/Private/Widgets/Sequencer/SAddToSequencerButton.cpp
@@ -0,0 +1,65 @@
+﻿// Copyright Epic Games, Inc. All Rights Reserved.
+
+#include "SAddToSequencerButton.h"
+
+#include "Sequencer.h"
+#include "SequencerCommands.h"
+#include "SPositiveActionButton.h"
+#include "Filters/Menus/SequencerFilterMenuContext.h"
+#include "Framework/MultiBox/MultiBoxExtender.h"
+#include "MVVM/ViewModels/SequencerOutlinerViewModel.h"
+
+#define LOCTEXT_NAMESPACE "SAddToSequencerButton"
+
+class FSequencer;
+
+namespace UE::Sequencer
+{
+void SAddToSequencerButton::Construct(const FArguments& InArgs)
+{
+	WeakSequencer = InArgs._Sequencer;
+	ExtenderForAddMenuAttr = InArgs._ExtenderForAddMenu;
+	
+	ChildSlot
+	[
+		SNew(SPositiveActionButton)
+		.OnGetMenuContent(this, &SAddToSequencerButton::MakeAddMenu)
+		.Icon(FAppStyle::Get().GetBrush("Icons.Plus"))
+		.Text(LOCTEXT("Add", "Add"))
+		.IsEnabled_Lambda([this]()
+		{
+			const TSharedPtr<FSequencer> SequencerPin = WeakSequencer.Pin();
+			return SequencerPin && !SequencerPin->IsReadOnly();
+		})
+	];
+}
+
+TSharedRef<SWidget> SAddToSequencerButton::MakeAddMenu()
+{
+	using namespace UE::Sequencer;
+
+	const TSharedPtr<FSequencer> SequencerPin = WeakSequencer.Pin();
+	if (!SequencerPin)
+	{
+		return SNullWidget::NullWidget;
+	}
+
+	FMenuBuilder MenuBuilder(true, SequencerPin->GetCommandBindings(), ExtenderForAddMenuAttr.Get());
+
+	if (SequencerPin->GetHostCapabilities().bSupportsAddFromContentBrowser)
+	{
+		MenuBuilder.AddMenuEntry(
+			FSequencerCommands::Get().AddSelectionFromContentBrowser,
+			NAME_None, TAttribute<FText>(), TAttribute<FText>(), 
+			FSlateIcon(FAppStyle::GetAppStyleSetName(), "Icons.Use")
+			);
+	}
+
+	FSequencerOutlinerViewModel* OutlinerViewModel = SequencerPin->GetViewModel()->GetOutliner()->CastThisChecked<FSequencerOutlinerViewModel>();
+	OutlinerViewModel->BuildContextMenu(MenuBuilder);
+
+	return MenuBuilder.MakeWidget();
+}
+} // namespace UE::Sequencer
+
+#undef LOCTEXT_NAMESPACE
\ No newline at end of file
diff --git a/Engine/Source/Editor/Sequencer/Private/Widgets/Sequencer/SAddToSequencerButton.h b/Engine/Source/Editor/Sequencer/Private/Widgets/Sequencer/SAddToSequencerButton.h
new file mode 100644
--- /dev/null
+++ b/Engine/Source/Editor/Sequencer/Private/Widgets/Sequencer/SAddToSequencerButton.h
@@ -0,0 +1,39 @@
+﻿// Copyright Epic Games, Inc. All Rights Reserved.
+
+#pragma once
+
+#include "Widgets/DeclarativeSyntaxSupport.h"
+#include "Widgets/SCompoundWidget.h"
+
+class FExtender;
+class FSequencer;
+
+namespace UE::Sequencer
+{
+/** Implements the add button for adding tracks, etc. to Sequencer. */
+class SAddToSequencerButton : public SCompoundWidget
+{
+public:
+	
+	SLATE_BEGIN_ARGS(SAddToSequencerButton){}
+		/** Used to add items to the sequencer instance, detecting whether it is read-only, etc. */
+		SLATE_ARGUMENT(TWeakPtr<FSequencer>, Sequencer)
+		/** External extenders for the add button menu. */
+		SLATE_ATTRIBUTE(TSharedPtr<FExtender>, ExtenderForAddMenu)
+	SLATE_END_ARGS()
+
+	void Construct(const FArguments& InArgs);
+	
+private:
+	
+	/** Used to generate the menu content */
+	TWeakPtr<FSequencer> WeakSequencer;
+	
+	/** External extenders for the add button menu. */
+	TAttribute<TSharedPtr<FExtender>> ExtenderForAddMenuAttr;
+	
+	/** @return Add combo button content */
+	TSharedRef<SWidget> MakeAddMenu();
+};
+} // namespace UE::Sequencer
+
diff --git a/Engine/Source/Editor/Sequencer/Private/Widgets/Sequencer/SSequencerSearchAndFilterRow.cpp b/Engine/Source/Editor/Sequencer/Private/Widgets/Sequencer/SSequencerSearchAndFilterRow.cpp
new file mode 100644
--- /dev/null
+++ b/Engine/Source/Editor/Sequencer/Private/Widgets/Sequencer/SSequencerSearchAndFilterRow.cpp
@@ -0,0 +1,69 @@
+﻿// Copyright Epic Games, Inc. All Rights Reserved.
+
+#include "SSequencerSearchAndFilterRow.h"
+
+#include "Filters/Menus/SequencerViewOptionsMenu.h"
+#include "Filters/SequencerFilterBar.h"
+#include "Filters/Widgets/SFilterBarIsolateHideShow.h"
+#include "Filters/Widgets/SSearchAndFilterWidget.h"
+#include "Filters/Widgets/SSequencerSearchBox.h"
+#include "Sequencer.h"
+#include "Widgets/Sequencer/SAddToSequencerButton.h"
+#include "Widgets/Sequencer/SSequencerViewOptionsButton.h"
+
+#define LOCTEXT_NAMESPACE "SSearchAndFilterRow"
+
+namespace UE::Sequencer
+{
+void SSequencerSearchAndFilterRow::Construct(const FArguments& InArgs)
+{
+	WeakSequencer = InArgs._Sequencer;
+	
+	const TSharedPtr<FSequencer> SequencerPin = WeakSequencer.Pin();
+	check(SequencerPin);
+	
+	ChildSlot
+	[
+		SNew(SHorizontalBox)
+
+		// Add Button
+		+ SHorizontalBox::Slot()
+		.AutoWidth()
+		.VAlign(VAlign_Center)
+		[
+			SNew(SAddToSequencerButton)
+			.Sequencer(InArgs._Sequencer)
+			.ExtenderForAddMenu(InArgs._ExtenderForAddMenu)
+		]
+
+		// Advanced Search Filter Combo Button
+		+ SHorizontalBox::Slot()
+		.FillWidth(1.f)
+		.VAlign(VAlign_Center)
+		[
+			SNew(SSearchAndFilterWidget, SequencerPin->GetFilterBar())
+			.OnSearchTextChanged(this, &SSequencerSearchAndFilterRow::OnSearchTextChanged)
+		]
+		
+		// View Options Combo Button
+		+ SHorizontalBox::Slot()
+		.AutoWidth()
+		.VAlign(VAlign_Center)
+		.Padding(2.f, 0.f, 0.f, 0.f)
+		[
+			SNew(SSequencerViewOptionsButton)
+			.Sequencer(InArgs._Sequencer)
+		]
+	];
+}
+
+void SSequencerSearchAndFilterRow::OnSearchTextChanged(const FText& InSearchText)
+{
+	if (const TSharedPtr<FSequencer> Sequencer = WeakSequencer.Pin())
+	{
+		Sequencer->GetNodeTree()->SetTextFilterString(InSearchText.ToString());
+	}
+}
+} // namespace UE::Sequencer
+
+#undef LOCTEXT_NAMESPACE
\ No newline at end of file
diff --git a/Engine/Source/Editor/Sequencer/Private/Widgets/Sequencer/SSequencerSearchAndFilterRow.h b/Engine/Source/Editor/Sequencer/Private/Widgets/Sequencer/SSequencerSearchAndFilterRow.h
new file mode 100644
--- /dev/null
+++ b/Engine/Source/Editor/Sequencer/Private/Widgets/Sequencer/SSequencerSearchAndFilterRow.h
@@ -0,0 +1,39 @@
+﻿// Copyright Epic Games, Inc. All Rights Reserved.
+
+#pragma once
+
+#include "Widgets/DeclarativeSyntaxSupport.h"
+#include "Widgets/SCompoundWidget.h"
+
+class FExtender;
+class FSequencer;
+class FSequencerFilterBar;
+
+namespace UE::Sequencer
+{
+/** The row of widgets displayed above the Sequence tree view. */
+class SSequencerSearchAndFilterRow : public SCompoundWidget
+{
+public:
+	
+	SLATE_BEGIN_ARGS(SSequencerSearchAndFilterRow){}
+		
+		/** Used to generate the subwidgets */
+		SLATE_ARGUMENT(TWeakPtr<FSequencer>, Sequencer)
+		
+		/** External extenders for the add button menu. */
+		SLATE_ATTRIBUTE(TSharedPtr<FExtender>, ExtenderForAddMenu)
+		
+	SLATE_END_ARGS()
+
+	void Construct(const FArguments& InArgs);
+	
+private:
+	
+	/** Used to update Sequencer's filters. */
+	TWeakPtr<FSequencer> WeakSequencer;
+	
+	void OnSearchTextChanged(const FText& InSearchText);
+};
+} // namespace UE::Sequencer
+
diff --git a/Engine/Source/Editor/Sequencer/Private/Widgets/Sequencer/SSequencerViewOptionsButton.cpp b/Engine/Source/Editor/Sequencer/Private/Widgets/Sequencer/SSequencerViewOptionsButton.cpp
new file mode 100644
--- /dev/null
+++ b/Engine/Source/Editor/Sequencer/Private/Widgets/Sequencer/SSequencerViewOptionsButton.cpp
@@ -0,0 +1,43 @@
+﻿// Copyright Epic Games, Inc. All Rights Reserved.
+
+#include "SSequencerViewOptionsButton.h"
+
+#include "Filters/Menus/SequencerViewOptionsMenu.h"
+#include "Sequencer.h"
+#include "Widgets/Input/SComboButton.h"
+
+#define LOCTEXT_NAMESPACE "SSequencerViewOptionsButton"
+
+class FSequencer;
+
+namespace UE::Sequencer
+{
+void SSequencerViewOptionsButton::Construct(const FArguments& InArgs)
+{
+	WeakSequencer = InArgs._Sequencer;
+	ViewOptionsMenu = MakeShared<FSequencerViewOptionsMenu>();
+	
+	ChildSlot
+	[
+		SNew(SComboButton)
+			.ContentPadding(2.f)
+			.HAlign(HAlign_Center)
+			.VAlign(VAlign_Center)
+			.ToolTipText(LOCTEXT("ViewOptionsToolTip", "View Options"))
+			.ComboButtonStyle(FAppStyle::Get(), TEXT("SimpleComboButtonWithIcon"))
+			.OnGetMenuContent_Lambda([this]() -> TSharedRef<SWidget>
+			{
+				return ViewOptionsMenu->CreateMenu(WeakSequencer);
+			})
+			.HasDownArrow(false)
+			.ButtonContent()
+			[
+				SNew(SImage)
+				.ColorAndOpacity(FSlateColor::UseForeground())
+				.Image(FAppStyle::Get().GetBrush(TEXT("Icons.Settings")))
+			]
+	];
+}
+} // namespace UE::Sequencer
+
+#undef LOCTEXT_NAMESPACE
\ No newline at end of file
diff --git a/Engine/Source/Editor/Sequencer/Private/Widgets/Sequencer/SSequencerViewOptionsButton.h b/Engine/Source/Editor/Sequencer/Private/Widgets/Sequencer/SSequencerViewOptionsButton.h
new file mode 100644
--- /dev/null
+++ b/Engine/Source/Editor/Sequencer/Private/Widgets/Sequencer/SSequencerViewOptionsButton.h
@@ -0,0 +1,33 @@
+﻿// Copyright Epic Games, Inc. All Rights Reserved.
+
+#pragma once
+
+#include "Widgets/DeclarativeSyntaxSupport.h"
+#include "Widgets/SCompoundWidget.h"
+
+class FSequencer;
+class FSequencerViewOptionsMenu;
+
+namespace UE::Sequencer
+{
+/** Implements view options for Sequencer. */
+class SSequencerViewOptionsButton : public SCompoundWidget
+{
+public:
+	
+	SLATE_BEGIN_ARGS(SSequencerViewOptionsButton){}
+		/** Used to generate the content of the view options. */
+		SLATE_ARGUMENT(TWeakPtr<FSequencer>, Sequencer)
+	SLATE_END_ARGS()
+
+	void Construct(const FArguments& InArgs);
+	
+private:
+	
+	/** Used to implement some items in the view options. */
+	TWeakPtr<FSequencer> WeakSequencer;
+	/** Used to generate the content of the view options. */
+	TSharedPtr<FSequencerViewOptionsMenu> ViewOptionsMenu;
+};
+} // namespace UE::Sequencer
+
diff --git a/Engine/Source/Editor/Sequencer/Public/Filters/Widgets/SSequencerSearchBox.h b/Engine/Source/Editor/Sequencer/Public/Filters/Widgets/SSequencerSearchBox.h
--- a/Engine/Source/Editor/Sequencer/Public/Filters/Widgets/SSequencerSearchBox.h
+++ b/Engine/Source/Editor/Sequencer/Public/Filters/Widgets/SSequencerSearchBox.h
@@ -155,6 +155,9 @@ private:
 
 	/** Provides default value suggestions, if no suggestions are provided */
 	UE_API void DefaultValueSuggestions(const FString& InKeyValue, TArray<FSequencerFilterSuggestion>& OutPossibleSuggestions);
+	
+	/** Handles ISequencerFilterBar::OnTextFilterTextChanged. */
+	void OnFilterTextChanged();
 
 	static UE_API bool IsOperatorToken(const TExpressionToken<TCHAR>& InToken);
 	static UE_API bool IsLogicalOperatorToken(const TExpressionToken<TCHAR>& InToken);
