[Sidebar]
- Change to use RegisterActiveTimer to avoid processing Sequencer selection changes when the selection details widget is hidden
- Add OnDrawerOpenClose event to sidebar
- Add notification to drawer sections when the drawer is open/closed

#Synced-CL 50035786

diff --git a/Engine/Source/Developer/ToolWidgets/Private/Sidebar/SSidebar.cpp b/Engine/Source/Developer/ToolWidgets/Private/Sidebar/SSidebar.cpp
--- a/Engine/Source/Developer/ToolWidgets/Private/Sidebar/SSidebar.cpp
+++ b/Engine/Source/Developer/ToolWidgets/Private/Sidebar/SSidebar.cpp
@@ -231,7 +231,7 @@ bool SSidebar::RegisterDrawerSection(const FName InDrawerId, const TSharedPtr<IS
 		DrawerSection->BuildContent();
 	}
 
-	return false;
+	return true;
 }
 
 bool SSidebar::UnregisterDrawerSection(const FName InDrawerId, const FName InSectionId)
@@ -249,7 +249,7 @@ bool SSidebar::UnregisterDrawerSection(const FName InDrawerId, const FName InSec
 
 	Drawer->ContentSections.Remove(InSectionId);
 
-	return false;
+	return true;
 }
 
 bool SSidebar::TryOpenDrawer(const FName InDrawerId)
@@ -434,8 +434,6 @@ void SSidebar::RemoveDrawer(const TSharedRef<FSidebarDrawer>& InDrawer)
 		Drawers.RemoveAt(IndexToRemove);
 	}
 
-	InDrawer->DrawerClosedDelegate.ExecuteIfBound(DrawerId);
-
 	if (ContainerWidget.IsValid())
 	{
 		ContainerWidget->UpdateDrawerTabAppearance();
@@ -444,7 +442,9 @@ void SSidebar::RemoveDrawer(const TSharedRef<FSidebarDrawer>& InDrawer)
 
 void SSidebar::RemoveAllDrawers()
 {
-	for (const TSharedRef<FSidebarDrawer>& Drawer : Drawers)
+	// Iterate over a copy to avoid modifying the array during iteration
+	const TArray<TSharedRef<FSidebarDrawer>> DrawersCopy = Drawers;
+	for (const TSharedRef<FSidebarDrawer>& Drawer : DrawersCopy)
 	{
 		RemoveDrawer(Drawer);
 	}
diff --git a/Engine/Source/Developer/ToolWidgets/Private/Sidebar/SSidebarContainer.cpp b/Engine/Source/Developer/ToolWidgets/Private/Sidebar/SSidebarContainer.cpp
--- a/Engine/Source/Developer/ToolWidgets/Private/Sidebar/SSidebarContainer.cpp
+++ b/Engine/Source/Developer/ToolWidgets/Private/Sidebar/SSidebarContainer.cpp
@@ -2,6 +2,7 @@
 
 #include "Sidebar/SSidebarContainer.h"
 #include "Framework/Application/SlateApplication.h"
+#include "Sidebar/ISidebarDrawerContent.h"
 #include "Sidebar/SSidebar.h"
 #include "Sidebar/SSidebarDrawer.h"
 #include "Widgets/Layout/SSplitter.h"
@@ -468,13 +469,14 @@ void SSidebarContainer::OpenDrawer_Internal(const TSharedRef<FSidebarDrawer>& In
 
 	InDrawer->DrawerWidget->Open(bInAnimate);
 	InDrawer->bIsOpen = true;
-	InDrawer->DrawerOpenedDelegate.ExecuteIfBound(InDrawer->GetUniqueId());
 
 	UpdateDrawerTabAppearance();
 
 	// This changes the focus and will trigger focus-related events, such as closing other tabs,
 	// so it's important that we only call it after we added the new drawer to OpenedDrawers.
 	FSlateApplication::Get().SetKeyboardFocus(InDrawer->DrawerWidget);
+
+	BroadcastDrawerOpenClose(InDrawer, /*bOpen=*/true);
 }
 
 void SSidebarContainer::CloseDrawer_Internal(const TSharedRef<FSidebarDrawer>& InDrawer, const bool bInAnimate
@@ -500,6 +502,8 @@ void SSidebarContainer::CloseDrawer_Internal(const TSharedRef<FSidebarDrawer>& I
 	{
 		SummonPinnedTabIfNothingOpened();
 	}
+
+	BroadcastDrawerOpenClose(InDrawer, /*bOpen=*/false);
 }
 
 void SSidebarContainer::SummonPinnedTabIfNothingOpened()
@@ -552,6 +556,8 @@ void SSidebarContainer::DockDrawer_Internal(const TSharedRef<FSidebarDrawer>& In
 	Reconstruct();
 
 	UpdateDrawerTabAppearance();
+
+	BroadcastDrawerOpenClose(InDrawer, /*bOpen=*/true);
 }
 
 void SSidebarContainer::UndockDrawer_Internal(const TSharedRef<FSidebarDrawer>& InDrawer)
@@ -562,6 +568,8 @@ void SSidebarContainer::UndockDrawer_Internal(const TSharedRef<FSidebarDrawer>&
 	Reconstruct();
 
 	UpdateDrawerTabAppearance();
+
+	BroadcastDrawerOpenClose(InDrawer, /*bOpen=*/false);
 }
 
 TSharedPtr<SSidebarDrawer> SSidebarContainer::FindOpenDrawerWidget(const TSharedRef<FSidebarDrawer>& InDrawer) const
@@ -708,4 +716,17 @@ void SSidebarContainer::OnSplitterResized()
 	SidebarWidget->OnStateChanged.ExecuteIfBound(SidebarWidget->GetState());
 }
 
+void SSidebarContainer::BroadcastDrawerOpenClose(const TSharedRef<FSidebarDrawer>& InDrawer, const bool bOpen)
+{
+	InDrawer->OnDrawerOpenCloseDelegate.ExecuteIfBound(InDrawer->GetUniqueId(), bOpen);
+
+	TArray<TSharedRef<ISidebarDrawerContent>> ContentSections;
+	InDrawer->ContentSections.GenerateValueArray(ContentSections);
+
+	for (const TSharedRef<ISidebarDrawerContent>& DrawerSection : ContentSections)
+	{
+		DrawerSection->OnDrawerOpenClose(bOpen);
+	}
+}
+
 #undef LOCTEXT_NAMESPACE
diff --git a/Engine/Source/Developer/ToolWidgets/Private/Sidebar/SidebarDrawer.h b/Engine/Source/Developer/ToolWidgets/Private/Sidebar/SidebarDrawer.h
--- a/Engine/Source/Developer/ToolWidgets/Private/Sidebar/SidebarDrawer.h
+++ b/Engine/Source/Developer/ToolWidgets/Private/Sidebar/SidebarDrawer.h
@@ -11,8 +11,7 @@ class ISidebarDrawerContent;
 class SSidebarDrawer;
 class SWidget;
 
-DECLARE_DELEGATE_OneParam(FOnSidebarDrawerOpened, const FName /*InUniqueId*/);
-DECLARE_DELEGATE_OneParam(FOnSidebarDrawerClosed, const FName /*InUniqueId*/);
+DECLARE_DELEGATE_TwoParams(FOnSidebarDrawerOpenClose, const FName /*InUniqueId*/, const bool /*bOpen*/)
 
 class FSidebarDrawer
 {
@@ -65,6 +64,5 @@ public:
 	bool bIsOpen = false;
 	FSidebarDrawerState State;
 
-	FOnSidebarDrawerOpened DrawerOpenedDelegate;
-	FOnSidebarDrawerClosed DrawerClosedDelegate;
+	FOnSidebarDrawerOpenClose OnDrawerOpenCloseDelegate;
 };
diff --git a/Engine/Source/Developer/ToolWidgets/Public/Sidebar/ISidebarDrawerContent.h b/Engine/Source/Developer/ToolWidgets/Public/Sidebar/ISidebarDrawerContent.h
--- a/Engine/Source/Developer/ToolWidgets/Public/Sidebar/ISidebarDrawerContent.h
+++ b/Engine/Source/Developer/ToolWidgets/Public/Sidebar/ISidebarDrawerContent.h
@@ -34,4 +34,7 @@ public:
 
 	/** Constructs the content widget for a sidebar drawer. */
 	virtual TSharedRef<SWidget> CreateContentWidget() = 0;
+
+	/** Event that fires when the drawer is opened or closed. */
+	virtual void OnDrawerOpenClose(const bool bOpen) {}
 };
diff --git a/Engine/Source/Developer/ToolWidgets/Public/Sidebar/SSidebarContainer.h b/Engine/Source/Developer/ToolWidgets/Public/Sidebar/SSidebarContainer.h
--- a/Engine/Source/Developer/ToolWidgets/Public/Sidebar/SSidebarContainer.h
+++ b/Engine/Source/Developer/ToolWidgets/Public/Sidebar/SSidebarContainer.h
@@ -101,6 +101,8 @@ protected:
 	UE_API int32 GetContentSlotIndex() const;
 	UE_API int32 GetSidebarSlotIndex() const;
 
+	void BroadcastDrawerOpenClose(const TSharedRef<FSidebarDrawer>& InDrawer, const bool bOpen);
+
 	/** The sidebar widget associated with this container. One sidebar widget per container. */
 	TSharedPtr<SSidebar> SidebarWidget;
 
diff --git a/Engine/Source/Editor/Sequencer/Private/Widgets/Sidebar/SequencerSelectionDrawer.cpp b/Engine/Source/Editor/Sequencer/Private/Widgets/Sidebar/SequencerSelectionDrawer.cpp
--- a/Engine/Source/Editor/Sequencer/Private/Widgets/Sidebar/SequencerSelectionDrawer.cpp
+++ b/Engine/Source/Editor/Sequencer/Private/Widgets/Sidebar/SequencerSelectionDrawer.cpp
@@ -1,7 +1,6 @@
 // Copyright Epic Games, Inc. All Rights Reserved.
 
 #include "SequencerSelectionDrawer.h"
-#include "DetailsViewArgs.h"
 #include "Framework/MultiBox/MultiBoxBuilder.h"
 #include "IKeyArea.h"
 #include "ISequencerSection.h"
@@ -18,7 +17,6 @@
 #include "MVVM/ViewModels/SectionOutlinerModel.h"
 #include "MVVM/ViewModels/SequencerEditorViewModel.h"
 #include "MVVM/ViewModels/TrackRowModel.h"
-#include "PropertyEditorModule.h"
 #include "Sequencer.h"
 #include "SequencerCommonHelpers.h"
 #include "SequencerUtilities.h"
@@ -144,7 +142,9 @@ TSharedRef<SWidget> FSequencerSelectionDrawer::CreateContentWidget()
 			]
 		];
 
-	OnSequencerSelectionChanged();
+	ContentBox->RegisterActiveTimer(0.f, FWidgetActiveTimerDelegate::CreateSP(this, &FSequencerSelectionDrawer::TickContentBox));
+
+	RequestSelectionUpdate();
 
 	return OutWidget;
 }
@@ -161,7 +161,7 @@ void FSequencerSelectionDrawer::BindEvents()
 
 	ActorAddedToSequencerDelegate = Sequencer->OnActorAddedToSequencer().AddLambda([this](AActor* InActor, const FGuid InGuid)
 		{
-			OnSequencerSelectionChanged();
+			RequestSelectionUpdate();
 		});
 
 	MovieSceneDataChangedDelegate = Sequencer->OnMovieSceneDataChanged().AddLambda([this](const EMovieSceneDataChangeType InChangeType)
@@ -169,13 +169,13 @@ void FSequencerSelectionDrawer::BindEvents()
 			// Update when a channel is overriden with a curve extension
 			if (InChangeType == EMovieSceneDataChangeType::MovieSceneStructureItemsChanged)
 			{
-				OnSequencerSelectionChanged();
+				RequestSelectionUpdate();
 			}
 		});
 
 	if (const TSharedPtr<FSequencerSelection> SequencerSelection = Private::GetSelection(*Sequencer.Get()))
 	{
-		SequencerSelection->OnChanged.AddSP(this, &FSequencerSelectionDrawer::OnSequencerSelectionChanged);
+		SequencerSelection->OnChanged.AddSP(this, &FSequencerSelectionDrawer::RequestSelectionUpdate);
 	}
 }
 
@@ -216,34 +216,24 @@ void FSequencerSelectionDrawer::ResetContent()
 	ChannelExtensions.Reset();
 }
 
-void FSequencerSelectionDrawer::OnSequencerSelectionChanged()
+EActiveTimerReturnType FSequencerSelectionDrawer::TickContentBox(const double InCurrentTime, const float InDeltaTime)
 {
-	if (!ContentBox.IsValid() || !GEditor)
+	if (ContentBox.IsValid() && bRequestSelectionUpdate)
 	{
-		return;
+		bRequestSelectionUpdate = false;
+		UpdateFromSelectionNextFrame();
 	}
 
-	if (bWaitingToHandleSelectionChanged)
-	{
-		return;
-	}
-
-	bWaitingToHandleSelectionChanged = true;
+	return EActiveTimerReturnType::Continue;
+}
 
-	const TWeakPtr<FSequencerSelectionDrawer> ThisWeak = SharedThis(this);
-	GEditor->GetTimerManager()->SetTimerForNextTick([ThisWeak]()
-		{
-			if (const TSharedPtr<FSequencerSelectionDrawer> This = ThisWeak.Pin())
-			{
-				This->UpdateFromSelectionNextFrame();
-			}
-		});
+void FSequencerSelectionDrawer::RequestSelectionUpdate()
+{
+	bRequestSelectionUpdate = true;
 }
 
 void FSequencerSelectionDrawer::UpdateFromSelectionNextFrame()
 {
-	bWaitingToHandleSelectionChanged = false;
-
 	ResetContent();
 
 	const TSharedPtr<FSequencer> Sequencer = WeakSequencer.Pin();
diff --git a/Engine/Source/Editor/Sequencer/Private/Widgets/Sidebar/SequencerSelectionDrawer.h b/Engine/Source/Editor/Sequencer/Private/Widgets/Sidebar/SequencerSelectionDrawer.h
--- a/Engine/Source/Editor/Sequencer/Private/Widgets/Sidebar/SequencerSelectionDrawer.h
+++ b/Engine/Source/Editor/Sequencer/Private/Widgets/Sidebar/SequencerSelectionDrawer.h
@@ -13,6 +13,7 @@ class FText;
 class ISidebarChannelExtension;
 class SVerticalBox;
 class SWidget;
+enum class EActiveTimerReturnType : uint8;
 struct FCurveChannelSectionSidebarExtension;
 struct FKeyEditData;
 
@@ -41,7 +42,9 @@ protected:
 	void BindEvents();
 	void UnbindEvents();
 
-	void OnSequencerSelectionChanged();
+	EActiveTimerReturnType TickContentBox(const double InCurrentTime, const float InDeltaTime);
+
+	void RequestSelectionUpdate();
 	void UpdateFromSelectionNextFrame();
 
 	void BuildKeySelectionDetails(const TSharedRef<UE::Sequencer::FSequencerSelection>& InSelection, FMenuBuilder& MenuBuilder);
@@ -69,7 +72,7 @@ protected:
 
 	TArray<TSharedPtr<ISidebarChannelExtension>> ChannelExtensions;
 
-	bool bWaitingToHandleSelectionChanged = false;
+	bool bRequestSelectionUpdate = false;
 
 	FDelegateHandle ActorAddedToSequencerDelegate;
 	FDelegateHandle MovieSceneDataChangedDelegate;
