[Sequencer] Refactor outliner widget to not size proportionately

#Synced-CL 50092247

diff --git a/Engine/Source/Editor/Sequencer/Private/SSequencer.cpp b/Engine/Source/Editor/Sequencer/Private/SSequencer.cpp
--- a/Engine/Source/Editor/Sequencer/Private/SSequencer.cpp
+++ b/Engine/Source/Editor/Sequencer/Private/SSequencer.cpp
@@ -167,6 +167,11 @@ FSequencerOutlinerColumnVisibility::FSequencerOutlinerColumnVisibility(TSharedPt
 	, bIsColumnVisible(bInIsColumnVisible)
 {}
 
+bool FSequencerOutlinerColumnVisibility::operator==(const FSequencerOutlinerColumnVisibility& InOther) const
+{
+	return Column->GetColumnName() == InOther.Column->GetColumnName();
+}
+
 /* SSequencer interface
  *****************************************************************************/
 UE_DISABLE_OPTIMIZATION_SHIP
@@ -279,16 +284,96 @@ void SSequencer::Construct(const FArguments& InArgs, TSharedRef<FSequencer> InSe
 
 	PlayTimeDisplay = StaticCastSharedRef<STemporarilyFocusedSpinBox<double>>(SequencerPtr.Pin()->MakePlayTimeDisplay());
 
-	TAttribute<FAnimatedRange> ViewRangeAttribute = InArgs._ViewRange;
+	check(SequencerPtr.IsValid());
+	const TSharedRef<FSequencer> SequencerRef = SequencerPtr.Pin().ToSharedRef();
+
+	TimeSliderController = MakeShared<FSequencerTimeSliderController>(TimeSliderArgs, SequencerPtr);
+	TSharedRef<FSequencerTimeSliderController> TimeSliderControllerRef = TimeSliderController.ToSharedRef();
+
+	ISequencerWidgetsModule& SequencerWidgets = FModuleManager::Get().LoadModuleChecked<ISequencerWidgetsModule>(TEXT("SequencerWidgets"));
+
+	ScrollBar = SNew(SScrollBar)
+		.Thickness(FVector2D(9.f, 9.f));
+
+	PinnedAreaScrollBar = SNew(SScrollBar)
+		.Thickness(FVector2D(9.f, 9.f));
+
+	TSharedPtr<FSequencerEditorViewModel> SequencerViewModel = SequencerRef->GetViewModel()->CastThisShared<FSequencerEditorViewModel>();
+	SAssignNew(PinnedTrackArea, SSequencerTrackAreaView, SequencerViewModel->GetPinnedTrackArea(), TimeSliderControllerRef);
+	SAssignNew(PinnedTreeView, SSequencerOutlinerView, SequencerViewModel->GetOutliner(), PinnedTrackArea.ToSharedRef())
+		.Selection(SequencerRef->GetViewModel()->GetSelection())
+		.ExternalScrollbar(PinnedAreaScrollBar)
+		.Clipping(EWidgetClipping::ClipToBounds);
+
+	PinnedTrackArea->SetOutliner(PinnedTreeView);
+	PinnedTrackArea->SetShowPinned(true);
+	PinnedTrackArea->SetIsPinned(true);
+	PinnedTreeView->SetShowPinned(true);
+
+	SAssignNew(TrackArea, SSequencerTrackAreaView, SequencerViewModel->GetTrackArea(), TimeSliderControllerRef);
+	SAssignNew(TreeView, SSequencerOutlinerView, SequencerViewModel->GetOutliner(), TrackArea.ToSharedRef())
+		.Selection(SequencerRef->GetViewModel()->GetSelection())
+		.ExternalScrollbar(ScrollBar)
+		.Clipping(EWidgetClipping::ClipToBounds);
 
-	GridPanel = ConstructTrackAreaGridPanel(InArgs, TimeSliderArgs);
+	TrackArea->SetOutliner(TreeView);
+
+	TreeView->AddPinnedTreeView(PinnedTreeView);
 
-	if (const TSharedPtr<FSequencerFilterBar> FilterBar = GetFilterBar())
+	USequencerSettings* const SequencerSettings = GetSequencerSettings();
+
+	if (SequencerSettings)
 	{
-		FilterBar->OnStateChanged().AddSP(this, &SSequencer::OnFilterBarStateChanged);
-		FilterBar->OnFiltersChanged().AddSP(this, &SSequencer::OnTrackFiltersChanged);
+		InitializeOutlinerColumns();
+		
+		SequencerViewModel->SetViewDensity(SequencerSettings->GetViewDensity());
 	}
 
+	SequencerViewModel->GetTrackArea()->CastThisChecked<FSequencerTrackAreaViewModel>()->InitializeDefaultEditTools(*TrackArea);
+	SequencerViewModel->GetPinnedTrackArea()->CastThisChecked<FSequencerTrackAreaViewModel>()->InitializeDefaultEditTools(*PinnedTrackArea);
+
+	// Create the top and bottom sliders
+	bool bMirrorLabels = false;
+	TopTimeSlider = SequencerWidgets.CreateTimeSlider(TimeSliderControllerRef, bMirrorLabels);
+	bMirrorLabels = true;
+	BottomTimeSlider = SequencerWidgets.CreateTimeSlider(TimeSliderControllerRef, TAttribute<EVisibility>(this, &SSequencer::GetBottomTimeSliderVisibility), bMirrorLabels);
+	
+	// Create bottom time range slider
+	EShowRange Ranges = EShowRange(EShowRange::WorkingRange | EShowRange::ViewRange);
+	if(InArgs._ShowPlaybackRangeInTimeSlider)
+	{
+		Ranges |= EShowRange::PlaybackRange;
+	}
+
+	BottomTimeRange = SequencerWidgets.CreateTimeRange(
+		FTimeRangeArgs(
+			Ranges,
+			TimeSliderControllerRef,
+			TAttribute<EVisibility>(this, &SSequencer::GetTimeRangeVisibility),
+			MakeAttributeSP(SequencerRef, &FSequencer::GetNumericTypeInterface, ENumericIntent::Position),
+			SequencerSettings->GetPlaybackRangeStartColor(),
+			SequencerSettings->GetPlaybackRangeEndColor()
+		),
+		SequencerWidgets.CreateTimeRangeSlider(TimeSliderControllerRef)
+	);
+
+	OutlinerAndTrackAreaContent = ConstructOutlinerAndTrackArea(InArgs, TimeSliderArgs);
+
+	static const FName TransportControlsLeft("Sequencer.TransportControls.Left");
+	static const FName TransportControlsRight("Sequencer.TransportControls.Right");
+	if (!UToolMenus::Get()->IsMenuRegistered(TransportControlsLeft))
+	{
+		UToolMenus::Get()->RegisterMenu(TransportControlsLeft, NAME_None, EMultiBoxType::ToolBar);
+		UToolMenus::Get()->RegisterMenu(TransportControlsRight, NAME_None, EMultiBoxType::ToolBar);
+	}
+
+	USequencerToolMenuContext* ContextObject = NewObject<USequencerToolMenuContext>();
+	ContextObject->WeakSequencer = SequencerRef;
+
+	FilterRow = SNew(SSequencerSearchAndFilterRow)
+		.ExtenderForAddMenu_Lambda([this]{ return FExtender::Combine(AddMenuExtenders); })
+		.Sequencer(SequencerPtr);
+
 	ChildSlot
 	[
 		SNew(SVerticalBox)
@@ -467,7 +552,6 @@ void SSequencer::Construct(const FArguments& InArgs, TSharedRef<FSequencer> InSe
 
 	if (InSequencer->GetHostCapabilities().bSupportsCurveEditor)
 	{
-		TSharedPtr<FSequencerEditorViewModel> SequencerViewModel = InSequencer->GetViewModel()->CastThisShared<FSequencerEditorViewModel>();
 		FCurveEditorExtension* CurveEditorExtension = SequencerViewModel->CastDynamicChecked<FCurveEditorExtension>();
 		CurveEditorExtension->CreateCurveEditor(TimeSliderArgs);
 	}
@@ -492,7 +576,7 @@ TSharedRef<SWidget> SSequencer::ConstructSidebarContent()
 	const bool bSupportsSidebar = Sequencer->GetHostCapabilities().bSupportsSidebar;
 
 	USequencerSettings* const SequencerSettings = GetSequencerSettings();
-	check(IsValid(SequencerSettings));
+	check(SequencerSettings);
 	const FSidebarState& SidebarState = SequencerSettings->GetSidebarState();
 
 	TSharedPtr<SWidget> OutWidget;
@@ -553,7 +637,7 @@ TSharedRef<SWidget> SSequencer::ConstructFilterBarContent()
 	
 	if (!IsFilterBarVisible() || GetFilterBarLayout() == EFilterBarLayout::Horizontal)
 	{
-		return ConstructGridOverlayContent();
+		return OutlinerAndTrackAreaContent.ToSharedRef();
 	}
 
 	const TSharedRef<SSequencerFilterBar> FilterBarWidgetRef = FilterBarWidget.ToSharedRef();
@@ -587,493 +671,397 @@ TSharedRef<SWidget> SSequencer::ConstructFilterBarContent()
 		+ SSplitter::Slot()
 		.Value(0.94f)
 		[
-			ConstructGridOverlayContent()
+			OutlinerAndTrackAreaContent.ToSharedRef()
 		];
 }
 
-TSharedRef<SWidget> SSequencer::ConstructGridOverlayContent()
-{
-	TAttribute<float> FillCoefficient_0, FillCoefficient_1;
-	{
-		FillCoefficient_0.Bind(TAttribute<float>::FGetter::CreateSP(this, &SSequencer::GetColumnFillCoefficient, 0));
-		FillCoefficient_1.Bind(TAttribute<float>::FGetter::CreateSP(this, &SSequencer::GetColumnFillCoefficient, 1));
-	}
-
-	return
-		SNew(SOverlay)
-		+ SOverlay::Slot()
-		[
-			GridPanel.ToSharedRef()
-		]
-		+ SOverlay::Slot()
-		[
-			// track area virtual splitter overlay
-			SAssignNew(TreeViewSplitter, SSequencerSplitterOverlay)
-			.Style(FAppStyle::Get(), TEXT("Sequencer.AnimationOutliner.Splitter"))
-			.Visibility(EVisibility::SelfHitTestInvisible)
-			.OnSplitterFinishedResizing(this, &SSequencer::OnSplitterFinishedResizing)
-
-			+ SSplitter::Slot()
-			.Value(FillCoefficient_0)
-			// Can't use a minsize here because the grid panel that is actually being used to
-			//   lay out the widgets only supports fill coefficients and this leads to a disparity between the two
-			// .MinSize(200)
-			.OnSlotResized(SSplitter::FOnSlotResized::CreateSP(this, &SSequencer::OnColumnFillCoefficientChanged, 0))
-			[
-				SNew(SSpacer)
-			]
-
-			+ SSplitter::Slot()
-			.Value(FillCoefficient_1)
-			.OnSlotResized(SSplitter::FOnSlotResized::CreateSP(this, &SSequencer::OnColumnFillCoefficientChanged, 1))
-			[
-				SNew(SSpacer)
-			]
-		];
-}
-
-TSharedRef<SGridPanel> SSequencer::ConstructTrackAreaGridPanel(const FArguments& InArgs, const FTimeSliderArgs& InTimeSliderArgs)
+TSharedRef<SSplitter> SSequencer::ConstructOutlinerAndTrackArea(const FArguments& InArgs, const FTimeSliderArgs& InTimeSliderArgs)
 {
 	using namespace UE::Sequencer;
 
-	check(SequencerPtr.IsValid());
-	const TSharedRef<FSequencer> SequencerRef = SequencerPtr.Pin().ToSharedRef();
-
-	constexpr int32 Column0 = 0, Column1 = 1;
-	constexpr int32 Row0 = 0, Row1 = 1, Row2 = 2, Row3 = 3, Row4 = 4;
-	const FMargin ResizeBarPadding(4.f, 0, 0, 0);
-
-	TAttribute<float> FillCoefficient_0, FillCoefficient_1;
-	{
-		FillCoefficient_0.Bind(TAttribute<float>::FGetter::CreateSP(this, &SSequencer::GetColumnFillCoefficient, 0));
-		FillCoefficient_1.Bind(TAttribute<float>::FGetter::CreateSP(this, &SSequencer::GetColumnFillCoefficient, 1));
-	}
+	// This will add a thin, light gray, horizontal line to match the right side of the sequencer widget
+	const FMargin ResizeBarPadding(1.f, 0, 0, 0);
 
+	const TSharedRef<FSequencer> SequencerRef = SequencerPtr.Pin().ToSharedRef();
+	const TSharedRef<FSequencerTimeSliderController> TimeSliderControllerRef = TimeSliderController.ToSharedRef();
 	USequencerSettings* const SequencerSettings = GetSequencerSettings();
 
-	ColumnFillCoefficients[0] = 0.3f;
-	ColumnFillCoefficients[1] = 0.7f;
-
-	if (IsValid(SequencerSettings))
-	{
-		const float TreeViewWidth = SequencerSettings->GetTreeViewWidth();
-		const float TimelineWidth = 1.f - TreeViewWidth;
-		if (TreeViewWidth > 0.f && TimelineWidth > 0.f)
-		{
-			ColumnFillCoefficients[0] = TreeViewWidth;
-			ColumnFillCoefficients[1] = TimelineWidth;
-		}
-	}
-
-	TimeSliderController = MakeShared<FSequencerTimeSliderController>(InTimeSliderArgs, SequencerPtr);
-	TSharedRef<FSequencerTimeSliderController> TimeSliderControllerRef = TimeSliderController.ToSharedRef();
-
-	ISequencerWidgetsModule& SequencerWidgets = FModuleManager::Get().LoadModuleChecked<ISequencerWidgetsModule>(TEXT("SequencerWidgets"));
-
-	ScrollBar = SNew(SScrollBar)
-		.Thickness(FVector2D(9.f, 9.f));
-
-	PinnedAreaScrollBar = SNew(SScrollBar)
-		.Thickness(FVector2D(9.f, 9.f));
-
-	TSharedPtr<FSequencerEditorViewModel> SequencerViewModel = SequencerRef->GetViewModel()->CastThisShared<FSequencerEditorViewModel>();
-	SAssignNew(PinnedTrackArea, SSequencerTrackAreaView, SequencerViewModel->GetPinnedTrackArea(), TimeSliderControllerRef);
-	SAssignNew(PinnedTreeView, SSequencerOutlinerView, SequencerViewModel->GetOutliner(), PinnedTrackArea.ToSharedRef())
-		.Selection(SequencerRef->GetViewModel()->GetSelection())
-		.ExternalScrollbar(PinnedAreaScrollBar)
-		.Clipping(EWidgetClipping::ClipToBounds);
-
-	PinnedTrackArea->SetOutliner(PinnedTreeView);
-	PinnedTrackArea->SetShowPinned(true);
-	PinnedTrackArea->SetIsPinned(true);
-	PinnedTreeView->SetShowPinned(true);
-
-	SAssignNew(TrackArea, SSequencerTrackAreaView, SequencerViewModel->GetTrackArea(), TimeSliderControllerRef);
-	SAssignNew(TreeView, SSequencerOutlinerView, SequencerViewModel->GetOutliner(), TrackArea.ToSharedRef())
-		.Selection(SequencerRef->GetViewModel()->GetSelection())
-		.ExternalScrollbar(ScrollBar)
-		.Clipping(EWidgetClipping::ClipToBounds);
-
-	TrackArea->SetOutliner(TreeView);
-
-	TreeView->AddPinnedTreeView(PinnedTreeView);
-
-	if (SequencerSettings)
-	{
-		InitializeOutlinerColumns();
-	}
-
-	SequencerViewModel->GetTrackArea()->CastThisChecked<FSequencerTrackAreaViewModel>()->InitializeDefaultEditTools(*TrackArea);
-	SequencerViewModel->GetPinnedTrackArea()->CastThisChecked<FSequencerTrackAreaViewModel>()->InitializeDefaultEditTools(*PinnedTrackArea);
-
-	if (SequencerSettings)
-	{
-		SequencerViewModel->SetViewDensity(SequencerSettings->GetViewDensity());
-	}
-
-	// Create the top and bottom sliders
-	bool bMirrorLabels = false;
-	TopTimeSlider = SequencerWidgets.CreateTimeSlider(TimeSliderControllerRef, bMirrorLabels);
-	bMirrorLabels = true;
-	BottomTimeSlider = SequencerWidgets.CreateTimeSlider(TimeSliderControllerRef, TAttribute<EVisibility>(this, &SSequencer::GetBottomTimeSliderVisibility), bMirrorLabels);
-	
-	// Create bottom time range slider
-	EShowRange Ranges = EShowRange(EShowRange::WorkingRange | EShowRange::ViewRange);
-	if(InArgs._ShowPlaybackRangeInTimeSlider)
-	{
-		Ranges |= EShowRange::PlaybackRange;
-	}
-
-	BottomTimeRange = SequencerWidgets.CreateTimeRange(
-		FTimeRangeArgs(
-			Ranges,
-			TimeSliderControllerRef,
-			TAttribute<EVisibility>(this, &SSequencer::GetTimeRangeVisibility),
-			MakeAttributeSP(SequencerRef, &FSequencer::GetNumericTypeInterface, UE::Sequencer::ENumericIntent::Position),
-			GetSequencerSettings()->GetPlaybackRangeStartColor(),
-			GetSequencerSettings()->GetPlaybackRangeEndColor()
-		),
-		SequencerWidgets.CreateTimeRangeSlider(TimeSliderControllerRef)
-	);
-
-	static const FName TransportControlsLeft("Sequencer.TransportControls.Left");
-	static const FName TransportControlsRight("Sequencer.TransportControls.Right");
-	if (!UToolMenus::Get()->IsMenuRegistered(TransportControlsLeft))
-	{
-		UToolMenus::Get()->RegisterMenu(TransportControlsLeft, NAME_None, EMultiBoxType::ToolBar);
-		UToolMenus::Get()->RegisterMenu(TransportControlsRight, NAME_None, EMultiBoxType::ToolBar);
-	}
-
-	USequencerToolMenuContext* ContextObject = NewObject<USequencerToolMenuContext>();
-	ContextObject->WeakSequencer = SequencerRef;
-	
-	FilterRow = SNew(SSequencerSearchAndFilterRow)
-		.ExtenderForAddMenu_Lambda([this]{ return FExtender::Combine(AddMenuExtenders); })
-		.Sequencer(SequencerPtr);
-	
-	return SNew(SGridPanel)
-		.FillRow(2, 1.f)
-		.FillColumn(0, FillCoefficient_0)
-		.FillColumn(1, FillCoefficient_1)
-
-		+ SGridPanel::Slot(Column0, Row1)
-		[
-			SNew(SBorder)
-			.BorderImage(FAppStyle::GetBrush(TEXT("ToolPanel.GroupBorder")))
-			[
-				SNew(SSpacer)
-			]
-		]
-
-		// outliner search box
-		+ SGridPanel::Slot(Column0, Row1, SGridPanel::Layer(10))
+	return SNew(SSplitter)
+		+ SSplitter::Slot()
+		.Resizable(true)
+		.SizeRule(SSplitter::SizeToContent)
+		.OnSlotResized_Lambda([this](float InSlotSize)
+			{
+				GetSequencerSettings()->SetTreeViewWidth(InSlotSize);
+			})
 		[
-			SNew(SBorder)
-			.BorderImage(FAppStyle::GetBrush(TEXT("ToolPanel.GroupBorder")))
-			.Padding(FMargin(CommonPadding*2, CommonPadding))
-			.Clipping(EWidgetClipping::ClipToBounds)
+			SNew(SBox)
+			.MinDesiredWidth(MinOutlinerWidth)
+			.WidthOverride(this, &SSequencer::GetOutlinerWidth)
 			[
-				SAssignNew(SearchAndFilterRow, SVerticalBox)
+				SNew(SVerticalBox)
+				+ SVerticalBox::Slot()
+				.Expose(OutlinerSearchAndFilterSlot)
+				.AutoHeight()
+				[
+					SNew(SBorder)
+					.BorderImage(FAppStyle::GetBrush(TEXT("ToolPanel.GroupBorder")))
+					.Padding(FMargin(CommonPadding * 2, CommonPadding))
+					.Clipping(EWidgetClipping::ClipToBounds)
+					[
+						SAssignNew(SearchAndFilterRow, SVerticalBox)
+					]
+				]
+				+ SVerticalBox::Slot()
+				[
+					ConstructMainTreeView(SequencerRef)
+				]
+				+ SVerticalBox::Slot()
+				.AutoHeight()
+				[
+					ConstructTransportPanel(SequencerRef)
+				]
 			]
 		]
-
-		// main sequencer area
-		+ SGridPanel::Slot(Column0, Row2, SGridPanel::Layer(10))
-		.ColumnSpan(2)
+		+ SSplitter::Slot()
 		[
-			SAssignNew(MainSequencerArea, SVerticalBox)
-
-			+ SVerticalBox::Slot()
-			.AutoHeight()
-			.MaxHeight(TAttribute<float>::Create(TAttribute<float>::FGetter::CreateSP(this, &SSequencer::GetPinnedAreaMaxHeight)))
-			.Padding(FMargin(0.f, 0.f, 0.f, CommonPadding))
+			SNew(SBox)
+			.MinDesiredWidth(MinOutlinerWidth)
+			.Padding(ResizeBarPadding)
 			[
-				SNew(SOverlay)
-				.Visibility(this, &SSequencer::GetPinnedAreaVisibility)
-
-				+ SOverlay::Slot()
+				SNew(SVerticalBox)
+				+ SVerticalBox::Slot()
+				.AutoHeight()
 				[
-					SNew(SScrollBorder, PinnedTreeView.ToSharedRef())
+					SNew(SBox)
+					.HeightOverride_Lambda([this]()
+						{
+							return OutlinerSearchAndFilterSlot
+								? OutlinerSearchAndFilterSlot->GetWidget()->GetCachedGeometry().GetLocalSize().Y
+								: 0.;
+						})
 					[
-						SNew(SHorizontalBox)
-
-						// outliner tree
-						+ SHorizontalBox::Slot()
-						.FillWidth(FillCoefficient_0)
+						SNew(SBorder)
+						.BorderImage(FAppStyle::GetBrush(TEXT("ToolPanel.GroupBorder")))
+						.BorderBackgroundColor(FLinearColor(.5f, .5f, .5f, 1.f))
+						.Padding(0)
+						.Clipping(EWidgetClipping::ClipToBounds)
 						[
-							PinnedTreeView.ToSharedRef()
-						]
-
-						// track area
-						+ SHorizontalBox::Slot()
-						.FillWidth(FillCoefficient_1)
-						[
-							SNew(SBox)
-							.Padding(ResizeBarPadding)
-							.Clipping(EWidgetClipping::ClipToBounds)
-							[
-								PinnedTrackArea.ToSharedRef()
-							]
+							SNew(SOverlay)
+							+ SOverlay::Slot().ZOrder(1) [ SNew(STopSliderAreaOverlay, SequencerRef->GetViewModel()) ]
+							+ SOverlay::Slot().ZOrder(2) [ TopTimeSlider.ToSharedRef() ]
 						]
 					]
 				]
-
-				+ SOverlay::Slot()
-				.HAlign(HAlign_Right)
+				+ SVerticalBox::Slot()
 				[
-					PinnedAreaScrollBar.ToSharedRef()
-				]
-			]
-
-			+ SVerticalBox::Slot()
-			[
-				SNew(SOverlay)
+					SNew(SOverlay)
 
-				+ SOverlay::Slot()
-				[
-					SNew(SScrollBorder, TreeView.ToSharedRef())
+					+ SOverlay::Slot()
 					[
-						SNew(SHorizontalBox)
+						SNew(SVerticalBox)
 
-						// outliner tree
-						+ SHorizontalBox::Slot()
-						.FillWidth(FillCoefficient_0)
+						+ SVerticalBox::Slot()
+						.AutoHeight()
+						.MaxHeight(TAttribute<float>::Create(TAttribute<float>::FGetter::CreateSP(this, &SSequencer::GetPinnedAreaMaxHeight)))
+						.Padding(FMargin(0.f, 0.f, 0.f, CommonPadding))
 						[
-							TreeView.ToSharedRef()
+							SNew(SBox)
+							.HeightOverride_Lambda([this]()
+								{
+									return OutlinerPinnedSlot
+										? OutlinerPinnedSlot->GetWidget()->GetCachedGeometry().GetLocalSize().Y
+										: 0.;
+								})
+							[
+								SNew(SOverlay)
+								.Visibility(this, &SSequencer::GetPinnedAreaVisibility)
+
+								+ SOverlay::Slot()
+								[
+									SNew(SScrollBorder, PinnedTreeView.ToSharedRef())
+									[
+										SNew(SBox)
+										.Padding(ResizeBarPadding)
+										.Clipping(EWidgetClipping::ClipToBounds)
+										[
+											PinnedTrackArea.ToSharedRef()
+										]
+									]
+								]
+
+								+ SOverlay::Slot()
+								.HAlign(HAlign_Right)
+								[
+									PinnedAreaScrollBar.ToSharedRef()
+								]
+							]
 						]
 
-						// track area
-						+ SHorizontalBox::Slot()
-						.FillWidth(FillCoefficient_1)
+						+ SVerticalBox::Slot()
 						[
-							SNew(SBox)
-							.Padding(ResizeBarPadding)
-							.Clipping(EWidgetClipping::ClipToBounds)
+							SNew(SOverlay)
+
+							+ SOverlay::Slot()
 							[
-								TrackArea.ToSharedRef()
+								SNew(SScrollBorder, TreeView.ToSharedRef())
+								[
+									SNew(SBox)
+									.Padding(ResizeBarPadding)
+									.Clipping(EWidgetClipping::ClipToBounds)
+									[
+										TrackArea.ToSharedRef()
+									]
+								]
+							]
+
+							+ SOverlay::Slot()
+							.HAlign(HAlign_Right)
+							[
+								ScrollBar.ToSharedRef()
+							]
+
+							+ SOverlay::Slot()
+							.VAlign(VAlign_Bottom)
+							[
+								SNew(SHorizontalBox)
+
+								+ SHorizontalBox::Slot()
+								.AutoWidth()
+								[
+									SAssignNew(SequencerTreeFilterStatusBar, SSequencerTreeFilterStatusBar, SequencerRef)
+									.Visibility(EVisibility::Hidden) // Initially hidden, visible on hover of the info button
+								]
 							]
 						]
 					]
-				]
-
-				+ SOverlay::Slot()
-				.HAlign(HAlign_Right)
-				[
-					ScrollBar.ToSharedRef()
-				]
 
-				+ SOverlay::Slot()
-				.VAlign(VAlign_Bottom)
-				[
-					SNew(SHorizontalBox)
+					+ SOverlay::Slot()
+					.ZOrder(10)
+					[
+						SNew(SSequencerSectionOverlay, TimeSliderControllerRef)
+						.Visibility(this, &SSequencer::GetShowTickLines)
+						.DisplayScrubPosition(false)
+						.DisplayTickLines(true)
+						.Clipping(EWidgetClipping::ClipToBounds)
+					]
+					
+					// Overlay that allows drawing over the (unpinned) track area
+					+ SOverlay::Slot()
+					.ZOrder(15)
+					[
+						SNew(STrackAreaOverlay, SequencerRef->GetViewModel())
+					]
 
-					+ SHorizontalBox::Slot()
-					.AutoWidth()
+					// Overlay that draws the scrub position
+					+ SOverlay::Slot()
+					.ZOrder(20)
 					[
-						SAssignNew(SequencerTreeFilterStatusBar, SSequencerTreeFilterStatusBar, SequencerRef)
-						.Visibility(EVisibility::Hidden) // Initially hidden, visible on hover of the info button
+						SNew(SSequencerSectionOverlay, TimeSliderControllerRef)
+						.Visibility(EVisibility::HitTestInvisible)
+						.DisplayScrubPosition(true)
+						.DisplayTickLines(false)
+						.DisplayMarkedFrames(true)
+						.PaintPlaybackRangeArgs(this, &SSequencer::GetSectionPlaybackRangeArgs)
+						.Clipping(EWidgetClipping::ClipToBounds)
 					]
-				]
-			]
-		]
-		// Info Button, Transport Controls and Current Frame
-		+ SGridPanel::Slot(Column0, Row4, SGridPanel::Layer(10))
-		[
-			SNew(SBorder)
-			.BorderImage(FAppStyle::GetBrush(TEXT("ToolPanel.GroupBorder")))
-			.Clipping(EWidgetClipping::ClipToBounds)
-			[
-				SNew(SHorizontalBox)
-	
-				+ SHorizontalBox::Slot()
-				.AutoWidth()
-				.VAlign(VAlign_Center)
-				.HAlign(HAlign_Left)
-				[
-					SNew(SButton)
-					.VAlign(EVerticalAlignment::VAlign_Center)
-					.ButtonStyle(FAppStyle::Get(), TEXT("SimpleButton"))
-					.ToolTipText_Lambda([this] { return LOCTEXT("ShowStatus", "Show Status"); })
-					.ContentPadding(FMargin(1, 0))
-					.Visibility(this, &SSequencer::GetInfoButtonVisibility)
-					.OnHovered_Lambda([this] { SequencerTreeFilterStatusBar->ShowStatusBar(); })
-					.OnUnhovered_Lambda([this] { SequencerTreeFilterStatusBar->FadeOutStatusBar(); })
-					.OnClicked_Lambda([this] { SequencerTreeFilterStatusBar->HideStatusBar(); return FReply::Handled(); })
+
+					+ SOverlay::Slot()
+					.ZOrder(30)
+					.HAlign(HAlign_Left)
+					.VAlign(VAlign_Top)
 					[
-						SNew(SImage)
-						.ColorAndOpacity(FSlateColor::UseForeground())
-						.Image(FAppStyle::Get().GetBrush(TEXT("Icons.Info.Small")))
+						// Transform box
+						SAssignNew(TransformBox, SSequencerTransformBox, SequencerRef, *SequencerSettings, SequencerRef->GetNumericTypeInterface().ToSharedRef())
 					]
-				]
 
-				+ SHorizontalBox::Slot()
-				.AutoWidth()
-				.VAlign(VAlign_Center)
-				.HAlign(HAlign_Left)
-				.Padding(FMargin(CommonPadding, 0.f, 0.f, 0.f))
-				[
-					UToolMenus::Get()->GenerateWidget(TransportControlsLeft, FToolMenuContext(SequencerRef->GetCommandBindings(), nullptr, ContextObject))
+					+ SOverlay::Slot()
+					.ZOrder(40)
+					.HAlign(HAlign_Left)
+					.VAlign(VAlign_Top)
+					[
+						// Stretch box
+						SAssignNew(StretchBox, SSequencerStretchBox, SequencerRef, *SequencerSettings, SequencerRef->GetNumericTypeInterface().ToSharedRef())
+					]
 				]
 
-				+ SHorizontalBox::Slot()
-				.HAlign(HAlign_Center)
+				+ SVerticalBox::Slot()
+				.AutoHeight()
 				[
-					SequencerPtr.Pin()->MakeTransportControls(true)
+					SNew(SSequencerDebugVisualizer, SequencerRef)
+					.ViewRange(FAnimatedRange::WrapAttribute(InArgs._ViewRange))
+					.Visibility(this, &SSequencer::GetDebugVisualizerVisibility)
 				]
 
-				+ SHorizontalBox::Slot()
-				.AutoWidth()
-				.VAlign(VAlign_Center)
-				.HAlign(HAlign_Right)
+				+ SVerticalBox::Slot()
+				.AutoHeight()
 				[
-					SNew(SButton)
-					.VAlign(EVerticalAlignment::VAlign_Center)
-					.ButtonStyle(FAppStyle::Get(), TEXT("NoBorder"))
-					.ContentPadding(FMargin(1, 0))
+					SNew(SBorder)
+					.BorderImage(FAppStyle::GetBrush(TEXT("ToolPanel.GroupBorder")))
+					.BorderBackgroundColor(FLinearColor(.5f, .5f, .5f, 1.f))
+					.Clipping(EWidgetClipping::ClipToBounds)
+					.Padding(0)
 					[
-						SNew(SHorizontalBox)
-
-						+ SHorizontalBox::Slot()
-						.AutoWidth()
-						.VAlign(VAlign_Center)
-						.HAlign(HAlign_Right)
-						.Padding(FMargin(CommonPadding, 0.f, 0.f, 0.f))
+						SNew(SOverlay)
+						+ SOverlay::Slot()
 						[
-							UToolMenus::Get()->GenerateWidget(TransportControlsRight, FToolMenuContext(SequencerRef->GetCommandBindings(), nullptr, ContextObject))
+							BottomTimeSlider.ToSharedRef()
 						]
-
-						+ SHorizontalBox::Slot()
-						.AutoWidth()
-						.VAlign(VAlign_Center)
-						.HAlign(HAlign_Right)
-						.Padding(FMargin(CommonPadding, 0.f, 0.f, 0.f))
+						+ SOverlay::Slot()
 						[
-							PlayTimeDisplay.ToSharedRef()
-						]
-
-						+ SHorizontalBox::Slot()
-						.AutoWidth()
-						.VAlign(VAlign_Center)
-						.HAlign(HAlign_Right)
-						.Padding(FMargin(CommonPadding, 0.f, 0.f, 0.f))
-						[
-							// Current loop index, if any
-							SAssignNew(LoopIndexDisplay, STextBlock)
-							.Font(FCoreStyle::GetDefaultFontStyle(TEXT("Regular"), 8))
-							.Text_Lambda([this]() -> FText {
-								TOptional<int32> LoopIndex = SequencerPtr.Pin()->GetLocalLoopIndex();
-								return LoopIndex ? FText::AsNumber(LoopIndex.GetValue()) : FText();
-							})
+							BottomTimeRange.ToSharedRef()
 						]
 					]
 				]
 			]
-		]
+		];
+}
 
-		// Second column
+TSharedRef<SWidget> SSequencer::ConstructMainTreeView(const TSharedRef<FSequencer>& InSequencer)
+{
+	return SAssignNew(MainSequencerArea, SVerticalBox)
 
-		+ SGridPanel::Slot(Column1, Row1)
-		.Padding(ResizeBarPadding)
-		.RowSpan(3)
+		// Pinned Tree View
+		+ SVerticalBox::Slot()
+		.Expose(OutlinerPinnedSlot)
+		.AutoHeight()
+		.MaxHeight(TAttribute<float>::Create(TAttribute<float>::FGetter::CreateSP(this, &SSequencer::GetPinnedAreaMaxHeight)))
+		.Padding(FMargin(0.f, 0.f, 0.f, CommonPadding))
 		[
-			SNew(SBorder)
-			.BorderImage(FAppStyle::GetBrush(TEXT("ToolPanel.GroupBorder")))
+			SNew(SScrollBorder, PinnedTreeView.ToSharedRef())
+			.Visibility(this, &SSequencer::GetPinnedAreaVisibility)
 			[
-				SNew(SSpacer)
+				PinnedTreeView.ToSharedRef()
 			]
 		]
 
-		+ SGridPanel::Slot(Column1, Row1, SGridPanel::Layer(10))
-		.Padding(ResizeBarPadding)
+		// Tree View
+		+ SVerticalBox::Slot()
 		[
-			SNew(SBorder)
-			.BorderImage(FAppStyle::GetBrush(TEXT("ToolPanel.GroupBorder")))
-			.BorderBackgroundColor(FLinearColor(.5f, .5f, .5f, 1.f))
-			.Padding(0)
-			.Clipping(EWidgetClipping::ClipToBounds)
+			SNew(SOverlay)
+
+			+ SOverlay::Slot()
 			[
-				SNew(SOverlay)
-				+ SOverlay::Slot().ZOrder(1) [ SNew(STopSliderAreaOverlay, SequencerRef->GetViewModel()) ]
-				+ SOverlay::Slot().ZOrder(2) [ TopTimeSlider.ToSharedRef() ]
+				SNew(SScrollBorder, TreeView.ToSharedRef())
+				[
+					TreeView.ToSharedRef()
+				]
 			]
-		]
-		
-		// Overlay that draws the tick lines
-		+ SGridPanel::Slot(Column1, Row2, SGridPanel::Layer(10))
-		.Padding(ResizeBarPadding)
-		[
-			SNew(SSequencerSectionOverlay, TimeSliderControllerRef)
-			.Visibility(this, &SSequencer::GetShowTickLines)
-			.DisplayScrubPosition(false)
-			.DisplayTickLines(true)
-			.Clipping(EWidgetClipping::ClipToBounds)
-		]
-		
-		// Overlay that allows drawing over the (unpinned) track area
-		+ SGridPanel::Slot(Column1, Row2, SGridPanel::Layer(15))
-		.Padding(ResizeBarPadding)
-		[
-			SNew(STrackAreaOverlay, SequencerRef->GetViewModel())
-		]
 
-		// Overlay that draws the scrub position
-		+ SGridPanel::Slot(Column1, Row2, SGridPanel::Layer(20))
-		.Padding(ResizeBarPadding)
-		[
-			SNew(SSequencerSectionOverlay, TimeSliderControllerRef)
-			.Visibility(EVisibility::HitTestInvisible)
-			.DisplayScrubPosition(true)
-			.DisplayTickLines(false)
-			.DisplayMarkedFrames(true)
-			.PaintPlaybackRangeArgs(this, &SSequencer::GetSectionPlaybackRangeArgs)
-			.Clipping(EWidgetClipping::ClipToBounds)
-		]
+			+ SOverlay::Slot()
+			.VAlign(VAlign_Bottom)
+			[
+				SNew(SHorizontalBox)
 
-		+ SGridPanel::Slot(Column1, Row2, SGridPanel::Layer(30))
-		.Padding(ResizeBarPadding)
-		.HAlign(HAlign_Left)
-		.VAlign(VAlign_Top)
-		[
-			// Transform box
-			SAssignNew(TransformBox, SSequencerTransformBox, SequencerRef, *SequencerSettings, SequencerRef->GetNumericTypeInterface().ToSharedRef())
-		]
+				+ SHorizontalBox::Slot()
+				.AutoWidth()
+				[
+					SAssignNew(SequencerTreeFilterStatusBar, SSequencerTreeFilterStatusBar, InSequencer)
+					.Visibility(EVisibility::Hidden) // Initially hidden, visible on hover of the info button
+				]
+			]
+		];
+}
 
-		+ SGridPanel::Slot(Column1, Row2, SGridPanel::Layer(40))
-		.Padding(ResizeBarPadding)
-		.HAlign(HAlign_Left)
-		.VAlign(VAlign_Top)
-		[
-			// Stretch box
-			SAssignNew(StretchBox, SSequencerStretchBox, SequencerRef, *SequencerSettings, SequencerRef->GetNumericTypeInterface().ToSharedRef())
-		]
+TSharedRef<SWidget> SSequencer::ConstructTransportPanel(const TSharedRef<FSequencer>& InSequencer)
+{
+	static const FName TransportControlsLeft("Sequencer.TransportControls.Left");
+	static const FName TransportControlsRight("Sequencer.TransportControls.Right");
+	if (!UToolMenus::Get()->IsMenuRegistered(TransportControlsLeft))
+	{
+		UToolMenus::Get()->RegisterMenu(TransportControlsLeft, NAME_None, EMultiBoxType::ToolBar);
+		UToolMenus::Get()->RegisterMenu(TransportControlsRight, NAME_None, EMultiBoxType::ToolBar);
+	}
 
-		// debug vis
-		+ SGridPanel::Slot(Column1, Row3, SGridPanel::Layer(10))
-		.Padding(ResizeBarPadding)
-		[
-			SNew(SSequencerDebugVisualizer, SequencerRef)
-			.ViewRange(FAnimatedRange::WrapAttribute(InArgs._ViewRange))
-			.Visibility(this, &SSequencer::GetDebugVisualizerVisibility)
-		]
+	USequencerToolMenuContext* ContextObject = NewObject<USequencerToolMenuContext>();
+	ContextObject->WeakSequencer = SequencerPtr;
 
-		// play range sliders
-		+ SGridPanel::Slot(Column1, Row4, SGridPanel::Layer(10))
-		.Padding(ResizeBarPadding)
+	return SNew(SBorder)
+		.BorderImage(FAppStyle::GetBrush(TEXT("ToolPanel.GroupBorder")))
+		.Clipping(EWidgetClipping::ClipToBounds)
 		[
-			SNew(SBorder)
-			.BorderImage(FAppStyle::GetBrush(TEXT("ToolPanel.GroupBorder")))
-			.BorderBackgroundColor(FLinearColor(.5f, .5f, .5f, 1.f))
-			.Clipping(EWidgetClipping::ClipToBounds)
-			.Padding(0)
+			SNew(SHorizontalBox)
+
+			+ SHorizontalBox::Slot()
+			.AutoWidth()
+			.VAlign(VAlign_Center)
+			.HAlign(HAlign_Left)
 			[
-				SNew(SOverlay)
-				+ SOverlay::Slot()
+				SNew(SButton)
+				.VAlign(EVerticalAlignment::VAlign_Center)
+				.ButtonStyle(FAppStyle::Get(), TEXT("SimpleButton"))
+				.ToolTipText_Lambda([this] { return LOCTEXT("ShowStatus", "Show Status"); })
+				.ContentPadding(FMargin(1, 0))
+				.Visibility(this, &SSequencer::GetInfoButtonVisibility)
+				.OnHovered_Lambda([this] { SequencerTreeFilterStatusBar->ShowStatusBar(); })
+				.OnUnhovered_Lambda([this] { SequencerTreeFilterStatusBar->FadeOutStatusBar(); })
+				.OnClicked_Lambda([this] { SequencerTreeFilterStatusBar->HideStatusBar(); return FReply::Handled(); })
 				[
-					BottomTimeSlider.ToSharedRef()
+					SNew(SImage)
+					.ColorAndOpacity(FSlateColor::UseForeground())
+					.Image(FAppStyle::Get().GetBrush(TEXT("Icons.Info.Small")))
 				]
-				+ SOverlay::Slot()
+			]
+
+			+ SHorizontalBox::Slot()
+			.AutoWidth()
+			.VAlign(VAlign_Center)
+			.HAlign(HAlign_Left)
+			.Padding(FMargin(CommonPadding, 0.f, 0.f, 0.f))
+			[
+				UToolMenus::Get()->GenerateWidget(TransportControlsLeft, FToolMenuContext(InSequencer->GetCommandBindings(), nullptr, ContextObject))
+			]
+
+			+ SHorizontalBox::Slot()
+			.HAlign(HAlign_Center)
+			[
+				SequencerPtr.Pin()->MakeTransportControls(true)
+			]
+
+			+ SHorizontalBox::Slot()
+			.AutoWidth()
+			.VAlign(VAlign_Center)
+			.HAlign(HAlign_Right)
+			[
+				SNew(SButton)
+				.VAlign(EVerticalAlignment::VAlign_Center)
+				.ButtonStyle(FAppStyle::Get(), TEXT("NoBorder"))
+				.ContentPadding(FMargin(1, 0))
 				[
-					BottomTimeRange.ToSharedRef()
+					SNew(SHorizontalBox)
+
+					+ SHorizontalBox::Slot()
+					.AutoWidth()
+					.VAlign(VAlign_Center)
+					.HAlign(HAlign_Right)
+					.Padding(FMargin(CommonPadding, 0.f, 0.f, 0.f))
+					[
+						UToolMenus::Get()->GenerateWidget(TransportControlsRight, FToolMenuContext(InSequencer->GetCommandBindings(), nullptr, ContextObject))
+					]
+
+					+ SHorizontalBox::Slot()
+					.AutoWidth()
+					.VAlign(VAlign_Center)
+					.HAlign(HAlign_Right)
+					.Padding(FMargin(CommonPadding, 0.f, 0.f, 0.f))
+					[
+						PlayTimeDisplay.ToSharedRef()
+					]
+
+					+ SHorizontalBox::Slot()
+					.AutoWidth()
+					.VAlign(VAlign_Center)
+					.HAlign(HAlign_Right)
+					.Padding(FMargin(CommonPadding, 0.f, 0.f, 0.f))
+					[
+						// Current loop index, if any
+						SAssignNew(LoopIndexDisplay, STextBlock)
+						.Font(FCoreStyle::GetDefaultFontStyle(TEXT("Regular"), 8))
+						.Text_Lambda([this]() -> FText {
+							TOptional<int32> LoopIndex = SequencerPtr.Pin()->GetLocalLoopIndex();
+							return LoopIndex ? FText::AsNumber(LoopIndex.GetValue()) : FText();
+						})
+					]
 				]
 			]
 		];
@@ -1252,7 +1240,7 @@ void SSequencer::UpdateOutlinerViewColumns()
 
 	for (FSequencerOutlinerColumnVisibility ColumnVisibility : OutlinerColumnVisibilities)
 	{
-		ColumnVisibilitySettings.Add(FColumnVisibilitySetting(ColumnVisibility.Column->GetColumnName(), ColumnVisibility.bIsColumnVisible));
+		ColumnVisibilitySettings.AddUnique(FColumnVisibilitySetting(ColumnVisibility.Column->GetColumnName(), ColumnVisibility.bIsColumnVisible));
 	}
 
 	GetSequencerSettings()->SetOutlinerColumnVisibility(ColumnVisibilitySettings);
@@ -1282,6 +1270,8 @@ void SSequencer::InitializeOutlinerColumns()
 		return;
 	}
 
+	OutlinerColumnVisibilities.Reset();
+
 	const TMap<FName, TSharedPtr<IOutlinerColumn>>& RegisteredColumns = Sequencer->GetOutlinerColumns();
 	
 	// Retrieve previously saved column names and visibilities
@@ -1295,7 +1285,7 @@ void SSequencer::InitializeOutlinerColumns()
 		if (OutlinerColumn)
 		{
 			ColumnNamesFoundInSettings.Add(ColumnVisibility.ColumnName);
-			OutlinerColumnVisibilities.Add(FSequencerOutlinerColumnVisibility(*OutlinerColumn, ColumnVisibility.bIsVisible));
+			OutlinerColumnVisibilities.AddUnique(FSequencerOutlinerColumnVisibility(*OutlinerColumn, ColumnVisibility.bIsVisible));
 		}
 	}
 
@@ -1304,7 +1294,7 @@ void SSequencer::InitializeOutlinerColumns()
 	{
 		if (!ColumnNamesFoundInSettings.Contains(RegisteredColumn.Key))
 		{
-			OutlinerColumnVisibilities.Add(FSequencerOutlinerColumnVisibility(RegisteredColumn.Value));
+			OutlinerColumnVisibilities.AddUnique(FSequencerOutlinerColumnVisibility(RegisteredColumn.Value));
 		}
 	}
 
@@ -3391,22 +3381,6 @@ EFrameNumberDisplayFormats SSequencer::GetTimeDisplayFormat() const
 	return GetSequencerSettings()->GetTimeDisplayFormat();
 }
 
-void SSequencer::OnSplitterFinishedResizing()
-{
-	SSplitter::FSlot const& LeftSplitterSlot = TreeViewSplitter->Splitter->SlotAt(0);
-	SSplitter::FSlot const& RightSplitterSlot = TreeViewSplitter->Splitter->SlotAt(1);
-
-	OnColumnFillCoefficientChanged(LeftSplitterSlot.GetSizeValue(), 0);
-	OnColumnFillCoefficientChanged(RightSplitterSlot.GetSizeValue(), 1);
-
-	GetSequencerSettings()->SetTreeViewWidth(LeftSplitterSlot.GetSizeValue());
-}
-
-void SSequencer::OnColumnFillCoefficientChanged(float FillCoefficient, int32 ColumnIndex)
-{
-	ColumnFillCoefficients[ColumnIndex] = FillCoefficient;
-}
-
 void SSequencer::OnCurveEditorVisibilityChanged(bool bShouldBeVisible)
 {
 	using namespace UE::Sequencer;
@@ -4333,4 +4307,36 @@ FText SSequencer::GetSidebarSelectionDrawerToolTipText() const
 	return ToolTipText;
 }
 
+FOptionalSize SSequencer::GetOutlinerWidth() const
+{
+	const USequencerSettings* const SequencerSettings = GetSequencerSettings();
+	if (!SequencerSettings)
+	{
+		return DefaultOutlinerWidth;
+	}
+
+	float CurrentOutlinerWidth = SequencerSettings->GetTreeViewWidth();
+
+	if (OutlinerAndTrackAreaContent.IsValid())
+	{
+		const FGeometry& CachedGeometry = OutlinerAndTrackAreaContent->GetCachedGeometry();
+		const FVector2f LocalSize = CachedGeometry.GetLocalSize();
+
+		if (LocalSize.X > 0.f)
+		{
+			// Limit the outliner size to the parent size
+			const float MaxDynamicWidth = LocalSize.X - MinOutlinerWidth;
+			if (CurrentOutlinerWidth > MaxDynamicWidth)
+			{
+				CurrentOutlinerWidth = MaxDynamicWidth;
+			}
+
+			const float MinWidth = FMath::Min<float>(MaxDynamicWidth, MinOutlinerWidth);
+			return FMath::Clamp<float>(CurrentOutlinerWidth, MinWidth, MaxDynamicWidth);
+		}
+	}
+
+	return FMath::Max<float>(CurrentOutlinerWidth, MinOutlinerWidth);
+}
+
 #undef LOCTEXT_NAMESPACE
diff --git a/Engine/Source/Editor/Sequencer/Private/SSequencer.h b/Engine/Source/Editor/Sequencer/Private/SSequencer.h
--- a/Engine/Source/Editor/Sequencer/Private/SSequencer.h
+++ b/Engine/Source/Editor/Sequencer/Private/SSequencer.h
@@ -42,7 +42,6 @@ class SCurveEditorPanel;
 class SBox;
 class SDockTab;
 class SFilterSearchBox;
-class SGridPanel;
 class SSidebar;
 class SSidebarContainer;
 class SWindow;
@@ -122,6 +121,8 @@ struct FSequencerOutlinerColumnVisibility
 
 	FSequencerOutlinerColumnVisibility(TSharedPtr<UE::Sequencer::IOutlinerColumn> InColumn);
 	FSequencerOutlinerColumnVisibility(TSharedPtr<UE::Sequencer::IOutlinerColumn> InColumn, bool bInIsColumnVisible);
+
+	bool operator==(const FSequencerOutlinerColumnVisibility& InOther) const;
 };
 
 
@@ -135,6 +136,9 @@ class SSequencer
 {
 public:
 
+	static constexpr float DefaultOutlinerWidth = 250.f;
+	static constexpr float MinOutlinerWidth = 150.f;
+
 	DECLARE_DELEGATE_OneParam( FOnToggleBoolOption, bool )
 	SLATE_BEGIN_ARGS( SSequencer )
 	{ }
@@ -374,8 +378,9 @@ private:
 
 	TSharedRef<SWidget> ConstructSidebarContent();
 	TSharedRef<SWidget> ConstructFilterBarContent();
-	TSharedRef<SWidget> ConstructGridOverlayContent();
-	TSharedRef<SGridPanel> ConstructTrackAreaGridPanel(const FArguments& InArgs, const FTimeSliderArgs& InTimeSliderArgs);
+	TSharedRef<SSplitter> ConstructOutlinerAndTrackArea(const FArguments& InArgs, const FTimeSliderArgs& InTimeSliderArgs);
+	TSharedRef<SWidget> ConstructMainTreeView(const TSharedRef<FSequencer>& InSequencer);
+	TSharedRef<SWidget> ConstructTransportPanel(const TSharedRef<FSequencer>& InSequencer);
 
 	/** Can be called anytime after Construct() to rebuild most of the main content (except for the toolbar). */
 	void RebuildForSidebar();
@@ -488,15 +493,6 @@ private:
 	*/
 	void OnValueSnapIntervalChanged( float InInterval );
 
-	/**
-	 * @return The fill percentage of the animation outliner
-	 */
-	float GetColumnFillCoefficient(int32 ColumnIndex) const
-	{
-		ensure(ColumnIndex == 0 || ColumnIndex == 1);
-		return ColumnFillCoefficients[ColumnIndex];
-	}
-
 	void AddFromContentBrowser();
 	bool CanAddFromContentBrowser() const;
 
@@ -567,11 +563,6 @@ private:
 	/** What is the preferred display format for time values. */
 	EFrameNumberDisplayFormats GetTimeDisplayFormat() const;
 
-	/** Called when a column fill percentage is changed by a splitter slot. */
-	void OnColumnFillCoefficientChanged(float FillCoefficient, int32 ColumnIndex);
-
-	void OnSplitterFinishedResizing();
-
 	/** Gets paint options for painting the playback range on sequencer */
 	FPaintPlaybackRangeArgs GetSectionPlaybackRangeArgs() const;
 
@@ -723,13 +714,15 @@ private:
 
 	FText GetSidebarSelectionDrawerToolTipText() const;
 
+	FOptionalSize GetOutlinerWidth() const;
+
 	/** This is the main container widget for the Sequencer (minus the toolbar). This is the parent of
 	 * the GridPanel widget below and is used to rebuild content when the sidebar is added/removed. */
 	TSharedPtr<SBox> MainContentContainer;
 
-	/** Grid panel that holds most of the content and contains MainSequencerArea widget below.
+	/** Splitter that holds most of the content and contains MainSequencerArea widget below.
 	 * Built once at Sequencer initialize and used when rebuilding the main content container when the sidebar is added/removed. */
-	TSharedPtr<SGridPanel> GridPanel;
+	TSharedPtr<SSplitter> OutlinerAndTrackAreaContent;
 
 	/** Contains the outliner tree and track area. Used by GetPinnedAreaMaxHeight().
 	 * Built once at Sequencer initialize and used when rebuilding the main content container when the sidebar is added/removed. */
@@ -794,8 +787,9 @@ private:
 	/** Container for the toolbar, so that we can re-create it as needed. */
 	TSharedPtr<SBox> ToolbarContainer;
 
-	/** The fill coefficients of each column in the grid. */
-	float ColumnFillCoefficients[2];
+	/** Exposing these slot so the equivalent track areas can match their height */
+	SVerticalBox::FSlot* OutlinerSearchAndFilterSlot = nullptr;
+	SVerticalBox::FSlot* OutlinerPinnedSlot = nullptr;
 
 	/** List of registered outliner columns with their visibility states */
 	TArray<FSequencerOutlinerColumnVisibility> OutlinerColumnVisibilities;
diff --git a/Engine/Source/Editor/Sequencer/Private/SequencerSettings.cpp b/Engine/Source/Editor/Sequencer/Private/SequencerSettings.cpp
--- a/Engine/Source/Editor/Sequencer/Private/SequencerSettings.cpp
+++ b/Engine/Source/Editor/Sequencer/Private/SequencerSettings.cpp
@@ -74,7 +74,7 @@ USequencerSettings::USequencerSettings( const FObjectInitializer& ObjectInitiali
 	FrameNumberDisplayFormat = EFrameNumberDisplayFormats::Seconds;
 	bAutoExpandNodesOnSelection = true;
 	bRestoreOriginalViewportOnCameraCutUnlock = true;
-	TreeViewWidth = 0.3f;
+	TreeViewWidth = SSequencer::DefaultOutlinerWidth;
 	bShowTickLines = true;
 	bShowSequencerToolbar = true;
 	bShowMarkedFrames = true;
@@ -104,6 +104,18 @@ USequencerSettings::USequencerSettings( const FObjectInitializer& ObjectInitiali
 	CurveEditorZoomScaling.MaxVerticalZoomOut = 10000;
 }
 
+void USequencerSettings::PostLoad()
+{
+	Super::PostLoad();
+
+	// Because we want to remove proportional outliner sizing,
+	// we need to convert the percentage to a pixel width
+	if (TreeViewWidth <= 1.f)
+	{
+		TreeViewWidth = SSequencer::DefaultOutlinerWidth;
+	}
+}
+
 void USequencerSettings::PostEditChangeProperty(FPropertyChangedEvent& PropertyChangedEvent)
 {
 	Super::PostEditChangeProperty(PropertyChangedEvent);
diff --git a/Engine/Source/Editor/Sequencer/Private/SequencerTimeSliderController.cpp b/Engine/Source/Editor/Sequencer/Private/SequencerTimeSliderController.cpp
--- a/Engine/Source/Editor/Sequencer/Private/SequencerTimeSliderController.cpp
+++ b/Engine/Source/Editor/Sequencer/Private/SequencerTimeSliderController.cpp
@@ -2346,6 +2346,14 @@ void FSequencerTimeSliderController::ClampViewRange(double& NewRangeMin, double&
 
 void FSequencerTimeSliderController::SetViewRange( double NewRangeMin, double NewRangeMax, EViewRangeInterpolation Interpolation )
 {
+	// It's possible to quickly drag the outliner / track area splitter and trigger 0 value calls to
+	// this function when the parent widget is sized smaller than the outliner. Guard against this 
+	// "infinite" zoom to avoid SlateVertex buffer overflow ensure being triggered.
+	if (NewRangeMin == 0.f || NewRangeMax == 0.f)
+	{
+		return;
+	}
+
 	// Clamp to a minimum size to avoid zero-sized or negative visible ranges
 	double MinVisibleTimeRange = FFrameNumber(1) / GetTickResolution();
 	TRange<double> ExistingViewRange  = GetViewRange();
diff --git a/Engine/Source/Editor/Sequencer/Public/SequencerSettings.h b/Engine/Source/Editor/Sequencer/Public/SequencerSettings.h
--- a/Engine/Source/Editor/Sequencer/Public/SequencerSettings.h
+++ b/Engine/Source/Editor/Sequencer/Public/SequencerSettings.h
@@ -141,8 +141,11 @@ public:
 	DECLARE_MULTICAST_DELEGATE_OneParam( FOnAllowEditsModeChanged, EAllowEditsMode );
 	DECLARE_MULTICAST_DELEGATE(FOnLoopStateChanged);
 	DECLARE_MULTICAST_DELEGATE(FOnTimeDisplayFormatChanged);
-	
+
+	//~ Begin UObject
+	virtual void PostLoad() override;
 	UE_API virtual void PostEditChangeProperty(FPropertyChangedEvent& PropertyChangedEvent) override;
+	//~ End UObject
 
 	/** Gets the current auto change mode. */
 	UE_API EAutoChangeMode GetAutoChangeMode() const;
diff --git a/Engine/Source/Runtime/Core/Public/Misc/FrameTime.h b/Engine/Source/Runtime/Core/Public/Misc/FrameTime.h
--- a/Engine/Source/Runtime/Core/Public/Misc/FrameTime.h
+++ b/Engine/Source/Runtime/Core/Public/Misc/FrameTime.h
@@ -268,7 +268,7 @@ inline FFrameTime::FFrameTime(FFrameNumber InFrameNumber, float InSubFrame)
 	// Hack to ensure that SubFrames are in a sensible range of precision to work around
 	// problems with FloorToXYZ returning the wrong thing for very small negative numbers
 	SubFrame = FMath::Clamp(InSubFrame + 0.5f - 0.5f, 0.f, MaxSubframe);
-	checkSlow(InSubFrame >= 0.f && InSubFrame < 1.f);
+	checkSlow(SubFrame >= 0.f && SubFrame < 1.f);
 }
 
 
