[Sequencer]
- Add text expression that enables all common filters to be used as text expressions if they don't have a custom setup text expression class
- Abstraction for text expression help dialog so it doesn't only work with text expression contexts
- Fix True/False suggestions not showing for boolean text expressions
- Fix Unbound filter not working correctly

#Synced-CL 50092247

diff --git a/Engine/Plugins/Experimental/SequenceNavigator/Source/SequenceNavigator/Private/Filters/Filters/NavigationToolFilterCollection.cpp b/Engine/Plugins/Experimental/SequenceNavigator/Source/SequenceNavigator/Private/Filters/Filters/NavigationToolFilterCollection.cpp
--- a/Engine/Plugins/Experimental/SequenceNavigator/Source/SequenceNavigator/Private/Filters/Filters/NavigationToolFilterCollection.cpp
+++ b/Engine/Plugins/Experimental/SequenceNavigator/Source/SequenceNavigator/Private/Filters/Filters/NavigationToolFilterCollection.cpp
@@ -127,8 +127,7 @@ TArray<FText> FNavigationToolFilterCollection::GetFilterDisplayNames() const
 	return OutDisplayNames;
 }
 
-TArray<TSharedRef<FNavigationToolFilter>> FNavigationToolFilterCollection::GetAllFilters(const bool bInCheckSupportsSequence
-	, const TArray<TSharedRef<FFilterCategory>>& InCategories) const
+TArray<TSharedRef<FNavigationToolFilter>> FNavigationToolFilterCollection::GetAllFilters(const TArray<TSharedRef<FFilterCategory>>& InCategories) const
 {
 	TArray<TSharedRef<FNavigationToolFilter>> OutFilters;
 
diff --git a/Engine/Plugins/Experimental/SequenceNavigator/Source/SequenceNavigator/Private/Filters/NavigationToolFilterBar.cpp b/Engine/Plugins/Experimental/SequenceNavigator/Source/SequenceNavigator/Private/Filters/NavigationToolFilterBar.cpp
--- a/Engine/Plugins/Experimental/SequenceNavigator/Source/SequenceNavigator/Private/Filters/NavigationToolFilterBar.cpp
+++ b/Engine/Plugins/Experimental/SequenceNavigator/Source/SequenceNavigator/Private/Filters/NavigationToolFilterBar.cpp
@@ -562,6 +562,14 @@ void FNavigationToolFilterBar::EnableAllFilters(const bool bInEnable, const TArr
 	EnableCustomTextFilters(bInEnable, ExceptionCustomTextFilters);
 }
 
+void FNavigationToolFilterBar::GetCommonFiltersImpl(TArray<TSharedPtr<void>>& OutTypeErasedFilters) const
+{
+	for (const TSharedRef<FNavigationToolFilter>& Filter : CommonFilters->GetAllFilters())
+	{
+		OutTypeErasedFilters.Add(Filter);
+	}
+}
+
 void FNavigationToolFilterBar::ActivateCommonFilters(const bool bInActivate, const TArray<FString>& InExceptionFilterNames)
 {
 	TArray<TSharedRef<FNavigationToolFilter>> ExceptionFilters;
@@ -941,7 +949,7 @@ int32 FNavigationToolFilterBar::GetFilteredDisplayNodeCount() const
 
 TArray<TSharedRef<FNavigationToolFilter>> FNavigationToolFilterBar::GetCommonFilters(const TArray<TSharedRef<FFilterCategory>>& InCategories) const
 {
-	return CommonFilters->GetAllFilters(true, InCategories);
+	return CommonFilters->GetAllFilters(InCategories);
 }
 
 bool FNavigationToolFilterBar::AnyCustomTextFilterActive() const
@@ -1340,7 +1348,7 @@ TArray<TSharedRef<FNavigationToolFilter>> FNavigationToolFilterBar::GetFilterLis
 {
 	TArray<TSharedRef<FNavigationToolFilter>> AllFilters;
 
-	AllFilters.Append(CommonFilters->GetAllFilters(true));
+	AllFilters.Append(CommonFilters->GetAllFilters());
 
 	AllFilters.Add(TextFilter);
 
@@ -1779,7 +1787,8 @@ void FNavigationToolFilterBar::OpenTextExpressionHelp()
 	FFilterExpressionHelpDialogConfig Config;
 	Config.IdentifierName = TEXT("NavigationToolCustomTextFilterHelp");
 	Config.DialogTitle = LOCTEXT("CustomTextFilterHelp", "Sequence Navigator Custom Text Filter Help");
-	Config.TextFilterExpressionContexts = GetTextFilterExpressionContexts();
+	Config.AddCommonFilterHelpEntries<FNavigationToolFilter>(SharedThis(this));
+
 	SFilterExpressionHelpDialog::Open(MoveTemp(Config));
 }
 
@@ -1795,6 +1804,51 @@ void FNavigationToolFilterBar::CreateNewTextFilter()
 	CreateWindow_AddCustomTextFilter(DefaultNewCustomTextFilterData(NewFilterString));
 }
 
+bool FNavigationToolFilterBar::DoesCommonFilterPass(const TViewModelPtr<FViewModel>& InViewModel, const FString& InFilterName) const
+{
+	bool bAnyPassed = false;
+
+	CommonFilters->ForEachFilter([this, &bAnyPassed, &InViewModel, &InFilterName]
+		(const TSharedRef<FNavigationToolFilter>& InFilter)
+		{
+			if (InFilter->GetName().Equals(InFilterName, ESearchCase::IgnoreCase)
+				&& InFilter->PassesFilter(InViewModel.ImplicitCast()))
+			{
+				bAnyPassed = true;
+				return false; // Stop looping
+			}
+			return true;
+		});
+
+	return bAnyPassed;
+}
+
+void FNavigationToolFilterBar::GetCommonFilterKeySuggestions(TArray<FSequencerFilterSuggestion>& OutSuggestions) const
+{
+	for (const TSharedRef<FNavigationToolFilter>& CommonFilter : CommonFilters->GetAllFilters())
+	{
+		FSequencerFilterSuggestion NewSuggestion;
+		NewSuggestion.Suggestion = CommonFilter->GetName();
+		NewSuggestion.DisplayName = CommonFilter->GetDisplayName();
+		NewSuggestion.Description = CommonFilter->GetDefaultToolTipText();
+
+		OutSuggestions.AddUnique(NewSuggestion);
+	}
+}
+
+void FNavigationToolFilterBar::GetCommonFilterValueSuggestions(const FString& InKey, TArray<FSequencerFilterSuggestion>& OutSuggestions) const
+{
+	for (const TSharedRef<FNavigationToolFilter>& CommonFilter : CommonFilters->GetAllFilters())
+	{
+		if (InKey.Equals(CommonFilter->GetName(), ESearchCase::IgnoreCase))
+		{
+			// Suggestions for common filters text expressions will always be boolean for turning on/off
+			FSequencerFilterSuggestion::MakeBooleanSuggestions(OutSuggestions);
+			break;
+		}
+	}
+}
+
 } // namespace UE::SequenceNavigator
 
 #undef LOCTEXT_NAMESPACE
diff --git a/Engine/Plugins/Experimental/SequenceNavigator/Source/SequenceNavigator/Public/Filters/Filters/NavigationToolFilterCollection.h b/Engine/Plugins/Experimental/SequenceNavigator/Source/SequenceNavigator/Public/Filters/Filters/NavigationToolFilterCollection.h
--- a/Engine/Plugins/Experimental/SequenceNavigator/Source/SequenceNavigator/Public/Filters/Filters/NavigationToolFilterCollection.h
+++ b/Engine/Plugins/Experimental/SequenceNavigator/Source/SequenceNavigator/Public/Filters/Filters/NavigationToolFilterCollection.h
@@ -64,8 +64,7 @@ public:
 	/** Gets all the available track filter names */
 	TArray<FText> GetFilterDisplayNames() const;
 
-	TArray<TSharedRef<FNavigationToolFilter>> GetAllFilters(const bool bInCheckSupportsSequence
-		, const TArray<TSharedRef<FFilterCategory>>& InCategories = {}) const;
+	TArray<TSharedRef<FNavigationToolFilter>> GetAllFilters(const TArray<TSharedRef<FFilterCategory>>& InCategories = {}) const;
 
 	TSet<TSharedRef<FFilterCategory>> GetCategories(const TSet<TSharedRef<FNavigationToolFilter>>* InFilters = nullptr) const;
 
diff --git a/Engine/Plugins/Experimental/SequenceNavigator/Source/SequenceNavigator/Public/Filters/NavigationToolFilterBar.h b/Engine/Plugins/Experimental/SequenceNavigator/Source/SequenceNavigator/Public/Filters/NavigationToolFilterBar.h
--- a/Engine/Plugins/Experimental/SequenceNavigator/Source/SequenceNavigator/Public/Filters/NavigationToolFilterBar.h
+++ b/Engine/Plugins/Experimental/SequenceNavigator/Source/SequenceNavigator/Public/Filters/NavigationToolFilterBar.h
@@ -209,7 +209,14 @@ public:
 	virtual void SaveCurrentFilterSetAsCustomTextFilter() override;
 	virtual void CreateNewTextFilter() override;
 
+	virtual bool DoesCommonFilterPass(const Sequencer::TViewModelPtr<Sequencer::FViewModel>& InViewModel, const FString& InFilterName) const override;
+
+	virtual void GetCommonFilterKeySuggestions(TArray<FSequencerFilterSuggestion>& OutSuggestions) const override;
+	virtual void GetCommonFilterValueSuggestions(const FString& InKey, TArray<FSequencerFilterSuggestion>& OutSuggestions) const override;
+
 protected:
+	virtual void GetCommonFiltersImpl(TArray<TSharedPtr<void>>& OutTypeErasedFilters) const override;
+
 	void CreateDefaultFilters();
 	void CreateCustomTextFiltersFromConfig();
 
diff --git a/Engine/Source/Editor/Sequencer/Private/Filters/Filters/SequencerTrackFilter_Text.cpp b/Engine/Source/Editor/Sequencer/Private/Filters/Filters/SequencerTrackFilter_Text.cpp
--- a/Engine/Source/Editor/Sequencer/Private/Filters/Filters/SequencerTrackFilter_Text.cpp
+++ b/Engine/Source/Editor/Sequencer/Private/Filters/Filters/SequencerTrackFilter_Text.cpp
@@ -15,6 +15,7 @@
 #include "Filters/TextExpressions/SequencerTextFilterExpression_CustomBinding.h"
 #include "Filters/TextExpressions/SequencerTextFilterExpression_Deactivated.h"
 #include "Filters/TextExpressions/SequencerTextFilterExpression_EmptyBinding.h"
+#include "Filters/TextExpressions/SequencerTextFilterExpression_Filters.h"
 #include "Filters/TextExpressions/SequencerTextFilterExpression_Group.h"
 #include "Filters/TextExpressions/SequencerTextFilterExpression_Keyed.h"
 #include "Filters/TextExpressions/SequencerTextFilterExpression_Level.h"
@@ -45,6 +46,7 @@ FSequencerTrackFilter_Text::FSequencerTrackFilter_Text(ISequencerTrackFilters& I
 		});
 
 	// Ordered by importance and most often used. This will dictate the order of display in the text expressions help dialog.
+	TextFilterExpressionContexts.Add(MakeShared<FSequencerTextFilterExpression_Filters>(InFilterInterface));
 	TextFilterExpressionContexts.Add(MakeShared<FSequencerTextFilterExpression_Name>(InFilterInterface));
 	TextFilterExpressionContexts.Add(MakeShared<FSequencerTextFilterExpression_TrackClass>(InFilterInterface));
 	TextFilterExpressionContexts.Add(MakeShared<FSequencerTextFilterExpression_ObjectClass>(InFilterInterface));
diff --git a/Engine/Source/Editor/Sequencer/Private/Filters/Filters/SequencerTrackFilter_Unbound.cpp b/Engine/Source/Editor/Sequencer/Private/Filters/Filters/SequencerTrackFilter_Unbound.cpp
--- a/Engine/Source/Editor/Sequencer/Private/Filters/Filters/SequencerTrackFilter_Unbound.cpp
+++ b/Engine/Source/Editor/Sequencer/Private/Filters/Filters/SequencerTrackFilter_Unbound.cpp
@@ -49,12 +49,24 @@ bool FSequencerTrackFilter_Unbound::PassesFilter(FSequencerTrackFilterType InIte
 {
 	bool bPassed = false;
 
-	if (const TViewModelPtr<FObjectBindingModel> ObjectBindingModel = InItem->FindAncestorOfType<FObjectBindingModel>())
+	if (const TViewModelPtr<FObjectBindingModel> ObjectBindingModel = InItem->FindAncestorOfType<FObjectBindingModel>(true))
 	{
 		if (const TViewModelPtr<FSequenceModel> SequenceModel = InItem->FindAncestorOfType<FSequenceModel>())
 		{
 			const TArrayView<TWeakObjectPtr<>> BoundObjects = FilterInterface.GetSequencer().FindBoundObjects(ObjectBindingModel->GetObjectGuid(), SequenceModel->GetSequenceID());
-			if (BoundObjects.IsEmpty())
+
+			bool bAnyValidBinding = false;
+
+			for (const TWeakObjectPtr<>& WeakObj : BoundObjects)
+			{
+				if (WeakObj.IsValid())
+				{
+					bAnyValidBinding = true;
+					break;
+				}
+			}
+
+			if (!bAnyValidBinding)
 			{
 				bPassed = true;
 			}
diff --git a/Engine/Source/Editor/Sequencer/Private/Filters/SequencerFilterBar.cpp b/Engine/Source/Editor/Sequencer/Private/Filters/SequencerFilterBar.cpp
--- a/Engine/Source/Editor/Sequencer/Private/Filters/SequencerFilterBar.cpp
+++ b/Engine/Source/Editor/Sequencer/Private/Filters/SequencerFilterBar.cpp
@@ -19,6 +19,7 @@
 #include "Filters/SequencerTrackFilterCollection.h"
 #include "Filters/SequencerTrackFilterCommands.h"
 #include "Filters/SequencerTrackFilterExtension.h"
+#include "Filters/Utils/FilterViewUtils.h"
 #include "Filters/Widgets/SFilterBarIsolateHideShow.h"
 #include "Filters/Widgets/SFilterExpressionHelpDialog.h"
 #include "Filters/Widgets/SSequencerFilterBar.h"
@@ -658,6 +659,14 @@ void FSequencerFilterBar::EnableAllFilters(const bool bInEnable, const TArray<FS
 	EnableCustomTextFilters(bInEnable, ExceptionCustomTextFilters);
 }
 
+void FSequencerFilterBar::GetCommonFiltersImpl(TArray<TSharedPtr<void>>& OutTypeErasedFilters) const
+{
+	for (const TSharedRef<FSequencerTrackFilter>& Filter : CommonFilters->GetAllFilters())
+	{
+		OutTypeErasedFilters.Add(Filter);
+	}
+}
+
 void FSequencerFilterBar::ActivateCommonFilters(const bool bInActivate, const TArray<FString>& InExceptionFilterNames)
 {
 	TArray<TSharedRef<FSequencerTrackFilter>> ExceptionFilters;
@@ -1456,11 +1465,20 @@ const FTextFilterExpressionEvaluator& FSequencerFilterBar::GetTextFilterExpressi
 TArray<TSharedRef<ISequencerTextFilterExpressionContext>> FSequencerFilterBar::GetTextFilterExpressionContexts() const
 {
 	TArray<TSharedRef<ISequencerTextFilterExpressionContext>> TextExpressions;
+
 	Algo::Transform(TextFilter->GetTextFilterExpressionContexts(), TextExpressions
 		, [](TSharedRef<FSequencerTextFilterExpressionContext> InExpressionContext)
 		{
 			return InExpressionContext;
 		});
+
+	CommonFilters->ForEachFilter([]
+		(const TSharedRef<FSequencerTrackFilter>& InFilter)
+		{
+
+			return true;
+		});
+
 	return TextExpressions;
 }
 
@@ -1864,11 +1882,7 @@ bool FSequencerFilterBar::IsFilterSupported(const FString& InFilterName) const
 
 void FSequencerFilterBar::OpenTextExpressionHelp()
 {
-	FFilterExpressionHelpDialogConfig Config;
-	Config.IdentifierName = TEXT("SequencerCustomTextFilterHelp");
-	Config.DialogTitle = LOCTEXT("CustomTextFilterHelp", "Sequencer Custom Text Filter Help");
-	Config.TextFilterExpressionContexts = GetTextFilterExpressionContexts();
-	SFilterExpressionHelpDialog::Open(MoveTemp(Config));
+	OpenDialog_TextExpressionHelp(SharedThis(this));
 }
 
 void FSequencerFilterBar::SaveCurrentFilterSetAsCustomTextFilter()
@@ -1885,4 +1899,48 @@ void FSequencerFilterBar::CreateNewTextFilter()
 	//CreateWindow_AddCustomTextFilter(DefaultNewCustomTextFilterData(NewFilterString));
 }
 
+bool FSequencerFilterBar::DoesCommonFilterPass(const TViewModelPtr<FViewModel>& InViewModel, const FString& InFilterName) const
+{
+	bool bAnyPassed = false;
+
+	CommonFilters->ForEachFilter([this, &bAnyPassed, &InViewModel, &InFilterName]
+		(const TSharedRef<FSequencerTrackFilter>& InFilter)
+		{
+			if (InFilter->GetName().Equals(InFilterName, ESearchCase::IgnoreCase)
+				&& InFilter->PassesFilter(InViewModel.ImplicitCast()))
+			{
+				bAnyPassed = true;
+				return false; // Stop looping
+			}
+			return true;
+		});
+
+	return bAnyPassed;
+}
+
+void FSequencerFilterBar::GetCommonFilterKeySuggestions(TArray<FSequencerFilterSuggestion>& OutSuggestions) const
+{
+	for (const TSharedRef<FSequencerTrackFilter>& CommonFilter : CommonFilters->GetAllFilters())
+	{
+		FSequencerFilterSuggestion NewSuggestion;
+		NewSuggestion.Suggestion = CommonFilter->GetName();
+		NewSuggestion.DisplayName = CommonFilter->GetDisplayName();
+		NewSuggestion.Description = CommonFilter->GetDefaultToolTipText();
+
+		OutSuggestions.AddUnique(NewSuggestion);
+	}
+}
+
+void FSequencerFilterBar::GetCommonFilterValueSuggestions(const FString& InKey, TArray<FSequencerFilterSuggestion>& OutSuggestions) const
+{
+	for (const TSharedRef<FSequencerTrackFilter>& CommonFilter : CommonFilters->GetAllFilters())
+	{
+		if (InKey.Equals(CommonFilter->GetName(), ESearchCase::IgnoreCase))
+		{
+			FSequencerFilterSuggestion::MakeBooleanSuggestions(OutSuggestions);
+			break;
+		}
+	}
+}
+
 #undef LOCTEXT_NAMESPACE
diff --git a/Engine/Source/Editor/Sequencer/Private/Filters/SequencerFilterBar.h b/Engine/Source/Editor/Sequencer/Private/Filters/SequencerFilterBar.h
--- a/Engine/Source/Editor/Sequencer/Private/Filters/SequencerFilterBar.h
+++ b/Engine/Source/Editor/Sequencer/Private/Filters/SequencerFilterBar.h
@@ -254,7 +254,14 @@ public:
 	virtual void SaveCurrentFilterSetAsCustomTextFilter() override;
 	virtual void CreateNewTextFilter() override;
 
+	virtual bool DoesCommonFilterPass(const UE::Sequencer::TViewModelPtr<UE::Sequencer::FViewModel>& InViewModel, const FString& InFilterName) const override;
+
+	virtual void GetCommonFilterKeySuggestions(TArray<FSequencerFilterSuggestion>& OutSuggestions) const override;
+	virtual void GetCommonFilterValueSuggestions(const FString& InKey, TArray<FSequencerFilterSuggestion>& OutSuggestions) const override;
+
 protected:
+	virtual void GetCommonFiltersImpl(TArray<TSharedPtr<void>>& OutTypeErasedFilters) const override;
+
 	void CreateDefaultFilters();
 	void CreateCustomTextFiltersFromConfig();
 
diff --git a/Engine/Source/Editor/Sequencer/Private/Filters/TextExpressions/SequencerTextFilterExpression_Filters.cpp b/Engine/Source/Editor/Sequencer/Private/Filters/TextExpressions/SequencerTextFilterExpression_Filters.cpp
new file mode 100644
--- /dev/null
+++ b/Engine/Source/Editor/Sequencer/Private/Filters/TextExpressions/SequencerTextFilterExpression_Filters.cpp
@@ -0,0 +1,50 @@
+﻿// Copyright Epic Games, Inc. All Rights Reserved.
+
+#include "SequencerTextFilterExpression_Filters.h"
+#include "Filters/ISequencerFilterBar.h"
+#include "MVVM/ViewModelPtr.h"
+#include "MVVM/ViewModels/SequenceModel.h"
+#include "Sequencer.h"
+
+using namespace UE::Sequencer;
+
+#define LOCTEXT_NAMESPACE "SequencerTextFilterExpression_Filters"
+
+FSequencerTextFilterExpression_Filters::FSequencerTextFilterExpression_Filters(ISequencerTrackFilters& InFilterInterface)
+	: FSequencerTextFilterExpressionContext(InFilterInterface)
+{
+}
+
+TSet<FName> FSequencerTextFilterExpression_Filters::GetKeys() const
+{
+	// No keys means this will not act like a traditional text expression, but as a general text expression
+	// that works for all common filters.
+	// This also means this will not show up in the text expression help dialog.
+	return {};
+}
+
+ESequencerTextFilterValueType FSequencerTextFilterExpression_Filters::GetValueType() const
+{
+	return ESequencerTextFilterValueType::Boolean;
+}
+
+FText FSequencerTextFilterExpression_Filters::GetDescription() const
+{
+	return FText();
+}
+
+bool FSequencerTextFilterExpression_Filters::TestComplexExpression(const FName& InKey
+	, const FTextFilterString& InValue
+	, const ETextFilterComparisonOperation InComparisonOperation
+	, const ETextFilterTextComparisonMode InTextComparisonMode) const
+{
+	if (!FSequencerTextFilterExpressionContext::TestComplexExpression(InKey, InValue, InComparisonOperation, InTextComparisonMode))
+	{
+		return true;
+	}
+
+	const bool bPassesFilter = FilterInterface.DoesCommonFilterPass(FilterItem.ImplicitCast(), InKey.ToString());
+	return CompareFStringForExactBool(InValue, InComparisonOperation, bPassesFilter);
+}
+
+#undef LOCTEXT_NAMESPACE
diff --git a/Engine/Source/Editor/Sequencer/Private/Filters/TextExpressions/SequencerTextFilterExpression_Filters.h b/Engine/Source/Editor/Sequencer/Private/Filters/TextExpressions/SequencerTextFilterExpression_Filters.h
new file mode 100644
--- /dev/null
+++ b/Engine/Source/Editor/Sequencer/Private/Filters/TextExpressions/SequencerTextFilterExpression_Filters.h
@@ -0,0 +1,24 @@
+﻿// Copyright Epic Games, Inc. All Rights Reserved.
+
+#pragma once
+
+#include "Filters/SequencerTextFilterExpressionContext.h"
+
+class FSequencerTextFilterExpression_Filters : public FSequencerTextFilterExpressionContext
+{
+public:
+	FSequencerTextFilterExpression_Filters(ISequencerTrackFilters& InFilterInterface);
+
+	//~ Begin FSequencerTextFilterExpressionContext
+	virtual TSet<FName> GetKeys() const override;
+	virtual ESequencerTextFilterValueType GetValueType() const override;
+	virtual FText GetDescription() const override;
+	//~ End FSequencerTextFilterExpressionContext
+
+	//~ Begin ITextFilterExpressionContext
+	virtual bool TestComplexExpression(const FName& InKey
+		, const FTextFilterString& InValue
+		, const ETextFilterComparisonOperation InComparisonOperation
+		, const ETextFilterTextComparisonMode InTextComparisonMode) const override;
+	//~ End ITextFilterExpressionContext
+};
diff --git a/Engine/Source/Editor/Sequencer/Private/Filters/Utils/FilterViewUtils.cpp b/Engine/Source/Editor/Sequencer/Private/Filters/Utils/FilterViewUtils.cpp
--- a/Engine/Source/Editor/Sequencer/Private/Filters/Utils/FilterViewUtils.cpp
+++ b/Engine/Source/Editor/Sequencer/Private/Filters/Utils/FilterViewUtils.cpp
@@ -28,7 +28,7 @@ void OpenDialog_TextExpressionHelp(const TSharedRef<ISequencerTrackFilters>& InF
 	FFilterExpressionHelpDialogConfig Config;
 	Config.IdentifierName = TEXT("SequencerCustomTextFilterHelp");
 	Config.DialogTitle = LOCTEXT("SequencerCustomTextFilterHelp", "Sequencer Custom Text Filter Help");
-	Config.TextFilterExpressionContexts = InFilterBar->GetTextFilterExpressionContexts();
+	Config.AddCommonFilterHelpEntries<FSequencerTrackFilter>(InFilterBar);
 
 	SFilterExpressionHelpDialog::Open(MoveTemp(Config));
 }
diff --git a/Engine/Source/Editor/Sequencer/Private/Filters/Widgets/SFilterExpressionHelpDialog.cpp b/Engine/Source/Editor/Sequencer/Private/Filters/Widgets/SFilterExpressionHelpDialog.cpp
--- a/Engine/Source/Editor/Sequencer/Private/Filters/Widgets/SFilterExpressionHelpDialog.cpp
+++ b/Engine/Source/Editor/Sequencer/Private/Filters/Widgets/SFilterExpressionHelpDialog.cpp
@@ -78,7 +78,9 @@ void SFilterExpressionHelpDialog::CloseWindow(const FName InName)
 
 void SFilterExpressionHelpDialog::Construct(const FArguments& InArgs, FFilterExpressionHelpDialogConfig&& InConfig)
 {
+	PRAGMA_DISABLE_DEPRECATION_WARNINGS
 	Config = MoveTemp(InConfig);
+	PRAGMA_ENABLE_DEPRECATION_WARNINGS
 
 	SWindow::Construct(SWindow::FArguments()
 		.Title(Config.DialogTitle)
@@ -159,11 +161,28 @@ TSharedRef<SWidget> SFilterExpressionHelpDialog::ConstructDialogHeader()
 
 TSharedRef<SWidget> SFilterExpressionHelpDialog::ConstructExpressionWidgetList()
 {
+	Config.ExpressionHelpEntries.Sort([](const FFilterExpressionHelpEntry& InA, const FFilterExpressionHelpEntry& InB)
+		{
+			const int32 CountA = InA.Keys.Num();
+			if (CountA == 0)
+			{
+				return false;
+			}
+
+			const int32 CountB = InB.Keys.Num();
+			if (CountB == 0)
+			{
+				return true;
+			}
+
+			return InA.Keys.Array()[0].LexicalLess(InB.Keys.Array()[0]);
+		});
+
 	const TSharedRef<SScrollBox> Container = SNew(SScrollBox);
 
 	bool bIsFirst = true;
 
-	for (const TSharedRef<ISequencerTextFilterExpressionContext>& ExpressionContext : Config.TextFilterExpressionContexts)
+	for (const FFilterExpressionHelpEntry& HelpEntry : Config.ExpressionHelpEntries)
 	{
 		if (bIsFirst)
 		{
@@ -182,14 +201,14 @@ TSharedRef<SWidget> SFilterExpressionHelpDialog::ConstructExpressionWidgetList()
 			.AutoSize()
 			.Padding(0.f, 0.f, 5.f, 0.f)
 			[
-				ConstructExpressionWidget(ExpressionContext)
+				ConstructExpressionWidget(HelpEntry)
 			];
 	}
 
 	return Container;
 }
 
-TSharedRef<SWidget> SFilterExpressionHelpDialog::ConstructExpressionWidget(const TSharedPtr<ISequencerTextFilterExpressionContext>& InExpressionContext)
+TSharedRef<SWidget> SFilterExpressionHelpDialog::ConstructExpressionWidget(const FFilterExpressionHelpEntry& InHelpEntry)
 {
 	return SNew(SVerticalBox)
 		+ SVerticalBox::Slot()
@@ -200,13 +219,13 @@ TSharedRef<SWidget> SFilterExpressionHelpDialog::ConstructExpressionWidget(const
 			+ SHorizontalBox::Slot()
 			.AutoWidth()
 			[
-				ConstructKeysWidget(InExpressionContext->GetKeys())
+				ConstructKeysWidget(InHelpEntry.Keys)
 			]
 			+ SHorizontalBox::Slot()
 			.AutoWidth()
 			.Padding(10.f, 0.f, 1.f, 0.f)
 			[
-				ConstructValueWidget(InExpressionContext->GetValueType())
+				ConstructValueWidget(InHelpEntry.ValueType)
 			]
 		]
 		+ SVerticalBox::Slot()
@@ -216,7 +235,7 @@ TSharedRef<SWidget> SFilterExpressionHelpDialog::ConstructExpressionWidget(const
 			SNew(STextBlock)
 			.AutoWrapText(true)
 			.TextStyle(FAppStyle::Get(), TEXT("NormalText"))
-			.Text(InExpressionContext->GetDescription())
+			.Text(InHelpEntry.Description)
 		];
 }
 
diff --git a/Engine/Source/Editor/Sequencer/Private/Filters/Widgets/SSequencerSearchBox.cpp b/Engine/Source/Editor/Sequencer/Private/Filters/Widgets/SSequencerSearchBox.cpp
--- a/Engine/Source/Editor/Sequencer/Private/Filters/Widgets/SSequencerSearchBox.cpp
+++ b/Engine/Source/Editor/Sequencer/Private/Filters/Widgets/SSequencerSearchBox.cpp
@@ -3,6 +3,7 @@
 #include "Filters/Widgets/SSequencerSearchBox.h"
 #include "Filters/ISequencerFilterBar.h"
 #include "Filters/Filters/SequencerTrackFilter_Text.h"
+#include "Filters/SequencerTrackFilterCollection.h"
 #include "Filters/Widgets/SSequencerFilterSuggestionListRow.h"
 #include "Framework/Application/SlateApplication.h"
 #include "SequencerSettings.h"
@@ -582,17 +583,18 @@ void SSequencerSearchBox::DefaultKeySuggestions(const FString& InKeyValue, TArra
 		const TArray<FName> Keys = Expression->GetKeys().Array();
 		for (const FName Key : Keys)
 		{
-			const FString KeyString = Key.ToString();
-
 			FSequencerFilterSuggestion NewSuggestion;
-			NewSuggestion.Suggestion = FString::Printf(TEXT("%s"), *KeyString);
-			NewSuggestion.DisplayName = FText::FromString(KeyString);
+			NewSuggestion.Suggestion = Key.ToString(); 
+			NewSuggestion.DisplayName = FText::FromName(Key);
 			NewSuggestion.Description = Expression->GetDescription();
 
 			OutPossibleSuggestions.Add(NewSuggestion);
 		}
 	}
 
+	// Add keys for all common filters to support suggestions for these text expressions
+	FilterBar->GetCommonFilterKeySuggestions(OutPossibleSuggestions);
+
 	OutPossibleSuggestions.Sort([](const FSequencerFilterSuggestion& InA, const FSequencerFilterSuggestion& InB)
 		{
 			return InA.DisplayName.CompareTo(InB.DisplayName) < 0;
@@ -627,22 +629,34 @@ void SSequencerSearchBox::DefaultValueSuggestions(const FString& InKeyValue, TAr
 end:
 	if (FoundExpression.IsValid())
 	{
-		const TArray<FSequencerTextFilterKeyword> TextFilterKeywords = FoundExpression->GetValueKeywords();
-		for (const FSequencerTextFilterKeyword& TextFilterKeyword : TextFilterKeywords)
+		if (FoundExpression->GetValueType() == ESequencerTextFilterValueType::Boolean)
 		{
-			FSequencerFilterSuggestion NewSuggestion;
-			NewSuggestion.Suggestion = FString::Printf(TEXT("%s"), *TextFilterKeyword.Keyword);
-			NewSuggestion.DisplayName = FText::FromString(TextFilterKeyword.Keyword);
-			NewSuggestion.Description = TextFilterKeyword.Description;
-
-			OutPossibleSuggestions.Add(NewSuggestion);
+			FSequencerFilterSuggestion::MakeBooleanSuggestions(OutPossibleSuggestions);
 		}
-
-		OutPossibleSuggestions.Sort([](const FSequencerFilterSuggestion& InA, const FSequencerFilterSuggestion& InB)
+		else
+		{
+			const TArray<FSequencerTextFilterKeyword> ValueKeywords = FoundExpression->GetValueKeywords();
+			for (const FSequencerTextFilterKeyword& ValueKeyword : ValueKeywords)
 			{
-				return InA.DisplayName.CompareTo(InB.DisplayName) < 0;
-			});
+				FSequencerFilterSuggestion NewSuggestion;
+				NewSuggestion.Suggestion = ValueKeyword.Keyword;
+				NewSuggestion.DisplayName = FText::FromString(ValueKeyword.Keyword);
+				NewSuggestion.Description = ValueKeyword.Description;
+
+				OutPossibleSuggestions.Add(NewSuggestion);
+			}
+		}
 	}
+	else
+	{
+		// Add keys for all common filters to support suggestions for these text expressions
+		FilterBar->GetCommonFilterValueSuggestions(InKeyValue, OutPossibleSuggestions);
+	}
+
+	OutPossibleSuggestions.Sort([](const FSequencerFilterSuggestion& InA, const FSequencerFilterSuggestion& InB)
+		{
+			return InA.DisplayName.CompareTo(InB.DisplayName) < 0;
+		});
 }
 
 void SSequencerSearchBox::OnFilterTextChanged()
diff --git a/Engine/Source/Editor/Sequencer/Public/Filters/ISequencerFilterBar.h b/Engine/Source/Editor/Sequencer/Public/Filters/ISequencerFilterBar.h
--- a/Engine/Source/Editor/Sequencer/Public/Filters/ISequencerFilterBar.h
+++ b/Engine/Source/Editor/Sequencer/Public/Filters/ISequencerFilterBar.h
@@ -4,14 +4,16 @@
 
 #include "Filters/SBasicFilterBar.h"
 #include "Framework/Commands/UICommandList.h"
+#include "MVVM/ViewModelPtr.h"
 
+template<typename InFilterType> class FSequencerFilterBase;
 class FSequencerTrackFilter;
-class FSequencerTrackFilter_CustomText;
 class FString;
 class FText;
 class FTextFilterExpressionEvaluator;
 class ISequencer;
 class ISequencerTextFilterExpressionContext;
+struct FSequencerFilterSuggestion;
 
 /** Generic interface for all Sequencer filter bar implementations. */
 class ISequencerFilterBar : public TSharedFromThis<ISequencerFilterBar>
@@ -53,6 +55,9 @@ public:
 
 	virtual void EnableAllFilters(const bool bInEnable, const TArray<FString>& InExceptionFilterNames) = 0;
 
+	template<typename InFilterClass>
+	TArray<TSharedRef<FSequencerFilterBase<InFilterClass>>> GetCommonFilters() const;
+
 	virtual void ActivateCommonFilters(const bool bInActivate, const TArray<FString>& InExceptionFilterNames) = 0;
 
 	virtual bool AreAllEnabledFiltersActive(const bool bInActive, const TArray<FString> InExceptionFilterNames) const = 0;
@@ -102,4 +107,37 @@ public:
 	virtual void SetToVerticalLayout() = 0;
 	virtual void SetToHorizontalLayout() = 0;
 	virtual void ToggleFilterBarLayout() = 0;
+
+	/** @return True if a common filter, specified by InFilterName, passes its filter. Used to support default text
+	 * expressions for common filters without needing to create an actual text expression class for each common filter. */
+	virtual bool DoesCommonFilterPass(const UE::Sequencer::TViewModelPtr<UE::Sequencer::FViewModel>& InViewModel, const FString& InFilterName) const = 0;
+
+	/**
+	 * Populates a list of suggested filter keys for the common filters used in the Sequencer filter bar.
+	 *
+	 * This function allows users to query and retrieve possible key suggestions for building advanced
+	 * filter expressions. The resulting suggestions are typically used by UI elements like search boxes
+	 * to provide auto-completion or hints.
+	 *
+	 * @param OutSuggestions An array to populate with filter key suggestions
+	 */
+	virtual void GetCommonFilterKeySuggestions(TArray<FSequencerFilterSuggestion>& OutSuggestions) const = 0;
+
+	/**
+	 * Populates a list of suggested filter values for the specified filter in the Sequencer filter bar.
+	 *
+	 * This method is useful for supporting auto-completion and providing context-aware suggestions
+	 * when users interact with filtering interfaces in the Sequencer UI. It queries possible values
+	 * for a given filter key and populates the output array with relevant suggestions.
+	 *
+	 * @param InFilterName The name of the filter to query for value suggestions
+	 * @param OutSuggestions An array to populate with filter value suggestions
+	 */
+	virtual void GetCommonFilterValueSuggestions(const FString& InFilterName, TArray<FSequencerFilterSuggestion>& OutSuggestions) const = 0;
+
+protected:
+	/** Internal implementation to bypass template limitations in virtual functions */
+	virtual void GetCommonFiltersImpl(TArray<TSharedPtr<void>>& OutTypeErasedFilters) const = 0;
 };
+
+#include "Filters/ISequencerFilterBar.inl"
diff --git a/Engine/Source/Editor/Sequencer/Public/Filters/ISequencerFilterBar.inl b/Engine/Source/Editor/Sequencer/Public/Filters/ISequencerFilterBar.inl
new file mode 100644
--- /dev/null
+++ b/Engine/Source/Editor/Sequencer/Public/Filters/ISequencerFilterBar.inl
@@ -0,0 +1,13 @@
+﻿// Copyright Epic Games, Inc. All Rights Reserved.
+
+#pragma once
+
+#include "Filters/SequencerFilterBase.h"
+
+template<typename InFilterClass>
+TArray<TSharedRef<FSequencerFilterBase<InFilterClass>>> ISequencerFilterBar::GetCommonFilters() const
+{
+	TArray<TSharedRef<FSequencerFilterBase<InFilterClass>>> OutFilters;
+	GetCommonFiltersImpl(*reinterpret_cast<TArray<TSharedPtr<void>>*>(&OutFilters));
+	return OutFilters;
+}
diff --git a/Engine/Source/Editor/Sequencer/Public/Filters/SequencerFilterBase.h b/Engine/Source/Editor/Sequencer/Public/Filters/SequencerFilterBase.h
--- a/Engine/Source/Editor/Sequencer/Public/Filters/SequencerFilterBase.h
+++ b/Engine/Source/Editor/Sequencer/Public/Filters/SequencerFilterBase.h
@@ -3,10 +3,15 @@
 #pragma once
 
 #include "Filters/FilterBase.h"
-#include "Filters/ISequencerFilterBar.h"
+#include "Framework/Commands/UIAction.h"
+#include "Framework/Commands/UICommandInfo.h"
+#include "Framework/Commands/InputChord.h"
 #include "Templates/SharedPointer.h"
 #include "MovieSceneSequence.h"
 
+class ISequencer;
+class ISequencerFilterBar;
+
 /** Sequencer base filter for all Sequencer filters (Sequencer and Sequence Navigator) */
 template<typename InFilterType>
 class FSequencerFilterBase
diff --git a/Engine/Source/Editor/Sequencer/Public/Filters/SequencerFilterSuggestion.h b/Engine/Source/Editor/Sequencer/Public/Filters/SequencerFilterSuggestion.h
--- a/Engine/Source/Editor/Sequencer/Public/Filters/SequencerFilterSuggestion.h
+++ b/Engine/Source/Editor/Sequencer/Public/Filters/SequencerFilterSuggestion.h
@@ -3,10 +3,43 @@
 #pragma once
 
 #include "Containers/UnrealString.h"
+#include "Internationalization/Internationalization.h"
 #include "Internationalization/Text.h"
 
+#define LOCTEXT_NAMESPACE "SequencerFilterSuggestion"
+
 struct FSequencerFilterSuggestion
 {
+	static FSequencerFilterSuggestion MakeSimpleSuggestion(FString InSuggestionString)
+	{
+		FSequencerFilterSuggestion Suggestion;
+		Suggestion.Suggestion = MoveTemp(InSuggestionString);
+		Suggestion.DisplayName = FText::FromString(Suggestion.Suggestion);
+		return MoveTemp(Suggestion);
+	}
+
+	static void MakeBooleanSuggestions(TArray<FSequencerFilterSuggestion>& OutSuggestions
+		, const FText& InOptionalTrueDescription = FText()
+		, const FText& InOptionalFalseDescription = FText())
+	{
+		FSequencerFilterSuggestion TrueSuggestion;
+		TrueSuggestion.Suggestion = TEXT("True");
+		TrueSuggestion.DisplayName = LOCTEXT("BooleanTrueSuggestion", "True");
+		TrueSuggestion.Description = InOptionalTrueDescription;
+		OutSuggestions.AddUnique(MoveTemp(TrueSuggestion));
+
+		FSequencerFilterSuggestion FalseSuggestion;
+		FalseSuggestion.Suggestion = TEXT("False");
+		FalseSuggestion.DisplayName = LOCTEXT("BooleanFalseSuggestion", "False");
+		FalseSuggestion.Description = InOptionalFalseDescription;
+		OutSuggestions.AddUnique(MoveTemp(FalseSuggestion));
+	}
+
+	bool operator==(const FSequencerFilterSuggestion& InOther) const
+	{
+		return Suggestion.Equals(InOther.Suggestion, ESearchCase::IgnoreCase);
+	}
+
 	/** The raw suggestion string that should be used with the search box */
 	FString Suggestion;
 
@@ -18,12 +51,6 @@ struct FSequencerFilterSuggestion
 
 	/** Describes what this suggestion will do */
 	FText Description;
-
-	static FSequencerFilterSuggestion MakeSimpleSuggestion(FString InSuggestionString)
-	{
-		FSequencerFilterSuggestion Suggestion;
-		Suggestion.Suggestion = MoveTemp(InSuggestionString);
-		Suggestion.DisplayName = FText::FromString(Suggestion.Suggestion);
-		return MoveTemp(Suggestion);
-	}
 };
+
+#undef LOCTEXT_NAMESPACE
diff --git a/Engine/Source/Editor/Sequencer/Public/Filters/Widgets/SFilterExpressionHelpDialog.h b/Engine/Source/Editor/Sequencer/Public/Filters/Widgets/SFilterExpressionHelpDialog.h
--- a/Engine/Source/Editor/Sequencer/Public/Filters/Widgets/SFilterExpressionHelpDialog.h
+++ b/Engine/Source/Editor/Sequencer/Public/Filters/Widgets/SFilterExpressionHelpDialog.h
@@ -2,13 +2,39 @@
 
 #pragma once
 
+#include "Filters/ISequencerFilterBar.h"
+#include "Filters/ISequencerTextFilterExpressionContext.h"
+#include "Filters/SequencerFilterBase.h"
 #include "Widgets/SWindow.h"
 
 #define UE_API SEQUENCER_API
 
-class ISequencerTextFilterExpressionContext;
+class ISequencerFilterBar;
 enum class ESequencerTextFilterValueType : uint8;
 
+struct FFilterExpressionHelpEntry
+{
+	UE_API FFilterExpressionHelpEntry(const TSet<FName>& InKeys, const ESequencerTextFilterValueType InValueType, const FText& InDescription)
+	{
+		Keys = InKeys;
+		ValueType = InValueType;
+		Description = InDescription;
+	}
+
+	UE_API FFilterExpressionHelpEntry(const TSharedRef<ISequencerTextFilterExpressionContext>& InExpressionContext)
+	{
+		Keys = InExpressionContext->GetKeys();
+		ValueType = InExpressionContext->GetValueType();
+		Description = InExpressionContext->GetDescription();
+	}
+
+	TSet<FName> Keys;
+
+	ESequencerTextFilterValueType ValueType;
+
+	FText Description;
+};
+
 struct FFilterExpressionHelpDialogConfig
 {
 	static constexpr float DefaultMaxDesiredWidth = 460.f;
@@ -16,15 +42,70 @@ struct FFilterExpressionHelpDialogConfig
 
 	UE_API FFilterExpressionHelpDialogConfig();
 
+	// Explicitly define special members so the deprecated field is not referenced implicitly
+	FFilterExpressionHelpDialogConfig(const FFilterExpressionHelpDialogConfig&) = default;
+	FFilterExpressionHelpDialogConfig(FFilterExpressionHelpDialogConfig&& InOther) noexcept
+		: IdentifierName(MoveTemp(InOther.IdentifierName))
+		, DialogTitle(MoveTemp(InOther.DialogTitle))
+		, DocumentationLink(MoveTemp(InOther.DocumentationLink))
+		, MaxDesiredWidth(InOther.MaxDesiredWidth)
+		, MaxDesiredHeight(InOther.MaxDesiredHeight)
+		, ExpressionHelpEntries(MoveTemp(InOther.ExpressionHelpEntries))
+	{}
+
+	FFilterExpressionHelpDialogConfig& operator=(const FFilterExpressionHelpDialogConfig&) = default;
+	FFilterExpressionHelpDialogConfig& operator=(FFilterExpressionHelpDialogConfig&& InOther) noexcept
+	{
+		if (this != &InOther)
+		{
+			IdentifierName = MoveTemp(InOther.IdentifierName);
+			DialogTitle = MoveTemp(InOther.DialogTitle);
+			DocumentationLink = MoveTemp(InOther.DocumentationLink);
+			MaxDesiredWidth = InOther.MaxDesiredWidth;
+			MaxDesiredHeight = InOther.MaxDesiredHeight;
+			ExpressionHelpEntries = MoveTemp(InOther.ExpressionHelpEntries);
+		}
+		return *this;
+	}
+
+	template<typename InFilterClass>
+	void AddCommonFilterHelpEntries(const TSharedRef<ISequencerFilterBar>& InFilterBar)
+	{
+		TSet<FName> EntryNames;
+
+		for (const TSharedRef<ISequencerTextFilterExpressionContext>& ExpressionContext : InFilterBar->GetTextFilterExpressionContexts())
+		{
+			if (!ExpressionContext->GetKeys().IsEmpty())
+			{
+				ExpressionHelpEntries.Add(ExpressionContext);
+				EntryNames.Append(ExpressionContext->GetKeys());
+			}
+		}
+
+		for (const TSharedRef<FSequencerFilterBase<InFilterClass>> CommonFilter : InFilterBar->template GetCommonFilters<InFilterClass>())
+		{
+			// Avoid adding duplicate entries for common filters if they've already been added through a text expression
+			const FName CommonFilterName = *CommonFilter->GetName();
+			if (!EntryNames.Contains(CommonFilterName))
+			{
+				ExpressionHelpEntries.Add(FFilterExpressionHelpEntry(
+					{ CommonFilterName }, ESequencerTextFilterValueType::Boolean, CommonFilter->GetDefaultToolTipText()));
+			}
+		}
+	}
+
 	FName IdentifierName;
 
 	FText DialogTitle;
 	FString DocumentationLink;
 
-	TArray<TSharedRef<ISequencerTextFilterExpressionContext>> TextFilterExpressionContexts;
-
 	float MaxDesiredWidth = DefaultMaxDesiredWidth;
 	float MaxDesiredHeight = DefaultMaxDesiredHeight;
+
+	TArray<FFilterExpressionHelpEntry> ExpressionHelpEntries;
+
+	UE_DEPRECATED(5.8, "Use ExpressionHelpEntries instead")
+	TArray<TSharedRef<ISequencerTextFilterExpressionContext>> TextFilterExpressionContexts_DEPRECATED;
 };
 
 class SFilterExpressionHelpDialog : public SWindow
@@ -45,8 +126,10 @@ protected:
 	static TMap<FName, TSharedPtr<SFilterExpressionHelpDialog>> DialogInstance;
 
 	TSharedRef<SWidget> ConstructDialogHeader();
+
 	TSharedRef<SWidget> ConstructExpressionWidgetList();
-	TSharedRef<SWidget> ConstructExpressionWidget(const TSharedPtr<ISequencerTextFilterExpressionContext>& InExpressionContext);
+	TSharedRef<SWidget> ConstructExpressionWidget(const FFilterExpressionHelpEntry& InHelpEntry);
+
 	TSharedRef<SWidget> ConstructKeysWidget(const TSet<FName>& InKeys);
 	TSharedRef<SWidget> ConstructValueWidget(const ESequencerTextFilterValueType InValueType);
 
