[Sequencer] Remove export macro from FFilterExpressionHelpEntry constructors

#Synced-CL 50092247

diff --git a/Engine/Source/Editor/Sequencer/Public/Filters/Widgets/SFilterExpressionHelpDialog.h b/Engine/Source/Editor/Sequencer/Public/Filters/Widgets/SFilterExpressionHelpDialog.h
--- a/Engine/Source/Editor/Sequencer/Public/Filters/Widgets/SFilterExpressionHelpDialog.h
+++ b/Engine/Source/Editor/Sequencer/Public/Filters/Widgets/SFilterExpressionHelpDialog.h
@@ -14,14 +14,14 @@ enum class ESequencerTextFilterValueType : uint8;
 
 struct FFilterExpressionHelpEntry
 {
-	UE_API FFilterExpressionHelpEntry(const TSet<FName>& InKeys, const ESequencerTextFilterValueType InValueType, const FText& InDescription)
+	FFilterExpressionHelpEntry(const TSet<FName>& InKeys, const ESequencerTextFilterValueType InValueType, const FText& InDescription)
 	{
 		Keys = InKeys;
 		ValueType = InValueType;
 		Description = InDescription;
 	}
 
-	UE_API FFilterExpressionHelpEntry(const TSharedRef<ISequencerTextFilterExpressionContext>& InExpressionContext)
+	FFilterExpressionHelpEntry(const TSharedRef<ISequencerTextFilterExpressionContext>& InExpressionContext)
 	{
 		Keys = InExpressionContext->GetKeys();
 		ValueType = InExpressionContext->GetValueType();
