[Sequencer] Adjust track area padding to fix fallout from outliner sizing refactor

#Synced-CL 50092247

diff --git a/Engine/Source/Editor/Sequencer/Private/SSequencer.cpp b/Engine/Source/Editor/Sequencer/Private/SSequencer.cpp
--- a/Engine/Source/Editor/Sequencer/Private/SSequencer.cpp
+++ b/Engine/Source/Editor/Sequencer/Private/SSequencer.cpp
@@ -764,7 +764,6 @@ TSharedRef<SSplitter> SSequencer::ConstructOutlinerAndTrackArea(const FArguments
 						+ SVerticalBox::Slot()
 						.AutoHeight()
 						.MaxHeight(TAttribute<float>::Create(TAttribute<float>::FGetter::CreateSP(this, &SSequencer::GetPinnedAreaMaxHeight)))
-						.Padding(FMargin(0.f, 0.f, 0.f, CommonPadding))
 						[
 							SNew(SBox)
 							.HeightOverride_Lambda([this]()
