[Sequencer] Slow Sequencer window resizing

#Synced-CL 50092247

diff --git a/Engine/Source/Editor/Sequencer/Private/SSequencer.cpp b/Engine/Source/Editor/Sequencer/Private/SSequencer.cpp
--- a/Engine/Source/Editor/Sequencer/Private/SSequencer.cpp
+++ b/Engine/Source/Editor/Sequencer/Private/SSequencer.cpp
@@ -686,13 +686,19 @@ TSharedRef<SSplitter> SSequencer::ConstructOutlinerAndTrackArea(const FArguments
 	const TSharedRef<FSequencerTimeSliderController> TimeSliderControllerRef = TimeSliderController.ToSharedRef();
 	USequencerSettings* const SequencerSettings = GetSequencerSettings();
 
+	TreeViewWidth = SequencerSettings->GetTreeViewWidth();
+
 	return SNew(SSplitter)
+		.OnSplitterFinishedResizing_Lambda([this]()
+			{
+				GetSequencerSettings()->SetTreeViewWidth(TreeViewWidth);
+			})
 		+ SSplitter::Slot()
 		.Resizable(true)
 		.SizeRule(SSplitter::SizeToContent)
-		.OnSlotResized_Lambda([this](float InSlotSize)
+		.OnSlotResized_Lambda([this](const float InSlotSize)
 			{
-				GetSequencerSettings()->SetTreeViewWidth(InSlotSize);
+				TreeViewWidth = InSlotSize;
 			})
 		[
 			SNew(SBox)
@@ -4314,7 +4320,7 @@ FOptionalSize SSequencer::GetOutlinerWidth() const
 		return DefaultOutlinerWidth;
 	}
 
-	float CurrentOutlinerWidth = SequencerSettings->GetTreeViewWidth();
+	float ActualOutlinerWidth = TreeViewWidth;
 
 	if (OutlinerAndTrackAreaContent.IsValid())
 	{
@@ -4325,17 +4331,17 @@ FOptionalSize SSequencer::GetOutlinerWidth() const
 		{
 			// Limit the outliner size to the parent size
 			const float MaxDynamicWidth = LocalSize.X - MinOutlinerWidth;
-			if (CurrentOutlinerWidth > MaxDynamicWidth)
+			if (ActualOutlinerWidth > MaxDynamicWidth)
 			{
-				CurrentOutlinerWidth = MaxDynamicWidth;
+				ActualOutlinerWidth = MaxDynamicWidth;
 			}
 
 			const float MinWidth = FMath::Min<float>(MaxDynamicWidth, MinOutlinerWidth);
-			return FMath::Clamp<float>(CurrentOutlinerWidth, MinWidth, MaxDynamicWidth);
+			return FMath::Clamp<float>(ActualOutlinerWidth, MinWidth, MaxDynamicWidth);
 		}
 	}
 
-	return FMath::Max<float>(CurrentOutlinerWidth, MinOutlinerWidth);
+	return FMath::Max<float>(ActualOutlinerWidth, MinOutlinerWidth);
 }
 
 #undef LOCTEXT_NAMESPACE
diff --git a/Engine/Source/Editor/Sequencer/Private/SSequencer.h b/Engine/Source/Editor/Sequencer/Private/SSequencer.h
--- a/Engine/Source/Editor/Sequencer/Private/SSequencer.h
+++ b/Engine/Source/Editor/Sequencer/Private/SSequencer.h
@@ -787,10 +787,13 @@ private:
 	/** Container for the toolbar, so that we can re-create it as needed. */
 	TSharedPtr<SBox> ToolbarContainer;
 
-	/** Exposing these slot so the equivalent track areas can match their height */
+	/** Exposing these slots so the equivalent track areas can match their height */
 	SVerticalBox::FSlot* OutlinerSearchAndFilterSlot = nullptr;
 	SVerticalBox::FSlot* OutlinerPinnedSlot = nullptr;
 
+	/** Width of the outliner tree. Mainly used while sizing so we don't need to SaveConfig() and query settings, which causes a performance hit */
+	float TreeViewWidth = 0.f;
+
 	/** List of registered outliner columns with their visibility states */
 	TArray<FSequencerOutlinerColumnVisibility> OutlinerColumnVisibilities;
 
