Sequencer: This refactor removes the SSequencerFilterBar dependency to SSequencerSearchBox.

#Synced-CL 50196422

diff --git a/Engine/Source/Editor/Sequencer/Private/Filters/Widgets/SSearchAndFilterWidget.cpp b/Engine/Source/Editor/Sequencer/Private/Filters/Widgets/SSearchAndFilterWidget.cpp
--- a/Engine/Source/Editor/Sequencer/Private/Filters/Widgets/SSearchAndFilterWidget.cpp
+++ b/Engine/Source/Editor/Sequencer/Private/Filters/Widgets/SSearchAndFilterWidget.cpp
@@ -78,6 +78,7 @@ void SSearchAndFilterWidget::OnOutlinerSearchCommitted(const FText& InSearchText
 void SSearchAndFilterWidget::OnOutlinerSearchSaved(const FText& InSearchText)
 {
 	FCustomTextFilterData CustomTextFilterData;
+	CustomTextFilterData.FilterLabel = LOCTEXT("NewFilterName", "New Filter Name");
 	CustomTextFilterData.FilterString = InSearchText;
 	SSequencerCustomTextFilterDialog::CreateWindow_AddCustomTextFilter(FilterBar.ToSharedRef(), CustomTextFilterData);
 }
diff --git a/Engine/Source/Editor/Sequencer/Private/Filters/Widgets/SSequencerFilterBar.cpp b/Engine/Source/Editor/Sequencer/Private/Filters/Widgets/SSequencerFilterBar.cpp
--- a/Engine/Source/Editor/Sequencer/Private/Filters/Widgets/SSequencerFilterBar.cpp
+++ b/Engine/Source/Editor/Sequencer/Private/Filters/Widgets/SSequencerFilterBar.cpp
@@ -39,7 +39,6 @@ void SSequencerFilterBar::Construct(const FArguments& InArgs, const TSharedRef<F
 {
 	WeakFilterBar = InFilterBar;
 
-	WeakSearchBox = InArgs._FilterSearchBox;
 	FilterBarLayout = InArgs._FilterBarLayout;
 	bCanChangeOrientation = InArgs._CanChangeOrientation;
 	FilterPillStyle = InArgs._FilterPillStyle;
@@ -80,13 +79,13 @@ void SSequencerFilterBar::Construct(const FArguments& InArgs, const TSharedRef<F
 				})
 		]
 	];
-
-	AttachFilterSearchBox(InArgs._FilterSearchBox);
-
+	
 	CreateFilterWidgetsFromConfig();
 
 	InFilterBar->OnFiltersChanged().AddSP(this, &SSequencerFilterBar::OnFiltersChanged);
 	InFilterBar->OnMuteFiltersChanged().AddSP(this, &SSequencerFilterBar::OnMuteStateChanged);
+	
+	OnMuteStateChanged(InFilterBar->AreFiltersMuted());
 }
 
 FReply SSequencerFilterBar::OnMouseButtonUp(const FGeometry& MyGeometry, const FPointerEvent& MouseEvent)
@@ -141,17 +140,6 @@ void SSequencerFilterBar::SetLayout(const EFilterBarLayout InFilterBarLayout)
 	}
 }
 
-void SSequencerFilterBar::AttachFilterSearchBox(const TSharedPtr<SSequencerSearchBox>& InFilterSearchBox)
-{
-	if (InFilterSearchBox)
-	{
-		WeakSearchBox = InFilterSearchBox;
-
-		InFilterSearchBox->SetOnSaveSearchHandler(
-			 SFilterSearchBox::FOnSaveSearchClicked::CreateSP(this, &SSequencerFilterBar::CreateAddCustomTextFilterWindowFromSearch));
-	}
-}
-
 bool SSequencerFilterBar::HasAnyFilterWidgets() const
 {
 	return FilterWidgets.Num() > 0;
@@ -414,29 +402,6 @@ UWorld* SSequencerFilterBar::GetWorld() const
 	return PlaybackContext->GetWorld();
 }
 
-TWeakPtr<SSequencerSearchBox> SSequencerFilterBar::GetSearchBox() const
-{
-	return WeakSearchBox;
-}
-
-void SSequencerFilterBar::SetMuted(bool bInMuted)
-{
-	if (HorizontalContainerWidget.IsValid())
-	{
-		HorizontalContainerWidget->SetEnabled(!bInMuted);
-	}
-
-	if (VerticalContainerWidget.IsValid())
-	{
-		VerticalContainerWidget->SetEnabled(!bInMuted);
-	}
-
-	if (WeakSearchBox.IsValid())
-	{
-		WeakSearchBox.Pin()->SetEnabled(!bInMuted);
-	}
-}
-
 void SSequencerFilterBar::OnFiltersChanged(const ESequencerFilterChange InChangeType, const TSharedRef<FSequencerTrackFilter>& InFilter)
 {
 	switch (InChangeType)
@@ -464,25 +429,16 @@ void SSequencerFilterBar::OnFiltersChanged(const ESequencerFilterChange InChange
 
 void SSequencerFilterBar::OnMuteStateChanged(bool bNewIsMuted)
 {
-	if (const TSharedPtr<FSequencerFilterBar> FilterBar = GetFilterBar())
+	const bool bIsEnabled = !bNewIsMuted;
+	if (HorizontalContainerWidget.IsValid())
 	{
-		SetMuted(bNewIsMuted);
+		HorizontalContainerWidget->SetEnabled(bIsEnabled);
 	}
-}
 
-void SSequencerFilterBar::CreateAddCustomTextFilterWindowFromSearch(const FText& InSearchText)
-{
-	const TSharedPtr<FSequencerFilterBar> FilterBar = WeakFilterBar.Pin();
-	if (!FilterBar.IsValid())
+	if (VerticalContainerWidget.IsValid())
 	{
-		return;
+		VerticalContainerWidget->SetEnabled(bIsEnabled);
 	}
-
-	FCustomTextFilterData CustomTextFilterData;
-	CustomTextFilterData.FilterLabel = LOCTEXT("NewFilterName", "New Filter Name");
-	CustomTextFilterData.FilterString = InSearchText;
-
-	SSequencerCustomTextFilterDialog::CreateWindow_AddCustomTextFilter(FilterBar.ToSharedRef(), MoveTemp(CustomTextFilterData));
 }
 
 void SSequencerFilterBar::CreateFilterWidgetsFromConfig()
diff --git a/Engine/Source/Editor/Sequencer/Private/Filters/Widgets/SSequencerFilterBar.h b/Engine/Source/Editor/Sequencer/Private/Filters/Widgets/SSequencerFilterBar.h
--- a/Engine/Source/Editor/Sequencer/Private/Filters/Widgets/SSequencerFilterBar.h
+++ b/Engine/Source/Editor/Sequencer/Private/Filters/Widgets/SSequencerFilterBar.h
@@ -25,11 +25,6 @@ public:
 		, _FilterPillStyle(EFilterPillStyle::Default)
 		, _UseSectionsForCategories(true)
 	{}
-		/** An SFilterSearchBox that can be attached to this filter bar. When provided along with a CreateTextFilter
-		 *  delegate, allows the user to save searches from the Search Box as text filters for the filter bar.
-		 *	NOTE: Will bind a delegate to SFilterSearchBox::OnClickedAddSearchHistoryButton */
-		SLATE_ARGUMENT(TSharedPtr<SSequencerSearchBox>, FilterSearchBox)
-
 		/** The layout that determines how the filters are laid out */
 		SLATE_ARGUMENT(EFilterBarLayout, FilterBarLayout)
 
@@ -58,17 +53,8 @@ public:
 	EFilterBarLayout GetLayout() const;
 	void SetLayout(const EFilterBarLayout InFilterBarLayout);
 
-	void AttachFilterSearchBox(const TSharedPtr<SSequencerSearchBox>& InFilterSearchBox);
-
 	bool HasAnyFilterWidgets() const;
 
-	void CreateAddCustomTextFilterWindowFromSearch(const FText& InSearchText);
-
-	TWeakPtr<SSequencerSearchBox> GetSearchBox() const;
-
-	// Set the state of the filter bar. Muted means that the filters are muted, but the context menu is still enabled and accessible.
-	void SetMuted(bool bInMuted);
-
 protected:
 	void AddWidgetToLayout(const TSharedRef<SWidget>& InWidget);
 	void RemoveWidgetFromLayout(const TSharedRef<SWidget>& InWidget);
@@ -108,8 +94,6 @@ protected:
 
 	TWeakPtr<FSequencerFilterBar> WeakFilterBar;
 
-	TWeakPtr<SSequencerSearchBox> WeakSearchBox;
-
 	TSharedPtr<SWidgetSwitcher> FilterBoxWidget;
 	TSharedPtr<SFilterBarClippingHorizontalBox> HorizontalContainerWidget;
 	TSharedPtr<SScrollBox> VerticalContainerWidget;
diff --git a/Engine/Source/Editor/Sequencer/Private/Filters/Widgets/SSequencerSearchBox.cpp b/Engine/Source/Editor/Sequencer/Private/Filters/Widgets/SSequencerSearchBox.cpp
--- a/Engine/Source/Editor/Sequencer/Private/Filters/Widgets/SSequencerSearchBox.cpp
+++ b/Engine/Source/Editor/Sequencer/Private/Filters/Widgets/SSequencerSearchBox.cpp
@@ -89,6 +89,10 @@ void SSequencerSearchBox::Construct(const FArguments& InArgs, const TWeakPtr<ISe
 	const TSharedPtr<ISequencerFilterBar> FilterBarPin = WeakFilterBar.Pin();
 	check(FilterBarPin);
 	FilterBarPin->OnTextFilterTextChanged().AddSP(this, &SSequencerSearchBox::OnFilterTextChanged);
+	FilterBarPin->OnMuteFiltersChanged().AddSP(this, &SSequencerSearchBox::OnMuteStateChanged);
+	
+	OnFilterTextChanged();
+	OnMuteStateChanged(FilterBarPin->AreFiltersMuted());
 }
 
 TSharedRef<SWidget> SSequencerSearchBox::GetSuggestionListMenuContent()
@@ -668,6 +672,11 @@ void SSequencerSearchBox::OnFilterTextChanged()
 	}
 }
 
+void SSequencerSearchBox::OnMuteStateChanged(bool bIsMuted)
+{
+	SetEnabled(!bIsMuted);
+}
+
 bool SSequencerSearchBox::IsOperatorToken(const FExpressionToken& InToken)
 {
 	return InToken.Node.Cast<TextFilterExpressionParser::FEqual>()
diff --git a/Engine/Source/Editor/Sequencer/Private/SSequencer.cpp b/Engine/Source/Editor/Sequencer/Private/SSequencer.cpp
--- a/Engine/Source/Editor/Sequencer/Private/SSequencer.cpp
+++ b/Engine/Source/Editor/Sequencer/Private/SSequencer.cpp
@@ -4186,9 +4186,7 @@ void SSequencer::RebuildSearchAndFilterRow()
 	const TSharedRef<FSequencerFilterBar> FilterBar = GetFilterBar().ToSharedRef();
 	FilterBarWidget = SNew(SSequencerFilterBar, FilterBar)
 		.FilterBarLayout(GetFilterBarLayout())
-		.AddMetaData<FTagMetaData>(FTagMetaData(TEXT("SequencerTrackFilters")))
-		.FilterSearchBox(SearchBox);
-	FilterBarWidget->SetMuted(FilterBar->AreFiltersMuted());
+		.AddMetaData<FTagMetaData>(FTagMetaData(TEXT("SequencerTrackFilters")));
 
 	if (IsFilterBarVisible() && GetFilterBarLayout() == EFilterBarLayout::Horizontal)
 	{
diff --git a/Engine/Source/Editor/Sequencer/Public/Filters/Widgets/SSequencerSearchBox.h b/Engine/Source/Editor/Sequencer/Public/Filters/Widgets/SSequencerSearchBox.h
--- a/Engine/Source/Editor/Sequencer/Public/Filters/Widgets/SSequencerSearchBox.h
+++ b/Engine/Source/Editor/Sequencer/Public/Filters/Widgets/SSequencerSearchBox.h
@@ -156,8 +156,10 @@ private:
 	/** Provides default value suggestions, if no suggestions are provided */
 	UE_API void DefaultValueSuggestions(const FString& InKeyValue, TArray<FSequencerFilterSuggestion>& OutPossibleSuggestions);
 	
-	/** Handles ISequencerFilterBar::OnTextFilterTextChanged. */
+	/** Updates the widget search text when the filter changes (@see ISequencerFilterBar::OnTextFilterTextChanged). */
 	void OnFilterTextChanged();
+	/** Calls SetEnabled when the mute state changes (@see ISequencerFilterBar::OnMuteStateChanged). */
+	void OnMuteStateChanged(bool bIsMuted);
 
 	static UE_API bool IsOperatorToken(const TExpressionToken<TCHAR>& InToken);
 	static UE_API bool IsLogicalOperatorToken(const TExpressionToken<TCHAR>& InToken);
