Sequencer: SFilterBarIsolateHideShow now uses a view-mode.

The view-model accepts FSequencer selection via injection. This is relevant so, in a follow up change, the view-model can be constructed for other views, specifically the curve editor view. The goal will be for curve editor's view to inject whatever it is filtering.

Effectively, this makes SFilterBarIsolateHideShow reusable in contexts without a hard reference to FSequencer nor FSequencerFilterBar.

#Synced-CL 50196422

diff --git a/Engine/Source/Editor/Sequencer/Private/Filters/Filters/SequencerTrackFilter_HideIsolate.cpp b/Engine/Source/Editor/Sequencer/Private/Filters/Filters/SequencerTrackFilter_HideIsolate.cpp
--- a/Engine/Source/Editor/Sequencer/Private/Filters/Filters/SequencerTrackFilter_HideIsolate.cpp
+++ b/Engine/Source/Editor/Sequencer/Private/Filters/Filters/SequencerTrackFilter_HideIsolate.cpp
@@ -68,12 +68,12 @@ void FSequencerTrackFilter_HideIsolate::ResetFilter()
 	BroadcastChangedEvent();
 }
 
-TSet<TWeakViewModelPtr<IOutlinerExtension>> FSequencerTrackFilter_HideIsolate::GetHiddenTracks() const
+const TSet<TWeakViewModelPtr<IOutlinerExtension>>& FSequencerTrackFilter_HideIsolate::GetHiddenTracks() const
 {
 	return HiddenTracks;
 }
 
-TSet<TWeakViewModelPtr<IOutlinerExtension>> FSequencerTrackFilter_HideIsolate::GetIsolatedTracks() const
+const TSet<TWeakViewModelPtr<IOutlinerExtension>>& FSequencerTrackFilter_HideIsolate::GetIsolatedTracks() const
 {
 	return IsolatedTracks;
 }
diff --git a/Engine/Source/Editor/Sequencer/Private/Filters/SequencerFilterBar.cpp b/Engine/Source/Editor/Sequencer/Private/Filters/SequencerFilterBar.cpp
--- a/Engine/Source/Editor/Sequencer/Private/Filters/SequencerFilterBar.cpp
+++ b/Engine/Source/Editor/Sequencer/Private/Filters/SequencerFilterBar.cpp
@@ -36,6 +36,7 @@
 #include "SSequencer.h"
 #include "Stats/StatsMisc.h"
 #include "UObject/UObjectIterator.h"
+#include "Utils/HideIsolateViewModelUtils.h"
 
 using namespace UE::Sequencer;
 
@@ -368,19 +369,14 @@ TSharedRef<FSequencerTrackFilter_Text> FSequencerFilterBar::GetTextFilter() cons
 	return TextFilter;
 }
 
-FText FSequencerFilterBar::GetFilterErrorText() const
-{
-	return TextFilter->GetFilterErrorText();
-}
-
-TSet<TWeakViewModelPtr<IOutlinerExtension>> FSequencerFilterBar::GetHiddenTracks() const
+TSharedRef<FSequencerTrackFilter_HideIsolate> FSequencerFilterBar::GetHideIsolateFilter() const
 {
-	return HideIsolateFilter->GetHiddenTracks();
+	return HideIsolateFilter;
 }
 
-TSet<TWeakViewModelPtr<IOutlinerExtension>> FSequencerFilterBar::GetIsolatedTracks() const
+FText FSequencerFilterBar::GetFilterErrorText() const
 {
-	return HideIsolateFilter->GetIsolatedTracks();
+	return TextFilter->GetFilterErrorText();
 }
 
 void FSequencerFilterBar::HideTracks(const TSet<TWeakViewModelPtr<IOutlinerExtension>>& InTracks, const bool bInAddToExisting)
@@ -403,23 +399,6 @@ void FSequencerFilterBar::UnisolateTracks(const TSet<TWeakViewModelPtr<IOutliner
 	HideIsolateFilter->UnisolateTracks(InTracks);
 }
 
-void FSequencerFilterBar::ShowAllTracks()
-{
-	HideIsolateFilter->ShowAllTracks();
-
-	if (const TSharedPtr<SSequencer> SequencerWidget = StaticCastSharedRef<SSequencer>(Sequencer.GetSequencerWidget()))
-	{
-		const TSharedPtr<FSequencerSelection> Selection = Sequencer.GetViewModel()->GetSelection();
-		const TArray<TWeakViewModelPtr<IOutlinerExtension>> SelectedTracks = Selection->Outliner.GetSelected().Array();
-		if (SelectedTracks.Num() > 0)
-		{
-			SequencerWidget->GetTreeView()->RequestScrollIntoView(SelectedTracks[0]);
-		}
-	}
-
-	RequestFilterUpdate();
-}
-
 bool FSequencerFilterBar::HasHiddenTracks() const
 {
 	return HideIsolateFilter->HasHiddenTracks();
@@ -430,20 +409,6 @@ bool FSequencerFilterBar::HasIsolatedTracks() const
 	return HideIsolateFilter->HasIsolatedTracks();
 }
 
-void FSequencerFilterBar::EmptyHiddenTracks()
-{
-	HideIsolateFilter->EmptyHiddenTracks();
-
-	RequestFilterUpdate();
-}
-
-void FSequencerFilterBar::EmptyIsolatedTracks()
-{
-	HideIsolateFilter->EmptyIsolatedTracks();
-
-	RequestFilterUpdate();
-}
-
 void FSequencerFilterBar::RequestFilterUpdate()
 {
 	Sequencer.GetNodeTree()->RequestFilterUpdate();
@@ -1596,29 +1561,7 @@ void FSequencerFilterBar::HandleTextFilterChanged()
 TSet<TWeakViewModelPtr<IOutlinerExtension>> FSequencerFilterBar::GetSelectedTracksOrAll() const
 {
 	TSharedPtr<FSequencerEditorViewModel> SequencerViewModel = GetSequencer().GetViewModel();
-	if (!SequencerViewModel.IsValid())
-	{
-		return {};
-	}
-
-	const TSharedPtr<FSequencerSelection> Selection = SequencerViewModel->GetSelection();
-	if (!Selection.IsValid())
-	{
-		return {};
-	}
-
-	const TSet<TWeakViewModelPtr<IOutlinerExtension>> SelectedSet = Selection->Outliner.GetSelected();
-	if (SelectedSet.IsEmpty())
-	{
-		TSet<TWeakViewModelPtr<IOutlinerExtension>> OutTracks;
-		for (const TViewModelPtr<IOutlinerExtension>& TrackModel : SequencerViewModel->GetRootModel()->GetDescendantsOfType<IOutlinerExtension>())
-		{
-			OutTracks.Add(TrackModel);
-		}
-		return OutTracks;
-	}
-
-	return SelectedSet;
+	return SequencerViewModel ? ::GetSelectedTracksOrAll(*SequencerViewModel) : TSet<TWeakViewModelPtr<IOutlinerExtension>>();
 }
 
 bool FSequencerFilterBar::HasSelectedTracks() const
@@ -1707,11 +1650,6 @@ bool FSequencerFilterBar::ShouldUpdateOnTrackValueChanged() const
 	return false;
 }
 
-TSharedRef<SFilterBarIsolateHideShow> FSequencerFilterBar::MakeIsolateHideShowPanel()
-{
-	return SNew(SFilterBarIsolateHideShow, SharedThis(this));
-}
-
 TSharedRef<SComboButton> FSequencerFilterBar::MakeAddFilterButton()
 {
 	const TSharedPtr<SLayeredImage> FilterImage = SNew(SLayeredImage)
@@ -1737,15 +1675,16 @@ TSharedRef<SComboButton> FSequencerFilterBar::MakeAddFilterButton()
 			return FAppStyle::Get().GetBrush(TEXT("Icons.Badge"));
 		}));
 
-	const TSharedRef<SToolBarButtonBlock> asdf = SNew(SToolBarButtonBlock);
 	const TSharedRef<SComboButton> ComboButton = SNew(SComboButton)
 		.ComboButtonStyle(&FAppStyle::Get().GetWidgetStyle<FComboButtonStyle>(TEXT("SimpleComboButtonWithIcon")))
 		.ForegroundColor(FSlateColor::UseStyle())
 		.ToolTipText_Lambda([this]()
 			{
-				return FText::Format(LOCTEXT("AddFilterToolTip", "Open the Add Filter Menu to add or manage filters\n\n"
-					"Shift + Click to temporarily mute all active filters\n\n{0}")
-					, SFilterBarIsolateHideShow::MakeLongDisplaySummaryText(*this));
+				return FText::Format(
+					LOCTEXT("AddFilterToolTip", "Open the Add Filter Menu to add or manage filters\n\n"
+					"Shift + Click to temporarily mute all active filters\n\n{0}"), 
+					MakeHiddenAndIsolatedCountText_WithTotal(FilterData, *HideIsolateFilter)
+					);
 			})
 		.OnComboBoxOpened_Lambda([this]()
 			{
diff --git a/Engine/Source/Editor/Sequencer/Private/Filters/SequencerFilterBar.h b/Engine/Source/Editor/Sequencer/Private/Filters/SequencerFilterBar.h
--- a/Engine/Source/Editor/Sequencer/Private/Filters/SequencerFilterBar.h
+++ b/Engine/Source/Editor/Sequencer/Private/Filters/SequencerFilterBar.h
@@ -65,27 +65,20 @@ public:
 	FSequencerFiltersChanged& OnFiltersChanged() { return FiltersChangedEvent; }
 
 	TSharedRef<FSequencerTrackFilter_Text> GetTextFilter() const;
+	TSharedRef<FSequencerTrackFilter_HideIsolate> GetHideIsolateFilter() const;
 	FText GetFilterErrorText() const;
 
 	/** Hide/Isolate/Show Filter Functions */
 
-	TSet<UE::Sequencer::TWeakViewModelPtr<UE::Sequencer::IOutlinerExtension>> GetHiddenTracks() const;
-	TSet<UE::Sequencer::TWeakViewModelPtr<UE::Sequencer::IOutlinerExtension>> GetIsolatedTracks() const;
-
 	void HideTracks(const TSet<UE::Sequencer::TWeakViewModelPtr<UE::Sequencer::IOutlinerExtension>>& InTracks, const bool bInAddToExisting);
 	void UnhideTracks(const TSet<UE::Sequencer::TWeakViewModelPtr<UE::Sequencer::IOutlinerExtension>>& InTracks);
 
 	void IsolateTracks(const TSet<UE::Sequencer::TWeakViewModelPtr<UE::Sequencer::IOutlinerExtension>>& InTracks, const bool bInAddToExisting);
 	void UnisolateTracks(const TSet<UE::Sequencer::TWeakViewModelPtr<UE::Sequencer::IOutlinerExtension>>& InTracks);
-
-	void ShowAllTracks();
-
+	
 	bool HasHiddenTracks() const;
 	bool HasIsolatedTracks() const;
 
-	void EmptyHiddenTracks();
-	void EmptyIsolatedTracks();
-
 	TSharedPtr<FSequencerTrackFilter> FindFilterByDisplayName(const FString& InFilterName) const;
 	TSharedPtr<FSequencerTrackFilter_CustomText> FindCustomTextFilterByDisplayName(const FString& InFilterName) const;
 
@@ -219,8 +212,6 @@ public:
 
 	/** Misc Functions */
 
-	TSet<UE::Sequencer::TWeakViewModelPtr<UE::Sequencer::IOutlinerExtension>> GetSelectedTracksOrAll() const;
-
 	void SetTrackParentsExpanded(const UE::Sequencer::TViewModelPtr<UE::Sequencer::IOutlinerExtension>& InNode, const bool bInExpanded);
 
 	UWorld* GetWorld() const;
@@ -229,7 +220,6 @@ public:
 
 	bool ShouldUpdateOnTrackValueChanged() const;
 
-	TSharedRef<SFilterBarIsolateHideShow> MakeIsolateHideShowPanel();
 	TSharedRef<SComboButton> MakeAddFilterButton();
 
 	virtual bool ShouldShowFilterBarWidget() const override;
@@ -310,8 +300,12 @@ protected:
 	FMuteStateChangedDelegate OnMuteFiltersChangedEvent;
 
 private:
+	
 	static int32 InstanceCount;
-
+	
+	/** @return Gets the selected IOutlinerExtensions, or all of them if nothing is selected. */
+	TSet<UE::Sequencer::TWeakViewModelPtr<UE::Sequencer::IOutlinerExtension>> GetSelectedTracksOrAll() const;
+	
 	/** Do not call directly! Should only be called by FilterNodes(). */
 	bool FilterNodesRecursive(const bool bInHasActiveFilter, const UE::Sequencer::TViewModelPtr<UE::Sequencer::IOutlinerExtension>& InStartNode);
 	
diff --git a/Engine/Source/Editor/Sequencer/Private/Filters/Utils/HideIsolateViewModelUtils.cpp b/Engine/Source/Editor/Sequencer/Private/Filters/Utils/HideIsolateViewModelUtils.cpp
new file mode 100644
--- /dev/null
+++ b/Engine/Source/Editor/Sequencer/Private/Filters/Utils/HideIsolateViewModelUtils.cpp
@@ -0,0 +1,100 @@
+﻿// Copyright Epic Games, Inc. All Rights Reserved.
+
+#include "HideIsolateViewModelUtils.h"
+
+#include "Filters/Filters/SequencerTrackFilter_HideIsolate.h"
+#include "Filters/SequencerFilterData.h"
+#include "Internationalization/Text.h"
+#include "MVVM/Selection/Selection.h"
+#include "MVVM/ViewModels/SequencerEditorViewModel.h"
+
+#define LOCTEXT_NAMESPACE "FitlerViewModelUtils"
+
+namespace UE::Sequencer
+{
+TSet<TWeakViewModelPtr<IOutlinerExtension>> GetSelectedTracksOrAll(const FSequencerEditorViewModel& InViewModel)
+{
+	const TSharedPtr<FSequencerSelection> Selection = InViewModel.GetSelection();
+	if (!Selection.IsValid())
+	{
+		return {};
+	}
+
+	const TSet<TWeakViewModelPtr<IOutlinerExtension>> SelectedSet = Selection->Outliner.GetSelected();
+	return SelectedSet.IsEmpty() ? GetAllTracks(InViewModel) : SelectedSet;
+}
+
+TSet<TWeakViewModelPtr<IOutlinerExtension>> GetAllTracks(const FSequencerEditorViewModel& InViewModel)
+{
+	TSet<TWeakViewModelPtr<IOutlinerExtension>> OutTracks;
+	for (const TViewModelPtr<IOutlinerExtension>& TrackModel : InViewModel.GetRootModel()->GetDescendantsOfType<IOutlinerExtension>())
+	{
+		OutTracks.Add(TrackModel);
+	}
+	return OutTracks;
+}
+
+FText MakeHiddenTracksCountText(const FSequencerTrackFilter_HideIsolate& InFilter)
+{
+	return FText::Format(
+		LOCTEXT("HiddenTracksSummary", "{0} hidden tracks"), 
+		InFilter.GetHiddenTracks().Num()
+		);
+}
+
+FText MakeHiddenTracksCountText_WithTotal(
+	const FSequencerFilterData& InFilterData, const FSequencerTrackFilter_HideIsolate& InFilter
+	)
+{
+	return FText::Format(
+		LOCTEXT("HiddenTracksSummaryWithTotal", "{0} of {1} total tracks"), 
+		MakeHiddenTracksCountText(InFilter),
+		InFilterData.GetTotalNodeCount()
+		);
+}
+
+FText MakeIsolatedTracksCountText(const FSequencerTrackFilter_HideIsolate& InFilter)
+{
+	return FText::Format(
+		LOCTEXT("IsolatedTracksSummary", "{0} isolated tracks"), 
+		InFilter.GetHiddenTracks().Num()
+		);
+}
+
+FText MakeIsolatedTracksCountText_WithTotal(
+	const FSequencerFilterData& InFilterData, const FSequencerTrackFilter_HideIsolate& InFilterBar
+	)
+{
+	return FText::Format(
+		LOCTEXT("IsolatedTracksSummaryWithTotal", "{0} of {1} total tracks"), 
+		MakeIsolatedTracksCountText(InFilterBar),
+		InFilterData.GetTotalNodeCount()
+		);
+}
+
+FText MakeHiddenAndIsolatedCountText(const FSequencerTrackFilter_HideIsolate& InFilterBar)
+{
+	const FText HiddenTracksSummary = MakeHiddenTracksCountText(InFilterBar);
+	const FText IsolatedTracksSummary = MakeIsolatedTracksCountText(InFilterBar);
+	return FText::Format(
+		LOCTEXT("HideIsolateSummary", "{0}, {1}"),
+		HiddenTracksSummary, 
+		IsolatedTracksSummary
+		);
+}
+
+FText MakeHiddenAndIsolatedCountText_WithTotal(
+	const FSequencerFilterData& InFilterData, const FSequencerTrackFilter_HideIsolate& InFilterBar
+	)
+{
+	const int32 FilteredCount = InFilterData.GetDisplayNodeCount();
+	const int32 TotalCount = InFilterData.GetTotalNodeCount();
+	const FText HideIsolateSummary = MakeHiddenAndIsolatedCountText(InFilterBar);
+	return FText::Format(
+		LOCTEXT("LongDisplaySummary", "Showing {0} of {1} total tracks\n{2}"),
+		FilteredCount, TotalCount, HideIsolateSummary
+		);
+}
+} // namespace UE::Sequencer
+
+#undef LOCTEXT_NAMESPACE
\ No newline at end of file
diff --git a/Engine/Source/Editor/Sequencer/Private/Filters/Utils/HideIsolateViewModelUtils.h b/Engine/Source/Editor/Sequencer/Private/Filters/Utils/HideIsolateViewModelUtils.h
new file mode 100644
--- /dev/null
+++ b/Engine/Source/Editor/Sequencer/Private/Filters/Utils/HideIsolateViewModelUtils.h
@@ -0,0 +1,37 @@
+﻿// Copyright Epic Games, Inc. All Rights Reserved.
+
+#pragma once
+
+#include "Containers/ContainersFwd.h"
+#include "MVVM/ViewModelPtr.h"
+
+class FSequencerTrackFilter_HideIsolate;
+class FText;
+class ISequencerFilterBar;
+struct FSequencerFilterData;
+
+namespace UE::Sequencer
+{
+class FSequencerEditorViewModel;
+class IOutlinerExtension;
+
+/** @return Gets the selected IOutlinerExtensions, or all of them if nothing is selected. */
+TSet<TWeakViewModelPtr<IOutlinerExtension>> GetSelectedTracksOrAll(const FSequencerEditorViewModel& InViewModel);
+/** @return Gets all tracks */
+TSet<TWeakViewModelPtr<IOutlinerExtension>> GetAllTracks(const FSequencerEditorViewModel& InViewModel);
+
+/** @return "{0} hidden tracks" */
+FText MakeHiddenTracksCountText(const FSequencerTrackFilter_HideIsolate& InFilter);
+/** @return "{0} hidden tracks of {1} total tracks" */
+FText MakeHiddenTracksCountText_WithTotal(const FSequencerFilterData& InFilterData, const FSequencerTrackFilter_HideIsolate& InFilter);
+
+/** @return "{0} isolated tracks" */
+FText MakeIsolatedTracksCountText(const FSequencerTrackFilter_HideIsolate& InFilter);
+/** @return "{0} isolated tracks of {1} total tracks" */
+FText MakeIsolatedTracksCountText_WithTotal(const FSequencerFilterData& InFilterData, const FSequencerTrackFilter_HideIsolate& InFilterBar);
+
+/** @return "{0} hidden tracks, {1} isolated tracks" */
+FText MakeHiddenAndIsolatedCountText(const FSequencerTrackFilter_HideIsolate& InFilterBar);
+/** @return "Showing {0} of {1} total tracks\n{0} hidden tracks, {1} isolated tracks" */
+FText MakeHiddenAndIsolatedCountText_WithTotal(const FSequencerFilterData& InFilterData, const FSequencerTrackFilter_HideIsolate& InFilterBar);
+} // namespace UE::Sequencer
diff --git a/Engine/Source/Editor/Sequencer/Private/Filters/ViewModels/HideIsolateShowViewModel.cpp b/Engine/Source/Editor/Sequencer/Private/Filters/ViewModels/HideIsolateShowViewModel.cpp
new file mode 100644
--- /dev/null
+++ b/Engine/Source/Editor/Sequencer/Private/Filters/ViewModels/HideIsolateShowViewModel.cpp
@@ -0,0 +1,163 @@
+﻿// Copyright Epic Games, Inc. All Rights Reserved.
+
+#include "HideIsolateShowViewModel.h"
+
+#include "Filters/Filters/SequencerTrackFilter_HideIsolate.h"
+#include "Filters/ISequencerFilterBar.h"
+#include "Filters/SequencerTrackFilterCommands.h"
+#include "Filters/Utils/HideIsolateViewModelUtils.h"
+#include "MVVM/Selection/Selection.h"
+#include "MVVM/ViewModels/SequencerEditorViewModel.h"
+
+#define LOCTEXT_NAMESPACE "FIsolateHideShowViewModel"
+
+namespace UE::Sequencer
+{
+FHideIsolateShowViewModel::FHideIsolateShowViewModel(
+	const TSharedRef<ISequencerTrackFilters>& InFilterBar,
+	const TSharedRef<FSequencerTrackFilter_HideIsolate>& InIsolateFilter,
+	TAttribute<TSet<TWeakViewModelPtr<IOutlinerExtension>>> InSelectedTracksAttr,
+	TAttribute<TSet<TWeakViewModelPtr<IOutlinerExtension>>> InAllTracksAttr
+	)
+	: FilterBar(InFilterBar)
+	, IsolateFilter(InIsolateFilter)
+	, SelectedTracksAttr(MoveTemp(InSelectedTracksAttr))
+	, AllTracksAttr(MoveTemp(InAllTracksAttr))
+{}
+
+void FHideIsolateShowViewModel::EmptyIsolatedTracks()
+{
+	IsolateFilter->EmptyIsolatedTracks();
+	FilterBar->RequestFilterUpdate();
+}
+
+void FHideIsolateShowViewModel::IsolateSelectedTracks(bool bAddToExisting)
+{
+	IsolateFilter->IsolateTracks(GetSelectedTracksOrAll(), bAddToExisting);
+}
+
+void FHideIsolateShowViewModel::EmptyHiddenTracks()
+{
+	IsolateFilter->EmptyHiddenTracks();
+	FilterBar->RequestFilterUpdate();
+}
+
+void FHideIsolateShowViewModel::HideSelectedTracks(const bool bAddToExisting)
+{
+	IsolateFilter->HideTracks(GetSelectedTracksOrAll(), bAddToExisting);
+}
+
+void FHideIsolateShowViewModel::ShowAllTracks()
+{
+	IsolateFilter->ShowAllTracks();
+	
+	const TSet<TWeakViewModelPtr<IOutlinerExtension>> SelectedTracks = SelectedTracksAttr.Get();
+	if (SelectedTracks.Num() > 0)
+	{
+		const TWeakViewModelPtr<IOutlinerExtension> FirstItem = *SelectedTracks.CreateConstIterator();
+		RequestScrollToViewDelegate.Broadcast(FirstItem);
+	}
+	
+	FilterBar->RequestFilterUpdate();
+}
+
+bool FHideIsolateShowViewModel::AreFiltersMuted() const
+{
+	return FilterBar->AreFiltersMuted();
+}
+
+bool FHideIsolateShowViewModel::HasIsolatedTracks() const
+{
+	return IsolateFilter->GetIsolatedTracks().Num() > 0;
+}
+
+bool FHideIsolateShowViewModel::HasHiddenTracks() const
+{
+	return IsolateFilter->GetHiddenTracks().Num() > 0;
+}
+
+FText FHideIsolateShowViewModel::GetHideTracksButtonTooltipText() const
+{
+	FText TooltipText = LOCTEXT("HideTracksButtonToolTip", "Hide selected tracks");
+
+	const TSharedRef<const FInputChord> Chord = FSequencerTrackFilterCommands::Get().HideSelectedTracks->GetFirstValidChord();
+	if (Chord->IsValidChord())
+	{
+		TooltipText = FText::Format(
+			LOCTEXT("HideTracksButtonToolTipExtended", "{0} ({1})"), 
+			TooltipText, 
+			Chord->GetInputText()
+			);
+	}
+
+	const FText SummaryWithTotal = MakeHiddenTracksCountText_WithTotal(FilterBar->GetFilterData(), *IsolateFilter);
+	TooltipText = FText::Format(
+		LOCTEXT("HideTracksButtonToolTipExtendedWithTotal", "{0}\n\nUse the Control modifier to reset the hidden track list.\n\n{1}"), 
+		TooltipText, 
+		SummaryWithTotal
+		);
+
+	return TooltipText;
+}
+
+FText FHideIsolateShowViewModel::GetIsolateTracksButtonTooltipText() const
+{
+	FText TooltipText = LOCTEXT("IsolateTracksButtonToolTip", "Isolate selected tracks");
+
+	const TSharedRef<const FInputChord> Chord = FSequencerTrackFilterCommands::Get().IsolateSelectedTracks->GetFirstValidChord();
+	if (Chord->IsValidChord())
+	{
+		TooltipText = FText::Format(
+			LOCTEXT("IsolateTracksButtonToolTipExtended", "{0} ({1})"),
+			TooltipText,
+			Chord->GetInputText()
+			);
+	}
+
+	const FText SummaryWithTotal = MakeIsolatedTracksCountText_WithTotal(FilterBar->GetFilterData(), *IsolateFilter);
+	TooltipText = FText::Format(
+		LOCTEXT("IsolateTracksButtonToolTipExtendedWithTotal", "{0}\n\n"
+			"Use the Shift modifier to additively isolate.\n"
+			"Use the Control modifier to reset the isolated track list.\n\n{1}"
+			),
+		TooltipText, 
+		SummaryWithTotal
+		);
+
+	return TooltipText;
+}
+
+FText FHideIsolateShowViewModel::GetShowAllTracksButtonTooltipText() const
+{
+	FText TooltipText = LOCTEXT("ShowAllTracksButtonToolTip", "Show all tracks");
+
+	const TSharedRef<const FInputChord> Chord = FSequencerTrackFilterCommands::Get().ShowAllTracks->GetFirstValidChord();
+	if (Chord->IsValidChord())
+	{
+		TooltipText = FText::Format(
+			LOCTEXT("ShowAllTracksButtonToolTipExtended", "{0} ({1})"), 
+			TooltipText,
+			Chord->GetInputText()
+			);
+	}
+
+	const FText SummaryText = MakeHiddenAndIsolatedCountText_WithTotal(FilterBar->GetFilterData(), *IsolateFilter);
+	TooltipText = FText::Format(
+		LOCTEXT("ShowAllTracksButtonToolTipExtendedWithSummary", "{0}\n\n{1}"), 
+		TooltipText,
+		SummaryText
+		);
+
+	return TooltipText;
+}
+
+TSet<TWeakViewModelPtr<IOutlinerExtension>> FHideIsolateShowViewModel::GetSelectedTracksOrAll() const
+{
+	const TSet<TWeakViewModelPtr<IOutlinerExtension>> Selection = SelectedTracksAttr.Get();
+	return Selection.IsEmpty()
+		? AllTracksAttr.Get()
+		: Selection;
+}
+} // namespace UE::Sequencer
+
+#undef LOCTEXT_NAMESPACE
\ No newline at end of file
diff --git a/Engine/Source/Editor/Sequencer/Private/Filters/ViewModels/HideIsolateShowViewModel.h b/Engine/Source/Editor/Sequencer/Private/Filters/ViewModels/HideIsolateShowViewModel.h
new file mode 100644
--- /dev/null
+++ b/Engine/Source/Editor/Sequencer/Private/Filters/ViewModels/HideIsolateShowViewModel.h
@@ -0,0 +1,85 @@
+﻿// Copyright Epic Games, Inc. All Rights Reserved.
+
+#pragma once
+
+#include "Delegates/Delegate.h"
+#include "MVVM/ViewModelPtr.h"
+#include "Misc/Attribute.h"
+#include "Templates/SharedPointer.h"
+
+class FSequencerTrackFilter_HideIsolate;
+class ISequencerFilterBar;
+class ISequencerTrackFilters;
+
+namespace UE::Sequencer
+{
+class FSequencerEditorViewModel;
+class IOutlinerExtension;
+
+/** 
+ * View model for the SFilterBarIsolateHideShow widget.
+ * 
+ * It is carefully engineered to be independent of FSequencer references so UI external to Sequencer can potentially create instances.
+ * This API is private for now but can be made public more easily in the future.
+ */
+class FHideIsolateShowViewModel
+{
+public:
+	
+	/** Requests UI to view the given view model into view. */
+	DECLARE_MULTICAST_DELEGATE_OneParam(FRequestScrollToViewEvent, const TWeakViewModelPtr<IOutlinerExtension>&);
+	
+	explicit FHideIsolateShowViewModel(
+		const TSharedRef<ISequencerTrackFilters>& InFilterBar,
+		const TSharedRef<FSequencerTrackFilter_HideIsolate>& InIsolateFilter,
+		TAttribute<TSet<TWeakViewModelPtr<IOutlinerExtension>>> InSelectedTracksAttr,
+		TAttribute<TSet<TWeakViewModelPtr<IOutlinerExtension>>> InAllTracksAttr
+		);
+	
+	/** Clears all the isolated tracks. */
+	void EmptyIsolatedTracks();
+	/** Isolates the selected tracks, or all tracks if none are selected. */
+	void IsolateSelectedTracks(bool bAddToExisting);
+	
+	/** Clears all the hidden tracks. */
+	void EmptyHiddenTracks();
+	/** Hides the selected tracks, or all tracks if none are selected. */
+	void HideSelectedTracks(const bool bAddToExisting);
+	
+	/** Shows all tracks effectively clearing the hidden and isolate states. */
+	void ShowAllTracks();
+	
+	/** @return Whether the filters are muted.  */
+	bool AreFiltersMuted() const;
+	/** @return Whether any tracks are isolated */
+	bool HasIsolatedTracks() const;
+	/** @return Whether any tracks are hidden */
+	bool HasHiddenTracks() const;
+	
+	FText GetHideTracksButtonTooltipText() const;
+	FText GetIsolateTracksButtonTooltipText() const;
+	FText GetShowAllTracksButtonTooltipText() const;
+	
+	/** @return Broadcasts when an item is supposed to be focused, e.g. scrolled to view. */
+	FRequestScrollToViewEvent& OnRequestScrollToView() { return RequestScrollToViewDelegate; }
+	
+private:
+	
+	/** The filter bar for which this widget is being displayed. */
+	const TSharedRef<ISequencerTrackFilters> FilterBar;
+	/** The filter instance for which this widget is being displayed. */
+	const TSharedRef<FSequencerTrackFilter_HideIsolate> IsolateFilter;
+	
+	/** Gets all selected tracks. */
+	const TAttribute<TSet<TWeakViewModelPtr<IOutlinerExtension>>> SelectedTracksAttr;
+	/** Gets all tracks. */
+	const TAttribute<TSet<TWeakViewModelPtr<IOutlinerExtension>>> AllTracksAttr;
+	
+	/** Broadcasts when an item is supposed to be focused, e.g. scrolled to view. */
+	FRequestScrollToViewEvent RequestScrollToViewDelegate;
+	
+	/** @return Selected tracks, or all of them. */
+	TSet<TWeakViewModelPtr<IOutlinerExtension>> GetSelectedTracksOrAll() const;
+};
+} // namespace UE::Sequencer
+
diff --git a/Engine/Source/Editor/Sequencer/Private/Filters/Widgets/SFilterBarIsolateHideShow.cpp b/Engine/Source/Editor/Sequencer/Private/Filters/Widgets/SFilterBarIsolateHideShow.cpp
--- a/Engine/Source/Editor/Sequencer/Private/Filters/Widgets/SFilterBarIsolateHideShow.cpp
+++ b/Engine/Source/Editor/Sequencer/Private/Filters/Widgets/SFilterBarIsolateHideShow.cpp
@@ -1,8 +1,11 @@
 ﻿// Copyright Epic Games, Inc. All Rights Reserved.
 
 #include "SFilterBarIsolateHideShow.h"
+
 #include "Filters/SequencerFilterBar.h"
 #include "Filters/SequencerTrackFilterCommands.h"
+#include "Filters/ViewModels/HideIsolateShowViewModel.h"
+
 #include "Framework/Application/SlateApplication.h"
 #include "Styling/StyleColors.h"
 #include "Widgets/Images/SImage.h"
@@ -12,9 +15,11 @@
 
 #define LOCTEXT_NAMESPACE "SFilterBarIsolateHideShow"
 
-void SFilterBarIsolateHideShow::Construct(const FArguments& InArgs, const TSharedRef<FSequencerFilterBar>& InFilterBar)
+namespace UE::Sequencer
+{
+void SFilterBarIsolateHideShow::Construct(const FArguments& InArgs, const TSharedRef<FHideIsolateShowViewModel>& InViewModel)
 {
-	WeakFilterBar = InFilterBar;
+	ViewModel = InViewModel;
 
 	constexpr float ButtonContentPadding = 2.f;
 	constexpr float ButtonSpacing = 1.f;
@@ -34,12 +39,12 @@ void SFilterBarIsolateHideShow::Construct(const FArguments& InArgs, const TShare
 			.HAlign(HAlign_Center)
 			.VAlign(VAlign_Center)
 			.ButtonStyle(FAppStyle::Get(), TEXT("SimpleButton"))
-			.ToolTipText(this, &SFilterBarIsolateHideShow::GetIsolateTracksButtonTooltipText)
-			.IsEnabled(this, &SFilterBarIsolateHideShow::AreFiltersMuted)
+			.ToolTipText(InViewModel, &FHideIsolateShowViewModel::GetIsolateTracksButtonTooltipText)
+			.IsEnabled(this, &SFilterBarIsolateHideShow::AreFiltersUnmuted)
 			.OnClicked(this, &SFilterBarIsolateHideShow::HandleIsolateTracksClick)
 			[
 				ConstructLayeredImage(TEXT("Sequencer.TrackIsolate")
-					, TAttribute<bool>::CreateSP(this, &SFilterBarIsolateHideShow::HasIsolatedTracks))
+					, TAttribute<bool>::Create(TAttribute<bool>::FGetter::CreateSP(InViewModel, &FHideIsolateShowViewModel::HasIsolatedTracks)))
 			]
 		]
 
@@ -54,12 +59,12 @@ void SFilterBarIsolateHideShow::Construct(const FArguments& InArgs, const TShare
 			.HAlign(HAlign_Center)
 			.VAlign(VAlign_Center)
 			.ButtonStyle(FAppStyle::Get(), TEXT("SimpleButton"))
-			.ToolTipText(this, &SFilterBarIsolateHideShow::GetHideTracksButtonTooltipText)
-			.IsEnabled(this, &SFilterBarIsolateHideShow::AreFiltersMuted)
+			.ToolTipText(InViewModel, &FHideIsolateShowViewModel::GetHideTracksButtonTooltipText)
+			.IsEnabled(this, &SFilterBarIsolateHideShow::AreFiltersUnmuted)
 			.OnClicked(this, &SFilterBarIsolateHideShow::HandleHideTracksClick)
 			[
 				ConstructLayeredImage(TEXT("Sequencer.TrackHide")
-					, TAttribute<bool>::CreateSP(this, &SFilterBarIsolateHideShow::HasHiddenTracks))
+					, TAttribute<bool>::Create(TAttribute<bool>::FGetter::CreateSP(InViewModel, &FHideIsolateShowViewModel::HasHiddenTracks)))
 			]
 		]
 
@@ -74,8 +79,8 @@ void SFilterBarIsolateHideShow::Construct(const FArguments& InArgs, const TShare
 			.HAlign(HAlign_Center)
 			.VAlign(VAlign_Center)
 			.ButtonStyle(FAppStyle::Get(), TEXT("SimpleButton"))
-			.ToolTipText(this, &SFilterBarIsolateHideShow::GetShowAllTracksButtonTooltipText)
-			.IsEnabled(this, &SFilterBarIsolateHideShow::AreFiltersMuted)
+			.ToolTipText(InViewModel, &FHideIsolateShowViewModel::GetShowAllTracksButtonTooltipText)
+			.IsEnabled(this, &SFilterBarIsolateHideShow::AreFiltersUnmuted)
 			.OnClicked(this, &SFilterBarIsolateHideShow::HandleShowAllTracksClick)
 			[
 				SNew(SImage)
@@ -104,28 +109,23 @@ TSharedRef<SWidget> SFilterBarIsolateHideShow::ConstructLayeredImage(const FName
 	return LayeredImage;
 }
 
-bool SFilterBarIsolateHideShow::AreFiltersMuted() const
+bool SFilterBarIsolateHideShow::AreFiltersUnmuted() const
 {
-	return WeakFilterBar.IsValid() ? !WeakFilterBar.Pin()->AreFiltersMuted() : true;
+	return !ViewModel->AreFiltersMuted();
 }
 
 FReply SFilterBarIsolateHideShow::HandleHideTracksClick()
 {
-	const TSharedPtr<FSequencerFilterBar> FilterBar = WeakFilterBar.Pin();
-	if (!FilterBar.IsValid())
-	{
-		return FReply::Handled();
-	}
-
 	const FModifierKeysState ModifierKeys = FSlateApplication::Get().GetModifierKeys();
 	const bool bControlDown = ModifierKeys.AreModifersDown(EModifierKey::Control);
 	if (bControlDown)
 	{
-		FilterBar->EmptyHiddenTracks();
+		ViewModel->EmptyHiddenTracks();
 	}
 	else
 	{
-		FilterBar->HideSelectedTracks();
+		const bool bAddToExisting = !FSlateApplication::Get().GetModifierKeys().AreModifersDown(EModifierKey::Shift);
+		ViewModel->HideSelectedTracks(bAddToExisting);
 	}
 
 	return FReply::Handled();
@@ -133,21 +133,16 @@ FReply SFilterBarIsolateHideShow::HandleHideTracksClick()
 
 FReply SFilterBarIsolateHideShow::HandleIsolateTracksClick()
 {
-	const TSharedPtr<FSequencerFilterBar> FilterBar = WeakFilterBar.Pin();
-	if (!FilterBar.IsValid())
-	{
-		return FReply::Handled();
-	}
-
 	const FModifierKeysState ModifierKeys = FSlateApplication::Get().GetModifierKeys();
 	const bool bControlDown = ModifierKeys.AreModifersDown(EModifierKey::Control);
 	if (bControlDown)
 	{
-		FilterBar->EmptyIsolatedTracks();
+		ViewModel->EmptyIsolatedTracks();
 	}
 	else
 	{
-		FilterBar->IsolateSelectedTracks();
+		const bool bAddToExisting = FSlateApplication::Get().GetModifierKeys().AreModifersDown(EModifierKey::Shift);
+		ViewModel->IsolateSelectedTracks(bAddToExisting);
 	}
 
 	return FReply::Handled();
@@ -155,157 +150,15 @@ FReply SFilterBarIsolateHideShow::HandleIsolateTracksClick()
 
 FReply SFilterBarIsolateHideShow::HandleShowAllTracksClick()
 {
-	const TSharedPtr<FSequencerFilterBar> FilterBar = WeakFilterBar.Pin();
-	if (!FilterBar.IsValid())
-	{
-		return FReply::Handled();
-	}
-
-	FilterBar->ShowAllTracks();
-
+	ViewModel->ShowAllTracks();
 	return FReply::Handled();
 }
 
-bool SFilterBarIsolateHideShow::HasIsolatedTracks() const
-{
-	return WeakFilterBar.IsValid() && WeakFilterBar.Pin()->GetIsolatedTracks().Num() > 0;
-}
-
-bool SFilterBarIsolateHideShow::HasHiddenTracks() const
-{
-	return WeakFilterBar.IsValid() && WeakFilterBar.Pin()->GetHiddenTracks().Num() > 0;
-}
-
 FSlateColor SFilterBarIsolateHideShow::GetShowAllTracksButtonTextColor() const
 {
-	if (const TSharedPtr<FSequencerFilterBar> FilterBar = WeakFilterBar.Pin())
-	{
-		return (FilterBar->GetHiddenTracks().Num() == 0 && FilterBar->GetIsolatedTracks().Num() == 0)
-			? FStyleColors::Foreground : FStyleColors::Warning;
-	}
-	return FStyleColors::Foreground;
-}
-
-FText SFilterBarIsolateHideShow::GetHideTracksButtonTooltipText() const
-{
-	const TSharedPtr<FSequencerFilterBar> FilterBar = WeakFilterBar.Pin();
-	if (!FilterBar.IsValid())
-	{
-		return FText::GetEmpty();
-	}
-
-	FText TooltipText = LOCTEXT("HideTracksButtonToolTip", "Hide selected tracks");
-
-	const TSharedRef<const FInputChord> Chord = FSequencerTrackFilterCommands::Get().HideSelectedTracks->GetFirstValidChord();
-	if (Chord->IsValidChord())
-	{
-		TooltipText = FText::Format(LOCTEXT("HideTracksButtonToolTipExtended", "{0} ({1})")
-			, TooltipText, Chord->GetInputText());
-	}
-
-	const FText SummaryWithTotal = MakeHiddenTracksSummaryText(*FilterBar, true);
-	TooltipText = FText::Format(LOCTEXT("HideTracksButtonToolTipExtendedWithTotal", "{0}\n\n"
-		"Use the Control modifier to reset the hidden track list.\n\n{1}")
-		, TooltipText, SummaryWithTotal);
-
-	return TooltipText;
-}
-
-FText SFilterBarIsolateHideShow::GetIsolateTracksButtonTooltipText() const
-{
-	const TSharedPtr<FSequencerFilterBar> FilterBar = WeakFilterBar.Pin();
-	if (!FilterBar.IsValid())
-	{
-		return FText::GetEmpty();
-	}
-
-	FText TooltipText = LOCTEXT("IsolateTracksButtonToolTip", "Isolate selected tracks");
-
-	const TSharedRef<const FInputChord> Chord = FSequencerTrackFilterCommands::Get().IsolateSelectedTracks->GetFirstValidChord();
-	if (Chord->IsValidChord())
-	{
-		TooltipText = FText::Format(LOCTEXT("IsolateTracksButtonToolTipExtended", "{0} ({1})")
-			, TooltipText, Chord->GetInputText());
-	}
-
-	const FText SummaryWithTotal = MakeIsolatedTracksSummaryText(*FilterBar, true);
-	TooltipText = FText::Format(LOCTEXT("IsolateTracksButtonToolTipExtendedWithTotal", "{0}\n\n"
-		"Use the Shift modifier to additively isolate.\n"
-		"Use the Control modifier to reset the isolated track list.\n\n{1}")
-		, TooltipText, SummaryWithTotal);
-
-	return TooltipText;
-}
-
-FText SFilterBarIsolateHideShow::GetShowAllTracksButtonTooltipText() const
-{
-	const TSharedPtr<FSequencerFilterBar> FilterBar = WeakFilterBar.Pin();
-	if (!FilterBar.IsValid())
-	{
-		return FText::GetEmpty();
-	}
-
-	FText TooltipText = LOCTEXT("ShowAllTracksButtonToolTip", "Show all tracks");
-
-	const TSharedRef<const FInputChord> Chord = FSequencerTrackFilterCommands::Get().ShowAllTracks->GetFirstValidChord();
-	if (Chord->IsValidChord())
-	{
-		TooltipText = FText::Format(LOCTEXT("ShowAllTracksButtonToolTipExtended", "{0} ({1})")
-			, TooltipText, Chord->GetInputText());
-	}
-
-	const FText SummaryText = MakeLongDisplaySummaryText(*FilterBar);
-	TooltipText = FText::Format(LOCTEXT("ShowAllTracksButtonToolTipExtendedWithSummary", "{0}\n\n{1}")
-		, TooltipText, SummaryText);
-
-	return TooltipText;
-}
-
-FText SFilterBarIsolateHideShow::MakeHiddenTracksSummaryText(FSequencerFilterBar& InFilterBar, const bool bInShowTotalCount)
-{
-	FText Summary = FText::Format(LOCTEXT("HiddenTracksSummary", "{0} hidden tracks")
-		, InFilterBar.GetHiddenTracks().Num());
-
-	if (bInShowTotalCount)
-	{
-		Summary = FText::Format(LOCTEXT("HiddenTracksSummaryWithTotal", "{0} of {1} total tracks")
-			, Summary, InFilterBar.GetFilterData().GetTotalNodeCount());
-	}
-
-	return Summary;
-}
-
-FText SFilterBarIsolateHideShow::MakeIsolatedTracksSummaryText(FSequencerFilterBar& InFilterBar, const bool bInShowTotalCount)
-{
-	FText Summary = FText::Format(LOCTEXT("IsolatedTracksSummary", "{0} isolated tracks")
-		, InFilterBar.GetIsolatedTracks().Num());
-
-	if (bInShowTotalCount)
-	{
-		Summary = FText::Format(LOCTEXT("IsolatedTracksSummaryWithTotal", "{0} of {1} total tracks")
-			, Summary, InFilterBar.GetFilterData().GetTotalNodeCount());
-	}
-
-	return Summary;
-}
-
-FText SFilterBarIsolateHideShow::MakeHideIsolateTracksSummaryText(FSequencerFilterBar& InFilterBar)
-{
-	const FText HiddenTracksSummary = MakeHiddenTracksSummaryText(InFilterBar, false);
-	const FText IsolatedTracksSummary = MakeIsolatedTracksSummaryText(InFilterBar, false);
-	return FText::Format(LOCTEXT("HideIsolateSummary", "{0}, {1}")
-		, HiddenTracksSummary, IsolatedTracksSummary);
-}
-
-FText SFilterBarIsolateHideShow::MakeLongDisplaySummaryText(FSequencerFilterBar& InFilterBar)
-{
-	const FSequencerFilterData& FilterData = InFilterBar.GetFilterData();
-	const int32 FilteredCount = FilterData.GetDisplayNodeCount();
-	const int32 TotalCount = FilterData.GetTotalNodeCount();
-	const FText HideIsolateSummary = MakeHideIsolateTracksSummaryText(InFilterBar);
-
-	return FText::Format(LOCTEXT("LongDisplaySummary", "Showing {0} of {1} total tracks\n{2}"),
-		FilteredCount, TotalCount, HideIsolateSummary);
+	const bool bHasActiveFilters = ViewModel->HasHiddenTracks() || ViewModel->HasIsolatedTracks();
+	return bHasActiveFilters ? FStyleColors::Warning : FStyleColors::Foreground;
 }
+} // namespace UE::Sequencer
 
 #undef LOCTEXT_NAMESPACE
diff --git a/Engine/Source/Editor/Sequencer/Private/Filters/Widgets/SFilterBarIsolateHideShow.h b/Engine/Source/Editor/Sequencer/Private/Filters/Widgets/SFilterBarIsolateHideShow.h
--- a/Engine/Source/Editor/Sequencer/Private/Filters/Widgets/SFilterBarIsolateHideShow.h
+++ b/Engine/Source/Editor/Sequencer/Private/Filters/Widgets/SFilterBarIsolateHideShow.h
@@ -4,39 +4,36 @@
 
 #include "Widgets/SCompoundWidget.h"
 
-class FSequencerFilterBar;
+class ISequencerFilterBar;
 
+namespace UE::Sequencer
+{
+class FHideIsolateShowViewModel;
+
+/** View for visualizing the FSequencerTrackFilter_HideIsolate filter. */
 class SFilterBarIsolateHideShow : public SCompoundWidget
 {
 public:
-	static FText MakeHiddenTracksSummaryText(FSequencerFilterBar& InFilterBar, const bool bInShowTotalCount);
-	static FText MakeIsolatedTracksSummaryText(FSequencerFilterBar& InFilterBar, const bool bInShowTotalCount);
-	static FText MakeHideIsolateTracksSummaryText(FSequencerFilterBar& InFilterBar);
-	static FText MakeLongDisplaySummaryText(FSequencerFilterBar& InFilterBar);
-
-	SLATE_BEGIN_ARGS(SFilterBarIsolateHideShow)
-	{}
+	
+	SLATE_BEGIN_ARGS(SFilterBarIsolateHideShow) {}
 	SLATE_END_ARGS()
 
-	void Construct(const FArguments& InArgs, const TSharedRef<FSequencerFilterBar>& InFilterBar);
+	void Construct(const FArguments& InArgs, const TSharedRef<FHideIsolateShowViewModel>& InViewModel);
 
-protected:
+private:
+	
+	/** View model for this widget. */
+	TSharedPtr<FHideIsolateShowViewModel> ViewModel;
+	
 	TSharedRef<SWidget> ConstructLayeredImage(const FName InBaseImageName, const TAttribute<bool>& InShowBadge);
 	
-	bool AreFiltersMuted() const;
+	/** @return Whether filters are not muted, i.e. whether filtering UI should be enabled. */
+	bool AreFiltersUnmuted() const;
 
 	FReply HandleHideTracksClick();
 	FReply HandleIsolateTracksClick();
 	FReply HandleShowAllTracksClick();
 
-	bool HasIsolatedTracks() const;
-	bool HasHiddenTracks() const;
-
 	FSlateColor GetShowAllTracksButtonTextColor() const;
-
-	FText GetHideTracksButtonTooltipText() const;
-	FText GetIsolateTracksButtonTooltipText() const;
-	FText GetShowAllTracksButtonTooltipText() const;
-
-	TWeakPtr<FSequencerFilterBar> WeakFilterBar;
 };
+}
\ No newline at end of file
diff --git a/Engine/Source/Editor/Sequencer/Private/Filters/Widgets/SSearchAndFilterWidget.cpp b/Engine/Source/Editor/Sequencer/Private/Filters/Widgets/SSearchAndFilterWidget.cpp
--- a/Engine/Source/Editor/Sequencer/Private/Filters/Widgets/SSearchAndFilterWidget.cpp
+++ b/Engine/Source/Editor/Sequencer/Private/Filters/Widgets/SSearchAndFilterWidget.cpp
@@ -2,21 +2,24 @@
 
 #include "SSearchAndFilterWidget.h"
 
-#include "Sequencer.h"
 #include "Filters/Menus/SequencerViewOptionsMenu.h"
 #include "Filters/SequencerFilterBar.h"
+#include "Filters/SequencerFilterBar.h"
+#include "Filters/ViewModels/HideIsolateShowViewModel.h"
 #include "Filters/Widgets/SFilterBarIsolateHideShow.h"
 #include "Filters/Widgets/SSequencerSearchBox.h"
 #include "SPositiveActionButton.h"
 #include "SSequencerFilterBar.h"
-#include "Styling/AppStyle.h"
+#include "Sequencer.h"
 
 #define LOCTEXT_NAMESPACE "SSearchAndFilterRow"
 
 namespace UE::Sequencer
 {
 void SSearchAndFilterWidget::Construct(
-	const FArguments& InArgs, const TSharedRef<FSequencerFilterBar>& InFilterBar
+	const FArguments& InArgs, 
+	const TSharedRef<FSequencerFilterBar>& InFilterBar,
+	const TSharedRef<FHideIsolateShowViewModel>& InIsolateHideViewModel
 	)
 {
 	FilterBar = InFilterBar;
@@ -55,7 +58,7 @@ void SSearchAndFilterWidget::Construct(
 		.VAlign(VAlign_Center)
 		.Padding(2.f, 0.f, 0.f, 0.f)
 		[
-			InFilterBar->MakeIsolateHideShowPanel()
+			SNew(SFilterBarIsolateHideShow, InIsolateHideViewModel)
 		]
 	];
 }
diff --git a/Engine/Source/Editor/Sequencer/Private/Filters/Widgets/SSearchAndFilterWidget.h b/Engine/Source/Editor/Sequencer/Private/Filters/Widgets/SSearchAndFilterWidget.h
--- a/Engine/Source/Editor/Sequencer/Private/Filters/Widgets/SSearchAndFilterWidget.h
+++ b/Engine/Source/Editor/Sequencer/Private/Filters/Widgets/SSearchAndFilterWidget.h
@@ -2,17 +2,24 @@
 
 #pragma once
 
+#include "MVVM/ViewModelPtr.h"
 #include "SSequencerCustomTextFilterDialog.h"
 #include "Widgets/DeclarativeSyntaxSupport.h"
 #include "Widgets/SCompoundWidget.h"
 
 class FSequencerFilterBar;
-class ISequencerFilterBar;
+class ISequencerTrackFilters;
 class SSequencerSearchBox;
 
 namespace UE::Sequencer
 {
-/** Contains widgets for searching and filtering the Sequencer tree view. */
+class FHideIsolateShowViewModel;
+class IOutlinerExtension;
+
+/** 
+ * Contains widgets for searching and filtering a tree view displaying IOutlinerExtensions. 
+ * This view is intentionally engineered to be used independent of SSequencer.
+ */
 class SSearchAndFilterWidget : public SCompoundWidget
 {
 public:
@@ -22,12 +29,16 @@ public:
 		SLATE_EVENT(FOnTextChanged, OnSearchTextChanged)
 	SLATE_END_ARGS()
 
-	void Construct(const FArguments& InArgs, const TSharedRef<FSequencerFilterBar>& InFilterBar);
+	void Construct(
+		const FArguments& InArgs, 
+		const TSharedRef<FSequencerFilterBar>& InFilterBar,
+		const TSharedRef<FHideIsolateShowViewModel>& InIsolateHideViewModel
+		);
 
 private:
 	
 	/** The filter bar that is being displayed. */
-	TSharedPtr<FSequencerFilterBar> FilterBar;
+	TSharedPtr<ISequencerTrackFilters> FilterBar;
 	
 	/** Invoked when the search text changes. */
 	FOnTextChanged OnSearchTextChanged;
diff --git a/Engine/Source/Editor/Sequencer/Private/Widgets/Sequencer/SSequencerSearchAndFilterRow.cpp b/Engine/Source/Editor/Sequencer/Private/Widgets/Sequencer/SSequencerSearchAndFilterRow.cpp
--- a/Engine/Source/Editor/Sequencer/Private/Widgets/Sequencer/SSequencerSearchAndFilterRow.cpp
+++ b/Engine/Source/Editor/Sequencer/Private/Widgets/Sequencer/SSequencerSearchAndFilterRow.cpp
@@ -4,9 +4,14 @@
 
 #include "Filters/Menus/SequencerViewOptionsMenu.h"
 #include "Filters/SequencerFilterBar.h"
+#include "Filters/Utils/HideIsolateViewModelUtils.h"
+#include "Filters/ViewModels/HideIsolateShowViewModel.h"
 #include "Filters/Widgets/SFilterBarIsolateHideShow.h"
 #include "Filters/Widgets/SSearchAndFilterWidget.h"
 #include "Filters/Widgets/SSequencerSearchBox.h"
+#include "MVVM/Selection/Selection.h"
+#include "MVVM/Views/SOutlinerView.h"
+#include "SSequencer.h"
 #include "Sequencer.h"
 #include "Widgets/Sequencer/SAddToSequencerButton.h"
 #include "Widgets/Sequencer/SSequencerViewOptionsButton.h"
@@ -15,12 +20,51 @@
 
 namespace UE::Sequencer
 {
+namespace FilterRowDetail
+{
+using FOutlinerExtensionsAttr = TAttribute<TSet<TWeakViewModelPtr<IOutlinerExtension>>>; 
+
+static FOutlinerExtensionsAttr MakeSelectedTracksAttr(const TSharedRef<FSequencer>& InSequencer)
+{
+	return FOutlinerExtensionsAttr::CreateLambda([WeakSequencer = InSequencer.ToWeakPtr()]() -> TSet<TWeakViewModelPtr<IOutlinerExtension>>
+	{
+		const TSharedPtr<FSequencer> SequencerPin = WeakSequencer.Pin();
+		if (!SequencerPin)
+		{
+			return {};
+		}
+			
+		const TSharedPtr<FSequencerSelection> Selection = SequencerPin->GetViewModel()->GetSelection();
+		if (!Selection)
+		{
+			return {};
+		}
+			
+		return Selection->Outliner.GetSelected();
+	});
+}
+
+static FOutlinerExtensionsAttr MakeAllTracksAttr(const TSharedRef<FSequencer>& InSequencer)
+{
+	return FOutlinerExtensionsAttr::CreateLambda([WeakSequencer = InSequencer.ToWeakPtr()]()
+	{
+		const TSharedPtr<FSequencer> SequencerPin = WeakSequencer.Pin();
+		return SequencerPin ? GetAllTracks(*SequencerPin->GetViewModel()) : TSet<TWeakViewModelPtr<IOutlinerExtension>>();
+	});
+}
+}
 void SSequencerSearchAndFilterRow::Construct(const FArguments& InArgs)
 {
 	WeakSequencer = InArgs._Sequencer;
 	
-	const TSharedPtr<FSequencer> SequencerPin = WeakSequencer.Pin();
-	check(SequencerPin);
+	const TSharedRef<FSequencer> SequencerPin = WeakSequencer.Pin().ToSharedRef();
+	
+	const TSharedRef<FSequencerFilterBar> FilterBar = SequencerPin->GetFilterBar();
+	const TSharedRef<FHideIsolateShowViewModel> IsolateViewModel = MakeShared<FHideIsolateShowViewModel>(
+		FilterBar, FilterBar->GetHideIsolateFilter(),
+		FilterRowDetail::MakeSelectedTracksAttr(SequencerPin), FilterRowDetail::MakeAllTracksAttr(SequencerPin)
+		);
+	IsolateViewModel->OnRequestScrollToView().AddSP(this, &SSequencerSearchAndFilterRow::ScrollIntoView);
 	
 	ChildSlot
 	[
@@ -41,7 +85,7 @@ void SSequencerSearchAndFilterRow::Construct(const FArguments& InArgs)
 		.FillWidth(1.f)
 		.VAlign(VAlign_Center)
 		[
-			SNew(SSearchAndFilterWidget, SequencerPin->GetFilterBar())
+			SNew(SSearchAndFilterWidget, FilterBar, IsolateViewModel)
 			.OnSearchTextChanged(this, &SSequencerSearchAndFilterRow::OnSearchTextChanged)
 		]
 		
@@ -64,6 +108,15 @@ void SSequencerSearchAndFilterRow::OnSearchTextChanged(const FText& InSearchText
 		Sequencer->GetNodeTree()->SetTextFilterString(InSearchText.ToString());
 	}
 }
+
+void SSequencerSearchAndFilterRow::ScrollIntoView(const TWeakViewModelPtr<IOutlinerExtension>& InItem)
+{
+	if (const TSharedPtr<FSequencer> SequencerPin = WeakSequencer.Pin())
+	{
+		const TSharedRef<SSequencer> SequencerWidget = StaticCastSharedRef<SSequencer>(SequencerPin->GetSequencerWidget());
+		SequencerWidget->GetTreeView()->RequestScrollIntoView(InItem);
+	}
+}
 } // namespace UE::Sequencer
 
 #undef LOCTEXT_NAMESPACE
\ No newline at end of file
diff --git a/Engine/Source/Editor/Sequencer/Private/Widgets/Sequencer/SSequencerSearchAndFilterRow.h b/Engine/Source/Editor/Sequencer/Private/Widgets/Sequencer/SSequencerSearchAndFilterRow.h
--- a/Engine/Source/Editor/Sequencer/Private/Widgets/Sequencer/SSequencerSearchAndFilterRow.h
+++ b/Engine/Source/Editor/Sequencer/Private/Widgets/Sequencer/SSequencerSearchAndFilterRow.h
@@ -2,6 +2,7 @@
 
 #pragma once
 
+#include "MVVM/ViewModelPtr.h"
 #include "Widgets/DeclarativeSyntaxSupport.h"
 #include "Widgets/SCompoundWidget.h"
 
@@ -11,6 +12,8 @@ class FSequencerFilterBar;
 
 namespace UE::Sequencer
 {
+class IOutlinerExtension;
+
 /** The row of widgets displayed above the Sequence tree view. */
 class SSequencerSearchAndFilterRow : public SCompoundWidget
 {
@@ -33,7 +36,11 @@ private:
 	/** Used to update Sequencer's filters. */
 	TWeakPtr<FSequencer> WeakSequencer;
 	
+	/** Propagates changes in the search widget to the filter system */
 	void OnSearchTextChanged(const FText& InSearchText);
+	
+	/** Handles request to scroll the item into view. */
+	void ScrollIntoView(const TWeakViewModelPtr<IOutlinerExtension>& InItem);
 };
 } // namespace UE::Sequencer
 
diff --git a/Engine/Source/Editor/Sequencer/Public/Filters/Filters/SequencerTrackFilter_HideIsolate.h b/Engine/Source/Editor/Sequencer/Public/Filters/Filters/SequencerTrackFilter_HideIsolate.h
--- a/Engine/Source/Editor/Sequencer/Public/Filters/Filters/SequencerTrackFilter_HideIsolate.h
+++ b/Engine/Source/Editor/Sequencer/Public/Filters/Filters/SequencerTrackFilter_HideIsolate.h
@@ -13,8 +13,8 @@ public:
 
 	void ResetFilter();
 
-	TSet<UE::Sequencer::TWeakViewModelPtr<UE::Sequencer::IOutlinerExtension>> GetHiddenTracks() const;
-	TSet<UE::Sequencer::TWeakViewModelPtr<UE::Sequencer::IOutlinerExtension>> GetIsolatedTracks() const;
+	const TSet<UE::Sequencer::TWeakViewModelPtr<UE::Sequencer::IOutlinerExtension>>& GetHiddenTracks() const;
+	const TSet<UE::Sequencer::TWeakViewModelPtr<UE::Sequencer::IOutlinerExtension>>& GetIsolatedTracks() const;
 
 	void HideTracks(const TSet<UE::Sequencer::TWeakViewModelPtr<UE::Sequencer::IOutlinerExtension>>& InTracks, const bool bInAddToExisting);
 	void UnhideTracks(const TSet<UE::Sequencer::TWeakViewModelPtr<UE::Sequencer::IOutlinerExtension>>& InTracks);
