Modifier Outliner UI

Adding a generic view model for dispalying decorations/modifiers in the Sequencer tree.

The DecorationOutlinerModel can be added to existing outliner lists to display modifiers in their own subtree.

Currently the decoration model only supports displaying decorators that supply sections, but future support for properties is planned.

#Synced-CL 49762952

diff --git a/Engine/Source/Editor/Sequencer/Private/MVVM/ViewModels/DecorationOutlinerModel.cpp b/Engine/Source/Editor/Sequencer/Private/MVVM/ViewModels/DecorationOutlinerModel.cpp
new file mode 100644
--- /dev/null
+++ b/Engine/Source/Editor/Sequencer/Private/MVVM/ViewModels/DecorationOutlinerModel.cpp
@@ -0,0 +1,305 @@
+﻿// Copyright Epic Games, Inc. All Rights Reserved.
+
+#include "MVVM/ViewModels/DecorationOutlinerModel.h"
+
+#include "ISequencer.h"
+#include "MovieScene.h"
+#include "SequencerCommonHelpers.h"
+#include "Decorations/IMovieSceneSectionProviderDecoration.h"
+#include "MVVM/SectionModelStorageExtension.h"
+#include "MVVM/Extensions/IRecyclableExtension.h"
+#include "MVVM/ViewModels/ObjectBindingModel.h"
+#include "MVVM/ViewModels/SequenceModel.h"
+#include "MVVM/ViewModels/SectionModel.h"
+#include "MVVM/ViewModels/SectionOutlinerModel.h"
+#include "MVVM/ViewModels/SequencerEditorViewModel.h"
+#include "MVVM/ViewModels/TrackModelLayoutBuilder.h"
+#include "MVVM/ViewModels/TrackRowModel.h"
+#include "MVVM/ViewModels/ViewDensity.h"
+
+#define LOCTEXT_NAMESPACE "DecorationOutlinerModel"
+
+namespace UE::Sequencer
+{
+UE_SEQUENCER_DEFINE_CASTABLE(FDecorationOutlinerModel);
+
+FDecorationOutlinerModel::FDecorationOutlinerModel(UMovieSceneTrack* InParentTrack, UMovieSceneDecorationContainerObject* InDecorationContainer)
+	: WeakParentTrack(InParentTrack)
+	, WeakDecorationContainer(InDecorationContainer)
+{
+	RegisterChildList(&DecoratorList);
+}
+
+FDecorationOutlinerModel::~FDecorationOutlinerModel()
+{
+}
+
+void FDecorationOutlinerModel::OnConstruct()
+{
+	if (!WeakParentTrack.IsValid() || !WeakDecorationContainer.IsValid())
+	{
+		return;
+	}
+	
+	// Get ISequencer pointer from the ancestor view model.
+	TSharedPtr<FSequenceModel> SequenceModel = FindAncestorOfType<FSequenceModel>();
+	if (!SequenceModel)
+	{
+		return;
+	}
+
+	FSectionModelStorageExtension* SectionModelStorage = SequenceModel->CastDynamic<FSectionModelStorageExtension>();
+	if (!SectionModelStorage)
+	{
+		return;
+	}
+
+	const TSharedPtr<ISequencer> Sequencer = SequenceModel->GetSequencer();
+	if (!Sequencer)
+	{
+		return;
+	}
+	
+	TrackEditor = Sequencer->GetTrackEditor(WeakParentTrack.Get());
+	FViewModelChildren OutlinerChildren = GetTopLevelChannels();
+	
+	FScopedViewModelListHead RecycledModels(AsShared(), EViewModelListType::Recycled);
+	OutlinerChildren.MoveChildrenTo<IRecyclableExtension>(RecycledModels.GetChildren(), IRecyclableExtension::CallOnRecycle);
+	
+	TSharedPtr<FViewModel> OutlinerTail;
+	for (TObjectPtr<UObject> Decoration : WeakDecorationContainer->GetDecorations())
+	{
+		// For decorations that have sections, create and add models for those sections.
+		if (IMovieSceneSectionProviderDecoration* SectionProvider = Cast<IMovieSceneSectionProviderDecoration>(Decoration))
+		{
+			for (UMovieSceneSection* Section : SectionProvider->GetSections())
+			{
+				if (!Section)
+				{
+					continue;
+				}
+				const TSharedRef<ISequencerSection> SectionInterface = TrackEditor->MakeSectionInterface(*Section, *WeakParentTrack.Get(), WeakParentTrack.Get()->FindObjectBindingGuid());
+				const TSharedPtr<FSectionModel> SectionModel = SectionModelStorage->CreateModelForSection(Section, SectionInterface);
+
+				TSharedPtr<FSectionOutlinerModel> SectionOutlinerModel;
+				for (TSharedPtr<FSectionOutlinerModel> ExistingOutliner : OutlinerChildren.IterateSubList<FSectionOutlinerModel>())
+				{
+					if (ExistingOutliner && ExistingOutliner->GetSection() == Section)
+					{
+						SectionOutlinerModel = ExistingOutliner;
+						break;
+					}
+				}
+
+				if (!SectionOutlinerModel)
+				{
+					SectionOutlinerModel = MakeShared<FSectionOutlinerModel>(Section, SectionModel);
+				}
+
+				OutlinerChildren.InsertChild(SectionOutlinerModel, OutlinerTail);
+				OutlinerTail = SectionOutlinerModel;
+
+				SectionOutlinerModel->OnConstruct();
+				SectionOutlinerModel->GetTrackAreaChildren().AddChild(SectionModel);
+				SectionModel->SetLinkedOutlinerItem(SectionOutlinerModel);
+
+				FTrackModelLayoutBuilder LayoutBuilder(SectionOutlinerModel);
+				LayoutBuilder.RefreshLayout(SectionModel);
+			}
+		}
+	}
+}
+
+/*------------------------------------------------------------------------------------
+	FOutlinerItemModel
+ -----------------------------------------------------------------------------------*/
+void FDecorationOutlinerModel::BuildContextMenu(FMenuBuilder& MenuBuilder)
+{
+	const TSharedPtr<FSequencerEditorViewModel> EditorViewModel = GetEditor();
+	if (!EditorViewModel.IsValid())
+	{
+		return;
+	}
+	
+	const TSharedPtr<FSequencer> Sequencer = EditorViewModel->GetSequencerImpl();
+	if (!Sequencer.IsValid())
+	{
+		return;
+	}
+	
+	if (WeakDecorationContainer.IsValid())
+	{
+		for (TObjectPtr<UObject> Decoration : WeakDecorationContainer->GetDecorations())
+		{
+			if (Decoration->Implements<UMovieSceneSectionProviderDecoration>())
+			{
+				const TArray<TWeakObjectPtr<>> TrackAreaModels = SequencerHelpers::GetSectionObjectsFromTrackAreaModels(GetTrackAreaModelList());
+				SequencerHelpers::BuildEditSectionMenu(Sequencer, TrackAreaModels, MenuBuilder, true);
+			}
+		}
+	}
+	FOutlinerItemModel::BuildContextMenu(MenuBuilder);
+}
+
+void FDecorationOutlinerModel::BuildSidebarMenu(FMenuBuilder& MenuBuilder)
+{
+	BuildContextMenu(MenuBuilder);
+}
+
+/*------------------------------------------------------------------------------------
+	IOutlinerExtension
+ -----------------------------------------------------------------------------------*/
+FOutlinerSizing FDecorationOutlinerModel::GetOutlinerSizing() const
+{
+	const FViewDensityInfo Density = GetEditor()->GetViewDensity();
+	constexpr float SectionAreaDefaultHeight = 27.0f;
+	float Height = Density.UniformHeight.Get(SectionAreaDefaultHeight);
+	for (const TSharedPtr<FSectionModel> Section : DecoratorList.Iterate<FSectionModel>())
+	{
+		Height = Section->GetSectionInterface()->GetSectionHeight(Density);
+		break;
+	}
+	return FOutlinerSizing(Height);
+}
+
+TSharedPtr<SWidget> FDecorationOutlinerModel::CreateOutlinerViewForColumn(const FCreateOutlinerViewParams& InParams,
+																						const FName& InColumnName)
+{
+	if (TrackEditor)
+	{
+		const FBuildColumnWidgetParams Params(SharedThis(this), InParams);
+		return TrackEditor->BuildOutlinerColumnWidget(Params, InColumnName);
+	}
+	return nullptr;
+}
+
+FText FDecorationOutlinerModel::GetLabel() const
+{
+	return LOCTEXT("ModifiersLabel", "Modifiers");
+}
+
+FViewModelChildren FDecorationOutlinerModel::GetTopLevelChannels()
+{
+	return GetChildrenForList(&DecoratorList);
+}
+
+/*------------------------------------------------------------------------------------
+	ITrackAreaExtension
+------------------------------------------------------------------------------------*/
+
+FTrackAreaParameters FDecorationOutlinerModel::GetTrackAreaParameters() const
+{
+	FTrackAreaParameters Parameters;
+	Parameters.LaneType = ETrackAreaLaneType::Inline;
+	return Parameters;
+}
+
+FViewModelVariantIterator FDecorationOutlinerModel::GetTrackAreaModelList() const
+{
+	return &DecoratorList;
+}
+
+FViewModelVariantIterator FDecorationOutlinerModel::GetTopLevelChildTrackAreaModels() const
+{
+	return &DecoratorList;
+}
+
+/*------------------------------------------------------------------------------------
+	IDeletableExtension
+------------------------------------------------------------------------------------*/
+
+bool FDecorationOutlinerModel::CanDelete(FText* OutErrorMessage) const
+{
+	// This outliner only exists if there are decorations to display, so assume that this outliner can always be deleted if it exists.
+	return true;
+}
+
+void FDecorationOutlinerModel::Delete()
+{
+	UMovieSceneDecorationContainerObject* DecorationContainer = WeakDecorationContainer.Get();
+	UMovieSceneTrack* ParentTrack = WeakParentTrack.Get();
+
+	if (!DecorationContainer || !ParentTrack)
+	{
+		return;
+	}
+
+	ParentTrack->Modify();
+	DecorationContainer->Modify();
+
+	TArray<UClass*> DecorationsToRemove;
+
+	for (const UObject* Decoration : DecorationContainer->GetDecorations())
+	{
+		DecorationsToRemove.Add(Decoration->GetClass());
+	}
+
+	for (UClass* DecorationClass : DecorationsToRemove)
+	{
+		DecorationContainer->RemoveDecoration(DecorationClass);
+	}
+
+	// Notify parent model that this outliner has been deleted.
+	OnOutlinerDeletedDelegate.ExecuteIfBound();
+}
+
+/*-----------------------------------------------------------------------------
+	ILockableExtension
+-----------------------------------------------------------------------------*/
+
+ELockableLockState FDecorationOutlinerModel::GetLockState() const
+{
+	int32 NumSections = 0;
+	int32 NumLockedSections = 0;
+	
+	for (const TViewModelPtr<FSectionModel>& Section : DecoratorList.Iterate<FSectionModel>())
+	{
+		++NumSections;
+
+		const UMovieSceneSection* SectionObject = Section->GetSection();
+		if (SectionObject && SectionObject->IsLocked())
+		{
+			++NumLockedSections;
+		}
+	}
+
+	if (NumSections == 0 || NumLockedSections == 0)
+	{
+		return ELockableLockState::None;
+	}
+
+	return NumLockedSections == NumSections ? ELockableLockState::Locked : ELockableLockState::PartiallyLocked;
+}
+
+void FDecorationOutlinerModel::SetIsLocked(bool bIsLocked)
+{
+	for (const TViewModelPtr<FSectionModel>& Section : DecoratorList.Iterate<FSectionModel>())
+	{
+		UMovieSceneSection* SectionObject = Section->GetSection();
+		if (SectionObject && SectionObject->IsLocked() != bIsLocked)
+		{
+			SectionObject->Modify();
+			SectionObject->SetIsLocked(bIsLocked);
+		}
+	}
+}
+
+TViewModelPtr<ITrackRowExtension> FDecorationOutlinerModel::GetParentTrackRow() const
+{
+	const UMovieSceneTrack* Track = WeakParentTrack.Get();
+	if (!Track)
+	{
+		return nullptr;
+	}
+
+	const TViewModelPtr<FViewModel> ParentModel = FindAncestorOfType<FViewModel>();
+	if (!ParentModel.IsValid())
+	{
+		return nullptr;
+	}
+	
+	return ParentModel.ImplicitCast();
+}
+}
+
+#undef LOCTEXT_NAMESPACE
\ No newline at end of file
diff --git a/Engine/Source/Editor/Sequencer/Public/MVVM/ViewModels/DecorationOutlinerModel.h b/Engine/Source/Editor/Sequencer/Public/MVVM/ViewModels/DecorationOutlinerModel.h
new file mode 100644
--- /dev/null
+++ b/Engine/Source/Editor/Sequencer/Public/MVVM/ViewModels/DecorationOutlinerModel.h
@@ -0,0 +1,91 @@
+﻿// Copyright Epic Games, Inc. All Rights Reserved.
+
+#pragma once
+#include "ISequencerTrackEditor.h"
+#include "OutlinerItemModel.h"
+#include "MVVM/Extensions/IConditionableExtension.h"
+#include "MVVM/Extensions/IDeletableExtension.h"
+#include "MVVM/Extensions/ILockableExtension.h"
+#include "MVVM/Extensions/ITopLevelChannelHolderExtension.h"
+#include "MVVM/Extensions/ITrackAreaExtension.h"
+
+class UMovieSceneTrackRowDecoration;
+
+#define UE_API SEQUENCER_API
+
+namespace UE::Sequencer
+{
+DECLARE_DELEGATE(FOnOutlinerDeleted)
+
+class FDecorationOutlinerModel
+	: public FOutlinerItemModel
+	, public ITrackAreaExtension
+	, public ILockableExtension
+	, public IDeletableExtension
+	, public ITopLevelChannelHolderExtension
+
+{
+public:
+	UE_SEQUENCER_DECLARE_CASTABLE_API(SEQUENCER_API, FDecorationOutlinerModel
+		, FOutlinerItemModel
+		, ITrackAreaExtension
+		, ILockableExtension
+		, IDeletableExtension
+		, ITopLevelChannelHolderExtension);
+
+	UE_API explicit FDecorationOutlinerModel(UMovieSceneTrack* InParentTrack, UMovieSceneDecorationContainerObject* InDecorationContainer);
+	UE_API virtual ~FDecorationOutlinerModel() override;
+
+	/* FViewModel Interface */
+	UE_API virtual void OnConstruct() override;
+
+	/* FOutlinerItemModel */
+	UE_API virtual void BuildContextMenu(FMenuBuilder& MenuBuilder) override;
+	UE_API virtual void BuildSidebarMenu(FMenuBuilder& MenuBuilder) override;
+
+	/* IOutlinerExtension */
+	UE_API virtual FOutlinerSizing GetOutlinerSizing() const override;
+	UE_API virtual TSharedPtr<SWidget> CreateOutlinerViewForColumn(const FCreateOutlinerViewParams& InParams, const FName& InColumnName) override;
+	UE_API virtual FText GetLabel() const override;
+
+	/* ITopLevelChannelHolderExtension */
+	UE_API virtual FViewModelChildren GetTopLevelChannels() override;
+
+	/* ITrackAreaExtension */
+	UE_API virtual FTrackAreaParameters GetTrackAreaParameters() const override;
+	UE_API virtual FViewModelVariantIterator GetTrackAreaModelList() const override;
+	UE_API virtual FViewModelVariantIterator GetTopLevelChildTrackAreaModels() const override;
+
+	/* IDeletableExtension */
+	UE_API virtual bool CanDelete(FText* OutErrorMessage) const override;
+	UE_API virtual void Delete() override;
+
+	/* ILockableExtension */
+	UE_API virtual ELockableLockState GetLockState() const override;
+	UE_API virtual void SetIsLocked(bool bIsLocked) override;
+
+	UE_API TWeakObjectPtr<UMovieSceneDecorationContainerObject> GetDecorationContainer() const { return WeakDecorationContainer; }
+
+	UE_API FOnOutlinerDeleted& OnOutlinerDeleted() { return OnOutlinerDeletedDelegate; };
+
+private:
+
+	TViewModelPtr<class ITrackRowExtension> GetParentTrackRow() const;
+	
+	FViewModelListHead DecoratorList{EViewModelListType::Outliner};
+
+	TWeakObjectPtr<UMovieSceneTrack> WeakParentTrack;
+	
+	TWeakObjectPtr<UMovieSceneDecorationContainerObject> WeakDecorationContainer;
+
+	/** Track editor for the parent track */
+	TSharedPtr<ISequencerTrackEditor> TrackEditor;
+
+	/** Called to let the owning model know this outliner has been deleted. */
+	FOnOutlinerDeleted OnOutlinerDeletedDelegate;
+	
+};
+#undef UE_API
+}
+
+
diff --git a/Engine/Source/Runtime/MovieScene/Public/Decorations/IMovieSceneSectionProviderDecoration.h b/Engine/Source/Runtime/MovieScene/Public/Decorations/IMovieSceneSectionProviderDecoration.h
new file mode 100644
--- /dev/null
+++ b/Engine/Source/Runtime/MovieScene/Public/Decorations/IMovieSceneSectionProviderDecoration.h
@@ -0,0 +1,28 @@
+﻿// Copyright Epic Games, Inc. All Rights Reserved.
+
+#pragma once
+
+#include "CoreTypes.h"
+#include "UObject/Interface.h"
+#include "IMovieSceneSectionProviderDecoration.generated.h"
+
+class UMovieSceneSection;
+
+UINTERFACE(MinimalAPI)
+class UMovieSceneSectionProviderDecoration : public UInterface
+{
+public:
+	GENERATED_BODY()
+};
+
+/**
+ *  Optional interface for movie scene decorations that contain their own sections.
+ */
+class IMovieSceneSectionProviderDecoration
+{
+public:
+	GENERATED_BODY()
+
+	virtual TArray<TObjectPtr<UMovieSceneSection>>& GetSections() = 0;
+	
+};
