Sequencer: Fix crash when using command CTRL+F (QuickTreeSearch). This addresses a regression added by a change made earlier today.

#Synced-CL 50196422

diff --git a/Engine/Source/Editor/Sequencer/Private/Filters/Widgets/SSearchAndFilterWidget.cpp b/Engine/Source/Editor/Sequencer/Private/Filters/Widgets/SSearchAndFilterWidget.cpp
--- a/Engine/Source/Editor/Sequencer/Private/Filters/Widgets/SSearchAndFilterWidget.cpp
+++ b/Engine/Source/Editor/Sequencer/Private/Filters/Widgets/SSearchAndFilterWidget.cpp
@@ -43,7 +43,7 @@ void SSearchAndFilterWidget::Construct(
 		.VAlign(VAlign_Center)
 		.Padding(2.f, 0.f, 0.f, 0.f)
 		[
-			SNew(SSequencerSearchBox, InFilterBar)
+			SAssignNew(SearchBox, SSequencerSearchBox, InFilterBar)
 			.AddMetaData<FTagMetaData>(FTagMetaData(TEXT("SequencerFilterSearch")))
 			.HintText(LOCTEXT("FilterSearch", "Search..."))
 			.ToolTipText(LOCTEXT("FilterSearchHint", "Type here to search"))
diff --git a/Engine/Source/Editor/Sequencer/Private/Filters/Widgets/SSearchAndFilterWidget.h b/Engine/Source/Editor/Sequencer/Private/Filters/Widgets/SSearchAndFilterWidget.h
--- a/Engine/Source/Editor/Sequencer/Private/Filters/Widgets/SSearchAndFilterWidget.h
+++ b/Engine/Source/Editor/Sequencer/Private/Filters/Widgets/SSearchAndFilterWidget.h
@@ -34,12 +34,18 @@ public:
 		const TSharedRef<FSequencerFilterBar>& InFilterBar,
 		const TSharedRef<FHideIsolateShowViewModel>& InIsolateHideViewModel
 		);
+	
+	/** @return The search box */
+	TSharedPtr<SSequencerSearchBox> GetSearchBox() const { return SearchBox; }
 
 private:
 	
 	/** The filter bar that is being displayed. */
 	TSharedPtr<ISequencerTrackFilters> FilterBar;
 	
+	/** Search box with which the user affects the text filter. */
+	TSharedPtr<SSequencerSearchBox> SearchBox;
+	
 	/** Invoked when the search text changes. */
 	FOnTextChanged OnSearchTextChanged;
 	
diff --git a/Engine/Source/Editor/Sequencer/Private/SSequencer.cpp b/Engine/Source/Editor/Sequencer/Private/SSequencer.cpp
--- a/Engine/Source/Editor/Sequencer/Private/SSequencer.cpp
+++ b/Engine/Source/Editor/Sequencer/Private/SSequencer.cpp
@@ -1129,7 +1129,7 @@ void SSequencer::BindCommands(TSharedRef<FUICommandList> SequencerCommandBinding
 	// Allow jumping to the Sequencer tree search if you have Sequencer focused
 	SequencerCommandBindings->MapAction(
 		FSequencerCommands::Get().QuickTreeSearch,
-		FExecuteAction::CreateLambda([this] { FSlateApplication::Get().SetKeyboardFocus(SearchBox, EFocusCause::SetDirectly); })
+		FExecuteAction::CreateLambda([this] { FSlateApplication::Get().SetKeyboardFocus(FilterRow->GetActiveSearchBox(), EFocusCause::SetDirectly); })
 	);
 	
 	SequencerCommandBindings->MapAction(
diff --git a/Engine/Source/Editor/Sequencer/Private/SSequencer.h b/Engine/Source/Editor/Sequencer/Private/SSequencer.h
--- a/Engine/Source/Editor/Sequencer/Private/SSequencer.h
+++ b/Engine/Source/Editor/Sequencer/Private/SSequencer.h
@@ -752,9 +752,6 @@ private:
 	/** The breadcrumb trail widget for this sequencer */
 	TSharedPtr<SBreadcrumbTrail<FSequencerBreadcrumb>> BreadcrumbTrail;
 
-	/** The search box for filtering tracks. */
-	TSharedPtr<SSequencerSearchBox> SearchBox;
-
 	/** The current playback time display. */
 	TSharedPtr<STemporarilyFocusedSpinBox<double>> PlayTimeDisplay;
 
diff --git a/Engine/Source/Editor/Sequencer/Private/Widgets/Sequencer/SSequencerSearchAndFilterRow.cpp b/Engine/Source/Editor/Sequencer/Private/Widgets/Sequencer/SSequencerSearchAndFilterRow.cpp
--- a/Engine/Source/Editor/Sequencer/Private/Widgets/Sequencer/SSequencerSearchAndFilterRow.cpp
+++ b/Engine/Source/Editor/Sequencer/Private/Widgets/Sequencer/SSequencerSearchAndFilterRow.cpp
@@ -85,7 +85,7 @@ void SSequencerSearchAndFilterRow::Construct(const FArguments& InArgs)
 		.FillWidth(1.f)
 		.VAlign(VAlign_Center)
 		[
-			SNew(SSearchAndFilterWidget, FilterBar, IsolateViewModel)
+			SAssignNew(FilterRow, SSearchAndFilterWidget, FilterBar, IsolateViewModel)
 			.OnSearchTextChanged(this, &SSequencerSearchAndFilterRow::OnSearchTextChanged)
 		]
 		
@@ -101,6 +101,11 @@ void SSequencerSearchAndFilterRow::Construct(const FArguments& InArgs)
 	];
 }
 
+TSharedPtr<SSequencerSearchBox> SSequencerSearchAndFilterRow::GetActiveSearchBox() const
+{
+	return FilterRow->GetSearchBox();
+}
+
 void SSequencerSearchAndFilterRow::OnSearchTextChanged(const FText& InSearchText)
 {
 	if (const TSharedPtr<FSequencer> Sequencer = WeakSequencer.Pin())
diff --git a/Engine/Source/Editor/Sequencer/Private/Widgets/Sequencer/SSequencerSearchAndFilterRow.h b/Engine/Source/Editor/Sequencer/Private/Widgets/Sequencer/SSequencerSearchAndFilterRow.h
--- a/Engine/Source/Editor/Sequencer/Private/Widgets/Sequencer/SSequencerSearchAndFilterRow.h
+++ b/Engine/Source/Editor/Sequencer/Private/Widgets/Sequencer/SSequencerSearchAndFilterRow.h
@@ -6,6 +6,7 @@
 #include "Widgets/DeclarativeSyntaxSupport.h"
 #include "Widgets/SCompoundWidget.h"
 
+class SSequencerSearchBox;
 class FExtender;
 class FSequencer;
 class FSequencerFilterBar;
@@ -13,6 +14,7 @@ class FSequencerFilterBar;
 namespace UE::Sequencer
 {
 class IOutlinerExtension;
+class SSearchAndFilterWidget;
 
 /** The row of widgets displayed above the Sequence tree view. */
 class SSequencerSearchAndFilterRow : public SCompoundWidget
@@ -31,11 +33,17 @@ public:
 
 	void Construct(const FArguments& InArgs);
 	
+	/** @return The search box that is currently being displayed. */
+	TSharedPtr<SSequencerSearchBox> GetActiveSearchBox() const;
+	
 private:
 	
 	/** Used to update Sequencer's filters. */
 	TWeakPtr<FSequencer> WeakSequencer;
 	
+	/** The widget displaying the search bar, and other filter widgets. */
+	TSharedPtr<SSearchAndFilterWidget> FilterRow;
+	
 	/** Propagates changes in the search widget to the filter system */
 	void OnSearchTextChanged(const FText& InSearchText);
 	
