[Sequencer] Quick fix to disable text expression that enables/disables common filters until it can play nicely with the rest of the text expressions

#Synced-CL 50270751

diff --git a/Engine/Source/Editor/Sequencer/Private/Filters/Filters/SequencerTrackFilter_Text.cpp b/Engine/Source/Editor/Sequencer/Private/Filters/Filters/SequencerTrackFilter_Text.cpp
--- a/Engine/Source/Editor/Sequencer/Private/Filters/Filters/SequencerTrackFilter_Text.cpp
+++ b/Engine/Source/Editor/Sequencer/Private/Filters/Filters/SequencerTrackFilter_Text.cpp
@@ -46,7 +46,7 @@ FSequencerTrackFilter_Text::FSequencerTrackFilter_Text(ISequencerTrackFilters& I
 		});
 
 	// Ordered by importance and most often used. This will dictate the order of display in the text expressions help dialog.
-	TextFilterExpressionContexts.Add(MakeShared<FSequencerTextFilterExpression_Filters>(InFilterInterface));
+	//TextFilterExpressionContexts.Add(MakeShared<FSequencerTextFilterExpression_Filters>(InFilterInterface));
 	TextFilterExpressionContexts.Add(MakeShared<FSequencerTextFilterExpression_Name>(InFilterInterface));
 	TextFilterExpressionContexts.Add(MakeShared<FSequencerTextFilterExpression_TrackClass>(InFilterInterface));
 	TextFilterExpressionContexts.Add(MakeShared<FSequencerTextFilterExpression_ObjectClass>(InFilterInterface));
