[Sequencer] Fix common text expression filter not working correctly with other text expressions
The problem was that the new FSequencerTextFilterExpression_Filters was trying to call DoesCommonFilterPass with the search term.
This works when the search term is the name of a filter, e.g. "Selected". However, it would not work when using e.g. "Locked". DoesCommonFilterPass would return false because "Locked" is not a name of a common filter.

#Synced-CL 50270751

#Synced-CL 50270751

diff --git a/Engine/Plugins/Experimental/SequenceNavigator/Source/SequenceNavigator/Private/Filters/NavigationToolFilterBar.cpp b/Engine/Plugins/Experimental/SequenceNavigator/Source/SequenceNavigator/Private/Filters/NavigationToolFilterBar.cpp
--- a/Engine/Plugins/Experimental/SequenceNavigator/Source/SequenceNavigator/Private/Filters/NavigationToolFilterBar.cpp
+++ b/Engine/Plugins/Experimental/SequenceNavigator/Source/SequenceNavigator/Private/Filters/NavigationToolFilterBar.cpp
@@ -1849,6 +1849,20 @@ void FNavigationToolFilterBar::GetCommonFilterValueSuggestions(const FString& In
 	}
 }
 
+bool FNavigationToolFilterBar::IsCommonFilterName(const FString& InFilterName) const
+{
+	bool bFoundFilter = false;
+
+	CommonFilters->ForEachFilter([&InFilterName, &bFoundFilter](const TSharedRef<FNavigationToolFilter>& InFilter)
+		{
+			bFoundFilter = InFilter->GetName().Equals(InFilterName, ESearchCase::IgnoreCase);
+			const bool bContinue = !bFoundFilter;
+			return bContinue;
+		});
+
+	return bFoundFilter;
+}
+
 } // namespace UE::SequenceNavigator
 
 #undef LOCTEXT_NAMESPACE
diff --git a/Engine/Plugins/Experimental/SequenceNavigator/Source/SequenceNavigator/Public/Filters/NavigationToolFilterBar.h b/Engine/Plugins/Experimental/SequenceNavigator/Source/SequenceNavigator/Public/Filters/NavigationToolFilterBar.h
--- a/Engine/Plugins/Experimental/SequenceNavigator/Source/SequenceNavigator/Public/Filters/NavigationToolFilterBar.h
+++ b/Engine/Plugins/Experimental/SequenceNavigator/Source/SequenceNavigator/Public/Filters/NavigationToolFilterBar.h
@@ -214,6 +214,8 @@ public:
 	virtual void GetCommonFilterKeySuggestions(TArray<FSequencerFilterSuggestion>& OutSuggestions) const override;
 	virtual void GetCommonFilterValueSuggestions(const FString& InKey, TArray<FSequencerFilterSuggestion>& OutSuggestions) const override;
 
+	virtual bool IsCommonFilterName(const FString& InFilterName) const override;
+
 protected:
 	virtual void GetCommonFiltersImpl(TArray<TSharedPtr<void>>& OutTypeErasedFilters) const override;
 
diff --git a/Engine/Source/Editor/Sequencer/Private/Filters/Filters/SequencerTrackFilter_Text.cpp b/Engine/Source/Editor/Sequencer/Private/Filters/Filters/SequencerTrackFilter_Text.cpp
--- a/Engine/Source/Editor/Sequencer/Private/Filters/Filters/SequencerTrackFilter_Text.cpp
+++ b/Engine/Source/Editor/Sequencer/Private/Filters/Filters/SequencerTrackFilter_Text.cpp
@@ -46,7 +46,7 @@ FSequencerTrackFilter_Text::FSequencerTrackFilter_Text(ISequencerTrackFilters& I
 		});
 
 	// Ordered by importance and most often used. This will dictate the order of display in the text expressions help dialog.
-	//TextFilterExpressionContexts.Add(MakeShared<FSequencerTextFilterExpression_Filters>(InFilterInterface));
+	TextFilterExpressionContexts.Add(MakeShared<FSequencerTextFilterExpression_Filters>(InFilterInterface));
 	TextFilterExpressionContexts.Add(MakeShared<FSequencerTextFilterExpression_Name>(InFilterInterface));
 	TextFilterExpressionContexts.Add(MakeShared<FSequencerTextFilterExpression_TrackClass>(InFilterInterface));
 	TextFilterExpressionContexts.Add(MakeShared<FSequencerTextFilterExpression_ObjectClass>(InFilterInterface));
diff --git a/Engine/Source/Editor/Sequencer/Private/Filters/SequencerFilterBar.cpp b/Engine/Source/Editor/Sequencer/Private/Filters/SequencerFilterBar.cpp
--- a/Engine/Source/Editor/Sequencer/Private/Filters/SequencerFilterBar.cpp
+++ b/Engine/Source/Editor/Sequencer/Private/Filters/SequencerFilterBar.cpp
@@ -1882,4 +1882,18 @@ void FSequencerFilterBar::GetCommonFilterValueSuggestions(const FString& InKey,
 	}
 }
 
+bool FSequencerFilterBar::IsCommonFilterName(const FString& InFilterName) const
+{
+	bool bFoundFilter = false;
+
+	CommonFilters->ForEachFilter([&InFilterName, &bFoundFilter](const TSharedRef<FSequencerTrackFilter>& InFilter)
+		{
+			bFoundFilter = InFilter->GetName().Equals(InFilterName, ESearchCase::IgnoreCase);
+			const bool bContinue = !bFoundFilter;
+			return bContinue;
+		});
+
+	return bFoundFilter;
+}
+
 #undef LOCTEXT_NAMESPACE
diff --git a/Engine/Source/Editor/Sequencer/Private/Filters/SequencerFilterBar.h b/Engine/Source/Editor/Sequencer/Private/Filters/SequencerFilterBar.h
--- a/Engine/Source/Editor/Sequencer/Private/Filters/SequencerFilterBar.h
+++ b/Engine/Source/Editor/Sequencer/Private/Filters/SequencerFilterBar.h
@@ -249,6 +249,8 @@ public:
 	virtual void GetCommonFilterKeySuggestions(TArray<FSequencerFilterSuggestion>& OutSuggestions) const override;
 	virtual void GetCommonFilterValueSuggestions(const FString& InKey, TArray<FSequencerFilterSuggestion>& OutSuggestions) const override;
 
+	virtual bool IsCommonFilterName(const FString& InFilterName) const override;
+
 protected:
 	virtual void GetCommonFiltersImpl(TArray<TSharedPtr<void>>& OutTypeErasedFilters) const override;
 
diff --git a/Engine/Source/Editor/Sequencer/Private/Filters/TextExpressions/SequencerTextFilterExpression_Filters.cpp b/Engine/Source/Editor/Sequencer/Private/Filters/TextExpressions/SequencerTextFilterExpression_Filters.cpp
--- a/Engine/Source/Editor/Sequencer/Private/Filters/TextExpressions/SequencerTextFilterExpression_Filters.cpp
+++ b/Engine/Source/Editor/Sequencer/Private/Filters/TextExpressions/SequencerTextFilterExpression_Filters.cpp
@@ -43,6 +43,15 @@ bool FSequencerTextFilterExpression_Filters::TestComplexExpression(const FName&
 		return true;
 	}
 
+	// This expression tries to expose the names of common filters as expression.
+	// If the expression is an actual expression, like "Lock", "Locked", etc., then ignore it. We'll only handle actual common filter names.
+	const FString FilterName = InKey.ToString();
+	const bool bIsCommonFilter = FilterInterface.IsCommonFilterName(FilterName);
+	if (!bIsCommonFilter)
+	{
+		return true;
+	}
+
 	const bool bPassesFilter = FilterInterface.DoesCommonFilterPass(FilterItem.ImplicitCast(), InKey.ToString());
 	return CompareFStringForExactBool(InValue, InComparisonOperation, bPassesFilter);
 }
diff --git a/Engine/Source/Editor/Sequencer/Public/Filters/ISequencerFilterBar.h b/Engine/Source/Editor/Sequencer/Public/Filters/ISequencerFilterBar.h
--- a/Engine/Source/Editor/Sequencer/Public/Filters/ISequencerFilterBar.h
+++ b/Engine/Source/Editor/Sequencer/Public/Filters/ISequencerFilterBar.h
@@ -135,6 +135,9 @@ public:
 	 */
 	virtual void GetCommonFilterValueSuggestions(const FString& InFilterName, TArray<FSequencerFilterSuggestion>& OutSuggestions) const = 0;
 
+	/** @return Whether there is a common filter with InFilterName. */
+	virtual bool IsCommonFilterName(const FString& InFilterName) const = 0;
+
 protected:
 	/** Internal implementation to bypass template limitations in virtual functions */
 	virtual void GetCommonFiltersImpl(TArray<TSharedPtr<void>>& OutTypeErasedFilters) const = 0;
