Sequencer: Fix issue with baking code (and other code) that uses FindSpawnableObjectTemplate in the case that the guid passed in is a possessable that is possessing a spawnable in another sequence. We now check for that.

#Synced-CL 50251927

diff --git a/Engine/Source/Editor/Sequencer/Private/Sequencer.cpp b/Engine/Source/Editor/Sequencer/Private/Sequencer.cpp
--- a/Engine/Source/Editor/Sequencer/Private/Sequencer.cpp
+++ b/Engine/Source/Editor/Sequencer/Private/Sequencer.cpp
@@ -5085,6 +5085,25 @@ UObject* FSequencer::FindSpawnedObjectOrTemplate(const FGuid& BindingId)
 			}
 		}
 	}
+	// If we're a possessable that references a spawnable in another sequence via SpawnableObjectBindingID,
+	// resolve the spawnable's template from that sequence (needed for operations like Bake Animation Sequence
+	// when evaluating in isolation mode where the referenced spawnable isn't actually spawned)
+	else if (Possessable && Possessable->GetSpawnableObjectBindingID().IsValid())
+	{
+		const FMovieSceneObjectBindingID& SpawnableBindingID = Possessable->GetSpawnableObjectBindingID();
+		FMovieSceneSequenceID ResolvedSequenceID = SpawnableBindingID.ResolveSequenceID(GetFocusedTemplateID(), GetSharedPlaybackState());
+		PRAGMA_DISABLE_DEPRECATION_WARNINGS
+		UMovieSceneSequence* ResolvedSequence = State.FindSequence(ResolvedSequenceID);			
+		PRAGMA_ENABLE_DEPRECATION_WARNINGS
+
+		if (ResolvedSequence)
+		{
+			if (UObject* Template = MovieSceneHelpers::GetObjectTemplate(ResolvedSequence, SpawnableBindingID.GetGuid(), GetSharedPlaybackState()))
+			{
+				return Template;
+			}
+		}
+	}
 	// If we're a spawnable and we don't have the object, use the default object to build up the track menu
 	else if (UObject* Template = MovieSceneHelpers::GetObjectTemplate(Sequence, BindingId, GetSharedPlaybackState()))
 	{
