[Sequencer] Add back removed filter bar delegates that were accidently removed in merge

#Synced-CL 50270751

diff --git a/Engine/Source/Editor/Sequencer/Private/SSequencer.cpp b/Engine/Source/Editor/Sequencer/Private/SSequencer.cpp
--- a/Engine/Source/Editor/Sequencer/Private/SSequencer.cpp
+++ b/Engine/Source/Editor/Sequencer/Private/SSequencer.cpp
@@ -357,22 +357,17 @@ void SSequencer::Construct(const FArguments& InArgs, TSharedRef<FSequencer> InSe
 		SequencerWidgets.CreateTimeRangeSlider(TimeSliderControllerRef)
 	);
 
-	OutlinerAndTrackAreaContent = ConstructOutlinerAndTrackArea(InArgs, TimeSliderArgs);
+	FilterRow = SNew(SSequencerSearchAndFilterRow)
+		.ExtenderForAddMenu_Lambda([this]{ return FExtender::Combine(AddMenuExtenders); })
+		.Sequencer(SequencerPtr);
 
-	static const FName TransportControlsLeft("Sequencer.TransportControls.Left");
-	static const FName TransportControlsRight("Sequencer.TransportControls.Right");
-	if (!UToolMenus::Get()->IsMenuRegistered(TransportControlsLeft))
+	if (const TSharedPtr<FSequencerFilterBar> FilterBar = GetFilterBar())
 	{
-		UToolMenus::Get()->RegisterMenu(TransportControlsLeft, NAME_None, EMultiBoxType::ToolBar);
-		UToolMenus::Get()->RegisterMenu(TransportControlsRight, NAME_None, EMultiBoxType::ToolBar);
+		FilterBar->OnStateChanged().AddSP(this, &SSequencer::OnFilterBarStateChanged);
+		FilterBar->OnFiltersChanged().AddSP(this, &SSequencer::OnTrackFiltersChanged);
 	}
 
-	USequencerToolMenuContext* ContextObject = NewObject<USequencerToolMenuContext>();
-	ContextObject->WeakSequencer = SequencerRef;
-
-	FilterRow = SNew(SSequencerSearchAndFilterRow)
-		.ExtenderForAddMenu_Lambda([this]{ return FExtender::Combine(AddMenuExtenders); })
-		.Sequencer(SequencerPtr);
+	OutlinerAndTrackAreaContent = ConstructOutlinerAndTrackArea(InArgs, TimeSliderArgs);
 
 	ChildSlot
 	[
