[Sequencer] Fix search box caret location always reverted to the end after text change

#Synced-CL 50270751

diff --git a/Engine/Source/Editor/Sequencer/Private/Filters/Widgets/SSequencerSearchBox.cpp b/Engine/Source/Editor/Sequencer/Private/Filters/Widgets/SSequencerSearchBox.cpp
--- a/Engine/Source/Editor/Sequencer/Private/Filters/Widgets/SSequencerSearchBox.cpp
+++ b/Engine/Source/Editor/Sequencer/Private/Filters/Widgets/SSequencerSearchBox.cpp
@@ -665,9 +665,19 @@ end:
 
 void SSequencerSearchBox::OnFilterTextChanged()
 {
-	if (const TSharedPtr<ISequencerFilterBar> FilterBarPin = WeakFilterBar.Pin())
+	const TSharedPtr<ISequencerFilterBar> FilterBar = WeakFilterBar.Pin();
+	if (!FilterBar)
+	{
+		return;
+	}
+
+	// If the filter was changed externally, keep the widget in sync.
+	// Skip SetText if the change was instantiated by this widget: when focused and calling SetText it can cause the cursor location to change.
+	const FString FilterString = FilterBar->GetTextFilterString();
+	const bool bHasChanged = GetText().ToString() != FilterString;
+	if (bHasChanged)
 	{
-		const FText NewText = FText::FromString(FilterBarPin->GetTextFilterString());
+		const FText NewText = FText::FromString(FilterString);
 		SetText(NewText);
 	}
 }
