[Sequencer]
- Move SSequencerSectionOverlay widget below track area in z order
- Fixes vertical tick lines being drawn over track sections

#Synced-CL 50319781

diff --git a/Engine/Source/Editor/Sequencer/Private/SSequencer.cpp b/Engine/Source/Editor/Sequencer/Private/SSequencer.cpp
--- a/Engine/Source/Editor/Sequencer/Private/SSequencer.cpp
+++ b/Engine/Source/Editor/Sequencer/Private/SSequencer.cpp
@@ -759,9 +759,21 @@ TSharedRef<SSplitter> SSequencer::ConstructOutlinerAndTrackArea(const FArguments
 					SNew(SOverlay)
 
 					+ SOverlay::Slot()
+					.ZOrder(0)
+					[
+						SNew(SSequencerSectionOverlay, TimeSliderControllerRef)
+						.Visibility(this, &SSequencer::GetShowTickLines)
+						.DisplayScrubPosition(false)
+						.DisplayTickLines(true)
+						.Clipping(EWidgetClipping::ClipToBounds)
+					]
+
+					+ SOverlay::Slot()
+					.ZOrder(10)
 					[
 						SNew(SVerticalBox)
 
+						// Pinned Track Area
 						+ SVerticalBox::Slot()
 						.AutoHeight()
 						.MaxHeight(TAttribute<float>::Create(TAttribute<float>::FGetter::CreateSP(this, &SSequencer::GetPinnedAreaMaxHeight)))
@@ -798,6 +810,7 @@ TSharedRef<SSplitter> SSequencer::ConstructOutlinerAndTrackArea(const FArguments
 							]
 						]
 
+						// Track Area
 						+ SVerticalBox::Slot()
 						[
 							SNew(SOverlay)
@@ -836,16 +849,6 @@ TSharedRef<SSplitter> SSequencer::ConstructOutlinerAndTrackArea(const FArguments
 						]
 					]
 
-					+ SOverlay::Slot()
-					.ZOrder(10)
-					[
-						SNew(SSequencerSectionOverlay, TimeSliderControllerRef)
-						.Visibility(this, &SSequencer::GetShowTickLines)
-						.DisplayScrubPosition(false)
-						.DisplayTickLines(true)
-						.Clipping(EWidgetClipping::ClipToBounds)
-					]
-					
 					// Overlay that allows drawing over the (unpinned) track area
 					+ SOverlay::Slot()
 					.ZOrder(15)
