Sequencer: Attempt fix at shutdown. There is no repro so this is just a reasonable attempt.

RequestSyncSelection enqueues a delegate to execute next tick. The issue is if the engine is requested to exit before that delegate executes.
Then the destructor of FTimerManager is run too late: ~FScopedTransaction will dereference GEditor and cause a crash on exit.

#Synced-CL 50196422

diff --git a/Engine/Source/Editor/Sequencer/Private/MVVM/CurveEditorExtension.cpp b/Engine/Source/Editor/Sequencer/Private/MVVM/CurveEditorExtension.cpp
--- a/Engine/Source/Editor/Sequencer/Private/MVVM/CurveEditorExtension.cpp
+++ b/Engine/Source/Editor/Sequencer/Private/MVVM/CurveEditorExtension.cpp
@@ -304,6 +304,13 @@ FCurveEditorExtension::FCurveEditorExtension()
 	
 }
 
+FCurveEditorExtension::~FCurveEditorExtension()
+{
+	// Prevent crash on shutdown. RequestSyncSelection sets a timer that keeps a FScopedTransaction. 
+	// ~FScopedTransaction expects GEditor to be valid, but it is null at shutdown.
+	CleanUpPendingSyncSelection();
+}
+
 void FCurveEditorExtension::OnCreated(TSharedRef<FViewModel> InWeakOwner)
 {
 	ensureMsgf(!WeakOwnerModel.Pin().IsValid(), TEXT("This extension was already created!"));
@@ -679,7 +686,7 @@ void FCurveEditorExtension::RequestSyncSelection()
 	TSharedPtr<CurveEditor::FScopedSelectionChange> SelectionChange = bShouldRecord
 		? MakeShared<CurveEditor::FScopedSelectionChange>(CurveEditorModel).ToSharedPtr() : nullptr;
 	
-	GEditor->GetTimerManager()->SetTimerForNextTick(
+	SyncSelectionHandle = GEditor->GetTimerManager()->SetTimerForNextTick(
 	[WeakRootViewModel, Transaction = MoveTemp(Transaction), SelectionChange = MoveTemp(SelectionChange)]() mutable
 	{
 		bSyncSelectionRequested = false;
@@ -698,6 +705,7 @@ void FCurveEditorExtension::RequestSyncSelection()
 		FCurveEditorExtension* This = RootViewModel->CastDynamic<FCurveEditorExtension>();
 		if (This)
 		{
+			This->SyncSelectionHandle.Invalidate();
 			This->SyncSelection();
 		}
 
@@ -829,6 +837,16 @@ void FCurveEditorExtension::SyncSelection()
 	CurveEditorModel->ResumeBroadcast();
 }
 
+void FCurveEditorExtension::CleanUpPendingSyncSelection()
+{
+	bSyncSelectionRequested = false;
+	
+	if (GEditor)
+	{
+		GEditor->GetTimerManager()->ClearTimer(SyncSelectionHandle);
+	}
+}
+
 } // namespace Sequencer
 } // namespace UE
 
diff --git a/Engine/Source/Editor/Sequencer/Public/MVVM/CurveEditorExtension.h b/Engine/Source/Editor/Sequencer/Public/MVVM/CurveEditorExtension.h
--- a/Engine/Source/Editor/Sequencer/Public/MVVM/CurveEditorExtension.h
+++ b/Engine/Source/Editor/Sequencer/Public/MVVM/CurveEditorExtension.h
@@ -10,6 +10,7 @@
 #include "CurveEditorTypes.h"
 #include "Tree/ICurveEditorTreeItem.h"
 #include "STemporarilyFocusedSpinBox.h"
+#include "Engine/TimerHandle.h"
 
 #define UE_API SEQUENCER_API
 
@@ -33,7 +34,8 @@ public:
 
 	UE_SEQUENCER_DECLARE_VIEW_MODEL_TYPE_ID_API(UE_API, FCurveEditorExtension)
 
-	UE_API FCurveEditorExtension();
+	FCurveEditorExtension();
+	~FCurveEditorExtension();
 
 	UE_API virtual void OnCreated(TSharedRef<FViewModel> InWeakOwner) override;
 
@@ -73,6 +75,9 @@ private:
 
 	/** Executed when the filter class has changed */
 	UE_API void FilterClassChanged();
+	
+	/** Cancels any pending transactions. */
+	void CleanUpPendingSyncSelection();
 
 private:
 
@@ -94,6 +99,9 @@ private:
 	/** The current playback time display. */
 	TSharedPtr<STemporarilyFocusedSpinBox<double>> PlayTimeDisplay;
 
+	/** Handle to latent refresh of selection. */
+	FTimerHandle SyncSelectionHandle;
+	
 	friend class FCurveEditorIntegrationExtension;
 };
 
