[Sequencer]
- Popout transport controls to their own tab window
- Add transport controls visibility toggle command that controls visibility of the transport controls in the curve editor only
- If the transport controls are popped out to their own tab window, this visibility toggle is ignored and the transport controls are always hidden inline

#Synced-CL 50367394

diff --git a/Engine/Source/Editor/Sequencer/Private/SSequencer.cpp b/Engine/Source/Editor/Sequencer/Private/SSequencer.cpp
--- a/Engine/Source/Editor/Sequencer/Private/SSequencer.cpp
+++ b/Engine/Source/Editor/Sequencer/Private/SSequencer.cpp
@@ -150,7 +150,9 @@
 #include "MVVM/Extensions/ITopTimeSliderOverlayExtension.h"
 #include "MVVM/Extensions/ITrackAreaOverlayExtension.h"
 #include "Settings/EditorExperimentalSettings.h"
+#include "SPopoutTabInlineContent.h"
 #include "Widgets/Docking/SDockTab.h"
+#include "Widgets/Layout/SScaleBox.h"
 #include "Widgets/Sidebar/SequencerSelectionDrawer.h"
 #include "WorkspaceMenuStructure.h"
 #include "WorkspaceMenuStructureModule.h"
@@ -904,11 +906,24 @@ TSharedRef<SSplitter> SSequencer::ConstructOutlinerAndTrackArea(const FArguments
 					.BorderBackgroundColor(FLinearColor(.5f, .5f, .5f, 1.f))
 					.Clipping(EWidgetClipping::ClipToBounds)
 					.Padding(0)
+					
 					[
 						SNew(SOverlay)
 						+ SOverlay::Slot()
 						[
-							BottomTimeSlider.ToSharedRef()
+							SNew(SBox)
+							.Visibility_Lambda([this]()
+								{
+									if (const USequencerSettings* const SequencerSettings = GetSequencerSettings())
+									{
+										return SequencerSettings->GetShowBottomTimeSlider()
+											? EVisibility::Visible : EVisibility::Collapsed;
+									}
+									return EVisibility::Collapsed;
+								})
+							[
+								BottomTimeSlider.ToSharedRef()
+							]
 						]
 						+ SOverlay::Slot()
 						[
@@ -968,6 +983,18 @@ TSharedRef<SWidget> SSequencer::ConstructMainTreeView(const TSharedRef<FSequence
 
 TSharedRef<SWidget> SSequencer::ConstructTransportPanel(const TSharedRef<FSequencer>& InSequencer)
 {
+	const TSharedPtr<FTabManager> TabManager = GetTabManager();
+	if (!ensure(TabManager.IsValid()))
+	{
+		return SNullWidget::NullWidget;
+	}
+
+	const USequencerSettings* const SequencerSettings = GetSequencerSettings();
+	if (!ensure(SequencerSettings))
+	{
+		return SNullWidget::NullWidget;
+	}
+
 	static const FName TransportControlsLeft("Sequencer.TransportControls.Left");
 	static const FName TransportControlsRight("Sequencer.TransportControls.Right");
 	if (!UToolMenus::Get()->IsMenuRegistered(TransportControlsLeft))
@@ -982,6 +1009,15 @@ TSharedRef<SWidget> SSequencer::ConstructTransportPanel(const TSharedRef<FSequen
 	return SNew(SBorder)
 		.BorderImage(FAppStyle::GetBrush(TEXT("ToolPanel.GroupBorder")))
 		.Clipping(EWidgetClipping::ClipToBounds)
+		.Visibility_Lambda([this]()
+			{
+				if (const USequencerSettings* const SequencerSettings = GetSequencerSettings())
+				{
+					return !SequencerSettings->ShouldShowTransportControls() || SequencerSettings->AreTransportControlsPoppedOut()
+						? EVisibility::Collapsed : EVisibility::Visible;
+				}
+				return EVisibility::Collapsed;
+			})
 		[
 			SNew(SHorizontalBox)
 
@@ -1007,18 +1043,55 @@ TSharedRef<SWidget> SSequencer::ConstructTransportPanel(const TSharedRef<FSequen
 			]
 
 			+ SHorizontalBox::Slot()
-			.AutoWidth()
-			.VAlign(VAlign_Center)
-			.HAlign(HAlign_Left)
-			.Padding(FMargin(CommonPadding, 0.f, 0.f, 0.f))
-			[
-				UToolMenus::Get()->GenerateWidget(TransportControlsLeft, FToolMenuContext(InSequencer->GetCommandBindings(), nullptr, ContextObject))
-			]
-
-			+ SHorizontalBox::Slot()
-			.HAlign(HAlign_Center)
+			.HAlign(HAlign_Fill)
 			[
-				SequencerPtr.Pin()->MakeTransportControls(true)
+				SAssignNew(TransportControlsPopoutTab, SPopoutTabInlineContent)
+				.TabManager(TabManager.ToSharedRef())
+				.TabRole(ETabRole::PanelTab)
+				.TabId(TEXT("Sequencer.TransportControls"))
+				.TabDisplayName(LOCTEXT("TransportControlsTabName", "Transport"))
+				.TabIcon(FSlateIcon(FAppStyle::GetAppStyleSetName(), TEXT("Sequencer.Transport.Looping")))
+				.TabGroup(WorkspaceMenu::GetMenuStructure().GetLevelEditorCinematicsCategory())
+				.ShouldAutosize(true)
+				.InitiallyPoppedOut(SequencerSettings->AreTransportControlsPoppedOut())
+				.OnPopoutStateChanged_Lambda([this](const bool bPoppedOut)
+					{
+						if (USequencerSettings* const SequencerSettings = GetSequencerSettings())
+						{
+							SequencerSettings->SetTransportControlsPoppedOut(bPoppedOut);
+						}
+					})
+				[
+					SNew(SScaleBox)
+					.Stretch(EStretch::ScaleToFit)
+					.StretchDirection(EStretchDirection::DownOnly)
+					.HAlign(HAlign_Fill)
+					.VAlign(VAlign_Fill)
+					[
+						SNew(SHorizontalBox)
+						+ SHorizontalBox::Slot()
+						.AutoWidth()
+						.HAlign(HAlign_Left)
+						.VAlign(VAlign_Center)
+						.Padding(FMargin(CommonPadding, 0.f, 0.f, 0.f))
+						[
+							UToolMenus::Get()->GenerateWidget(TransportControlsLeft, FToolMenuContext(InSequencer->GetCommandBindings(), nullptr, ContextObject))
+						]
+						+ SHorizontalBox::Slot()
+						.HAlign(HAlign_Center)
+						[
+							InSequencer->MakeTransportControls(/*bExtended*/true)
+						]
+						+ SHorizontalBox::Slot()
+						.AutoWidth()
+						.HAlign(HAlign_Right)
+						.VAlign(VAlign_Center)
+						.Padding(FMargin(CommonPadding, 0.f, 0.f, 0.f))
+						[
+							UToolMenus::Get()->GenerateWidget(TransportControlsRight, FToolMenuContext(InSequencer->GetCommandBindings(), nullptr, ContextObject))
+						]
+					]
+				]
 			]
 
 			+ SHorizontalBox::Slot()
@@ -1033,15 +1106,6 @@ TSharedRef<SWidget> SSequencer::ConstructTransportPanel(const TSharedRef<FSequen
 				[
 					SNew(SHorizontalBox)
 
-					+ SHorizontalBox::Slot()
-					.AutoWidth()
-					.VAlign(VAlign_Center)
-					.HAlign(HAlign_Right)
-					.Padding(FMargin(CommonPadding, 0.f, 0.f, 0.f))
-					[
-						UToolMenus::Get()->GenerateWidget(TransportControlsRight, FToolMenuContext(InSequencer->GetCommandBindings(), nullptr, ContextObject))
-					]
-
 					+ SHorizontalBox::Slot()
 					.AutoWidth()
 					.VAlign(VAlign_Center)
@@ -1887,7 +1951,10 @@ void SSequencer::PopulateViewOptionMenu(UToolMenu* InToolMenu)
 	FToolMenuSection& SequencerSettingsSection = InToolMenu->AddSection("SequencerSettings", LOCTEXT("SequencerSettings", "Sequencer Settings"));
 	{
 		SequencerSettingsSection.AddMenuEntry( FSequencerCommands::Get().ToggleAutoScroll );
+		SequencerSettingsSection.AddMenuEntry( FSequencerCommands::Get().ToggleShowTransportControl );
+		SequencerSettingsSection.AddMenuEntry( FSequencerCommands::Get().ToggleTransportControlsPoppedOut );
 		SequencerSettingsSection.AddMenuEntry( FSequencerCommands::Get().ToggleShowRangeSlider );
+		SequencerSettingsSection.AddMenuEntry( FSequencerCommands::Get().ToggleShowBottomTimeSlider );
 		SequencerSettingsSection.AddMenuEntry( FSequencerCommands::Get().ToggleLayerBars );
 		SequencerSettingsSection.AddMenuEntry( FSequencerCommands::Get().ToggleKeyBars );
 		SequencerSettingsSection.AddMenuEntry( FSequencerCommands::Get().ToggleChannelColors );
@@ -3361,7 +3428,7 @@ EVisibility SSequencer::GetBottomTimeSliderVisibility() const
 
 EVisibility SSequencer::GetTimeRangeVisibility() const
 {
-	return GetSequencerSettings()->GetShowRangeSlider() ? EVisibility::Visible : EVisibility::Hidden;
+	return GetSequencerSettings()->GetShowRangeSlider() ? EVisibility::Visible : EVisibility::Collapsed;
 }
 
 EVisibility SSequencer::GetInfoButtonVisibility() const
@@ -4340,4 +4407,30 @@ FOptionalSize SSequencer::GetOutlinerWidth() const
 	return FMath::Max<float>(ActualOutlinerWidth, MinOutlinerWidth);
 }
 
+bool SSequencer::IsTransportControlsPoppedOut() const
+{
+	return TransportControlsPopoutTab.IsValid() ? TransportControlsPopoutTab->IsPoppedOut() : false;
+}
+
+void SSequencer::SetTransportControlsPoppedOut(const bool bInPoppedOut)
+{
+	USequencerSettings* const SequencerSettings = GetSequencerSettings();
+	if (!SequencerSettings)
+	{
+		return;
+	}
+
+	SequencerSettings->SetTransportControlsPoppedOut(bInPoppedOut);
+
+	if (TransportControlsPopoutTab.IsValid())
+	{
+		TransportControlsPopoutTab->Popout(bInPoppedOut);
+	}
+}
+
+void SSequencer::ToggleTransportControlsPoppedOut()
+{
+	SetTransportControlsPoppedOut(!IsTransportControlsPoppedOut());
+}
+
 #undef LOCTEXT_NAMESPACE
diff --git a/Engine/Source/Editor/Sequencer/Private/SSequencer.h b/Engine/Source/Editor/Sequencer/Private/SSequencer.h
--- a/Engine/Source/Editor/Sequencer/Private/SSequencer.h
+++ b/Engine/Source/Editor/Sequencer/Private/SSequencer.h
@@ -46,6 +46,7 @@ class SSidebar;
 class SSidebarContainer;
 class SWindow;
 class USequencerSettings;
+class SPopoutTabInlineContent;
 class SSequencerGroupManager;
 class SSequencerSearchBox;
 class SSequencerTreeFilterStatusBar;
@@ -687,6 +688,10 @@ public:
 	EFilterBarLayout GetFilterBarLayout() const;
 	void SetFilterBarLayout(const EFilterBarLayout InLayout);
 
+	bool IsTransportControlsPoppedOut() const;
+	void SetTransportControlsPoppedOut(const bool bInPoppedOut);
+	void ToggleTransportControlsPoppedOut();
+
 private:
 
 	/** Pending focus handler */
@@ -901,4 +906,7 @@ private:
 	TSharedPtr<SVerticalBox> SearchAndFilterRow;
 	/** The widgets displayed above the tree view. */
 	TSharedPtr<UE::Sequencer::SSequencerSearchAndFilterRow> FilterRow;
+
+	/** Widget that manages popped out state (tab window registration) of the transport controls */
+	TSharedPtr<SPopoutTabInlineContent> TransportControlsPopoutTab;
 };
diff --git a/Engine/Source/Editor/Sequencer/Private/Sequencer.cpp b/Engine/Source/Editor/Sequencer/Private/Sequencer.cpp
--- a/Engine/Source/Editor/Sequencer/Private/Sequencer.cpp
+++ b/Engine/Source/Editor/Sequencer/Private/Sequencer.cpp
@@ -11622,12 +11622,30 @@ void FSequencer::BindCommands()
 		} ),
 		FCanExecuteAction::CreateLambda([] { return true; }));
 
+	SequencerCommandBindings->MapAction(
+		Commands.ToggleShowTransportControl,
+		FExecuteAction::CreateLambda( [this]{ Settings->SetShowTransportControls( !Settings->ShouldShowTransportControls() ); } ),
+		FCanExecuteAction(),
+		FIsActionChecked::CreateLambda( [this]{ return Settings->ShouldShowTransportControls(); } ) );
+
+	SequencerCommandBindings->MapAction(
+		Commands.ToggleTransportControlsPoppedOut,
+		FExecuteAction::CreateLambda( [this]{ SequencerWidget->ToggleTransportControlsPoppedOut(); } ),
+		FCanExecuteAction(),
+		FIsActionChecked::CreateLambda( [this]{ return SequencerWidget->IsTransportControlsPoppedOut(); } ) );
+
 	SequencerCommandBindings->MapAction(
 		Commands.ToggleShowRangeSlider,
 		FExecuteAction::CreateLambda( [this]{ Settings->SetShowRangeSlider( !Settings->GetShowRangeSlider() ); } ),
 		FCanExecuteAction::CreateLambda( []{ return true; } ),
 		FIsActionChecked::CreateLambda( [this]{ return Settings->GetShowRangeSlider(); } ) );
 
+	SequencerCommandBindings->MapAction(
+		Commands.ToggleShowBottomTimeSlider,
+		FExecuteAction::CreateLambda( [this]{ Settings->SetShowBottomTimeSlider( !Settings->GetShowBottomTimeSlider() ); } ),
+		FCanExecuteAction::CreateLambda( []{ return true; } ),
+		FIsActionChecked::CreateLambda( [this]{ return Settings->GetShowBottomTimeSlider(); } ) );
+
 	SequencerCommandBindings->MapAction(
 		Commands.ToggleIsSnapEnabled,
 		FExecuteAction::CreateLambda( [this]{ Settings->SetIsSnapEnabled( !Settings->GetIsSnapEnabled() ); } ),
diff --git a/Engine/Source/Editor/Sequencer/Private/SequencerCommands.cpp b/Engine/Source/Editor/Sequencer/Private/SequencerCommands.cpp
--- a/Engine/Source/Editor/Sequencer/Private/SequencerCommands.cpp
+++ b/Engine/Source/Editor/Sequencer/Private/SequencerCommands.cpp
@@ -136,7 +136,10 @@ void FSequencerCommands::RegisterCommands()
 	UI_COMMAND( OpenDirectorBlueprint, "Open Director Blueprint", "Opens the director blueprint for this sequence.", EUserInterfaceActionType::Button, FInputChord() );
 	UI_COMMAND( OpenTaggedBindingManager, "Open Binding Tag Manager", "Specifies options for tagging bindings within this sequence for external systems to reference as a persistent name.", EUserInterfaceActionType::Button, FInputChord() );
 	UI_COMMAND( OpenNodeGroupsManager, "Open Sequencer Group Manager", "Manage groups within this sequence.", EUserInterfaceActionType::Button, FInputChord());
+	UI_COMMAND( ToggleShowTransportControl, "Transport Controls", "Enables and disables showing the transport controls panel", EUserInterfaceActionType::ToggleButton, FInputChord());
+	UI_COMMAND( ToggleTransportControlsPoppedOut, "Popout Transport Controls", "Toggles the popped out state of the transport controls panel", EUserInterfaceActionType::ToggleButton, FInputChord());
 	UI_COMMAND( ToggleShowRangeSlider, "Range Slider", "Enables and disables showing the time range slider", EUserInterfaceActionType::ToggleButton, FInputChord());
+	UI_COMMAND( ToggleShowBottomTimeSlider, "Bottom Time Slider", "Enables and disables showing the bottom time slider", EUserInterfaceActionType::ToggleButton, FInputChord());
 	UI_COMMAND( ToggleShowInfoButton, "Info Button", "Enables and disables showing the information button in the playback controls", EUserInterfaceActionType::ToggleButton, FInputChord());
 	UI_COMMAND( ToggleIsSnapEnabled, "Enable Snapping", "Enables and disables snapping", EUserInterfaceActionType::ToggleButton, FInputChord() );
 
diff --git a/Engine/Source/Editor/Sequencer/Private/SequencerSettings.cpp b/Engine/Source/Editor/Sequencer/Private/SequencerSettings.cpp
--- a/Engine/Source/Editor/Sequencer/Private/SequencerSettings.cpp
+++ b/Engine/Source/Editor/Sequencer/Private/SequencerSettings.cpp
@@ -20,6 +20,7 @@ USequencerSettings::USequencerSettings( const FObjectInitializer& ObjectInitiali
 	SpawnPosition = SSP_Origin;
 	bCreateSpawnableCameras = true;
 	bShowRangeSlider = false;
+	bShowBottomTimeSlider = true;
 	bIsSnapEnabled = true;
 	bSnapKeyTimesToElements = false;
 	bSnapSectionTimesToElements = false;
@@ -85,6 +86,8 @@ USequencerSettings::USequencerSettings( const FObjectInitializer& ObjectInitiali
 	bFilterBarVisible = true;
 	bNavigationToolVisible = false;
 	bSelectionDetailsPoppedOut = false;
+	bShowTransportControls = true;
+	bTransportControlsPoppedOut = false;
 
 	SectionColorTints.Add(FColor(88, 102, 142, 255)); // blue
 	SectionColorTints.Add(FColor(99, 137, 132, 255)); // blue-green
@@ -221,6 +224,20 @@ void USequencerSettings::SetShowRangeSlider(bool InbShowRangeSlider)
 	}
 }
 
+bool USequencerSettings::GetShowBottomTimeSlider() const
+{
+	return bShowBottomTimeSlider;
+}
+
+void USequencerSettings::SetShowBottomTimeSlider(const bool bInShow)
+{
+	if (bShowBottomTimeSlider != bInShow)
+	{
+		bShowBottomTimeSlider = bInShow;
+		SaveConfig();
+	}
+}
+
 bool USequencerSettings::GetIsSnapEnabled() const
 {
 	return bIsSnapEnabled;
@@ -1334,3 +1351,31 @@ void USequencerSettings::SetSelectionDetailsPoppedOut(const bool bInPoppedOut)
 		SaveConfig();
 	}
 }
+
+bool USequencerSettings::ShouldShowTransportControls() const
+{
+	return bShowTransportControls;
+}
+
+void USequencerSettings::SetShowTransportControls(const bool bInShowTransportControls)
+{
+	if (bShowTransportControls != bInShowTransportControls)
+	{
+		bShowTransportControls = bInShowTransportControls;
+		SaveConfig();
+	}
+}
+
+bool USequencerSettings::AreTransportControlsPoppedOut() const
+{
+	return bTransportControlsPoppedOut;
+}
+
+void USequencerSettings::SetTransportControlsPoppedOut(const bool bInPoppedOut)
+{
+	if (bTransportControlsPoppedOut != bInPoppedOut)
+	{
+		bTransportControlsPoppedOut = bInPoppedOut;
+		SaveConfig();
+	}
+}
diff --git a/Engine/Source/Editor/Sequencer/Public/SequencerCommands.h b/Engine/Source/Editor/Sequencer/Public/SequencerCommands.h
--- a/Engine/Source/Editor/Sequencer/Public/SequencerCommands.h
+++ b/Engine/Source/Editor/Sequencer/Public/SequencerCommands.h
@@ -351,9 +351,18 @@ public:
 	/** Remove selected nodes from folder. */
 	TSharedPtr< FUICommandInfo > RemoveFromFolder;
 
+	/** Toggles display of the transport controls panel. */
+	TSharedPtr< FUICommandInfo > ToggleShowTransportControl;
+
+	/** Toggles popped out state of the transport controls panel. */
+	TSharedPtr< FUICommandInfo > ToggleTransportControlsPoppedOut;
+
 	/** Turns the range slider on and off. */
 	TSharedPtr< FUICommandInfo > ToggleShowRangeSlider;
 
+	/** Turns the bottom time slider on and off. */
+	TSharedPtr< FUICommandInfo > ToggleShowBottomTimeSlider;
+
 	/** Turns snapping on and off. */
 	TSharedPtr< FUICommandInfo > ToggleIsSnapEnabled;
 
diff --git a/Engine/Source/Editor/Sequencer/Public/SequencerSettings.h b/Engine/Source/Editor/Sequencer/Public/SequencerSettings.h
--- a/Engine/Source/Editor/Sequencer/Public/SequencerSettings.h
+++ b/Engine/Source/Editor/Sequencer/Public/SequencerSettings.h
@@ -181,9 +181,14 @@ public:
 
 	/** Gets whether or not to show the time range slider. */
 	UE_API bool GetShowRangeSlider() const;
-	/** Sets whether or not to show frame numbers. */
+	/** Sets whether or not to show the time range slider. */
 	UE_API void SetShowRangeSlider(bool InbShowRangeSlider);
 
+	/** Gets whether or not to show the bottom time slider. */
+	UE_API bool GetShowBottomTimeSlider() const;
+	/** Sets whether or not to show the bottom time slider. */
+	UE_API void SetShowBottomTimeSlider(const bool bInbShow);
+
 	/** Gets whether or not snapping is enabled. */
 	UE_API bool GetIsSnapEnabled() const;
 	/** Sets whether or not snapping is enabled. */
@@ -595,6 +600,16 @@ public:
 	/** Sets the popped out state of the selection details in the Sequencer tab */
 	UE_API void SetSelectionDetailsPoppedOut(const bool bInPoppedOut);
 
+	/** @return True if the transport control panel is visible in the Sequencer tab */
+	UE_API bool ShouldShowTransportControls() const;
+	/** Sets the visibility of the transport control panel in the Sequencer tab */
+	UE_API void SetShowTransportControls(const bool bInShow);
+
+	/** @return True if the transport control panel is popped out from the Sequencer tab to its own tab */
+	UE_API bool AreTransportControlsPoppedOut() const;
+	/** Sets the popped out state of the transport controls panel in the Sequencer tab */
+	UE_API void SetTransportControlsPoppedOut(const bool bInPoppedOut);
+
 protected:
 
 	/** The auto change mode (auto-key, auto-track or none). */
@@ -629,6 +644,10 @@ protected:
 	UPROPERTY( config, EditAnywhere, Category=Timeline )
 	bool bShowRangeSlider;
 
+	/** Show the bottom time slider, which is only visible if bShowRangeSlider is false. */
+	UPROPERTY( config, EditAnywhere, Category=Timeline )
+	bool bShowBottomTimeSlider;
+
 	/** Enable or disable snapping in the timeline. */
 	UPROPERTY( config, EditAnywhere, Category=Snapping )
 	bool bIsSnapEnabled;
@@ -954,6 +973,15 @@ protected:
 	UPROPERTY(config, EditAnywhere, Category = General)
 	bool bSelectionDetailsPoppedOut;
 
+	/** Last saved visibility of the transport controls */
+	UPROPERTY(config, EditAnywhere, Category = General)
+	bool bShowTransportControls;
+
+	/** Last saved state of the transport controls: popped out to its own tab (that can be docked/undocked,
+	 * thus the term of "Popout" is used) instead of "Dock" or popped into inline content */
+	UPROPERTY(config, EditAnywhere, Category = General)
+	bool bTransportControlsPoppedOut;
+
 	UFUNCTION()
 	static UE_API bool ShouldShowThumbnailCaptureSettings();
 
