[CurveEditor]
- Add transport controls visibility toggle command that controls visibility of the transport controls in the curve editor only
- If the transport controls are popped out to their own tab window, this visibility toggle is ignored and the transport controls are always hidden inline
- Used in conjuction with change 50096406 for popping out the transport controls in Sequencer

#Synced-CL 50367394

diff --git a/Engine/Source/Editor/CurveEditor/Private/CurveEditor.cpp b/Engine/Source/Editor/CurveEditor/Private/CurveEditor.cpp
--- a/Engine/Source/Editor/CurveEditor/Private/CurveEditor.cpp
+++ b/Engine/Source/Editor/CurveEditor/Private/CurveEditor.cpp
@@ -553,6 +553,12 @@ void FCurveEditor::BindCommands()
 		FIsActionChecked::CreateLambda( [CurveSettings]{ return CurveSettings->GetAutoFrameCurveEditor(); } )
 	);
 
+	CommandList->MapAction(FCurveEditorCommands::Get().ToggleShowTransportControls,
+		FExecuteAction::CreateLambda([this, CurveSettings] { CurveSettings->SetShowTransportControls(!CurveSettings->GetShowTransportControls()); }),
+		FCanExecuteAction(),
+		FIsActionChecked::CreateLambda([CurveSettings] { return CurveSettings->GetShowTransportControls(); })
+	);
+
 	CommandList->MapAction(FCurveEditorCommands::Get().ToggleShowBars,
 		FExecuteAction::CreateLambda([this, CurveSettings] { CurveSettings->SetShowBars(!CurveSettings->GetShowBars()); Tree.RecreateModelsFromExistingSelection(this); }),
 		FCanExecuteAction(),
diff --git a/Engine/Source/Editor/CurveEditor/Private/CurveEditorCommands.cpp b/Engine/Source/Editor/CurveEditor/Private/CurveEditorCommands.cpp
--- a/Engine/Source/Editor/CurveEditor/Private/CurveEditorCommands.cpp
+++ b/Engine/Source/Editor/CurveEditor/Private/CurveEditorCommands.cpp
@@ -70,6 +70,7 @@ void FCurveEditorCommands::RegisterCommands()
 	UI_COMMAND(ToggleSnapTimeToSelection, "Snap Time to Selection", "Snap the current time to the first selected key time.", EUserInterfaceActionType::ToggleButton, FInputChord() );
 	UI_COMMAND(ToggleShowBufferedCurves, "Buffered Curves", "Show buffered curves for the selected curves.", EUserInterfaceActionType::ToggleButton, FInputChord() );
 	UI_COMMAND(ToggleShowCurveEditorCurveToolTips, "Curve Tool Tips", "Show a tool tip with name and values when hovering over a curve.", EUserInterfaceActionType::ToggleButton, FInputChord() );
+	UI_COMMAND(ToggleShowTransportControls, "Show Transport Controls", "Show transport controls (only used if the transport controls are not popped out)", EUserInterfaceActionType::ToggleButton, FInputChord());
 	UI_COMMAND(ToggleShowBars, "Show Bars", "Show Bars like Constraints and Spaces", EUserInterfaceActionType::ToggleButton, FInputChord());
 	UI_COMMAND(ToggleShowValueIndicatorLines, "Show Value Indicators", "Show horizontal, dotted lines on the selected min & max valued key handles", EUserInterfaceActionType::ToggleButton, FInputChord());
 
diff --git a/Engine/Source/Editor/CurveEditor/Private/CurveEditorSettings.cpp b/Engine/Source/Editor/CurveEditor/Private/CurveEditorSettings.cpp
--- a/Engine/Source/Editor/CurveEditor/Private/CurveEditorSettings.cpp
+++ b/Engine/Source/Editor/CurveEditor/Private/CurveEditorSettings.cpp
@@ -10,6 +10,7 @@ UCurveEditorSettings::UCurveEditorSettings()
 {
 	PanningMouseButton = ECurveEditorPanningMouseButton::RightOrAltMiddle;
 	bAutoFrameCurveEditor = true;
+	bShowTransportControls = true;
 	bShowBars = true;
 	FrameInputPadding = 50;
 	FrameOutputPadding = 50;
@@ -91,6 +92,20 @@ void UCurveEditorSettings::SetAutoFrameCurveEditor(bool InbAutoFrameCurveEditor)
 	}
 }
 
+bool UCurveEditorSettings::GetShowTransportControls() const
+{
+	return bShowTransportControls;
+}
+
+void UCurveEditorSettings::SetShowTransportControls(const bool bInShow)
+{
+	if (bShowTransportControls != bInShow)
+	{
+		bShowTransportControls = bInShow;
+		SaveConfig();
+	}
+}
+
 bool UCurveEditorSettings::GetShowBars() const
 {
 	return bShowBars;
diff --git a/Engine/Source/Editor/CurveEditor/Private/SCurveEditorPanel.cpp b/Engine/Source/Editor/CurveEditor/Private/SCurveEditorPanel.cpp
--- a/Engine/Source/Editor/CurveEditor/Private/SCurveEditorPanel.cpp
+++ b/Engine/Source/Editor/CurveEditor/Private/SCurveEditorPanel.cpp
@@ -1197,6 +1197,7 @@ TSharedRef<SWidget> SCurveEditorPanel::MakeCurveEditorCurveViewOptionsMenu(TShar
 	MenuBuilder.AddMenuEntry(FCurveEditorCommands::Get().ToggleAutoFrameCurveEditor);
 	MenuBuilder.AddMenuEntry(FCurveEditorCommands::Get().ToggleSnapTimeToSelection);
 	MenuBuilder.AddMenuEntry(FCurveEditorCommands::Get().ToggleShowBufferedCurves);
+	MenuBuilder.AddMenuEntry(FCurveEditorCommands::Get().ToggleShowTransportControls);
 	MenuBuilder.AddMenuEntry(FCurveEditorCommands::Get().ToggleShowBars);
 	MenuBuilder.AddMenuEntry(FCurveEditorCommands::Get().ToggleShowCurveEditorCurveToolTips);
 	MenuBuilder.AddMenuEntry(FCurveEditorCommands::Get().ToggleShowValueIndicatorLines);
diff --git a/Engine/Source/Editor/CurveEditor/Public/CurveEditorCommands.h b/Engine/Source/Editor/CurveEditor/Public/CurveEditorCommands.h
--- a/Engine/Source/Editor/CurveEditor/Public/CurveEditorCommands.h
+++ b/Engine/Source/Editor/CurveEditor/Public/CurveEditorCommands.h
@@ -85,6 +85,7 @@ public:
 	TSharedPtr<FUICommandInfo> ToggleSnapTimeToSelection;
 	TSharedPtr<FUICommandInfo> ToggleShowBufferedCurves;
 	TSharedPtr<FUICommandInfo> ToggleShowCurveEditorCurveToolTips;
+	TSharedPtr<FUICommandInfo> ToggleShowTransportControls;
 	TSharedPtr<FUICommandInfo> ToggleShowBars;
 	TSharedPtr<FUICommandInfo> ToggleShowValueIndicatorLines;
 
diff --git a/Engine/Source/Editor/CurveEditor/Public/CurveEditorSettings.h b/Engine/Source/Editor/CurveEditor/Public/CurveEditorSettings.h
--- a/Engine/Source/Editor/CurveEditor/Public/CurveEditorSettings.h
+++ b/Engine/Source/Editor/CurveEditor/Public/CurveEditorSettings.h
@@ -124,6 +124,11 @@ public:
 	/** Sets whether or not the curve editor auto frames the selected curves. */
 	UE_API void SetAutoFrameCurveEditor(bool InbAutoFrameCurveEditor);
 
+	/** Gets whether or not the curve editor shows the transport controls panel. */
+	UE_API bool GetShowTransportControls() const;
+	/** Sets whether or not the curve editor shows the transport controls panel. */
+	UE_API void SetShowTransportControls(const bool bInShow);
+
 	/** Gets whether or not the curve editor shows key bar style curves, like for constraints and spaces. */
 	UE_API bool GetShowBars() const;
 	/** Sets whether or not the curve editor shows key bar style curves, like for constraints and spaces. */
@@ -228,6 +233,9 @@ protected:
 	UPROPERTY( config, EditAnywhere, Category="Curve Editor" )
 	bool bAutoFrameCurveEditor;
 
+	UPROPERTY(config, EditAnywhere, Category = "Curve Editor")
+	bool bShowTransportControls;
+
 	UPROPERTY(config, EditAnywhere, Category = "Curve Editor")
 	bool bShowBars;
 
diff --git a/Engine/Source/Editor/Sequencer/Private/MVVM/CurveEditorExtension.cpp b/Engine/Source/Editor/Sequencer/Private/MVVM/CurveEditorExtension.cpp
--- a/Engine/Source/Editor/Sequencer/Private/MVVM/CurveEditorExtension.cpp
+++ b/Engine/Source/Editor/Sequencer/Private/MVVM/CurveEditorExtension.cpp
@@ -18,6 +18,7 @@
 #include "SCurveEditorPanel.h"
 #include "SCurveEditorToolProperties.h"
 #include "SCurveKeyDetailPanel.h"
+#include "SPopoutTabInlineContent.h"
 #include "SSequencerTreeFilterStatusBar.h"
 #include "STemporarilyFocusedSpinBox.h"
 #include "Sequencer.h"
@@ -35,6 +36,8 @@
 #include "Widgets/Docking/SDockTab.h"
 #include "Widgets/Input/SButton.h"
 #include "Widgets/Layout/SScrollBorder.h"
+#include "WorkspaceMenuStructure.h"
+#include "WorkspaceMenuStructureModule.h"
 
 #define LOCTEXT_NAMESPACE "SequencerCurveEditorExtension"
 
@@ -311,6 +314,25 @@ FCurveEditorExtension::~FCurveEditorExtension()
 	CleanUpPendingSyncSelection();
 }
 
+TSharedPtr<ISequencer> FCurveEditorExtension::GetSequencer() const
+{
+	const TSharedPtr<FSequencerEditorViewModel> OwnerModel = WeakOwnerModel.Pin();
+	return OwnerModel.IsValid() ? OwnerModel->GetSequencerImpl() : nullptr;
+}
+
+USequencerSettings* FCurveEditorExtension::GetSequencerSettings() const
+{
+	const TSharedPtr<ISequencer> Sequencer = GetSequencer();
+	return Sequencer.IsValid() ? Sequencer->GetSequencerSettings() : nullptr;
+}
+
+TSharedPtr<FTabManager> FCurveEditorExtension::GetTabManager() const
+{
+	const TSharedPtr<ISequencer> Sequencer = GetSequencer();
+	const TSharedPtr<IToolkitHost> ToolkitHost = Sequencer.IsValid() ? Sequencer->GetToolkitHost() : nullptr;
+	return ToolkitHost.IsValid() ? ToolkitHost->GetTabManager() : nullptr;
+}
+
 void FCurveEditorExtension::OnCreated(TSharedRef<FViewModel> InWeakOwner)
 {
 	ensureMsgf(!WeakOwnerModel.Pin().IsValid(), TEXT("This extension was already created!"));
@@ -319,29 +341,40 @@ void FCurveEditorExtension::OnCreated(TSharedRef<FViewModel> InWeakOwner)
 
 void FCurveEditorExtension::CreateCurveEditor(const FTimeSliderArgs& TimeSliderArgs)
 {
-	TSharedPtr<FSequencerEditorViewModel> OwnerModel = WeakOwnerModel.Pin();
+	const TSharedPtr<FSequencerEditorViewModel> OwnerModel = WeakOwnerModel.Pin();
 	if (!ensure(OwnerModel))
 	{
 		return;
 	}
 
-	TSharedPtr<FSequencer> Sequencer = OwnerModel->GetSequencerImpl();
+	const TSharedPtr<FSequencer> Sequencer = OwnerModel->GetSequencerImpl();
 	if (!ensure(Sequencer))
 	{
 		return;
 	}
 
+	USequencerSettings* const SequencerSettings = Sequencer->GetSequencerSettings();
+	if (!ensure(SequencerSettings))
+	{
+		return;
+	}
+
 	// If they've said they want to support the curve editor then they need to provide a toolkit host
 	// so that we know where to spawn our tab into.
-	if (!ensure(Sequencer->GetToolkitHost().IsValid()))
+	const TSharedPtr<IToolkitHost> ToolkitHost = Sequencer->GetToolkitHost();
+	if (!ensure(ToolkitHost.IsValid()))
+	{
+		return;
+	}
+
+	const TSharedPtr<FTabManager> TabManager = GetTabManager();
+	if (!ensure(TabManager.IsValid()))
 	{
 		return;
 	}
 
 	// Create the curve editor;
 	{
-		USequencerSettings* SequencerSettings = Sequencer->GetSequencerSettings();
-		
 		FCurveEditorInitParams CurveEditorInitParams;
 		CurveEditorInitParams.AdditionalEditorExtensions = { MakeShared<FSequencerCurveEditorToolbarExtender>(Sequencer.ToWeakPtr()) };
 		CurveEditorInitParams.ZoomScalingAttr.BindLambda([SequencerSettings]
@@ -381,7 +414,7 @@ void FCurveEditorExtension::CreateCurveEditor(const FTimeSliderArgs& TimeSliderA
 		.GridLineTint(FLinearColor(0.f, 0.f, 0.f, 0.3f))
 		.ExternalTimeSliderController(CurveEditorTimeSliderController)
 		.MinimumViewPanelHeight(0.f)
-		.TabManager(Sequencer->GetToolkitHost()->GetTabManager())
+		.TabManager(TabManager)
 		.DisabledTimeSnapTooltip(LOCTEXT("CurveEditorTimeSnapDisabledTooltip", "Time Snapping is currently driven by Sequencer."))
 		.TreeContent()
 		[
@@ -430,8 +463,15 @@ void FCurveEditorExtension::CreateCurveEditor(const FTimeSliderArgs& TimeSliderA
 				SNew(SBorder)
 				.BorderImage(FAppStyle::GetBrush("ToolPanel.GroupBorder"))
 				.Clipping(EWidgetClipping::ClipToBounds)
+				.Visibility_Lambda([this]()
+					{
+						const USequencerSettings* const SequencerSettings = GetSequencerSettings();
+						const bool bTransportControlsPoppedOut = SequencerSettings
+							? SequencerSettings->AreTransportControlsPoppedOut() : false;
+						return bTransportControlsPoppedOut || !GetDefault<UCurveEditorSettings>()->GetShowTransportControls()
+							? EVisibility::Collapsed : EVisibility::Visible;
+					})
 				[
-
 					SNew(SHorizontalBox)
 
 					+ SHorizontalBox::Slot()
@@ -455,10 +495,24 @@ void FCurveEditorExtension::CreateCurveEditor(const FTimeSliderArgs& TimeSliderA
 					]
 
 					+ SHorizontalBox::Slot()
+					.HAlign(HAlign_Fill)
 					[
-						SNew(SBorder)
-						.BorderImage(FAppStyle::GetBrush("ToolPanel.GroupBorder"))
-						.HAlign(HAlign_Center)
+						SNew(SPopoutTabInlineContent)
+						.TabManager(TabManager)
+						.TabRole(ETabRole::PanelTab)
+						.TabId(TEXT("Sequencer.TransportControls"))
+						.TabDisplayName(LOCTEXT("TransportControlsTabName", "Transport"))
+						.TabIcon(FSlateIcon(FAppStyle::GetAppStyleSetName(), TEXT("Sequencer.Transport.Looping")))
+						.TabGroup(WorkspaceMenu::GetMenuStructure().GetLevelEditorCinematicsCategory())
+						.ShouldAutosize(true)
+						.InitiallyPoppedOut(SequencerSettings->AreTransportControlsPoppedOut())
+						.OnPopoutStateChanged_Lambda([this](const bool bPoppedOut)
+							{
+								if (USequencerSettings* const SequencerSettings = GetSequencerSettings())
+								{
+									SequencerSettings->SetTransportControlsPoppedOut(bPoppedOut);
+								}
+							})
 						[
 							Sequencer->MakeTransportControls(true)
 						]
@@ -519,9 +573,7 @@ void FCurveEditorExtension::CreateCurveEditor(const FTimeSliderArgs& TimeSliderA
 	CurveEditorPanel->OnFilterClassChanged.BindRaw(this, &FCurveEditorExtension::FilterClassChanged);
 
 	// Check to see if the tab is already opened due to the saved window layout.
-	FTabId TabId = FTabId(FCurveEditorExtension::CurveEditorTabName);
-	TSharedPtr<SDockTab> ExistingCurveEditorTab = Sequencer->GetToolkitHost()->GetTabManager()->FindExistingLiveTab(TabId);
-	if (ExistingCurveEditorTab)
+	if (const TSharedPtr<SDockTab> ExistingCurveEditorTab = TabManager->FindExistingLiveTab(CurveEditorTabName))
 	{
 		ExistingCurveEditorTab->SetContent(CurveEditorWidget.ToSharedRef());
 	}
@@ -556,23 +608,15 @@ void FCurveEditorExtension::FilterClassChanged()
 
 void FCurveEditorExtension::OpenCurveEditor()
 {
-	TSharedPtr<FSequencerEditorViewModel> OwnerModel = WeakOwnerModel.Pin();
-	if (!ensure(OwnerModel))
-	{
-		return;
-	}
-
-	TSharedPtr<ISequencer> Sequencer = OwnerModel->GetSequencer();
-	if (!Sequencer)
+	const TSharedPtr<FTabManager> TabManager = GetTabManager();
+	if (!ensure(TabManager))
 	{
 		return;
 	}
 
 	// Request the Tab Manager invoke the tab. This will spawn the tab if needed, otherwise pull it to focus. This assumes
 	// that the Toolkit Host's Tab Manager has already registered a tab with a NullWidget for content.
-	FTabId TabId = FTabId(FCurveEditorExtension::CurveEditorTabName);
-	TSharedPtr<SDockTab> CurveEditorTab = Sequencer->GetToolkitHost()->GetTabManager()->TryInvokeTab(TabId);
-	if (CurveEditorTab.IsValid())
+	if (const TSharedPtr<SDockTab> CurveEditorTab = TabManager->TryInvokeTab(CurveEditorTabName))
 	{
 		CurveEditorTab->SetContent(CurveEditorWidget.ToSharedRef());
 
@@ -587,51 +631,22 @@ void FCurveEditorExtension::OpenCurveEditor()
 
 bool FCurveEditorExtension::IsCurveEditorOpen() const
 {
-	TSharedPtr<FSequencerEditorViewModel> OwnerModel = WeakOwnerModel.Pin();
-	if (!ensure(OwnerModel))
-	{
-		return false;
-	}
-
-	TSharedPtr<ISequencer> Sequencer = OwnerModel->GetSequencer();
-	if (!Sequencer)
-	{
-		return false;
-	}
-
-	TSharedPtr<IToolkitHost> ToolkitHost = Sequencer->GetToolkitHost();
-	if (!ToolkitHost)
+	if (const TSharedPtr<FTabManager> TabManager = GetTabManager())
 	{
-		return false;
+		return TabManager->FindExistingLiveTab(CurveEditorTabName).IsValid();
 	}
-
-	TSharedPtr<FTabManager> TabManager = ToolkitHost->GetTabManager();
-	if (!TabManager)
-	{
-		return false;
-	}
-
-	FTabId TabId = FTabId(FCurveEditorExtension::CurveEditorTabName);
-	return TabManager->FindExistingLiveTab(TabId).IsValid();
+	return false;
 }
 
 void FCurveEditorExtension::CloseCurveEditor()
 {
-	TSharedPtr<FSequencerEditorViewModel> OwnerModel = WeakOwnerModel.Pin();
-	if (!ensure(OwnerModel))
+	const TSharedPtr<FTabManager> TabManager = GetTabManager();
+	if (!TabManager.IsValid())
 	{
 		return;
 	}
 
-	TSharedPtr<ISequencer> Sequencer = OwnerModel->GetSequencer();
-	if (!Sequencer)
-	{
-		return;
-	}
-
-	FTabId TabId = FTabId(FCurveEditorExtension::CurveEditorTabName);
-	TSharedPtr<SDockTab> CurveEditorTab = Sequencer->GetToolkitHost()->GetTabManager()->FindExistingLiveTab(TabId);
-	if (CurveEditorTab)
+	if (const TSharedPtr<SDockTab> CurveEditorTab = TabManager->FindExistingLiveTab(CurveEditorTabName))
 	{
 		CurveEditorTab->RequestCloseTab();
 	}
@@ -639,11 +654,7 @@ void FCurveEditorExtension::CloseCurveEditor()
 
 FKeyAttributes FCurveEditorExtension::GetDefaultKeyAttributes() const
 {
-	TSharedPtr<FSequencerEditorViewModel> OwnerModel = WeakOwnerModel.Pin();
-	check(OwnerModel);
-	TSharedPtr<ISequencer> Sequencer = OwnerModel->GetSequencer();
-	check(Sequencer);
-	USequencerSettings* Settings = Sequencer->GetSequencerSettings();
+	USequencerSettings* const Settings = GetSequencerSettings();
 	check(Settings);
 
 	switch (Settings->GetKeyInterpolation())
diff --git a/Engine/Source/Editor/Sequencer/Public/MVVM/CurveEditorExtension.h b/Engine/Source/Editor/Sequencer/Public/MVVM/CurveEditorExtension.h
--- a/Engine/Source/Editor/Sequencer/Public/MVVM/CurveEditorExtension.h
+++ b/Engine/Source/Editor/Sequencer/Public/MVVM/CurveEditorExtension.h
@@ -14,11 +14,13 @@
 
 #define UE_API SEQUENCER_API
 
+class FTabManager;
 class ISequencer;
 class IPropertyTypeCustomization;
 class SCurveEditorPanel;
 class SCurveEditorTree;
 class SCurveEditorTreeFilterStatusBar;
+class USequencerSettings;
 struct FTimeSliderArgs;
 
 namespace UE
@@ -67,6 +69,10 @@ public:
 
 private:
 
+	TSharedPtr<ISequencer> GetSequencer() const;
+	USequencerSettings* GetSequencerSettings() const;
+	TSharedPtr<FTabManager> GetTabManager() const;
+
 	/** Get the default key attributes to apply to newly created keys on the curve editor */
 	UE_API FKeyAttributes GetDefaultKeyAttributes() const;
 
