Sequencer: Add base framework to switch Sequencer between linked and instanced filtering.

=== Feature brief ===

Going forward, the idea is to give other tools with a view into Sequencer tree view, such as Curve Editor, a way to link filtering.
The user can switch between instanced filtering (= only affect the UI in for that tool), and linked (= affect all UI instances that the user has also put into linked mode).
By default, Sequencer will now be in linked mode.
Follow up changes:
- Add a toggle button UI to Sequencer for switching between these modes
- Add the same filter widgets to Curve Editor's tree view

=== Notable changes ===

The filtering models have been abstracted into interfaces so we can easily expose it to the public API in the future.
ILinkedFilterFactoryViewModel and ILinkedFilterViewModel will remain in the private interface for now.
The intention is to move them to the public interface in the future: but only once needed to retain API flexibility.

This introduces FSequencerEditorViewModelImpl: as subclass of FSequencerEditorViewModel, it holds the models that should remain private to Sequencer code.
Other private details, such as FSequencerEditorViewModel::GetSequencerImpl could be moved here, too.

FSequencerFilterBar had to be adjusted:
- to no longer call IOutlinerExtension::SetFilteredOut directly. Instead, FSequencerFilterBar & FSequencerFilterData now exposes a delegate that FLinkedFilteringMediator binds to.
- to save the filter state to separate config entries. For this, we inject a subidentifier. Each FSequencerFilterBar instance will have a different subidentifier so the filter state can be saved. For now, there is no migration in place: the default linked filter mode will use the data from the same config entry as the released 5.7 version.

#Synced-CL 50395143

diff --git a/Engine/Source/Editor/Sequencer/Private/Filters/SequencerFilterBar.cpp b/Engine/Source/Editor/Sequencer/Private/Filters/SequencerFilterBar.cpp
--- a/Engine/Source/Editor/Sequencer/Private/Filters/SequencerFilterBar.cpp
+++ b/Engine/Source/Editor/Sequencer/Private/Filters/SequencerFilterBar.cpp
@@ -46,7 +46,7 @@ const FName FSequencerFilterBar::SharedIdentifier = TEXT("SharedSequencerTrackFi
 
 int32 FSequencerFilterBar::InstanceCount = 0;
 
-FSequencerFilterBar::FSequencerFilterBar(FSequencer& InSequencer)
+FSequencerFilterBar::FSequencerFilterBar(FSequencer& InSequencer, FString InSubIdentifier)
 	: Sequencer(InSequencer)
 	, CommandList(MakeShared<FUICommandList>())
 	, ClassTypeCategory(MakeShared<FFilterCategory>(LOCTEXT("ActorTypeFilterCategory", "Actor Type Filters"), FText()))
@@ -63,6 +63,7 @@ FSequencerFilterBar::FSequencerFilterBar(FSequencer& InSequencer)
 	, ModifiedFilter(MakeShared<FSequencerTrackFilter_Modified>(*this, MiscCategory))
 	, FilterMenu(MakeShared<FSequencerTrackFilterMenu>())
 	, FilterData(FString())
+	, SubIdentifier(MoveTemp(InSubIdentifier))
 {
 	InstanceCount++;
 
@@ -78,6 +79,11 @@ FSequencerFilterBar::FSequencerFilterBar(FSequencer& InSequencer)
 	CreateDefaultFilters();
 
 	CreateCustomTextFiltersFromConfig();
+	
+	FilterData.OnChangeFilterState().AddLambda([this](const TViewModelPtr<IOutlinerExtension>& InItem, bool bIsFilteredOut)
+	{
+		ChangeOutlinerExtensionFilterStateEvent.Broadcast(InItem, bIsFilteredOut);
+	});
 }
 
 FSequencerFilterBar::~FSequencerFilterBar()
@@ -308,7 +314,10 @@ FName FSequencerFilterBar::GetIdentifier() const
 {
 	if (const USequencerSettings* const SequencerSettings = Sequencer.GetSequencerSettings())
 	{
-		return *SequencerSettings->GetName();
+		return SubIdentifier.IsEmpty() 
+			? *SequencerSettings->GetName()
+			// E.g. "LevelSequencerEditor.Linked"
+			: *FString::Printf(TEXT("%s.%s"), *SequencerSettings->GetName(), *SubIdentifier);
 	}
 	return TEXT("SequencerMain");
 }
@@ -411,8 +420,6 @@ bool FSequencerFilterBar::HasIsolatedTracks() const
 
 void FSequencerFilterBar::RequestFilterUpdate()
 {
-	Sequencer.GetNodeTree()->RequestFilterUpdate();
-	
 	RequestUpdateEvent.Broadcast();
 }
 
@@ -1410,7 +1417,7 @@ const FSequencerFilterData& FSequencerFilterBar::FilterNodes()
 	{
 		if (const TViewModelPtr<IOutlinerExtension> SpacerNode = CastViewModelChecked<IOutlinerExtension>(SequenceModel->GetBottomSpacer()))
 		{
-			SpacerNode->SetFilteredOut(false);
+			ChangeOutlinerExtensionFilterStateEvent.Broadcast(SpacerNode, false);
 		}
 	}
 	
diff --git a/Engine/Source/Editor/Sequencer/Private/Filters/SequencerFilterBar.h b/Engine/Source/Editor/Sequencer/Private/Filters/SequencerFilterBar.h
--- a/Engine/Source/Editor/Sequencer/Private/Filters/SequencerFilterBar.h
+++ b/Engine/Source/Editor/Sequencer/Private/Filters/SequencerFilterBar.h
@@ -3,6 +3,7 @@
 #pragma once
 
 #include "Filters/CustomTextFilters.h"
+#include "Filters/FilterBarDelegates.h"
 #include "Filters/ISequencerTrackFilters.h"
 #include "Filters/SequencerFilterData.h"
 #include "Filters/SequencerTrackFilterBase.h"
@@ -47,8 +48,16 @@ public:
 	
 	DECLARE_EVENT_TwoParams(FSequencerFilterBar, FSequencerFiltersChanged, const ESequencerFilterChange /*InChangeType*/, const TSharedRef<FSequencerTrackFilter>& /*InFilter*/);
 	DECLARE_EVENT_TwoParams(FSequencerFilterBar, FSequencerCustomTextFiltersChanged, const ESequencerFilterChange /*InChangeType*/, const TSharedRef<FSequencerTrackFilter_CustomText>& /*InFilter*/);
+	
+	
 
-	FSequencerFilterBar(FSequencer& InSequencer);
+	/**
+	 * @param InSequencer The owning Sequencer instance 
+	 * @param InSubIdentifier Used to identify config entries. This is relevant saving settings to config entries corresponding to linked & instanced
+	 *	filtering. The config entry has the format "{0}.{1}" where {0} = FSequencer::Settings->GetName(), and {1} is SubIdentifier. If empty, only the
+	 *	settings name is used, i.e. FSequencer::Settings->GetName().
+	 */
+	explicit FSequencerFilterBar(FSequencer& InSequencer, FString InSubIdentifier = TEXT(""));
 	virtual ~FSequencerFilterBar() override;
 
 	void BindCommands();
@@ -239,6 +248,7 @@ public:
 	virtual FSimpleMulticastDelegate& OnRequestUpdate() override { return RequestUpdateEvent; }
 	virtual FSimpleMulticastDelegate& OnTextFilterTextChanged() override { return OnTextFilterTextChangedEvent; }
 	virtual FMuteStateChangedDelegate& OnMuteFiltersChanged() override { return OnMuteFiltersChangedEvent; }
+	UE::Sequencer::FChangeOutlinerExtensionFilterState& OnChangeOutlinerExtensionFilterState() { return ChangeOutlinerExtensionFilterStateEvent; }
 	
 	virtual void OpenTextExpressionHelp() override;
 	virtual void SaveCurrentFilterSetAsCustomTextFilter() override;
@@ -300,11 +310,22 @@ protected:
 	FSimpleMulticastDelegate RequestUpdateEvent;
 	FSimpleMulticastDelegate OnTextFilterTextChangedEvent;
 	FMuteStateChangedDelegate OnMuteFiltersChangedEvent;
+	/** Broadcasts when the filter state of IOutlinerExtension is changed. */
+	UE::Sequencer::FChangeOutlinerExtensionFilterState ChangeOutlinerExtensionFilterStateEvent;
 
 private:
 	
 	static int32 InstanceCount;
 	
+	/**
+	 * Used to build the config entry for the filters.
+	 * This is relevant for the settings to use for reading and saving config.
+	 * 
+	 * The final key is "{0}.{1}" where {0} = FSequencer::Settings->GetFName(), and {1} is SubIdentifier.
+	 * Examples: "LevelSequenceEditor.Linked", "LevelSequenceEditor.Instanced".
+	 */
+	const FString SubIdentifier;
+	
 	/** @return Gets the selected IOutlinerExtensions, or all of them if nothing is selected. */
 	TSet<UE::Sequencer::TWeakViewModelPtr<UE::Sequencer::IOutlinerExtension>> GetSelectedTracksOrAll() const;
 	
diff --git a/Engine/Source/Editor/Sequencer/Private/Filters/SequencerFilterData.cpp b/Engine/Source/Editor/Sequencer/Private/Filters/SequencerFilterData.cpp
--- a/Engine/Source/Editor/Sequencer/Private/Filters/SequencerFilterData.cpp
+++ b/Engine/Source/Editor/Sequencer/Private/Filters/SequencerFilterData.cpp
@@ -69,7 +69,7 @@ void FSequencerFilterData::FilterInNode(TWeakViewModelPtr<IOutlinerExtension> In
 
 	if (const TViewModelPtr<IOutlinerExtension> Node = InNodeWeak.Pin())
 	{
-		Node->SetFilteredOut(false);
+		ChangeFilterStateEvent.Broadcast(Node, false);
 	}
 }
 
@@ -83,7 +83,7 @@ void FSequencerFilterData::FilterOutNode(TWeakViewModelPtr<IOutlinerExtension> I
 
 	if (const TViewModelPtr<IOutlinerExtension> Node = InNodeWeak.Pin())
 	{
-		Node->SetFilteredOut(true);
+		ChangeFilterStateEvent.Broadcast(Node, true);
 	}
 }
 
diff --git a/Engine/Source/Editor/Sequencer/Private/Filters/ViewModels/Linking/ILinkedFilterFactoryViewModel.h b/Engine/Source/Editor/Sequencer/Private/Filters/ViewModels/Linking/ILinkedFilterFactoryViewModel.h
new file mode 100644
--- /dev/null
+++ b/Engine/Source/Editor/Sequencer/Private/Filters/ViewModels/Linking/ILinkedFilterFactoryViewModel.h
@@ -0,0 +1,29 @@
+﻿// Copyright Epic Games, Inc. All Rights Reserved.
+
+#pragma once
+
+#include "Templates/SharedPointerFwd.h"
+
+namespace UE::Sequencer
+{
+class ILinkedFilterViewModel;
+struct FMakeLinkedFilterViewModelArgs;
+
+/** 
+ * Factory for creating ILinkedFilterViewModel.
+ * Use this if you want to create UI that can be linked to the filters in the Sequencer UI.
+ */
+class ILinkedFilterFactoryViewModel
+{
+public:
+	
+	/**
+	 * Creates a new view model instance for UI that wants to be able to link to the filter state in Sequencer.
+	 * The view model has a reference to the shared linked filter state and its own filter state, which you can switch between.
+	 * @return View model that manages linked and instanced filter state for your UI.
+	 */
+	virtual TSharedRef<ILinkedFilterViewModel> MakeFilteringModel(const FMakeLinkedFilterViewModelArgs& InArgs) = 0;
+	
+	virtual ~ILinkedFilterFactoryViewModel() = default;
+};
+} // namespace UE::Sequencer
diff --git a/Engine/Source/Editor/Sequencer/Private/Filters/ViewModels/Linking/ILinkedFilterViewModel.h b/Engine/Source/Editor/Sequencer/Private/Filters/ViewModels/Linking/ILinkedFilterViewModel.h
new file mode 100644
--- /dev/null
+++ b/Engine/Source/Editor/Sequencer/Private/Filters/ViewModels/Linking/ILinkedFilterViewModel.h
@@ -0,0 +1,51 @@
+﻿// Copyright Epic Games, Inc. All Rights Reserved.
+
+#pragma once
+
+#include "HAL/Platform.h"
+#include "MVVM/ViewModelPtr.h"
+
+class ISequencerTrackFilters;
+
+namespace UE::Sequencer
+{
+class IOutlinerExtension;
+
+enum class ELinkedFilterMode : uint8
+{
+	/** Filtering is using a shared filterinstance */
+	Linked,
+	/** Filtering is using its own filter state. */
+	Instanced,
+	
+	/** Not a real mode. Useful for static asserting. */
+	Count
+};
+
+/** 
+ * View-model for UI to switch between:
+ * - shared filter state: this state is shared by multiple UI, allowing the user to sync multiple UI instances (Sequencer, Curve Editor, etc.)
+ * - instanced filter state: this state is used only by the UI instance
+ */
+class ILinkedFilterViewModel
+{
+public:
+	
+	/** Sets the current filtering mode. */
+	virtual void SetFilterMode(ELinkedFilterMode InNewMode) = 0;
+	/** @return The current filtering mode. */
+	virtual ELinkedFilterMode GetFilterMode() const = 0;
+	
+	/** @return The currently active filter bar instance according to the linked mode. */
+	virtual TSharedRef<ISequencerTrackFilters> GetActiveFilterBar() const = 0; 
+	/** @return The filter bar that is shared with all other UI. */
+	virtual TSharedRef<ISequencerTrackFilters> GetLinkedFilterBar() const = 0;
+	/** @return The filter bar that is only used by this instance. */
+	virtual TSharedRef<ISequencerTrackFilters> GetInstancedFilterBar() const = 0;
+	
+	/** @return Event when the filter mode changes, i.e. the instance returned by GetActiveFilterBar changes. */
+	virtual FSimpleMulticastDelegate& OnFilterModeChanged() = 0;
+	
+	virtual ~ILinkedFilterViewModel() = default;
+};
+} // namespace UE::Sequencer
diff --git a/Engine/Source/Editor/Sequencer/Private/Filters/ViewModels/Linking/LinkedFilterFactoryViewModel.cpp b/Engine/Source/Editor/Sequencer/Private/Filters/ViewModels/Linking/LinkedFilterFactoryViewModel.cpp
new file mode 100644
--- /dev/null
+++ b/Engine/Source/Editor/Sequencer/Private/Filters/ViewModels/Linking/LinkedFilterFactoryViewModel.cpp
@@ -0,0 +1,28 @@
+﻿// Copyright Epic Games, Inc. All Rights Reserved.
+
+#include "LinkedFilterFactoryViewModel.h"
+
+#include "LinkedFilterViewModel.h"
+#include "Filters/ViewModels/Linking/MakeLinkedFilterViewModelArgs.h"
+
+namespace UE::Sequencer
+{
+FLinkedFilterFactoryViewModel::FLinkedFilterFactoryViewModel(
+	FSequencer& InSequencer UE_LIFETIMEBOUND,
+	const TSharedRef<FSequencerFilterBar>& InSharedFilterBar
+	)
+	: Sequencer(InSequencer)
+	, SharedFilterBar(InSharedFilterBar)
+{}
+
+TSharedRef<FLinkedFilterViewModel> FLinkedFilterFactoryViewModel::MakeFilteringModelImpl(const FMakeLinkedFilterViewModelArgs& InArgs) const
+{
+	const TSharedRef<FSequencerFilterBar> InstancedFilterBar = MakeShared<FSequencerFilterBar>(Sequencer, InArgs.ConfigSubIdentifier);
+	return MakeShared<FLinkedFilterViewModel>(SharedFilterBar, InstancedFilterBar);
+}
+
+TSharedRef<ILinkedFilterViewModel> FLinkedFilterFactoryViewModel::MakeFilteringModel(const FMakeLinkedFilterViewModelArgs& InArgs)
+{
+	return MakeFilteringModelImpl(InArgs);
+}
+} // namespace UE::Sequencer
diff --git a/Engine/Source/Editor/Sequencer/Private/Filters/ViewModels/Linking/LinkedFilterFactoryViewModel.h b/Engine/Source/Editor/Sequencer/Private/Filters/ViewModels/Linking/LinkedFilterFactoryViewModel.h
new file mode 100644
--- /dev/null
+++ b/Engine/Source/Editor/Sequencer/Private/Filters/ViewModels/Linking/LinkedFilterFactoryViewModel.h
@@ -0,0 +1,37 @@
+﻿// Copyright Epic Games, Inc. All Rights Reserved.
+
+#pragma once
+
+#include "Filters/SequencerFilterBar.h"
+#include "Filters/ViewModels/Linking/ILinkedFilterFactoryViewModel.h"
+#include "HAL/Platform.h"
+
+namespace UE::Sequencer
+{
+class FLinkedFilterViewModel;
+
+class FLinkedFilterFactoryViewModel : public ILinkedFilterFactoryViewModel
+{
+public:
+	
+	explicit FLinkedFilterFactoryViewModel(
+		FSequencer& InSequencer UE_LIFETIMEBOUND,
+		const TSharedRef<FSequencerFilterBar>& InSharedFilterBar
+		);
+	
+	/** @return View model that manages linked and instanced filter state for your UI.  */
+	TSharedRef<FLinkedFilterViewModel> MakeFilteringModelImpl(const FMakeLinkedFilterViewModelArgs& InArgs) const;
+	
+	//~ Begin ILinkedFilterFactoryViewModel
+	virtual TSharedRef<ILinkedFilterViewModel> MakeFilteringModel(const FMakeLinkedFilterViewModelArgs& InArgs) override;
+	//~ End ILinkedFilterFactoryViewModel
+	
+private:
+	
+	/** Sequencer */
+	FSequencer& Sequencer;
+	
+	/** The filter bar instance that model edit when they are in linked mode. All created models reference this. */
+	const TSharedRef<FSequencerFilterBar> SharedFilterBar;
+};
+} // namespace UE::Sequencer
diff --git a/Engine/Source/Editor/Sequencer/Private/Filters/ViewModels/Linking/LinkedFilterViewModel.cpp b/Engine/Source/Editor/Sequencer/Private/Filters/ViewModels/Linking/LinkedFilterViewModel.cpp
new file mode 100644
--- /dev/null
+++ b/Engine/Source/Editor/Sequencer/Private/Filters/ViewModels/Linking/LinkedFilterViewModel.cpp
@@ -0,0 +1,62 @@
+﻿// Copyright Epic Games, Inc. All Rights Reserved.
+
+#include "LinkedFilterViewModel.h"
+
+#include "Filters/SequencerFilterBar.h"
+
+namespace UE::Sequencer
+{
+FLinkedFilterViewModel::FLinkedFilterViewModel(
+	const TSharedRef<FSequencerFilterBar>& InSharedFilterBar,
+	const TSharedRef<FSequencerFilterBar>& InInstancedFilterBar
+	)
+	: LinkedFilterBar(InSharedFilterBar)
+	, InstancedFilterBar(InInstancedFilterBar)
+{}
+
+TSharedRef<FSequencerFilterBar> FLinkedFilterViewModel::GetActiveFilterBarImpl() const
+{
+	static_assert(static_cast<int32>(ELinkedFilterMode::Count) == 2, "Update this switch");
+	switch (FilterMode)
+	{
+	case ELinkedFilterMode::Linked: return LinkedFilterBar;
+	case ELinkedFilterMode::Instanced: [[fallthrough]];
+	default: return InstancedFilterBar;
+	}
+}
+
+TSharedRef<FSequencerFilterBar> FLinkedFilterViewModel::GetLinkedFilterBarImpl() const
+{
+	return LinkedFilterBar;
+}
+
+TSharedRef<FSequencerFilterBar> FLinkedFilterViewModel::GetInstancedFilterBarImpl() const
+{
+	return InstancedFilterBar;
+}
+
+void FLinkedFilterViewModel::SetFilterMode(ELinkedFilterMode InNewMode)
+{
+	if (InNewMode != FilterMode)
+	{
+		FilterMode = InNewMode;
+		GetActiveFilterBar()->RequestFilterUpdate();
+		OnFilterModeChangedDelegate.Broadcast();
+	}
+}
+
+TSharedRef<ISequencerTrackFilters> FLinkedFilterViewModel::GetActiveFilterBar() const
+{
+	return GetActiveFilterBarImpl();
+}
+
+TSharedRef<ISequencerTrackFilters> FLinkedFilterViewModel::GetLinkedFilterBar() const
+{
+	return LinkedFilterBar;
+}
+
+TSharedRef<ISequencerTrackFilters> FLinkedFilterViewModel::GetInstancedFilterBar() const
+{
+	return InstancedFilterBar;
+}
+} // namespace UE::Sequencer
diff --git a/Engine/Source/Editor/Sequencer/Private/Filters/ViewModels/Linking/LinkedFilterViewModel.h b/Engine/Source/Editor/Sequencer/Private/Filters/ViewModels/Linking/LinkedFilterViewModel.h
new file mode 100644
--- /dev/null
+++ b/Engine/Source/Editor/Sequencer/Private/Filters/ViewModels/Linking/LinkedFilterViewModel.h
@@ -0,0 +1,48 @@
+﻿// Copyright Epic Games, Inc. All Rights Reserved.
+
+#pragma once
+
+#include "Filters/ViewModels/Linking/ILinkedFilterViewModel.h"
+#include "Templates/SharedPointer.h"
+
+class FSequencer;
+class FSequencerFilterBar;
+
+namespace UE::Sequencer
+{
+class FLinkedFilterViewModel : public ILinkedFilterViewModel
+{
+public:
+	
+	explicit FLinkedFilterViewModel(
+		const TSharedRef<FSequencerFilterBar>& InSharedFilterBar,
+		const TSharedRef<FSequencerFilterBar>& InInstancedFilterBar
+		);
+	
+	TSharedRef<FSequencerFilterBar> GetActiveFilterBarImpl() const;
+	TSharedRef<FSequencerFilterBar> GetLinkedFilterBarImpl() const;
+	TSharedRef<FSequencerFilterBar> GetInstancedFilterBarImpl() const;
+	
+	//~ Begin ILinkedFilterViewModel
+	virtual void SetFilterMode(ELinkedFilterMode InNewMode) override;
+	virtual ELinkedFilterMode GetFilterMode() const override { return FilterMode; }
+	virtual TSharedRef<ISequencerTrackFilters> GetActiveFilterBar() const override;
+	virtual TSharedRef<ISequencerTrackFilters> GetLinkedFilterBar() const override;
+	virtual TSharedRef<ISequencerTrackFilters> GetInstancedFilterBar() const override;
+	virtual FSimpleMulticastDelegate& OnFilterModeChanged() override { return OnFilterModeChangedDelegate; }
+	//~ End ILinkedFilterViewModel
+	
+private:
+	
+	/** The filter bar instance that is shared across all UI. Other UI can choose to link to this instance. */
+	const TSharedRef<FSequencerFilterBar> LinkedFilterBar;
+	/** The filter bar instance that is owned by the particular view that is displaying the filters, e.g. Curve Editor. */
+	const TSharedRef<FSequencerFilterBar> InstancedFilterBar;
+	
+	/** The current filter mode. It determines which filter bar instance is used. */
+	ELinkedFilterMode FilterMode = ELinkedFilterMode::Linked;
+	
+	/** Event when the filter mode changes, i.e. the instance returned by GetActiveFilterBar changes. */
+	FSimpleMulticastDelegate OnFilterModeChangedDelegate;
+};
+} // namespace UE::Sequencer
diff --git a/Engine/Source/Editor/Sequencer/Private/Filters/ViewModels/Linking/MakeLinkedFilterViewModelArgs.h b/Engine/Source/Editor/Sequencer/Private/Filters/ViewModels/Linking/MakeLinkedFilterViewModelArgs.h
new file mode 100644
--- /dev/null
+++ b/Engine/Source/Editor/Sequencer/Private/Filters/ViewModels/Linking/MakeLinkedFilterViewModelArgs.h
@@ -0,0 +1,26 @@
+﻿// Copyright Epic Games, Inc. All Rights Reserved.
+
+#pragma once
+
+#include "Containers/UnrealString.h"
+
+namespace UE::Sequencer
+{
+/** @see ILinkedFilterFactoryViewModel::MakeFilteringModel */
+struct FMakeLinkedFilterViewModelArgs
+{
+	/**
+	 * The sub-identifier to used for saving and loading filter entries.
+	 * 
+	 * The final identifier is the name of the settings (USequencerSettings::GetName) + "." + ConfigSubIdentifier.
+	 * Example: "LevelSequencerEditor.Linked", "LevelSequencerEditor.CurveEditor", etc.
+	 * 
+	 * @see USequencerSettings::TrackFilterBars.
+	 */
+	FString ConfigSubIdentifier;
+	
+	explicit FMakeLinkedFilterViewModelArgs(FString InConfigSubIdentifier)
+		: ConfigSubIdentifier(MoveTemp(InConfigSubIdentifier))
+		{}
+};
+} // namespace UE::Sequencer
diff --git a/Engine/Source/Editor/Sequencer/Private/MVVM/ViewModels/LinkedFilteringMediator.cpp b/Engine/Source/Editor/Sequencer/Private/MVVM/ViewModels/LinkedFilteringMediator.cpp
new file mode 100644
--- /dev/null
+++ b/Engine/Source/Editor/Sequencer/Private/MVVM/ViewModels/LinkedFilteringMediator.cpp
@@ -0,0 +1,107 @@
+﻿// Copyright Epic Games, Inc. All Rights Reserved.
+
+#include "LinkedFilteringMediator.h"
+
+#include "Filters/SequencerFilterBar.h"
+#include "Filters/ViewModels/Linking/ILinkedFilterViewModel.h"
+#include "Filters/ViewModels/Linking/LinkedFilterViewModel.h"
+#include "SequencerNodeTree.h"
+
+namespace UE::Sequencer
+{
+FLinkedFilteringMediator::FLinkedFilteringMediator(
+	const TSharedRef<FSequencerNodeTree>& InNodeTree,
+	const TSharedRef<FLinkedFilterViewModel>& InSequencerFilterViewModel
+	)
+	: WeakNodeTree(InNodeTree)
+	, SequencerFilterViewModel(InSequencerFilterViewModel)
+{
+	SequencerFilterViewModel->OnFilterModeChanged().AddRaw(this, &FLinkedFilteringMediator::OnFilterModeChanged);
+	
+	const TSharedRef<FSequencerFilterBar> SharedFilterBar = SequencerFilterViewModel->GetLinkedFilterBarImpl();
+	SharedFilterBar->OnRequestUpdate().AddRaw(this, &FLinkedFilteringMediator::OnRequestUpdate_LinkedFilterBar);
+	SharedFilterBar->OnChangeOutlinerExtensionFilterState().AddRaw(this, &FLinkedFilteringMediator::OnFilterStateChanged_LinkedFilterBar);
+	
+	const TSharedRef<FSequencerFilterBar> InstancedFilterBar = SequencerFilterViewModel->GetInstancedFilterBarImpl();
+	InstancedFilterBar->OnChangeOutlinerExtensionFilterState().AddRaw(this, &FLinkedFilteringMediator::OnFilterStateChanged_InstancedFilterBar);
+	InstancedFilterBar->OnRequestUpdate().AddRaw(this, &FLinkedFilteringMediator::OnRequestUpdate_InstancedFilterBar);
+}
+
+FLinkedFilteringMediator::~FLinkedFilteringMediator()
+{
+	SequencerFilterViewModel->OnFilterModeChanged().RemoveAll(this);
+	
+	const TSharedRef<FSequencerFilterBar> SharedFilterBar = SequencerFilterViewModel->GetLinkedFilterBarImpl();
+	SharedFilterBar->OnChangeOutlinerExtensionFilterState().RemoveAll(this);
+	SharedFilterBar->OnRequestUpdate().RemoveAll(this);
+	
+	const TSharedRef<FSequencerFilterBar> InstancedFilterBar = SequencerFilterViewModel->GetInstancedFilterBarImpl();
+	InstancedFilterBar->OnChangeOutlinerExtensionFilterState().RemoveAll(this);
+	InstancedFilterBar->OnRequestUpdate().RemoveAll(this);
+}
+
+void FLinkedFilteringMediator::BindCommands() const
+{
+	const TSharedRef<FSequencerFilterBar> SharedFilterBar = SequencerFilterViewModel->GetLinkedFilterBarImpl();
+	SharedFilterBar->BindCommands();
+	
+	const TSharedRef<FSequencerFilterBar> InstancedFilterBar = SequencerFilterViewModel->GetInstancedFilterBarImpl();
+	InstancedFilterBar->BindCommands();
+}
+
+bool FLinkedFilteringMediator::ProcessCommandBindings(const FKeyEvent& InKeyEvent) const
+{
+	const TSharedPtr<FUICommandList> CommandList = SequencerFilterViewModel->GetActiveFilterBar()->GetCommandList();
+	return CommandList && CommandList->ProcessCommandBindings(InKeyEvent);
+}
+
+ELinkedFilterMode FLinkedFilteringMediator::GetFilterMode() const
+{
+	return SequencerFilterViewModel->GetFilterMode();
+}
+
+void FLinkedFilteringMediator::OnFilterModeChanged()
+{
+	// TODO: The goal of this is to call SetFilteredOut on all items again, but it does more work then it needs to. Instead of rerunning the filter,
+	// we could iterate through the hierarchy, and use FSequencerFilterBar::GetFilterData to call SetFilteredOut with the correct filter state.
+	SequencerFilterViewModel->GetActiveFilterBar()->RequestFilterUpdate();
+}
+
+void FLinkedFilteringMediator::OnFilterStateChanged_LinkedFilterBar(
+	const TViewModelPtr<IOutlinerExtension>& InItem, bool bInIsFilteredOut
+	) const
+{
+	if (GetFilterMode() == ELinkedFilterMode::Linked)
+	{
+		InItem->SetFilteredOut(bInIsFilteredOut);
+	}
+}
+
+void FLinkedFilteringMediator::OnRequestUpdate_LinkedFilterBar() const
+{
+	const TSharedPtr<FSequencerNodeTree> NodeTreePin = WeakNodeTree.Pin();
+	if (NodeTreePin && GetFilterMode() == ELinkedFilterMode::Linked)
+	{
+		NodeTreePin->RequestFilterUpdate();
+	}
+}
+
+void FLinkedFilteringMediator::OnFilterStateChanged_InstancedFilterBar(
+	const TViewModelPtr<IOutlinerExtension>& InItem, bool bInIsFilteredOut
+	) const
+{
+	if (GetFilterMode() == ELinkedFilterMode::Instanced)
+	{
+		InItem->SetFilteredOut(bInIsFilteredOut);
+	}
+}
+
+void FLinkedFilteringMediator::OnRequestUpdate_InstancedFilterBar() const
+{
+	const TSharedPtr<FSequencerNodeTree> NodeTreePin = WeakNodeTree.Pin();
+	if (NodeTreePin && GetFilterMode() == ELinkedFilterMode::Instanced)
+	{
+		NodeTreePin->RequestFilterUpdate();
+	}
+}
+} // namespace UE::Sequencer
diff --git a/Engine/Source/Editor/Sequencer/Private/MVVM/ViewModels/LinkedFilteringMediator.h b/Engine/Source/Editor/Sequencer/Private/MVVM/ViewModels/LinkedFilteringMediator.h
new file mode 100644
--- /dev/null
+++ b/Engine/Source/Editor/Sequencer/Private/MVVM/ViewModels/LinkedFilteringMediator.h
@@ -0,0 +1,64 @@
+﻿// Copyright Epic Games, Inc. All Rights Reserved.
+
+#pragma once
+
+#include "MVVM/Extensions/IOutlinerExtension.h"
+#include "MVVM/ViewModels/ViewModel.h"
+
+class FSequencerNodeTree;
+
+namespace UE::Sequencer
+{
+class FLinkedFilterViewModel;
+enum class ELinkedFilterMode : uint8;
+
+/** 
+ * Hooks up Sequencer's own FLinkedFilterViewModel with the relevant systems.
+ * 
+ * Its responsibility is to handle updating any systems when the linked filtering mode changes.
+ * For example:
+ * - Calling IOutlinerExtension::SetFilteredOut when the item's filter state changes
+ * - Causing the FSequencerNodeTree to refilter & refresh the tree view when the active FSequencerFilterBar is requested to refilter
+ * - Binding commands that are bound by FSequencerFilterBar and making sure correct commands are executed depending on which one is active.
+ */
+class FLinkedFilteringMediator : public FViewModel
+{
+public:
+	
+	explicit FLinkedFilteringMediator(
+		const TSharedRef<FSequencerNodeTree>& InNodeTree,
+		const TSharedRef<FLinkedFilterViewModel>& InSequencerFilterViewModel
+		);
+	virtual ~FLinkedFilteringMediator() override;
+	
+	/** Binds the filter bar commands. */
+	void BindCommands() const;
+	
+	/** Process the commands of the active filter bar. */
+	bool ProcessCommandBindings(const FKeyEvent& InKeyEvent) const;
+	
+private:
+	
+	/** Needed to cause the tree view to refresh. */
+	const TWeakPtr<FSequencerNodeTree> WeakNodeTree;
+	
+	/** Knows whether Sequencer is using linked or instanced filtering. */
+	const TSharedRef<FLinkedFilterViewModel> SequencerFilterViewModel;
+	
+	/** @return The filter mode that Sequencer is currently using. */
+	ELinkedFilterMode GetFilterMode() const;
+
+	/** Schedules the IOutlinerExtensions to be updated with SetFilteredOut calls to reflect the state of the new filter bar. */
+	void OnFilterModeChanged();
+	
+	/** Handles propagating the IOutlinerExtension's filter out state if the active filter bar is the linked one. */
+	void OnFilterStateChanged_LinkedFilterBar(const TViewModelPtr<IOutlinerExtension>& InItem, bool bInIsFilteredOut) const;
+	/** Propagates filter changes from the linked filter bar to the Sequencer tree.  */
+	void OnRequestUpdate_LinkedFilterBar() const;
+	
+	/** Handles propagating the IOutlinerExtension's filter out state if the active filter bar is the instanced one. */
+	void OnFilterStateChanged_InstancedFilterBar(const TViewModelPtr<IOutlinerExtension>& InItem, bool bInIsFilteredOut) const;
+	/** Propagates filter changes from the instanced filter bar to the Sequencer tree.  */
+	void OnRequestUpdate_InstancedFilterBar() const;
+};
+} // namespace UE::Sequencer
diff --git a/Engine/Source/Editor/Sequencer/Private/MVVM/ViewModels/SequencerEditorViewModelImpl.cpp b/Engine/Source/Editor/Sequencer/Private/MVVM/ViewModels/SequencerEditorViewModelImpl.cpp
new file mode 100644
--- /dev/null
+++ b/Engine/Source/Editor/Sequencer/Private/MVVM/ViewModels/SequencerEditorViewModelImpl.cpp
@@ -0,0 +1,52 @@
+﻿// Copyright Epic Games, Inc. All Rights Reserved.
+
+#include "SequencerEditorViewModelImpl.h"
+
+#include "Filters/ViewModels/Linking/LinkedFilterFactoryViewModel.h"
+#include "Filters/ViewModels/Linking/LinkedFilterViewModel.h"
+#include "Filters/ViewModels/Linking/MakeLinkedFilterViewModelArgs.h"
+#include "LinkedFilteringMediator.h"
+#include "Sequencer.h"
+
+namespace UE::Sequencer
+{
+/** Sub-identifier used for saving filter states to config for the linked FSequencerFilterBar. */
+const TCHAR* LinkedConfigSubId = TEXT(""); // TODO UE-362151: For now empty so we load using the old config entry, see below.
+/** Sub-identifier used for saving filter states to config for the instanced FSequencerFilterBar displayed in SSequencer. */
+const TCHAR* InstancedConfigSubId = TEXT("MainInstance"); 
+
+FSequencerEditorViewModelImpl::FSequencerEditorViewModelImpl(
+	TSharedRef<ISequencer> InSequencer, const FSequencerHostCapabilities& InHostCapabilities
+	)
+	: FSequencerEditorViewModel(InSequencer, InHostCapabilities)
+{}
+
+void FSequencerEditorViewModelImpl::InitializeEditorImpl()
+{
+	Super::InitializeEditorImpl();
+	
+	InitLinkedFiltering();
+}
+
+void FSequencerEditorViewModelImpl::InitLinkedFiltering()
+{
+	const TSharedPtr<FSequencer> Sequencer = GetSequencerImpl();
+	// TODO UE-362151: Implement data migration.
+	// Before linked filtering, the filter state was saved into the section string "*SequencerSettings->GetName()", e.g. "LevelSequencerEditor".
+	// Going forward, we can separate linked and instanced, e.g. ""LevelSequencerEditor.Linked", "LevelSequencerEditor.MainInstance", respectively.
+	// For now, filtering will default to linked mode. For linked mode, we specify no sub-identifier: that causes FSequencerFilterBar to continue 
+	// using the old config entry, which would be SequencerSettings->GetName(), e.g. "LevelSequencerEditor".
+	const TCHAR* ConfigId = LinkedConfigSubId;
+	LinkedFilterFactory = MakeShared<FLinkedFilterFactoryViewModel>(
+		*Sequencer,
+		MakeShared<FSequencerFilterBar>(*Sequencer, ConfigId)
+		);
+	
+	SequenceFilterInstanceViewModel = LinkedFilterFactory->MakeFilteringModelImpl(
+		FMakeLinkedFilterViewModelArgs(InstancedConfigSubId)
+		);
+	FilterMediator = MakeShared<FLinkedFilteringMediator>(
+		Sequencer->GetNodeTree(), SequenceFilterInstanceViewModel.ToSharedRef()
+		);
+}
+}
diff --git a/Engine/Source/Editor/Sequencer/Private/MVVM/ViewModels/SequencerEditorViewModelImpl.h b/Engine/Source/Editor/Sequencer/Private/MVVM/ViewModels/SequencerEditorViewModelImpl.h
new file mode 100644
--- /dev/null
+++ b/Engine/Source/Editor/Sequencer/Private/MVVM/ViewModels/SequencerEditorViewModelImpl.h
@@ -0,0 +1,47 @@
+﻿// Copyright Epic Games, Inc. All Rights Reserved.
+
+#pragma once
+
+#include "MVVM/ViewModels/SequencerEditorViewModel.h"
+
+namespace UE::Sequencer
+{
+class FLinkedFilteringMediator;
+class FLinkedFilterViewModel;
+class FLinkedFilterFactoryViewModel;
+
+/** Exposes private API for Sequencer. */
+class FSequencerEditorViewModelImpl : public FSequencerEditorViewModel
+{
+	using Super = FSequencerEditorViewModel;
+public:
+	
+	explicit FSequencerEditorViewModelImpl(TSharedRef<ISequencer> InSequencer, const FSequencerHostCapabilities& InHostCapabilities);
+	
+	/** Private version of GetLinkedFilterViewModel */
+	TSharedPtr<FLinkedFilterViewModel> GetLinkedFilterViewModelImpl() const { return SequenceFilterInstanceViewModel; }
+	
+	/** @return The object responsible for syncing the linked filters to Sequencer's systems. */
+	TSharedPtr<FLinkedFilteringMediator> GetFilterMediator() const { return FilterMediator; }
+	
+protected:
+	
+	//~ Begin FEditorViewModel Interface
+	virtual void InitializeEditorImpl() override;
+	//~ End FEditorViewModel Interface
+	
+private:
+	
+	/** Creates view models so (external) UI instances can sync to the filter state of the Sequencer UI. */
+	TSharedPtr<FLinkedFilterFactoryViewModel> LinkedFilterFactory;
+	/** Filter state view model that Sequencer UI uses when in ELinkedFilterMode::Instanced. */
+	TSharedPtr<FLinkedFilterViewModel> SequenceFilterInstanceViewModel;
+	
+	/** Hooks up Sequencer's own FLinkedFilterViewModel with the relevant systems.*/
+	TSharedPtr<FLinkedFilteringMediator> FilterMediator;
+	
+	/** Inits the systems relevant for linked filtering. */
+	void InitLinkedFiltering();
+};
+}
+
diff --git a/Engine/Source/Editor/Sequencer/Private/Sequencer.cpp b/Engine/Source/Editor/Sequencer/Private/Sequencer.cpp
--- a/Engine/Source/Editor/Sequencer/Private/Sequencer.cpp
+++ b/Engine/Source/Editor/Sequencer/Private/Sequencer.cpp
@@ -188,6 +188,10 @@
 #include "Bindings/MovieSceneSpawnableActorBinding.h"
 #include "Bindings/MovieSceneReplaceableActorBinding.h"
 #include "Decorations/MovieSceneMuteSoloDecoration.h"
+#include "Filters/SequencerTrackFilterCommands.h"
+#include "Filters/ViewModels/Linking/LinkedFilterViewModel.h"
+#include "MVVM/ViewModels/LinkedFilteringMediator.h"
+#include "MVVM/ViewModels/SequencerEditorViewModelImpl.h"
 #include "TrackEditors/TimeWarpTrackEditor.h"
 #include "Variants/MovieSceneTimeWarpGetter.h"
 
@@ -624,8 +628,8 @@ void FSequencer::InitSequencer(const FSequencerInitParams& InitParams
 		DeferredSignedObjectChangeHandler = MakeShared<FDeferredSignedObjectChangeHandler>();
 		UMovieSceneSignedObject::SetDeferredHandler(DeferredSignedObjectChangeHandler);
 	}
-
-	ViewModel = MakeShared<FSequencerEditorViewModel>(SharedThisRef, GetHostCapabilities());
+	
+	ViewModel = MakeShared<FSequencerEditorViewModelImpl>(SharedThisRef, GetHostCapabilities());
 	ViewModel->InitializeEditor();
 	ViewModel->SetSequence(InitParams.RootSequence);
 	ViewModel->GetRootSequenceModel()->InitializeExtensions();
@@ -651,8 +655,6 @@ void FSequencer::InitSequencer(const FSequencerInitParams& InitParams
 		Settings->SetLoopMode(SLM_LoopSelectionRange);
 	}
 
-	FilterBar = MakeShared<FSequencerFilterBar>(*this);
-
 	// Make internal widgets
 	SequencerWidget = SNew( SSequencer, SharedThis( this ) )
 		.ViewRange( this, &FSequencer::GetViewRange )
@@ -904,8 +906,6 @@ FSequencer::~FSequencer()
 
 	TrackEditors.Empty();
 
-	FilterBar.Reset();
-
 	RootTemplateInstance.TearDown();
 }
 
@@ -947,8 +947,6 @@ void FSequencer::Close()
 	SequencerWidget.Reset();
 	TrackEditors.Empty();
 
-	FilterBar.Reset();
-
 	GUnrealEd->UpdatePivotLocationForSelection();
 
 	// Redraw viewports after restoring pre animated state in case viewports are not set to realtime
@@ -6305,6 +6303,7 @@ bool FSequencer::OnRequestNodeDeleted( TSharedRef<FViewModel> NodeToBeDeleted, c
 
 		const TViewModelPtr<IOutlinerExtension> OutlinerItem = NodeToBeDeleted->CastThisShared<IOutlinerExtension>();
 
+		const TSharedPtr<FSequencerFilterBar> FilterBar = GetFilterBar();
 		if (FilterBar->HasIsolatedTracks())
 		{
 			FilterBar->UnisolateTracks({ OutlinerItem });
@@ -6887,8 +6886,8 @@ void FSequencer::AddNodesToExistingNodeGroup(TArrayView<const UE::Sequencer::TWe
 
 void FSequencer::ClearFilters()
 {
-	FilterBar->SetTextFilterString(FString());
-	FilterBar->EnableFilters(false);
+	GetFilterBar()->SetTextFilterString(FString());
+	GetFilterBar()->EnableFilters(false);
 	GetSequencerSettings()->SetShowSelectedNodesOnly(false);
 
 	UMovieSceneSequence* FocusedMovieSequence = GetFocusedMovieSceneSequence();
@@ -8973,6 +8972,7 @@ void FSequencer::ObjectImplicitlyRemoved(UObject* InObject) const
 
 void FSequencer::SetTrackFilterEnabled(const FText& InFilterName, bool bInEnabled)
 {
+	const TSharedPtr<FSequencerFilterBar> FilterBar = GetFilterBar();
 	if (const TSharedPtr<FSequencerTrackFilter> Filter = FilterBar->FindFilterByDisplayName(InFilterName.ToString()))
 	{
 		FilterBar->SetFilterEnabled(Filter.ToSharedRef(), bInEnabled, true);
@@ -8981,6 +8981,7 @@ void FSequencer::SetTrackFilterEnabled(const FText& InFilterName, bool bInEnable
 
 bool FSequencer::IsTrackFilterEnabled(const FText& InFilterName) const
 {
+	const TSharedPtr<FSequencerFilterBar> FilterBar = GetFilterBar();
 	if (const TSharedPtr<FSequencerTrackFilter> Filter = FilterBar->FindFilterByDisplayName(InFilterName.ToString()))
 	{
 		return FilterBar->IsFilterEnabled(Filter.ToSharedRef());
@@ -8990,6 +8991,7 @@ bool FSequencer::IsTrackFilterEnabled(const FText& InFilterName) const
 
 TArray<FText> FSequencer::GetTrackFilterNames() const
 {
+	const TSharedPtr<FSequencerFilterBar> FilterBar = GetFilterBar();
 	return FilterBar->GetFilterDisplayNames();
 }
 
@@ -9607,6 +9609,7 @@ void FSequencer::OnAddBinding(const FGuid& ObjectBinding, UMovieScene* MovieScen
 
 	NodeTree->SortAllNodesAndDescendants();
 
+	const TSharedPtr<FSequencerFilterBar> FilterBar = GetFilterBar();
 	if (FilterBar->HasIsolatedTracks())
 	{
 		for (const TViewModelPtr<IObjectBindingExtension>& ObjectBindingItem : NodeTree->GetRootNode()->GetDescendantsOfType<IObjectBindingExtension>())
@@ -9673,6 +9676,7 @@ void FSequencer::OnAddTrack(const TWeakObjectPtr<UMovieSceneTrack>&InTrack, cons
 
 	NodeTree->SortAllNodesAndDescendants();
 
+	const TSharedPtr<FSequencerFilterBar> FilterBar = GetFilterBar();
 	if (FilterBar->HasIsolatedTracks())
 	{
 		SequencerWidget->AddNewNodePathsToIsolate({ NewNodePath });
@@ -12292,8 +12296,8 @@ void FSequencer::BindCommands()
 
 	// bind widget specific commands
 	SequencerWidget->BindCommands(SequencerCommandBindings, CurveEditorSharedBindings);
-
-	FilterBar->BindCommands();
+	
+	ViewModel->GetFilterMediator()->BindCommands();
 }
 
 void FSequencer::BuildAddTrackMenu(class FMenuBuilder& MenuBuilder)
@@ -12992,12 +12996,12 @@ FText FSequencer::GetSidebarSelectionDrawerToolTipText() const
 
 TSharedRef<ISequencerTrackFilters> FSequencer::GetFilterInterface() const
 {
-	return FilterBar.ToSharedRef();
+	return GetFilterBar();
 }
 
 TSharedRef<FSequencerFilterBar> FSequencer::GetFilterBar() const
 {
-	return FilterBar.ToSharedRef();
+	return ViewModel->GetLinkedFilterViewModelImpl()->GetActiveFilterBarImpl();
 }
 
 EVisibility FSequencer::GetPlayRateComboVisibility() const
diff --git a/Engine/Source/Editor/Sequencer/Private/Sequencer.h b/Engine/Source/Editor/Sequencer/Private/Sequencer.h
--- a/Engine/Source/Editor/Sequencer/Private/Sequencer.h
+++ b/Engine/Source/Editor/Sequencer/Private/Sequencer.h
@@ -103,6 +103,7 @@ namespace UE
 		struct FViewModelTypeID;
 		class FSequenceModel;
 		class FSequencerEditorViewModel;
+		class FSequencerEditorViewModelImpl;
 		class FViewModel;
 		class FSequencerSelection;
 
@@ -248,6 +249,11 @@ public:
 	{
 		return NodeTree;
 	}
+	
+	TSharedRef<UE::Sequencer::FSequencerEditorViewModelImpl> GetViewModelImpl() const
+	{
+		return ViewModel.ToSharedRef();
+	}
 
 	FObjectBindingTagCache* GetObjectBindingTagCache()
 	{
@@ -1625,7 +1631,7 @@ private:
 
 	TUniquePtr<FObjectBindingTagCache> ObjectBindingTagCache;
 
-	TSharedPtr<UE::Sequencer::FSequencerEditorViewModel> ViewModel;
+	TSharedPtr<UE::Sequencer::FSequencerEditorViewModelImpl> ViewModel;
 
 	struct FCachedViewState
 	{
@@ -1660,8 +1666,6 @@ private:
 	/** Delegate that is called when selection limiting has been toggled on or off */
 	FOnViewportSelectionLimitedChanged OnSelectionLimitedChangedDelegate;
 
-	TSharedPtr<FSequencerFilterBar> FilterBar;
-
 	enum class ETimeControllerState
 	{
 		ReadyToPlay,
diff --git a/Engine/Source/Editor/Sequencer/Public/Filters/FilterBarDelegates.h b/Engine/Source/Editor/Sequencer/Public/Filters/FilterBarDelegates.h
new file mode 100644
--- /dev/null
+++ b/Engine/Source/Editor/Sequencer/Public/Filters/FilterBarDelegates.h
@@ -0,0 +1,13 @@
+﻿// Copyright Epic Games, Inc. All Rights Reserved.
+
+#pragma once
+
+#include "Delegates/Delegate.h"
+#include "MVVM/ViewModelPtr.h"
+
+namespace UE::Sequencer
+{
+class IOutlinerExtension;
+
+DECLARE_MULTICAST_DELEGATE_TwoParams(FChangeOutlinerExtensionFilterState, const TViewModelPtr<IOutlinerExtension>&, bool bIsFilteredOut);
+}
\ No newline at end of file
diff --git a/Engine/Source/Editor/Sequencer/Public/Filters/SequencerFilterData.h b/Engine/Source/Editor/Sequencer/Public/Filters/SequencerFilterData.h
--- a/Engine/Source/Editor/Sequencer/Public/Filters/SequencerFilterData.h
+++ b/Engine/Source/Editor/Sequencer/Public/Filters/SequencerFilterData.h
@@ -2,6 +2,7 @@
 
 #pragma once
 
+#include "FilterBarDelegates.h"
 #include "MovieSceneTrack.h"
 #include "MVVM/ViewModelPtr.h"
 
@@ -57,6 +58,9 @@ struct FSequencerFilterData
 	TMap<UE::Sequencer::TWeakViewModelPtr<UE::Sequencer::IOutlinerExtension>, UE::Sequencer::TWeakViewModelPtr<UE::Sequencer::ITrackExtension>> ResolvedTracks;
 	TMap<UE::Sequencer::TWeakViewModelPtr<UE::Sequencer::IOutlinerExtension>, TWeakObjectPtr<UMovieSceneTrack>> ResolvedTrackObjects;
 	TMap<UE::Sequencer::TWeakViewModelPtr<UE::Sequencer::IOutlinerExtension>, TWeakObjectPtr<>> ResolvedObjects;
+	
+	/** Broadcasts when the filter state of IOutlinerExtension is changed. */
+	UE::Sequencer::FChangeOutlinerExtensionFilterState& OnChangeFilterState() { return ChangeFilterStateEvent; }
 
 protected:
 	FString RawFilterText;
@@ -65,6 +69,9 @@ protected:
 
 	/** Nodes to be displayed in the UI */
 	TSet<UE::Sequencer::TWeakViewModelPtr<UE::Sequencer::IOutlinerExtension>> FilterInNodes;
+	
+	/** Broadcasts when the filter state of IOutlinerExtension is changed. */
+	UE::Sequencer::FChangeOutlinerExtensionFilterState ChangeFilterStateEvent;
 };
 
 #undef UE_API
