Sequencer: Add a toggle button to switch Sequencer between linked and instanced filtering. The button is invisible by default and can be enabled with the "Sequencer.UseLinkedFiltering 1" CVar.

Temporarily re-uses 2 icons from modeling toolkit for linked and instanced filtering. Those will be changed later. See UE-363318.

#Synced-CL 50395143

#Synced-CL 50395143

diff --git a/Engine/Source/Editor/EditorStyle/Private/StarshipStyle.cpp b/Engine/Source/Editor/EditorStyle/Private/StarshipStyle.cpp
--- a/Engine/Source/Editor/EditorStyle/Private/StarshipStyle.cpp
+++ b/Engine/Source/Editor/EditorStyle/Private/StarshipStyle.cpp
@@ -2347,6 +2347,19 @@ void FStarshipEditorStyle::FStyle::SetupSequencerStyles()
 		Set("Sequencer.AnimationOutliner.ItalicFont", DEFAULT_FONT("Italic", 10));
 
 		Set("Sequencer.Outliner.Indicators.TimeWarp", new IMAGE_BRUSH_SVG("Starship/Sequencer/TimeWarp_12", Icon12x12));
+		
+		FCheckBoxStyle BaseCheckBoxStyle = FAppStyle::Get().GetWidgetStyle<FCheckBoxStyle>("ToggleButtonCheckbox");
+		Set("Sequencer.LinkedFilterToggleButton", BaseCheckBoxStyle
+			.SetUncheckedImage(FSlateRoundedBoxBrush(FStyleColors::Transparent, 4.0f))
+			.SetUncheckedHoveredImage(FSlateRoundedBoxBrush(FStyleColors::Hover, 4.0f))
+			.SetUncheckedPressedImage(FSlateRoundedBoxBrush(FStyleColors::Hover, 4.0f))
+			.SetCheckedImage(FSlateRoundedBoxBrush(FStyleColors::Primary, 4.0f))
+			.SetCheckedHoveredImage(FSlateRoundedBoxBrush(FStyleColors::PrimaryHover, 4.0f))
+			.SetCheckedPressedImage(FSlateRoundedBoxBrush(FStyleColors::PrimaryHover, 4.0f))
+			.SetPadding(FMargin(2, 2))
+		);
+		Set("Sequencer.LinkedFiltering.LinkedIcon", new IMAGE_BRUSH("Starship/Sequencer/FilteringLinked_x20", Icon20x20));
+		Set("Sequencer.LinkedFiltering.InstancedIcon", new IMAGE_BRUSH("Starship/Sequencer/FilteringInstanced_x20", Icon20x20));
 
 		Set("Sequencer.ShotFilter", new IMAGE_BRUSH("Sequencer/FilteredArea", FVector2f(74.f, 74.f), FLinearColor::White, ESlateBrushTileType::Both));
 		Set("Sequencer.KeyMark", new IMAGE_BRUSH("Sequencer/KeyMark", FVector2f(3.f, 21.f), FLinearColor::White, ESlateBrushTileType::NoTile));
diff --git a/Engine/Source/Editor/Sequencer/Private/Filters/SequencerTrackFilterCommands.cpp b/Engine/Source/Editor/Sequencer/Private/Filters/SequencerTrackFilterCommands.cpp
--- a/Engine/Source/Editor/Sequencer/Private/Filters/SequencerTrackFilterCommands.cpp
+++ b/Engine/Source/Editor/Sequencer/Private/Filters/SequencerTrackFilterCommands.cpp
@@ -30,6 +30,8 @@ void FSequencerTrackFilterCommands::RegisterCommands()
 
 	UI_COMMAND(ActivateAllFilters, "Activate All Filters", "Activates all enabled filters", EUserInterfaceActionType::Button, FInputChord());
 	UI_COMMAND(DeactivateAllFilters, "Deactivate All Filters", "Deactivates all enabled filters", EUserInterfaceActionType::Button, FInputChord());
+	
+	UI_COMMAND(ToggleLinkedFiltering, "Toggle Linked Filtering", "Toggles whether filtering should use the filter state shared with other UI or its own instance.", EUserInterfaceActionType::ToggleButton, FInputChord());
 
 	UI_COMMAND(HideSelectedTracks, "Hide Selected Tracks", "Hide all selected tracks", EUserInterfaceActionType::Button, FInputChord(EModifierKey::Alt, EKeys::W));
 	UI_COMMAND(IsolateSelectedTracks, "Isolate Selected Tracks", "Isolate all selected tracks", EUserInterfaceActionType::Button, FInputChord(EModifierKey::Alt, EKeys::Q));
diff --git a/Engine/Source/Editor/Sequencer/Private/Filters/SequencerTrackFilterCommands.h b/Engine/Source/Editor/Sequencer/Private/Filters/SequencerTrackFilterCommands.h
--- a/Engine/Source/Editor/Sequencer/Private/Filters/SequencerTrackFilterCommands.h
+++ b/Engine/Source/Editor/Sequencer/Private/Filters/SequencerTrackFilterCommands.h
@@ -34,6 +34,8 @@ public:
 
 	TSharedPtr<FUICommandInfo> ActivateAllFilters;
 	TSharedPtr<FUICommandInfo> DeactivateAllFilters;
+	
+	TSharedPtr<FUICommandInfo> ToggleLinkedFiltering;
 
 	/** Hide/Isolate Filter Commands */
 
diff --git a/Engine/Source/Editor/Sequencer/Private/Filters/Utils/LinkedFilteringViewUtils.cpp b/Engine/Source/Editor/Sequencer/Private/Filters/Utils/LinkedFilteringViewUtils.cpp
new file mode 100644
--- /dev/null
+++ b/Engine/Source/Editor/Sequencer/Private/Filters/Utils/LinkedFilteringViewUtils.cpp
@@ -0,0 +1,42 @@
+﻿// Copyright Epic Games, Inc. All Rights Reserved.
+
+#include "LinkedFilteringViewUtils.h"
+
+#include "Filters/ViewModels/Linking/ILinkedFilterViewModel.h"
+#include "Misc/ConsoleVariables.h"
+#include "Templates/SharedPointer.h"
+
+namespace UE::Sequencer
+{
+void ToggleLinkedFiltering(TWeakPtr<ILinkedFilterViewModel> InWeakViewModel)
+{
+	const TSharedPtr<ILinkedFilterViewModel> ViewModelPin = InWeakViewModel.Pin();
+	if (!ViewModelPin)
+	{
+		return;
+	}
+	
+	static_assert(static_cast<int32>(ELinkedFilterMode::Count) == 2, "Update this switch");
+	switch (ViewModelPin->GetFilterMode())
+	{
+	case ELinkedFilterMode::Linked: ViewModelPin->SetFilterMode(ELinkedFilterMode::Instanced); break;
+	case ELinkedFilterMode::Instanced: ViewModelPin->SetFilterMode(ELinkedFilterMode::Linked); break;
+	default: checkNoEntry(); break;
+	}
+}
+
+bool IsLinkedFilteringActionChecked(TWeakPtr<ILinkedFilterViewModel> InWeakViewModel)
+{
+	const TSharedPtr<ILinkedFilterViewModel> ViewModelPin = InWeakViewModel.Pin();
+	return ViewModelPin && ViewModelPin->GetFilterMode() == ELinkedFilterMode::Linked;
+}
+
+TAutoConsoleVariable<bool> CVarSequencerUseLinkedFiltering(
+	TEXT("Sequencer.UseLinkedFiltering"), false, TEXT("Whether to use the new linked filtering system")
+	);
+
+bool IsLinkedFilteringFeatureEnabled()
+{
+	return CVarSequencerUseLinkedFiltering.GetValueOnAnyThread();
+}
+} // namespace UE::Sequencer
diff --git a/Engine/Source/Editor/Sequencer/Private/Filters/Utils/LinkedFilteringViewUtils.h b/Engine/Source/Editor/Sequencer/Private/Filters/Utils/LinkedFilteringViewUtils.h
new file mode 100644
--- /dev/null
+++ b/Engine/Source/Editor/Sequencer/Private/Filters/Utils/LinkedFilteringViewUtils.h
@@ -0,0 +1,18 @@
+﻿// Copyright Epic Games, Inc. All Rights Reserved.
+
+#pragma once
+
+#include "Templates/SharedPointerFwd.h"
+
+namespace UE::Sequencer
+{
+class ILinkedFilterViewModel;
+
+/** Utility for binding ToggleLinkedFiltering to FExecuteAction. */
+void ToggleLinkedFiltering(TWeakPtr<ILinkedFilterViewModel> InWeakViewModel);
+/** Utility for binding ToggleLinkedFiltering to FIsActionChecked. */
+bool IsLinkedFilteringActionChecked(TWeakPtr<ILinkedFilterViewModel> InWeakViewModel);
+
+/** @return Whether Sequencer should use linked filtering. */
+bool IsLinkedFilteringFeatureEnabled();
+}
diff --git a/Engine/Source/Editor/Sequencer/Private/Filters/Widgets/Linking/SFilterBarSwitcher.cpp b/Engine/Source/Editor/Sequencer/Private/Filters/Widgets/Linking/SFilterBarSwitcher.cpp
new file mode 100644
--- /dev/null
+++ b/Engine/Source/Editor/Sequencer/Private/Filters/Widgets/Linking/SFilterBarSwitcher.cpp
@@ -0,0 +1,39 @@
+﻿// Copyright Epic Games, Inc. All Rights Reserved.
+
+#include "SFilterBarSwitcher.h"
+
+#include "SFilterLinkStateSwitcher.h"
+#include "Filters/Widgets/SSequencerFilterBar.h"
+#include "Filters/ViewModels/Linking/LinkedFilterViewModel.h"
+
+namespace UE::Sequencer
+{
+void SFilterBarSwitcher::Construct(const FArguments& InArgs)
+{
+	const TSharedPtr<FLinkedFilterViewModel> LinkedFilterViewModel = InArgs._LinkedFilterViewModel;
+	check(LinkedFilterViewModel);
+	
+	ChildSlot
+	[
+		SAssignNew(WidgetSwitcher, SFilterLinkStateSwitcher)
+		.LinkedFilterViewModel(LinkedFilterViewModel)
+		.LinkedContent()
+		[
+			SNew(SSequencerFilterBar, LinkedFilterViewModel->GetLinkedFilterBarImpl())
+			.FilterBarLayout(InArgs._FilterBarLayout)
+			.AddMetaData<FTagMetaData>(FTagMetaData(TEXT("SequencerTrackFilters")))
+		]
+		.InstancedContent()
+		[
+			SNew(SSequencerFilterBar, LinkedFilterViewModel->GetInstancedFilterBarImpl())
+			.FilterBarLayout(InArgs._FilterBarLayout)
+			.AddMetaData<FTagMetaData>(FTagMetaData(TEXT("SequencerTrackFilters")))
+		]
+	];
+}
+
+FReply SFilterBarSwitcher::OnMouseButtonUp(const FGeometry& MyGeometry, const FPointerEvent& MouseEvent)
+{
+	return WidgetSwitcher->GetActiveWidget()->OnMouseButtonUp(MyGeometry, MouseEvent);
+}
+} // namespace UE::Sequencer
diff --git a/Engine/Source/Editor/Sequencer/Private/Filters/Widgets/Linking/SFilterBarSwitcher.h b/Engine/Source/Editor/Sequencer/Private/Filters/Widgets/Linking/SFilterBarSwitcher.h
new file mode 100644
--- /dev/null
+++ b/Engine/Source/Editor/Sequencer/Private/Filters/Widgets/Linking/SFilterBarSwitcher.h
@@ -0,0 +1,38 @@
+﻿// Copyright Epic Games, Inc. All Rights Reserved.
+
+#pragma once
+
+#include "Filters/SequencerFilterBar.h"
+#include "Widgets/DeclarativeSyntaxSupport.h"
+#include "Widgets/SCompoundWidget.h"
+
+namespace UE::Sequencer
+{
+class SFilterLinkStateSwitcher;
+class FLinkedFilterViewModel;
+
+/** Switches between SSequencerFilterBar instances based on the active ELinkedFilterMode. */
+class SFilterBarSwitcher : public SCompoundWidget
+{
+public:
+	
+	SLATE_BEGIN_ARGS(SFilterBarSwitcher) {}
+		/** View-model determines the ELinkedFilterMode driving the displayed content. */
+		SLATE_ARGUMENT(TSharedPtr<FLinkedFilterViewModel>, LinkedFilterViewModel)
+
+		/** The layout that determines how the filters are laid out */
+		SLATE_ARGUMENT(EFilterBarLayout, FilterBarLayout)
+	SLATE_END_ARGS()
+	
+	void Construct(const FArguments& InArgs);
+	
+	//~ Begin SWidget
+	virtual FReply OnMouseButtonUp(const FGeometry& MyGeometry, const FPointerEvent& MouseEvent) override;
+	//~ End SWidget
+	
+private:
+	
+	/** Determines the filter bar that is displayed. */
+	TSharedPtr<SFilterLinkStateSwitcher> WidgetSwitcher;
+};
+} // namespace UE::Sequencer
diff --git a/Engine/Source/Editor/Sequencer/Private/Filters/Widgets/Linking/SFilterLinkStateSwitcher.cpp b/Engine/Source/Editor/Sequencer/Private/Filters/Widgets/Linking/SFilterLinkStateSwitcher.cpp
new file mode 100644
--- /dev/null
+++ b/Engine/Source/Editor/Sequencer/Private/Filters/Widgets/Linking/SFilterLinkStateSwitcher.cpp
@@ -0,0 +1,45 @@
+﻿// Copyright Epic Games, Inc. All Rights Reserved.
+
+#include "SFilterLinkStateSwitcher.h"
+
+#include "Filters/ViewModels/Linking/ILinkedFilterViewModel.h"
+#include "Widgets/Layout/SWidgetSwitcher.h"
+
+namespace UE::Sequencer
+{
+void SFilterLinkStateSwitcher::Construct(const FArguments& InArgs)
+{
+	LinkedFilterViewModel = InArgs._LinkedFilterViewModel;
+	LinkedFilterViewModel->OnFilterModeChanged().AddSP(this, &SFilterLinkStateSwitcher::RefreshContent);
+	
+	ChildSlot
+	[
+		SAssignNew(WidgetSwitcher, SWidgetSwitcher)
+		+SWidgetSwitcher::Slot() [ InArgs._LinkedContent.Widget ]
+		+SWidgetSwitcher::Slot() [ InArgs._InstancedContent.Widget ]
+	];
+	
+	RefreshContent();
+}
+
+TSharedPtr<SWidget> SFilterLinkStateSwitcher::GetActiveWidget() const
+{
+	return WidgetSwitcher->GetActiveWidget();
+}
+
+int32 SFilterLinkStateSwitcher::GetActiveIndex() const
+{
+	static_assert(static_cast<int32>(ELinkedFilterMode::Count) == 2, "Update this switch");
+	switch (LinkedFilterViewModel->GetFilterMode())
+	{
+	case ELinkedFilterMode::Linked: return 0;
+	case ELinkedFilterMode::Instanced: return 1;
+	default: checkNoEntry(); return 0;
+	}
+}
+
+void SFilterLinkStateSwitcher::RefreshContent() const
+{
+	WidgetSwitcher->SetActiveWidgetIndex(GetActiveIndex());
+}
+} // namespace UE::Sequencer
diff --git a/Engine/Source/Editor/Sequencer/Private/Filters/Widgets/Linking/SFilterLinkStateSwitcher.h b/Engine/Source/Editor/Sequencer/Private/Filters/Widgets/Linking/SFilterLinkStateSwitcher.h
new file mode 100644
--- /dev/null
+++ b/Engine/Source/Editor/Sequencer/Private/Filters/Widgets/Linking/SFilterLinkStateSwitcher.h
@@ -0,0 +1,48 @@
+﻿// Copyright Epic Games, Inc. All Rights Reserved.
+
+#pragma once
+
+#include "Widgets/DeclarativeSyntaxSupport.h"
+#include "Widgets/SCompoundWidget.h"
+
+class SWidgetSwitcher;
+
+namespace UE::Sequencer
+{
+class ILinkedFilterViewModel;
+
+/** This widget switches it content based ELinkedFilterMode. */
+class SFilterLinkStateSwitcher : public SCompoundWidget
+{
+public:
+	
+	SLATE_BEGIN_ARGS(SFilterLinkStateSwitcher) {}
+		/** View-model determines the ELinkedFilterMode driving the displayed content. */
+		SLATE_ARGUMENT(TSharedPtr<ILinkedFilterViewModel>, LinkedFilterViewModel)
+
+		/** Content to display when in linked mode. */
+		SLATE_NAMED_SLOT(FArguments, LinkedContent)
+		/** Content to display when in instanced mode. */
+		SLATE_NAMED_SLOT(FArguments, InstancedContent)
+		
+	SLATE_END_ARGS()
+	
+	void Construct(const FArguments& InArgs);
+	
+	/** @return The active content widget. */
+	TSharedPtr<SWidget> GetActiveWidget() const;
+	
+private:
+	
+	/** View-model determines the ELinkedFilterMode driving the displayed content. */
+	TSharedPtr<ILinkedFilterViewModel> LinkedFilterViewModel;
+	
+	/** Used to switch between the content. */
+	TSharedPtr<SWidgetSwitcher> WidgetSwitcher;
+	
+	/** @return The widget index that should be active */
+	int32 GetActiveIndex() const;
+	/** Updates the widget content */
+	void RefreshContent() const;
+};
+} // namespace UE::Sequencer
diff --git a/Engine/Source/Editor/Sequencer/Private/Filters/Widgets/Linking/SLinkedFilteringToggleButton.cpp b/Engine/Source/Editor/Sequencer/Private/Filters/Widgets/Linking/SLinkedFilteringToggleButton.cpp
new file mode 100644
--- /dev/null
+++ b/Engine/Source/Editor/Sequencer/Private/Filters/Widgets/Linking/SLinkedFilteringToggleButton.cpp
@@ -0,0 +1,71 @@
+﻿// Copyright Epic Games, Inc. All Rights Reserved.
+
+#include "SLinkedFilteringToggleButton.h"
+
+#include "Filters/ViewModels/Linking/ILinkedFilterViewModel.h"
+#include "Widgets/Images/SImage.h"
+#include "Widgets/Input/SCheckBox.h"
+
+#define LOCTEXT_NAMESPACE "SLinkedFilteringToggleButton"
+
+namespace UE::Sequencer
+{
+void SLinkedFilteringToggleButton::Construct(
+	const FArguments& InArgs, const TSharedRef<ILinkedFilterViewModel>& InLinkedFilterViewModel
+	)
+{
+	LinkedFilterViewModel = InLinkedFilterViewModel;
+	
+	ChildSlot
+	[
+		SNew(SCheckBox)
+		.ToolTipText(this, &SLinkedFilteringToggleButton::GetToolTipText)
+		.Style(FAppStyle::Get(), TEXT("Sequencer.LinkedFilterToggleButton"))
+		.IsChecked(this, &SLinkedFilteringToggleButton::IsChecked)
+		.OnCheckStateChanged(this, &SLinkedFilteringToggleButton::OnButtonClicked)
+		[
+			SNew(SImage)
+			.Image(this, &SLinkedFilteringToggleButton::GetIcon)
+		]
+	];
+}
+
+ECheckBoxState SLinkedFilteringToggleButton::IsChecked() const
+{
+	return LinkedFilterViewModel->GetFilterMode() == ELinkedFilterMode::Linked
+		? ECheckBoxState::Checked
+		: ECheckBoxState::Unchecked;
+}
+
+FText SLinkedFilteringToggleButton::GetToolTipText() const
+{
+	static_assert(static_cast<int32>(ELinkedFilterMode::Count) == 2, "Update this switch");
+	switch (LinkedFilterViewModel->GetFilterMode())
+	{
+	case ELinkedFilterMode::Linked: return LOCTEXT("Tooltip.Linked", "Linked filtering.\nChanges made to filters affect all synced UI.");
+	case ELinkedFilterMode::Instanced: return LOCTEXT("Tooltip.Instanced", "Instanced filtering.\nChanges made to filters only affect this UI.");
+	default: checkNoEntry(); return FText::GetEmpty();
+	}
+}
+
+void SLinkedFilteringToggleButton::OnButtonClicked(ECheckBoxState CheckBoxState)
+{
+	const ELinkedFilterMode NewMode = CheckBoxState == ECheckBoxState::Checked
+		? ELinkedFilterMode::Linked
+		: ELinkedFilterMode::Instanced;
+	LinkedFilterViewModel->SetFilterMode(NewMode);
+}
+
+const FSlateBrush* SLinkedFilteringToggleButton::GetIcon() const
+{
+	static_assert(static_cast<int32>(ELinkedFilterMode::Count) == 2, "Update this switch");
+	switch (LinkedFilterViewModel->GetFilterMode())
+	{
+	case ELinkedFilterMode::Linked: return FAppStyle::Get().GetBrush("Sequencer.LinkedFiltering.LinkedIcon");
+	case ELinkedFilterMode::Instanced: return FAppStyle::Get().GetBrush("Sequencer.LinkedFiltering.InstancedIcon");;
+	default: checkNoEntry(); return nullptr;
+	}
+}
+} // namespace UE::Sequencer
+
+#undef LOCTEXT_NAMESPACE 
\ No newline at end of file
diff --git a/Engine/Source/Editor/Sequencer/Private/Filters/Widgets/Linking/SLinkedFilteringToggleButton.h b/Engine/Source/Editor/Sequencer/Private/Filters/Widgets/Linking/SLinkedFilteringToggleButton.h
new file mode 100644
--- /dev/null
+++ b/Engine/Source/Editor/Sequencer/Private/Filters/Widgets/Linking/SLinkedFilteringToggleButton.h
@@ -0,0 +1,32 @@
+﻿// Copyright Epic Games, Inc. All Rights Reserved.
+
+#pragma once
+
+#include "Widgets/DeclarativeSyntaxSupport.h"
+#include "Widgets/SCompoundWidget.h"
+
+namespace UE::Sequencer
+{
+class ILinkedFilterViewModel;
+
+/** Toggle button for switching between the linked and instanced filter modes. */
+class SLinkedFilteringToggleButton : public SCompoundWidget
+{
+public:
+	
+	SLATE_BEGIN_ARGS(SLinkedFilteringToggleButton){}
+	SLATE_END_ARGS()
+
+	void Construct(const FArguments& InArgs, const TSharedRef<ILinkedFilterViewModel>& InLinkedFilterViewModel);
+	
+private:
+	
+	/** Used to toggle between the filter modes. */
+	TSharedPtr<ILinkedFilterViewModel> LinkedFilterViewModel;
+
+	ECheckBoxState IsChecked() const;
+	FText GetToolTipText() const;
+	void OnButtonClicked(ECheckBoxState CheckBoxState);
+	const FSlateBrush* GetIcon() const;
+};
+} // namespace UE::Sequencer
diff --git a/Engine/Source/Editor/Sequencer/Private/MVVM/ViewModels/LinkedFilteringMediator.cpp b/Engine/Source/Editor/Sequencer/Private/MVVM/ViewModels/LinkedFilteringMediator.cpp
--- a/Engine/Source/Editor/Sequencer/Private/MVVM/ViewModels/LinkedFilteringMediator.cpp
+++ b/Engine/Source/Editor/Sequencer/Private/MVVM/ViewModels/LinkedFilteringMediator.cpp
@@ -3,6 +3,8 @@
 #include "LinkedFilteringMediator.h"
 
 #include "Filters/SequencerFilterBar.h"
+#include "Filters/SequencerTrackFilterCommands.h"
+#include "Filters/Utils/LinkedFilteringViewUtils.h"
 #include "Filters/ViewModels/Linking/ILinkedFilterViewModel.h"
 #include "Filters/ViewModels/Linking/LinkedFilterViewModel.h"
 #include "SequencerNodeTree.h"
@@ -40,13 +42,20 @@ FLinkedFilteringMediator::~FLinkedFilteringMediator()
 	InstancedFilterBar->OnRequestUpdate().RemoveAll(this);
 }
 
-void FLinkedFilteringMediator::BindCommands() const
+void FLinkedFilteringMediator::BindCommands(const TSharedRef<FUICommandList>& InSequencerCommandList) const
 {
 	const TSharedRef<FSequencerFilterBar> SharedFilterBar = SequencerFilterViewModel->GetLinkedFilterBarImpl();
 	SharedFilterBar->BindCommands();
 	
 	const TSharedRef<FSequencerFilterBar> InstancedFilterBar = SequencerFilterViewModel->GetInstancedFilterBarImpl();
 	InstancedFilterBar->BindCommands();
+	
+	InSequencerCommandList->MapAction(
+		FSequencerTrackFilterCommands::Get().ToggleLinkedFiltering ,
+		FExecuteAction::CreateLambda([WeakViewModel = SequencerFilterViewModel.ToWeakPtr()] { ToggleLinkedFiltering(WeakViewModel); }),
+		FCanExecuteAction::CreateLambda([]{ return IsLinkedFilteringFeatureEnabled(); }),
+		FIsActionChecked::CreateLambda([WeakViewModel = SequencerFilterViewModel.ToWeakPtr()] { return IsLinkedFilteringActionChecked(WeakViewModel); })
+		);
 }
 
 bool FLinkedFilteringMediator::ProcessCommandBindings(const FKeyEvent& InKeyEvent) const
diff --git a/Engine/Source/Editor/Sequencer/Private/MVVM/ViewModels/LinkedFilteringMediator.h b/Engine/Source/Editor/Sequencer/Private/MVVM/ViewModels/LinkedFilteringMediator.h
--- a/Engine/Source/Editor/Sequencer/Private/MVVM/ViewModels/LinkedFilteringMediator.h
+++ b/Engine/Source/Editor/Sequencer/Private/MVVM/ViewModels/LinkedFilteringMediator.h
@@ -5,6 +5,7 @@
 #include "MVVM/Extensions/IOutlinerExtension.h"
 #include "MVVM/ViewModels/ViewModel.h"
 
+class FUICommandList;
 class FSequencerNodeTree;
 
 namespace UE::Sequencer
@@ -32,7 +33,7 @@ public:
 	virtual ~FLinkedFilteringMediator() override;
 	
 	/** Binds the filter bar commands. */
-	void BindCommands() const;
+	void BindCommands(const TSharedRef<FUICommandList>& InSequencerCommandList) const;
 	
 	/** Process the commands of the active filter bar. */
 	bool ProcessCommandBindings(const FKeyEvent& InKeyEvent) const;
diff --git a/Engine/Source/Editor/Sequencer/Private/SSequencer.cpp b/Engine/Source/Editor/Sequencer/Private/SSequencer.cpp
--- a/Engine/Source/Editor/Sequencer/Private/SSequencer.cpp
+++ b/Engine/Source/Editor/Sequencer/Private/SSequencer.cpp
@@ -156,6 +156,9 @@
 #include "Widgets/Sidebar/SequencerSelectionDrawer.h"
 #include "WorkspaceMenuStructure.h"
 #include "WorkspaceMenuStructureModule.h"
+#include "Filters/Widgets/Linking/SFilterBarSwitcher.h"
+#include "MVVM/ViewModels/LinkedFilteringMediator.h"
+#include "MVVM/ViewModels/SequencerEditorViewModelImpl.h"
 
 #define LOCTEXT_NAMESPACE "Sequencer"
 
@@ -359,10 +362,7 @@ void SSequencer::Construct(const FArguments& InArgs, TSharedRef<FSequencer> InSe
 		SequencerWidgets.CreateTimeRangeSlider(TimeSliderControllerRef)
 	);
 
-	FilterRow = SNew(SSequencerSearchAndFilterRow)
-		.ExtenderForAddMenu_Lambda([this]{ return FExtender::Combine(AddMenuExtenders); })
-		.Sequencer(SequencerPtr);
-
+	// TODO UE-363287
 	if (const TSharedPtr<FSequencerFilterBar> FilterBar = GetFilterBar())
 	{
 		FilterBar->OnStateChanged().AddSP(this, &SSequencer::OnFilterBarStateChanged);
@@ -371,6 +371,11 @@ void SSequencer::Construct(const FArguments& InArgs, TSharedRef<FSequencer> InSe
 
 	OutlinerAndTrackAreaContent = ConstructOutlinerAndTrackArea(InArgs, TimeSliderArgs);
 
+	FilterRow = SNew(SSequencerSearchAndFilterRow)
+		.ExtenderForAddMenu_Lambda([this]{ return FExtender::Combine(AddMenuExtenders); })
+		.FilterViewModel(SequencerRef->GetViewModelImpl()->GetLinkedFilterViewModelImpl())
+		.Sequencer(SequencerPtr);
+
 	ChildSlot
 	[
 		SNew(SVerticalBox)
@@ -630,6 +635,8 @@ void SSequencer::RebuildFilterBarContent()
 
 TSharedRef<SWidget> SSequencer::ConstructFilterBarContent()
 {
+	using namespace UE::Sequencer;
+	
 	RebuildSearchAndFilterRow();
 	
 	if (!IsFilterBarVisible() || GetFilterBarLayout() == EFilterBarLayout::Horizontal)
@@ -637,7 +644,7 @@ TSharedRef<SWidget> SSequencer::ConstructFilterBarContent()
 		return OutlinerAndTrackAreaContent.ToSharedRef();
 	}
 
-	const TSharedRef<SSequencerFilterBar> FilterBarWidgetRef = FilterBarWidget.ToSharedRef();
+	const TSharedRef<SFilterBarSwitcher> FilterBarWidgetRef = FilterBarWidget.ToSharedRef();
 
 	return SNew(SSplitter)
 		.Orientation(Orient_Horizontal)
@@ -662,8 +669,9 @@ TSharedRef<SWidget> SSequencer::ConstructFilterBarContent()
 				GetSequencerSettings()->SetLastFilterBarSizeCoefficient(InNewCoefficient, false);
 			})
 		[
-			SFilterBarClippingHorizontalBox::WrapVerticalListWithHeading(FilterBarWidgetRef
-				, FPointerEventHandler::CreateSP(FilterBarWidgetRef, &SSequencerFilterBar::OnMouseButtonUp))
+			SFilterBarClippingHorizontalBox::WrapVerticalListWithHeading(
+				FilterBarWidgetRef, FPointerEventHandler::CreateSP(FilterBarWidgetRef, &SFilterBarSwitcher::OnMouseButtonUp)
+				)
 		]
 		+ SSplitter::Slot()
 		.Value(0.94f)
@@ -2985,13 +2993,19 @@ FReply SSequencer::OnKeyDown( const FGeometry& MyGeometry, const FKeyEvent& InKe
 {
 	// A toolkit tab is active, so direct all command processing to it
 	const TSharedPtr<FSequencer> SequencerPin = SequencerPtr.Pin();
-	if (SequencerPin && SequencerPin->GetCommandBindings()->ProcessCommandBindings(InKeyEvent))
+	if (!SequencerPin)
+	{
+		return FReply::Unhandled();
+	}
+	
+	if (SequencerPin->GetCommandBindings()->ProcessCommandBindings(InKeyEvent))
 	{
 		return FReply::Handled();
 	}
 
-	const TSharedPtr<FSequencerFilterBar> FilterBar = FilterBarWidget->GetFilterBar();
-	if (FilterBar && FilterBar->GetCommandList()->ProcessCommandBindings(InKeyEvent))
+	using namespace UE::Sequencer;
+	const TSharedPtr<FLinkedFilteringMediator> FilterMediator = SequencerPin->GetViewModelImpl()->GetFilterMediator();
+	if (FilterMediator && FilterMediator->ProcessCommandBindings(InKeyEvent))
 	{
 		return FReply::Handled();
 	}
@@ -4241,6 +4255,8 @@ void SSequencer::OnTrackFiltersChanged(const ESequencerFilterChange InChangeType
 
 void SSequencer::RebuildSearchAndFilterRow()
 {
+	using namespace UE::Sequencer;
+	
 	SearchAndFilterRow->ClearChildren();
 	SearchAndFilterRow->AddSlot()
 		.AutoHeight()
@@ -4248,10 +4264,12 @@ void SSequencer::RebuildSearchAndFilterRow()
 			FilterRow.ToSharedRef()
 		];
 	
-	const TSharedRef<FSequencerFilterBar> FilterBar = GetFilterBar().ToSharedRef();
-	FilterBarWidget = SNew(SSequencerFilterBar, FilterBar)
+	const TSharedPtr<FSequencer> SequencerPin = SequencerPtr.Pin();
+	check(SequencerPin);
+	FilterBarWidget = SNew(SFilterBarSwitcher)
 		.FilterBarLayout(GetFilterBarLayout())
-		.AddMetaData<FTagMetaData>(FTagMetaData(TEXT("SequencerTrackFilters")));
+		.AddMetaData<FTagMetaData>(FTagMetaData(TEXT("SequencerTrackFilters")))
+		.LinkedFilterViewModel(SequencerPin->GetViewModelImpl()->GetLinkedFilterViewModelImpl());
 
 	if (IsFilterBarVisible() && GetFilterBarLayout() == EFilterBarLayout::Horizontal)
 	{
diff --git a/Engine/Source/Editor/Sequencer/Private/SSequencer.h b/Engine/Source/Editor/Sequencer/Private/SSequencer.h
--- a/Engine/Source/Editor/Sequencer/Private/SSequencer.h
+++ b/Engine/Source/Editor/Sequencer/Private/SSequencer.h
@@ -61,6 +61,7 @@ namespace UE
 {
 namespace Sequencer
 {
+class SFilterBarSwitcher;
 class SSequencerSearchAndFilterRow;
 
 class SOutlinerView;
@@ -900,7 +901,7 @@ private:
 	TSharedPtr<FSequencerSelectionDrawer> SelectionDetailsDrawer;
 
 	/** Displays the filter pills. Placed at different UI locations depending on the layout mode. */
-	TSharedPtr<SSequencerFilterBar> FilterBarWidget;
+	TSharedPtr<UE::Sequencer::SFilterBarSwitcher> FilterBarWidget;
 	
 	/** Contains the filters widgets plus optionally the filter bar when in horizontal layout mode. */
 	TSharedPtr<SVerticalBox> SearchAndFilterRow;
diff --git a/Engine/Source/Editor/Sequencer/Private/Sequencer.cpp b/Engine/Source/Editor/Sequencer/Private/Sequencer.cpp
--- a/Engine/Source/Editor/Sequencer/Private/Sequencer.cpp
+++ b/Engine/Source/Editor/Sequencer/Private/Sequencer.cpp
@@ -12297,7 +12297,7 @@ void FSequencer::BindCommands()
 	// bind widget specific commands
 	SequencerWidget->BindCommands(SequencerCommandBindings, CurveEditorSharedBindings);
 	
-	ViewModel->GetFilterMediator()->BindCommands();
+	ViewModel->GetFilterMediator()->BindCommands(SequencerCommandBindings);
 }
 
 void FSequencer::BuildAddTrackMenu(class FMenuBuilder& MenuBuilder)
diff --git a/Engine/Source/Editor/Sequencer/Private/Widgets/Sequencer/SSequencerSearchAndFilterRow.cpp b/Engine/Source/Editor/Sequencer/Private/Widgets/Sequencer/SSequencerSearchAndFilterRow.cpp
--- a/Engine/Source/Editor/Sequencer/Private/Widgets/Sequencer/SSequencerSearchAndFilterRow.cpp
+++ b/Engine/Source/Editor/Sequencer/Private/Widgets/Sequencer/SSequencerSearchAndFilterRow.cpp
@@ -5,14 +5,17 @@
 #include "Filters/Menus/SequencerViewOptionsMenu.h"
 #include "Filters/SequencerFilterBar.h"
 #include "Filters/Utils/HideIsolateViewModelUtils.h"
+#include "Filters/Utils/LinkedFilteringViewUtils.h"
 #include "Filters/ViewModels/HideIsolateShowViewModel.h"
-#include "Filters/Widgets/SFilterBarIsolateHideShow.h"
+#include "Filters/ViewModels/Linking/ILinkedFilterViewModel.h"
+#include "Filters/ViewModels/Linking/LinkedFilterViewModel.h"
+#include "Filters/Widgets/Linking/SFilterLinkStateSwitcher.h"
+#include "Filters/Widgets/Linking/SLinkedFilteringToggleButton.h"
 #include "Filters/Widgets/SSearchAndFilterWidget.h"
 #include "Filters/Widgets/SSequencerSearchBox.h"
 #include "MVVM/Selection/Selection.h"
 #include "MVVM/Views/SOutlinerView.h"
 #include "SSequencer.h"
-#include "Sequencer.h"
 #include "Widgets/Sequencer/SAddToSequencerButton.h"
 #include "Widgets/Sequencer/SSequencerViewOptionsButton.h"
 
@@ -56,15 +59,8 @@ static FOutlinerExtensionsAttr MakeAllTracksAttr(const TSharedRef<FSequencer>& I
 void SSequencerSearchAndFilterRow::Construct(const FArguments& InArgs)
 {
 	WeakSequencer = InArgs._Sequencer;
-	
-	const TSharedRef<FSequencer> SequencerPin = WeakSequencer.Pin().ToSharedRef();
-	
-	const TSharedRef<FSequencerFilterBar> FilterBar = SequencerPin->GetFilterBar();
-	const TSharedRef<FHideIsolateShowViewModel> IsolateViewModel = MakeShared<FHideIsolateShowViewModel>(
-		FilterBar, FilterBar->GetHideIsolateFilter(),
-		FilterRowDetail::MakeSelectedTracksAttr(SequencerPin), FilterRowDetail::MakeAllTracksAttr(SequencerPin)
-		);
-	IsolateViewModel->OnRequestScrollToView().AddSP(this, &SSequencerSearchAndFilterRow::ScrollIntoView);
+	SequencerFilterModel = InArgs._FilterViewModel;
+	check(SequencerFilterModel);
 	
 	ChildSlot
 	[
@@ -80,13 +76,23 @@ void SSequencerSearchAndFilterRow::Construct(const FArguments& InArgs)
 			.ExtenderForAddMenu(InArgs._ExtenderForAddMenu)
 		]
 
+		+ SHorizontalBox::Slot()
+		.AutoWidth()
+		.VAlign(VAlign_Center)
+		[
+			SNew(SLinkedFilteringToggleButton, InArgs._FilterViewModel.ToSharedRef())
+			.Visibility_Lambda([]
+			{
+				return IsLinkedFilteringFeatureEnabled() ? EVisibility::Visible : EVisibility::Collapsed;
+			})
+		]
+
 		// Advanced Search Filter Combo Button
 		+ SHorizontalBox::Slot()
 		.FillWidth(1.f)
 		.VAlign(VAlign_Center)
 		[
-			SAssignNew(FilterRow, SSearchAndFilterWidget, FilterBar, IsolateViewModel)
-			.OnSearchTextChanged(this, &SSequencerSearchAndFilterRow::OnSearchTextChanged)
+			MakeFilterContent(InArgs)
 		]
 		
 		// View Options Combo Button
@@ -103,7 +109,43 @@ void SSequencerSearchAndFilterRow::Construct(const FArguments& InArgs)
 
 TSharedPtr<SSequencerSearchBox> SSequencerSearchAndFilterRow::GetActiveSearchBox() const
 {
-	return FilterRow->GetSearchBox();
+	static_assert(static_cast<int32>(ELinkedFilterMode::Count) == 2, "Update this switch");
+	switch (SequencerFilterModel->GetFilterMode())
+	{
+	case ELinkedFilterMode::Linked: return LinkedFilterWidgets->GetSearchBox();
+	case ELinkedFilterMode::Instanced: return InstancedFilterWidgets->GetSearchBox();
+	default: checkNoEntry(); return nullptr;
+	}
+}
+
+TSharedRef<SWidget> SSequencerSearchAndFilterRow::MakeFilterContent(const FArguments& InArgs)
+{
+	const auto MakeHideIsolateViewModel = [this](const TSharedRef<FSequencerFilterBar>& InFilterBar)
+	{
+		const TSharedRef<FSequencer> SequencerPin = WeakSequencer.Pin().ToSharedRef();
+		const TSharedRef<FHideIsolateShowViewModel> IsolateViewModel = MakeShared<FHideIsolateShowViewModel>(
+			InFilterBar, InFilterBar->GetHideIsolateFilter(),
+			FilterRowDetail::MakeSelectedTracksAttr(SequencerPin), FilterRowDetail::MakeAllTracksAttr(SequencerPin)
+			);
+		IsolateViewModel->OnRequestScrollToView().AddSP(this, &SSequencerSearchAndFilterRow::ScrollIntoView);
+		return IsolateViewModel;
+	};
+	
+	const TSharedPtr<FLinkedFilterViewModel> LinkedFilterViewModel = InArgs._FilterViewModel;
+	TSharedRef<FSequencerFilterBar> LinkedFilterBar = LinkedFilterViewModel->GetLinkedFilterBarImpl();
+	TSharedRef<FSequencerFilterBar> InstancedFilterBar = LinkedFilterViewModel->GetInstancedFilterBarImpl();
+	return SNew(SFilterLinkStateSwitcher)
+		.LinkedFilterViewModel(InArgs._FilterViewModel)
+		.LinkedContent()
+		[
+			SAssignNew(LinkedFilterWidgets, SSearchAndFilterWidget, LinkedFilterBar, MakeHideIsolateViewModel(LinkedFilterBar))
+			.OnSearchTextChanged(this, &SSequencerSearchAndFilterRow::OnSearchTextChanged)
+		]
+		.InstancedContent()
+		[
+			SAssignNew(InstancedFilterWidgets, SSearchAndFilterWidget, InstancedFilterBar, MakeHideIsolateViewModel(InstancedFilterBar))
+			.OnSearchTextChanged(this, &SSequencerSearchAndFilterRow::OnSearchTextChanged)
+		];
 }
 
 void SSequencerSearchAndFilterRow::OnSearchTextChanged(const FText& InSearchText)
diff --git a/Engine/Source/Editor/Sequencer/Private/Widgets/Sequencer/SSequencerSearchAndFilterRow.h b/Engine/Source/Editor/Sequencer/Private/Widgets/Sequencer/SSequencerSearchAndFilterRow.h
--- a/Engine/Source/Editor/Sequencer/Private/Widgets/Sequencer/SSequencerSearchAndFilterRow.h
+++ b/Engine/Source/Editor/Sequencer/Private/Widgets/Sequencer/SSequencerSearchAndFilterRow.h
@@ -13,6 +13,7 @@ class FSequencerFilterBar;
 
 namespace UE::Sequencer
 {
+class FLinkedFilterViewModel;
 class IOutlinerExtension;
 class SSearchAndFilterWidget;
 
@@ -25,6 +26,12 @@ public:
 		
 		/** Used to generate the subwidgets */
 		SLATE_ARGUMENT(TWeakPtr<FSequencer>, Sequencer)
+
+		/** 
+		 * Determines the filtering widgets displayed. 
+		 * Some widgets are stateful, e.g. the search box keeps a search history that should be different between the modes.
+		 */
+		SLATE_ARGUMENT(TSharedPtr<FLinkedFilterViewModel>, FilterViewModel)
 		
 		/** External extenders for the add button menu. */
 		SLATE_ATTRIBUTE(TSharedPtr<FExtender>, ExtenderForAddMenu)
@@ -43,6 +50,16 @@ private:
 	
 	/** The widget displaying the search bar, and other filter widgets. */
 	TSharedPtr<SSearchAndFilterWidget> FilterRow;
+
+	/** Filter model for Sequencer. */
+	TSharedPtr<FLinkedFilterViewModel> SequencerFilterModel;
+	/** Content displayed when ELinkedFilterMode::Linked */
+	TSharedPtr<SSearchAndFilterWidget> LinkedFilterWidgets;
+	/** Content displayed when ELinkedFilterMode::Instanced */
+	TSharedPtr<SSearchAndFilterWidget> InstancedFilterWidgets;
+	
+	/** @return The filtering widget. */
+	TSharedRef<SWidget> MakeFilterContent(const FArguments& InArgs);
 	
 	/** Propagates changes in the search widget to the filter system */
 	void OnSearchTextChanged(const FText& InSearchText);
