Sequencer: Move OnActorAddedToSequencerEvent broadcast to after the possessable has been converted to a spawnable

#rnx

#Synced-CL 50515454

diff --git a/Engine/Source/Editor/Sequencer/Private/Sequencer.cpp b/Engine/Source/Editor/Sequencer/Private/Sequencer.cpp
--- a/Engine/Source/Editor/Sequencer/Private/Sequencer.cpp
+++ b/Engine/Source/Editor/Sequencer/Private/Sequencer.cpp
@@ -6400,8 +6400,6 @@ void FSequencer::OnNewActorsDropped(const TArray<UObject*>& DroppedObjects, cons
 				FGuid PossessableGuid = FSequencerUtilities::CreateBinding(AsShared(), *NewActor);
 				FGuid NewGuid = PossessableGuid;
 
-				OnActorAddedToSequencerEvent.Broadcast(NewActor, PossessableGuid);
-
 				if (bAddSpawnable || bAddReplaceable)
 				{
 					TSubclassOf<UMovieSceneCustomBinding> CustomBindingClass = bAddSpawnable ? UMovieSceneSpawnableActorBinding::StaticClass() : UMovieSceneReplaceableActorBinding::StaticClass();
@@ -6440,6 +6438,8 @@ void FSequencer::OnNewActorsDropped(const TArray<UObject*>& DroppedObjects, cons
 					}
 				}
 
+				OnActorAddedToSequencerEvent.Broadcast(NewActor, NewGuid);
+
 				if (NewActor && (NewActor->GetClass() == ACameraRig_Rail::StaticClass() || NewActor->GetClass() == ACameraRig_Crane::StaticClass()))
 				{
 					ACineCameraActor* OutActor;
