Sequencer: Extract widget creation from FCurveEditorExtension to FSequencerCurveEditorPresenter.

This introduces SOldCurveEditorTreeContent, which is the currently latest state of 5.8 Curve Editor UI.
In a follow up change, a new CVar will be introduced that wil switch between that legacy UI and the new linked filtering UI.

#Synced-CL 50532451

diff --git a/Engine/Source/Editor/Sequencer/Private/MVVM/CurveEditorExtension.cpp b/Engine/Source/Editor/Sequencer/Private/MVVM/CurveEditorExtension.cpp
--- a/Engine/Source/Editor/Sequencer/Private/MVVM/CurveEditorExtension.cpp
+++ b/Engine/Source/Editor/Sequencer/Private/MVVM/CurveEditorExtension.cpp
@@ -8,18 +8,17 @@
 #include "Framework/Docking/TabManager.h"
 #include "IPropertyRowGenerator.h"
 #include "IStructureDetailsView.h"
-#include "Misc/TransportControlsHelper.h"
 #include "MVVM/ViewModels/ViewModelIterators.h"
 #include "MVVM/ViewModels/SequencerEditorViewModel.h"
 #include "MVVM/ViewModels/ChannelModel.h"
 #include "MVVM/Extensions/IOutlinerExtension.h"
 #include "CurveEditorAxis.h"
 #include "ICurveEditorExtension.h"
-#include "LevelSequence.h"
 #include "SCurveEditorView.h"
 #include "SCurveEditorPanel.h"
 #include "SCurveEditorToolProperties.h"
 #include "SCurveKeyDetailPanel.h"
+#include "SPopoutTabInlineContent.h"
 #include "SSequencerTreeFilterStatusBar.h"
 #include "STemporarilyFocusedSpinBox.h"
 #include "Sequencer.h"
@@ -28,7 +27,6 @@
 #include "Modification/Resolution/CurveMetaDataIdentifiers.h"
 #include "Modification/Resolution/CurveModelLookUpInfo.h"
 #include "Modification/Utils/ScopedSelectionChange.h"
-#include "SPopoutTabInlineContent.h"
 #include "Toolkits/IToolkitHost.h"
 #include "Tree/SCurveEditorTree.h"
 #include "Tree/SCurveEditorTreeFilterStatusBar.h"
@@ -40,6 +38,7 @@
 #include "Widgets/Layout/SScrollBorder.h"
 #include "WorkspaceMenuStructure.h"
 #include "WorkspaceMenuStructureModule.h"
+#include "Views/CurveEditor/CurveEditorWidgetOwner.h"
 
 #define LOCTEXT_NAMESPACE "SequencerCurveEditorExtension"
 
@@ -375,8 +374,6 @@ void FCurveEditorExtension::CreateCurveEditor(const FTimeSliderArgs& TimeSliderA
 		return;
 	}
 
-	const TSharedRef<FSequencer> SequencerRef = Sequencer.ToSharedRef();
-
 	// Create the curve editor;
 	{
 		FCurveEditorInitParams CurveEditorInitParams;
@@ -388,7 +385,7 @@ void FCurveEditorExtension::CreateCurveEditor(const FTimeSliderArgs& TimeSliderA
 		CurveEditorInitParams.ResolveCurveModelDelegate.BindStatic(&ResolveCurveEditorModel, Sequencer.ToWeakPtr());
 		
 		CurveEditorModel = MakeShared<FSequencerCurveEditor>(Sequencer, TimeSliderArgs.NumericTypeInterface);
-		CurveEditorModel->SetBounds(MakeUnique<FSequencerCurveEditorBounds>(SequencerRef));
+		CurveEditorModel->SetBounds(MakeUnique<FSequencerCurveEditorBounds>(Sequencer.ToSharedRef()));
 		CurveEditorModel->InitCurveEditor(CurveEditorInitParams);
 
 		CurveEditorModel->InputSnapEnabledAttribute = MakeAttributeLambda([SequencerSettings] { return SequencerSettings->GetIsSnapEnabled(); });
@@ -407,227 +404,38 @@ void FCurveEditorExtension::CreateCurveEditor(const FTimeSliderArgs& TimeSliderA
 	// aware of our custom bounds logic. Currently the range the bar displays is tied to Sequencer 
 	// timeline and not the Bounds, so we need a way of changing it to look at the Bounds but only for 
 	// the Curve Editor time slider controller. We want everything else to just pass through though.
-	TSharedRef<ITimeSliderController> CurveEditorTimeSliderController = MakeShared<FSequencerCurveEditorTimeSliderController>(
-			TimeSliderArgs, Sequencer, CurveEditorModel.ToSharedRef());
+	const TSharedRef<ITimeSliderController> CurveEditorTimeSliderController = MakeShared<FSequencerCurveEditorTimeSliderController>(
+		TimeSliderArgs, Sequencer, CurveEditorModel.ToSharedRef());
 	
-	PlayTimeDisplay = StaticCastSharedRef<STemporarilyFocusedSpinBox<double>>(Sequencer->MakePlayTimeDisplay());
-
-	CurveEditorTreeView = SNew(SCurveEditorTree, CurveEditorModel);
-	CurveEditorPanel = SNew(SCurveEditorPanel, CurveEditorModel.ToSharedRef())
-		// Grid lines match the color specified in FSequencerTimeSliderController::OnPaintViewArea
-		.GridLineTint(FLinearColor(0.f, 0.f, 0.f, 0.3f))
-		.ExternalTimeSliderController(CurveEditorTimeSliderController)
-		.MinimumViewPanelHeight(0.f)
-		.TabManager(TabManager)
-		.DisabledTimeSnapTooltip(LOCTEXT("CurveEditorTimeSnapDisabledTooltip", "Time Snapping is currently driven by Sequencer."))
-		.TreeContent()
-		[
-			SNew(SVerticalBox)
-
-			+ SVerticalBox::Slot()
-			.AutoHeight()
-			[
-				SNew(SBorder)
-				.BorderImage(FAppStyle::GetBrush("ToolPanel.GroupBorder"))
-				.Clipping(EWidgetClipping::ClipToBounds)
-				[
-					SAssignNew(CurveEditorSearchBox, SCurveEditorTreeTextFilter, CurveEditorModel)
-				]
-			]
-
-			+ SVerticalBox::Slot()
-			[
-				SNew(SOverlay)
-
-				+ SOverlay::Slot()
-				[
-					SNew(SScrollBorder, CurveEditorTreeView.ToSharedRef())
-					[
-						CurveEditorTreeView.ToSharedRef()
-					]
-				]
-
-				+ SOverlay::Slot()
-				.VAlign(VAlign_Bottom)
-				[
-					SNew(SHorizontalBox)
-
-					+ SHorizontalBox::Slot()
-					.AutoWidth()
-					[
-						SAssignNew(CurveEditorTreeFilterStatusBar, SCurveEditorTreeFilterStatusBar, CurveEditorModel)
-						.Visibility(EVisibility::Hidden) // Initially hidden, visible on hover of the info button
-					]
-				]
-			]
-
-			+ SVerticalBox::Slot()
-			.AutoHeight()
-			[
-				SNew(SBorder)
-				.BorderImage(FAppStyle::GetBrush("ToolPanel.GroupBorder"))
-				.Clipping(EWidgetClipping::ClipToBounds)
-				.Visibility(this, &FCurveEditorExtension::GetPopoutTransportControlsVisibility)
-				[
-					SNew(SHorizontalBox)
-
-					+ SHorizontalBox::Slot()
-					.AutoWidth()
-					.VAlign(VAlign_Center)
-					.HAlign(HAlign_Left)
-					[
-						SNew(SButton)
-						.VAlign(EVerticalAlignment::VAlign_Center)
-						.ButtonStyle(FAppStyle::Get(), "SimpleButton")
-						.ToolTipText_Lambda([this] { return LOCTEXT("ShowStatus", "Show Status"); })
-						.ContentPadding(FMargin(1, 0))
-						.OnHovered_Lambda([this] { CurveEditorTreeFilterStatusBar->ShowStatusBar(); })
-						.OnUnhovered_Lambda([this] { CurveEditorTreeFilterStatusBar->FadeOutStatusBar(); })
-						.OnClicked_Lambda([this] { CurveEditorTreeFilterStatusBar->HideStatusBar(); return FReply::Handled(); })
-						[
-							SNew(SImage)
-							.ColorAndOpacity(FSlateColor::UseForeground())
-							.Image(FAppStyle::Get().GetBrush("Icons.Info.Small"))
-						]
-					]
-
-					+ SHorizontalBox::Slot()
-					.HAlign(HAlign_Fill)
-					[
-						MakePopoutTransportControls(SequencerRef, TabManager)
-					]
-
-					+ SHorizontalBox::Slot()
-					.AutoWidth()
-					.VAlign(VAlign_Center)
-					.HAlign(HAlign_Right)
-					[
-						SNew(SButton)
-						.VAlign(EVerticalAlignment::VAlign_Center)
-						.ButtonStyle(FAppStyle::Get(), "NoBorder")
-						.ContentPadding(FMargin(1, 0))
-						[
-							SNew(SHorizontalBox)
-
-							+ SHorizontalBox::Slot()
-							.AutoWidth()
-							.VAlign(VAlign_Center)
-							.HAlign(HAlign_Right)
-							.Padding(FMargin(3.f, 0.f, 0.f, 0.f))
-							[
-								SNew(SBorder)
-								.BorderImage(nullptr)
-								[
-									PlayTimeDisplay.ToSharedRef()
-								]
-							]
-						]
-					]
-				]
-			]
-		];
-
-	// Register an instanced custom property type layout to handle converting FFrameNumber from Tick Resolution to Display Rate.
-	CurveEditorPanel->GetKeyDetailsView()->GetPropertyRowGenerator()->RegisterInstancedCustomPropertyTypeLayout(
-			"FrameNumber", 
-			FOnGetPropertyTypeCustomizationInstance::CreateSP(SequencerRef, &FSequencer::MakeFrameNumberDetailsCustomization));
-	CurveEditorPanel->GetToolPropertiesPanel()->GetStructureDetailsView()->GetDetailsView()->RegisterInstancedCustomPropertyTypeLayout(
-		"FrameNumber",
-		FOnGetPropertyTypeCustomizationInstance::CreateSP(SequencerRef, &FSequencer::MakeFrameNumberDetailsCustomization));
-
-	// And jump to the Curve Editor tree search if you have the Curve Editor focused
-	CurveEditorModel->GetCommands()->MapAction(
-		FSequencerCommands::Get().QuickTreeSearch,
-		FExecuteAction::CreateLambda([this] { FSlateApplication::Get().SetKeyboardFocus(CurveEditorSearchBox, EFocusCause::SetDirectly); })
-	);
-
-	CurveEditorModel->GetCommands()->MapAction(
-		FSequencerCommands::Get().ToggleShowGotoBox,
-		FExecuteAction::CreateLambda([this] { PlayTimeDisplay->Setup(); FSlateApplication::Get().SetKeyboardFocus(PlayTimeDisplay, EFocusCause::SetDirectly); })
-	);
-
-	CurveEditorWidget = SNew(SSequencerCurveEditor, CurveEditorPanel.ToSharedRef(), Sequencer);
-
-	CurveEditorPanel->OnFilterClassChanged.BindRaw(this, &FCurveEditorExtension::FilterClassChanged);
-
-	// Check to see if the tab is already opened due to the saved window layout.
-	if (const TSharedPtr<SDockTab> ExistingCurveEditorTab = TabManager->FindExistingLiveTab(CurveEditorTabName))
-	{
-		ExistingCurveEditorTab->SetContent(CurveEditorWidget.ToSharedRef());
-	}
-}
-
-void FCurveEditorExtension::FilterClassChanged()
-{
-	TSharedPtr<FSequencerEditorViewModel> OwnerModel = WeakOwnerModel.Pin();
-	if (!ensure(OwnerModel))
-	{
-		return;
-	}
-
-	TSharedPtr<FSequencer> Sequencer = OwnerModel->GetSequencerImpl();
-	if (!ensure(Sequencer))
-	{
-		return;
-	}
-
-	if (CurveEditorPanel)
-	{
-		TSharedPtr<SCurveEditorFilterPanel> FilterPanel = CurveEditorPanel->GetFilterPanel();
-		if (FilterPanel)
-		{
-			TWeakPtr<ISequencer> WeakSequencer(Sequencer);
-			FilterPanel->GetDetailsView()->RegisterInstancedCustomPropertyTypeLayout(
-				"FrameNumber",
-				FOnGetPropertyTypeCustomizationInstance::CreateSP(Sequencer.ToSharedRef(), &FSequencer::MakeFrameNumberDetailsCustomization));
-		}
-	}
+	WidgetOwner = MakePimpl<FCurveEditorWidgetOwner>(
+		Sequencer.ToSharedRef(), CurveEditorModel.ToSharedRef(), CurveEditorTimeSliderController
+		);
 }
 
 void FCurveEditorExtension::OpenCurveEditor()
 {
-	const TSharedPtr<FTabManager> TabManager = GetTabManager();
-	if (!ensure(TabManager))
-	{
-		return;
-	}
-
-	// Request the Tab Manager invoke the tab. This will spawn the tab if needed, otherwise pull it to focus. This assumes
-	// that the Toolkit Host's Tab Manager has already registered a tab with a NullWidget for content.
-	if (const TSharedPtr<SDockTab> CurveEditorTab = TabManager->TryInvokeTab(CurveEditorTabName))
+	if (WidgetOwner)
 	{
-		CurveEditorTab->SetContent(CurveEditorWidget.ToSharedRef());
-
-		const FSlateIcon SequencerGraphIcon = FSlateIcon(FAppStyle::GetAppStyleSetName(), "GenericCurveEditor.TabIcon");
-		CurveEditorTab->SetTabIcon(SequencerGraphIcon.GetIcon());
-
-		CurveEditorTab->SetLabel(LOCTEXT("SequencerMainGraphEditorTitle", "Sequencer Curves"));
-
-		CurveEditorModel->ZoomToFit();
+		WidgetOwner->OpenCurveEditor();
 	}
 }
 
 bool FCurveEditorExtension::IsCurveEditorOpen() const
 {
-	if (const TSharedPtr<FTabManager> TabManager = GetTabManager())
-	{
-		return TabManager->FindExistingLiveTab(CurveEditorTabName).IsValid();
-	}
-	return false;
+	return WidgetOwner && WidgetOwner->IsCurveEditorOpen();
 }
 
 void FCurveEditorExtension::CloseCurveEditor()
 {
-	const TSharedPtr<FTabManager> TabManager = GetTabManager();
-	if (!TabManager.IsValid())
+	if (WidgetOwner)
 	{
-		return;
+		WidgetOwner->CloseCurveEditor();
 	}
+}
 
-	if (const TSharedPtr<SDockTab> CurveEditorTab = TabManager->FindExistingLiveTab(CurveEditorTabName))
-	{
-		CurveEditorTab->RequestCloseTab();
-	}
+TSharedPtr<SCurveEditorTree> FCurveEditorExtension::GetCurveEditorTreeView() const
+{
+	return WidgetOwner ? WidgetOwner->GetTreeView() : nullptr;
 }
 
 FKeyAttributes FCurveEditorExtension::GetDefaultKeyAttributes() const
@@ -650,7 +458,7 @@ FKeyAttributes FCurveEditorExtension::GetDefaultKeyAttributes() const
 static bool bSyncSelectionRequested = false;
 void FCurveEditorExtension::RequestSyncSelection()
 {
-	if (bSyncSelectionRequested)
+	if (bSyncSelectionRequested || !WidgetOwner)
 	{
 		return;
 	}
@@ -691,11 +499,11 @@ void FCurveEditorExtension::RequestSyncSelection()
 			return;
 		}
 
-		FCurveEditorExtension* This = RootViewModel->CastDynamic<FCurveEditorExtension>();
-		if (This)
+		
+		if (FCurveEditorExtension* This = RootViewModel->CastDynamic<FCurveEditorExtension>())
 		{
 			This->SyncSelectionHandle.Invalidate();
-			This->SyncSelection();
+			This->WidgetOwner->SyncSelection();
 		}
 
 		// Order matters: this appends a sub-transaction for the selection change...
@@ -704,127 +512,6 @@ void FCurveEditorExtension::RequestSyncSelection()
 		Transaction.Reset();
 	});
 }
-static void GetAllChildrenRecursively(TSharedPtr<SCurveEditorTree>& CurveEditorTreeView, const FCurveEditorTreeItemID Item, int32& LastAddedItem,
-	TMap<const FCurveEditorTreeItemID, int32>& ItemToSortOrder)
-{
-	ItemToSortOrder.Add(Item, LastAddedItem++);
-
-	TArray<FCurveEditorTreeItemID> Children;
-
-	CurveEditorTreeView->GetTreeItemChildren(Item, Children);
-	// Add them to the output array
-	for (const FCurveEditorTreeItemID ChildItemID : Children)
-	{
-		GetAllChildrenRecursively(CurveEditorTreeView, ChildItemID, LastAddedItem, ItemToSortOrder);
-	}
-}
-
-//we now sync to the first non-channel in the tree hierarchy, this is needed for hosted channels with control rigs
-//we need to sort the hierarchy each time since the curve id's being used previously were not consistent in order of viewing/creationg, persistent now for undo/redo
-void FCurveEditorExtension::SyncSelection()
-{
-	if (!ensure(CurveEditorModel && CurveEditorTreeView))
-	{
-		return;
-	}
-
-	TSharedPtr<FSequencerEditorViewModel> OwnerModel = WeakOwnerModel.Pin();
-	if (!ensure(OwnerModel))
-	{
-		return;
-	}
-
-	TSharedPtr<ISequencer> Sequencer = OwnerModel->GetSequencer();
-	if (!ensure(Sequencer))
-	{
-		return;
-	}
-
-	CurveEditorModel->SuspendBroadcast();
-
-	CurveEditorTreeView->ClearSelection();
-
-	//map of tree items id's to their actually sort order, need to use this since tree items ids are not in order anymore
-	TMap<const FCurveEditorTreeItemID, int32>  ItemToSortOrder;
-	int32 LastAddedItem = 0;
-	TArray<FCurveEditorTreeItemID> Children;
-	TArrayView<const FCurveEditorTreeItemID> RootItems = CurveEditorTreeView->GetRootItems();
-
-
-	for (const FCurveEditorTreeItemID RootItemID : RootItems)
-	{
-		GetAllChildrenRecursively(CurveEditorTreeView, RootItemID, LastAddedItem, ItemToSortOrder);
-	}
-
-	auto GetSortOrder = [&ItemToSortOrder](FCurveEditorTreeItemID InTreeItemID) -> int32
-	{
-		if (const int32* SortOrder = ItemToSortOrder.Find(InTreeItemID))
-		{
-			return *SortOrder;
-		}
-		return INDEX_NONE;
-	};
-
-	FCurveEditorTreeItemID FirstCurveEditorTreeItemID;
-	int32 FirstSortOrder = INDEX_NONE;
-	bool bSelectedIsChannel = true;
-
-	for (TViewModelPtr<IOutlinerExtension> SelectedItem : OwnerModel->GetSelection()->Outliner)
-	{
-		if (TViewModelPtr<ICurveEditorTreeItemExtension> CurveEditorItem = SelectedItem.ImplicitCast())
-		{
-			FCurveEditorTreeItemID CurveEditorTreeItem = CurveEditorItem->GetCurveEditorItemID();
-			if (CurveEditorTreeItem != FCurveEditorTreeItemID::Invalid())
-			{
-				if (!CurveEditorTreeView->IsItemSelected(CurveEditorTreeItem))
-				{
-					CurveEditorTreeView->SetItemSelection(CurveEditorTreeItem, true);
-					bool bIsChannel = SelectedItem.AsModel()->IsA<FChannelModel>() || SelectedItem.AsModel()->IsA<FChannelGroupOutlinerModel>()
-					|| SelectedItem.AsModel()->IsA<FChannelGroupModel>();
-					if (!FirstCurveEditorTreeItemID.IsValid())
-					{
-						FirstCurveEditorTreeItemID = CurveEditorTreeItem;
-						FirstSortOrder = GetSortOrder(CurveEditorTreeItem);
-						bSelectedIsChannel = bIsChannel;
-					}
-					else
-					{
-						int32 SortOrder = GetSortOrder(CurveEditorTreeItem);
-						if (bIsChannel == false)
-						{
-							if (bSelectedIsChannel == true)
-							{
-								FirstCurveEditorTreeItemID = CurveEditorTreeItem;
-								FirstSortOrder = SortOrder;
-								bSelectedIsChannel = false;
-							}
-							else if (SortOrder < FirstSortOrder)
-							{
-								FirstCurveEditorTreeItemID = CurveEditorTreeItem;
-								FirstSortOrder = SortOrder;
-								bSelectedIsChannel = false;
-							}
-						}
-						else if (bSelectedIsChannel == true) //its a channel only sort if selected is a non channel
-						{
-							if (SortOrder < FirstSortOrder)
-							{
-								FirstCurveEditorTreeItemID = CurveEditorTreeItem;
-								FirstSortOrder = SortOrder;
-							}
-						}
-					}
-				}
-			}
-		}
-	}
-	if (FirstCurveEditorTreeItemID.IsValid())
-	{
-		CurveEditorTreeView->RequestScrollIntoView(FirstCurveEditorTreeItemID);
-	}
-
-	CurveEditorModel->ResumeBroadcast();
-}
 
 void FCurveEditorExtension::CleanUpPendingSyncSelection()
 {
@@ -836,31 +523,6 @@ void FCurveEditorExtension::CleanUpPendingSyncSelection()
 	}
 }
 
-EVisibility FCurveEditorExtension::GetPopoutTransportControlsVisibility() const
-{
-	const TSharedPtr<ISequencer> Sequencer = GetSequencer();
-	if (!Sequencer.IsValid())
-	{
-		return EVisibility::Collapsed;
-	}
-
-	const UMovieSceneSequence* const RootSequence = Sequencer->GetRootMovieSceneSequence();
-	if (!RootSequence || !RootSequence->IsA<ULevelSequence>())
-	{
-		return EVisibility::Visible;
-	}
-
-	const USequencerSettings* const SequencerSettings = Sequencer->GetSequencerSettings();
-	if (!SequencerSettings)
-	{
-		return EVisibility::Visible;
-	}
-
-	const bool ShowTransportControls = GetDefault<UCurveEditorSettings>()->GetShowTransportControls();
-	const bool TransportControlsPoppedOut = SequencerSettings->AreTransportControlsPoppedOut();
-	return (!ShowTransportControls || TransportControlsPoppedOut) ? EVisibility::Collapsed : EVisibility::Visible;
-}
-
 } // namespace Sequencer
 } // namespace UE
 
diff --git a/Engine/Source/Editor/Sequencer/Private/MVVM/Views/CurveEditor/CurveEditorWidgetOwner.cpp b/Engine/Source/Editor/Sequencer/Private/MVVM/Views/CurveEditor/CurveEditorWidgetOwner.cpp
new file mode 100644
--- /dev/null
+++ b/Engine/Source/Editor/Sequencer/Private/MVVM/Views/CurveEditor/CurveEditorWidgetOwner.cpp
@@ -0,0 +1,192 @@
+﻿// Copyright Epic Games, Inc. All Rights Reserved.
+
+#include "CurveEditorWidgetOwner.h"
+
+#include "IPropertyRowGenerator.h"
+#include "IStructureDetailsView.h"
+#include "SCurveEditorPanel.h"
+#include "SCurveEditorToolProperties.h"
+#include "SCurveKeyDetailPanel.h"
+#include "Sequencer.h"
+#include "SequencerCommands.h"
+#include "SOldCurveTreeContent.h"
+#include "Filters/SCurveEditorFilterPanel.h"
+#include "MVVM/CurveEditorExtension.h"
+#include "Widgets/DeclarativeSyntaxSupport.h"
+#include "Widgets/CurveEditor/SSequencerCurveEditor.h"
+#include "Widgets/Docking/SDockTab.h"
+
+struct FTimeSliderArgs;
+
+#define LOCTEXT_NAMESPACE "FSequencerCurveEditorPresenter"
+
+namespace UE::Sequencer
+{
+namespace PresenterDetail
+{
+static FCurveEditorWidgets MakePresenterWidgets(
+	const TSharedRef<FSequencer>& InSequencer, 
+	const TSharedRef<FCurveEditor>& InCurveEditor, 
+	const TSharedRef<ITimeSliderController>& InTimeSliderController,
+	const TSharedPtr<FTabManager> InTabManager
+	)
+{
+	const TSharedRef<SOldCurveTreeContent> OldContent = SNew(SOldCurveTreeContent, InSequencer, InCurveEditor);
+	// TODO UE-363527: Create a new widget here and then hook it up using SWidgetSwitcher driven by the new CVar.
+	
+	const TSharedRef<SCurveEditorPanel> CurveEditorPanel = SNew(SCurveEditorPanel, InCurveEditor)
+		// Grid lines match the color specified in FSequencerTimeSliderController::OnPaintViewArea
+		.GridLineTint(FLinearColor(0.f, 0.f, 0.f, 0.3f))
+		.ExternalTimeSliderController(InTimeSliderController)
+		.MinimumViewPanelHeight(0.f)
+		.TabManager(InTabManager)
+		.DisabledTimeSnapTooltip(LOCTEXT("CurveEditorTimeSnapDisabledTooltip", "Time Snapping is currently driven by Sequencer."))
+		.TreeContent()
+		[
+			OldContent
+		];
+	
+	const TSharedRef<SSequencerCurveEditor> RootWidget = SNew(SSequencerCurveEditor, CurveEditorPanel, InSequencer);
+	return FCurveEditorWidgets(RootWidget, CurveEditorPanel, OldContent);
+}
+}
+
+FCurveEditorWidgetOwner::FCurveEditorWidgetOwner(
+	const TSharedRef<FSequencer>& InSequencer,
+	const TSharedRef<FCurveEditor>& InCurveEditor,
+	const TSharedRef<ITimeSliderController>& InTimeSliderController
+	)
+	: WeakSequencer(InSequencer)
+	, CurveEditor(InCurveEditor)
+	, Widgets(PresenterDetail::MakePresenterWidgets(InSequencer, InCurveEditor, InTimeSliderController, GetTabManager()))
+{
+	BindCommands();
+	InitFrameNumberPropertyLayout(InSequencer);
+	
+	const TSharedPtr<FTabManager> TabManager = GetTabManager();
+	if (const TSharedPtr<SDockTab> ExistingCurveEditorTab = TabManager->FindExistingLiveTab(FCurveEditorExtension::CurveEditorTabName))
+	{
+		ExistingCurveEditorTab->SetContent(Widgets.RootWidget);
+	}
+}
+
+FCurveEditorWidgetOwner::~FCurveEditorWidgetOwner()
+{
+	// It should not happen defensively prevent crash in case the widget outlives FSequencerCurveEditorPresenter.
+	Widgets.CurveEditorPanel->OnFilterClassChanged.Unbind();
+}
+
+bool FCurveEditorWidgetOwner::IsCurveEditorOpen() const
+{
+	const TSharedPtr<FTabManager> TabManager = GetTabManager();
+	return TabManager && TabManager->FindExistingLiveTab(FCurveEditorExtension::CurveEditorTabName).IsValid();
+}
+
+void FCurveEditorWidgetOwner::OpenCurveEditor()
+{
+	const TSharedPtr<FTabManager> TabManager = GetTabManager();
+	if (!ensure(TabManager))
+	{
+		return;
+	}
+
+	// Request the Tab Manager invoke the tab. This will spawn the tab if needed, otherwise pull it to focus. This assumes
+	// that the Toolkit Host's Tab Manager has already registered a tab with a NullWidget for content.
+	if (const TSharedPtr<SDockTab> CurveEditorTab = TabManager->TryInvokeTab(FCurveEditorExtension::CurveEditorTabName))
+	{
+		CurveEditorTab->SetContent(Widgets.CurveEditorPanel);
+
+		const FSlateIcon SequencerGraphIcon = FSlateIcon(FAppStyle::GetAppStyleSetName(), "GenericCurveEditor.TabIcon");
+		CurveEditorTab->SetTabIcon(SequencerGraphIcon.GetIcon());
+		CurveEditorTab->SetLabel(LOCTEXT("SequencerMainGraphEditorTitle", "Sequencer Curves"));
+
+		CurveEditor->ZoomToFit();
+	}
+}
+
+void FCurveEditorWidgetOwner::CloseCurveEditor()
+{
+	const TSharedPtr<FTabManager> TabManager = GetTabManager();
+	if (!TabManager.IsValid())
+	{
+		return;
+	}
+
+	if (const TSharedPtr<SDockTab> CurveEditorTab = TabManager->FindExistingLiveTab(FCurveEditorExtension::CurveEditorTabName))
+	{
+		CurveEditorTab->RequestCloseTab();
+	}
+}
+
+void FCurveEditorWidgetOwner::SyncSelection() const
+{
+	// TODO UE-363527: Depending on the new CVar, reroute the selection sync to the correct widget.
+	Widgets.OldTreeContent->SyncSelection();
+}
+
+TSharedPtr<SCurveEditorTree> FCurveEditorWidgetOwner::GetTreeView() const
+{
+	// TODO UE-363527: Depending on the new CVar, reroute to the correct widget.
+	return Widgets.OldTreeContent->GetCurveEditorTreeView();
+}
+
+USequencerSettings* FCurveEditorWidgetOwner::GetSequencerSettings() const
+{
+	const TSharedPtr<ISequencer> Sequencer = WeakSequencer.Pin();
+	return Sequencer ? Sequencer->GetSequencerSettings() : nullptr;
+}
+
+TSharedPtr<FTabManager> FCurveEditorWidgetOwner::GetTabManager() const
+{
+	const TSharedPtr<ISequencer> Sequencer = WeakSequencer.Pin();
+	const TSharedPtr<IToolkitHost> ToolkitHost = Sequencer ? Sequencer->GetToolkitHost() : nullptr;
+	return ToolkitHost.IsValid() ? ToolkitHost->GetTabManager() : nullptr;
+}
+
+void FCurveEditorWidgetOwner::BindCommands()
+{
+	const TSharedPtr<FUICommandList> CurveEditorCommandList = CurveEditor->GetCommands();
+	if (!ensure(CurveEditorCommandList))
+	{
+		return;
+	}
+	
+	// TODO UE-363527: Depending on the new CVar, reroute the commands to the correct widget.
+	CurveEditorCommandList->MapAction(
+		FSequencerCommands::Get().QuickTreeSearch,
+		FExecuteAction::CreateLambda([this] { Widgets.OldTreeContent->FocusSearchBox(); })
+	);
+	CurveEditorCommandList->MapAction(
+		FSequencerCommands::Get().ToggleShowGotoBox,
+		FExecuteAction::CreateLambda([this] { Widgets.OldTreeContent->FocusPlayTimeDisplay(); })
+	);
+}
+
+void FCurveEditorWidgetOwner::InitFrameNumberPropertyLayout(const TSharedRef<FSequencer>& InSequencer)
+{
+	const TSharedRef<SCurveEditorPanel>& CurveEditorPanel = Widgets.CurveEditorPanel;
+	
+	CurveEditorPanel->GetKeyDetailsView()->GetPropertyRowGenerator()->RegisterInstancedCustomPropertyTypeLayout(
+		"FrameNumber", 
+		FOnGetPropertyTypeCustomizationInstance::CreateSP(InSequencer, &FSequencer::MakeFrameNumberDetailsCustomization));
+	CurveEditorPanel->GetToolPropertiesPanel()->GetStructureDetailsView()->GetDetailsView()->RegisterInstancedCustomPropertyTypeLayout(
+		"FrameNumber",
+		FOnGetPropertyTypeCustomizationInstance::CreateSP(InSequencer, &FSequencer::MakeFrameNumberDetailsCustomization));
+	
+	Widgets.CurveEditorPanel->OnFilterClassChanged.BindRaw(this, &FCurveEditorWidgetOwner::OnFilterClassChanged);
+}
+
+void FCurveEditorWidgetOwner::OnFilterClassChanged()
+{
+	const TSharedPtr<FSequencer> Sequencer = WeakSequencer.Pin();
+	const TSharedPtr<SCurveEditorFilterPanel> FilterPanel = Widgets.CurveEditorPanel->GetFilterPanel();
+	if (Sequencer && FilterPanel)
+	{
+		FilterPanel->GetDetailsView()->RegisterInstancedCustomPropertyTypeLayout(
+			TEXT("FrameNumber"),
+			FOnGetPropertyTypeCustomizationInstance::CreateSP(Sequencer.ToSharedRef(), &FSequencer::MakeFrameNumberDetailsCustomization));
+	}
+}
+} // namespace UE::Sequencer
+
+#undef LOCTEXT_NAMESPACE
\ No newline at end of file
diff --git a/Engine/Source/Editor/Sequencer/Private/MVVM/Views/CurveEditor/CurveEditorWidgetOwner.h b/Engine/Source/Editor/Sequencer/Private/MVVM/Views/CurveEditor/CurveEditorWidgetOwner.h
new file mode 100644
--- /dev/null
+++ b/Engine/Source/Editor/Sequencer/Private/MVVM/Views/CurveEditor/CurveEditorWidgetOwner.h
@@ -0,0 +1,76 @@
+﻿// Copyright Epic Games, Inc. All Rights Reserved.
+
+#pragma once
+
+#include "CurveEditorWidgets.h"
+#include "Templates/SharedPointer.h"
+
+class FCurveEditor;
+class FSequencer;
+class FTabManager;
+class ITimeSliderController;
+class SCurveEditorPanel;
+class SCurveEditorTree;
+class SSequencerCurveEditor;
+class USequencerSettings;
+struct FTimeSliderArgs;
+
+namespace UE::Sequencer
+{
+/**
+ * Instantiates the UI displaying Sequencer's FCurveEditor.
+ * 
+ * Manages the UI:
+ * - binds UI commands
+ * - updates the details customization for the key settings in the toolbar
+ * - interacts with the tab manager
+ */
+class FCurveEditorWidgetOwner
+{
+public:
+	
+	explicit FCurveEditorWidgetOwner(
+		const TSharedRef<FSequencer>& InSequencer, 
+		const TSharedRef<FCurveEditor>& InCurveEditor,
+		const TSharedRef<ITimeSliderController>& InTimeSliderController
+		);
+	~FCurveEditorWidgetOwner();
+	
+	/** Returns whether the curve editor is open */
+	bool IsCurveEditorOpen() const;
+	
+	/** Opens the curve editor */
+	void OpenCurveEditor();
+	/** Closes the curve editor */
+	void CloseCurveEditor();
+	
+	/** Requests that the UI syncs its selection to that of Sequencer. */
+	void SyncSelection() const;
+	
+	/**
+	 * @return The currently active tree view widget.
+	 * TODO UE-363527: Maybe this can be factored out. It breaks encapsulation. For that, FCurveEditorExtension::GetCurveEditorTreeView must be deprecated.
+	 */
+	TSharedPtr<SCurveEditorTree> GetTreeView() const;
+	
+private:
+	
+	/** Sequencer instance that owns us. */
+	const TWeakPtr<FSequencer> WeakSequencer;
+	/** The curve editor being displayed. */
+	const TSharedRef<FCurveEditor> CurveEditor;
+	
+	/** Widgets this presenter manages. */
+	const FCurveEditorWidgets Widgets;
+	
+	USequencerSettings* GetSequencerSettings() const;
+	TSharedPtr<FTabManager> GetTabManager() const;
+	
+	void BindCommands();
+	
+	/** Register an instanced custom property type layout to handle converting FFrameNumber from Tick Resolution to Display Rate. */
+	void InitFrameNumberPropertyLayout(const TSharedRef<FSequencer>& InSequencer);
+	/** Updates the frame number layout */
+	void OnFilterClassChanged();
+};
+} // namespace UE::Sequencer
diff --git a/Engine/Source/Editor/Sequencer/Private/MVVM/Views/CurveEditor/CurveEditorWidgets.h b/Engine/Source/Editor/Sequencer/Private/MVVM/Views/CurveEditor/CurveEditorWidgets.h
new file mode 100644
--- /dev/null
+++ b/Engine/Source/Editor/Sequencer/Private/MVVM/Views/CurveEditor/CurveEditorWidgets.h
@@ -0,0 +1,38 @@
+﻿// Copyright Epic Games, Inc. All Rights Reserved.
+
+#pragma once
+
+#include "Templates/SharedPointer.h"
+
+class SCurveEditorPanel;
+class SSequencerCurveEditor;
+
+namespace UE::Sequencer
+{
+class SOldCurveTreeContent;
+
+/** @see FSequencerCurveEditorPresenter */
+struct FCurveEditorWidgets
+{
+	/** The curve editor widget containing the curve editor panel */
+	TSharedRef<SSequencerCurveEditor> RootWidget;
+	/** The curve editor panel. This is created and updated even if it is not currently visible. */
+	TSharedRef<SCurveEditorPanel> CurveEditorPanel;
+	
+	/** 
+	 * The tree content displayed inside CurveEditorPanel.
+	 * TODO UE-363527: A CVar will be added to switch between this in the new linked filtering UI.
+	 */
+	TSharedRef<SOldCurveTreeContent> OldTreeContent;
+
+	explicit FCurveEditorWidgets(
+		const TSharedRef<SSequencerCurveEditor>& RootWidget, 
+		const TSharedRef<SCurveEditorPanel>& CurveEditorPanel,
+		const TSharedRef<SOldCurveTreeContent>& OldTreeContent
+		)
+		: RootWidget(RootWidget)
+		, CurveEditorPanel(CurveEditorPanel)
+		, OldTreeContent(OldTreeContent)
+	{}
+};
+} // namespace UE::Sequencer
diff --git a/Engine/Source/Editor/Sequencer/Private/MVVM/Views/CurveEditor/SOldCurveTreeContent.cpp b/Engine/Source/Editor/Sequencer/Private/MVVM/Views/CurveEditor/SOldCurveTreeContent.cpp
new file mode 100644
--- /dev/null
+++ b/Engine/Source/Editor/Sequencer/Private/MVVM/Views/CurveEditor/SOldCurveTreeContent.cpp
@@ -0,0 +1,207 @@
+﻿// Copyright Epic Games, Inc. All Rights Reserved.
+
+#include "SOldCurveTreeContent.h"
+
+#include "CurveEditorSettings.h"
+#include "LevelSequence.h"
+#include "SelectionSyncUtils.h"
+#include "Sequencer.h"
+#include "SequencerSettings.h"
+#include "SPopoutTabInlineContent.h"
+#include "WorkspaceMenuStructure.h"
+#include "WorkspaceMenuStructureModule.h"
+#include "Misc/TransportControlsHelper.h"
+#include "Styling/AppStyle.h"
+#include "Textures/SlateIcon.h"
+#include "Tree/SCurveEditorTree.h"
+#include "Tree/SCurveEditorTreeFilterStatusBar.h"
+#include "Tree/SCurveEditorTreeTextFilter.h"
+#include "Widgets/SBoxPanel.h"
+#include "Widgets/SOverlay.h"
+#include "Widgets/Docking/SDockTab.h"
+#include "Widgets/Layout/SBorder.h"
+#include "Widgets/Layout/SScrollBorder.h"
+
+#define LOCTEXT_NAMESPACE "SOldCurveTreeContent"
+
+namespace UE::Sequencer
+{
+void SOldCurveTreeContent::Construct(
+	const FArguments& InArgs, const TSharedRef<FSequencer> InSequencer, const TSharedRef<FCurveEditor>& InCurveEditor
+	)
+{
+	WeakSequencer = InSequencer;
+	CurveEditor = InCurveEditor;
+	
+	ChildSlot
+	[
+		SNew(SVerticalBox)
+
+		+ SVerticalBox::Slot()
+		.AutoHeight()
+		[
+			SNew(SBorder)
+			.BorderImage(FAppStyle::GetBrush("ToolPanel.GroupBorder"))
+			.Clipping(EWidgetClipping::ClipToBounds)
+			[
+				SAssignNew(CurveEditorSearchBox, SCurveEditorTreeTextFilter, InCurveEditor)
+			]
+		]
+
+		+ SVerticalBox::Slot()
+		[
+			MakeTreeContent(InCurveEditor)
+		]
+
+		+ SVerticalBox::Slot()
+		.AutoHeight()
+		[
+			MakeFooter(InArgs, InSequencer)
+		]
+	];
+}
+
+void SOldCurveTreeContent::FocusSearchBox()
+{
+	FSlateApplication::Get().SetKeyboardFocus(CurveEditorSearchBox, EFocusCause::SetDirectly);
+}
+
+void SOldCurveTreeContent::FocusPlayTimeDisplay()
+{
+	PlayTimeDisplay->Setup(); 
+	FSlateApplication::Get().SetKeyboardFocus(PlayTimeDisplay, EFocusCause::SetDirectly);
+}
+
+void SOldCurveTreeContent::SyncSelection()
+{
+	if (const TSharedPtr<FSequencer> Sequencer = WeakSequencer.Pin())
+	{
+		Sequencer::SyncSelection(
+			Sequencer->GetViewModel()->GetSelection(), CurveEditor.ToSharedRef(), CurveEditorTreeView.ToSharedRef()
+			);
+	}
+}
+
+TSharedRef<SWidget> SOldCurveTreeContent::MakeTreeContent(const TSharedRef<FCurveEditor>& InCurveEditor)
+{
+	CurveEditorTreeView = SNew(SCurveEditorTree, InCurveEditor);
+	return SNew(SOverlay)
+
+		+ SOverlay::Slot()
+		[
+			SNew(SScrollBorder, CurveEditorTreeView.ToSharedRef())
+			[
+				CurveEditorTreeView.ToSharedRef()
+			]
+		]
+
+		+ SOverlay::Slot()
+		.VAlign(VAlign_Bottom)
+		[
+			SNew(SHorizontalBox)
+
+			+ SHorizontalBox::Slot()
+			.AutoWidth()
+			[
+				SAssignNew(CurveEditorTreeFilterStatusBar, SCurveEditorTreeFilterStatusBar, InCurveEditor)
+				.Visibility(EVisibility::Hidden) // Initially hidden, visible on hover of the info button
+			]
+		];
+}
+
+TSharedRef<SWidget> SOldCurveTreeContent::MakeFooter(const FArguments& InArgs, const TSharedRef<FSequencer> InSequencer)
+{
+	PlayTimeDisplay = InSequencer->MakePlayTimeDisplay();
+	
+	return SNew(SBorder)
+		.BorderImage(FAppStyle::GetBrush("ToolPanel.GroupBorder"))
+		.Clipping(EWidgetClipping::ClipToBounds)
+		.Visibility(this, &SOldCurveTreeContent::GetPopoutTransportControlsVisibility)
+		[
+			SNew(SHorizontalBox)
+
+			+ SHorizontalBox::Slot()
+			.AutoWidth()
+			.VAlign(VAlign_Center)
+			.HAlign(HAlign_Left)
+			[
+				SNew(SButton)
+				.VAlign(EVerticalAlignment::VAlign_Center)
+				.ButtonStyle(FAppStyle::Get(), "SimpleButton")
+				.ToolTipText_Lambda([this] { return LOCTEXT("ShowStatus", "Show Status"); })
+				.ContentPadding(FMargin(1, 0))
+				.OnHovered_Lambda([this] { CurveEditorTreeFilterStatusBar->ShowStatusBar(); })
+				.OnUnhovered_Lambda([this] { CurveEditorTreeFilterStatusBar->FadeOutStatusBar(); })
+				.OnClicked_Lambda([this] { CurveEditorTreeFilterStatusBar->HideStatusBar(); return FReply::Handled(); })
+				[
+					SNew(SImage)
+					.ColorAndOpacity(FSlateColor::UseForeground())
+					.Image(FAppStyle::Get().GetBrush("Icons.Info.Small"))
+				]
+			]
+
+			+ SHorizontalBox::Slot()
+			.HAlign(HAlign_Fill)
+			[
+				MakePopoutTransportControls(InSequencer, InArgs._TabManager)
+			]
+
+			+ SHorizontalBox::Slot()
+			.AutoWidth()
+			.VAlign(VAlign_Center)
+			.HAlign(HAlign_Right)
+			[
+				SNew(SButton)
+				.VAlign(EVerticalAlignment::VAlign_Center)
+				.ButtonStyle(FAppStyle::Get(), "NoBorder")
+				.ContentPadding(FMargin(1, 0))
+				[
+					SNew(SHorizontalBox)
+
+					+ SHorizontalBox::Slot()
+					.AutoWidth()
+					.VAlign(VAlign_Center)
+					.HAlign(HAlign_Right)
+					.Padding(FMargin(3.f, 0.f, 0.f, 0.f))
+					[
+						SNew(SBorder)
+						.BorderImage(nullptr)
+						[
+							PlayTimeDisplay.ToSharedRef()
+						]
+					]
+				]
+			]
+		];
+}
+
+USequencerSettings* SOldCurveTreeContent::GetSequencerSettings() const
+{
+	const TSharedPtr<ISequencer> Sequencer = WeakSequencer.Pin();
+	return Sequencer.IsValid() ? Sequencer->GetSequencerSettings() : nullptr;
+}
+
+EVisibility SOldCurveTreeContent::GetPopoutTransportControlsVisibility() const
+{
+	const TSharedPtr<ISequencer> Sequencer = WeakSequencer.Pin();
+	if (!Sequencer.IsValid())
+	{
+		return EVisibility::Collapsed;
+	}
+	const UMovieSceneSequence* const RootSequence = Sequencer->GetRootMovieSceneSequence();
+	if (!RootSequence || !RootSequence->IsA<ULevelSequence>())
+	{
+		return EVisibility::Visible;
+	}
+	const USequencerSettings* const SequencerSettings = Sequencer->GetSequencerSettings();
+	if (!SequencerSettings)
+	{
+		return EVisibility::Visible;
+	}
+	const bool bShowTransportControls = GetDefault<UCurveEditorSettings>()->GetShowTransportControls();
+	const bool bTransportControlsPoppedOut = SequencerSettings->AreTransportControlsPoppedOut();
+	return (!bShowTransportControls || bTransportControlsPoppedOut) ? EVisibility::Collapsed : EVisibility::Visible;
+}
+} // namespace UE::Sequencer
+
+#undef LOCTEXT_NAMESPACE
\ No newline at end of file
diff --git a/Engine/Source/Editor/Sequencer/Private/MVVM/Views/CurveEditor/SOldCurveTreeContent.h b/Engine/Source/Editor/Sequencer/Private/MVVM/Views/CurveEditor/SOldCurveTreeContent.h
new file mode 100644
--- /dev/null
+++ b/Engine/Source/Editor/Sequencer/Private/MVVM/Views/CurveEditor/SOldCurveTreeContent.h
@@ -0,0 +1,68 @@
+﻿// Copyright Epic Games, Inc. All Rights Reserved.
+
+#pragma once
+
+#include "STemporarilyFocusedSpinBox.h"
+#include "Widgets/DeclarativeSyntaxSupport.h"
+#include "Widgets/SCompoundWidget.h"
+
+class FCurveEditor;
+class FSequencer;
+class FTabManager;
+class SCurveEditorTree;
+class SCurveEditorTreeFilterStatusBar;
+class USequencerSettings;
+
+namespace UE::Sequencer
+{
+/**
+ * Curve Editor tree content before 5.8. 
+ * 
+ * This widget will be removed in the future. 
+ * For now, it is used to get the legacy UI looks (see UE-363527)
+ */
+class SOldCurveTreeContent : public SCompoundWidget
+{
+public:
+	
+	SLATE_BEGIN_ARGS(SOldCurveTreeContent){}
+		/** Explicit tab manager to use, otherwise uses the global tab manager */
+		SLATE_ARGUMENT(TSharedPtr<FTabManager>, TabManager)
+	SLATE_END_ARGS()
+	
+	void Construct(const FArguments& InArgs, const TSharedRef<FSequencer> InSequencer, const TSharedRef<FCurveEditor>& InCurveEditor);
+	
+	/** Sets keyboard focus to the search box. */
+	void FocusSearchBox();
+	/** Sets keyboard focus to the play time display. */
+	void FocusPlayTimeDisplay();
+	
+	/** Syncs the tree view's selection to that of Sequencer. */
+	void SyncSelection();
+	
+	const TSharedPtr<SCurveEditorTree>& GetCurveEditorTreeView() const { return CurveEditorTreeView; }
+	
+private:
+	
+	/** The owning Sequencer instance that the Curve Editor is being displayed for. */
+	TWeakPtr<FSequencer> WeakSequencer;
+	/** The curve editor being displayed. */
+	TSharedPtr<FCurveEditor> CurveEditor;
+	
+	/** Curve editor tree widget */
+	TSharedPtr<SCurveEditorTree> CurveEditorTreeView;
+	/** Filter Status Bar */
+	TSharedPtr<SCurveEditorTreeFilterStatusBar> CurveEditorTreeFilterStatusBar;
+	/** The current playback time display. */
+	TSharedPtr<STemporarilyFocusedSpinBox<double>> PlayTimeDisplay;
+	/** The search widget for filtering curves in the Curve Editor tree. */
+	TSharedPtr<SWidget> CurveEditorSearchBox;
+	
+	TSharedRef<SWidget> MakeTreeContent(const TSharedRef<FCurveEditor>& InCurveEditor);
+	TSharedRef<SWidget> MakeFooter(const FArguments& InArgs, const TSharedRef<FSequencer> InSequencer);
+	
+	USequencerSettings* GetSequencerSettings() const;
+	
+	EVisibility GetPopoutTransportControlsVisibility() const;
+};
+} // namespace UE::Sequencer
diff --git a/Engine/Source/Editor/Sequencer/Private/MVVM/Views/CurveEditor/SelectionSyncUtils.cpp b/Engine/Source/Editor/Sequencer/Private/MVVM/Views/CurveEditor/SelectionSyncUtils.cpp
new file mode 100644
--- /dev/null
+++ b/Engine/Source/Editor/Sequencer/Private/MVVM/Views/CurveEditor/SelectionSyncUtils.cpp
@@ -0,0 +1,127 @@
+﻿// Copyright Epic Games, Inc. All Rights Reserved.
+
+#include "SelectionSyncUtils.h"
+
+#include "CurveEditor.h"
+#include "MVVM/Extensions/ICurveEditorTreeItemExtension.h"
+#include "MVVM/Extensions/IOutlinerExtension.h"
+#include "MVVM/Selection/Selection.h"
+#include "MVVM/ViewModelPtr.h"
+#include "Templates/SharedPointer.h"
+#include "Tree/SCurveEditorTree.h"
+
+namespace UE::Sequencer
+{
+namespace SyncDetail
+{
+static void GetAllChildrenRecursively(
+	const TSharedRef<SCurveEditorTree>& CurveEditorTreeView, const FCurveEditorTreeItemID Item, int32& LastAddedItem,
+	TMap<const FCurveEditorTreeItemID, int32>& ItemToSortOrder
+	)
+{
+	ItemToSortOrder.Add(Item, LastAddedItem++);
+
+	TArray<FCurveEditorTreeItemID> Children;
+
+	CurveEditorTreeView->GetTreeItemChildren(Item, Children);
+	// Add them to the output array
+	for (const FCurveEditorTreeItemID ChildItemID : Children)
+	{
+		GetAllChildrenRecursively(CurveEditorTreeView, ChildItemID, LastAddedItem, ItemToSortOrder);
+	}
+}
+} // namespace SyncDetail
+
+void SyncSelection(
+	const TSharedPtr<FSequencerSelection>& InSelection, 
+	const TSharedRef<FCurveEditor>& InCurveEditorModel, 
+	const TSharedRef<SCurveEditorTree>& InCurveEditorTreeView
+	)
+{
+	//we now sync to the first non-channel in the tree hierarchy, this is needed for hosted channels with control rigs
+	//we need to sort the hierarchy each time since the curve id's being used previously were not consistent in order of viewing/creation, persistent now for undo/redo
+	InCurveEditorModel->SuspendBroadcast();
+	InCurveEditorTreeView->ClearSelection();
+
+	//map of tree items id's to their actually sort order, need to use this since tree items ids are not in order anymore
+	TMap<const FCurveEditorTreeItemID, int32>  ItemToSortOrder;
+	int32 LastAddedItem = 0;
+	TArray<FCurveEditorTreeItemID> Children;
+	TArrayView<const FCurveEditorTreeItemID> RootItems = InCurveEditorTreeView->GetRootItems();
+	
+	for (const FCurveEditorTreeItemID RootItemID : RootItems)
+	{
+		SyncDetail::GetAllChildrenRecursively(InCurveEditorTreeView, RootItemID, LastAddedItem, ItemToSortOrder);
+	}
+
+	auto GetSortOrder = [&ItemToSortOrder](FCurveEditorTreeItemID InTreeItemID) -> int32
+	{
+		if (const int32* SortOrder = ItemToSortOrder.Find(InTreeItemID))
+		{
+			return *SortOrder;
+		}
+		return INDEX_NONE;
+	};
+
+	FCurveEditorTreeItemID FirstCurveEditorTreeItemID;
+	int32 FirstSortOrder = INDEX_NONE;
+	bool bSelectedIsChannel = true;
+
+	for (TViewModelPtr<IOutlinerExtension> SelectedItem : InSelection->Outliner)
+	{
+		if (TViewModelPtr<ICurveEditorTreeItemExtension> CurveEditorItem = SelectedItem.ImplicitCast())
+		{
+			FCurveEditorTreeItemID CurveEditorTreeItem = CurveEditorItem->GetCurveEditorItemID();
+			if (CurveEditorTreeItem != FCurveEditorTreeItemID::Invalid())
+			{
+				if (!InCurveEditorTreeView->IsItemSelected(CurveEditorTreeItem))
+				{
+					InCurveEditorTreeView->SetItemSelection(CurveEditorTreeItem, true);
+					bool bIsChannel = SelectedItem.AsModel()->IsA<FChannelModel>() || SelectedItem.AsModel()->IsA<FChannelGroupOutlinerModel>()
+					|| SelectedItem.AsModel()->IsA<FChannelGroupModel>();
+					if (!FirstCurveEditorTreeItemID.IsValid())
+					{
+						FirstCurveEditorTreeItemID = CurveEditorTreeItem;
+						FirstSortOrder = GetSortOrder(CurveEditorTreeItem);
+						bSelectedIsChannel = bIsChannel;
+					}
+					else
+					{
+						int32 SortOrder = GetSortOrder(CurveEditorTreeItem);
+						if (bIsChannel == false)
+						{
+							if (bSelectedIsChannel == true)
+							{
+								FirstCurveEditorTreeItemID = CurveEditorTreeItem;
+								FirstSortOrder = SortOrder;
+								bSelectedIsChannel = false;
+							}
+							else if (SortOrder < FirstSortOrder)
+							{
+								FirstCurveEditorTreeItemID = CurveEditorTreeItem;
+								FirstSortOrder = SortOrder;
+								bSelectedIsChannel = false;
+							}
+						}
+						else if (bSelectedIsChannel == true) //its a channel only sort if selected is a non channel
+						{
+							if (SortOrder < FirstSortOrder)
+							{
+								FirstCurveEditorTreeItemID = CurveEditorTreeItem;
+								FirstSortOrder = SortOrder;
+							}
+						}
+					}
+				}
+			}
+		}
+	}
+	
+	if (FirstCurveEditorTreeItemID.IsValid())
+	{
+		InCurveEditorTreeView->RequestScrollIntoView(FirstCurveEditorTreeItemID);
+	}
+
+	InCurveEditorModel->ResumeBroadcast();
+}
+} // namespace UE::Sequencer
diff --git a/Engine/Source/Editor/Sequencer/Private/MVVM/Views/CurveEditor/SelectionSyncUtils.h b/Engine/Source/Editor/Sequencer/Private/MVVM/Views/CurveEditor/SelectionSyncUtils.h
new file mode 100644
--- /dev/null
+++ b/Engine/Source/Editor/Sequencer/Private/MVVM/Views/CurveEditor/SelectionSyncUtils.h
@@ -0,0 +1,20 @@
+﻿// Copyright Epic Games, Inc. All Rights Reserved.
+
+#pragma once
+
+#include "Templates/SharedPointerFwd.h"
+
+class FCurveEditor;
+class SCurveEditorTree;
+
+namespace UE::Sequencer
+{
+class FSequencerSelection;
+
+/** Updates SCurveEditorTree to have the same selection as InSelection */
+void SyncSelection(
+	const TSharedPtr<FSequencerSelection>& InSelection, 
+	const TSharedRef<FCurveEditor>& InCurveEditorModel, 
+	const TSharedRef<SCurveEditorTree>& InCurveEditorTreeView
+	);
+} // namespace UE::Sequencer
diff --git a/Engine/Source/Editor/Sequencer/Private/SSequencer.cpp b/Engine/Source/Editor/Sequencer/Private/SSequencer.cpp
--- a/Engine/Source/Editor/Sequencer/Private/SSequencer.cpp
+++ b/Engine/Source/Editor/Sequencer/Private/SSequencer.cpp
@@ -288,7 +288,7 @@ void SSequencer::Construct(const FArguments& InArgs, TSharedRef<FSequencer> InSe
 	RootCustomization.AddMenuExtender = InArgs._AddMenuExtender;
 	RootCustomization.ToolbarExtender = InArgs._ToolbarExtender;
 
-	PlayTimeDisplay = StaticCastSharedRef<STemporarilyFocusedSpinBox<double>>(SequencerPtr.Pin()->MakePlayTimeDisplay());
+	PlayTimeDisplay = SequencerPtr.Pin()->MakePlayTimeDisplay();
 
 	check(SequencerPtr.IsValid());
 	const TSharedRef<FSequencer> SequencerRef = SequencerPtr.Pin().ToSharedRef();
diff --git a/Engine/Source/Editor/Sequencer/Private/Sequencer.cpp b/Engine/Source/Editor/Sequencer/Private/Sequencer.cpp
--- a/Engine/Source/Editor/Sequencer/Private/Sequencer.cpp
+++ b/Engine/Source/Editor/Sequencer/Private/Sequencer.cpp
@@ -4787,7 +4787,7 @@ TSharedRef<SWidget> FSequencer::MakeTransportControls(bool bExtended)
 	return EditorWidgetsModule.CreateTransportControl( TransportControlArgs );
 }
 
-TSharedRef<SWidget> FSequencer::MakePlayTimeDisplay()
+TSharedRef<STemporarilyFocusedSpinBox<double>> FSequencer::MakePlayTimeDisplay()
 {
 	return SNew(STemporarilyFocusedSpinBox<double>)
 		.Style(&FAppStyle::GetWidgetStyle<FSpinBoxStyle>("Sequencer.PlayTimeSpinBox"))
diff --git a/Engine/Source/Editor/Sequencer/Private/Sequencer.h b/Engine/Source/Editor/Sequencer/Private/Sequencer.h
--- a/Engine/Source/Editor/Sequencer/Private/Sequencer.h
+++ b/Engine/Source/Editor/Sequencer/Private/Sequencer.h
@@ -84,6 +84,7 @@ struct FNotificationInfo;
 struct FEditorViewportViewModifierParams;
 struct FMovieSceneMarkedFrame;
 struct FSidebarDrawerConfig;
+template<typename T> struct STemporarilyFocusedSpinBox;
 
 enum class EMapChangeType : uint8;
 enum class ENearestKeyOption : uint8;
@@ -1217,7 +1218,7 @@ public:
 	/** Get the name of the movie renderer to use, defaults to the first available if the setting is empty */
 	FString GetMovieRendererName() const;
 
-	TSharedRef<SWidget> MakePlayTimeDisplay();
+	TSharedRef<STemporarilyFocusedSpinBox<double>> MakePlayTimeDisplay();
 
 	/** @return True if the Sequencer is currently limiting selection to only Sequencer key-able objects */
 	bool IsViewportSelectionLimited() const override;
diff --git a/Engine/Source/Editor/Sequencer/Public/MVVM/CurveEditorExtension.h b/Engine/Source/Editor/Sequencer/Public/MVVM/CurveEditorExtension.h
--- a/Engine/Source/Editor/Sequencer/Public/MVVM/CurveEditorExtension.h
+++ b/Engine/Source/Editor/Sequencer/Public/MVVM/CurveEditorExtension.h
@@ -29,6 +29,7 @@ namespace Sequencer
 {
 
 class FSequencerEditorViewModel;
+class FCurveEditorWidgetOwner;
 
 class FCurveEditorExtension : public IDynamicExtension
 {
@@ -55,7 +56,7 @@ public:
 	UE_API void CloseCurveEditor();
 
 	/** Curve editor tree widget */
-	TSharedPtr<SCurveEditorTree> GetCurveEditorTreeView() const { return CurveEditorTreeView; }
+	UE_API TSharedPtr<SCurveEditorTree> GetCurveEditorTreeView() const;
 
 	/**
 	 * Synchronize curve editor selection with sequencer outliner selection on the next update.
@@ -75,18 +76,10 @@ private:
 
 	/** Get the default key attributes to apply to newly created keys on the curve editor */
 	UE_API FKeyAttributes GetDefaultKeyAttributes() const;
-
-	/** Syncs the selection of the curve editor with that of the sequencer */
-	UE_API void SyncSelection();
-
-	/** Executed when the filter class has changed */
-	UE_API void FilterClassChanged();
 	
 	/** Cancels any pending transactions. */
 	void CleanUpPendingSyncSelection();
 
-	EVisibility GetPopoutTransportControlsVisibility() const;
-
 private:
 
 	/** The sequencer editor we are extending with a curve editor */
@@ -94,18 +87,9 @@ private:
 
 	/** Curve editor */
 	TSharedPtr<FCurveEditor> CurveEditorModel;
-	/** Curve editor tree widget */
-	TSharedPtr<SCurveEditorTree> CurveEditorTreeView;
-	/** The search widget for filtering curves in the Curve Editor tree. */
-	TSharedPtr<SWidget> CurveEditorSearchBox;
-	/** The curve editor widget containing the curve editor panel */
-	TSharedPtr<SWidget> CurveEditorWidget;
-	/** The curve editor panel. This is created and updated even if it is not currently visible. */
-	TSharedPtr<SCurveEditorPanel> CurveEditorPanel;
-	/** Filter Status Bar */
-	TSharedPtr<SCurveEditorTreeFilterStatusBar> CurveEditorTreeFilterStatusBar;
-	/** The current playback time display. */
-	TSharedPtr<STemporarilyFocusedSpinBox<double>> PlayTimeDisplay;
+	
+	/** Responsible for creating the UI. */
+	TPimplPtr<FCurveEditorWidgetOwner> WidgetOwner; 
 
 	/** Handle to latent refresh of selection. */
 	FTimerHandle SyncSelectionHandle;
