Sequencer: Fix missing toolbar in curve editor.

The issue was a copy-paste typo.

#Synced-CL 50532451

diff --git a/Engine/Source/Editor/Sequencer/Private/MVVM/Views/CurveEditor/CurveEditorWidgetOwner.cpp b/Engine/Source/Editor/Sequencer/Private/MVVM/Views/CurveEditor/CurveEditorWidgetOwner.cpp
--- a/Engine/Source/Editor/Sequencer/Private/MVVM/Views/CurveEditor/CurveEditorWidgetOwner.cpp
+++ b/Engine/Source/Editor/Sequencer/Private/MVVM/Views/CurveEditor/CurveEditorWidgetOwner.cpp
@@ -31,7 +31,8 @@ static FCurveEditorWidgets MakePresenterWidgets(
 	const TSharedPtr<FTabManager> InTabManager
 	)
 {
-	const TSharedRef<SOldCurveTreeContent> OldContent = SNew(SOldCurveTreeContent, InSequencer, InCurveEditor);
+	const TSharedRef<SOldCurveTreeContent> OldContent = SNew(SOldCurveTreeContent, InSequencer, InCurveEditor)
+		.TabManager(InTabManager);
 	// TODO UE-363527: Create a new widget here and then hook it up using SWidgetSwitcher driven by the new CVar.
 	
 	const TSharedRef<SCurveEditorPanel> CurveEditorPanel = SNew(SCurveEditorPanel, InCurveEditor)
@@ -94,7 +95,7 @@ void FCurveEditorWidgetOwner::OpenCurveEditor()
 	// that the Toolkit Host's Tab Manager has already registered a tab with a NullWidget for content.
 	if (const TSharedPtr<SDockTab> CurveEditorTab = TabManager->TryInvokeTab(FCurveEditorExtension::CurveEditorTabName))
 	{
-		CurveEditorTab->SetContent(Widgets.CurveEditorPanel);
+		CurveEditorTab->SetContent(Widgets.RootWidget);
 
 		const FSlateIcon SequencerGraphIcon = FSlateIcon(FAppStyle::GetAppStyleSetName(), "GenericCurveEditor.TabIcon");
 		CurveEditorTab->SetTabIcon(SequencerGraphIcon.GetIcon());
