Sequencer: Eliminate FSequencerTreeNode::SetTextFilterString.

This is what happens under the hood:
- SSequencerSearchBox::OnTextChanged invokes the lambda in SSequencerSearchAndFilterRow::MakeFilterContent
- FSequencerFilterBar::SetTextFilterString invokes FSequencerFilterBar::OnRequestUpdate
- FLinkedFilteringMediator::OnRequestUpdate_XXX calls FSequencerNodeTree::RequestFilterUpdate
This way the FSequencerNodeTree::RequestFilterUpdate continues to be matched.

This is required for follow up changes for linked filtering.

#Synced-CL 50532451

diff --git a/Engine/Source/Editor/Sequencer/Private/Filters/SequencerFilterBar.cpp b/Engine/Source/Editor/Sequencer/Private/Filters/SequencerFilterBar.cpp
--- a/Engine/Source/Editor/Sequencer/Private/Filters/SequencerFilterBar.cpp
+++ b/Engine/Source/Editor/Sequencer/Private/Filters/SequencerFilterBar.cpp
@@ -366,6 +366,7 @@ FString FSequencerFilterBar::GetTextFilterString() const
 void FSequencerFilterBar::SetTextFilterString(const FString& InText)
 {
 	TextFilter->SetRawFilterText(FText::FromString(InText));
+	RequestFilterUpdate();
 }
 
 bool FSequencerFilterBar::DoesTextFilterStringContainExpressionPair(const ISequencerTextFilterExpressionContext& InExpression) const
diff --git a/Engine/Source/Editor/Sequencer/Private/SequencerNodeTree.cpp b/Engine/Source/Editor/Sequencer/Private/SequencerNodeTree.cpp
--- a/Engine/Source/Editor/Sequencer/Private/SequencerNodeTree.cpp
+++ b/Engine/Source/Editor/Sequencer/Private/SequencerNodeTree.cpp
@@ -355,17 +355,6 @@ int32 FSequencerNodeTree::GetFilteredDisplayNodeCount() const
 	return Sequencer.GetFilterBar()->GetFilterData().GetDisplayNodeCount();
 }
 
-void FSequencerNodeTree::SetTextFilterString(const FString& InFilter)
-{
-	const TSharedRef<FSequencerFilterBar> FilterBar = Sequencer.GetFilterBar();
-	const FString FilterString = FilterBar->GetTextFilterString();
-	if (InFilter != FilterString)
-	{
-		bFilterUpdateRequested = true;
-		FilterBar->SetTextFilterString(InFilter);
-	}
-}
-
 void FSequencerNodeTree::NodeGroupsCollectionChanged()
 {
 	if (Sequencer.GetFocusedMovieSceneSequence()->GetMovieScene()->GetNodeGroups().HasAnyActiveFilter() || bFilteringOnNodeGroups)
diff --git a/Engine/Source/Editor/Sequencer/Private/SequencerNodeTree.h b/Engine/Source/Editor/Sequencer/Private/SequencerNodeTree.h
--- a/Engine/Source/Editor/Sequencer/Private/SequencerNodeTree.h
+++ b/Engine/Source/Editor/Sequencer/Private/SequencerNodeTree.h
@@ -88,12 +88,6 @@ public:
 	 */
 	bool NeedsFilterUpdate() const { return bFilterUpdateRequested; }
 
-	/**
-	 * Filters the nodes based on the passed in filter terms
-	 *
-	 * @param InFilter	The filter terms
-	 */
-	void SetTextFilterString( const FString& InFilter );
 
 	/** Called when the active MovieScene's node group colletion has been modifed */
 	void NodeGroupsCollectionChanged();
diff --git a/Engine/Source/Editor/Sequencer/Private/Widgets/Sequencer/SSequencerSearchAndFilterRow.cpp b/Engine/Source/Editor/Sequencer/Private/Widgets/Sequencer/SSequencerSearchAndFilterRow.cpp
--- a/Engine/Source/Editor/Sequencer/Private/Widgets/Sequencer/SSequencerSearchAndFilterRow.cpp
+++ b/Engine/Source/Editor/Sequencer/Private/Widgets/Sequencer/SSequencerSearchAndFilterRow.cpp
@@ -120,7 +120,7 @@ TSharedPtr<SSequencerSearchBox> SSequencerSearchAndFilterRow::GetActiveSearchBox
 
 TSharedRef<SWidget> SSequencerSearchAndFilterRow::MakeFilterContent(const FArguments& InArgs)
 {
-	const auto MakeHideIsolateViewModel = [this](const TSharedRef<FSequencerFilterBar>& InFilterBar)
+	const auto MakeWidget = [this](const TSharedRef<FSequencerFilterBar>& InFilterBar)
 	{
 		const TSharedRef<FSequencer> SequencerPin = WeakSequencer.Pin().ToSharedRef();
 		const TSharedRef<FHideIsolateShowViewModel> IsolateViewModel = MakeShared<FHideIsolateShowViewModel>(
@@ -128,34 +128,28 @@ TSharedRef<SWidget> SSequencerSearchAndFilterRow::MakeFilterContent(const FArgum
 			FilterRowDetail::MakeSelectedTracksAttr(SequencerPin), FilterRowDetail::MakeAllTracksAttr(SequencerPin)
 			);
 		IsolateViewModel->OnRequestScrollToView().AddSP(this, &SSequencerSearchAndFilterRow::ScrollIntoView);
-		return IsolateViewModel;
+	
+		const TSharedPtr<ISequencerTrackFilters> TrackFilters = StaticCastSharedPtr<ISequencerTrackFilters>(InFilterBar.ToSharedPtr());
+		return SNew(SSearchAndFilterWidget, InFilterBar, IsolateViewModel)
+			.OnSearchTextChanged_Lambda([InFilterBar](const FText& InSearchText)
+			{
+				InFilterBar->SetTextFilterString(InSearchText.ToString());
+			});
 	};
 	
 	const TSharedPtr<FLinkedFilterViewModel> LinkedFilterViewModel = InArgs._FilterViewModel;
-	TSharedRef<FSequencerFilterBar> LinkedFilterBar = LinkedFilterViewModel->GetLinkedFilterBarImpl();
-	TSharedRef<FSequencerFilterBar> InstancedFilterBar = LinkedFilterViewModel->GetInstancedFilterBarImpl();
 	return SNew(SFilterLinkStateSwitcher)
 		.LinkedFilterViewModel(InArgs._FilterViewModel)
-		.LinkedContent()
+		.LinkedContent() 
 		[
-			SAssignNew(LinkedFilterWidgets, SSearchAndFilterWidget, LinkedFilterBar, MakeHideIsolateViewModel(LinkedFilterBar))
-			.OnSearchTextChanged(this, &SSequencerSearchAndFilterRow::OnSearchTextChanged)
+			MakeWidget(LinkedFilterViewModel->GetLinkedFilterBarImpl())
 		]
-		.InstancedContent()
+		.InstancedContent() 
 		[
-			SAssignNew(InstancedFilterWidgets, SSearchAndFilterWidget, InstancedFilterBar, MakeHideIsolateViewModel(InstancedFilterBar))
-			.OnSearchTextChanged(this, &SSequencerSearchAndFilterRow::OnSearchTextChanged)
+			MakeWidget(LinkedFilterViewModel->GetInstancedFilterBarImpl())
 		];
 }
 
-void SSequencerSearchAndFilterRow::OnSearchTextChanged(const FText& InSearchText)
-{
-	if (const TSharedPtr<FSequencer> Sequencer = WeakSequencer.Pin())
-	{
-		Sequencer->GetNodeTree()->SetTextFilterString(InSearchText.ToString());
-	}
-}
-
 void SSequencerSearchAndFilterRow::ScrollIntoView(const TWeakViewModelPtr<IOutlinerExtension>& InItem)
 {
 	if (const TSharedPtr<FSequencer> SequencerPin = WeakSequencer.Pin())
diff --git a/Engine/Source/Editor/Sequencer/Private/Widgets/Sequencer/SSequencerSearchAndFilterRow.h b/Engine/Source/Editor/Sequencer/Private/Widgets/Sequencer/SSequencerSearchAndFilterRow.h
--- a/Engine/Source/Editor/Sequencer/Private/Widgets/Sequencer/SSequencerSearchAndFilterRow.h
+++ b/Engine/Source/Editor/Sequencer/Private/Widgets/Sequencer/SSequencerSearchAndFilterRow.h
@@ -6,10 +6,11 @@
 #include "Widgets/DeclarativeSyntaxSupport.h"
 #include "Widgets/SCompoundWidget.h"
 
-class SSequencerSearchBox;
 class FExtender;
 class FSequencer;
 class FSequencerFilterBar;
+class ISequencerTrackFilters;
+class SSequencerSearchBox;
 
 namespace UE::Sequencer
 {
@@ -61,9 +62,6 @@ private:
 	/** @return The filtering widget. */
 	TSharedRef<SWidget> MakeFilterContent(const FArguments& InArgs);
 	
-	/** Propagates changes in the search widget to the filter system */
-	void OnSearchTextChanged(const FText& InSearchText);
-	
 	/** Handles request to scroll the item into view. */
 	void ScrollIntoView(const TWeakViewModelPtr<IOutlinerExtension>& InItem);
 };
