Sequencer: Refactor  SSearchAndFilterWidget to take ISequencerTrackFilters.
- Promote FSequencerFilterBar::GetHideIsolateFilter to ISequencerTrackFilters
- Extract FSequencerFilterBar::MakeAddFilterButton to a global function taking FGetMenuContent
- FGetMenuConten is injected to SSearchAndFilterWidget by SSequencerSearchAndFilterRow

#Synced-CL 50532451

diff --git a/Engine/Source/Editor/Sequencer/Private/Filters/SequencerFilterBar.cpp b/Engine/Source/Editor/Sequencer/Private/Filters/SequencerFilterBar.cpp
--- a/Engine/Source/Editor/Sequencer/Private/Filters/SequencerFilterBar.cpp
+++ b/Engine/Source/Editor/Sequencer/Private/Filters/SequencerFilterBar.cpp
@@ -26,7 +26,6 @@
 #include "Filters/Widgets/SSequencerSearchBox.h"
 #include "Framework/Application/SlateApplication.h"
 #include "Framework/MultiBox/SToolBarButtonBlock.h"
-#include "Menus/SequencerTrackFilterMenu.h"
 #include "MovieScene.h"
 #include "MVVM/CurveEditorExtension.h"
 #include "MVVM/Selection/Selection.h"
@@ -61,7 +60,6 @@ FSequencerFilterBar::FSequencerFilterBar(FSequencer& InSequencer, FString InSubI
 	, GroupFilter(MakeShared<FSequencerTrackFilter_Group>(*this, TransientCategory))
 	, SelectedFilter(MakeShared<FSequencerTrackFilter_Selected>(*this, MiscCategory))
 	, ModifiedFilter(MakeShared<FSequencerTrackFilter_Modified>(*this, MiscCategory))
-	, FilterMenu(MakeShared<FSequencerTrackFilterMenu>())
 	, FilterData(FString())
 	, SubIdentifier(MoveTemp(InSubIdentifier))
 {
@@ -1658,70 +1656,6 @@ bool FSequencerFilterBar::ShouldUpdateOnTrackValueChanged() const
 	return false;
 }
 
-TSharedRef<SComboButton> FSequencerFilterBar::MakeAddFilterButton()
-{
-	const TSharedPtr<SLayeredImage> FilterImage = SNew(SLayeredImage)
-		.Image(FAppStyle::Get().GetBrush(TEXT("Icons.Filter")))
-		.ColorAndOpacity_Lambda([this]()
-			{
-				return AreFiltersMuted() ? FLinearColor(1.f, 1.f, 1.f, 0.2f) : FSlateColor::UseForeground();
-			});
-
-	// Badge the filter icon if there are filters enabled or active
-	FilterImage->AddLayer(TAttribute<const FSlateBrush*>::CreateLambda([this]() -> const FSlateBrush*
-		{
-			if (AreFiltersMuted() || !HasAnyFilterEnabled())
-			{
-				return nullptr;
-			}
-
-			if (HasAnyFilterActive(false, false))
-			{
-				return FAppStyle::Get().GetBrush(TEXT("Icons.BadgeModified"));
-			}
-
-			return FAppStyle::Get().GetBrush(TEXT("Icons.Badge"));
-		}));
-
-	const TSharedRef<SComboButton> ComboButton = SNew(SComboButton)
-		.ComboButtonStyle(&FAppStyle::Get().GetWidgetStyle<FComboButtonStyle>(TEXT("SimpleComboButtonWithIcon")))
-		.ForegroundColor(FSlateColor::UseStyle())
-		.ToolTipText_Lambda([this]()
-			{
-				return FText::Format(
-					LOCTEXT("AddFilterToolTip", "Open the Add Filter Menu to add or manage filters\n\n"
-					"Shift + Click to temporarily mute all active filters\n\n{0}"), 
-					MakeHiddenAndIsolatedCountText_WithTotal(FilterData, *HideIsolateFilter)
-					);
-			})
-		.OnComboBoxOpened_Lambda([this]()
-			{
-				// Don't allow opening the menu if filters are muted or we are toggling the filter mute state
-				if (AreFiltersMuted() || FSlateApplication::Get().GetModifierKeys().IsShiftDown())
-				{
-					FSlateApplication::Get().DismissAllMenus();
-				}
-			})
-		.OnGetMenuContent_Lambda([this]() -> TSharedRef<SWidget>
-			{
-				if (FSlateApplication::Get().GetModifierKeys().IsShiftDown())
-				{
-					MuteFilters(!AreFiltersMuted());
-					FSlateApplication::Get().DismissAllMenus();
-					return SNullWidget::NullWidget;
-				}
-				return FilterMenu->CreateMenu(SharedThis(this));
-			})
-		.ContentPadding(FMargin(1, 0))
-		.ButtonContent()
-		[
-			FilterImage.ToSharedRef()
-		];
-	ComboButton->AddMetadata(MakeShared<FTagMetaData>(TEXT("SequencerTrackFiltersCombo")));
-
-	return ComboButton;
-}
-
 bool FSequencerFilterBar::ShouldShowFilterBarWidget() const
 {
 	if (const USequencerSettings* const Settings = Sequencer.GetSequencerSettings())
diff --git a/Engine/Source/Editor/Sequencer/Private/Filters/SequencerFilterBar.h b/Engine/Source/Editor/Sequencer/Private/Filters/SequencerFilterBar.h
--- a/Engine/Source/Editor/Sequencer/Private/Filters/SequencerFilterBar.h
+++ b/Engine/Source/Editor/Sequencer/Private/Filters/SequencerFilterBar.h
@@ -74,7 +74,7 @@ public:
 	FSequencerFiltersChanged& OnFiltersChanged() { return FiltersChangedEvent; }
 
 	TSharedRef<FSequencerTrackFilter_Text> GetTextFilter() const;
-	TSharedRef<FSequencerTrackFilter_HideIsolate> GetHideIsolateFilter() const;
+	virtual TSharedRef<FSequencerTrackFilter_HideIsolate> GetHideIsolateFilter() const override;
 	FText GetFilterErrorText() const;
 
 	/** Hide/Isolate/Show Filter Functions */
@@ -229,8 +229,6 @@ public:
 
 	bool ShouldUpdateOnTrackValueChanged() const;
 
-	TSharedRef<SComboButton> MakeAddFilterButton();
-
 	virtual bool ShouldShowFilterBarWidget() const override;
 
 	virtual bool IsFilterBarVisible() const override;
@@ -301,8 +299,6 @@ protected:
 	
 	TArray<TSharedRef<FSequencerTrackFilter_CustomText>> CustomTextFilters;
 
-	TSharedPtr<FSequencerTrackFilterMenu> FilterMenu;
-
 	FSequencerFilterData FilterData;
 
 	FSequencerFiltersChanged FiltersChangedEvent;
diff --git a/Engine/Source/Editor/Sequencer/Private/Filters/Utils/FilterViewUtils.cpp b/Engine/Source/Editor/Sequencer/Private/Filters/Utils/FilterViewUtils.cpp
--- a/Engine/Source/Editor/Sequencer/Private/Filters/Utils/FilterViewUtils.cpp
+++ b/Engine/Source/Editor/Sequencer/Private/Filters/Utils/FilterViewUtils.cpp
@@ -2,7 +2,9 @@
 
 #include "FilterViewUtils.h"
 
+#include "HideIsolateViewModelUtils.h"
 #include "Filters/ISequencerTrackFilters.h"
+#include "Filters/Menus/SequencerTrackFilterMenu.h"
 #include "Filters/Widgets/SFilterExpressionHelpDialog.h"
 #include "Filters/Widgets/SSequencerCustomTextFilterDialog.h"
 #include "Templates/SharedPointer.h"
@@ -32,6 +34,91 @@ void OpenDialog_TextExpressionHelp(const TSharedRef<ISequencerTrackFilters>& InF
 
 	SFilterExpressionHelpDialog::Open(MoveTemp(Config));
 }
+
+TSharedRef<SComboButton> MakeAddFilterButton(
+	const TSharedRef<ISequencerTrackFilters>& InTrackFilters,
+	FOnGetContent InGetMenuContent
+	)
+{
+	const TSharedPtr<SLayeredImage> FilterImage = SNew(SLayeredImage)
+		.Image(FAppStyle::Get().GetBrush(TEXT("Icons.Filter")))
+		.ColorAndOpacity_Lambda([WeakFilters = InTrackFilters.ToWeakPtr()]
+		{
+			const TSharedPtr<ISequencerTrackFilters> Filters = WeakFilters.Pin();
+			return Filters && Filters->AreFiltersMuted() ? FLinearColor(1.f, 1.f, 1.f, 0.2f) : FSlateColor::UseForeground();
+		});
+
+	// Badge the filter icon if there are filters enabled or active
+	FilterImage->AddLayer(
+		TAttribute<const FSlateBrush*>::CreateLambda([WeakFilters = InTrackFilters.ToWeakPtr()]() -> const FSlateBrush*
+		{
+			const TSharedPtr<ISequencerTrackFilters> Filters = WeakFilters.Pin();
+			if (!Filters || Filters->AreFiltersMuted() || !Filters->HasAnyFilterEnabled())
+			{
+				return nullptr;
+			}
+
+			if (Filters->HasAnyFilterActive(false, false))
+			{
+				return FAppStyle::Get().GetBrush(TEXT("Icons.BadgeModified"));
+			}
+
+			return FAppStyle::Get().GetBrush(TEXT("Icons.Badge"));
+		}));
+
+	const TSharedRef<SComboButton> ComboButton = SNew(SComboButton)
+		.ComboButtonStyle(&FAppStyle::Get().GetWidgetStyle<FComboButtonStyle>(TEXT("SimpleComboButtonWithIcon")))
+		.ForegroundColor(FSlateColor::UseStyle())
+		.ToolTipText_Lambda([WeakFilters = InTrackFilters.ToWeakPtr()]
+		{
+			const TSharedPtr<ISequencerTrackFilters> Filters = WeakFilters.Pin();
+			if (!Filters)
+			{
+				return FText::GetEmpty();
+			}
+				
+			return FText::Format(
+				LOCTEXT("AddFilterToolTip", "Open the Add Filter Menu to add or manage filters\n\n"
+				"Shift + Click to temporarily mute all active filters\n\n{0}"), 
+				MakeHiddenAndIsolatedCountText_WithTotal(Filters->GetFilterData(), *Filters->GetHideIsolateFilter())
+				);
+		})
+		.OnComboBoxOpened_Lambda([WeakFilters = InTrackFilters.ToWeakPtr()]
+		{
+			const TSharedPtr<ISequencerTrackFilters> Filters = WeakFilters.Pin();
+				
+			// Don't allow opening the menu if filters are muted or we are toggling the filter mute state
+			if (!Filters || Filters->AreFiltersMuted() || FSlateApplication::Get().GetModifierKeys().IsShiftDown())
+			{
+				FSlateApplication::Get().DismissAllMenus();
+			}
+		})
+		.OnGetMenuContent_Lambda([WeakFilters = InTrackFilters.ToWeakPtr(), InGetMenuContent = MoveTemp(InGetMenuContent)]() -> TSharedRef<SWidget>
+		{
+			const TSharedPtr<ISequencerTrackFilters> Filters = WeakFilters.Pin();
+			if (!Filters)
+			{
+				return SNullWidget::NullWidget;
+			}
+				
+			if (FSlateApplication::Get().GetModifierKeys().IsShiftDown())
+			{
+				Filters->MuteFilters(!Filters->AreFiltersMuted());
+				FSlateApplication::Get().DismissAllMenus();
+				return SNullWidget::NullWidget;
+			}
+				
+			return InGetMenuContent.IsBound() ? InGetMenuContent.Execute() : SNullWidget::NullWidget;
+		})
+		.ContentPadding(FMargin(1, 0))
+		.ButtonContent()
+		[
+			FilterImage.ToSharedRef()
+		];
+	ComboButton->AddMetadata(MakeShared<FTagMetaData>(TEXT("SequencerTrackFiltersCombo")));
+
+	return ComboButton;
+}
 }
 
 #undef LOCTEXT_NAMESPACE
\ No newline at end of file
diff --git a/Engine/Source/Editor/Sequencer/Private/Filters/Utils/FilterViewUtils.h b/Engine/Source/Editor/Sequencer/Private/Filters/Utils/FilterViewUtils.h
--- a/Engine/Source/Editor/Sequencer/Private/Filters/Utils/FilterViewUtils.h
+++ b/Engine/Source/Editor/Sequencer/Private/Filters/Utils/FilterViewUtils.h
@@ -2,8 +2,10 @@
 
 #pragma once
 
+#include "Framework/SlateDelegates.h"
 #include "Templates/SharedPointerFwd.h"
 
+class SComboButton;
 class ISequencerTrackFilters;
 
 namespace UE::Sequencer
@@ -13,5 +15,11 @@ void OpenDialog_SaveCurrentFilterSetAsCustomTextFilter(const TSharedRef<ISequenc
 
 /** Opens a dialog that helps the user for creating custom text expressions. */
 void OpenDialog_TextExpressionHelp(const TSharedRef<ISequencerTrackFilters>& InFilterBar);
+
+/** @return Combo button for adding filters to ISequencerTrackFilters. */
+TSharedRef<SComboButton> MakeAddFilterButton(
+	const TSharedRef<ISequencerTrackFilters>& InTrackFilters,
+	FOnGetContent InGetMenuContent
+	);
 }
 
diff --git a/Engine/Source/Editor/Sequencer/Private/Filters/Widgets/SSearchAndFilterWidget.cpp b/Engine/Source/Editor/Sequencer/Private/Filters/Widgets/SSearchAndFilterWidget.cpp
--- a/Engine/Source/Editor/Sequencer/Private/Filters/Widgets/SSearchAndFilterWidget.cpp
+++ b/Engine/Source/Editor/Sequencer/Private/Filters/Widgets/SSearchAndFilterWidget.cpp
@@ -11,6 +11,7 @@
 #include "SPositiveActionButton.h"
 #include "SSequencerFilterBar.h"
 #include "Sequencer.h"
+#include "Filters/Utils/FilterViewUtils.h"
 
 #define LOCTEXT_NAMESPACE "SSearchAndFilterRow"
 
@@ -18,7 +19,7 @@ namespace UE::Sequencer
 {
 void SSearchAndFilterWidget::Construct(
 	const FArguments& InArgs, 
-	const TSharedRef<FSequencerFilterBar>& InFilterBar,
+	const TSharedRef<ISequencerTrackFilters>& InFilterBar,
 	const TSharedRef<FHideIsolateShowViewModel>& InIsolateHideViewModel
 	)
 {
@@ -35,7 +36,7 @@ void SSearchAndFilterWidget::Construct(
 		.VAlign(VAlign_Center)
 		.Padding(2.f, 0.f, 0.f, 0.f)
 		[
-			InFilterBar->MakeAddFilterButton()
+			MakeAddFilterButton(InFilterBar, InArgs._OnMakeAddFilterMenuContent)
 		]
 
 		// Advanced Search Box
diff --git a/Engine/Source/Editor/Sequencer/Private/Filters/Widgets/SSearchAndFilterWidget.h b/Engine/Source/Editor/Sequencer/Private/Filters/Widgets/SSearchAndFilterWidget.h
--- a/Engine/Source/Editor/Sequencer/Private/Filters/Widgets/SSearchAndFilterWidget.h
+++ b/Engine/Source/Editor/Sequencer/Private/Filters/Widgets/SSearchAndFilterWidget.h
@@ -7,7 +7,6 @@
 #include "Widgets/DeclarativeSyntaxSupport.h"
 #include "Widgets/SCompoundWidget.h"
 
-class FSequencerFilterBar;
 class ISequencerTrackFilters;
 class SSequencerSearchBox;
 
@@ -27,11 +26,14 @@ public:
 	SLATE_BEGIN_ARGS(SSearchAndFilterWidget){}
 		/** Invoked when the search text changes. */
 		SLATE_EVENT(FOnTextChanged, OnSearchTextChanged)
+
+		/** Makes the content for the add filter combo button. */
+		SLATE_EVENT(FOnGetContent, OnMakeAddFilterMenuContent)
 	SLATE_END_ARGS()
 
 	void Construct(
 		const FArguments& InArgs, 
-		const TSharedRef<FSequencerFilterBar>& InFilterBar,
+		const TSharedRef<ISequencerTrackFilters>& InFilterBar,
 		const TSharedRef<FHideIsolateShowViewModel>& InIsolateHideViewModel
 		);
 	
diff --git a/Engine/Source/Editor/Sequencer/Private/Widgets/Sequencer/SSequencerSearchAndFilterRow.cpp b/Engine/Source/Editor/Sequencer/Private/Widgets/Sequencer/SSequencerSearchAndFilterRow.cpp
--- a/Engine/Source/Editor/Sequencer/Private/Widgets/Sequencer/SSequencerSearchAndFilterRow.cpp
+++ b/Engine/Source/Editor/Sequencer/Private/Widgets/Sequencer/SSequencerSearchAndFilterRow.cpp
@@ -16,6 +16,7 @@
 #include "MVVM/Selection/Selection.h"
 #include "MVVM/Views/SOutlinerView.h"
 #include "SSequencer.h"
+#include "Filters/Menus/SequencerTrackFilterMenu.h"
 #include "Widgets/Sequencer/SAddToSequencerButton.h"
 #include "Widgets/Sequencer/SSequencerViewOptionsButton.h"
 
@@ -55,6 +56,15 @@ static FOutlinerExtensionsAttr MakeAllTracksAttr(const TSharedRef<FSequencer>& I
 		return SequencerPin ? GetAllTracks(*SequencerPin->GetViewModel()) : TSet<TWeakViewModelPtr<IOutlinerExtension>>();
 	});
 }
+
+static FOnGetContent MakeDelegate_AddFilterMenuContent(const TSharedRef<FSequencerFilterBar>& InFilterBar)
+{
+	return FOnGetContent::CreateLambda([Menu = MakeShared<FSequencerTrackFilterMenu>(), WeakFilterBar = InFilterBar.ToWeakPtr()]
+	{
+		const TSharedPtr<FSequencerFilterBar> FilterBar = WeakFilterBar.Pin();
+		return ensure(FilterBar) ? Menu->CreateMenu(FilterBar.ToSharedRef()) : SNullWidget::NullWidget;
+	});
+}
 }
 void SSequencerSearchAndFilterRow::Construct(const FArguments& InArgs)
 {
@@ -128,9 +138,10 @@ TSharedRef<SWidget> SSequencerSearchAndFilterRow::MakeFilterContent(const FArgum
 			FilterRowDetail::MakeSelectedTracksAttr(SequencerPin), FilterRowDetail::MakeAllTracksAttr(SequencerPin)
 			);
 		IsolateViewModel->OnRequestScrollToView().AddSP(this, &SSequencerSearchAndFilterRow::ScrollIntoView);
-	
+
 		const TSharedPtr<ISequencerTrackFilters> TrackFilters = StaticCastSharedPtr<ISequencerTrackFilters>(InFilterBar.ToSharedPtr());
 		return SNew(SSearchAndFilterWidget, InFilterBar, IsolateViewModel)
+			.OnMakeAddFilterMenuContent(FilterRowDetail::MakeDelegate_AddFilterMenuContent(InFilterBar))
 			.OnSearchTextChanged_Lambda([InFilterBar](const FText& InSearchText)
 			{
 				InFilterBar->SetTextFilterString(InSearchText.ToString());
diff --git a/Engine/Source/Editor/Sequencer/Public/Filters/ISequencerTrackFilters.h b/Engine/Source/Editor/Sequencer/Public/Filters/ISequencerTrackFilters.h
--- a/Engine/Source/Editor/Sequencer/Public/Filters/ISequencerTrackFilters.h
+++ b/Engine/Source/Editor/Sequencer/Public/Filters/ISequencerTrackFilters.h
@@ -4,6 +4,7 @@
 
 #include "ISequencerFilterBar.h"
 
+class FSequencerTrackFilter_HideIsolate;
 class FSequencerTrackFilter;
 class FSequencerTrackFilter_CustomText;
 struct FSequencerFilterData;
@@ -21,7 +22,8 @@ public:
 
 	virtual bool AddCustomTextFilter(const TSharedRef<FSequencerTrackFilter_CustomText>& InFilter, const bool bInAddToConfig) = 0;
 	virtual bool RemoveCustomTextFilter(const TSharedRef<FSequencerTrackFilter_CustomText>& InFilter, const bool bInAddToConfig) = 0;
-
+	
+	virtual TSharedRef<FSequencerTrackFilter_HideIsolate> GetHideIsolateFilter() const = 0;
 	virtual void HideSelectedTracks() = 0;
 	virtual void IsolateSelectedTracks() = 0;
 
