Fix OnUnlink not being called when closing Sequencer after PIE.

When entering and exiting PIE Sequencer's playback context changes, causing the Root Template Instances's runner to change, but the Sequence's Runner ptr was not updated to the new one, causing OnUnlink to not be called. We now update the Runner pointer when the playback context changes.

#Synced-CL 50607091

diff --git a/Engine/Source/Editor/Sequencer/Private/Sequencer.cpp b/Engine/Source/Editor/Sequencer/Private/Sequencer.cpp
--- a/Engine/Source/Editor/Sequencer/Private/Sequencer.cpp
+++ b/Engine/Source/Editor/Sequencer/Private/Sequencer.cpp
@@ -810,6 +810,8 @@ void FSequencer::SetSequencerSettings(USequencerSettings* InSettings)
 void FSequencer::OnPlaybackContextChanged()
 {
 	RootTemplateInstance.PlaybackContextChanged(*this);
+	// Update runner to match the newly created shared payback state of the root template instance.
+	Runner = RootTemplateInstance.GetRunner();
 	InitRootSequenceInstance();
 }
 
