Sequencer: Introduces FFilterEvaluator, which can be used to register multiple FSequencerFilterBar instances to evaluate filters with.
- FSequencerNodeTree now uses FFilterEvaluator to evaluate filters
- FLinkedFilteringViewModel handles registering the currently active FSequencerFilterBar with FFilterEvaluator

Linked filtering now supports filters that update on track value changes.

#Synced-CL 50532451

diff --git a/Engine/Source/Editor/Sequencer/Private/Filters/ViewModels/FilterEvaluator.cpp b/Engine/Source/Editor/Sequencer/Private/Filters/ViewModels/FilterEvaluator.cpp
new file mode 100644
--- /dev/null
+++ b/Engine/Source/Editor/Sequencer/Private/Filters/ViewModels/FilterEvaluator.cpp
@@ -0,0 +1,130 @@
+﻿// Copyright Epic Games, Inc. All Rights Reserved.
+
+#include "FilterEvaluator.h"
+
+#include "Filters/SequencerFilterBar.h"
+
+namespace UE::Sequencer
+{
+bool FFilterEvaluator::UpdateFiltersOnTrackValueChanged()
+{
+	if (bHasRequestedFilterUpdate)
+	{
+		return false;
+	}
+	
+	bool bAnyFilterNeedsRerun = false;
+	for (FFilterBarData& FilterBarData : RegisteredFilterBars)
+	{
+		const bool bNeedsRerun = FilterBarData.FilterBar->ShouldUpdateOnTrackValueChanged();
+		FilterBarData.bHasRequestedRefilter |= bNeedsRerun;
+		bAnyFilterNeedsRerun |= bNeedsRerun;
+	}
+	
+	if (bAnyFilterNeedsRerun)
+	{
+		constexpr bool bForce = true;
+		const bool bFiltersUpdated = UpdateFilters(bForce);
+
+		// If the filter list was modified, set bHasRequestedFilterUpdate to suppress excessive re-filters between tree update
+		bHasRequestedFilterUpdate = bFiltersUpdated;
+		return bFiltersUpdated;
+	}
+
+	return false;
+}
+
+bool FFilterEvaluator::UpdateFilters(bool bForce)
+{
+	if (!bForce && !bHasRequestedFilterUpdate)
+	{
+		return false;
+	}
+
+	bool bHasFilterDataChanged = false; 
+	for (FFilterBarData& FilterBarData : RegisteredFilterBars)
+	{
+		const TSharedRef<FSequencerFilterBar>& FilterBar = FilterBarData.FilterBar;
+		
+		const FSequencerFilterData PreviousFilterData = FilterBar->GetFilterData();
+		const FSequencerFilterData& FilterData = FilterBar->FilterNodes();
+		const bool bHasFilteredDataChanged = FilterData != PreviousFilterData;
+		bHasFilterDataChanged |= bHasFilteredDataChanged;
+		FilterBarData.bHasRequestedRefilter = false;
+		
+		OnFilterBarFilteredDelegate.Broadcast(FilterBar);
+	}
+	bHasRequestedFilterUpdate = false;
+
+	return bHasFilterDataChanged;
+}
+
+void FFilterEvaluator::RequestFilterUpdate(const TSharedRef<FSequencerFilterBar>& InFilterBar)
+{
+	const int32 Index = IndexOf(InFilterBar);
+	if (RegisteredFilterBars.IsValidIndex(Index))
+	{
+		bHasRequestedFilterUpdate = true;
+		RegisteredFilterBars[Index].bHasRequestedRefilter = true;
+	}
+}
+
+bool FFilterEvaluator::NeedsFilterUpdate() const
+{
+	return bHasRequestedFilterUpdate;
+}
+
+void FFilterEvaluator::IncrementFilterBarUsage(const TSharedRef<FSequencerFilterBar>& InFilterBar)
+{
+	const int32 Index = IndexOf(InFilterBar);
+	if (RegisteredFilterBars.IsValidIndex(Index))
+	{
+		++RegisteredFilterBars[Index].UsageCount;
+	}
+	else
+	{
+		FSimpleDelegate HandleRequestUpdate = FSimpleDelegate::CreateRaw(this, &FFilterEvaluator::OnFilterRequestedRefresh, InFilterBar.ToWeakPtr());
+		RegisteredFilterBars.Emplace(InFilterBar, MoveTemp(HandleRequestUpdate));
+	}
+	
+	// We run the filter just in case it's not been evaluated yet. In the future, we could add a dirty flag or serial number to the filter bar.
+	bHasRequestedFilterUpdate = true;
+}
+
+void FFilterEvaluator::DecrementFilterBarUsage(const TSharedRef<FSequencerFilterBar>& InFilterBar)
+{
+	const int32 Index = IndexOf(InFilterBar);
+	if (!ensure(RegisteredFilterBars.IsValidIndex(Index)))
+	{
+		return;
+	}
+	
+	FFilterBarData& Data = RegisteredFilterBars[Index];
+	check(Data.UsageCount > 0);
+	
+	if (Data.UsageCount == 1)
+	{
+		RegisteredFilterBars.RemoveAtSwap(Index);
+	}
+	else
+	{
+		--Data.UsageCount;
+	}
+}
+
+int32 FFilterEvaluator::IndexOf(const TSharedRef<FSequencerFilterBar>& InFilterBar) const
+{
+	return RegisteredFilterBars.IndexOfByPredicate([&InFilterBar](const FFilterBarData& Data)
+	{
+		return InFilterBar == Data.FilterBar;
+	});
+}
+
+void FFilterEvaluator::OnFilterRequestedRefresh(TWeakPtr<FSequencerFilterBar> InWeakFilterBar)
+{
+	if (const TSharedPtr<FSequencerFilterBar> FilterBar = InWeakFilterBar.Pin(); ensure(FilterBar))
+	{
+		RequestFilterUpdate(FilterBar.ToSharedRef());
+	}
+}
+} // namespace UE::Sequencer
diff --git a/Engine/Source/Editor/Sequencer/Private/Filters/ViewModels/FilterEvaluator.h b/Engine/Source/Editor/Sequencer/Private/Filters/ViewModels/FilterEvaluator.h
new file mode 100644
--- /dev/null
+++ b/Engine/Source/Editor/Sequencer/Private/Filters/ViewModels/FilterEvaluator.h
@@ -0,0 +1,106 @@
+﻿// Copyright Epic Games, Inc. All Rights Reserved.
+
+#pragma once
+
+#include "Containers/Array.h"
+#include "Delegates/Delegate.h"
+#include "Filters/SequencerFilterBar.h"
+#include "Templates/SharedPointer.h"
+#include "Templates/UnrealTemplate.h"
+
+class FSequencerFilterBar;
+
+namespace UE::Sequencer
+{
+/**
+ * Knows of all filter bar instances that need to be run in Sequencer.
+ * This is relevant for multiple UI endpoints using the linked filtering system.
+ * 
+ * Each endpoint can register a filter bar they want evaluated, and request explicit refreshes.
+ * The evaluator will be refreshed next tick by Sequencer.
+ * @see IncrementFilterBarUsage, DecrementFilterBarUsage
+ * 
+ * If your filter UI becomes hidden (e.g. tab in background), as optimization, it would be valid to call DecrementFilterBarUsage and follow it up 
+ * again with a IncrementFilterBarUsage once the UI becomes active again.
+ */
+class FFilterEvaluator : public FNoncopyable
+{
+public:
+	
+	/** 
+	 * Evaluates all filters bars that need updating when a track value changes. 
+	 * @return Whether the new list of filtered nodes is different from the previous list
+	 */
+	bool UpdateFiltersOnTrackValueChanged();
+	
+	/**
+	 * Evaluates all filter bars. 
+	 * @return Whether the new list of filtered nodes is different from the previous list
+	 */
+	bool UpdateFilters(bool bForce = false);
+	
+	/** Schedules an update of all filters on the given filter bar. The filter bar must be registered. */
+	void RequestFilterUpdate(const TSharedRef<FSequencerFilterBar>& InFilterBar);
+	/** @return Whether there is a filter update scheduled */
+	bool NeedsFilterUpdate() const;
+
+	/** 
+	 * Tracks a filter bar that some UI endpoints is displaying filter results for.
+	 * 
+	 * If this is already tracked, increments the usage counter by 1.
+	 * If previously untracked, this creates a subscription ISequencerFilterBar::OnRequestUpdate.
+	 * 
+	 * @note You should prefer FScopedFilterBarUsage for auto deregistration
+	 * @see ScopedFilterBarUsage.
+	 * 
+	 * Each call must be matched with DecrementFilterBarUsage.
+	 * @see DecrementFilterBarUsage
+	 */
+	void IncrementFilterBarUsage(const TSharedRef<FSequencerFilterBar>& InFilterBar);
+	/** Decrements the usage counter by 1. If it reaches 0, the filter bar becomes untracked for filter evaluation. */
+	void DecrementFilterBarUsage(const TSharedRef<FSequencerFilterBar>& InFilterBar);
+	
+	DECLARE_MULTICAST_DELEGATE_OneParam(FOnFilterBarFiltered, const TSharedRef<FSequencerFilterBar>&)
+	/** @return Delegate invoked when a filter bar is refiltered. */
+	FOnFilterBarFiltered& OnFilterBarFiltered() { return OnFilterBarFilteredDelegate; }
+	
+private:
+	
+	/** Whether any of the filters has a refilter requested. */
+	bool bHasRequestedFilterUpdate = false;
+	
+	struct FFilterBarData : FNoncopyable // Avoid accidental copying so we don't unsubscribe by accident.
+	{
+		/** The filter bar this data corresponds to. */
+		const TSharedRef<FSequencerFilterBar> FilterBar;
+		/** Handle to ISequencerFilterBar::OnRequestUpdate. */
+		const FDelegateHandle OnRequestUpdateHandle;
+		
+		/** How many systems have registered to be using this filter bar. Once this reaches 0, the entry is removed. */
+		int32 UsageCount = 1;
+		/** Whether this filter bar has been marked to be refiltered. */
+		bool bHasRequestedRefilter = false;
+
+		explicit FFilterBarData(const TSharedRef<FSequencerFilterBar>& InFilterBar, FSimpleDelegate HandleFilterRequestRefresh)
+			: FilterBar(InFilterBar)
+			, OnRequestUpdateHandle(FilterBar->OnRequestUpdate().Add(MoveTemp(HandleFilterRequestRefresh)))
+		{}
+		
+		~FFilterBarData()
+		{
+			FilterBar->OnRequestUpdate().Remove(OnRequestUpdateHandle);
+		}
+	};
+	/** Filters that should be considered for evaluation. */
+	TArray<FFilterBarData> RegisteredFilterBars;
+	
+	/** Invoked when a filter bar is refiltered. */
+	FOnFilterBarFiltered OnFilterBarFilteredDelegate;
+	
+	/** @return Index to RegisteredFilterBars entry that contains InFilterBar. */
+	int32 IndexOf(const TSharedRef<FSequencerFilterBar>& InFilterBar) const;
+	
+	/** Handles ISequencerTrackFilters::OnRequestRefresh by calling RequestFilterUpdate. */
+	void OnFilterRequestedRefresh(TWeakPtr<FSequencerFilterBar> InWeakFilterBar);
+};
+} // namespace UE::Sequencer
diff --git a/Engine/Source/Editor/Sequencer/Private/Filters/ViewModels/Linking/LinkedFilterFactoryViewModel.cpp b/Engine/Source/Editor/Sequencer/Private/Filters/ViewModels/Linking/LinkedFilterFactoryViewModel.cpp
--- a/Engine/Source/Editor/Sequencer/Private/Filters/ViewModels/Linking/LinkedFilterFactoryViewModel.cpp
+++ b/Engine/Source/Editor/Sequencer/Private/Filters/ViewModels/Linking/LinkedFilterFactoryViewModel.cpp
@@ -9,16 +9,18 @@ namespace UE::Sequencer
 {
 FLinkedFilterFactoryViewModel::FLinkedFilterFactoryViewModel(
 	FSequencer& InSequencer UE_LIFETIMEBOUND,
-	const TSharedRef<FSequencerFilterBar>& InSharedFilterBar
+	const TSharedRef<FSequencerFilterBar>& InSharedFilterBar,
+	const TSharedRef<FFilterEvaluator> InFilterEvaluator
 	)
 	: Sequencer(InSequencer)
 	, SharedFilterBar(InSharedFilterBar)
+	, FilterEvaluator(InFilterEvaluator)
 {}
 
 TSharedRef<FLinkedFilterViewModel> FLinkedFilterFactoryViewModel::MakeFilteringModelImpl(const FMakeLinkedFilterViewModelArgs& InArgs) const
 {
 	const TSharedRef<FSequencerFilterBar> InstancedFilterBar = MakeShared<FSequencerFilterBar>(Sequencer, InArgs.ConfigSubIdentifier);
-	return MakeShared<FLinkedFilterViewModel>(SharedFilterBar, InstancedFilterBar);
+	return MakeShared<FLinkedFilterViewModel>(SharedFilterBar, InstancedFilterBar, FilterEvaluator);
 }
 
 TSharedRef<ILinkedFilterViewModel> FLinkedFilterFactoryViewModel::MakeFilteringModel(const FMakeLinkedFilterViewModelArgs& InArgs)
diff --git a/Engine/Source/Editor/Sequencer/Private/Filters/ViewModels/Linking/LinkedFilterFactoryViewModel.h b/Engine/Source/Editor/Sequencer/Private/Filters/ViewModels/Linking/LinkedFilterFactoryViewModel.h
--- a/Engine/Source/Editor/Sequencer/Private/Filters/ViewModels/Linking/LinkedFilterFactoryViewModel.h
+++ b/Engine/Source/Editor/Sequencer/Private/Filters/ViewModels/Linking/LinkedFilterFactoryViewModel.h
@@ -8,6 +8,7 @@
 
 namespace UE::Sequencer
 {
+class FFilterEvaluator;
 class FLinkedFilterViewModel;
 
 class FLinkedFilterFactoryViewModel : public ILinkedFilterFactoryViewModel
@@ -16,7 +17,8 @@ public:
 	
 	explicit FLinkedFilterFactoryViewModel(
 		FSequencer& InSequencer UE_LIFETIMEBOUND,
-		const TSharedRef<FSequencerFilterBar>& InSharedFilterBar
+		const TSharedRef<FSequencerFilterBar>& InSharedFilterBar,
+		const TSharedRef<FFilterEvaluator> InFilterEvaluator
 		);
 	
 	/** @return View model that manages linked and instanced filter state for your UI.  */
@@ -33,5 +35,8 @@ private:
 	
 	/** The filter bar instance that model edit when they are in linked mode. All created models reference this. */
 	const TSharedRef<FSequencerFilterBar> SharedFilterBar;
+	
+	/** Instance that the active filter bar instances are registered with. This way Sequencer will issue a refilter eventually. */
+	const TSharedRef<FFilterEvaluator> FilterEvaluator;
 };
 } // namespace UE::Sequencer
diff --git a/Engine/Source/Editor/Sequencer/Private/Filters/ViewModels/Linking/LinkedFilterViewModel.cpp b/Engine/Source/Editor/Sequencer/Private/Filters/ViewModels/Linking/LinkedFilterViewModel.cpp
--- a/Engine/Source/Editor/Sequencer/Private/Filters/ViewModels/Linking/LinkedFilterViewModel.cpp
+++ b/Engine/Source/Editor/Sequencer/Private/Filters/ViewModels/Linking/LinkedFilterViewModel.cpp
@@ -3,15 +3,19 @@
 #include "LinkedFilterViewModel.h"
 
 #include "Filters/SequencerFilterBar.h"
+#include "Filters/ViewModels/FilterEvaluator.h"
 
 namespace UE::Sequencer
 {
 FLinkedFilterViewModel::FLinkedFilterViewModel(
 	const TSharedRef<FSequencerFilterBar>& InSharedFilterBar,
-	const TSharedRef<FSequencerFilterBar>& InInstancedFilterBar
+	const TSharedRef<FSequencerFilterBar>& InInstancedFilterBar,
+	const TSharedRef<FFilterEvaluator> InFilterEvaluator
 	)
 	: LinkedFilterBar(InSharedFilterBar)
 	, InstancedFilterBar(InInstancedFilterBar)
+	, FilterEvaluator(InFilterEvaluator)
+	, ActiveFilterBarUsage(GetActiveFilterBarImpl(), FilterEvaluator)
 {}
 
 TSharedRef<FSequencerFilterBar> FLinkedFilterViewModel::GetActiveFilterBarImpl() const
@@ -20,8 +24,8 @@ TSharedRef<FSequencerFilterBar> FLinkedFilterViewModel::GetActiveFilterBarImpl()
 	switch (FilterMode)
 	{
 	case ELinkedFilterMode::Linked: return LinkedFilterBar;
-	case ELinkedFilterMode::Instanced: [[fallthrough]];
-	default: return InstancedFilterBar;
+	case ELinkedFilterMode::Instanced: return InstancedFilterBar;
+	default: checkNoEntry(); return LinkedFilterBar;
 	}
 }
 
@@ -40,7 +44,9 @@ void FLinkedFilterViewModel::SetFilterMode(ELinkedFilterMode InNewMode)
 	if (InNewMode != FilterMode)
 	{
 		FilterMode = InNewMode;
-		GetActiveFilterBar()->RequestFilterUpdate();
+		
+		const TSharedRef<FSequencerFilterBar> NewActiveFilterBar = GetActiveFilterBarImpl();
+		ActiveFilterBarUsage = FScopedFilterBarUsage(NewActiveFilterBar, FilterEvaluator);
 		OnFilterModeChangedDelegate.Broadcast();
 	}
 }
diff --git a/Engine/Source/Editor/Sequencer/Private/Filters/ViewModels/Linking/LinkedFilterViewModel.h b/Engine/Source/Editor/Sequencer/Private/Filters/ViewModels/Linking/LinkedFilterViewModel.h
--- a/Engine/Source/Editor/Sequencer/Private/Filters/ViewModels/Linking/LinkedFilterViewModel.h
+++ b/Engine/Source/Editor/Sequencer/Private/Filters/ViewModels/Linking/LinkedFilterViewModel.h
@@ -2,6 +2,7 @@
 
 #pragma once
 
+#include "Filters/ViewModels/ScopedFilterBarUsage.h"
 #include "Filters/ViewModels/Linking/ILinkedFilterViewModel.h"
 #include "Templates/SharedPointer.h"
 
@@ -10,13 +11,16 @@ class FSequencerFilterBar;
 
 namespace UE::Sequencer
 {
+class FFilterEvaluator;
+
 class FLinkedFilterViewModel : public ILinkedFilterViewModel
 {
 public:
 	
 	explicit FLinkedFilterViewModel(
 		const TSharedRef<FSequencerFilterBar>& InSharedFilterBar,
-		const TSharedRef<FSequencerFilterBar>& InInstancedFilterBar
+		const TSharedRef<FSequencerFilterBar>& InInstancedFilterBar,
+		const TSharedRef<FFilterEvaluator> InFilterEvaluator
 		);
 	
 	TSharedRef<FSequencerFilterBar> GetActiveFilterBarImpl() const;
@@ -42,6 +46,11 @@ private:
 	/** The current filter mode. It determines which filter bar instance is used. */
 	ELinkedFilterMode FilterMode = ELinkedFilterMode::Linked;
 	
+	/** Ensures that the active filter is refiltered when a change is requested on it. The active filter is always registered with it. */
+	const TSharedRef<FFilterEvaluator> FilterEvaluator;
+	/** When the active filter bar changes, this is leveraged to orchestrate FFilterEvaluator::IncrementFilterBarUsage and DecrementFilterBarUsage. */
+	FScopedFilterBarUsage ActiveFilterBarUsage;
+	
 	/** Event when the filter mode changes, i.e. the instance returned by GetActiveFilterBar changes. */
 	FSimpleMulticastDelegate OnFilterModeChangedDelegate;
 };
diff --git a/Engine/Source/Editor/Sequencer/Private/Filters/ViewModels/ScopedFilterBarUsage.h b/Engine/Source/Editor/Sequencer/Private/Filters/ViewModels/ScopedFilterBarUsage.h
new file mode 100644
--- /dev/null
+++ b/Engine/Source/Editor/Sequencer/Private/Filters/ViewModels/ScopedFilterBarUsage.h
@@ -0,0 +1,69 @@
+﻿// Copyright Epic Games, Inc. All Rights Reserved.
+
+#pragma once
+
+#include "FilterEvaluator.h"
+#include "Templates/SharedPointer.h"
+#include "Templates/UnrealTemplate.h"
+
+class FSequencerFilterBar;
+
+namespace UE::Sequencer
+{
+class FFilterEvaluator;
+
+/**
+ * Registers FSequencerFilterBar with FFilterEvaluator for the lifetime of this scope object.
+ * Handles calling FFilterEvaluator::IncrementFilterBarUsage and FFilterEvaluator::DecrementFilterBarUsage.
+ */
+class FScopedFilterBarUsage : public FNoncopyable
+{
+public:
+	
+	explicit FScopedFilterBarUsage(const TSharedRef<FSequencerFilterBar> InFilterBar, const TSharedRef<FFilterEvaluator>& InEvaluator)
+		: FilterBar(InFilterBar)
+		, Evaluator(InEvaluator)
+	{
+		Evaluator->IncrementFilterBarUsage(InFilterBar);
+	}
+	
+	FScopedFilterBarUsage(FScopedFilterBarUsage&& InOther)
+		: FilterBar(MoveTemp(InOther.FilterBar))
+		, Evaluator(MoveTemp(InOther.Evaluator))
+	{}
+
+	FScopedFilterBarUsage& operator=(FScopedFilterBarUsage&& Other)
+	{
+		if (this != &Other)
+		{
+			CleanUp();
+		}
+		
+		FilterBar = MoveTemp(Other.FilterBar);
+		Evaluator = MoveTemp(Other.Evaluator);
+		return *this;
+	}
+
+	~FScopedFilterBarUsage()
+	{
+		CleanUp();
+	}
+	
+private:
+	
+	/** The filter bar that is being used. */
+	TSharedPtr<FSequencerFilterBar> FilterBar;
+	
+	/** The filter evaluator to register and unregister FilterBar with. */
+	TSharedPtr<FFilterEvaluator> Evaluator;
+	
+	/** Calls DecrementFilterBarUsage if this instance state has not been moved. */
+	void CleanUp() const
+	{
+		if (Evaluator && FilterBar)
+		{
+			Evaluator->DecrementFilterBarUsage(FilterBar.ToSharedRef());
+		}
+	}
+};
+} // namespace UE::Sequencer
diff --git a/Engine/Source/Editor/Sequencer/Private/MVVM/ViewModels/LinkedFilteringMediator.cpp b/Engine/Source/Editor/Sequencer/Private/MVVM/ViewModels/LinkedFilteringMediator.cpp
--- a/Engine/Source/Editor/Sequencer/Private/MVVM/ViewModels/LinkedFilteringMediator.cpp
+++ b/Engine/Source/Editor/Sequencer/Private/MVVM/ViewModels/LinkedFilteringMediator.cpp
@@ -11,22 +11,16 @@
 
 namespace UE::Sequencer
 {
-FLinkedFilteringMediator::FLinkedFilteringMediator(
-	const TSharedRef<FSequencerNodeTree>& InNodeTree,
-	const TSharedRef<FLinkedFilterViewModel>& InSequencerFilterViewModel
-	)
-	: WeakNodeTree(InNodeTree)
-	, SequencerFilterViewModel(InSequencerFilterViewModel)
+FLinkedFilteringMediator::FLinkedFilteringMediator(const TSharedRef<FLinkedFilterViewModel>& InSequencerFilterViewModel)
+	: SequencerFilterViewModel(InSequencerFilterViewModel)
 {
 	SequencerFilterViewModel->OnFilterModeChanged().AddRaw(this, &FLinkedFilteringMediator::OnFilterModeChanged);
 	
 	const TSharedRef<FSequencerFilterBar> SharedFilterBar = SequencerFilterViewModel->GetLinkedFilterBarImpl();
-	SharedFilterBar->OnRequestUpdate().AddRaw(this, &FLinkedFilteringMediator::OnRequestUpdate_LinkedFilterBar);
 	SharedFilterBar->OnChangeOutlinerExtensionFilterState().AddRaw(this, &FLinkedFilteringMediator::OnFilterStateChanged_LinkedFilterBar);
 	
 	const TSharedRef<FSequencerFilterBar> InstancedFilterBar = SequencerFilterViewModel->GetInstancedFilterBarImpl();
 	InstancedFilterBar->OnChangeOutlinerExtensionFilterState().AddRaw(this, &FLinkedFilteringMediator::OnFilterStateChanged_InstancedFilterBar);
-	InstancedFilterBar->OnRequestUpdate().AddRaw(this, &FLinkedFilteringMediator::OnRequestUpdate_InstancedFilterBar);
 }
 
 FLinkedFilteringMediator::~FLinkedFilteringMediator()
@@ -35,11 +29,9 @@ FLinkedFilteringMediator::~FLinkedFilteringMediator()
 	
 	const TSharedRef<FSequencerFilterBar> SharedFilterBar = SequencerFilterViewModel->GetLinkedFilterBarImpl();
 	SharedFilterBar->OnChangeOutlinerExtensionFilterState().RemoveAll(this);
-	SharedFilterBar->OnRequestUpdate().RemoveAll(this);
 	
 	const TSharedRef<FSequencerFilterBar> InstancedFilterBar = SequencerFilterViewModel->GetInstancedFilterBarImpl();
 	InstancedFilterBar->OnChangeOutlinerExtensionFilterState().RemoveAll(this);
-	InstancedFilterBar->OnRequestUpdate().RemoveAll(this);
 }
 
 void FLinkedFilteringMediator::BindCommands(const TSharedRef<FUICommandList>& InSequencerCommandList) const
@@ -86,15 +78,6 @@ void FLinkedFilteringMediator::OnFilterStateChanged_LinkedFilterBar(
 	}
 }
 
-void FLinkedFilteringMediator::OnRequestUpdate_LinkedFilterBar() const
-{
-	const TSharedPtr<FSequencerNodeTree> NodeTreePin = WeakNodeTree.Pin();
-	if (NodeTreePin && GetFilterMode() == ELinkedFilterMode::Linked)
-	{
-		NodeTreePin->RequestFilterUpdate();
-	}
-}
-
 void FLinkedFilteringMediator::OnFilterStateChanged_InstancedFilterBar(
 	const TViewModelPtr<IOutlinerExtension>& InItem, bool bInIsFilteredOut
 	) const
@@ -104,13 +87,4 @@ void FLinkedFilteringMediator::OnFilterStateChanged_InstancedFilterBar(
 		InItem->SetFilteredOut(bInIsFilteredOut);
 	}
 }
-
-void FLinkedFilteringMediator::OnRequestUpdate_InstancedFilterBar() const
-{
-	const TSharedPtr<FSequencerNodeTree> NodeTreePin = WeakNodeTree.Pin();
-	if (NodeTreePin && GetFilterMode() == ELinkedFilterMode::Instanced)
-	{
-		NodeTreePin->RequestFilterUpdate();
-	}
-}
 } // namespace UE::Sequencer
diff --git a/Engine/Source/Editor/Sequencer/Private/MVVM/ViewModels/LinkedFilteringMediator.h b/Engine/Source/Editor/Sequencer/Private/MVVM/ViewModels/LinkedFilteringMediator.h
--- a/Engine/Source/Editor/Sequencer/Private/MVVM/ViewModels/LinkedFilteringMediator.h
+++ b/Engine/Source/Editor/Sequencer/Private/MVVM/ViewModels/LinkedFilteringMediator.h
@@ -6,7 +6,6 @@
 #include "MVVM/ViewModels/ViewModel.h"
 
 class FUICommandList;
-class FSequencerNodeTree;
 
 namespace UE::Sequencer
 {
@@ -19,17 +18,13 @@ enum class ELinkedFilterMode : uint8;
  * Its responsibility is to handle updating any systems when the linked filtering mode changes.
  * For example:
  * - Calling IOutlinerExtension::SetFilteredOut when the item's filter state changes
- * - Causing the FSequencerNodeTree to refilter & refresh the tree view when the active FSequencerFilterBar is requested to refilter
  * - Binding commands that are bound by FSequencerFilterBar and making sure correct commands are executed depending on which one is active.
  */
 class FLinkedFilteringMediator : public FViewModel
 {
 public:
 	
-	explicit FLinkedFilteringMediator(
-		const TSharedRef<FSequencerNodeTree>& InNodeTree,
-		const TSharedRef<FLinkedFilterViewModel>& InSequencerFilterViewModel
-		);
+	explicit FLinkedFilteringMediator(const TSharedRef<FLinkedFilterViewModel>& InSequencerFilterViewModel);
 	virtual ~FLinkedFilteringMediator() override;
 	
 	/** Binds the filter bar commands. */
@@ -40,9 +35,6 @@ public:
 	
 private:
 	
-	/** Needed to cause the tree view to refresh. */
-	const TWeakPtr<FSequencerNodeTree> WeakNodeTree;
-	
 	/** Knows whether Sequencer is using linked or instanced filtering. */
 	const TSharedRef<FLinkedFilterViewModel> SequencerFilterViewModel;
 	
@@ -54,12 +46,8 @@ private:
 	
 	/** Handles propagating the IOutlinerExtension's filter out state if the active filter bar is the linked one. */
 	void OnFilterStateChanged_LinkedFilterBar(const TViewModelPtr<IOutlinerExtension>& InItem, bool bInIsFilteredOut) const;
-	/** Propagates filter changes from the linked filter bar to the Sequencer tree.  */
-	void OnRequestUpdate_LinkedFilterBar() const;
 	
 	/** Handles propagating the IOutlinerExtension's filter out state if the active filter bar is the instanced one. */
 	void OnFilterStateChanged_InstancedFilterBar(const TViewModelPtr<IOutlinerExtension>& InItem, bool bInIsFilteredOut) const;
-	/** Propagates filter changes from the instanced filter bar to the Sequencer tree.  */
-	void OnRequestUpdate_InstancedFilterBar() const;
 };
 } // namespace UE::Sequencer
diff --git a/Engine/Source/Editor/Sequencer/Private/MVVM/ViewModels/SequencerEditorViewModelImpl.cpp b/Engine/Source/Editor/Sequencer/Private/MVVM/ViewModels/SequencerEditorViewModelImpl.cpp
--- a/Engine/Source/Editor/Sequencer/Private/MVVM/ViewModels/SequencerEditorViewModelImpl.cpp
+++ b/Engine/Source/Editor/Sequencer/Private/MVVM/ViewModels/SequencerEditorViewModelImpl.cpp
@@ -39,14 +39,13 @@ void FSequencerEditorViewModelImpl::InitLinkedFiltering()
 	const TCHAR* ConfigId = LinkedConfigSubId;
 	LinkedFilterFactory = MakeShared<FLinkedFilterFactoryViewModel>(
 		*Sequencer,
-		MakeShared<FSequencerFilterBar>(*Sequencer, ConfigId)
+		MakeShared<FSequencerFilterBar>(*Sequencer, ConfigId),
+		Sequencer->GetNodeTree()->GetFilterEvaluator()
 		);
 	
 	SequenceFilterInstanceViewModel = LinkedFilterFactory->MakeFilteringModelImpl(
 		FMakeLinkedFilterViewModelArgs(InstancedConfigSubId)
 		);
-	FilterMediator = MakeShared<FLinkedFilteringMediator>(
-		Sequencer->GetNodeTree(), SequenceFilterInstanceViewModel.ToSharedRef()
-		);
+	FilterMediator = MakeShared<FLinkedFilteringMediator>(SequenceFilterInstanceViewModel.ToSharedRef());
 }
 }
diff --git a/Engine/Source/Editor/Sequencer/Private/SequencerNodeTree.cpp b/Engine/Source/Editor/Sequencer/Private/SequencerNodeTree.cpp
--- a/Engine/Source/Editor/Sequencer/Private/SequencerNodeTree.cpp
+++ b/Engine/Source/Editor/Sequencer/Private/SequencerNodeTree.cpp
@@ -41,6 +41,7 @@
 #include "SequencerLog.h"
 #include "SequencerCommonHelpers.h"
 #include "SSequencer.h"
+#include "Filters/ViewModels/FilterEvaluator.h"
 
 FSequencerNodeTree::~FSequencerNodeTree()
 {
@@ -48,7 +49,7 @@ FSequencerNodeTree::~FSequencerNodeTree()
 
 FSequencerNodeTree::FSequencerNodeTree(FSequencer& InSequencer)
 	: Sequencer(InSequencer)
-	, bFilterUpdateRequested(false)
+	, FilterEvaluator(MakeShared<UE::Sequencer::FFilterEvaluator>())
 {
 }
 
@@ -68,24 +69,7 @@ TSharedPtr<UE::Sequencer::FObjectBindingModel> FSequencerNodeTree::FindObjectBin
 
 bool FSequencerNodeTree::UpdateFiltersOnTrackValueChanged()
 {
-	// If filters are already scheduled for update, we can defer until the next update
-	if (bFilterUpdateRequested)
-	{
-		return false;
-	}
-
-	if (Sequencer.GetFilterBar()->ShouldUpdateOnTrackValueChanged())
-	{
-		// UpdateFilters will only run if bFilterUpdateRequested is true
-		bFilterUpdateRequested = true;
-		bool bFiltersUpdated = UpdateFilters();
-
-		// If the filter list was modified, set bFilterUpdateRequested to suppress excessive re-filters between tree update
-		bFilterUpdateRequested = bFiltersUpdated;
-		return bFiltersUpdated;
-	}
-
-	return false;
+	return FilterEvaluator->UpdateFiltersOnTrackValueChanged();
 }
 
 void FSequencerNodeTree::Update()
@@ -107,8 +91,8 @@ void FSequencerNodeTree::Update()
 
 	// Re-filter the tree after updating 
 	// @todo sequencer: Newly added sections may need to be visible even when there is a filter
-	bFilterUpdateRequested = true;
-	UpdateFilters();
+	constexpr bool bForce = true;
+	FilterEvaluator->UpdateFilters(bForce);
 
 	// Sort all nodes
 	const bool bIncludeRootNode = true;
@@ -308,6 +292,16 @@ bool FSequencerNodeTree::IsNodeFiltered(const TSharedPtr<UE::Sequencer::FViewMod
 	return OutlinerItem && !OutlinerItem->IsFilteredOut();
 }
 
+void FSequencerNodeTree::RequestFilterUpdate()
+{
+	FilterEvaluator->RequestFilterUpdate(Sequencer.GetFilterBar());
+}
+
+bool FSequencerNodeTree::NeedsFilterUpdate() const
+{
+	return FilterEvaluator->NeedsFilterUpdate();
+}
+
 TSharedPtr<UE::Sequencer::FSectionModel> FSequencerNodeTree::GetSectionModel(const UMovieSceneSection* Section) const
 {
 	using namespace UE::Sequencer;
@@ -320,31 +314,6 @@ TSharedPtr<UE::Sequencer::FSectionModel> FSequencerNodeTree::GetSectionModel(con
 	return nullptr;
 }
 
-bool FSequencerNodeTree::UpdateFilters()
-{
-	using namespace UE::Sequencer;
-
-	if (!bFilterUpdateRequested)
-	{
-		return false;
-	}
-
-	const TSharedPtr<FSequencerFilterBar> FilterBar = Sequencer.GetFilterBar();
-	if (!FilterBar.IsValid())
-	{
-		return false;
-	}
-
-	const FSequencerFilterData PreviousFilterData = FilterBar->GetFilterData();
-	const FSequencerFilterData& FilterData = FilterBar->FilterNodes();
-
-	bFilteringOnNodeGroups = Sequencer.GetFocusedMovieSceneSequence()->GetMovieScene()->GetNodeGroups().HasAnyActiveFilter();
-	bFilterUpdateRequested = false;
-
-	// Return whether the new list of FilteredNodes is different than the previous list
-	return PreviousFilterData != FilterData;
-}
-
 int32 FSequencerNodeTree::GetTotalDisplayNodeCount() const 
 { 
 	return Sequencer.GetFilterBar()->GetFilterData().GetTotalNodeCount();
diff --git a/Engine/Source/Editor/Sequencer/Private/SequencerNodeTree.h b/Engine/Source/Editor/Sequencer/Private/SequencerNodeTree.h
--- a/Engine/Source/Editor/Sequencer/Private/SequencerNodeTree.h
+++ b/Engine/Source/Editor/Sequencer/Private/SequencerNodeTree.h
@@ -16,6 +16,11 @@
 #include "MVVM/Extensions/IOutlinerExtension.h"
 #include "MovieSceneSequence.h"
 
+namespace UE::Sequencer
+{
+class FFilterEvaluator;
+}
+
 class FSequencer;
 class FSequencerFilterBar;
 class FSequencerTrackFilterCollection;
@@ -78,15 +83,11 @@ public:
 	 */
 	bool IsNodeFiltered(const TSharedPtr<FViewModel>& Node) const;
 
-	/**
-	 * Schedules an update of all filters
-	 */
-	void RequestFilterUpdate() { bFilterUpdateRequested = true; }
+	/** Schedules an update of the Sequencer outliner filter bar. */
+	void RequestFilterUpdate();
 
-	/**
-	 * @return Whether there is a filter update scheduled
-	 */
-	bool NeedsFilterUpdate() const { return bFilterUpdateRequested; }
+	/** @return Whether any filter bar needs to be filtered. */
+	bool NeedsFilterUpdate() const;
 
 
 	/** Called when the active MovieScene's node group colletion has been modifed */
@@ -160,6 +161,9 @@ public:
 	 * @return The node located at NodePath, or nullptr if not found
 	 */
 	UE::Sequencer::TViewModelPtr<UE::Sequencer::IOutlinerExtension> GetNodeAtPath(const FString& NodePath) const;
+	
+	/** @return  */
+	const TSharedRef<UE::Sequencer::FFilterEvaluator>& GetFilterEvaluator() const { return FilterEvaluator; }
 
 public:
 
@@ -168,16 +172,6 @@ public:
 	 */
 	TSharedPtr<FSectionModel> GetSectionModel(const UMovieSceneSection* Section) const;
 
-private:
-
-	/**
-	 * Update the list of filters nodes based on current filter settings, if an update is scheduled
-	 * This is called by Update();
-	 * 
-	 * @return Whether the list of filtered nodes changed
-	 */
-	bool UpdateFilters();
-
 public:
 
 	int32 GetTotalDisplayNodeCount() const;
@@ -190,9 +184,12 @@ private:
 
 	/** Sequencer interface */
 	FSequencer& Sequencer;
+	
+	/** Knows about all filter bars that are dirty and handles evaluating them. */
+	const TSharedRef<UE::Sequencer::FFilterEvaluator> FilterEvaluator;
+	
 	/** A multicast delegate which is called whenever the node tree has been updated. */
 	FOnUpdated OnUpdatedDelegate;
 
-	bool bFilterUpdateRequested;
 	bool bFilteringOnNodeGroups;
 };
