Sequencer: Anim Mixer modifier/decorator display initial version + root motion changes.

* Renamed root motion anim mixer track level to 'Root Motion Target' and section level to 'Root Motion Settings'
* Made the Root Motion Target one a 'modifier' that provides a section rather than just a section
* For this version, the modifiers are shown under a custom outliner item 'Modifiers'- this will likely change in future version
* A bunch of MVVM work to centralize the logic of laying out the modifiers, their channels, and sections and hook into it from the various things that can show modifiers of different types (e.g. section outliner, tracks, layers, mixer track, etc.)
* Some initial refactors to root motion offset cacheing (that will probably be redone in the next version)
* Added concept of decoration editors for the purpose of laying out sections that decorations provide
* Added ability for a channel via trait to say it doesn't support keying so the key buttons don't show up- currently we had added a FMovieSceneByteChannelDefaultOnly which didn't accept keys but the key buttons were still there
* Added storage for decorations to better handle recycling/refresh (was having a lot of trouble with flickering, adding/deleting not working, etc.)
* Modify section rendering somewhat to support drawing background underneath channels added by decorations
* Handle showing properties in the selection drawer when selecting outliner item for modifiers
* Added the ability for decorations to provide sections and have those sections picked up by the compiled data manager properly.

#Synced-CL 50659531

diff --git a/Engine/Plugins/Experimental/MovieSceneAnimMixer/Source/MovieSceneAnimMixer/Private/MovieSceneAnimMixerModule.cpp b/Engine/Plugins/Experimental/MovieSceneAnimMixer/Source/MovieSceneAnimMixer/Private/MovieSceneAnimMixerModule.cpp
--- a/Engine/Plugins/Experimental/MovieSceneAnimMixer/Source/MovieSceneAnimMixer/Private/MovieSceneAnimMixerModule.cpp
+++ b/Engine/Plugins/Experimental/MovieSceneAnimMixer/Source/MovieSceneAnimMixer/Private/MovieSceneAnimMixerModule.cpp
@@ -3,9 +3,14 @@
 #include "MovieSceneAnimMixerModule.h"
 
 #include "AnimNode_SequencerMixerTarget.h"
+#include "IMovieSceneModule.h"
 #include "MovieSceneAnimMixerSettings.h"
+#include "MovieSceneAnimationMixerLayer.h"
+#include "MovieSceneAnimationMixerTrack.h"
+#include "MovieSceneRootMotionTargetDecoration.h"
 #include "Modules/ModuleInterface.h"
 #include "Modules/ModuleManager.h"
+#include "Sections/MovieSceneSkeletalAnimationSection.h"
 
 #if WITH_UNREAL_DEVELOPER_TOOLS
 #include "ISettingsModule.h"
@@ -39,12 +44,21 @@ void FMovieSceneAnimMixerModule::StartupModule()
 #endif
 
 	CVarDefaultAnimNextInjectionSite->Set(*GetMutableDefault<UMovieSceneAnimMixerSettings>()->DefaultInjectionSite.ToString(), ECVF_SetByProjectSetting);
-	
+
+	// Register Root Motion Target decoration as compatible with the Animation Mixer Track
+	IMovieSceneModule::Get().RegisterCompatibleDecoration(UMovieSceneAnimationMixerTrack::StaticClass(), UMovieSceneRootMotionTargetDecoration::StaticClass());
+
+	// Register Root Motion Settings decoration as compatible with Skeletal Animation Sections
+	IMovieSceneModule::Get().RegisterCompatibleDecoration(UMovieSceneSkeletalAnimationSection::StaticClass(), UMovieSceneRootMotionSettingsDecoration::StaticClass());
+
 	IModularFeatures::Get().RegisterModularFeature(IAnimGraphRuntime_SequencerMixerTargetConnector::GetModularFeatureName(), this);
 }
 
 void FMovieSceneAnimMixerModule::ShutdownModule()
 {
+	IMovieSceneModule::Get().UnregisterCompatibleDecoration(UMovieSceneAnimationMixerTrack::StaticClass(), UMovieSceneRootMotionTargetDecoration::StaticClass());
+	IMovieSceneModule::Get().UnregisterCompatibleDecoration(UMovieSceneSkeletalAnimationSection::StaticClass(), UMovieSceneRootMotionSettingsDecoration::StaticClass());
+
 	IModularFeatures::Get().UnregisterModularFeature(IAnimGraphRuntime_SequencerMixerTargetConnector::GetModularFeatureName(), this);
 }
 
diff --git a/Engine/Plugins/Experimental/MovieSceneAnimMixer/Source/MovieSceneAnimMixer/Private/MovieSceneAnimationMixerTrack.cpp b/Engine/Plugins/Experimental/MovieSceneAnimMixer/Source/MovieSceneAnimMixer/Private/MovieSceneAnimationMixerTrack.cpp
--- a/Engine/Plugins/Experimental/MovieSceneAnimMixer/Source/MovieSceneAnimMixer/Private/MovieSceneAnimationMixerTrack.cpp
+++ b/Engine/Plugins/Experimental/MovieSceneAnimMixer/Source/MovieSceneAnimMixer/Private/MovieSceneAnimationMixerTrack.cpp
@@ -31,7 +31,7 @@ bool FMovieSceneByteChannelDefaultOnly::SerializeFromMismatchedTag(const FProper
 	return false;
 }
 
-UMovieSceneAnimationBaseTransformDecoration::UMovieSceneAnimationBaseTransformDecoration(const FObjectInitializer& ObjInit)
+UMovieSceneRootMotionSettingsDecoration::UMovieSceneRootMotionSettingsDecoration(const FObjectInitializer& ObjInit)
 	: Super(ObjInit)
 {
 	Location[0].SetDefault(0.0);
@@ -48,17 +48,17 @@ UMovieSceneAnimationBaseTransformDecoration::UMovieSceneAnimationBaseTransformDe
 	TransformMode.SetDefault((uint8)EMovieSceneRootMotionTransformMode::Offset);
 }
 
-EMovieSceneRootMotionTransformMode UMovieSceneAnimationBaseTransformDecoration::GetRootTransformMode() const
+EMovieSceneRootMotionTransformMode UMovieSceneRootMotionSettingsDecoration::GetRootTransformMode() const
 {
 	return (EMovieSceneRootMotionTransformMode)TransformMode.GetDefault().Get((uint8)EMovieSceneRootMotionTransformMode::Offset);
 }
 
-EMovieSceneChannelProxyType UMovieSceneAnimationBaseTransformDecoration::PopulateChannelProxy(FMovieSceneChannelProxyData& OutProxyData)
+EMovieSceneChannelProxyType UMovieSceneRootMotionSettingsDecoration::PopulateChannelProxy(FMovieSceneChannelProxyData& OutProxyData)
 {
 	EMovieSceneRootMotionTransformMode CurrentTransformMode = GetRootTransformMode();
 
 #if WITH_EDITOR
-	FText Group = LOCTEXT("RootMotion", "Root Motion");
+	FText Group = LOCTEXT("RootMotionSettings", "Root Motion Settings");
 	{
 		FMovieSceneChannelMetaData ChannelMetaData;
 		ChannelMetaData.SetIdentifiers("Space", LOCTEXT("Space", "Space"), Group);
@@ -155,7 +155,7 @@ EMovieSceneChannelProxyType UMovieSceneAnimationBaseTransformDecoration::Populat
 	return EMovieSceneChannelProxyType::Dynamic;
 }
 
-void UMovieSceneAnimationBaseTransformDecoration::ExtendEntityImpl(UMovieSceneEntitySystemLinker* EntityLinker, const FEntityImportParams& Params, FImportedEntity* OutImportedEntity)
+void UMovieSceneRootMotionSettingsDecoration::ExtendEntityImpl(UMovieSceneEntitySystemLinker* EntityLinker, const FEntityImportParams& Params, FImportedEntity* OutImportedEntity)
 {
 	using namespace UE::MovieScene;
 	FBuiltInComponentTypes* BuiltInComponents = FBuiltInComponentTypes::Get();
@@ -262,10 +262,18 @@ void UMovieSceneAnimationMixerTrack::OnSectionAddedImpl(UMovieSceneSection* Sect
 		NewDecoration->MixerTrack = this;
 		Section->AddDecoration(NewDecoration);
 
+		// Auto-add Root Motion Settings decoration only for skeletal animation sections
+		// whose animations have root motion enabled
 		if (UMovieSceneSkeletalAnimationSection* AnimSection = Cast<UMovieSceneSkeletalAnimationSection>(Section))
 		{
-			UMovieSceneAnimationBaseTransformDecoration* BaseTransform = NewObject<UMovieSceneAnimationBaseTransformDecoration>(AnimSection, UMovieSceneAnimationBaseTransformDecoration::StaticClass(), NAME_None, RF_Transactional);
-			AnimSection->AddDecoration(BaseTransform);
+			if (UAnimSequenceBase* Animation = AnimSection->Params.Animation)
+			{
+				if (Animation->HasRootMotion())
+				{
+					UMovieSceneRootMotionSettingsDecoration* RootMotionSettings = NewObject<UMovieSceneRootMotionSettingsDecoration>(AnimSection, UMovieSceneRootMotionSettingsDecoration::StaticClass(), NAME_None, RF_Transactional);
+					AnimSection->AddDecoration(RootMotionSettings);
+				}
+			}
 		}
 
 #if WITH_EDITOR
@@ -287,6 +295,9 @@ void UMovieSceneAnimationMixerTrack::OnSectionAddedImpl(UMovieSceneSection* Sect
 		{
 			UpdateTransitionsForRow(RowIndex);
 		}
+
+		// Mark root motions as dirty so accumulated offset cache will be rebuilt
+		SetRootMotionsDirty();
 	}
 }
 
@@ -307,6 +318,9 @@ void UMovieSceneAnimationMixerTrack::OnSectionRemovedImpl(UMovieSceneSection* Se
 		{
 			RemoveSection(*Transition);
 		}
+
+		// Mark root motions as dirty so accumulated offset cache will be rebuilt
+		SetRootMotionsDirty();
 	}
 }
 
@@ -723,6 +737,9 @@ EMovieSceneSectionMovedResult UMovieSceneAnimationMixerTrack::OnSectionMoved(UMo
 
 	const bool bSectionsChanged = UpdateTransitionsForRow(Section.GetRowIndex());
 
+	// Mark root motions as dirty so accumulated offset cache will be rebuilt
+	SetRootMotionsDirty();
+
 	return bSectionsChanged ? EMovieSceneSectionMovedResult::SectionsChanged : EMovieSceneSectionMovedResult::None;
 }
 #endif // WITH_EDITOR
diff --git a/Engine/Plugins/Experimental/MovieSceneAnimMixer/Source/MovieSceneAnimMixer/Private/MovieSceneRootMotionSection.cpp b/Engine/Plugins/Experimental/MovieSceneAnimMixer/Source/MovieSceneAnimMixer/Private/MovieSceneRootMotionSection.cpp
--- a/Engine/Plugins/Experimental/MovieSceneAnimMixer/Source/MovieSceneAnimMixer/Private/MovieSceneRootMotionSection.cpp
+++ b/Engine/Plugins/Experimental/MovieSceneAnimMixer/Source/MovieSceneAnimMixer/Private/MovieSceneRootMotionSection.cpp
@@ -6,7 +6,9 @@
 #include "Channels/MovieSceneChannelProxy.h"
 #include "EntitySystem/BuiltInComponentTypes.h"
 #include "Evaluation/MovieSceneEvaluationField.h"
+#include "MovieSceneTrack.h"
 #include "Systems/MovieSceneSkeletalAnimationSystem.h"
+#include "Tracks/MovieSceneCommonAnimationTrack.h"
 
 UMovieSceneRootMotionSection::UMovieSceneRootMotionSection(const FObjectInitializer& ObjInit)
 	: Super(ObjInit)
@@ -24,7 +26,7 @@ EMovieSceneChannelProxyType UMovieSceneRootMotionSection::CacheChannelProxy()
 #if WITH_EDITOR
 
 	static FMovieSceneChannelMetaData MetaData("RootDestination", NSLOCTEXT("AnimMixer", "RootDestinationName", "Root Destination"));
-	MetaData.bCanCollapseToTrack = true;
+	MetaData.bCanCollapseToTrack = false;
 
 	Channels.Add(RootDestinationChannel, MetaData, TMovieSceneExternalValue<uint8>());
 
@@ -76,3 +78,98 @@ bool UMovieSceneRootMotionSection::PopulateEvaluationFieldImpl(const TRange<FFra
 
 	return true;
 }
+
+FTransform UMovieSceneRootMotionSection::GetAccumulatedOffsetAt(FFrameTime Time) const
+{
+	if (GapBehavior != EMovieSceneRootMotionGapBehavior::PersistPreviousTransform)
+	{
+		return FTransform::Identity;
+	}
+
+	// Check if the cache needs rebuilding (track's sections have changed)
+	if (UMovieSceneCommonAnimationTrack* Track = GetTypedOuter<UMovieSceneCommonAnimationTrack>())
+	{
+		if (Track->RootMotionParams.bRootMotionsDirty)
+		{
+			RebuildAccumulatedOffsetCache();
+			Track->RootMotionParams.bRootMotionsDirty = false;
+		}
+	}
+
+	// Find the most recent accumulated offset before the given time
+	FTransform Result = FTransform::Identity;
+	FFrameNumber TimeFrame = Time.FloorToFrame();
+
+	for (const FRootMotionAccumulatedOffset& Offset : AccumulatedOffsetCache)
+	{
+		if (Offset.StartFrame <= TimeFrame)
+		{
+			Result = Offset.AccumulatedOffset;
+		}
+		else
+		{
+			break;
+		}
+	}
+
+	return Result;
+}
+
+void UMovieSceneRootMotionSection::RebuildAccumulatedOffsetCache() const
+{
+	AccumulatedOffsetCache.Reset();
+
+	UMovieSceneTrack* OwningTrack = GetTypedOuter<UMovieSceneTrack>();
+	if (!OwningTrack)
+	{
+		return;
+	}
+
+	// Collect sections with root motion settings that have KeepState enabled
+	TArray<TPair<UMovieSceneSection*, UMovieSceneRootMotionSettingsDecoration*>> SectionsWithKeepState;
+
+	for (UMovieSceneSection* Section : OwningTrack->GetAllSections())
+	{
+		if (!Section || !Section->HasEndFrame())
+		{
+			continue;
+		}
+
+		UMovieSceneRootMotionSettingsDecoration* Settings = Section->FindDecoration<UMovieSceneRootMotionSettingsDecoration>();
+		if (Settings && Settings->bKeepStateAfterSectionEnds)
+		{
+			SectionsWithKeepState.Add(TPair<UMovieSceneSection*, UMovieSceneRootMotionSettingsDecoration*>(Section, Settings));
+		}
+	}
+
+	// Sort by section end frame (when the accumulated offset takes effect)
+	SectionsWithKeepState.Sort([](const auto& A, const auto& B)
+	{
+		return A.Key->GetExclusiveEndFrame() < B.Key->GetExclusiveEndFrame();
+	});
+
+	// Build the accumulated offset cache
+	FTransform RunningAccumulatedOffset = FTransform::Identity;
+
+	for (const auto& Pair : SectionsWithKeepState)
+	{
+		UMovieSceneSection* Section = Pair.Key;
+		UMovieSceneRootMotionSettingsDecoration* Settings = Pair.Value;
+
+		FVector LocationOffset = Settings->MatchedLocationOffset;
+		FRotator RotationOffset = Settings->MatchedRotationOffset;
+
+		FTransform SectionOffset(RotationOffset.Quaternion(), LocationOffset);
+
+		// Accumulate the offset
+		RunningAccumulatedOffset = SectionOffset * RunningAccumulatedOffset;
+
+		// Add to cache - this offset takes effect when the section ends
+		FRootMotionAccumulatedOffset CacheEntry;
+		CacheEntry.StartFrame = Section->GetExclusiveEndFrame();
+		CacheEntry.AccumulatedOffset = RunningAccumulatedOffset;
+		CacheEntry.Weight = 1.0f;
+
+		AccumulatedOffsetCache.Add(CacheEntry);
+	}
+}
diff --git a/Engine/Plugins/Experimental/MovieSceneAnimMixer/Source/MovieSceneAnimMixer/Private/MovieSceneRootMotionTargetDecoration.cpp b/Engine/Plugins/Experimental/MovieSceneAnimMixer/Source/MovieSceneAnimMixer/Private/MovieSceneRootMotionTargetDecoration.cpp
new file mode 100644
--- /dev/null
+++ b/Engine/Plugins/Experimental/MovieSceneAnimMixer/Source/MovieSceneAnimMixer/Private/MovieSceneRootMotionTargetDecoration.cpp
@@ -0,0 +1,33 @@
+// Copyright Epic Games, Inc. All Rights Reserved.
+
+#include "MovieSceneRootMotionTargetDecoration.h"
+
+#include "MovieSceneAnimationMixerTrack.h"
+#include "MovieSceneSection.h"
+#include "MovieSceneTrack.h"
+
+#include UE_INLINE_GENERATED_CPP_BY_NAME(MovieSceneRootMotionTargetDecoration)
+
+UMovieSceneRootMotionTargetDecoration::UMovieSceneRootMotionTargetDecoration(const FObjectInitializer& Init)
+	: Super(Init)
+{
+}
+
+void UMovieSceneRootMotionTargetDecoration::OnDecorationAdded(UMovieSceneTrack* Track)
+{
+	// Create the internal root motion section if it doesn't exist
+	if (!RootMotionSection)
+	{
+		RootMotionSection = NewObject<UMovieSceneRootMotionSection>(this, NAME_None, RF_Transactional);
+	}
+}
+
+void UMovieSceneRootMotionTargetDecoration::OnDecorationRemoved()
+{
+}
+
+TArrayView<TObjectPtr<UMovieSceneSection>> UMovieSceneRootMotionTargetDecoration::GetSections()
+{
+	return MakeArrayView(&RootMotionSection, RootMotionSection ? 1 : 0);
+}
+
diff --git a/Engine/Plugins/Experimental/MovieSceneAnimMixer/Source/MovieSceneAnimMixer/Private/Systems/MovieSceneAnimMixerSystem.cpp b/Engine/Plugins/Experimental/MovieSceneAnimMixer/Source/MovieSceneAnimMixer/Private/Systems/MovieSceneAnimMixerSystem.cpp
--- a/Engine/Plugins/Experimental/MovieSceneAnimMixer/Source/MovieSceneAnimMixer/Private/Systems/MovieSceneAnimMixerSystem.cpp
+++ b/Engine/Plugins/Experimental/MovieSceneAnimMixer/Source/MovieSceneAnimMixer/Private/Systems/MovieSceneAnimMixerSystem.cpp
@@ -331,26 +331,33 @@ namespace UE::MovieScene
 						{
 							ThisRootSpace = MixerEntry->RootMotionSettings->RootMotionSpace;
 
-							// If we have a root location or rotation offset/override, set that up
+							// Handle root transform offset/override
 							const bool bHasLocation = !MixerEntry->RootMotionSettings->RootLocation.IsZero();
 							const bool bHasRotation = !MixerEntry->RootMotionSettings->RootRotation.IsIdentity();
-							if (bHasLocation || bHasRotation)
-							{
-								if (MixerEntry->RootMotionSettings->TransformMode == EMovieSceneRootMotionTransformMode::Offset)
-								{
-									Conversion |= ESpaceConversions::RootTransformOffset;
-									RootOffsetOrigin = MixerEntry->RootMotionSettings->RootOriginLocation;
-								}
-								else
-								{
-									Conversion |= ESpaceConversions::RootTransformOverride;
-								}
 
+							if (MixerEntry->RootMotionSettings->TransformMode == EMovieSceneRootMotionTransformMode::Override)
+							{
+								// Override mode: Always apply, even with identity values (to null out root motion)
+								Conversion |= ESpaceConversions::RootTransformOverride;
 								RootTransform = FTransform(
 									MixerEntry->RootMotionSettings->RootRotation,
 									MixerEntry->RootMotionSettings->RootLocation
 									);
 							}
+							else if (MixerEntry->RootMotionSettings->TransformMode == EMovieSceneRootMotionTransformMode::Offset)
+							{
+								// Offset mode: Only apply if we have non-zero values to offset by
+								if (bHasLocation || bHasRotation)
+								{
+									Conversion |= ESpaceConversions::RootTransformOffset;
+									RootOffsetOrigin = MixerEntry->RootMotionSettings->RootOriginLocation;
+									RootTransform = FTransform(
+										MixerEntry->RootMotionSettings->RootRotation,
+										MixerEntry->RootMotionSettings->RootLocation
+										);
+								}
+							}
+							// Asset mode: No override/offset applied, use the animation's root motion as-is
 						}
 
 						if (ThisRootSpace == EMovieSceneRootMotionSpace::AnimationSpace)
diff --git a/Engine/Plugins/Experimental/MovieSceneAnimMixer/Source/MovieSceneAnimMixer/Public/MovieSceneAnimationMixerTrack.h b/Engine/Plugins/Experimental/MovieSceneAnimMixer/Source/MovieSceneAnimMixer/Public/MovieSceneAnimationMixerTrack.h
--- a/Engine/Plugins/Experimental/MovieSceneAnimMixer/Source/MovieSceneAnimMixer/Public/MovieSceneAnimationMixerTrack.h
+++ b/Engine/Plugins/Experimental/MovieSceneAnimMixer/Source/MovieSceneAnimMixer/Public/MovieSceneAnimationMixerTrack.h
@@ -38,6 +38,9 @@ struct TStructOpsTypeTraits<FMovieSceneByteChannelDefaultOnly> : public TStructO
 template<>
 struct TMovieSceneChannelTraits<FMovieSceneByteChannelDefaultOnly> : TMovieSceneChannelTraitsBase<FMovieSceneByteChannelDefaultOnly>
 {
+	/** Default-only channels cannot be keyed - they only support default values */
+	enum { SupportsKeying = false };
+
 #if WITH_EDITOR
 
 	/** Byte channels can have external values (ie, they can get their values from external objects for UI purposes) */
@@ -73,8 +76,13 @@ private:
 	FColor MixerTintOverride = FColor();
 };
 
-UCLASS(MinimalAPI)
-class UMovieSceneAnimationBaseTransformDecoration
+/**
+ * Section-level decoration that provides root motion offset settings.
+ * Adds transform offset channels to an existing animation section.
+ * Channels appear under "Modifiers" in the Sequencer outliner.
+ */
+UCLASS(MinimalAPI, collapseCategories, meta=(DisplayName="Root Motion Settings"))
+class UMovieSceneRootMotionSettingsDecoration
 	: public UMovieSceneSignedObject
 	, public IMovieSceneChannelDecoration
 	, public IMovieSceneEntityDecorator
@@ -83,13 +91,15 @@ public:
 
 	GENERATED_BODY()
 
-	UMovieSceneAnimationBaseTransformDecoration(const FObjectInitializer& ObjInit);
+	UMovieSceneRootMotionSettingsDecoration(const FObjectInitializer& ObjInit);
 
 	MOVIESCENEANIMMIXER_API EMovieSceneRootMotionTransformMode GetRootTransformMode() const;
 
 	virtual EMovieSceneChannelProxyType PopulateChannelProxy(FMovieSceneChannelProxyData& OutProxyData) override;
 	virtual void ExtendEntityImpl(UMovieSceneEntitySystemLinker* EntityLinker, const FEntityImportParams& Params, FImportedEntity* OutImportedEntity) override;
 
+	// Transform Offset Channels
+
 	UPROPERTY()
 	FMovieSceneDoubleChannel Location[3];
 
@@ -104,6 +114,50 @@ public:
 
 	UPROPERTY()
 	FMovieSceneByteChannelDefaultOnly TransformMode;
+
+	// Bone Matching Configuration
+
+	/** Name of the bone to match transforms to */
+	UPROPERTY(EditAnywhere, Category="Bone Matching")
+	FName MatchBoneName;
+
+	/** Whether to match with the previous section's end transform */
+	UPROPERTY(EditAnywhere, Category="Bone Matching")
+	bool bMatchWithPrevious = true;
+
+	/** Whether to match translation */
+	UPROPERTY(EditAnywhere, Category="Bone Matching")
+	bool bMatchTranslation = true;
+
+	/** Whether to include Z height when matching translation */
+	UPROPERTY(EditAnywhere, Category="Bone Matching", meta=(EditCondition="bMatchTranslation"))
+	bool bMatchIncludeZHeight = false;
+
+	/** Whether to match rotation yaw */
+	UPROPERTY(EditAnywhere, Category="Bone Matching")
+	bool bMatchRotationYaw = true;
+
+	/** Whether to match rotation pitch */
+	UPROPERTY(EditAnywhere, Category="Bone Matching")
+	bool bMatchRotationPitch = false;
+
+	/** Whether to match rotation roll */
+	UPROPERTY(EditAnywhere, Category="Bone Matching")
+	bool bMatchRotationRoll = false;
+
+	/** Computed location offset from bone matching (read-only) */
+	UPROPERTY(VisibleAnywhere, Category="Bone Matching")
+	FVector MatchedLocationOffset = FVector::ZeroVector;
+
+	/** Computed rotation offset from bone matching (read-only) */
+	UPROPERTY(VisibleAnywhere, Category="Bone Matching")
+	FRotator MatchedRotationOffset = FRotator::ZeroRotator;
+
+	// Keep State
+
+	/** If true, the transform offset persists after this section ends (for use with gap behavior) */
+	UPROPERTY(EditAnywhere, Category="Root Motion")
+	bool bKeepStateAfterSectionEnds = false;
 };
 
 
diff --git a/Engine/Plugins/Experimental/MovieSceneAnimMixer/Source/MovieSceneAnimMixer/Public/MovieSceneRootMotionSection.h b/Engine/Plugins/Experimental/MovieSceneAnimMixer/Source/MovieSceneAnimMixer/Public/MovieSceneRootMotionSection.h
--- a/Engine/Plugins/Experimental/MovieSceneAnimMixer/Source/MovieSceneAnimMixer/Public/MovieSceneRootMotionSection.h
+++ b/Engine/Plugins/Experimental/MovieSceneAnimMixer/Source/MovieSceneAnimMixer/Public/MovieSceneRootMotionSection.h
@@ -24,10 +24,40 @@ enum class EMovieSceneRootMotionDestination : uint8
 	Attribute,
 };
 
+/** Behavior during gaps where no root motion sections are active */
+UENUM()
+enum class EMovieSceneRootMotionGapBehavior : uint8
+{
+	/** Reset to origin during gaps - current default behavior */
+	ResetToOrigin,
+	/** Persist the transform from the previous section that had KeepState enabled */
+	PersistPreviousTransform
+};
+
+/** Cached accumulated offset from sections with KeepState enabled */
+USTRUCT()
+struct FRootMotionAccumulatedOffset
+{
+	GENERATED_BODY()
+
+	/** Frame at which this accumulated offset takes effect */
+	UPROPERTY()
+	FFrameNumber StartFrame;
+
+	/** The accumulated transform offset at this point */
+	UPROPERTY()
+	FTransform AccumulatedOffset = FTransform::Identity;
+
+	/** Weight of this offset for blending */
+	UPROPERTY()
+	float Weight = 1.0f;
+};
+
 /**
- * 
+ * Section that controls where root motion is applied for animation mixer tracks.
+ * Created internally by UMovieSceneRootMotionTargetDecoration - not added directly.
  */
-UCLASS(MinimalAPI, DisplayName="Root Motion")
+UCLASS(MinimalAPI, meta=(DisplayName="Root Motion Target", Hidden))
 class UMovieSceneRootMotionSection
 	: public UMovieSceneSection
 	, public IMovieSceneEntityProvider
@@ -51,6 +81,22 @@ public:
 	virtual void ImportEntityImpl(UMovieSceneEntitySystemLinker* EntityLinker, const FEntityImportParams& Params, FImportedEntity* OutImportedEntity) override;
 	virtual bool PopulateEvaluationFieldImpl(const TRange<FFrameNumber>& EffectiveRange, const FMovieSceneEvaluationFieldEntityMetaData& InMetaData, FMovieSceneEntityComponentFieldBuilder* OutFieldBuilder) override;
 
+	/** Get the accumulated offset at a given time*/
+	MOVIESCENEANIMMIXER_API FTransform GetAccumulatedOffsetAt(FFrameTime Time) const;
+
+	/** Rebuild the accumulated offset cache when sections change. */
+	MOVIESCENEANIMMIXER_API void RebuildAccumulatedOffsetCache() const;
+
+public:
+
+	/** How to handle gaps where no sections with root motion are active */
+	UPROPERTY(EditAnywhere, Category="Root Motion")
+	EMovieSceneRootMotionGapBehavior GapBehavior = EMovieSceneRootMotionGapBehavior::ResetToOrigin;
+
+	/** Cached accumulated offsets from sections with KeepState enabled. Mutable for lazy rebuilding. */
+	UPROPERTY(Transient)
+	mutable TArray<FRootMotionAccumulatedOffset> AccumulatedOffsetCache;
+
 private:
 
 	UPROPERTY()
diff --git a/Engine/Plugins/Experimental/MovieSceneAnimMixer/Source/MovieSceneAnimMixer/Public/MovieSceneRootMotionTargetDecoration.h b/Engine/Plugins/Experimental/MovieSceneAnimMixer/Source/MovieSceneAnimMixer/Public/MovieSceneRootMotionTargetDecoration.h
new file mode 100644
--- /dev/null
+++ b/Engine/Plugins/Experimental/MovieSceneAnimMixer/Source/MovieSceneAnimMixer/Public/MovieSceneRootMotionTargetDecoration.h
@@ -0,0 +1,46 @@
+// Copyright Epic Games, Inc. All Rights Reserved.
+
+#pragma once
+
+#include "CoreTypes.h"
+#include "Decorations/IMovieSceneSectionProviderDecoration.h"
+#include "Decorations/IMovieSceneTrackDecoration.h"
+#include "EntitySystem/IMovieSceneEntityDecorator.h"
+#include "MovieSceneSignedObject.h"
+#include "MovieSceneRootMotionSection.h"
+#include "MovieSceneRootMotionTargetDecoration.generated.h"
+
+class UMovieSceneTrack;
+
+/**
+ * Track-level decoration that controls where root motion is applied.
+ * Wraps a UMovieSceneRootMotionSection
+ */
+UCLASS(MinimalAPI, collapseCategories, meta=(DisplayName="Root Motion Target"))
+class UMovieSceneRootMotionTargetDecoration
+	: public UMovieSceneSignedObject
+	, public IMovieSceneTrackDecoration
+	, public IMovieSceneSectionProviderDecoration
+{
+	GENERATED_BODY()
+
+public:
+
+	UMovieSceneRootMotionTargetDecoration(const FObjectInitializer& Init);
+
+	// IMovieSceneTrackDecoration interface
+	virtual void OnDecorationAdded(UMovieSceneTrack* Track) override;
+	virtual void OnDecorationRemoved() override;
+
+	// IMovieSceneSectionProviderDecoration interface
+	virtual TArrayView<TObjectPtr<UMovieSceneSection>> GetSections() override;
+
+	/** Get the internal root motion section */
+	UMovieSceneRootMotionSection* GetRootMotionSection() const { return Cast<UMovieSceneRootMotionSection>(RootMotionSection); }
+
+protected:
+
+	/** Internal root motion section that provides the destination channel and entity logic */
+	UPROPERTY()
+	TObjectPtr<UMovieSceneSection> RootMotionSection;
+};
diff --git a/Engine/Plugins/Experimental/MovieSceneAnimMixer/Source/MovieSceneAnimMixerEditor/Private/MVVM/AnimationMixerLayerModel.cpp b/Engine/Plugins/Experimental/MovieSceneAnimMixer/Source/MovieSceneAnimMixerEditor/Private/MVVM/AnimationMixerLayerModel.cpp
--- a/Engine/Plugins/Experimental/MovieSceneAnimMixer/Source/MovieSceneAnimMixerEditor/Private/MVVM/AnimationMixerLayerModel.cpp
+++ b/Engine/Plugins/Experimental/MovieSceneAnimMixer/Source/MovieSceneAnimMixerEditor/Private/MVVM/AnimationMixerLayerModel.cpp
@@ -31,6 +31,7 @@
 #include "ISequencer.h"
 #include "ISequencerTrackEditor.h"
 #include "SequencerUtilities.h"
+#include "MVVM/DecorationModelStorageExtension.h"
 #include "SequencerCommonHelpers.h"
 #include "Decorations/MovieSceneTrackRowDecoration.h"
 
@@ -301,7 +302,7 @@ void FAnimationMixerLayerModel::RefreshLayout(bool bChildrenNeedLayout)
 
 			// Query saved expansion state directly - we can't use IsExpanded() here because
 			// it checks for children, but we're in the process of building them
-			bool bCreateOutliners = GetSavedExpansionState() && LayerSections.Num() > 1;
+			bool bCreateOutliners = GetSavedExpansionState() && (LayerSections.Num() > 1 || SequencerHelpers::HasOutlinerVisibleDecorations(Layer));
 
 			// Only recycle top level channels if children need layout
 			if (bChildrenNeedLayout)
@@ -327,6 +328,19 @@ void FAnimationMixerLayerModel::RefreshLayout(bool bChildrenNeedLayout)
 				PreviousLayoutNumSections);
 
 			PreviousLayoutNumSections = LayerSections.Num();
+
+			// Handle layer and section decorations
+			if (FDecorationModelStorageExtension* DecorationModelStorage = SequenceModel->CastDynamic<FDecorationModelStorageExtension>())
+			{
+				DecorationModelStorage->SyncDecorationModel(Layer, ParentTrack, OutlinerChildren);
+
+				// Section decorations for single-section case (when section outliners aren't created).
+				// Multi-section case: SyncDecorationModel's section outliner iteration handles stale cleanup.
+				if (LayerSections.Num() == 1)
+				{
+					DecorationModelStorage->SyncDecorationModel(LayerSections[0], ParentTrack, OutlinerChildren);
+				}
+			}
 		}
 	}
 }
diff --git a/Engine/Plugins/Experimental/MovieSceneAnimMixer/Source/MovieSceneAnimMixerEditor/Private/MVVM/AnimationMixerLayoutHelpers.cpp b/Engine/Plugins/Experimental/MovieSceneAnimMixer/Source/MovieSceneAnimMixerEditor/Private/MVVM/AnimationMixerLayoutHelpers.cpp
--- a/Engine/Plugins/Experimental/MovieSceneAnimMixer/Source/MovieSceneAnimMixerEditor/Private/MVVM/AnimationMixerLayoutHelpers.cpp
+++ b/Engine/Plugins/Experimental/MovieSceneAnimMixer/Source/MovieSceneAnimMixerEditor/Private/MVVM/AnimationMixerLayoutHelpers.cpp
@@ -64,6 +64,7 @@ TArray<TSharedPtr<FSectionModel>> BuildSectionModels(
 TArray<TSharedPtr<FViewModel>> CreateSectionOutliners(
 	const TArray<TSharedPtr<FSectionModel>>& SectionModels,
 	FViewModelChildren& OutlinerChildren,
+	const FViewModelChildren& RecycledChildren,
 	TSharedPtr<FViewModel> Owner)
 {
 	TArray<TSharedPtr<FViewModel>> SectionRoots;
@@ -84,9 +85,11 @@ TArray<TSharedPtr<FViewModel>> CreateSectionOutliners(
 			continue;
 		}
 
-		// Find existing section outliner model or create new one
+		// Find existing section outliner model in recycled children first, then current children
 		TSharedPtr<FSectionOutlinerModel> SectionOutliner;
-		for (TSharedPtr<FSectionOutlinerModel> ExistingOutliner : OutlinerChildren.IterateSubList<FSectionOutlinerModel>())
+
+		// Search recycled children for existing section outliner
+		for (TSharedPtr<FSectionOutlinerModel> ExistingOutliner : RecycledChildren.IterateSubList<FSectionOutlinerModel>())
 		{
 			if (ExistingOutliner && ExistingOutliner->GetSection() == Section)
 			{
@@ -95,6 +98,19 @@ TArray<TSharedPtr<FViewModel>> CreateSectionOutliners(
 			}
 		}
 
+		// Also check current children (in case of partial rebuild)
+		if (!SectionOutliner)
+		{
+			for (TSharedPtr<FSectionOutlinerModel> ExistingOutliner : OutlinerChildren.IterateSubList<FSectionOutlinerModel>())
+			{
+				if (ExistingOutliner && ExistingOutliner->GetSection() == Section)
+				{
+					SectionOutliner = ExistingOutliner;
+					break;
+				}
+			}
+		}
+
 		if (!SectionOutliner)
 		{
 			SectionOutliner = MakeShared<FSectionOutlinerModel>(Section, SectionModel);
@@ -210,7 +226,7 @@ void LayoutSections(
 		if (bCreateSectionOutliners && NumSections > 1)
 		{
 			// Multiple sections: Create section outliner models
-			SectionRoots = CreateSectionOutliners(SectionModels, OutlinerChildren, Owner);
+			SectionRoots = CreateSectionOutliners(SectionModels, OutlinerChildren, RecycledModels.GetChildren(), Owner);
 		}
 		else
 		{
@@ -220,6 +236,17 @@ void LayoutSections(
 
 		// Run layout builder on each section
 		RunLayoutBuilders(SectionModels, SectionRoots);
+
+		// Refresh decorations on section outliners after layout builders complete.
+		// Layout builders recycle the outliner child list, so decorations must be
+		// added after that process finishes.
+		for (TSharedPtr<FSectionOutlinerModel> SectionOutliner : OutlinerChildren.IterateSubList<FSectionOutlinerModel>())
+		{
+			if (SectionOutliner)
+			{
+				SectionOutliner->RefreshDecorations();
+			}
+		}
 	}
 }
 
diff --git a/Engine/Plugins/Experimental/MovieSceneAnimMixer/Source/MovieSceneAnimMixerEditor/Private/MVVM/AnimationMixerTrackModel.cpp b/Engine/Plugins/Experimental/MovieSceneAnimMixer/Source/MovieSceneAnimMixerEditor/Private/MVVM/AnimationMixerTrackModel.cpp
--- a/Engine/Plugins/Experimental/MovieSceneAnimMixer/Source/MovieSceneAnimMixerEditor/Private/MVVM/AnimationMixerTrackModel.cpp
+++ b/Engine/Plugins/Experimental/MovieSceneAnimMixer/Source/MovieSceneAnimMixerEditor/Private/MVVM/AnimationMixerTrackModel.cpp
@@ -6,6 +6,7 @@
 #include "MovieSceneAnimationMixerTrack.h"
 #include "MovieSceneAnimationMixerLayer.h"
 #include "MovieSceneAnimTransitionSectionBase.h"
+#include "MVVM/DecorationModelStorageExtension.h"
 #include "MVVM/ViewModels/TrackModel.h"
 #include "MVVM/ViewModels/TrackRowModel.h"
 #include "MVVM/ViewModels/SequenceModel.h"
@@ -92,6 +93,7 @@ void FAnimationMixerTrackModel::ForceUpdate()
 			// This follows the same pattern as FTrackModel for single-row tracks
 
 			FSectionModelStorageExtension* SectionModelStorage = SequenceModel->CastDynamic<FSectionModelStorageExtension>();
+			FDecorationModelStorageExtension* DecorationModelStorage = SequenceModel->CastDynamic<FDecorationModelStorageExtension>();
 			TSharedPtr<ISequencerTrackEditor> TrackEditorPtr = GetTrackEditor();
 			if (SectionModelStorage && TrackEditorPtr)
 			{
@@ -118,6 +120,19 @@ void FAnimationMixerTrackModel::ForceUpdate()
 					ObjectBinding,
 					bCreateOutliners,
 					bChildrenNeedLayout);
+
+				// Handle layer and section decorations
+				if (DecorationModelStorage)
+				{
+					DecorationModelStorage->SyncDecorationModel(Layer, MixerTrack, OutlinerChildren);
+
+					// Section decorations for single-section case (when section outliners aren't created).
+					// Multi-section case: SyncDecorationModel's section outliner iteration handles stale cleanup.
+					if (!bCreateOutliners && LayerSections.Num() == 1)
+					{
+						DecorationModelStorage->SyncDecorationModel(LayerSections[0], MixerTrack, OutlinerChildren);
+					}
+				}
 			}
 		}
 	}
@@ -173,6 +188,13 @@ void FAnimationMixerTrackModel::ForceUpdate()
 			LayerModel->RefreshLayout(bNeedsLayout);
 		}
 	}
+
+	// Handle track-level decorations
+	FDecorationModelStorageExtension* DecorationModelStorage = SequenceModel->CastDynamic<FDecorationModelStorageExtension>();
+	if (DecorationModelStorage)
+	{
+		DecorationModelStorage->SyncDecorationModel(MixerTrack, MixerTrack, OutlinerChildren);
+	}
 }
 
 bool FAnimationMixerTrackModel::IsRowOccupiedByChildTrack(int32 RowIndex) const
@@ -237,7 +259,7 @@ bool FAnimationMixerTrackModel::OnSectionVerticalDrag(const FSectionVerticalDrag
 {
 	TOptional<FAnimMixerDropPreview> NewPreview = ComputeDropPreview(Context);
 
-	bool bChanged = (NewPreview != CachedPreview);
+	bool bChanged = (CachedPreview.IsSet() && NewPreview != CachedPreview);
 	CachedPreview = NewPreview;
 	return bChanged;
 }
diff --git a/Engine/Plugins/Experimental/MovieSceneAnimMixer/Source/MovieSceneAnimMixerEditor/Private/MovieSceneAnimMixerEditorModule.cpp b/Engine/Plugins/Experimental/MovieSceneAnimMixer/Source/MovieSceneAnimMixerEditor/Private/MovieSceneAnimMixerEditorModule.cpp
--- a/Engine/Plugins/Experimental/MovieSceneAnimMixer/Source/MovieSceneAnimMixerEditor/Private/MovieSceneAnimMixerEditorModule.cpp
+++ b/Engine/Plugins/Experimental/MovieSceneAnimMixer/Source/MovieSceneAnimMixerEditor/Private/MovieSceneAnimMixerEditorModule.cpp
@@ -7,6 +7,7 @@
 #include "AnimTransitionSectionInterface.h"
 #include "IAnimGraph_SequencerMixerTargetConnector.h"
 #include "IMovieSceneAnimMixerTargetMenuProvider.h"
+#include "ISequencerDecorationEditor.h"
 #include "ISequencerModule.h"
 #include "Modules/ModuleInterface.h"
 #include "Modules/ModuleManager.h"
@@ -14,6 +15,7 @@
 #include "MovieSceneAnimationMixerTrackEditor.h"
 #include "MovieSceneAnimCrossfadeTransitionSection.h"
 #include "MovieSceneAnimInertialDeadBlendTransitionSection.h"
+#include "RootMotionTargetDecorationEditor.h"
 #include "Channels/BuiltInChannelEditors.h"
 #include "MVVM/AnimationMixerTrackModel.h"
 
@@ -135,10 +137,20 @@ namespace UE::MovieScene
 		IModularFeatures::Get().RegisterModularFeature(IMovieSceneAnimMixerTargetMenuProvider::GetModularFeatureName(), AnimInstanceTargetMenuProvider.Get());
 		IModularFeatures::Get().RegisterModularFeature(IMovieSceneAnimMixerTargetMenuProvider::GetModularFeatureName(), AnimBlueprintTargetMenuProvider.Get());
 		IModularFeatures::Get().RegisterModularFeature(IMovieSceneAnimMixerTargetMenuProvider::GetModularFeatureName(), AnimNextInjectionTargetMenuProvider.Get());
+
+		// Register decoration editors
+		RootMotionTargetDecorationEditor = MakeUnique<UE::Sequencer::FRootMotionTargetDecorationEditor>();
+		IModularFeatures::Get().RegisterModularFeature(ISequencerDecorationEditor::GetModularFeatureName(), RootMotionTargetDecorationEditor.Get());
 	}
 
 	void FMovieSceneAnimMixerEditorModule::ShutdownModule()
 	{
+		// Unregister decoration editors
+		if (RootMotionTargetDecorationEditor)
+		{
+			IModularFeatures::Get().UnregisterModularFeature(ISequencerDecorationEditor::GetModularFeatureName(), RootMotionTargetDecorationEditor.Get());
+		}
+
 		// Unregister target menu providers
 		if (AutomaticTargetMenuProvider)
 		{
diff --git a/Engine/Plugins/Experimental/MovieSceneAnimMixer/Source/MovieSceneAnimMixerEditor/Private/MovieSceneAnimMixerEditorModule.h b/Engine/Plugins/Experimental/MovieSceneAnimMixer/Source/MovieSceneAnimMixerEditor/Private/MovieSceneAnimMixerEditorModule.h
--- a/Engine/Plugins/Experimental/MovieSceneAnimMixer/Source/MovieSceneAnimMixerEditor/Private/MovieSceneAnimMixerEditorModule.h
+++ b/Engine/Plugins/Experimental/MovieSceneAnimMixer/Source/MovieSceneAnimMixerEditor/Private/MovieSceneAnimMixerEditorModule.h
@@ -11,6 +11,11 @@ class FAnimInstanceTargetMenuProvider;
 class FAnimBlueprintTargetMenuProvider;
 class FAnimNextInjectionTargetMenuProvider;
 
+namespace UE::Sequencer
+{
+	class FRootMotionTargetDecorationEditor;
+}
+
 namespace UE::MovieScene
 {
 	class FMovieSceneAnimMixerEditorModule : public IModuleInterface, public IAnimGraph_SequencerMixerTargetConnector
@@ -32,6 +37,9 @@ namespace UE::MovieScene
 		TUniquePtr<FAnimInstanceTargetMenuProvider> AnimInstanceTargetMenuProvider;
 		TUniquePtr<FAnimBlueprintTargetMenuProvider> AnimBlueprintTargetMenuProvider;
 		TUniquePtr<FAnimNextInjectionTargetMenuProvider> AnimNextInjectionTargetMenuProvider;
+
+		// Decoration editors
+		TUniquePtr<UE::Sequencer::FRootMotionTargetDecorationEditor> RootMotionTargetDecorationEditor;
 	};
 }
 
diff --git a/Engine/Plugins/Experimental/MovieSceneAnimMixer/Source/MovieSceneAnimMixerEditor/Private/MovieSceneAnimationMixerTrackEditor.cpp b/Engine/Plugins/Experimental/MovieSceneAnimMixer/Source/MovieSceneAnimMixerEditor/Private/MovieSceneAnimationMixerTrackEditor.cpp
--- a/Engine/Plugins/Experimental/MovieSceneAnimMixer/Source/MovieSceneAnimMixerEditor/Private/MovieSceneAnimationMixerTrackEditor.cpp
+++ b/Engine/Plugins/Experimental/MovieSceneAnimMixer/Source/MovieSceneAnimMixerEditor/Private/MovieSceneAnimationMixerTrackEditor.cpp
@@ -13,6 +13,7 @@
 #include "IMovieSceneAnimMixerItemMenuProvider.h"
 #include "IMovieSceneAnimMixerTargetMenuProvider.h"
 #include "AnimatedRange.h"
+#include "IMovieSceneModule.h"
 #include "ISequencer.h"
 #include "MovieSceneAnimationMixerItemInterface.h"
 #include "MovieSceneAnimationMixerLayer.h"
@@ -30,6 +31,7 @@
 #include "Sections/MovieSceneSkeletalAnimationSection.h"
 #include "MovieSceneAnimTransitionSectionBase.h"
 #include "MVVM/ViewModels/SectionOutlinerModel.h"
+#include "SequencerCommonHelpers.h"
 #include "SequencerSectionPainter.h"
 #include "SequencerUtilities.h"
 #include "TimeToPixel.h"
@@ -307,22 +309,22 @@ public:
 		}
 	}
 
-	UMovieSceneAnimationBaseTransformDecoration* FindRootDecoration() const
+	UMovieSceneRootMotionSettingsDecoration* FindRootDecoration() const
 	{
 		UMovieSceneSkeletalAnimationSection* AnimSection = Cast<UMovieSceneSkeletalAnimationSection>(WeakSection.Get());
-		return AnimSection ? AnimSection->FindDecoration<UMovieSceneAnimationBaseTransformDecoration>() : nullptr;
+		return AnimSection ? AnimSection->FindDecoration<UMovieSceneRootMotionSettingsDecoration>() : nullptr;
 	}
 
-	UMovieSceneAnimationBaseTransformDecoration* FindOrCreateRootDecoration() const
+	UMovieSceneRootMotionSettingsDecoration* FindOrCreateRootDecoration() const
 	{
 		UMovieSceneSkeletalAnimationSection* AnimSection = Cast<UMovieSceneSkeletalAnimationSection>(WeakSection.Get());
 		if (AnimSection)
 		{
-			UMovieSceneAnimationBaseTransformDecoration* Decoration = AnimSection->FindDecoration<UMovieSceneAnimationBaseTransformDecoration>();
+			UMovieSceneRootMotionSettingsDecoration* Decoration = AnimSection->FindDecoration<UMovieSceneRootMotionSettingsDecoration>();
 			if (!Decoration)
 			{
 				AnimSection->Modify();
-				Decoration = AnimSection->GetOrCreateDecoration<UMovieSceneAnimationBaseTransformDecoration>();
+				Decoration = AnimSection->GetOrCreateDecoration<UMovieSceneRootMotionSettingsDecoration>();
 			}
 			return Decoration;
 		}
@@ -376,7 +378,7 @@ public:
 		}
 		MenuBuilder.EndSection();
 
-		UMovieSceneAnimationBaseTransformDecoration* Decoration = FindRootDecoration();
+		UMovieSceneRootMotionSettingsDecoration* Decoration = FindRootDecoration();
 		if (Decoration && Decoration->GetRootTransformMode() == EMovieSceneRootMotionTransformMode::Offset)
 		{
 			MenuBuilder.BeginSection(NAME_None, LOCTEXT("RootTransformOffset_Label", "Offset"));
@@ -395,7 +397,7 @@ public:
 
 	bool IsTransformMode(EMovieSceneRootMotionTransformMode Mode) const
 	{
-		if (UMovieSceneAnimationBaseTransformDecoration* Decoration = FindRootDecoration())
+		if (UMovieSceneRootMotionSettingsDecoration* Decoration = FindRootDecoration())
 		{
 			return Decoration->GetRootTransformMode() == Mode;
 		}
@@ -407,7 +409,7 @@ public:
 		FScopedTransaction Transaction(LOCTEXT("SetRootTransformNone", "Use Root Transform from Asset"));
 
 		UMovieSceneSkeletalAnimationSection*         AnimSection = Cast<UMovieSceneSkeletalAnimationSection>(WeakSection.Get());
-		UMovieSceneAnimationBaseTransformDecoration* Decoration  = FindOrCreateRootDecoration();
+		UMovieSceneRootMotionSettingsDecoration* Decoration  = FindOrCreateRootDecoration();
 		if (Decoration && AnimSection)
 		{
 			Decoration->Modify();
@@ -420,7 +422,7 @@ public:
 		FScopedTransaction Transaction(LOCTEXT("SetRootTransformOffset", "Offset Root Transform from Asset"));
 
 		UMovieSceneSkeletalAnimationSection*         AnimSection = Cast<UMovieSceneSkeletalAnimationSection>(WeakSection.Get());
-		UMovieSceneAnimationBaseTransformDecoration* Decoration  = FindOrCreateRootDecoration();
+		UMovieSceneRootMotionSettingsDecoration* Decoration  = FindOrCreateRootDecoration();
 		if (Decoration && AnimSection)
 		{
 			Decoration->Modify();
@@ -433,7 +435,7 @@ public:
 		FScopedTransaction Transaction(LOCTEXT("SetRootTransformOverride", "Override Root Transform"));
 
 		UMovieSceneSkeletalAnimationSection*         AnimSection = Cast<UMovieSceneSkeletalAnimationSection>(WeakSection.Get());
-		UMovieSceneAnimationBaseTransformDecoration* Decoration  = FindOrCreateRootDecoration();
+		UMovieSceneRootMotionSettingsDecoration* Decoration  = FindOrCreateRootDecoration();
 		if (Decoration && AnimSection)
 		{
 			Decoration->Modify();
@@ -448,7 +450,7 @@ public:
 		TSharedPtr<ISequencer> SequencerPtr = Sequencer.Pin();
 
 		UMovieSceneSkeletalAnimationSection*         AnimSection = Cast<UMovieSceneSkeletalAnimationSection>(WeakSection.Get());
-		UMovieSceneAnimationBaseTransformDecoration* Decoration  = FindRootDecoration();
+		UMovieSceneRootMotionSettingsDecoration* Decoration  = FindRootDecoration();
 		UAnimSequenceBase*                           Animation   = AnimSection ? AnimSection->GetAnimation() : nullptr;
 
 		if (SequencerPtr && Decoration && Animation)
@@ -1243,6 +1245,31 @@ TSharedRef<SWidget> FAnimationMixerTrackEditor::BuildAddSectionSubMenu(TWeakView
 		MenuBuilder.EndSection();
 	}
 
+	// Add decoration/modifier menu entries for compatible decorations
+	// For track-level: add decorations to the track itself
+	// For layer-level: add decorations to the layer and its sections
+	if (LayerExtension)
+	{
+		// Layer-level: get the layer and add its compatible decorations
+		UMovieSceneAnimationMixerLayer* Layer = MixerTrack->GetLayer(LayerExtension->GetRowIndex());
+		if (Layer)
+		{
+			SequencerHelpers::BuildDecorationMenu(MenuBuilder, Layer);
+
+			// For single-section layers, also show decorations compatible with the section
+			TArray<UMovieSceneSection*> LayerSections = Layer->GetSections();
+			if (LayerSections.Num() == 1 && LayerSections[0])
+			{
+				SequencerHelpers::BuildDecorationMenu(MenuBuilder, LayerSections[0]);
+			}
+		}
+	}
+	else
+	{
+		// Track-level: add decorations to the track
+		SequencerHelpers::BuildDecorationMenu(MenuBuilder, MixerTrack);
+	}
+
 	return MenuBuilder.MakeWidget();
 }
 
diff --git a/Engine/Plugins/Experimental/MovieSceneAnimMixer/Source/MovieSceneAnimMixerEditor/Private/RootMotionTargetDecorationEditor.cpp b/Engine/Plugins/Experimental/MovieSceneAnimMixer/Source/MovieSceneAnimMixerEditor/Private/RootMotionTargetDecorationEditor.cpp
new file mode 100644
--- /dev/null
+++ b/Engine/Plugins/Experimental/MovieSceneAnimMixer/Source/MovieSceneAnimMixerEditor/Private/RootMotionTargetDecorationEditor.cpp
@@ -0,0 +1,66 @@
+// Copyright Epic Games, Inc. All Rights Reserved.
+
+#include "RootMotionTargetDecorationEditor.h"
+
+#include "MovieSceneRootMotionSection.h"
+#include "MovieSceneRootMotionTargetDecoration.h"
+#include "SequencerSectionPainter.h"
+
+#define LOCTEXT_NAMESPACE "RootMotionTargetDecorationEditor"
+
+namespace UE::Sequencer
+{
+
+//------------------------------------------------------------------------------
+// FRootMotionTargetSectionInterface
+//------------------------------------------------------------------------------
+
+FRootMotionTargetSectionInterface::FRootMotionTargetSectionInterface(
+	UMovieSceneSection& InSection,
+	UMovieSceneTrack& /*InTrack*/,
+	FGuid /*InObjectBinding*/)
+	: SectionPtr(&InSection)
+{
+}
+
+UMovieSceneSection* FRootMotionTargetSectionInterface::GetSectionObject()
+{
+	return SectionPtr.Get();
+}
+
+FText FRootMotionTargetSectionInterface::GetSectionTitle() const
+{
+	return FText::GetEmpty();
+}
+
+int32 FRootMotionTargetSectionInterface::OnPaintSection(FSequencerSectionPainter& Painter) const
+{
+	// Use the mixer item tint for the section background
+	if (const UMovieSceneRootMotionSection* RootMotionSection = Cast<UMovieSceneRootMotionSection>(SectionPtr.Get()))
+	{
+		FLinearColor Tint = FLinearColor(RootMotionSection->GetMixerItemTint());
+		return Painter.PaintSectionBackground(Tint);
+	}
+	return Painter.PaintSectionBackground();
+}
+
+//------------------------------------------------------------------------------
+// FRootMotionTargetDecorationEditor
+//------------------------------------------------------------------------------
+
+UClass* FRootMotionTargetDecorationEditor::GetDecorationClass() const
+{
+	return UMovieSceneRootMotionTargetDecoration::StaticClass();
+}
+
+TSharedRef<ISequencerSection> FRootMotionTargetDecorationEditor::MakeSectionInterface(
+	UMovieSceneSection& Section,
+	UMovieSceneTrack& OwningTrack,
+	FGuid ObjectBinding)
+{
+	return MakeShared<FRootMotionTargetSectionInterface>(Section, OwningTrack, ObjectBinding);
+}
+
+} // namespace UE::Sequencer
+
+#undef LOCTEXT_NAMESPACE
diff --git a/Engine/Plugins/Experimental/MovieSceneAnimMixer/Source/MovieSceneAnimMixerEditor/Private/RootMotionTargetDecorationEditor.h b/Engine/Plugins/Experimental/MovieSceneAnimMixer/Source/MovieSceneAnimMixerEditor/Private/RootMotionTargetDecorationEditor.h
new file mode 100644
--- /dev/null
+++ b/Engine/Plugins/Experimental/MovieSceneAnimMixer/Source/MovieSceneAnimMixerEditor/Private/RootMotionTargetDecorationEditor.h
@@ -0,0 +1,51 @@
+// Copyright Epic Games, Inc. All Rights Reserved.
+
+#pragma once
+
+#include "ISequencerDecorationEditor.h"
+#include "ISequencerSection.h"
+#include "Templates/SharedPointer.h"
+#include "UObject/WeakObjectPtr.h"
+
+class UMovieSceneSection;
+class UMovieSceneTrack;
+
+namespace UE::Sequencer
+{
+
+/**
+ * Section interface for root motion sections displayed under decoration outliners.
+ * Provides basic display with the channel data for root motion destination.
+ */
+class FRootMotionTargetSectionInterface
+	: public ISequencerSection
+	, public TSharedFromThis<FRootMotionTargetSectionInterface>
+{
+public:
+	FRootMotionTargetSectionInterface(UMovieSceneSection& InSection, UMovieSceneTrack& InTrack, FGuid InObjectBinding);
+
+	//~ ISequencerSection interface
+	virtual UMovieSceneSection* GetSectionObject() override;
+	virtual FText GetSectionTitle() const override;
+	virtual int32 OnPaintSection(FSequencerSectionPainter& Painter) const override;
+	virtual bool IsReadOnly() const override { return false; }
+
+private:
+	TWeakObjectPtr<UMovieSceneSection> SectionPtr;
+};
+
+/**
+ * Decoration editor for UMovieSceneRootMotionTargetDecoration.
+ * Creates section interfaces for the root motion sections that the decoration provides.
+ */
+class FRootMotionTargetDecorationEditor : public ISequencerDecorationEditor
+{
+public:
+	virtual UClass* GetDecorationClass() const override;
+	virtual TSharedRef<ISequencerSection> MakeSectionInterface(
+		UMovieSceneSection& Section,
+		UMovieSceneTrack& OwningTrack,
+		FGuid ObjectBinding) override;
+};
+
+} // namespace UE::Sequencer
diff --git a/Engine/Source/Editor/Sequencer/Private/ISequencerSection.cpp b/Engine/Source/Editor/Sequencer/Private/ISequencerSection.cpp
--- a/Engine/Source/Editor/Sequencer/Private/ISequencerSection.cpp
+++ b/Engine/Source/Editor/Sequencer/Private/ISequencerSection.cpp
@@ -9,6 +9,7 @@
 #include "Containers/ContainerAllocationPolicies.h"
 #include "Containers/Map.h"
 #include "Containers/UnrealString.h"
+#include "Decorations/IMovieSceneChannelDecoration.h"
 #include "HAL/PlatformCrt.h"
 #include "IKeyArea.h"
 #include "ISectionLayoutBuilder.h"
@@ -96,6 +97,15 @@ void ISequencerSection::GenerateSectionLayout( ISectionLayoutBuilder& LayoutBuil
 			const FMovieSceneChannelMetaData& MetaData = AllMetaData[Index];
 			if (MetaData.bEnabled)
 			{
+				// Skip channels owned by decorations - they're handled by FDecorationOutlinerModel
+				if (UObject* OwningObject = MetaData.WeakOwningObject.Get())
+				{
+					if (OwningObject->Implements<UMovieSceneChannelDecoration>())
+					{
+						continue;
+					}
+				}
+
 				FName GroupName = *MetaData.Group.ToString();
 
 				FGroupData* ExistingGroup = GroupToChannelsMap.Find(GroupName);
diff --git a/Engine/Source/Editor/Sequencer/Private/MVVM/DecorationModelStorageExtension.cpp b/Engine/Source/Editor/Sequencer/Private/MVVM/DecorationModelStorageExtension.cpp
new file mode 100644
--- /dev/null
+++ b/Engine/Source/Editor/Sequencer/Private/MVVM/DecorationModelStorageExtension.cpp
@@ -0,0 +1,124 @@
+// Copyright Epic Games, Inc. All Rights Reserved.
+
+#include "MVVM/DecorationModelStorageExtension.h"
+
+#include "MVVM/ViewModels/DecorationOutlinerModel.h"
+#include "MVVM/ViewModels/SectionOutlinerModel.h"
+#include "MVVM/ViewModels/ViewModelHierarchy.h"
+#include "SequencerCommonHelpers.h"
+
+namespace UE::Sequencer
+{
+
+FDecorationModelStorageExtension::FDecorationModelStorageExtension()
+{
+}
+
+void FDecorationModelStorageExtension::OnReinitialize()
+{
+	for (auto It = ContainerToModel.CreateIterator(); It; ++It)
+	{
+		if (It.Key().ResolveObjectPtr() == nullptr || It.Value().Pin().Get() == nullptr)
+		{
+			It.RemoveCurrent();
+		}
+	}
+	ContainerToModel.Compact();
+}
+
+TSharedPtr<FDecorationOutlinerModel> FDecorationModelStorageExtension::FindOrCreateModel(
+	UMovieSceneDecorationContainerObject* Container,
+	UMovieSceneTrack* ParentTrack,
+	bool& bOutNewlyCreated)
+{
+	bOutNewlyCreated = false;
+
+	if (!Container || !ParentTrack)
+	{
+		return nullptr;
+	}
+
+	FObjectKey ContainerKey(Container);
+	if (TSharedPtr<FDecorationOutlinerModel> Existing = ContainerToModel.FindRef(ContainerKey).Pin())
+	{
+		return Existing;
+	}
+
+	TSharedPtr<FDecorationOutlinerModel> NewModel = MakeShared<FDecorationOutlinerModel>(ParentTrack, Container);
+	ContainerToModel.Add(ContainerKey, NewModel);
+	bOutNewlyCreated = true;
+	return NewModel;
+}
+
+TSharedPtr<FDecorationOutlinerModel> FDecorationModelStorageExtension::FindModel(
+	const UMovieSceneDecorationContainerObject* Container) const
+{
+	FObjectKey ContainerKey(Container);
+	return ContainerToModel.FindRef(ContainerKey).Pin();
+}
+
+TSharedPtr<FDecorationOutlinerModel> FDecorationModelStorageExtension::SyncDecorationModel(
+	UMovieSceneDecorationContainerObject* Container,
+	UMovieSceneTrack* ParentTrack,
+	FViewModelChildren& OutlinerChildren)
+{
+	if (!Container || !ParentTrack)
+	{
+		return nullptr;
+	}
+
+	TSharedPtr<FDecorationOutlinerModel> Result;
+
+	if (SequencerHelpers::HasOutlinerVisibleDecorations(Container))
+	{
+		bool bNewlyCreated = false;
+		TSharedPtr<FDecorationOutlinerModel> Model = FindOrCreateModel(Container, ParentTrack, bNewlyCreated);
+		if (Model)
+		{
+			OutlinerChildren.AddChild(Model);
+			// New models get OnConstruct via the framework (SetSharedData). For existing models,
+			// only reconstruct if the container's decorations actually changed, otherwise keep
+			// existing children intact to avoid flickering during layout-only rebuilds (e.g. drags).
+			if (!bNewlyCreated && Model->NeedsReconstruct())
+			{
+				Model->OnConstruct();
+			}
+		}
+		Result = Model;
+	}
+	else
+	{
+		// Container has no visible decorations - remove any stale model
+		if (TSharedPtr<FDecorationOutlinerModel> StaleModel = FindModel(Container))
+		{
+			StaleModel->RemoveFromParent();
+		}
+	}
+
+	// Clean up stale decoration models on section outliners. When decorations are
+	// removed from a section (e.g. via Delete), the event fires before the decoration
+	// is actually removed from the container, so the event-driven cleanup can't detect
+	// the change. This deferred check catches it during the subsequent rebuild.
+	// Only does removal - creation/adoption is handled by RefreshDecorations after
+	// layout builders complete (avoids AddChild flickering on every rebuild).
+	for (TSharedPtr<FSectionOutlinerModel> SectionOutliner : OutlinerChildren.IterateSubList<FSectionOutlinerModel>())
+	{
+		if (!SectionOutliner)
+		{
+			continue;
+		}
+
+		UMovieSceneSection* Section = SectionOutliner->GetSection();
+		if (Section && !SequencerHelpers::HasOutlinerVisibleDecorations(Section))
+		{
+			if (TSharedPtr<FDecorationOutlinerModel> StaleModel = FindModel(Section))
+			{
+				StaleModel->RemoveFromParent();
+			}
+		}
+	}
+
+	return Result;
+}
+
+} // namespace UE::Sequencer
diff --git a/Engine/Source/Editor/Sequencer/Private/MVVM/SequencerViewModelTypeIDs.cpp b/Engine/Source/Editor/Sequencer/Private/MVVM/SequencerViewModelTypeIDs.cpp
--- a/Engine/Source/Editor/Sequencer/Private/MVVM/SequencerViewModelTypeIDs.cpp
+++ b/Engine/Source/Editor/Sequencer/Private/MVVM/SequencerViewModelTypeIDs.cpp
@@ -37,6 +37,7 @@
 #include "MVVM/PinEditorExtension.h"
 #include "MVVM/TrackModelStorageExtension.h"
 #include "MVVM/TrackRowModelStorageExtension.h"
+#include "MVVM/DecorationModelStorageExtension.h"
 #include "MVVM/SectionModelStorageExtension.h"
 
 
@@ -84,6 +85,7 @@ UE_SEQUENCER_DEFINE_VIEW_MODEL_TYPE_ID(FCurveEditorIntegrationExtension);
 UE_SEQUENCER_DEFINE_VIEW_MODEL_TYPE_ID(FFolderModelStorageExtension);
 UE_SEQUENCER_DEFINE_VIEW_MODEL_TYPE_ID(FObjectBindingModelStorageExtension);
 UE_SEQUENCER_DEFINE_VIEW_MODEL_TYPE_ID(FPinEditorExtension);
+UE_SEQUENCER_DEFINE_VIEW_MODEL_TYPE_ID(FDecorationModelStorageExtension);
 UE_SEQUENCER_DEFINE_VIEW_MODEL_TYPE_ID(FSectionModelStorageExtension);
 UE_SEQUENCER_DEFINE_VIEW_MODEL_TYPE_ID(FTrackModelStorageExtension);
 UE_SEQUENCER_DEFINE_VIEW_MODEL_TYPE_ID(FTrackRowModelStorageExtension);
diff --git a/Engine/Source/Editor/Sequencer/Private/MVVM/ViewModels/ChannelModel.cpp b/Engine/Source/Editor/Sequencer/Private/MVVM/ViewModels/ChannelModel.cpp
--- a/Engine/Source/Editor/Sequencer/Private/MVVM/ViewModels/ChannelModel.cpp
+++ b/Engine/Source/Editor/Sequencer/Private/MVVM/ViewModels/ChannelModel.cpp
@@ -1017,8 +1017,33 @@ TSharedPtr<SWidget> FChannelGroupOutlinerModel::CreateOutlinerViewForColumn(cons
 			.Visibility(this, &FChannelGroupOutlinerModel::GetKeyEditorVisibility);
 	}
 
+	// Helper to check if any channel in this group can be keyed
+	auto CanAnyChannelBeKeyed = [this]()
+	{
+		for (const TWeakViewModelPtr<FChannelModel>& WeakChannel : GetChannels())
+		{
+			if (TSharedPtr<FChannelModel> Channel = WeakChannel.Pin())
+			{
+				if (TSharedPtr<IKeyArea> KeyArea = Channel->GetKeyArea())
+				{
+					const FMovieSceneChannelMetaData* MetaData = KeyArea->GetChannel().GetMetaData();
+					if (MetaData && MetaData->bCanBeKeyed)
+					{
+						return true;
+					}
+				}
+			}
+		}
+		return false;
+	};
+
 	if (InColumnName == FCommonOutlinerNames::KeyFrame)
 	{
+		if (!CanAnyChannelBeKeyed())
+		{
+			return nullptr;
+		}
+
 		EKeyNavigationButtons Buttons = EKeyNavigationButtons::AddKey;
 		return SNew(SSequencerKeyNavigationButtons, SharedThis(this), Editor->GetSequencer())
 			.Buttons(Buttons);
@@ -1026,6 +1051,11 @@ TSharedPtr<SWidget> FChannelGroupOutlinerModel::CreateOutlinerViewForColumn(cons
 
 	if (InColumnName == FCommonOutlinerNames::Nav)
 	{
+		if (!CanAnyChannelBeKeyed())
+		{
+			return nullptr;
+		}
+
 		EKeyNavigationButtons Buttons = InParams.TreeViewRow->IsColumnVisible(FCommonOutlinerNames::KeyFrame)
 			? EKeyNavigationButtons::NavOnly
 			: EKeyNavigationButtons::All;
diff --git a/Engine/Source/Editor/Sequencer/Private/MVVM/ViewModels/DecorationOutlinerModel.cpp b/Engine/Source/Editor/Sequencer/Private/MVVM/ViewModels/DecorationOutlinerModel.cpp
--- a/Engine/Source/Editor/Sequencer/Private/MVVM/ViewModels/DecorationOutlinerModel.cpp
+++ b/Engine/Source/Editor/Sequencer/Private/MVVM/ViewModels/DecorationOutlinerModel.cpp
@@ -3,25 +3,31 @@
 #include "MVVM/ViewModels/DecorationOutlinerModel.h"
 
 #include "ISequencer.h"
-#include "MovieScene.h"
+#include "ISequencerDecorationEditor.h"
+#include "SequencerLog.h"
+#include "MovieSceneSection.h"
 #include "SequencerCommonHelpers.h"
+#include "Channels/MovieSceneChannelProxy.h"
+#include "Decorations/IMovieSceneChannelDecoration.h"
 #include "Decorations/IMovieSceneSectionProviderDecoration.h"
 #include "MVVM/SectionModelStorageExtension.h"
-#include "MVVM/Extensions/IRecyclableExtension.h"
-#include "MVVM/ViewModels/ObjectBindingModel.h"
+#include "MVVM/ViewModels/CategoryModel.h"
+#include "MVVM/ViewModels/ChannelModel.h"
 #include "MVVM/ViewModels/SequenceModel.h"
 #include "MVVM/ViewModels/SectionModel.h"
 #include "MVVM/ViewModels/SectionOutlinerModel.h"
 #include "MVVM/ViewModels/SequencerEditorViewModel.h"
+#include "MVVM/ViewModels/TrackModel.h"
 #include "MVVM/ViewModels/TrackModelLayoutBuilder.h"
-#include "MVVM/ViewModels/TrackRowModel.h"
 #include "MVVM/ViewModels/ViewDensity.h"
+#include "SSequencer.h"
 
 #define LOCTEXT_NAMESPACE "DecorationOutlinerModel"
 
 namespace UE::Sequencer
 {
 UE_SEQUENCER_DEFINE_CASTABLE(FDecorationOutlinerModel);
+UE_SEQUENCER_DEFINE_CASTABLE(FDecorationCategoryModel);
 
 FDecorationOutlinerModel::FDecorationOutlinerModel(UMovieSceneTrack* InParentTrack, UMovieSceneDecorationContainerObject* InDecorationContainer)
 	: WeakParentTrack(InParentTrack)
@@ -34,13 +40,45 @@ FDecorationOutlinerModel::~FDecorationOutlinerModel()
 {
 }
 
+bool FDecorationOutlinerModel::NeedsReconstruct() const
+{
+	UMovieSceneDecorationContainerObject* Container = WeakDecorationContainer.Get();
+	if (!Container)
+	{
+		return false;
+	}
+
+	TArrayView<const TObjectPtr<UObject>> CurrentDecorations = Container->GetDecorations();
+	if (CurrentDecorations.Num() != LastConstructedDecorations.Num())
+	{
+		return true;
+	}
+
+	for (int32 Index = 0; Index < CurrentDecorations.Num(); ++Index)
+	{
+		if (CurrentDecorations[Index] != LastConstructedDecorations[Index].Get())
+		{
+			return true;
+		}
+	}
+	return false;
+}
+
 void FDecorationOutlinerModel::OnConstruct()
 {
 	if (!WeakParentTrack.IsValid() || !WeakDecorationContainer.IsValid())
 	{
 		return;
 	}
-	
+
+	// Early out if we don't have shared data yet - this means we're being called
+	// during AddChild() before we're fully integrated into the hierarchy.
+	// The explicit OnConstruct() call after AddChild() will do the real work.
+	if (!GetSharedData())
+	{
+		return;
+	}
+
 	// Get ISequencer pointer from the ancestor view model.
 	TSharedPtr<FSequenceModel> SequenceModel = FindAncestorOfType<FSequenceModel>();
 	if (!SequenceModel)
@@ -62,27 +100,56 @@ void FDecorationOutlinerModel::OnConstruct()
 	
 	TrackEditor = Sequencer->GetTrackEditor(WeakParentTrack.Get());
 	FViewModelChildren OutlinerChildren = GetTopLevelChannels();
-	
+
 	FScopedViewModelListHead RecycledModels(AsShared(), EViewModelListType::Recycled);
 	OutlinerChildren.MoveChildrenTo<IRecyclableExtension>(RecycledModels.GetChildren(), IRecyclableExtension::CallOnRecycle);
-	
+	FViewModelChildren RecycledChildren = RecycledModels.GetChildren();
+
+	// Remove old channel models from their parent section models before destroying proxies.
+	// Channel models are added to section models for ancestor lookup (see later in this function),
+	// and section models are reused across OnConstruct() calls. Without this cleanup, old channels
+	// with stale handles (pointing to destroyed proxies) would persist alongside new channels.
+	for (const TWeakViewModelPtr<FChannelModel>& WeakChannel : TrackedChannelModels)
+	{
+		if (TViewModelPtr<FChannelModel> OldChannel = WeakChannel.Pin())
+		{
+			OldChannel->RemoveFromParent();
+		}
+	}
+	TrackedChannelModels.Empty();
+
+	// Clear old channel proxies before rebuilding
+	ChannelDecorationProxies.Reset();
+
 	TSharedPtr<FViewModel> OutlinerTail;
 	for (TObjectPtr<UObject> Decoration : WeakDecorationContainer->GetDecorations())
 	{
 		// For decorations that have sections, create and add models for those sections.
 		if (IMovieSceneSectionProviderDecoration* SectionProvider = Cast<IMovieSceneSectionProviderDecoration>(Decoration))
 		{
+			// Find the decoration editor for this decoration class
+			// Decorations are generic and can go on any track, so we can't use the parent track's
+			// track editor to create section interfaces. Instead, decoration editors are registered
+			// via modular features for each decoration type that provides sections.
+			ISequencerDecorationEditor* DecorationEditor = ISequencerDecorationEditor::FindDecorationEditor(Decoration->GetClass());
+			if (!DecorationEditor)
+			{
+				UE_LOG(LogSequencer, Warning, TEXT("No decoration editor registered for decoration class %s. Sections will not be displayed."), *Decoration->GetClass()->GetName());
+				continue;
+			}
+
 			for (UMovieSceneSection* Section : SectionProvider->GetSections())
 			{
 				if (!Section)
 				{
 					continue;
 				}
-				const TSharedRef<ISequencerSection> SectionInterface = TrackEditor->MakeSectionInterface(*Section, *WeakParentTrack.Get(), WeakParentTrack.Get()->FindObjectBindingGuid());
+				const TSharedRef<ISequencerSection> SectionInterface = DecorationEditor->MakeSectionInterface(*Section, *WeakParentTrack.Get(), WeakParentTrack.Get()->FindObjectBindingGuid());
 				const TSharedPtr<FSectionModel> SectionModel = SectionModelStorage->CreateModelForSection(Section, SectionInterface);
 
 				TSharedPtr<FSectionOutlinerModel> SectionOutlinerModel;
-				for (TSharedPtr<FSectionOutlinerModel> ExistingOutliner : OutlinerChildren.IterateSubList<FSectionOutlinerModel>())
+				// Search in RecycledChildren (where we moved existing models) rather than OutlinerChildren (which is now empty)
+				for (TSharedPtr<FSectionOutlinerModel> ExistingOutliner : RecycledChildren.IterateSubList<FSectionOutlinerModel>())
 				{
 					if (ExistingOutliner && ExistingOutliner->GetSection() == Section)
 					{
@@ -105,9 +172,153 @@ void FDecorationOutlinerModel::OnConstruct()
 
 				FTrackModelLayoutBuilder LayoutBuilder(SectionOutlinerModel);
 				LayoutBuilder.RefreshLayout(SectionModel);
+
+				// Refresh decorations after layout builder completes - the layout builder
+				// recycles the outliner child list, so decorations must be added after.
+				SectionOutlinerModel->RefreshDecorations();
+			}
+		}
+		// For decorations that add channels to existing sections
+		else if (IMovieSceneChannelDecoration* ChannelDecoration = Cast<IMovieSceneChannelDecoration>(Decoration))
+		{
+			if (!TrackEditor)
+			{
+				continue;
+			}
+
+			// Check if the decoration container is a section (for section-level decorations)
+			UMovieSceneSection* ParentSection = Cast<UMovieSceneSection>(WeakDecorationContainer.Get());
+			if (!ParentSection)
+			{
+				continue;
+			}
+
+			// Get or create the section model for the parent section
+			// We need to use the section interface from the section model (which is kept alive by the model)
+			// rather than creating a temporary one that would be destroyed when this function returns
+			TSharedRef<ISequencerSection> NewSectionInterface = TrackEditor->MakeSectionInterface(*ParentSection, *WeakParentTrack.Get(), WeakParentTrack.Get()->FindObjectBindingGuid());
+			const TSharedPtr<FSectionModel> SectionModel = SectionModelStorage->CreateModelForSection(ParentSection, NewSectionInterface);
+
+			// Use the section interface from the section model - this ensures it stays alive
+			// (CreateModelForSection may have returned an existing model with its own interface)
+			TSharedPtr<ISequencerSection> SectionInterface = SectionModel->GetSectionInterface();
+
+			// Get channels from the decoration
+			FMovieSceneChannelProxyData ChannelProxyData;
+			ChannelDecoration->PopulateChannelProxy(ChannelProxyData);
+
+			// Build channel proxy and store it to keep channel handles valid
+			TSharedPtr<FMovieSceneChannelProxy> ChannelProxyPtr = MakeShared<FMovieSceneChannelProxy>(MoveTemp(ChannelProxyData));
+			ChannelDecorationProxies.Add(ChannelProxyPtr);
+			FMovieSceneChannelProxy& ChannelProxy = *ChannelProxyPtr;
+
+			// Get the display name from the decoration class
+			FText DecorationDisplayName = Decoration->GetClass()->GetDisplayNameText();
+
+			// Create a category group for this decoration's channels
+			FName CategoryName(*Decoration->GetClass()->GetName());
+			TSharedPtr<FDecorationCategoryModel> CategoryOutlinerModel;
+			// Search in RecycledChildren (where we moved existing models) rather than OutlinerChildren (which is now empty)
+			for (TSharedPtr<FDecorationCategoryModel> ExistingCategory : RecycledChildren.IterateSubList<FDecorationCategoryModel>())
+			{
+				if (ExistingCategory && ExistingCategory->GetCategoryName() == CategoryName)
+				{
+					CategoryOutlinerModel = ExistingCategory;
+					break;
+				}
+			}
+
+			if (!CategoryOutlinerModel)
+			{
+				CategoryOutlinerModel = MakeShared<FDecorationCategoryModel>(CategoryName, DecorationDisplayName, Decoration);
+			}
+
+			OutlinerChildren.InsertChild(CategoryOutlinerModel, OutlinerTail);
+			OutlinerTail = CategoryOutlinerModel;
+
+			// Collect all channels with their metadata for sorting
+			struct FChannelSortData
+			{
+				FMovieSceneChannelHandle Handle;
+				const FMovieSceneChannelMetaData* MetaData;
+			};
+			TArray<FChannelSortData> SortedChannels;
+
+			for (const FMovieSceneChannelEntry& Entry : ChannelProxy.GetAllEntries())
+			{
+				TArrayView<const FMovieSceneChannelMetaData> AllMetaData = Entry.GetMetaData();
+
+				for (int32 ChannelIndex = 0; ChannelIndex < AllMetaData.Num(); ++ChannelIndex)
+				{
+					FChannelSortData SortData;
+					SortData.Handle = ChannelProxy.MakeHandle(Entry.GetChannelTypeName(), ChannelIndex);
+					SortData.MetaData = &AllMetaData[ChannelIndex];
+					SortedChannels.Add(SortData);
+				}
+			}
+
+			// Sort by SortOrder
+			SortedChannels.Sort([](const FChannelSortData& A, const FChannelSortData& B)
+			{
+				return A.MetaData->SortOrder < B.MetaData->SortOrder;
+			});
+
+			// Add channels from the decoration to the category in sorted order
+			FViewModelChildren CategoryChildren = CategoryOutlinerModel->GetChildList(EViewModelListType::Outliner);
+			TSharedPtr<FViewModel> ChannelTail;
+			for (const FChannelSortData& ChannelData : SortedChannels)
+			{
+				const FMovieSceneChannelMetaData& MetaData = *ChannelData.MetaData;
+				FMovieSceneChannelHandle ChannelHandle = ChannelData.Handle;
+
+				// Create channel group outliner model
+				FName ChannelName = MetaData.Name;
+				TSharedPtr<FChannelGroupOutlinerModel> ChannelGroupModel;
+				for (TSharedPtr<FChannelGroupOutlinerModel> ExistingGroup : CategoryChildren.IterateSubList<FChannelGroupOutlinerModel>())
+				{
+					if (ExistingGroup && ExistingGroup->GetChannelName() == ChannelName)
+					{
+						ChannelGroupModel = ExistingGroup;
+						break;
+					}
+				}
+
+				if (!ChannelGroupModel)
+				{
+					ChannelGroupModel = MakeShared<FChannelGroupOutlinerModel>(ChannelName, MetaData.DisplayText, MetaData.GetTooltipTextDelegate);
+				}
+				else
+				{
+					// Clear old channels when reusing an existing group to avoid stale handles
+					// from a previous proxy that has been destroyed
+					ChannelGroupModel->OnRecycle();
+				}
+
+				CategoryChildren.InsertChild(ChannelGroupModel, ChannelTail);
+				ChannelTail = ChannelGroupModel;
+
+				// Create channel model
+				TSharedPtr<FChannelModel> ChannelModel = MakeShared<FChannelModel>(ChannelName, SectionInterface, ChannelHandle);
+				ChannelGroupModel->AddChannel(ChannelModel);
+				ChannelModel->SetLinkedOutlinerItem(ChannelGroupModel.ToSharedRef());
+				ChannelModel->Initialize(SectionInterface, ChannelHandle);
+
+				// Add channel model as child of section model so FindAncestorOfType<FSectionModel>() works
+				// This is required for ArrangeVirtualTrackLaneView() to get the section range for key drawing
+				SectionModel->GetChildList(FTrackModel::GetTopLevelChannelType()).AddChild(ChannelModel);
+
+				// Track this channel so we can remove it from section model during cleanup
+				TrackedChannelModels.Add(ChannelModel);
 			}
 		}
 	}
+
+	// Store the decoration state so NeedsReconstruct() can detect changes
+	LastConstructedDecorations.Reset();
+	for (const TObjectPtr<UObject>& Decoration : WeakDecorationContainer->GetDecorations())
+	{
+		LastConstructedDecorations.Add(Decoration);
+	}
 }
 
 /*------------------------------------------------------------------------------------
@@ -131,7 +342,7 @@ void FDecorationOutlinerModel::BuildContextMenu(FMenuBuilder& MenuBuilder)
 	{
 		for (TObjectPtr<UObject> Decoration : WeakDecorationContainer->GetDecorations())
 		{
-			if (Decoration->Implements<UMovieSceneSectionProviderDecoration>())
+			if (Decoration && Decoration->Implements<UMovieSceneSectionProviderDecoration>())
 			{
 				const TArray<TWeakObjectPtr<>> TrackAreaModels = SequencerHelpers::GetSectionObjectsFromTrackAreaModels(GetTrackAreaModelList());
 				SequencerHelpers::BuildEditSectionMenu(Sequencer, TrackAreaModels, MenuBuilder, true);
@@ -152,8 +363,7 @@ void FDecorationOutlinerModel::BuildSidebarMenu(FMenuBuilder& MenuBuilder)
 FOutlinerSizing FDecorationOutlinerModel::GetOutlinerSizing() const
 {
 	const FViewDensityInfo Density = GetEditor()->GetViewDensity();
-	constexpr float SectionAreaDefaultHeight = 27.0f;
-	float Height = Density.UniformHeight.Get(SectionAreaDefaultHeight);
+	float Height = Density.UniformHeight.Get(SequencerLayoutConstants::SectionAreaDefaultHeight);
 	for (const TSharedPtr<FSectionModel> Section : DecoratorList.Iterate<FSectionModel>())
 	{
 		Height = Section->GetSectionInterface()->GetSectionHeight(Density);
@@ -231,16 +441,16 @@ void FDecorationOutlinerModel::Delete()
 
 	for (const UObject* Decoration : DecorationContainer->GetDecorations())
 	{
-		DecorationsToRemove.Add(Decoration->GetClass());
+		if (Decoration)
+		{
+			DecorationsToRemove.Add(Decoration->GetClass());
+		}
 	}
 
 	for (UClass* DecorationClass : DecorationsToRemove)
 	{
 		DecorationContainer->RemoveDecoration(DecorationClass);
 	}
-
-	// Notify parent model that this outliner has been deleted.
-	OnOutlinerDeletedDelegate.ExecuteIfBound();
 }
 
 /*-----------------------------------------------------------------------------
@@ -284,22 +494,77 @@ void FDecorationOutlinerModel::SetIsLocked(bool bIsLocked)
 	}
 }
 
-TViewModelPtr<ITrackRowExtension> FDecorationOutlinerModel::GetParentTrackRow() const
+
+TArray<UObject*> FDecorationOutlinerModel::GetDecorations() const
 {
-	const UMovieSceneTrack* Track = WeakParentTrack.Get();
-	if (!Track)
+	TArray<UObject*> Result;
+
+	if (UMovieSceneDecorationContainerObject* Container = WeakDecorationContainer.Get())
 	{
-		return nullptr;
+		for (const TObjectPtr<UObject>& Decoration : Container->GetDecorations())
+		{
+			if (Decoration)
+			{
+				Result.Add(Decoration.Get());
+			}
+		}
 	}
 
-	const TViewModelPtr<FViewModel> ParentModel = FindAncestorOfType<FViewModel>();
-	if (!ParentModel.IsValid())
+	return Result;
+}
+
+void FDecorationOutlinerModel::DeleteDecoration(UClass* DecorationClass)
+{
+	UMovieSceneDecorationContainerObject* DecorationContainer = WeakDecorationContainer.Get();
+	UMovieSceneTrack* ParentTrack = WeakParentTrack.Get();
+
+	if (!DecorationContainer || !ParentTrack || !DecorationClass)
 	{
-		return nullptr;
+		return;
 	}
-	
-	return ParentModel.ImplicitCast();
+
+	ParentTrack->Modify();
+	DecorationContainer->Modify();
+
+	DecorationContainer->RemoveDecoration(DecorationClass);
+}
+
+/*------------------------------------------------------------------------------------
+	FDecorationCategoryModel
+ -----------------------------------------------------------------------------------*/
+
+FDecorationCategoryModel::FDecorationCategoryModel(FName InCategoryName, const FText& InDisplayText, UObject* InDecoration)
+	: FCategoryGroupModel(InCategoryName, InDisplayText, FGetMovieSceneTooltipText())
+	, WeakDecoration(InDecoration)
+{
+}
+
+bool FDecorationCategoryModel::CanDelete(FText* OutErrorMessage) const
+{
+	// Can delete if we have a valid decoration and parent decoration model
+	if (!WeakDecoration.IsValid())
+	{
+		return false;
+	}
+
+	TSharedPtr<FDecorationOutlinerModel> DecorationModel = FindAncestorOfType<FDecorationOutlinerModel>();
+	return DecorationModel.IsValid();
 }
+
+void FDecorationCategoryModel::Delete()
+{
+	if (!WeakDecoration.IsValid())
+	{
+		return;
+	}
+
+	TSharedPtr<FDecorationOutlinerModel> DecorationModel = FindAncestorOfType<FDecorationOutlinerModel>();
+	if (DecorationModel)
+	{
+		DecorationModel->DeleteDecoration(WeakDecoration->GetClass());
+	}
+}
+
 }
 
 #undef LOCTEXT_NAMESPACE
\ No newline at end of file
diff --git a/Engine/Source/Editor/Sequencer/Private/MVVM/ViewModels/SectionOutlinerModel.cpp b/Engine/Source/Editor/Sequencer/Private/MVVM/ViewModels/SectionOutlinerModel.cpp
--- a/Engine/Source/Editor/Sequencer/Private/MVVM/ViewModels/SectionOutlinerModel.cpp
+++ b/Engine/Source/Editor/Sequencer/Private/MVVM/ViewModels/SectionOutlinerModel.cpp
@@ -11,14 +11,18 @@
 #include "MVVM/ViewModels/OutlinerColumns/OutlinerColumnTypes.h"
 #include "MVVM/Views/SOutlinerItemViewBase.h"
 #include "MVVM/Views/SSequencerKeyNavigationButtons.h"
+#include "MVVM/Views/ViewUtilities.h"
+#include "SKeyAreaEditorSwitcher.h"
 #include "ISequencerModule.h"
 #include "ISequencer.h"
 #include "ISequencerSection.h"
 #include "ISequencerTrackEditor.h"
 #include "MovieSceneSection.h"
 #include "MovieSceneTrack.h"
+#include "Decorations/IMovieSceneSectionProviderDecoration.h"
+#include "Decorations/MovieSceneDecorationContainer.h"
 #include "Modules/ModuleManager.h"
-#include "SequencerCommonHelpers.h"
+#include "MVVM/DecorationModelStorageExtension.h"
 #include "Widgets/SNullWidget.h"
 #include "Decorations/MovieSceneMuteSoloDecoration.h"
 
@@ -62,6 +66,7 @@ FSectionOutlinerModel::FSectionOutlinerModel(UMovieSceneSection* InSection, TSha
 
 FSectionOutlinerModel::~FSectionOutlinerModel()
 {
+	SectionEventHandler.Unlink();
 }
 
 void FSectionOutlinerModel::OnConstruct()
@@ -78,7 +83,16 @@ void FSectionOutlinerModel::OnConstruct()
 				TrackEditor = Sequencer->GetTrackEditor(Section->GetTypedOuter<UMovieSceneTrack>());
 			}
 		}
+
+		// Subscribe to section decoration events (unlink first in case OnConstruct is called again)
+		SectionEventHandler.Unlink();
+		Section->EventHandlers.Link(SectionEventHandler, this);
 	}
+
+	// Mark decorations for refresh. Don't call RefreshDecorations() here because
+	// callers typically run layout builders after OnConstruct, which recycle the
+	// outliner child list. RefreshDecorations must be called after layout completes.
+	bDecorationsNeedRefresh = true;
 }
 
 FViewModelChildren FSectionOutlinerModel::GetTrackAreaChildren()
@@ -240,6 +254,41 @@ TSharedPtr<SWidget> FSectionOutlinerModel::CreateOutlinerViewForColumn(const FCr
 		}
 	}
 
+	// Edit column - show key area editor for top-level channels
+	if (InColumnName == FCommonOutlinerNames::Edit)
+	{
+		TOptional<FViewModelChildren> TopLevelChannels = FindChildList(FTrackModel::GetTopLevelChannelGroupType());
+		TSharedPtr<FChannelGroupModel> TopLevelChannel = TopLevelChannels ? TopLevelChannels->FindFirstChildOfType<FChannelGroupModel>() : nullptr;
+		if (TopLevelChannel)
+		{
+			return SNew(SKeyAreaEditorSwitcher, TopLevelChannel, Editor);
+		}
+	}
+
+	// Add column - show '+' button for adding decorations/modifiers to the section
+	if (InColumnName == FCommonOutlinerNames::Add)
+	{
+		UMovieSceneSection* Section = GetSection();
+		if (Section && SequencerHelpers::HasCompatibleDecorations(Section))
+		{
+			TWeakObjectPtr<UMovieSceneSection> WeakSectionForMenu = Section;
+			auto GetMenuContent = [WeakSectionForMenu]() -> TSharedRef<SWidget>
+			{
+				FMenuBuilder MenuBuilder(true, nullptr);
+				if (UMovieSceneSection* PinnedSection = WeakSectionForMenu.Get())
+				{
+					SequencerHelpers::BuildDecorationMenu(MenuBuilder, PinnedSection);
+				}
+				return MenuBuilder.MakeWidget();
+			};
+
+			return MakeAddButton(
+				NSLOCTEXT("SectionOutlinerModel", "AddModifier", "Modifier"),
+				FOnGetContent::CreateLambda(GetMenuContent),
+				SharedThis(this));
+		}
+	}
+
 	return nullptr;
 }
 
@@ -494,6 +543,28 @@ void FSectionOutlinerModel::Delete()
 {
 	if (UMovieSceneSection* Section = WeakSection.Get())
 	{
+		// Check if this section is owned by a decoration that provides it
+		// (decoration-provided sections have the decoration as their outer, not the track)
+		UObject* Outer = Section->GetOuter();
+		if (IMovieSceneSectionProviderDecoration* SectionProvider = Cast<IMovieSceneSectionProviderDecoration>(Outer))
+		{
+			TArrayView<TObjectPtr<UMovieSceneSection>> ProviderSections = SectionProvider->GetSections();
+			if (ProviderSections.Num() <= 1)
+			{
+				// Last/only section - remove the entire decoration
+				if (UMovieSceneDecorationContainerObject* Container = Cast<UMovieSceneDecorationContainerObject>(Outer->GetOuter()))
+				{
+					Container->Modify();
+					Container->RemoveDecoration(Outer->GetClass());
+					return;
+				}
+			}
+			// Multi-section case: don't allow direct deletion from outliner.
+			// User must use the decoration editor UI to manage sections.
+			return;
+		}
+
+		// Standard case: section is owned directly by the track
 		if (UMovieSceneTrack* Track = Section->GetTypedOuter<UMovieSceneTrack>())
 		{
 			Track->Modify();
@@ -641,5 +712,57 @@ void FSectionOutlinerModel::SetIsSoloed(bool bIsSoloed)
 	}
 }
 
+/*-----------------------------------------------------------------------------
+	ISectionEventHandler
+-----------------------------------------------------------------------------*/
+
+void FSectionOutlinerModel::OnDecorationAdded(UObject*)
+{
+	bDecorationsNeedRefresh = true;
+	RefreshDecorations();
+}
+
+void FSectionOutlinerModel::OnDecorationRemoved(UObject*)
+{
+	bDecorationsNeedRefresh = true;
+	RefreshDecorations();
+}
+
+void FSectionOutlinerModel::RefreshDecorations()
+{
+	if (!bDecorationsNeedRefresh)
+	{
+		return;
+	}
+	bDecorationsNeedRefresh = false;
+
+	UMovieSceneSection* Section = WeakSection.Get();
+	if (!Section)
+	{
+		return;
+	}
+
+	UMovieSceneTrack* Track = Section->GetTypedOuter<UMovieSceneTrack>();
+	if (!Track)
+	{
+		return;
+	}
+
+	TSharedPtr<FSequenceModel> SequenceModel = FindAncestorOfType<FSequenceModel>();
+	if (!SequenceModel)
+	{
+		return;
+	}
+
+	FDecorationModelStorageExtension* DecorationModelStorage = SequenceModel->CastDynamic<FDecorationModelStorageExtension>();
+	if (!DecorationModelStorage)
+	{
+		return;
+	}
+
+	FViewModelChildren OutlinerChildren = GetChildList(EViewModelListType::Outliner);
+	DecorationModelStorage->SyncDecorationModel(Section, Track, OutlinerChildren);
+}
+
 } // namespace Sequencer
 } // namespace UE
diff --git a/Engine/Source/Editor/Sequencer/Private/MVVM/ViewModels/TrackModel.cpp b/Engine/Source/Editor/Sequencer/Private/MVVM/ViewModels/TrackModel.cpp
--- a/Engine/Source/Editor/Sequencer/Private/MVVM/ViewModels/TrackModel.cpp
+++ b/Engine/Source/Editor/Sequencer/Private/MVVM/ViewModels/TrackModel.cpp
@@ -14,6 +14,7 @@
 #include "MVVM/ViewModels/TrackModelLayoutBuilder.h"
 #include "MVVM/ViewModels/TrackRowModel.h"
 #include "MVVM/ViewModels/TrackRowReorderDragDropOp.h"
+#include "MVVM/DecorationModelStorageExtension.h"
 #include "MVVM/SectionModelStorageExtension.h"
 #include "MVVM/TrackRowModelStorageExtension.h"
 #include "MVVM/Views/SOutlinerTrackView.h"
@@ -27,7 +28,6 @@
 #include "MovieSceneNameableTrack.h"
 #include "Decorations/MovieSceneMuteSoloDecoration.h"
 #include "Tracks/MovieScene3DTransformTrack.h"
-#include "Tracks/MovieScene3DTransformTrack.h"
 #include "Tracks/MovieScenePrimitiveMaterialTrack.h"
 #include "EntitySystem/IMovieSceneBlenderSystemSupport.h"
 
@@ -182,6 +182,8 @@ void FTrackModel::ForceUpdate()
 
 	const int32 NumRows = PopulatedRows.CountSetBits();
 
+	FDecorationModelStorageExtension* DecorationModelStorage = SequenceModel->CastDynamic<FDecorationModelStorageExtension>();
+
 	if (NumRows == 0)
 	{
 		// Reset expansion state if this track can no longer be expanded
@@ -308,6 +310,12 @@ void FTrackModel::ForceUpdate()
 		}
 	}
 
+	// Handle track-level decorations
+	if (DecorationModelStorage)
+	{
+		DecorationModelStorage->SyncDecorationModel(Track, Track, OutlinerChildren);
+	}
+
 	PreviousLayoutNumRows = NumRows;
 }
 
diff --git a/Engine/Source/Editor/Sequencer/Private/MovieSceneTrackEditor.cpp b/Engine/Source/Editor/Sequencer/Private/MovieSceneTrackEditor.cpp
--- a/Engine/Source/Editor/Sequencer/Private/MovieSceneTrackEditor.cpp
+++ b/Engine/Source/Editor/Sequencer/Private/MovieSceneTrackEditor.cpp
@@ -14,6 +14,7 @@
 #include "Sequencer.h"
 #include "Framework/MultiBox/MultiBoxBuilder.h"
 #include "SequencerUtilities.h"
+#include "SequencerCommonHelpers.h"
 #include "MVVM/Views/ViewUtilities.h"
 #include "Framework/Application/SlateApplication.h"
 #include "MVVM/ViewModels/OutlinerColumns/OutlinerColumnTypes.h"
@@ -342,16 +343,29 @@ TSharedPtr<SWidget> FMovieSceneTrackEditor::BuildOutlinerColumnWidget(const FBui
 
 TSharedPtr<SWidget> FMovieSceneTrackEditor::BuildOutlinerEditWidget(const FGuid& ObjectBinding, UMovieSceneTrack* Track, const FBuildEditWidgetParams& Params)
 {
-	if (Track && Track->GetSupportedBlendTypes().Num() > 0)
+	const bool bHasSections = Track && Track->GetSupportedBlendTypes().Num() > 0;
+	const bool bHasDecorations = SequencerHelpers::HasCompatibleDecorations(Track);
+
+	if (bHasSections || bHasDecorations)
 	{
 		TWeakPtr<ISequencer> WeakSequencer = GetSequencer();
 
 		const int32 RowIndex = Params.TrackInsertRowIndex;
-		auto SubMenuCallback = [=]() -> TSharedRef<SWidget>
+		auto SubMenuCallback = [=, this]() -> TSharedRef<SWidget>
 		{
 			FMenuBuilder MenuBuilder(true, nullptr);
 
-			FSequencerUtilities::PopulateMenu_CreateNewSection(MenuBuilder, RowIndex, Track, WeakSequencer);
+			// Add section creation entries if track supports sections
+			if (bHasSections)
+			{
+				FSequencerUtilities::PopulateMenu_CreateNewSection(MenuBuilder, RowIndex, Track, WeakSequencer);
+			}
+
+			// Add decoration entries if track supports decorations
+			if (bHasDecorations)
+			{
+				SequencerHelpers::BuildDecorationMenu(MenuBuilder, Track);
+			}
 
 			return MenuBuilder.MakeWidget();
 		};
@@ -360,7 +374,7 @@ TSharedPtr<SWidget> FMovieSceneTrackEditor::BuildOutlinerEditWidget(const FGuid&
 	}
 	else
 	{
-		return TSharedPtr<SWidget>(); 
+		return TSharedPtr<SWidget>();
 	}
 }
 
diff --git a/Engine/Source/Editor/Sequencer/Private/SSequencerSection.cpp b/Engine/Source/Editor/Sequencer/Private/SSequencerSection.cpp
--- a/Engine/Source/Editor/Sequencer/Private/SSequencerSection.cpp
+++ b/Engine/Source/Editor/Sequencer/Private/SSequencerSection.cpp
@@ -264,9 +264,11 @@ struct FSequencerSectionPainterImpl : FSequencerSectionPainter
 			FLinearColor FillTint = BlendedTint.LinearRGBToHSV();
 			FillTint.G *= .5f;
 			FillTint.B = FMath::Max(.03f, FillTint.B*.1f);
+			// Draw contents background at base LayerId (not incremented) so child track lanes
+			// that paint at the same layer can draw their content on top
 			FSlateDrawElement::MakeBox(
 				DrawElements,
-				++LayerId,
+				LayerId,
 				ContentsGeometry.ToPaintGeometry(),
 				SectionContentsBackgroundBrush,
 				DrawEffects,
@@ -725,30 +727,21 @@ struct FSequencerSectionPainterImpl : FSequencerSectionPainter
 			return;
 		}
 
-		// Find all our children
-		TSet<FWeakViewModelPtr> AllDescendents;
-		for (const FViewModelPtr& Descendent : SectionModel->GetDescendants())
-		{
-			AllDescendents.Add(Descendent);
-		}
-
 		// Find all our child track lanes
 		TArray<TSharedPtr<STrackLane>> EmptyChildLanes;
 		Algo::Transform(OwningTrackLane->GetChildLanes(), EmptyChildLanes, [](const TWeakPtr<STrackLane>& In){ return In.Pin(); });
 
-		// Remove any invalid track lanes, or track lanes that we made a widget for
+		// Remove any invalid track lanes, or track lanes that already have widgets
+		// (widgets handle their own rendering, so we shouldn't draw backgrounds over them)
 		for (int32 Index = EmptyChildLanes.Num()-1; Index >= 0; --Index)
 		{
 			if (EmptyChildLanes[Index])
 			{
-				for (const TPair<FWeakViewModelPtr, TSharedPtr<ITrackLaneWidget>>& Pair : EmptyChildLanes[Index]->GetAllWidgets())
+				// If this track lane has any widgets at all, don't draw an empty background for it
+				// This includes decoration channel widgets which are not descendants of the section
+				if (EmptyChildLanes[Index]->GetAllWidgets().Num() > 0)
 				{
-					if (AllDescendents.Contains(Pair.Key))
-					{
-						EmptyChildLanes.RemoveAtSwap(Index, EAllowShrinking::No);
-						// Move onto the next child lane
-						break;
-					}
+					EmptyChildLanes.RemoveAtSwap(Index, EAllowShrinking::No);
 				}
 			}
 			else
@@ -757,9 +750,34 @@ struct FSequencerSectionPainterImpl : FSequencerSectionPainter
 			}
 		}
 
-		// Paint empty space for remaining track lanes
+		if (EmptyChildLanes.IsEmpty())
+		{
+			return;
+		}
+
+		// Paint section contents background for remaining track lanes (e.g., decoration channels)
+		// Use the same brush and color computation as the section contents area
 		const ESlateDrawEffect DrawEffects = bParentEnabled ? ESlateDrawEffect::None : ESlateDrawEffect::DisabledEffect;
-		static const FSlateBrush* EmptySpaceBrush = FAppStyle::GetBrush("Sequencer.Section.EmptySpace");
+		static const FSlateBrush* ContentsBrush = FAppStyle::GetBrush("Sequencer.Section.Background_Contents");
+		if (!ContentsBrush)
+		{
+			return;
+		}
+
+		// Compute the fill tint the same way PaintSectionBackground does for expanded section contents
+		const USequencerSettings* const SequencerSettings = Sequencer.GetSequencerSettings();
+		if (!SequencerSettings)
+		{
+			return;
+		}
+		UMovieSceneSection* SectionObject = SectionModel->GetSection();
+		UMovieSceneTrack* OuterTrack = SectionObject ? SectionObject->GetTypedOuter<UMovieSceneTrack>() : nullptr;
+		const FLinearColor TrackColor = OuterTrack ? OuterTrack->GetColorTint() : FLinearColor::White;
+		const FLinearColor BlendedTint = BlendColor(TrackColor).CopyWithNewOpacity(SequencerSettings->GetSectionColorAlpha());
+		FLinearColor FillTint = BlendedTint.LinearRGBToHSV();
+		FillTint.G *= .5f;
+		FillTint.B = FMath::Max(.03f, FillTint.B * .1f);
+		const FLinearColor FinalFillColor = FillTint.HSVToLinearRGB();
 
 		for (TSharedPtr<STrackLane> EmptyTrackLane : EmptyChildLanes)
 		{
@@ -773,8 +791,9 @@ struct FSequencerSectionPainterImpl : FSequencerSectionPainter
 					LaneSize,
 					FSlateLayoutTransform(FVector2D(0.f, LaneTop))
 				).ToPaintGeometry(),
-				EmptySpaceBrush,
-				DrawEffects
+				ContentsBrush,
+				DrawEffects,
+				FinalFillColor
 			);
 		}
 
diff --git a/Engine/Source/Editor/Sequencer/Private/Sequencer.cpp b/Engine/Source/Editor/Sequencer/Private/Sequencer.cpp
--- a/Engine/Source/Editor/Sequencer/Private/Sequencer.cpp
+++ b/Engine/Source/Editor/Sequencer/Private/Sequencer.cpp
@@ -98,6 +98,8 @@
 #include "ScopedTransaction.h"
 #include "ISequencerTrackEditor.h"
 #include "MovieSceneToolHelpers.h"
+#include "Decorations/IMovieSceneSectionProviderDecoration.h"
+#include "Decorations/MovieSceneDecorationContainer.h"
 #include "Evaluation/ViewportSettingsPlaybackCapability.h"
 #include "Sections/MovieSceneBoolSection.h"
 #include "Sections/MovieScene3DTransformSection.h"
@@ -1953,16 +1955,30 @@ void FSequencer::DeleteSections(const TSet<UMovieSceneSection*>& Sections)
 			continue;
 		}
 
-		// if this check fails then the section is outered to a type that doesnt know about the section
-		UMovieSceneTrack* Track = CastChecked<UMovieSceneTrack>(Section->GetOuter());
+		UObject* SectionOuter = Section->GetOuter();
+
+		// Check if section comes from a decoration that provides sections
+		if (IMovieSceneSectionProviderDecoration* SectionProvider = Cast<IMovieSceneSectionProviderDecoration>(SectionOuter))
 		{
+			// Section is provided by a decoration - remove the decoration from its container
+			UMovieSceneDecorationContainerObject* Container = Cast<UMovieSceneDecorationContainerObject>(SectionOuter->GetOuter());
+			if (Container)
+			{
+				Container->SetFlags(RF_Transactional);
+				Container->Modify();
+				Container->RemoveDecoration(SectionOuter->GetClass());
+				bAnythingRemoved = true;
+			}
+		}
+		else if (UMovieSceneTrack* Track = Cast<UMovieSceneTrack>(SectionOuter))
+		{
+			// Standard case: section is owned by a track
 			Track->SetFlags(RF_Transactional);
 			Track->Modify();
 			Track->RemoveSection(*Section);
 			Track->UpdateEasing();
+			bAnythingRemoved = true;
 		}
-
-		bAnythingRemoved = true;
 	}
 
 	if (bAnythingRemoved)
diff --git a/Engine/Source/Editor/Sequencer/Private/SequencerCommonHelpers.cpp b/Engine/Source/Editor/Sequencer/Private/SequencerCommonHelpers.cpp
--- a/Engine/Source/Editor/Sequencer/Private/SequencerCommonHelpers.cpp
+++ b/Engine/Source/Editor/Sequencer/Private/SequencerCommonHelpers.cpp
@@ -3,7 +3,12 @@
 #include "SequencerCommonHelpers.h"
 
 #include "Algo/Reverse.h"
+#include "Decorations/IMovieSceneChannelDecoration.h"
+#include "Decorations/IMovieSceneSectionProviderDecoration.h"
+#include "Decorations/MovieSceneDecorationContainer.h"
 #include "Decorations/MovieSceneDecorationContainerCustomization.h"
+#include "IMovieSceneModule.h"
+#include "ISequencer.h"
 #include "EntitySystem/IMovieSceneBlenderSystemSupport.h"
 #include "EntitySystem/MovieSceneBlenderSystem.h"
 #include "FrameNumberDetailsCustomization.h"
@@ -652,6 +657,131 @@ void SequencerHelpers::BuildEditTrackMenu(const TWeakPtr<FSequencer>& InWeakSequ
 	}
 }
 
+void SequencerHelpers::BuildEditDecorationMenu(const TWeakPtr<FSequencer>& InWeakSequencer
+	, const TArray<TWeakObjectPtr<>>& InWeakDecorations
+	, FMenuBuilder& MenuBuilder)
+{
+	using namespace UE::Sequencer;
+
+	if (InWeakDecorations.Num() == 0)
+	{
+		return;
+	}
+
+	const TSharedPtr<FSequencer> Sequencer = InWeakSequencer.Pin();
+	if (!Sequencer)
+	{
+		return;
+	}
+
+	// Group decorations by class so each class gets its own details view
+	TMap<UClass*, TArray<TWeakObjectPtr<>>> DecorationsByClass;
+	for (const TWeakObjectPtr<>& WeakDecoration : InWeakDecorations)
+	{
+		if (UObject* Decoration = WeakDecoration.Get())
+		{
+			DecorationsByClass.FindOrAdd(Decoration->GetClass()).Add(WeakDecoration);
+		}
+	}
+
+	for (const auto& Pair : DecorationsByClass)
+	{
+		UClass* DecorationClass = Pair.Key;
+		const TArray<TWeakObjectPtr<>>& ClassDecorations = Pair.Value;
+
+		// Check if this decoration class provides sections
+		// If so, we display the section properties under the decoration's class header
+		TArray<TWeakObjectPtr<>> ObjectsToDisplay;
+		bool bIsSectionProvider = false;
+
+		for (const TWeakObjectPtr<>& WeakDecoration : ClassDecorations)
+		{
+			if (IMovieSceneSectionProviderDecoration* SectionProvider = Cast<IMovieSceneSectionProviderDecoration>(WeakDecoration.Get()))
+			{
+				bIsSectionProvider = true;
+				for (UMovieSceneSection* Section : SectionProvider->GetSections())
+				{
+					if (Section)
+					{
+						ObjectsToDisplay.AddUnique(Section);
+					}
+				}
+			}
+		}
+
+		if (!bIsSectionProvider)
+		{
+			// Check if this class has any editable properties - skip if empty
+			bool bHasEditableProperties = false;
+			for (TFieldIterator<FProperty> PropIt(DecorationClass); PropIt; ++PropIt)
+			{
+				if (PropIt->HasAnyPropertyFlags(CPF_Edit))
+				{
+					bHasEditableProperties = true;
+					break;
+				}
+			}
+			if (!bHasEditableProperties)
+			{
+				continue;
+			}
+
+			// Use the decorations themselves for display
+			ObjectsToDisplay = ClassDecorations;
+		}
+
+		if (ObjectsToDisplay.IsEmpty())
+		{
+			continue;
+		}
+
+		TSharedRef<SSectionDetailsNotifyHookWrapper> DetailsNotifyWrapper = SNew(SSectionDetailsNotifyHookWrapper);
+		FDetailsViewArgs DetailsViewArgs;
+		{
+			DetailsViewArgs.bAllowSearch = false;
+			DetailsViewArgs.bCustomFilterAreaLocation = true;
+			DetailsViewArgs.bCustomNameAreaLocation = true;
+			DetailsViewArgs.bHideSelectionTip = true;
+			DetailsViewArgs.bLockable = false;
+			DetailsViewArgs.bSearchInitialKeyFocus = true;
+			DetailsViewArgs.bUpdatesFromSelection = false;
+			DetailsViewArgs.bShowOptions = false;
+			DetailsViewArgs.bShowModifiedPropertiesOption = false;
+			DetailsViewArgs.bShowScrollBar = false;
+			DetailsViewArgs.NotifyHook = &DetailsNotifyWrapper.Get();
+			DetailsViewArgs.ColumnWidth = 0.45f;
+		}
+
+		TSharedRef<IDetailsView> DetailsView = FModuleManager::GetModuleChecked<FPropertyEditorModule>("PropertyEditor").CreateDetailView(DetailsViewArgs);
+
+		DetailsView->SetIsPropertyVisibleDelegate(FIsPropertyVisible::CreateLambda([bIsSectionProvider](const FPropertyAndParent& PropertyAndParent)
+			{
+				// Hide the Decorations property when showing sections from section-providing decorations
+				// to prevent the "Modifiers" category with decoration management UI from appearing
+				if (bIsSectionProvider && PropertyAndParent.Property.GetFName() == TEXT("Decorations"))
+				{
+					return false;
+				}
+				return FPropertyEditorPermissionList::Get().DoesPropertyPassFilter(PropertyAndParent.Property.GetOwnerStruct(), PropertyAndParent.Property.GetFName());
+			})
+		);
+
+		Sequencer->OnInitializeDetailsPanel().Broadcast(DetailsView, Sequencer.ToSharedRef());
+		DetailsView->SetObjects(ObjectsToDisplay);
+
+		DetailsNotifyWrapper->SetDetailsAndSequencer(DetailsView, InWeakSequencer);
+		DetailsNotifyWrapper->SetEnabled(!Sequencer->IsReadOnly());
+
+		// Use the decoration class display name as the section header
+		FText SectionLabel = DecorationClass->GetDisplayNameText();
+		MenuBuilder.BeginSection(*DecorationClass->GetName(), SectionLabel);
+		{
+			MenuBuilder.AddWidget(DetailsNotifyWrapper, FText::GetEmpty(), true);
+		}
+		MenuBuilder.EndSection();
+	}
+}
+
 void SequencerHelpers::BuildBlendingMenu(const TWeakPtr<FSequencer>& InWeakSequencer
 	, const TWeakObjectPtr<UMovieSceneTrack>& InTrackWeak
 	, FMenuBuilder& MenuBuilder)
@@ -1011,3 +1141,104 @@ void SequencerHelpers::SortOutlinerItems(FSequencer& Sequencer
 
 	Sequencer.RefreshTree();
 }
+
+bool SequencerHelpers::HasCompatibleDecorations(const UMovieSceneDecorationContainerObject* DecorationContainer)
+{
+	if (!DecorationContainer)
+	{
+		return false;
+	}
+
+	TSet<UClass*> CompatibleDecorations;
+	IMovieSceneModule::Get().GetCompatibleDecorationsForContainer(DecorationContainer, CompatibleDecorations);
+
+	// Check if any decoration is not already added
+	for (UClass* DecorationClass : CompatibleDecorations)
+	{
+		bool bAlreadyExists = DecorationContainer->GetDecorations().ContainsByPredicate(
+			[DecorationClass](const TObjectPtr<UObject>& Existing)
+			{
+				return Existing && Existing->GetClass() == DecorationClass;
+			});
+		if (!bAlreadyExists)
+		{
+			return true;
+		}
+	}
+	return false;
+}
+
+bool SequencerHelpers::HasOutlinerVisibleDecorations(const UMovieSceneDecorationContainerObject* Container)
+{
+	if (!Container)
+	{
+		return false;
+	}
+
+	for (const TObjectPtr<UObject>& Decoration : Container->GetDecorations())
+	{
+		if (!Decoration)
+		{
+			continue;
+		}
+
+		if (Cast<IMovieSceneChannelDecoration>(Decoration.Get()) ||
+			Cast<IMovieSceneSectionProviderDecoration>(Decoration.Get()))
+		{
+			return true;
+		}
+	}
+
+	return false;
+}
+
+void SequencerHelpers::BuildDecorationMenu(
+	FMenuBuilder& MenuBuilder,
+	UMovieSceneDecorationContainerObject* DecorationContainer)
+{
+	if (!DecorationContainer)
+	{
+		return;
+	}
+
+	TSet<UClass*> CompatibleDecorations;
+	IMovieSceneModule::Get().GetCompatibleDecorationsForContainer(DecorationContainer, CompatibleDecorations);
+
+	if (CompatibleDecorations.Num() == 0)
+	{
+		return;
+	}
+
+	MenuBuilder.BeginSection(NAME_None, NSLOCTEXT("SequencerHelpers", "Modifiers_Label", "Modifiers"));
+	{
+		for (UClass* DecorationClass : CompatibleDecorations)
+		{
+			bool bAlreadyExists = DecorationContainer->GetDecorations().ContainsByPredicate(
+				[DecorationClass](const TObjectPtr<UObject>& Existing)
+				{
+					return Existing && Existing->GetClass() == DecorationClass;
+				});
+
+			if (!bAlreadyExists)
+			{
+				MenuBuilder.AddMenuEntry(
+					DecorationClass->GetDisplayNameText(),
+					DecorationClass->GetToolTipText(),
+					FSlateIcon(),
+					FUIAction(FExecuteAction::CreateLambda([DecorationContainer, DecorationClass]()
+					{
+						FScopedTransaction Transaction(
+							FText::Format(NSLOCTEXT("SequencerHelpers", "AddDecoration", "Add {0}"),
+							DecorationClass->GetDisplayNameText()));
+
+						DecorationContainer->Modify();
+						UObject* NewDecoration = NewObject<UObject>(
+							DecorationContainer, DecorationClass, NAME_None, RF_Transactional);
+						DecorationContainer->AddDecoration(NewDecoration);
+					}))
+				);
+			}
+		}
+	}
+	MenuBuilder.EndSection();
+}
diff --git a/Engine/Source/Editor/Sequencer/Private/SequencerModule.cpp b/Engine/Source/Editor/Sequencer/Private/SequencerModule.cpp
--- a/Engine/Source/Editor/Sequencer/Private/SequencerModule.cpp
+++ b/Engine/Source/Editor/Sequencer/Private/SequencerModule.cpp
@@ -30,6 +30,7 @@
 #include "MVVM/CurveEditorIntegrationExtension.h"
 #include "MVVM/FolderModelStorageExtension.h"
 #include "MVVM/ObjectBindingModelStorageExtension.h"
+#include "MVVM/DecorationModelStorageExtension.h"
 #include "MVVM/SectionModelStorageExtension.h"
 #include "MVVM/TrackModelStorageExtension.h"
 #include "MVVM/TrackRowModelStorageExtension.h"
@@ -529,6 +530,7 @@ public:
 				InModel->AddDynamicExtension(FTrackModelStorageExtension::ID, TrackModelDelegates);
 				InModel->AddDynamicExtension(FTrackRowModelStorageExtension::ID);
 				InModel->AddDynamicExtension(FSectionModelStorageExtension::ID);
+				InModel->AddDynamicExtension(FDecorationModelStorageExtension::ID);
 
 				// If the editor supports a curve editor, add an integration extension to
 				// sync view-model hierarchies between the outliner and curve editor.
diff --git a/Engine/Source/Editor/Sequencer/Private/Widgets/Sidebar/SequencerSelectionDrawer.cpp b/Engine/Source/Editor/Sequencer/Private/Widgets/Sidebar/SequencerSelectionDrawer.cpp
--- a/Engine/Source/Editor/Sequencer/Private/Widgets/Sidebar/SequencerSelectionDrawer.cpp
+++ b/Engine/Source/Editor/Sequencer/Private/Widgets/Sidebar/SequencerSelectionDrawer.cpp
@@ -8,10 +8,12 @@
 #include "Modules/ModuleManager.h"
 #include "MovieScene.h"
 #include "MovieSceneFolder.h"
+#include "MVVM/Extensions/ISectionOwnerExtension.h"
 #include "MVVM/Extensions/ITrackExtension.h"
 #include "MVVM/Selection/Selection.h"
 #include "MVVM/ViewModelPtr.h"
 #include "MVVM/ViewModels/CategoryModel.h"
+#include "MVVM/ViewModels/DecorationOutlinerModel.h"
 #include "MVVM/ViewModels/FolderModel.h"
 #include "MVVM/ViewModels/SectionModel.h"
 #include "MVVM/ViewModels/SectionOutlinerModel.h"
@@ -373,6 +375,7 @@ void FSequencerSelectionDrawer::BuildOutlinerDetails(const TSharedRef<FSequencer
 	TArray<TWeakObjectPtr<>> WeakFolderObjects;
 	TArray<TWeakObjectPtr<>> WeakSectionObjects;
 	TArray<TWeakObjectPtr<>> WeakTrackObjects;
+	TArray<TWeakObjectPtr<>> WeakDecorationObjects;
 	TSet<TViewModelPtr<FObjectBindingModel>> ObjectBindings;
 	TSet<TViewModelPtr<FChannelGroupOutlinerModel>> ChannelGroups;
 	TArray<TPair<TWeakObjectPtr<UMovieSceneTrack>, int32>> SelectedTrackRows;
@@ -421,6 +424,25 @@ void FSequencerSelectionDrawer::BuildOutlinerDetails(const TSharedRef<FSequencer
         {
         	ChannelGroups.Add(ChannelGroupOutlinerModel);
         }
+		// Handle decoration container (the "Modifiers" row) - check before FSectionOutlinerModel
+		else if (const TViewModelPtr<FDecorationOutlinerModel> DecorationOutlinerModel = OutlinerItem.ImplicitCast())
+		{
+			// Get decorations directly from the model
+			// Section-providing decorations will have their sections displayed under the decoration header
+			// by BuildEditDecorationMenu, so we don't collect sections separately here
+			for (UObject* Decoration : DecorationOutlinerModel->GetDecorations())
+			{
+				WeakDecorationObjects.AddUnique(Decoration);
+			}
+		}
+		// Handle individual decoration category (e.g., "Root Motion Settings")
+		else if (const TViewModelPtr<FDecorationCategoryModel> DecorationCategoryModel = OutlinerItem.ImplicitCast())
+		{
+			if (UObject* Decoration = DecorationCategoryModel->GetDecoration())
+			{
+				WeakDecorationObjects.AddUnique(Decoration);
+			}
+		}
 		else if (const TViewModelPtr<FSectionOutlinerModel> SectionOutlinerModel = OutlinerItem.ImplicitCast())
 		{
 			bIsSectionOutliner = true;
@@ -466,6 +488,12 @@ void FSequencerSelectionDrawer::BuildOutlinerDetails(const TSharedRef<FSequencer
 		SequencerHelpers::BuildEditTrackMenu(WeakSequencer, WeakTrackObjects, MenuBuilder, false);
 	}
 
+	// Show decoration properties before section properties so decoration class headers appear above their section content
+	if (!WeakDecorationObjects.IsEmpty())
+	{
+		SequencerHelpers::BuildEditDecorationMenu(WeakSequencer, WeakDecorationObjects, MenuBuilder);
+	}
+
 	// Only show section details for a selected track/trackrow if a single infinite section on the row or a section outliner
 	if (WeakSectionObjects.Num() == 1 && WeakSectionObjects[0].IsValid())
 	{
diff --git a/Engine/Source/Editor/Sequencer/Public/ISequencerDecorationEditor.h b/Engine/Source/Editor/Sequencer/Public/ISequencerDecorationEditor.h
new file mode 100644
--- /dev/null
+++ b/Engine/Source/Editor/Sequencer/Public/ISequencerDecorationEditor.h
@@ -0,0 +1,82 @@
+// Copyright Epic Games, Inc. All Rights Reserved.
+
+#pragma once
+
+#include "CoreMinimal.h"
+#include "Features/IModularFeatures.h"
+
+class ISequencerSection;
+class UMovieSceneSection;
+class UMovieSceneTrack;
+
+/**
+ * Interface for decoration editors that can create section interfaces for decorations
+ * that provide their own sections (via IMovieSceneSectionProviderDecoration).
+ *
+ * Decorations are generic and can be added to any track type, so the parent track's
+ * track editor cannot know how to create section interfaces for decoration-provided sections.
+ * Instead, each decoration that provides sections should have a corresponding decoration editor
+ * registered via modular features.
+ *
+ * Example registration in an editor module's StartupModule():
+ *
+ * class FMyDecorationEditor : public ISequencerDecorationEditor
+ * {
+ *     virtual UClass* GetDecorationClass() const override { return UMyDecoration::StaticClass(); }
+ *     virtual TSharedRef<ISequencerSection> MakeSectionInterface(...) override { return MakeShared<FMySectionInterface>(...); }
+ * };
+ *
+ * IModularFeatures::Get().RegisterModularFeature(ISequencerDecorationEditor::GetModularFeatureName(), MyDecorationEditorPtr);
+ */
+class ISequencerDecorationEditor : public IModularFeature
+{
+public:
+	virtual ~ISequencerDecorationEditor() = default;
+
+	/**
+	 * Get the decoration class this editor handles.
+	 * Used to match decorations to their editors.
+	 */
+	virtual UClass* GetDecorationClass() const = 0;
+
+	/**
+	 * Create a section interface for a section provided by this decoration.
+	 * The section interface is used for displaying/editing the section in Sequencer.
+	 */
+	virtual TSharedRef<ISequencerSection> MakeSectionInterface(
+		UMovieSceneSection& Section,
+		UMovieSceneTrack& OwningTrack,
+		FGuid ObjectBinding) = 0;
+
+	/**
+	 * Get the modular feature name for decoration editors.
+	 */
+	static FName GetModularFeatureName()
+	{
+		static FName Name(TEXT("SequencerDecorationEditor"));
+		return Name;
+	}
+
+	/**
+	 * Find a decoration editor for the given decoration class.
+	 * Returns nullptr if none is registered.
+	 */
+	static ISequencerDecorationEditor* FindDecorationEditor(UClass* DecorationClass)
+	{
+		if (!DecorationClass)
+		{
+			return nullptr;
+		}
+
+		TArray<ISequencerDecorationEditor*> Editors = IModularFeatures::Get().GetModularFeatureImplementations<ISequencerDecorationEditor>(GetModularFeatureName());
+		for (ISequencerDecorationEditor* Editor : Editors)
+		{
+			if (Editor && Editor->GetDecorationClass() == DecorationClass)
+			{
+				return Editor;
+			}
+		}
+
+		return nullptr;
+	}
+};
diff --git a/Engine/Source/Editor/Sequencer/Public/MVVM/DecorationModelStorageExtension.h b/Engine/Source/Editor/Sequencer/Public/MVVM/DecorationModelStorageExtension.h
new file mode 100644
--- /dev/null
+++ b/Engine/Source/Editor/Sequencer/Public/MVVM/DecorationModelStorageExtension.h
@@ -0,0 +1,67 @@
+// Copyright Epic Games, Inc. All Rights Reserved.
+
+#pragma once
+
+#include "Containers/Map.h"
+#include "MVVM/Extensions/DynamicExtensionContainer.h"
+#include "MVVM/ViewModelTypeID.h"
+#include "Templates/SharedPointer.h"
+#include "UObject/ObjectKey.h"
+
+#define UE_API SEQUENCER_API
+
+class UMovieSceneTrack;
+class UMovieSceneDecorationContainerObject;
+
+namespace UE
+{
+namespace Sequencer
+{
+
+class FDecorationOutlinerModel;
+struct FViewModelChildren;
+
+// Central storage for decoration outliner models, keyed by decoration container.
+// Registered on FSequenceModel's shared data, following FSectionModelStorageExtension pattern.
+class FDecorationModelStorageExtension
+	: public IDynamicExtension
+{
+public:
+
+	UE_SEQUENCER_DECLARE_VIEW_MODEL_TYPE_ID_API(UE_API, FDecorationModelStorageExtension);
+
+	UE_API FDecorationModelStorageExtension();
+
+	// Finds an existing model for the container, or creates a new one.
+	// ParentTrack is passed through to FDecorationOutlinerModel's constructor.
+	// bOutNewlyCreated is set to true if a new model was created, false if an existing one was found.
+	UE_API TSharedPtr<FDecorationOutlinerModel> FindOrCreateModel(
+		UMovieSceneDecorationContainerObject* Container,
+		UMovieSceneTrack* ParentTrack,
+		bool& bOutNewlyCreated);
+
+	// Finds an existing model for the container, returns nullptr if not found.
+	UE_API TSharedPtr<FDecorationOutlinerModel> FindModel(
+		const UMovieSceneDecorationContainerObject* Container) const;
+
+	// Ensures the decoration model for the given container is present in OutlinerChildren
+	// if the container has visible decorations, or removed if it doesn't. Handles all
+	// lifecycle: creation, re-adoption from recycled, NeedsReconstruct checks, and stale
+	// model removal. Callers don't need to check HasOutlinerVisibleDecorations - this
+	// method handles that internally. Returns the model if one was added, or nullptr.
+	UE_API TSharedPtr<FDecorationOutlinerModel> SyncDecorationModel(
+		UMovieSceneDecorationContainerObject* Container,
+		UMovieSceneTrack* ParentTrack,
+		FViewModelChildren& OutlinerChildren);
+
+	UE_API virtual void OnReinitialize() override;
+
+private:
+
+	TMap<FObjectKey, TWeakPtr<FDecorationOutlinerModel>> ContainerToModel;
+};
+
+} // namespace Sequencer
+} // namespace UE
+
+#undef UE_API
diff --git a/Engine/Source/Editor/Sequencer/Public/MVVM/ViewModels/DecorationOutlinerModel.h b/Engine/Source/Editor/Sequencer/Public/MVVM/ViewModels/DecorationOutlinerModel.h
--- a/Engine/Source/Editor/Sequencer/Public/MVVM/ViewModels/DecorationOutlinerModel.h
+++ b/Engine/Source/Editor/Sequencer/Public/MVVM/ViewModels/DecorationOutlinerModel.h
@@ -3,19 +3,19 @@
 #pragma once
 #include "ISequencerTrackEditor.h"
 #include "OutlinerItemModel.h"
-#include "MVVM/Extensions/IConditionableExtension.h"
 #include "MVVM/Extensions/IDeletableExtension.h"
 #include "MVVM/Extensions/ILockableExtension.h"
 #include "MVVM/Extensions/ITopLevelChannelHolderExtension.h"
 #include "MVVM/Extensions/ITrackAreaExtension.h"
+#include "MVVM/ViewModels/CategoryModel.h"
 
 class UMovieSceneTrackRowDecoration;
+struct FMovieSceneChannelProxy;
 
 #define UE_API SEQUENCER_API
 
 namespace UE::Sequencer
 {
-DECLARE_DELEGATE(FOnOutlinerDeleted)
 
 class FDecorationOutlinerModel
 	: public FOutlinerItemModel
@@ -23,7 +23,6 @@ class FDecorationOutlinerModel
 	, public ILockableExtension
 	, public IDeletableExtension
 	, public ITopLevelChannelHolderExtension
-
 {
 public:
 	UE_SEQUENCER_DECLARE_CASTABLE_API(SEQUENCER_API, FDecorationOutlinerModel
@@ -66,26 +65,59 @@ public:
 
 	UE_API TWeakObjectPtr<UMovieSceneDecorationContainerObject> GetDecorationContainer() const { return WeakDecorationContainer; }
 
-	UE_API FOnOutlinerDeleted& OnOutlinerDeleted() { return OnOutlinerDeletedDelegate; };
+	// Get all decoration objects from this model's container. Used for sidebar property display.
+	UE_API TArray<UObject*> GetDecorations() const;
+
+	// Delete a specific decoration by class. Called by FDecorationCategoryModel when its row is deleted.
+	UE_API void DeleteDecoration(UClass* DecorationClass);
+
+	// Returns true if the container's decorations have changed since the last OnConstruct.
+	UE_API bool NeedsReconstruct() const;
 
 private:
 
-	TViewModelPtr<class ITrackRowExtension> GetParentTrackRow() const;
-	
 	FViewModelListHead DecoratorList{EViewModelListType::Outliner};
 
 	TWeakObjectPtr<UMovieSceneTrack> WeakParentTrack;
-	
+
 	TWeakObjectPtr<UMovieSceneDecorationContainerObject> WeakDecorationContainer;
 
-	/** Track editor for the parent track */
 	TSharedPtr<ISequencerTrackEditor> TrackEditor;
 
-	/** Called to let the owning model know this outliner has been deleted. */
-	FOnOutlinerDeleted OnOutlinerDeletedDelegate;
-	
+	// Channel proxies for IMovieSceneChannelDecoration decorations. Stored to keep channel handles valid.
+	TArray<TSharedPtr<FMovieSceneChannelProxy>> ChannelDecorationProxies;
+
+	// Tracks channel models added to section models so they can be removed during cleanup.
+	TArray<TWeakViewModelPtr<FChannelModel>> TrackedChannelModels;
+
+	// Decoration objects present when OnConstruct last ran. Used to detect when rebuild is needed.
+	TArray<TWeakObjectPtr<UObject>> LastConstructedDecorations;
+
 };
-#undef UE_API
-}
 
+/**
+ * Specialized category model for decoration channels that supports deletion of the specific decoration.
+ * Used instead of FCategoryGroupModel when displaying channels from IMovieSceneChannelDecoration.
+ */
+class FDecorationCategoryModel
+	: public FCategoryGroupModel
+{
+public:
+	UE_SEQUENCER_DECLARE_CASTABLE_API(UE_API, FDecorationCategoryModel, FCategoryGroupModel);
 
+	UE_API explicit FDecorationCategoryModel(FName InCategoryName, const FText& InDisplayText, UObject* InDecoration);
+
+	/* IDeletableExtension - override to delete just this decoration */
+	UE_API virtual bool CanDelete(FText* OutErrorMessage) const override;
+	UE_API virtual void Delete() override;
+
+	/** Get the decoration object this category represents */
+	UE_API UObject* GetDecoration() const { return WeakDecoration.Get(); }
+
+private:
+	/** The decoration object this category represents */
+	TWeakObjectPtr<UObject> WeakDecoration;
+};
+
+#undef UE_API
+}
diff --git a/Engine/Source/Editor/Sequencer/Public/MVVM/ViewModels/SectionOutlinerModel.h b/Engine/Source/Editor/Sequencer/Public/MVVM/ViewModels/SectionOutlinerModel.h
--- a/Engine/Source/Editor/Sequencer/Public/MVVM/ViewModels/SectionOutlinerModel.h
+++ b/Engine/Source/Editor/Sequencer/Public/MVVM/ViewModels/SectionOutlinerModel.h
@@ -10,6 +10,8 @@
 #include "MVVM/Extensions/IConditionableExtension.h"
 #include "MVVM/Extensions/ISectionOwnerExtension.h"
 #include "MVVM/Extensions/ITopLevelChannelHolderExtension.h"
+#include "EventHandlers/ISectionEventHandler.h"
+#include "EventHandlers/MovieSceneDataEventContainer.h"
 
 #include "UObject/WeakObjectPtr.h"
 
@@ -44,6 +46,7 @@ class FSectionOutlinerModel
 	, public IConditionableExtension
 	, public ISectionOwnerExtension
 	, public ITopLevelChannelHolderExtension
+	, public UE::MovieScene::ISectionEventHandler
 {
 public:
 
@@ -124,6 +127,14 @@ public:
 	UE_API bool IsSolo() const override;
 	UE_API void SetIsSoloed(bool bIsSoloed) override;
 
+	/*~ ISectionEventHandler */
+	UE_API void OnDecorationAdded(UObject* AddedDecoration) override;
+	UE_API void OnDecorationRemoved(UObject* RemovedDecoration) override;
+
+	// Refreshes decoration models under this section outliner.
+	// Called during layout when bDecorationsNeedRefresh is true.
+	UE_API void RefreshDecorations();
+
 private:
 
 	FViewModelListHead TopLevelChannelList;
@@ -132,6 +143,9 @@ private:
 	TWeakViewModelPtr<FSectionModel> WeakSectionModel;
 
 	TSharedPtr<ISequencerTrackEditor> TrackEditor;
+
+	UE::MovieScene::TNonIntrusiveEventHandler<UE::MovieScene::ISectionEventHandler> SectionEventHandler;
+	bool bDecorationsNeedRefresh = false;
 };
 
 } // namespace Sequencer
diff --git a/Engine/Source/Editor/Sequencer/Public/SequencerCommonHelpers.h b/Engine/Source/Editor/Sequencer/Public/SequencerCommonHelpers.h
--- a/Engine/Source/Editor/Sequencer/Public/SequencerCommonHelpers.h
+++ b/Engine/Source/Editor/Sequencer/Public/SequencerCommonHelpers.h
@@ -13,6 +13,8 @@
 
 class FSequencer;
 class IKeyArea;
+class ISequencer;
+class UMovieSceneDecorationContainerObject;
 class UMovieSceneSection;
 class UMovieSceneTrack;
 struct FKeyHandle;
@@ -101,7 +103,12 @@ public:
 		, FMenuBuilder& MenuBuilder
 		, const bool bInSubMenu);
 
-
+	/*
+	 * Build an inline menu for editing decoration(s)
+	 */
+	static void BuildEditDecorationMenu(const TWeakPtr<FSequencer>& InWeakSequencer
+		, const TArray<TWeakObjectPtr<>>& InWeakDecorations
+		, FMenuBuilder& MenuBuilder);
 
 	/*
 	 * Build a menu for selection the blend algorithm
@@ -135,6 +142,23 @@ public:
 		, const bool bInSortByItemOrder
 		, const bool bInDescending
 		, const bool bInTransact = true);
+
+	/**
+	 * Check if a decoration container has any compatible decorations that can be added.
+	 * Returns true if there are compatible decorations that are not already present.
+	 */
+	static bool HasCompatibleDecorations(const UMovieSceneDecorationContainerObject* DecorationContainer);
+
+	// Returns true if the container has any decorations that should be visible in the outliner
+	// (implements IMovieSceneChannelDecoration or IMovieSceneSectionProviderDecoration).
+	static bool HasOutlinerVisibleDecorations(const UMovieSceneDecorationContainerObject* Container);
+
+	/**
+	 * Populates a menu with "Add Decoration" entries for compatible decorations.
+	 */
+	static void BuildDecorationMenu(
+		FMenuBuilder& MenuBuilder,
+		UMovieSceneDecorationContainerObject* DecorationContainer);
 };
 
 #undef LOCTEXT_NAMESPACE
diff --git a/Engine/Source/Runtime/MovieScene/Private/Compilation/MovieSceneCompiledDataManager.cpp b/Engine/Source/Runtime/MovieScene/Private/Compilation/MovieSceneCompiledDataManager.cpp
--- a/Engine/Source/Runtime/MovieScene/Private/Compilation/MovieSceneCompiledDataManager.cpp
+++ b/Engine/Source/Runtime/MovieScene/Private/Compilation/MovieSceneCompiledDataManager.cpp
@@ -12,6 +12,7 @@
 #include "Sections/MovieSceneSubSection.h"
 #include "Tracks/MovieSceneSubTrack.h"
 #include "Decorations/IMovieSceneDecoration.h"
+#include "Decorations/IMovieSceneSectionProviderDecoration.h"
 #include "Decorations/MovieSceneTimeWarpDecoration.h"
 #include "IMovieSceneModule.h"
 #include "MovieSceneTimeHelpers.h"
@@ -1535,6 +1536,48 @@ void UMovieSceneCompiledDataManager::GatherTrack(const FMovieSceneBinding* Objec
 
 				Provider->PopulateEvaluationField(Params.LocalClampRange, MetaData, &FieldBuilder);
 			}
+
+			// Handle decorations that provide sections for evaluation
+			if (IMovieSceneSectionProviderDecoration* SectionProvider = Cast<IMovieSceneSectionProviderDecoration>(Decoration))
+			{
+				for (UMovieSceneSection* Section : SectionProvider->GetSections())
+				{
+					if (!Section || !Section->IsActive())
+					{
+						continue;
+					}
+
+					TRange<FFrameNumber> SectionRange = Section->GetRange();
+					TRange<FFrameNumber> EffectiveRange = TRange<FFrameNumber>::Intersection(Params.LocalClampRange, SectionRange);
+					if (EffectiveRange.IsEmpty())
+					{
+						continue;
+					}
+
+					FMovieSceneEvaluationFieldEntityMetaData MetaData(PreCompileResult.DefaultMetaData);
+					MetaData.bEvaluateInSequencePreRoll  = Track->EvalOptions.bEvaluateInPreroll;
+					MetaData.bEvaluateInSequencePostRoll = Track->EvalOptions.bEvaluateInPostroll;
+					MetaData.Condition = Track->ConditionContainer.Condition;
+
+					// If the section is an entity provider, let it populate itself
+					IMovieSceneEntityProvider* SectionEntityProvider = Cast<IMovieSceneEntityProvider>(Section);
+					if (SectionEntityProvider)
+					{
+						if (!SectionEntityProvider->PopulateEvaluationField(EffectiveRange, MetaData, &FieldBuilder))
+						{
+							const int32 EntityIndex   = FieldBuilder.FindOrAddEntity(Section, 0);
+							const int32 MetaDataIndex = FieldBuilder.AddMetaData(MetaData);
+							FieldBuilder.AddPersistentEntity(EffectiveRange, EntityIndex, MetaDataIndex);
+						}
+					}
+					else
+					{
+						const int32 EntityIndex   = FieldBuilder.FindOrAddEntity(Section, 0);
+						const int32 MetaDataIndex = FieldBuilder.AddMetaData(MetaData);
+						FieldBuilder.AddPersistentEntity(EffectiveRange, EntityIndex, MetaDataIndex);
+					}
+				}
+			}
 		}
 
 		IMovieSceneEntityProvider* TrackEntityProvider = Cast<IMovieSceneEntityProvider>(Track);
diff --git a/Engine/Source/Runtime/MovieScene/Private/Decorations/MovieSceneDecorationContainer.cpp b/Engine/Source/Runtime/MovieScene/Private/Decorations/MovieSceneDecorationContainer.cpp
--- a/Engine/Source/Runtime/MovieScene/Private/Decorations/MovieSceneDecorationContainer.cpp
+++ b/Engine/Source/Runtime/MovieScene/Private/Decorations/MovieSceneDecorationContainer.cpp
@@ -88,13 +88,16 @@ void UMovieSceneDecorationContainerObject::AddDecoration(UObject* InDecoration)
 	Decorations.AddDecoration(InDecoration, this, [this](UObject* Decoration){
 		this->OnDecorationAdded(Decoration);
 	});
+	MarkAsChanged();
 }
 
 UObject* UMovieSceneDecorationContainerObject::GetOrCreateDecoration(const TSubclassOf<UObject>& InClass)
 {
-	return Decorations.GetOrCreateDecoration(InClass, this, [this](UObject* Decoration){
+	UObject* Result = Decorations.GetOrCreateDecoration(InClass, this, [this](UObject* Decoration){
 		this->OnDecorationAdded(Decoration);
 	});
+	MarkAsChanged();
+	return Result;
 }
 
 void UMovieSceneDecorationContainerObject::RemoveDecoration(const TSubclassOf<UObject>& InClass)
@@ -102,4 +105,5 @@ void UMovieSceneDecorationContainerObject::RemoveDecoration(const TSubclassOf<UO
 	Decorations.RemoveDecoration(InClass, [this](UObject* Decoration){
 		this->OnDecorationRemoved(Decoration);
 	});
+	MarkAsChanged();
 }
diff --git a/Engine/Source/Runtime/MovieScene/Public/Channels/MovieSceneChannelEditorDataEntry.h b/Engine/Source/Runtime/MovieScene/Public/Channels/MovieSceneChannelEditorDataEntry.h
--- a/Engine/Source/Runtime/MovieScene/Public/Channels/MovieSceneChannelEditorDataEntry.h
+++ b/Engine/Source/Runtime/MovieScene/Public/Channels/MovieSceneChannelEditorDataEntry.h
@@ -59,7 +59,13 @@ protected:
 		static_assert(std::is_same_v<typename TMovieSceneChannelTraits<ChannelType>::ExtendedEditorDataType, void>, "Must supply extended editor data according to the channel's traits.");
 
 		// Add the editor meta-data
-		MetaDataArray.Add(MetaData);
+		FMovieSceneChannelMetaData& AddedMetaData = MetaDataArray.Add_GetRef(MetaData);
+
+		// Automatically disable keying for channel types that don't support it (e.g., default-only channels)
+		if constexpr (!TMovieSceneChannelTraits<ChannelType>::SupportsKeying)
+		{
+			AddedMetaData.bCanBeKeyed = false;
+		}
 	}
 
 	/**
@@ -69,7 +75,13 @@ protected:
 	void AddMetaData(const FMovieSceneChannelMetaData& MetaData, ExtendedEditorDataType&& InExtendedEditorData)
 	{
 		// Add the editor meta-data
-		MetaDataArray.Add(MetaData);
+		FMovieSceneChannelMetaData& AddedMetaData = MetaDataArray.Add_GetRef(MetaData);
+
+		// Automatically disable keying for channel types that don't support it (e.g., default-only channels)
+		if constexpr (!TMovieSceneChannelTraits<ChannelType>::SupportsKeying)
+		{
+			AddedMetaData.bCanBeKeyed = false;
+		}
 
 		// Add the extended channel-type specific editor data
 		auto& TypedImpl = static_cast<TMovieSceneExtendedEditorDataArray<ChannelType>&>(ExtendedEditorDataArray.GetValue());
diff --git a/Engine/Source/Runtime/MovieScene/Public/Channels/MovieSceneChannelTraits.h b/Engine/Source/Runtime/MovieScene/Public/Channels/MovieSceneChannelTraits.h
--- a/Engine/Source/Runtime/MovieScene/Public/Channels/MovieSceneChannelTraits.h
+++ b/Engine/Source/Runtime/MovieScene/Public/Channels/MovieSceneChannelTraits.h
@@ -39,6 +39,9 @@ struct TMovieSceneChannelTraitsBase
 {
 	enum { SupportsDefaults = true };
 
+	/** Whether this channel type supports keying. Default-only channels should set this to false. */
+	enum { SupportsKeying = true };
+
 #if WITH_EDITOR
 
 	/** Type that specifies what editor data should be associated with ChannelType. Void (default) implies no extended data. */
diff --git a/Engine/Source/Runtime/MovieScene/Public/Decorations/IMovieSceneSectionProviderDecoration.h b/Engine/Source/Runtime/MovieScene/Public/Decorations/IMovieSceneSectionProviderDecoration.h
--- a/Engine/Source/Runtime/MovieScene/Public/Decorations/IMovieSceneSectionProviderDecoration.h
+++ b/Engine/Source/Runtime/MovieScene/Public/Decorations/IMovieSceneSectionProviderDecoration.h
@@ -16,13 +16,15 @@ public:
 };
 
 /**
- *  Optional interface for movie scene decorations that contain their own sections.
+ * Optional interface for movie scene decorations that contain their own sections.
+ * During compilation, the compiled data manager automatically adds sections from
+ * decorations implementing this interface to the entity field for evaluation.
  */
 class IMovieSceneSectionProviderDecoration
 {
 public:
 	GENERATED_BODY()
 
-	virtual TArray<TObjectPtr<UMovieSceneSection>>& GetSections() = 0;
-	
+	/** Returns the sections provided by this decoration */
+	virtual TArrayView<TObjectPtr<UMovieSceneSection>> GetSections() = 0;
 };
