Sequencer: Disable the track instead of disabling the sections when baking a transform track.

This is a follow up to 44627952 which did the corresponding for Bake to Control Rig.

#Synced-CL 50752143

diff --git a/Engine/Plugins/Animation/ControlRig/Source/ControlRigEditor/Private/ControlRigSequencerEditorLibrary.cpp b/Engine/Plugins/Animation/ControlRig/Source/ControlRigEditor/Private/ControlRigSequencerEditorLibrary.cpp
--- a/Engine/Plugins/Animation/ControlRig/Source/ControlRigEditor/Private/ControlRigSequencerEditorLibrary.cpp
+++ b/Engine/Plugins/Animation/ControlRig/Source/ControlRigEditor/Private/ControlRigSequencerEditorLibrary.cpp
@@ -32,6 +32,7 @@
 #include "Units/Execution/RigUnit_InverseExecution.h"
 #include "Tracks/MovieSceneSkeletalAnimationTrack.h"
 #include "Exporters/AnimSeqExportOption.h"
+#include "MovieSceneCommonHelpers.h"
 #include "MovieSceneTimeHelpers.h"
 #include "ScopedTransaction.h"
 #include "Sequencer/ControlRigParameterTrackEditor.h"
@@ -1422,18 +1423,7 @@ bool UControlRigSequencerEditorLibrary::BakeToControlRig(UWorld* World, ULevelSe
 					{
 						SkelTrack->Modify();
 
-						// Disable the entire track or each of the rows
-						if (SkelTrack->GetMaxRowIndex() == 0)
-						{
-							SkelTrack->SetEvalDisabled(true);
-						}
-						else
-						{
-							for (int32 RowIndex = 0; RowIndex <= SkelTrack->GetMaxRowIndex(); ++RowIndex)
-							{
-								SkelTrack->SetRowEvalDisabled(true, RowIndex);
-							}
-						}
+						MovieSceneHelpers::DisableTrack(SkelTrack);
 					}
 					//Finish Setup
 					if (ControlRigEditMode)
diff --git a/Engine/Plugins/Animation/ControlRig/Source/ControlRigEditor/Private/Sequencer/ControlRigParameterTrackEditor.cpp b/Engine/Plugins/Animation/ControlRig/Source/ControlRigEditor/Private/Sequencer/ControlRigParameterTrackEditor.cpp
--- a/Engine/Plugins/Animation/ControlRig/Source/ControlRigEditor/Private/Sequencer/ControlRigParameterTrackEditor.cpp
+++ b/Engine/Plugins/Animation/ControlRig/Source/ControlRigEditor/Private/Sequencer/ControlRigParameterTrackEditor.cpp
@@ -1363,18 +1363,7 @@ void FControlRigParameterTrackEditor::BakeToControlRig(const FControlRigAssetStr
 					{
 						SkelTrack->Modify();
 
-						// Disable the entire track or each of the rows
-						if (SkelTrack->GetMaxRowIndex() == 0)
-						{
-							SkelTrack->SetEvalDisabled(true);
-						}
-						else
-						{
-							for (int32 RowIndex = 0; RowIndex <= SkelTrack->GetMaxRowIndex(); ++RowIndex)
-							{
-								SkelTrack->SetRowEvalDisabled(true, RowIndex);
-							}
-						}
+						MovieSceneHelpers::DisableTrack(SkelTrack);
 					}
 
 					//Finish Setup
diff --git a/Engine/Source/Editor/Sequencer/Private/SequencerUtilities.cpp b/Engine/Source/Editor/Sequencer/Private/SequencerUtilities.cpp
--- a/Engine/Source/Editor/Sequencer/Private/SequencerUtilities.cpp
+++ b/Engine/Source/Editor/Sequencer/Private/SequencerUtilities.cpp
@@ -28,6 +28,7 @@
 #include "Widgets/Input/SButton.h"
 #include "Styling/CoreStyle.h"
 #include "Styling/AppStyle.h"
+#include "MovieSceneCommonHelpers.h"
 #include "MovieSceneCopyableBinding.h"
 #include "MovieSceneCopyableTrack.h"
 #include "EntitySystem/MovieSceneBlenderSystem.h"
@@ -4682,11 +4683,7 @@ bool FSequencerUtilities::BakeTransform(TSharedRef<ISequencer> Sequencer, const
 
 				if (bDisableSectionsAfterBaking)
 				{
-					for (UMovieSceneSection* ConstraintSection : ConstraintTrack->GetAllSections())
-					{
-						ConstraintSection->Modify();
-						ConstraintSection->SetIsActive(false);
-					}
+					MovieSceneHelpers::DisableTrack(ConstraintTrack);
 				}
 				else
 				{
@@ -4702,11 +4699,7 @@ bool FSequencerUtilities::BakeTransform(TSharedRef<ISequencer> Sequencer, const
 			{
 				if (bDisableSectionsAfterBaking)
 				{
-					for (UMovieSceneSection* TransformSection : TransformTrack->GetAllSections())
-					{
-						TransformSection->Modify();
-						TransformSection->SetIsActive(false);
-					}
+					MovieSceneHelpers::DisableTrack(TransformTrack);
 				}
 				else
 				{
@@ -4722,11 +4715,7 @@ bool FSequencerUtilities::BakeTransform(TSharedRef<ISequencer> Sequencer, const
 			{
 				if (bDisableSectionsAfterBaking)
 				{
-					for (UMovieSceneSection* CameraShakeSection : CameraShakeTrack->GetAllSections())
-					{
-						CameraShakeSection->Modify();
-						CameraShakeSection->SetIsActive(false);
-					}
+					MovieSceneHelpers::DisableTrack(CameraShakeTrack);
 				}
 				else
 				{
diff --git a/Engine/Source/Runtime/MovieScene/Private/MovieSceneCommonHelpers.cpp b/Engine/Source/Runtime/MovieScene/Private/MovieSceneCommonHelpers.cpp
--- a/Engine/Source/Runtime/MovieScene/Private/MovieSceneCommonHelpers.cpp
+++ b/Engine/Source/Runtime/MovieScene/Private/MovieSceneCommonHelpers.cpp
@@ -551,6 +551,22 @@ void MovieSceneHelpers::GetDescendantMovieScenes(UMovieSceneSequence* InSequence
 	}
 }
 
+void MovieSceneHelpers::DisableTrack(UMovieSceneTrack* Track)
+{
+	// Disable the entire track or each of the rows
+	if (Track->GetMaxRowIndex() == 0)
+	{
+		Track->SetEvalDisabled(true);
+	}
+	else
+	{
+		for (int32 RowIndex = 0; RowIndex <= Track->GetMaxRowIndex(); ++RowIndex)
+		{
+			Track->SetRowEvalDisabled(true, RowIndex);
+		}
+	}
+}
+
 void MovieSceneHelpers::GetDescendantSubSections(const UMovieScene* InMovieScene, TArray<UMovieSceneSubSection*>& InSubSections)
 {
 	if (!IsValid(InMovieScene))
diff --git a/Engine/Source/Runtime/MovieScene/Public/MovieSceneCommonHelpers.h b/Engine/Source/Runtime/MovieScene/Public/MovieSceneCommonHelpers.h
--- a/Engine/Source/Runtime/MovieScene/Public/MovieSceneCommonHelpers.h
+++ b/Engine/Source/Runtime/MovieScene/Public/MovieSceneCommonHelpers.h
@@ -130,6 +130,11 @@ public:
 	 */
 	static MOVIESCENE_API void GetDescendantSubSections(const UMovieScene* InMovieScene, TArray<UMovieSceneSubSection*>& InSubSections);
 	
+	/*
+	 * 	Disable the entire track or each of the rows
+ 	 */
+	static MOVIESCENE_API void DisableTrack(UMovieSceneTrack* InTrack);
+
 	/**
 	 * Get the scene component from the runtime object
 	 *
