[Sequencer] Added a needed API export.

#rnx

#Synced-CL 50703497

diff --git a/Engine/Source/Editor/Sequencer/Public/MVVM/FolderModelStorageExtension.h b/Engine/Source/Editor/Sequencer/Public/MVVM/FolderModelStorageExtension.h
--- a/Engine/Source/Editor/Sequencer/Public/MVVM/FolderModelStorageExtension.h
+++ b/Engine/Source/Editor/Sequencer/Public/MVVM/FolderModelStorageExtension.h
@@ -7,6 +7,8 @@
 #include "MVVM/Extensions/DynamicExtensionContainer.h"
 #include "EventHandlers/ISequenceDataEventHandler.h"
 
+#define UE_API SEQUENCER_API
+
 class UMovieSceneFolder;
 
 namespace UE
@@ -17,7 +19,7 @@ namespace Sequencer
 class FFolderModel;
 class FSequenceModel;
 
-class FFolderModelStorageExtension
+class UE_API FFolderModelStorageExtension
 	: public IDynamicExtension
 	, private UE::MovieScene::TIntrusiveEventHandler<UE::MovieScene::ISequenceDataEventHandler>
 {
@@ -49,3 +51,4 @@ private:
 } // namespace Sequencer
 } // namespace UE
 
+#undef UE_API
