Sequencer: Initial pass of the linked filter UI in Sequencer.
- It is enabled using the "Sequencer.CurveEditor.UseLinkedFiltering" CVar and is disabled by default.
- Allows users to filter curve editor tree view using the same filters present in Sequencer
- Allows users to sync the filters or fitler on their own separate instance.
- For now, there is no view options in the UI (follow up change)
Several misc follow up changes to improve the UI are planned.

#Synced-CL 50904440

diff --git a/Engine/Source/Editor/Sequencer/Private/Filters/FilterConfigIdentifiers.cpp b/Engine/Source/Editor/Sequencer/Private/Filters/FilterConfigIdentifiers.cpp
new file mode 100644
--- /dev/null
+++ b/Engine/Source/Editor/Sequencer/Private/Filters/FilterConfigIdentifiers.cpp
@@ -0,0 +1,11 @@
+﻿// Copyright Epic Games, Inc. All Rights Reserved.
+
+#include "FilterConfigIdentifiers.h"
+
+namespace UE::Sequencer::ConfigIds
+{
+const TCHAR* LinkedConfigSubId = TEXT(""); // TODO UE-362151: For now empty so we load using the old config entry, see below.
+
+const TCHAR* Sequencer_InstancedConfigSubId = TEXT("Outliner");
+const TCHAR* CurveEditor_InstancedConfigId = TEXT("CurveEditor");
+} // namespace UE::Sequencer
\ No newline at end of file
diff --git a/Engine/Source/Editor/Sequencer/Private/Filters/FilterConfigIdentifiers.h b/Engine/Source/Editor/Sequencer/Private/Filters/FilterConfigIdentifiers.h
new file mode 100644
--- /dev/null
+++ b/Engine/Source/Editor/Sequencer/Private/Filters/FilterConfigIdentifiers.h
@@ -0,0 +1,29 @@
+﻿// Copyright Epic Games, Inc. All Rights Reserved.
+
+#pragma once
+
+#include "CoreTypes.h"
+
+/** 
+ * Built-in IDs for sub-config identifiers in FSequencerFilterBar
+ * They affect the result of FSequencerFilterBar::GetIdentifier, which is important for the identifier in the filter state config.
+ */
+namespace UE::Sequencer::ConfigIds
+{
+/** 
+ * Config identifier for the linked filter bar in Sequencer. 
+ * @see FSequencerFilterBar::GetIdentifier.
+ */
+extern const TCHAR* LinkedConfigSubId; // TODO UE-362151: For now empty so we load using the old config entry, see below.
+
+/**
+ * Config identifier for the instanced filter bar in Sequencer's outliner. 
+ * @see FSequencerFilterBar::GetIdentifier.
+ */
+extern const TCHAR* Sequencer_InstancedConfigSubId; 
+/** 
+ * Config identifier for the instanced filter bar in Sequencer's Curve Editor. 
+ * @see FSequencerFilterBar::GetIdentifier.
+ */
+extern const TCHAR* CurveEditor_InstancedConfigId;
+} // namespace UE::Sequencer
diff --git a/Engine/Source/Editor/Sequencer/Private/Filters/ViewModels/Linking/ILinkedFilterViewModel.cpp b/Engine/Source/Editor/Sequencer/Private/Filters/ViewModels/Linking/ILinkedFilterViewModel.cpp
new file mode 100644
--- /dev/null
+++ b/Engine/Source/Editor/Sequencer/Private/Filters/ViewModels/Linking/ILinkedFilterViewModel.cpp
@@ -0,0 +1,14 @@
+﻿// Copyright Epic Games, Inc. All Rights Reserved.
+
+#include "ILinkedFilterViewModel.h"
+
+#include "Filters/ISequencerTrackFilters.h"
+#include "Filters/SequencerFilterData.h"
+
+namespace UE::Sequencer
+{
+bool ILinkedFilterViewModel::IsFilteredOut(const TViewModelPtr<IOutlinerExtension>& InItem) const
+{
+	return GetActiveFilterBar()->GetFilterData().IsFilteredOut(InItem);
+}
+} // namespace UE::Sequencer
diff --git a/Engine/Source/Editor/Sequencer/Private/Filters/ViewModels/Linking/ILinkedFilterViewModel.h b/Engine/Source/Editor/Sequencer/Private/Filters/ViewModels/Linking/ILinkedFilterViewModel.h
--- a/Engine/Source/Editor/Sequencer/Private/Filters/ViewModels/Linking/ILinkedFilterViewModel.h
+++ b/Engine/Source/Editor/Sequencer/Private/Filters/ViewModels/Linking/ILinkedFilterViewModel.h
@@ -2,6 +2,7 @@
 
 #pragma once
 
+#include "Delegates/Delegate.h"
 #include "HAL/Platform.h"
 #include "MVVM/ViewModelPtr.h"
 
@@ -36,6 +37,9 @@ public:
 	/** @return The current filtering mode. */
 	virtual ELinkedFilterMode GetFilterMode() const = 0;
 	
+	/** @return Whether InItem is filtered out by the active filter bar. */
+	bool IsFilteredOut(const TViewModelPtr<IOutlinerExtension>& InItem) const;
+	
 	/** @return The currently active filter bar instance according to the linked mode. */
 	virtual TSharedRef<ISequencerTrackFilters> GetActiveFilterBar() const = 0; 
 	/** @return The filter bar that is shared with all other UI. */
@@ -46,6 +50,9 @@ public:
 	/** @return Event when the filter mode changes, i.e. the instance returned by GetActiveFilterBar changes. */
 	virtual FSimpleMulticastDelegate& OnFilterModeChanged() = 0;
 	
+	/** @return Event when the items displayed need to be refreshed, e.g. when ELinkedFilterMode changes, or the active filter bar refilters its items. */
+	virtual FSimpleMulticastDelegate& OnFilteredItemsChanged() = 0;
+	
 	virtual ~ILinkedFilterViewModel() = default;
 };
-} // namespace UE::Sequencer
+} // namespace UE::Sequencer
\ No newline at end of file
diff --git a/Engine/Source/Editor/Sequencer/Private/Filters/ViewModels/Linking/LinkedFilterViewModel.cpp b/Engine/Source/Editor/Sequencer/Private/Filters/ViewModels/Linking/LinkedFilterViewModel.cpp
--- a/Engine/Source/Editor/Sequencer/Private/Filters/ViewModels/Linking/LinkedFilterViewModel.cpp
+++ b/Engine/Source/Editor/Sequencer/Private/Filters/ViewModels/Linking/LinkedFilterViewModel.cpp
@@ -16,7 +16,14 @@ FLinkedFilterViewModel::FLinkedFilterViewModel(
 	, InstancedFilterBar(InInstancedFilterBar)
 	, FilterEvaluator(InFilterEvaluator)
 	, ActiveFilterBarUsage(GetActiveFilterBarImpl(), FilterEvaluator)
-{}
+{
+	FilterEvaluator->OnFilterBarFiltered().AddRaw(this, &FLinkedFilterViewModel::OnFilterBarFiltered);
+}
+
+FLinkedFilterViewModel::~FLinkedFilterViewModel()
+{
+	FilterEvaluator->OnFilterBarFiltered().RemoveAll(this);
+}
 
 TSharedRef<FSequencerFilterBar> FLinkedFilterViewModel::GetActiveFilterBarImpl() const
 {
@@ -48,6 +55,7 @@ void FLinkedFilterViewModel::SetFilterMode(ELinkedFilterMode InNewMode)
 		const TSharedRef<FSequencerFilterBar> NewActiveFilterBar = GetActiveFilterBarImpl();
 		ActiveFilterBarUsage = FScopedFilterBarUsage(NewActiveFilterBar, FilterEvaluator);
 		OnFilterModeChangedDelegate.Broadcast();
+		OnFilteredItemsChangedDelegate.Broadcast();
 	}
 }
 
@@ -65,4 +73,12 @@ TSharedRef<ISequencerTrackFilters> FLinkedFilterViewModel::GetInstancedFilterBar
 {
 	return InstancedFilterBar;
 }
+
+void FLinkedFilterViewModel::OnFilterBarFiltered(const TSharedRef<FSequencerFilterBar>& InFilterBar)
+{
+	if (InFilterBar == GetActiveFilterBarImpl())
+	{
+		OnFilteredItemsChangedDelegate.Broadcast();
+	}
+}
 } // namespace UE::Sequencer
diff --git a/Engine/Source/Editor/Sequencer/Private/Filters/ViewModels/Linking/LinkedFilterViewModel.h b/Engine/Source/Editor/Sequencer/Private/Filters/ViewModels/Linking/LinkedFilterViewModel.h
--- a/Engine/Source/Editor/Sequencer/Private/Filters/ViewModels/Linking/LinkedFilterViewModel.h
+++ b/Engine/Source/Editor/Sequencer/Private/Filters/ViewModels/Linking/LinkedFilterViewModel.h
@@ -22,6 +22,7 @@ public:
 		const TSharedRef<FSequencerFilterBar>& InInstancedFilterBar,
 		const TSharedRef<FFilterEvaluator> InFilterEvaluator
 		);
+	~FLinkedFilterViewModel();
 	
 	TSharedRef<FSequencerFilterBar> GetActiveFilterBarImpl() const;
 	TSharedRef<FSequencerFilterBar> GetLinkedFilterBarImpl() const;
@@ -34,6 +35,7 @@ public:
 	virtual TSharedRef<ISequencerTrackFilters> GetLinkedFilterBar() const override;
 	virtual TSharedRef<ISequencerTrackFilters> GetInstancedFilterBar() const override;
 	virtual FSimpleMulticastDelegate& OnFilterModeChanged() override { return OnFilterModeChangedDelegate; }
+	virtual FSimpleMulticastDelegate& OnFilteredItemsChanged() override { return OnFilteredItemsChangedDelegate; }
 	//~ End ILinkedFilterViewModel
 	
 private:
@@ -53,5 +55,9 @@ private:
 	
 	/** Event when the filter mode changes, i.e. the instance returned by GetActiveFilterBar changes. */
 	FSimpleMulticastDelegate OnFilterModeChangedDelegate;
+	/** Event when the items displayed need to be refreshed, e.g. when ELinkedFilterMode changes, or the active filter bar refilters its items. */
+	FSimpleMulticastDelegate OnFilteredItemsChangedDelegate;
+	
+	void OnFilterBarFiltered(const TSharedRef<FSequencerFilterBar>& InFilterBar);
 };
 } // namespace UE::Sequencer
diff --git a/Engine/Source/Editor/Sequencer/Private/MVVM/CurveEditorExtension.cpp b/Engine/Source/Editor/Sequencer/Private/MVVM/CurveEditorExtension.cpp
--- a/Engine/Source/Editor/Sequencer/Private/MVVM/CurveEditorExtension.cpp
+++ b/Engine/Source/Editor/Sequencer/Private/MVVM/CurveEditorExtension.cpp
@@ -4,7 +4,12 @@
 #include "MVVM/Selection/Selection.h"
 
 #include "FrameNumberDetailsCustomization.h"
+#include "Filters/FilterConfigIdentifiers.h"
 #include "Filters/SCurveEditorFilterPanel.h"
+#include "Filters/ViewModels/Linking/ILinkedFilterViewModel.h"
+#include "Filters/ViewModels/Linking/LinkedFilterFactoryViewModel.h"
+#include "Filters/ViewModels/Linking/LinkedFilterViewModel.h"
+#include "Filters/ViewModels/Linking/MakeLinkedFilterViewModelArgs.h"
 #include "Framework/Docking/TabManager.h"
 #include "IPropertyRowGenerator.h"
 #include "IStructureDetailsView.h"
@@ -38,6 +43,7 @@
 #include "Widgets/Layout/SScrollBorder.h"
 #include "WorkspaceMenuStructure.h"
 #include "WorkspaceMenuStructureModule.h"
+#include "ViewModels/SequencerEditorViewModelImpl.h"
 #include "Views/CurveEditor/CurveEditorWidgetOwner.h"
 
 #define LOCTEXT_NAMESPACE "SequencerCurveEditorExtension"
@@ -399,7 +405,7 @@ void FCurveEditorExtension::CreateCurveEditor(const FTimeSliderArgs& TimeSliderA
 
 		CurveEditorModel->DefaultKeyAttributes = MakeAttributeLambda([this]() { return GetDefaultKeyAttributes(); });
 	}
-
+	
 	// We create a custom Time Slider Controller which is just a wrapper around the actual one, but is 
 	// aware of our custom bounds logic. Currently the range the bar displays is tied to Sequencer 
 	// timeline and not the Bounds, so we need a way of changing it to look at the Bounds but only for 
@@ -407,9 +413,18 @@ void FCurveEditorExtension::CreateCurveEditor(const FTimeSliderArgs& TimeSliderA
 	const TSharedRef<ITimeSliderController> CurveEditorTimeSliderController = MakeShared<FSequencerCurveEditorTimeSliderController>(
 		TimeSliderArgs, Sequencer, CurveEditorModel.ToSharedRef());
 	
+	const TSharedPtr<FLinkedFilterFactoryViewModel> FilterFactory = Sequencer->GetViewModelImpl()->GetLinkedFilterViewModelFactoryImpl();
+	const TSharedRef<FLinkedFilterViewModel> ViewModelImpl = FilterFactory->MakeFilteringModelImpl(
+		FMakeLinkedFilterViewModelArgs(ConfigIds::CurveEditor_InstancedConfigId)
+		);
+	CurveEditorFilterClient = ViewModelImpl;
 	WidgetOwner = MakePimpl<FCurveEditorWidgetOwner>(
-		Sequencer.ToSharedRef(), CurveEditorModel.ToSharedRef(), CurveEditorTimeSliderController
+		Sequencer.ToSharedRef(), CurveEditorModel.ToSharedRef(), CurveEditorTimeSliderController, ViewModelImpl
 		);
+	
+	// TODO UE-364009: Effectively, this makes FCurveEditorIntegrationExtension bind to FilterBarProxy. In the future, refactor 
+	// FCurveEditorIntegrationExtension (non-extension) to be a separate class that FCurveEditorExtension injects the dependencies to.
+	OnCurveEditorCreated.Broadcast();
 }
 
 void FCurveEditorExtension::OpenCurveEditor()
diff --git a/Engine/Source/Editor/Sequencer/Private/MVVM/CurveEditorIntegrationExtension.cpp b/Engine/Source/Editor/Sequencer/Private/MVVM/CurveEditorIntegrationExtension.cpp
--- a/Engine/Source/Editor/Sequencer/Private/MVVM/CurveEditorIntegrationExtension.cpp
+++ b/Engine/Source/Editor/Sequencer/Private/MVVM/CurveEditorIntegrationExtension.cpp
@@ -1,18 +1,20 @@
 // Copyright Epic Games, Inc. All Rights Reserved.
 
 #include "MVVM/CurveEditorIntegrationExtension.h"
+
+#include "CurveEditor.h"
+#include "Filters/ViewModels/Linking/LinkedFilterViewModel.h"
 #include "MVVM/CurveEditorExtension.h"
-#include "MVVM/TrackModelStorageExtension.h"
-#include "MVVM/ViewModels/ViewModel.h"
+#include "MVVM/Extensions/ICurveEditorTreeItemExtension.h"
 #include "MVVM/SharedViewModelData.h"
+#include "MVVM/TrackModelStorageExtension.h"
 #include "MVVM/ViewModelPtr.h"
 #include "MVVM/ViewModels/SequenceModel.h"
 #include "MVVM/ViewModels/SequencerEditorViewModel.h"
+#include "MVVM/ViewModels/ViewModel.h"
 #include "MVVM/ViewModels/ViewModelIterators.h"
-#include "MVVM/Extensions/ICurveEditorTreeItemExtension.h"
-#include "CurveEditor.h"
-
 #include "SequencerSettings.h"
+#include "Views/CurveEditor/CurveEditorWidgetOwner.h"
 
 namespace UE
 {
@@ -23,6 +25,11 @@ FCurveEditorIntegrationExtension::FCurveEditorIntegrationExtension()
 {
 }
 
+FCurveEditorIntegrationExtension::~FCurveEditorIntegrationExtension()
+{
+	UnbindFilterBarDependencies();
+}
+
 void FCurveEditorIntegrationExtension::OnCreated(TSharedRef<FViewModel> InWeakOwner)
 {
 	ensureMsgf(!WeakOwnerModel.Pin().IsValid(), TEXT("This extension was already created!"));
@@ -30,6 +37,8 @@ void FCurveEditorIntegrationExtension::OnCreated(TSharedRef<FViewModel> InWeakOw
 
 	FSimpleMulticastDelegate& HierarchyChanged = InWeakOwner->GetSharedData()->SubscribeToHierarchyChanged(InWeakOwner);
 	HierarchyChanged.AddSP(this, &FCurveEditorIntegrationExtension::OnHierarchyChanged);
+	
+	SetupFilterBarDependencies();
 }
 
 void FCurveEditorIntegrationExtension::OnHierarchyChanged()
@@ -68,8 +77,15 @@ void FCurveEditorIntegrationExtension::UpdateCurveEditor()
 	{
 		return;
 	}
-
+	
+	// TODO UE-364008: In the future, we can remove the paths with bUseFilteredOut.
 	const bool bUseFilteredOut = OwnerModel->GetSequencer()->GetSequencerSettings()->GetLinkFiltersWithCurveEditor();
+	const auto IsFilteredOut = [CurveEditorExtension, bUseFilteredOut](const TViewModelPtr<IOutlinerExtension>& InItem)
+	{
+		return UseLinkedFilteringUI()
+			? CurveEditorExtension->GetLinkedFilteringModel()->IsFilteredOut(InItem)
+			: bUseFilteredOut && InItem && InItem->IsFilteredOut();
+	};
 
 	FCurveEditorTree* CurveEditorTree = CurveEditor->GetTree();
 
@@ -85,7 +101,7 @@ void FCurveEditorIntegrationExtension::UpdateCurveEditor()
 		if (ViewModel)
 		{
 			TViewModelPtr<IOutlinerExtension> OutlinerModel = ViewModel.ImplicitCast();
-			const bool bIsVisible = bUseFilteredOut ? OutlinerModel && !OutlinerModel->IsFilteredOut() : true;
+			const bool bIsVisible = !IsFilteredOut(OutlinerModel);
 			
 			const FCurveEditorTreeItem* TreeItem = CurveEditorTree->FindItem(ItemID);
 
@@ -114,6 +130,7 @@ void FCurveEditorIntegrationExtension::UpdateCurveEditor()
 	{
 		if (CurveEditorTree->FindItem(It->Value) == nullptr)
 		{
+			TreeItemIDToViewModelMap.Remove(It->Value);
 			It.RemoveCurrent();
 		}
 	}
@@ -123,7 +140,7 @@ void FCurveEditorIntegrationExtension::UpdateCurveEditor()
 	bool bItemsAdded = false;
 	for (TParentFirstChildIterator<IOutlinerExtension> It(OwnerModel, bIncludeRootNode); It; ++It)
 	{
-		if (bUseFilteredOut && It->IsFilteredOut())
+		if (IsFilteredOut(*It))
 		{
 			It.IgnoreCurrentChildren();
 			continue;
@@ -180,8 +197,11 @@ FCurveEditorTreeItemID FCurveEditorIntegrationExtension::AddToCurveEditor(TViewM
 	NewItem->SetWeakItem(CurveEditorTreeItem);
 	TOptional<FString> UniquePathName =  InViewModel->GetUniquePathName();
 	NewItem->SetUniquePathName(UniquePathName);
+	
 	// Register the new ID in our map and notify the view model
 	ViewModelToTreeItemIDMap.Add(InViewModel.AsModel(), NewItem->GetID());
+	TreeItemIDToViewModelMap.Add(NewItem->GetID(), InViewModel.AsModel());
+	
 	InViewModel->OnAddedToCurveEditor(NewItem->GetID(), InCurveEditor);
 
 	return NewItem->GetID();
@@ -213,8 +233,64 @@ void FCurveEditorIntegrationExtension::ResetCurveEditor()
 	}
 
 	ViewModelToTreeItemIDMap.Reset();
+	TreeItemIDToViewModelMap.Reset();
 }
 
+TWeakPtr<FViewModel> FCurveEditorIntegrationExtension::FindViewModelById(const FCurveEditorTreeItemID& InTreeId) const
+{
+	const TWeakPtr<FViewModel>* ViewModel = TreeItemIDToViewModelMap.Find(InTreeId);
+	return ViewModel ? *ViewModel : nullptr;
+}
+
+FCurveEditorTreeItemID FCurveEditorIntegrationExtension::FindIdByViewModel(TWeakPtr<FViewModel> InViewModel) const
+{
+	const FCurveEditorTreeItemID* TreeItemID = ViewModelToTreeItemIDMap.Find(InViewModel);
+	return TreeItemID ? *TreeItemID : FCurveEditorTreeItemID::Invalid();
+}
+
+void FCurveEditorIntegrationExtension::SetupFilterBarDependencies()
+{
+	// TODO UE-364009: Effectively, this makes FCurveEditorIntegrationExtension bind to FilterBarProxy. In the future, refactor 
+	// FCurveEditorIntegrationExtension to be a separate class that FCurveEditorExtension injects the dependencies to.
+	FCurveEditorExtension* CurveEditorExtension = GetCurveEditorExtension();
+	if (!ensure(CurveEditorExtension))
+	{
+		return;
+	}
+	
+	if (CurveEditorExtension->GetLinkedFilteringModel())
+	{
+		BindFilterBarDependencies();
+	}
+	else
+	{
+		CurveEditorExtension->OnCurveEditorCreated.AddRaw(this, &FCurveEditorIntegrationExtension::BindFilterBarDependencies);
+	}
+}
+
+void FCurveEditorIntegrationExtension::BindFilterBarDependencies()
+{
+	if (FCurveEditorExtension* CurveEditorExtension = GetCurveEditorExtension())
+	{
+		CurveEditorExtension->GetLinkedFilteringModel()->OnFilteredItemsChanged().AddRaw(this, &FCurveEditorIntegrationExtension::UpdateCurveEditor);
+		UpdateCurveEditor();
+	}
+}
+
+void FCurveEditorIntegrationExtension::UnbindFilterBarDependencies()
+{
+	FCurveEditorExtension* CurveEditorExtension = GetCurveEditorExtension();
+	if (!CurveEditorExtension)
+	{
+		return;
+	}
+	
+	CurveEditorExtension->OnCurveEditorCreated.RemoveAll(this);
+	if (ILinkedFilterViewModel* ViewModel = CurveEditorExtension->GetLinkedFilteringModel())
+	{
+		ViewModel->OnFilteredItemsChanged().RemoveAll(this);
+	}
+}
 } // namespace Sequencer
 } // namespace UE
 
diff --git a/Engine/Source/Editor/Sequencer/Private/MVVM/ViewModels/SequencerEditorViewModelImpl.cpp b/Engine/Source/Editor/Sequencer/Private/MVVM/ViewModels/SequencerEditorViewModelImpl.cpp
--- a/Engine/Source/Editor/Sequencer/Private/MVVM/ViewModels/SequencerEditorViewModelImpl.cpp
+++ b/Engine/Source/Editor/Sequencer/Private/MVVM/ViewModels/SequencerEditorViewModelImpl.cpp
@@ -7,14 +7,10 @@
 #include "Filters/ViewModels/Linking/MakeLinkedFilterViewModelArgs.h"
 #include "LinkedFilteringMediator.h"
 #include "Sequencer.h"
+#include "Filters/FilterConfigIdentifiers.h"
 
 namespace UE::Sequencer
 {
-/** Sub-identifier used for saving filter states to config for the linked FSequencerFilterBar. */
-const TCHAR* LinkedConfigSubId = TEXT(""); // TODO UE-362151: For now empty so we load using the old config entry, see below.
-/** Sub-identifier used for saving filter states to config for the instanced FSequencerFilterBar displayed in SSequencer. */
-const TCHAR* InstancedConfigSubId = TEXT("MainInstance"); 
-
 FSequencerEditorViewModelImpl::FSequencerEditorViewModelImpl(
 	TSharedRef<ISequencer> InSequencer, const FSequencerHostCapabilities& InHostCapabilities
 	)
@@ -36,7 +32,7 @@ void FSequencerEditorViewModelImpl::InitLinkedFiltering()
 	// Going forward, we can separate linked and instanced, e.g. ""LevelSequencerEditor.Linked", "LevelSequencerEditor.MainInstance", respectively.
 	// For now, filtering will default to linked mode. For linked mode, we specify no sub-identifier: that causes FSequencerFilterBar to continue 
 	// using the old config entry, which would be SequencerSettings->GetName(), e.g. "LevelSequencerEditor".
-	const TCHAR* ConfigId = LinkedConfigSubId;
+	const TCHAR* ConfigId = ConfigIds::LinkedConfigSubId;
 	LinkedFilterFactory = MakeShared<FLinkedFilterFactoryViewModel>(
 		*Sequencer,
 		MakeShared<FSequencerFilterBar>(*Sequencer, ConfigId),
@@ -44,7 +40,7 @@ void FSequencerEditorViewModelImpl::InitLinkedFiltering()
 		);
 	
 	SequenceFilterInstanceViewModel = LinkedFilterFactory->MakeFilteringModelImpl(
-		FMakeLinkedFilterViewModelArgs(InstancedConfigSubId)
+		FMakeLinkedFilterViewModelArgs(ConfigIds::Sequencer_InstancedConfigSubId)
 		);
 	FilterMediator = MakeShared<FLinkedFilteringMediator>(SequenceFilterInstanceViewModel.ToSharedRef());
 }
diff --git a/Engine/Source/Editor/Sequencer/Private/MVVM/ViewModels/SequencerEditorViewModelImpl.h b/Engine/Source/Editor/Sequencer/Private/MVVM/ViewModels/SequencerEditorViewModelImpl.h
--- a/Engine/Source/Editor/Sequencer/Private/MVVM/ViewModels/SequencerEditorViewModelImpl.h
+++ b/Engine/Source/Editor/Sequencer/Private/MVVM/ViewModels/SequencerEditorViewModelImpl.h
@@ -18,8 +18,8 @@ public:
 	
 	explicit FSequencerEditorViewModelImpl(TSharedRef<ISequencer> InSequencer, const FSequencerHostCapabilities& InHostCapabilities);
 	
-	/** Private version of GetLinkedFilterViewModel */
 	TSharedPtr<FLinkedFilterViewModel> GetLinkedFilterViewModelImpl() const { return SequenceFilterInstanceViewModel; }
+	TSharedPtr<FLinkedFilterFactoryViewModel> GetLinkedFilterViewModelFactoryImpl() const { return LinkedFilterFactory; }
 	
 	/** @return The object responsible for syncing the linked filters to Sequencer's systems. */
 	TSharedPtr<FLinkedFilteringMediator> GetFilterMediator() const { return FilterMediator; }
diff --git a/Engine/Source/Editor/Sequencer/Private/MVVM/Views/CurveEditor/CurveEditorWidgetOwner.cpp b/Engine/Source/Editor/Sequencer/Private/MVVM/Views/CurveEditor/CurveEditorWidgetOwner.cpp
--- a/Engine/Source/Editor/Sequencer/Private/MVVM/Views/CurveEditor/CurveEditorWidgetOwner.cpp
+++ b/Engine/Source/Editor/Sequencer/Private/MVVM/Views/CurveEditor/CurveEditorWidgetOwner.cpp
@@ -2,18 +2,22 @@
 
 #include "CurveEditorWidgetOwner.h"
 
+#include "Filters/SCurveEditorFilterPanel.h"
+#include "Filters/ViewModels/Linking/LinkedFilterViewModel.h"
+#include "Filters/Widgets/Linking/SFilterBarSwitcher.h"
+#include "HAL/IConsoleManager.h"
 #include "IPropertyRowGenerator.h"
 #include "IStructureDetailsView.h"
+#include "MVVM/CurveEditorExtension.h"
 #include "SCurveEditorPanel.h"
 #include "SCurveEditorToolProperties.h"
 #include "SCurveKeyDetailPanel.h"
+#include "SCurveTreeContent.h"
+#include "SOldCurveTreeContent.h"
 #include "Sequencer.h"
 #include "SequencerCommands.h"
-#include "SOldCurveTreeContent.h"
-#include "Filters/SCurveEditorFilterPanel.h"
-#include "MVVM/CurveEditorExtension.h"
-#include "Widgets/DeclarativeSyntaxSupport.h"
 #include "Widgets/CurveEditor/SSequencerCurveEditor.h"
+#include "Widgets/DeclarativeSyntaxSupport.h"
 #include "Widgets/Docking/SDockTab.h"
 
 struct FTimeSliderArgs;
@@ -24,16 +28,37 @@ namespace UE::Sequencer
 {
 namespace PresenterDetail
 {
+static TAutoConsoleVariable<bool> CVarUseLinkedFilterUI(
+	TEXT("Sequencer.CurveEditor.UseLinkedFiltering"), false, TEXT("Whether to use the new linked filtering UI in Curve Editor")
+	);
+
 static FCurveEditorWidgets MakePresenterWidgets(
 	const TSharedRef<FSequencer>& InSequencer, 
 	const TSharedRef<FCurveEditor>& InCurveEditor, 
 	const TSharedRef<ITimeSliderController>& InTimeSliderController,
+	const TSharedRef<FLinkedFilterViewModel>& InFilteringViewModel,
 	const TSharedPtr<FTabManager> InTabManager
 	)
 {
 	const TSharedRef<SOldCurveTreeContent> OldContent = SNew(SOldCurveTreeContent, InSequencer, InCurveEditor)
 		.TabManager(InTabManager);
-	// TODO UE-363527: Create a new widget here and then hook it up using SWidgetSwitcher driven by the new CVar.
+	const TSharedRef<SCurveTreeContent> TreeContent = SNew(SCurveTreeContent)
+		.TabManager(InTabManager)
+		.Sequencer(InSequencer)
+		.CurveEditor(InCurveEditor)
+		.FilterViewModel(InFilteringViewModel)
+		.FilterPills()
+		[
+			SNew(SFilterBarSwitcher)
+			.LinkedFilterViewModel(InFilteringViewModel)
+			.FilterBarLayout(EFilterBarLayout::Horizontal)
+		];
+	
+	const TSharedRef<SWidgetSwitcher> Content = 
+		SNew(SWidgetSwitcher)
+		.WidgetIndex_Lambda([]{ return UseLinkedFilteringUI() ? 1 : 0; })
+		+SWidgetSwitcher::Slot()[ OldContent ]
+		+SWidgetSwitcher::Slot()[ TreeContent ];
 	
 	const TSharedRef<SCurveEditorPanel> CurveEditorPanel = SNew(SCurveEditorPanel, InCurveEditor)
 		// Grid lines match the color specified in FSequencerTimeSliderController::OnPaintViewArea
@@ -44,22 +69,28 @@ static FCurveEditorWidgets MakePresenterWidgets(
 		.DisabledTimeSnapTooltip(LOCTEXT("CurveEditorTimeSnapDisabledTooltip", "Time Snapping is currently driven by Sequencer."))
 		.TreeContent()
 		[
-			OldContent
+			Content
 		];
 	
 	const TSharedRef<SSequencerCurveEditor> RootWidget = SNew(SSequencerCurveEditor, CurveEditorPanel, InSequencer);
-	return FCurveEditorWidgets(RootWidget, CurveEditorPanel, OldContent);
+	return FCurveEditorWidgets(RootWidget, CurveEditorPanel, OldContent, TreeContent);
 }
 }
 
+bool UseLinkedFilteringUI()
+{
+	return PresenterDetail::CVarUseLinkedFilterUI.GetValueOnAnyThread();
+}
+
 FCurveEditorWidgetOwner::FCurveEditorWidgetOwner(
 	const TSharedRef<FSequencer>& InSequencer,
 	const TSharedRef<FCurveEditor>& InCurveEditor,
-	const TSharedRef<ITimeSliderController>& InTimeSliderController
+	const TSharedRef<ITimeSliderController>& InTimeSliderController,
+	const TSharedRef<FLinkedFilterViewModel>& InFilteringViewModel
 	)
 	: WeakSequencer(InSequencer)
 	, CurveEditor(InCurveEditor)
-	, Widgets(PresenterDetail::MakePresenterWidgets(InSequencer, InCurveEditor, InTimeSliderController, GetTabManager()))
+	, Widgets(PresenterDetail::MakePresenterWidgets(InSequencer, InCurveEditor, InTimeSliderController, InFilteringViewModel, GetTabManager()))
 {
 	BindCommands();
 	InitFrameNumberPropertyLayout(InSequencer);
@@ -121,14 +152,21 @@ void FCurveEditorWidgetOwner::CloseCurveEditor()
 
 void FCurveEditorWidgetOwner::SyncSelection() const
 {
-	// TODO UE-363527: Depending on the new CVar, reroute the selection sync to the correct widget.
-	Widgets.OldTreeContent->SyncSelection();
+	if (UseLinkedFilteringUI())
+	{
+		Widgets.TreeContent->SyncSelection();
+	}
+	else
+	{
+		Widgets.OldTreeContent->SyncSelection();
+	}
 }
 
 TSharedPtr<SCurveEditorTree> FCurveEditorWidgetOwner::GetTreeView() const
 {
-	// TODO UE-363527: Depending on the new CVar, reroute to the correct widget.
-	return Widgets.OldTreeContent->GetCurveEditorTreeView();
+	return UseLinkedFilteringUI() 
+		? Widgets.TreeContent->GetCurveEditorTreeView()
+		: Widgets.OldTreeContent->GetCurveEditorTreeView();
 }
 
 USequencerSettings* FCurveEditorWidgetOwner::GetSequencerSettings() const
@@ -152,17 +190,40 @@ void FCurveEditorWidgetOwner::BindCommands()
 		return;
 	}
 	
-	// TODO UE-363527: Depending on the new CVar, reroute the commands to the correct widget.
 	CurveEditorCommandList->MapAction(
 		FSequencerCommands::Get().QuickTreeSearch,
-		FExecuteAction::CreateLambda([this] { Widgets.OldTreeContent->FocusSearchBox(); })
+		FExecuteAction::CreateRaw(this, &FCurveEditorWidgetOwner::HandleQuickTreeSearch)
 	);
 	CurveEditorCommandList->MapAction(
 		FSequencerCommands::Get().ToggleShowGotoBox,
-		FExecuteAction::CreateLambda([this] { Widgets.OldTreeContent->FocusPlayTimeDisplay(); })
+		FExecuteAction::CreateRaw(this, &FCurveEditorWidgetOwner::HandleToggleShowGotoBox)
 	);
 }
 
+void FCurveEditorWidgetOwner::HandleQuickTreeSearch()
+{
+	if (UseLinkedFilteringUI())
+	{
+		Widgets.TreeContent->FocusSearchBox();
+	}
+	else
+	{
+		Widgets.OldTreeContent->FocusSearchBox();
+	}
+}
+
+void FCurveEditorWidgetOwner::HandleToggleShowGotoBox()
+{
+	if (UseLinkedFilteringUI())
+	{
+		Widgets.TreeContent->FocusPlayTimeDisplay();
+	}
+	else
+	{
+		Widgets.OldTreeContent->FocusPlayTimeDisplay();
+	}
+}
+
 void FCurveEditorWidgetOwner::InitFrameNumberPropertyLayout(const TSharedRef<FSequencer>& InSequencer)
 {
 	const TSharedRef<SCurveEditorPanel>& CurveEditorPanel = Widgets.CurveEditorPanel;
diff --git a/Engine/Source/Editor/Sequencer/Private/MVVM/Views/CurveEditor/CurveEditorWidgetOwner.h b/Engine/Source/Editor/Sequencer/Private/MVVM/Views/CurveEditor/CurveEditorWidgetOwner.h
--- a/Engine/Source/Editor/Sequencer/Private/MVVM/Views/CurveEditor/CurveEditorWidgetOwner.h
+++ b/Engine/Source/Editor/Sequencer/Private/MVVM/Views/CurveEditor/CurveEditorWidgetOwner.h
@@ -17,6 +17,11 @@ struct FTimeSliderArgs;
 
 namespace UE::Sequencer
 {
+class FLinkedFilterViewModel;
+
+/** @return Whether the linked filtering UI is enabled. */
+bool UseLinkedFilteringUI();
+
 /**
  * Instantiates the UI displaying Sequencer's FCurveEditor.
  * 
@@ -32,7 +37,8 @@ public:
 	explicit FCurveEditorWidgetOwner(
 		const TSharedRef<FSequencer>& InSequencer, 
 		const TSharedRef<FCurveEditor>& InCurveEditor,
-		const TSharedRef<ITimeSliderController>& InTimeSliderController
+		const TSharedRef<ITimeSliderController>& InTimeSliderController,
+		const TSharedRef<FLinkedFilterViewModel>& InFilteringViewModel
 		);
 	~FCurveEditorWidgetOwner();
 	
@@ -65,8 +71,10 @@ private:
 	
 	USequencerSettings* GetSequencerSettings() const;
 	TSharedPtr<FTabManager> GetTabManager() const;
-	
+
 	void BindCommands();
+	void HandleQuickTreeSearch();
+	void HandleToggleShowGotoBox();
 	
 	/** Register an instanced custom property type layout to handle converting FFrameNumber from Tick Resolution to Display Rate. */
 	void InitFrameNumberPropertyLayout(const TSharedRef<FSequencer>& InSequencer);
diff --git a/Engine/Source/Editor/Sequencer/Private/MVVM/Views/CurveEditor/CurveEditorWidgets.h b/Engine/Source/Editor/Sequencer/Private/MVVM/Views/CurveEditor/CurveEditorWidgets.h
--- a/Engine/Source/Editor/Sequencer/Private/MVVM/Views/CurveEditor/CurveEditorWidgets.h
+++ b/Engine/Source/Editor/Sequencer/Private/MVVM/Views/CurveEditor/CurveEditorWidgets.h
@@ -9,6 +9,7 @@ class SSequencerCurveEditor;
 
 namespace UE::Sequencer
 {
+class SCurveTreeContent;
 class SOldCurveTreeContent;
 
 /** @see FSequencerCurveEditorPresenter */
@@ -19,20 +20,22 @@ struct FCurveEditorWidgets
 	/** The curve editor panel. This is created and updated even if it is not currently visible. */
 	TSharedRef<SCurveEditorPanel> CurveEditorPanel;
 	
-	/** 
-	 * The tree content displayed inside CurveEditorPanel.
-	 * TODO UE-363527: A CVar will be added to switch between this in the new linked filtering UI.
-	 */
+	/** The tree content displayed inside CurveEditorPanel.*/
 	TSharedRef<SOldCurveTreeContent> OldTreeContent;
+	
+	/** The tree content for curve editor. It displays linked filtering UI. */
+	TSharedRef<SCurveTreeContent> TreeContent;
 
 	explicit FCurveEditorWidgets(
-		const TSharedRef<SSequencerCurveEditor>& RootWidget, 
-		const TSharedRef<SCurveEditorPanel>& CurveEditorPanel,
-		const TSharedRef<SOldCurveTreeContent>& OldTreeContent
+		const TSharedRef<SSequencerCurveEditor>& InRootWidget, 
+		const TSharedRef<SCurveEditorPanel>& InCurveEditorPanel,
+		const TSharedRef<SOldCurveTreeContent>& InOldTreeContent,
+		const TSharedRef<SCurveTreeContent>& InTreeContent
 		)
-		: RootWidget(RootWidget)
-		, CurveEditorPanel(CurveEditorPanel)
-		, OldTreeContent(OldTreeContent)
+		: RootWidget(InRootWidget)
+		, CurveEditorPanel(InCurveEditorPanel)
+		, OldTreeContent(InOldTreeContent)
+		, TreeContent(InTreeContent)
 	{}
 };
 } // namespace UE::Sequencer
diff --git a/Engine/Source/Editor/Sequencer/Private/MVVM/Views/CurveEditor/SCurveEditorSearchAndFilterRow.cpp b/Engine/Source/Editor/Sequencer/Private/MVVM/Views/CurveEditor/SCurveEditorSearchAndFilterRow.cpp
new file mode 100644
--- /dev/null
+++ b/Engine/Source/Editor/Sequencer/Private/MVVM/Views/CurveEditor/SCurveEditorSearchAndFilterRow.cpp
@@ -0,0 +1,89 @@
+﻿// Copyright Epic Games, Inc. All Rights Reserved.
+
+#include "SCurveEditorSearchAndFilterRow.h"
+
+#include "Filters/Menus/SequencerTrackFilterMenu.h"
+#include "Filters/SequencerFilterBar.h"
+#include "Filters/Utils/HideIsolateViewModelUtils.h"
+#include "Filters/ViewModels/HideIsolateShowViewModel.h"
+#include "Filters/ViewModels/Linking/ILinkedFilterViewModel.h"
+#include "Filters/ViewModels/Linking/LinkedFilterViewModel.h"
+#include "Filters/Widgets/Linking/SFilterLinkStateSwitcher.h"
+#include "Filters/Widgets/Linking/SLinkedFilteringToggleButton.h"
+#include "Filters/Widgets/SSearchAndFilterWidget.h"
+#include "MVVM/Selection/Selection.h"
+#include "Sequencer.h"
+#include "Widgets/SBoxPanel.h"
+#include "Widgets/Sequencer/SSequencerViewOptionsButton.h"
+
+namespace UE::Sequencer
+{
+namespace CurveEditorRowDetail
+{
+using FOutlinerExtensionsAttr = TAttribute<TSet<TWeakViewModelPtr<IOutlinerExtension>>>; 
+
+static FOnGetContent MakeDelegate_AddFilterMenuContent(const TSharedRef<FSequencerFilterBar>& InFilterBar)
+{
+	return FOnGetContent::CreateLambda([Menu = MakeShared<FSequencerTrackFilterMenu>(), WeakFilterBar = InFilterBar.ToWeakPtr()]
+	{
+		const TSharedPtr<FSequencerFilterBar> FilterBar = WeakFilterBar.Pin();
+		return ensure(FilterBar) ? Menu->CreateMenu(FilterBar.ToSharedRef()) : SNullWidget::NullWidget;
+	});
+}
+}
+
+void SCurveEditorSearchAndFilterRow::Construct(const FArguments& InArgs)
+{
+	WeakSequencer = InArgs._Sequencer;
+	FilterModel = InArgs._FilterViewModel;
+	
+	ChildSlot
+	[
+		SNew(SHorizontalBox)
+
+		+ SHorizontalBox::Slot()
+		.AutoWidth()
+		.VAlign(VAlign_Center)
+		[
+			SNew(SLinkedFilteringToggleButton, InArgs._FilterViewModel.ToSharedRef())
+		]
+
+		// Advanced Search Filter Combo Button
+		+ SHorizontalBox::Slot()
+		.FillWidth(1.f)
+		.VAlign(VAlign_Center)
+		[
+			MakeFilterContent(InArgs)
+		]
+	];
+}
+
+TSharedRef<SWidget> SCurveEditorSearchAndFilterRow::MakeFilterContent(const FArguments& InArgs)
+{
+	const auto MakeWidget = [this, &InArgs](const TSharedRef<FSequencerFilterBar>& InFilterBar)
+	{
+		const TSharedRef<FSequencer> SequencerPin = WeakSequencer.Pin().ToSharedRef();
+		const TSharedRef<FHideIsolateShowViewModel> IsolateViewModel = MakeShared<FHideIsolateShowViewModel>(
+			InFilterBar, InFilterBar->GetHideIsolateFilter(),
+			InArgs._SelectedItems, InArgs._AllItems
+			);
+		IsolateViewModel->OnRequestScrollToView().Add(InArgs._ScrollItemIntoView);
+			
+		const TSharedPtr<ISequencerTrackFilters> Filter = StaticCastSharedPtr<ISequencerTrackFilters>(InFilterBar.ToSharedPtr());
+		return SNew(SSearchAndFilterWidget, InFilterBar, IsolateViewModel)
+			.OnSearchTextChanged(this, &SCurveEditorSearchAndFilterRow::OnSearchTextChanged, Filter)
+			.OnMakeAddFilterMenuContent(CurveEditorRowDetail::MakeDelegate_AddFilterMenuContent(InFilterBar));
+	};
+	
+	const TSharedPtr<FLinkedFilterViewModel> LinkedFilterViewModel = InArgs._FilterViewModel;
+	return SNew(SFilterLinkStateSwitcher)
+		.LinkedFilterViewModel(InArgs._FilterViewModel)
+		.LinkedContent() [ MakeWidget(LinkedFilterViewModel->GetLinkedFilterBarImpl()) ]
+		.InstancedContent() [ MakeWidget(LinkedFilterViewModel->GetInstancedFilterBarImpl()) ];
+}
+
+void SCurveEditorSearchAndFilterRow::OnSearchTextChanged(const FText& InNewText, TSharedPtr<ISequencerTrackFilters> InFilter)
+{
+	InFilter->SetTextFilterString(InNewText.ToString());
+}
+} // namespace UE::Sequencer
diff --git a/Engine/Source/Editor/Sequencer/Private/MVVM/Views/CurveEditor/SCurveEditorSearchAndFilterRow.h b/Engine/Source/Editor/Sequencer/Private/MVVM/Views/CurveEditor/SCurveEditorSearchAndFilterRow.h
new file mode 100644
--- /dev/null
+++ b/Engine/Source/Editor/Sequencer/Private/MVVM/Views/CurveEditor/SCurveEditorSearchAndFilterRow.h
@@ -0,0 +1,59 @@
+﻿// Copyright Epic Games, Inc. All Rights Reserved.
+
+#pragma once
+
+#include "MVVM/Extensions/IOutlinerExtension.h"
+#include "Widgets/DeclarativeSyntaxSupport.h"
+#include "Widgets/SCompoundWidget.h"
+
+class ISequencerTrackFilters;
+class FSequencer;
+
+namespace UE::Sequencer
+{
+class FLinkedFilterViewModel;
+class ILinkedFilterViewModel;
+class SSearchAndFilterWidget;
+
+DECLARE_DELEGATE_OneParam(FScrollIntoView, const TWeakViewModelPtr<IOutlinerExtension>&);
+
+/** The row of widgets displayed above the Curve Editor tree view. */
+class SCurveEditorSearchAndFilterRow : public SCompoundWidget
+{
+public:
+	
+	SLATE_BEGIN_ARGS(SCurveEditorSearchAndFilterRow) {}
+		/** Used to generate the subwidgets */
+		SLATE_ARGUMENT(TWeakPtr<FSequencer>, Sequencer)
+		/** Determines the filtering widgets displayed. */
+		SLATE_ARGUMENT(TSharedPtr<FLinkedFilterViewModel>, FilterViewModel)
+
+		/** Gets all items displayed in the tree view. */
+		SLATE_ATTRIBUTE(TSet<TWeakViewModelPtr<IOutlinerExtension>>, AllItems)
+		/** Gets the items selected in the tree view. */
+		SLATE_ATTRIBUTE(TSet<TWeakViewModelPtr<IOutlinerExtension>>, SelectedItems)
+		/** Requests that the given item is scrolled into view. */
+		SLATE_EVENT(FScrollIntoView, ScrollItemIntoView)
+	SLATE_END_ARGS()
+	
+	void Construct(const FArguments& InArgs);
+
+private:
+	
+	/** Used to update Sequencer's filters. */
+	TWeakPtr<FSequencer> WeakSequencer;
+	
+	/** Filter model for the Curve Editor. */
+	TSharedPtr<ILinkedFilterViewModel> FilterModel;
+	
+	/** The widget displaying the search bar, and other filter widgets. */
+	TSharedPtr<SSearchAndFilterWidget> FilterRow;
+
+	/** @return The filtering widget. */
+	TSharedRef<SWidget> MakeFilterContent(const FArguments& InArgs);
+	
+	void OnSearchTextChanged(const FText& InNewText, TSharedPtr<ISequencerTrackFilters> InFilter);
+};
+} // namespace UE::Sequencer
+
+
diff --git a/Engine/Source/Editor/Sequencer/Private/MVVM/Views/CurveEditor/SCurveTreeContent.cpp b/Engine/Source/Editor/Sequencer/Private/MVVM/Views/CurveEditor/SCurveTreeContent.cpp
new file mode 100644
--- /dev/null
+++ b/Engine/Source/Editor/Sequencer/Private/MVVM/Views/CurveEditor/SCurveTreeContent.cpp
@@ -0,0 +1,226 @@
+﻿// Copyright Epic Games, Inc. All Rights Reserved.
+
+#include "SCurveTreeContent.h"
+
+#include "CurveEditor.h"
+#include "SCurveEditorSearchAndFilterRow.h"
+#include "SelectionSyncUtils.h"
+#include "Components/VerticalBox.h"
+#include "Misc/TransportControlsHelper.h"
+#include "SPopoutTabInlineContent.h"
+#include "STemporarilyFocusedSpinBox.h"
+#include "Sequencer.h"
+#include "Filters/Widgets/Linking/SFilterBarSwitcher.h"
+#include "MVVM/CurveEditorIntegrationExtension.h"
+#include "Styling/AppStyle.h"
+#include "Tree/SCurveEditorTree.h"
+#include "Widgets/Layout/SBorder.h"
+#include "Widgets/Layout/SScrollBorder.h"
+
+namespace UE::Sequencer
+{
+class FCurveEditorIntegrationExtension;
+
+void SCurveTreeContent::Construct(const FArguments& InArgs)
+{
+	check(InArgs._Sequencer && InArgs._CurveEditor && InArgs._FilterViewModel);
+	
+	WeakSequencer = InArgs._Sequencer;
+	CurveEditor = InArgs._CurveEditor;
+	
+	ChildSlot
+	[
+		SNew(SVerticalBox)
+
+		+ SVerticalBox::Slot()
+		.AutoHeight()
+		[
+			MakeFilterArea(InArgs)
+		]
+
+		+SVerticalBox::Slot()
+		[
+			MakeTreeContent(InArgs)
+		]
+
+		+ SVerticalBox::Slot()
+		.AutoHeight()
+		[
+			MakeFooter(InArgs)
+		]
+	];
+}
+
+void SCurveTreeContent::FocusSearchBox()
+{
+	// TODO DP
+	//FSlateApplication::Get().SetKeyboardFocus(CurveEditorSearchBox, EFocusCause::SetDirectly);
+}
+
+void SCurveTreeContent::FocusPlayTimeDisplay()
+{
+	PlayTimeDisplay->Setup(); 
+	FSlateApplication::Get().SetKeyboardFocus(PlayTimeDisplay, EFocusCause::SetDirectly);
+}
+
+void SCurveTreeContent::SyncSelection() const
+{
+	if (const TSharedPtr<FSequencer> Sequencer = WeakSequencer.Pin())
+	{
+		Sequencer::SyncSelection(
+			Sequencer->GetViewModel()->GetSelection(), CurveEditor.ToSharedRef(), CurveEditorTreeView.ToSharedRef()
+			);
+	}
+}
+
+TSharedRef<SWidget> SCurveTreeContent::MakeFilterArea(const FArguments& InArgs)
+{
+	return SNew(SVerticalBox)
+		
+		+SVerticalBox::Slot()
+		.AutoHeight()
+		[
+			SNew(SBorder)
+			.BorderImage(FAppStyle::GetBrush("ToolPanel.GroupBorder"))
+			.Clipping(EWidgetClipping::ClipToBounds)
+			[
+				SNew(SCurveEditorSearchAndFilterRow)
+				.Sequencer(InArgs._Sequencer)
+				.FilterViewModel(InArgs._FilterViewModel)
+				.ScrollItemIntoView(this, &SCurveTreeContent::ScrollIntoView)
+				.AllItems(this, &SCurveTreeContent::GetAllItems)
+				.SelectedItems(this, &SCurveTreeContent::GetSelectedItems)
+			]
+		]
+		
+		+SVerticalBox::Slot()
+		.AutoHeight()
+		[
+			InArgs._FilterPills.Widget
+		];
+}
+
+TSharedRef<SWidget> SCurveTreeContent::MakeTreeContent(const FArguments& InArgs)
+{
+	const TSharedRef<SCurveEditorTree> NewCurveEditorTree = SNew(SCurveEditorTree, InArgs._CurveEditor);
+	CurveEditorTreeView = NewCurveEditorTree;
+	return SNew(SScrollBorder, NewCurveEditorTree)
+		[
+			NewCurveEditorTree
+		];
+}
+
+TSharedRef<SWidget> SCurveTreeContent::MakeFooter(const FArguments& InArgs)
+{
+	const TSharedRef<FSequencer> Sequencer = InArgs._Sequencer.ToSharedRef();
+	
+	PlayTimeDisplay = Sequencer->MakePlayTimeDisplay();
+	return SNew(SBorder)
+		.BorderImage(FAppStyle::GetBrush("ToolPanel.GroupBorder"))
+		.Clipping(EWidgetClipping::ClipToBounds)
+		.Visibility_Static(&GetPopoutTransportControlsVisibility_CurveEditor, StaticCastWeakPtr<ISequencer>(Sequencer.ToWeakPtr()))
+		[
+			SNew(SHorizontalBox)
+
+			+ SHorizontalBox::Slot()
+			.HAlign(HAlign_Fill)
+			[
+				MakePopoutTransportControls(Sequencer, InArgs._TabManager)
+			]
+
+			+ SHorizontalBox::Slot()
+			.AutoWidth()
+			.VAlign(VAlign_Center)
+			.HAlign(HAlign_Right)
+			[
+				SNew(SButton)
+				.VAlign(EVerticalAlignment::VAlign_Center)
+				.ButtonStyle(FAppStyle::Get(), "NoBorder")
+				.ContentPadding(FMargin(1, 0))
+				[
+					SNew(SHorizontalBox)
+
+					+ SHorizontalBox::Slot()
+					.AutoWidth()
+					.VAlign(VAlign_Center)
+					.HAlign(HAlign_Right)
+					.Padding(FMargin(3.f, 0.f, 0.f, 0.f))
+					[
+						SNew(SBorder)
+						.BorderImage(nullptr)
+						[
+							PlayTimeDisplay.ToSharedRef()
+						]
+					]
+				]
+			]
+		];
+}
+
+void SCurveTreeContent::ScrollIntoView(const TWeakViewModelPtr<IOutlinerExtension>& WeakViewModel)
+{
+	// TODO UE-364009: Change this to use the new class that FCurveEditorIntegrationExtension is extracted to (and inject it to SCurveTreeContent)
+	const TSharedPtr<FSequencer> Sequencer = WeakSequencer.Pin();
+	FCurveEditorIntegrationExtension* CurveEditorIntegration = Sequencer 
+		? Sequencer->GetViewModel()->GetRootModel()->CastDynamic<FCurveEditorIntegrationExtension>() : nullptr;
+	if (!CurveEditorIntegration)
+	{
+		return;
+	}
+	
+	const FCurveEditorTreeItemID Id = CurveEditorIntegration->FindIdByViewModel(WeakViewModel);
+	if (Id.IsValid())
+	{
+		CurveEditorTreeView->RequestScrollIntoView(Id);
+	}
+}
+
+TSet<TWeakViewModelPtr<IOutlinerExtension>> SCurveTreeContent::GetAllItems() const
+{
+	// TODO UE-364009: Change this to use the new class that FCurveEditorIntegrationExtension is extracted to (and inject it to SCurveTreeContent)
+	const TSharedPtr<FSequencer> Sequencer = WeakSequencer.Pin();
+	FCurveEditorIntegrationExtension* CurveEditorIntegration = Sequencer 
+		? Sequencer->GetViewModel()->GetRootModel()->CastDynamic<FCurveEditorIntegrationExtension>() : nullptr;
+	if (!CurveEditorIntegration)
+	{
+		return {};
+	}
+	
+	TSet<TWeakViewModelPtr<IOutlinerExtension>> AllItems;
+	for (const TPair<FCurveEditorTreeItemID, FCurveEditorTreeItem>& Item : CurveEditor->GetTree()->GetAllItems())
+	{
+		const TSharedPtr<FViewModel> ViewModel = CurveEditorIntegration->FindViewModelById(Item.Key).Pin();
+		const TViewModelPtr<IOutlinerExtension> OutlinerNode = ViewModel ? CastViewModel<IOutlinerExtension>(ViewModel.ToSharedRef()) : nullptr;
+		if (OutlinerNode)
+		{
+			AllItems.Add(OutlinerNode);
+		}
+	}
+	return AllItems;
+}
+
+TSet<TWeakViewModelPtr<IOutlinerExtension>> SCurveTreeContent::GetSelectedItems() const
+{
+	// TODO UE-364009: Change this to use the new class that FCurveEditorIntegrationExtension is extracted to (and inject it to SCurveTreeContent)
+	const TSharedPtr<FSequencer> Sequencer = WeakSequencer.Pin();
+	FCurveEditorIntegrationExtension* CurveEditorIntegration = Sequencer 
+		? Sequencer->GetViewModel()->GetRootModel()->CastDynamic<FCurveEditorIntegrationExtension>() : nullptr;
+	if (!CurveEditorIntegration)
+	{
+		return {};
+	}
+	
+	TSet<TWeakViewModelPtr<IOutlinerExtension>> Selected;
+	for (const FCurveEditorTreeItemID& Item : CurveEditorTreeView->GetSelectedItems())
+	{
+		const TSharedPtr<FViewModel> ViewModel = CurveEditorIntegration->FindViewModelById(Item).Pin();
+		const TViewModelPtr<IOutlinerExtension> OutlinerNode = ViewModel ? CastViewModel<IOutlinerExtension>(ViewModel.ToSharedRef()) : nullptr;
+		if (OutlinerNode)
+		{
+			Selected.Add(OutlinerNode);
+		}
+	}
+	return Selected;
+	
+}
+} // namespace UE::Sequencer
diff --git a/Engine/Source/Editor/Sequencer/Private/MVVM/Views/CurveEditor/SCurveTreeContent.h b/Engine/Source/Editor/Sequencer/Private/MVVM/Views/CurveEditor/SCurveTreeContent.h
new file mode 100644
--- /dev/null
+++ b/Engine/Source/Editor/Sequencer/Private/MVVM/Views/CurveEditor/SCurveTreeContent.h
@@ -0,0 +1,74 @@
+﻿// Copyright Epic Games, Inc. All Rights Reserved.
+
+#pragma once
+
+#include "MVVM/Extensions/IOutlinerExtension.h"
+#include "Widgets/DeclarativeSyntaxSupport.h"
+#include "Widgets/SCompoundWidget.h"
+
+class FCurveEditor;
+class FSequencer;
+class FTabManager;
+class SCurveEditorTree;
+template<typename T> struct STemporarilyFocusedSpinBox;
+
+namespace UE::Sequencer
+{
+class FLinkedFilterViewModel;
+
+/** Displays Curve Editor in Sequencer with linked filtering. */
+class SCurveTreeContent : public SCompoundWidget
+{
+public:
+	
+	SLATE_BEGIN_ARGS(SCurveTreeContent) {}
+		/** Required. The Sequencer instance that owns this content */
+		SLATE_ARGUMENT(TSharedPtr<FSequencer>, Sequencer)
+		/** Required. The curve editor instance that is being displayed */
+		SLATE_ARGUMENT(TSharedPtr<FCurveEditor>, CurveEditor)
+		/** Required. Determines the filtering widgets displayed. */
+		SLATE_ARGUMENT(TSharedPtr<FLinkedFilterViewModel>, FilterViewModel)
+		
+		/** Explicit tab manager to use, otherwise uses the global tab manager */
+		SLATE_ARGUMENT(TSharedPtr<FTabManager>, TabManager)
+
+		/** Shows the filter pills. */
+		SLATE_NAMED_SLOT(FArguments, FilterPills)
+	SLATE_END_ARGS()
+	
+	void Construct(const FArguments& InArgs);
+	
+	/** Sets keyboard focus to the search box. */
+	void FocusSearchBox();
+	/** Sets keyboard focus to the play time display. */
+	void FocusPlayTimeDisplay();
+	
+	/** Syncs the tree view's selection to that of Sequencer. */
+	void SyncSelection() const;
+	
+	const TSharedPtr<SCurveEditorTree>& GetCurveEditorTreeView() const { return CurveEditorTreeView; }
+	
+private:
+	
+	/** The owning Sequencer instance that the Curve Editor is being displayed for. */
+	TWeakPtr<FSequencer> WeakSequencer;
+	/** The curve editor being displayed. */
+	TSharedPtr<FCurveEditor> CurveEditor;
+	
+	/** Curve editor tree widget */
+	TSharedPtr<SCurveEditorTree> CurveEditorTreeView;
+	/** The current playback time display. */
+	TSharedPtr<STemporarilyFocusedSpinBox<double>> PlayTimeDisplay;
+
+	TSharedRef<SWidget> MakeFilterArea(const FArguments& InArgs);
+	TSharedRef<SWidget> MakeTreeContent(const FArguments& InArgs);
+	TSharedRef<SWidget> MakeFooter(const FArguments& InArgs);
+	
+	/** Scrolls the given item into view. */
+	void ScrollIntoView(const TWeakViewModelPtr<IOutlinerExtension>& WeakViewModel);
+	/** @return All items visible in the tree view. */
+	TSet<TWeakViewModelPtr<IOutlinerExtension>> GetAllItems() const;
+	/** @return All items selected in the tree view. */
+	TSet<TWeakViewModelPtr<IOutlinerExtension>> GetSelectedItems() const;
+};
+} // namespace UE::Sequencer
diff --git a/Engine/Source/Editor/Sequencer/Private/MVVM/Views/CurveEditor/SOldCurveTreeContent.cpp b/Engine/Source/Editor/Sequencer/Private/MVVM/Views/CurveEditor/SOldCurveTreeContent.cpp
--- a/Engine/Source/Editor/Sequencer/Private/MVVM/Views/CurveEditor/SOldCurveTreeContent.cpp
+++ b/Engine/Source/Editor/Sequencer/Private/MVVM/Views/CurveEditor/SOldCurveTreeContent.cpp
@@ -116,7 +116,7 @@ TSharedRef<SWidget> SOldCurveTreeContent::MakeFooter(const FArguments& InArgs, c
 	return SNew(SBorder)
 		.BorderImage(FAppStyle::GetBrush("ToolPanel.GroupBorder"))
 		.Clipping(EWidgetClipping::ClipToBounds)
-		.Visibility(this, &SOldCurveTreeContent::GetPopoutTransportControlsVisibility)
+		.Visibility_Static(&UE::Sequencer::GetPopoutTransportControlsVisibility_CurveEditor, StaticCastWeakPtr<ISequencer>(InSequencer.ToWeakPtr()))
 		[
 			SNew(SHorizontalBox)
 
@@ -180,28 +180,6 @@ USequencerSettings* SOldCurveTreeContent::GetSequencerSettings() const
 	const TSharedPtr<ISequencer> Sequencer = WeakSequencer.Pin();
 	return Sequencer.IsValid() ? Sequencer->GetSequencerSettings() : nullptr;
 }
-
-EVisibility SOldCurveTreeContent::GetPopoutTransportControlsVisibility() const
-{
-	const TSharedPtr<ISequencer> Sequencer = WeakSequencer.Pin();
-	if (!Sequencer.IsValid())
-	{
-		return EVisibility::Collapsed;
-	}
-	const UMovieSceneSequence* const RootSequence = Sequencer->GetRootMovieSceneSequence();
-	if (!RootSequence || !RootSequence->IsA<ULevelSequence>())
-	{
-		return EVisibility::Visible;
-	}
-	const USequencerSettings* const SequencerSettings = Sequencer->GetSequencerSettings();
-	if (!SequencerSettings)
-	{
-		return EVisibility::Visible;
-	}
-	const bool bShowTransportControls = GetDefault<UCurveEditorSettings>()->GetShowTransportControls();
-	const bool bTransportControlsPoppedOut = SequencerSettings->AreTransportControlsPoppedOut();
-	return (!bShowTransportControls || bTransportControlsPoppedOut) ? EVisibility::Collapsed : EVisibility::Visible;
-}
 } // namespace UE::Sequencer
 
 #undef LOCTEXT_NAMESPACE
\ No newline at end of file
diff --git a/Engine/Source/Editor/Sequencer/Private/MVVM/Views/CurveEditor/SOldCurveTreeContent.h b/Engine/Source/Editor/Sequencer/Private/MVVM/Views/CurveEditor/SOldCurveTreeContent.h
--- a/Engine/Source/Editor/Sequencer/Private/MVVM/Views/CurveEditor/SOldCurveTreeContent.h
+++ b/Engine/Source/Editor/Sequencer/Private/MVVM/Views/CurveEditor/SOldCurveTreeContent.h
@@ -62,7 +62,5 @@ private:
 	TSharedRef<SWidget> MakeFooter(const FArguments& InArgs, const TSharedRef<FSequencer> InSequencer);
 	
 	USequencerSettings* GetSequencerSettings() const;
-	
-	EVisibility GetPopoutTransportControlsVisibility() const;
 };
 } // namespace UE::Sequencer
diff --git a/Engine/Source/Editor/Sequencer/Private/Misc/TransportControlsHelper.cpp b/Engine/Source/Editor/Sequencer/Private/Misc/TransportControlsHelper.cpp
--- a/Engine/Source/Editor/Sequencer/Private/Misc/TransportControlsHelper.cpp
+++ b/Engine/Source/Editor/Sequencer/Private/Misc/TransportControlsHelper.cpp
@@ -1,6 +1,8 @@
 ﻿// Copyright Epic Games, Inc. All Rights Reserved.
 
 #include "TransportControlsHelper.h"
+
+#include "CurveEditorSettings.h"
 #include "Framework/Docking/TabManager.h"
 #include "ISequencer.h"
 #include "LevelSequence.h"
@@ -52,6 +54,31 @@ TSharedRef<SPopoutTabInlineContent> MakePopoutTransportControls(const TSharedRef
 		];
 }
 
+EVisibility GetPopoutTransportControlsVisibility_CurveEditor(const TSharedRef<ISequencer>& InSequencer)
+{
+	const UMovieSceneSequence* const RootSequence = InSequencer->GetRootMovieSceneSequence();
+	if (!RootSequence || !RootSequence->IsA<ULevelSequence>())
+	{
+		return EVisibility::Visible;
+	}
+	
+	const USequencerSettings* const SequencerSettings = InSequencer->GetSequencerSettings();
+	if (!SequencerSettings)
+	{
+		return EVisibility::Visible;
+	}
+	
+	const bool bShowTransportControls = GetDefault<UCurveEditorSettings>()->GetShowTransportControls();
+	const bool bTransportControlsPoppedOut = SequencerSettings->AreTransportControlsPoppedOut();
+	return (!bShowTransportControls || bTransportControlsPoppedOut) ? EVisibility::Collapsed : EVisibility::Visible;
 }
 
+EVisibility GetPopoutTransportControlsVisibility_CurveEditor(TWeakPtr<ISequencer> InWeakSequencer)
+{
+	const TSharedPtr<ISequencer> Sequencer = InWeakSequencer.Pin();
+	return Sequencer ? GetPopoutTransportControlsVisibility_CurveEditor(Sequencer.ToSharedRef()) : EVisibility::Collapsed;
+}
+
+} // namespace UE::Sequencer
+
 #undef LOCTEXT_NAMESPACE
diff --git a/Engine/Source/Editor/Sequencer/Private/Misc/TransportControlsHelper.h b/Engine/Source/Editor/Sequencer/Private/Misc/TransportControlsHelper.h
--- a/Engine/Source/Editor/Sequencer/Private/Misc/TransportControlsHelper.h
+++ b/Engine/Source/Editor/Sequencer/Private/Misc/TransportControlsHelper.h
@@ -8,6 +8,7 @@ class FTabManager;
 class ISequencer;
 class SPopoutTabInlineContent;
 class SWidget;
+struct EVisibility;
 
 namespace UE::Sequencer
 {
@@ -19,4 +20,8 @@ namespace UE::Sequencer
 TSharedRef<SPopoutTabInlineContent> MakePopoutTransportControls(const TSharedRef<ISequencer>& InSequencer
 	, const TSharedPtr<FTabManager>& InTabManager, const TSharedPtr<SWidget>& InOverriddenContent = nullptr);
 
-}
+/** @return Visibility the popout transport controls should have in Sequencer's curve editor: must be configured to be visible and not popped out. */
+EVisibility GetPopoutTransportControlsVisibility_CurveEditor(const TSharedRef<ISequencer>& InSequencer);
+/** Overload useful for shorthand delegate construction. */
+EVisibility GetPopoutTransportControlsVisibility_CurveEditor(TWeakPtr<ISequencer> InWeakSequencer);
+} // namespace UE::Sequencer
diff --git a/Engine/Source/Editor/Sequencer/Public/MVVM/CurveEditorExtension.h b/Engine/Source/Editor/Sequencer/Public/MVVM/CurveEditorExtension.h
--- a/Engine/Source/Editor/Sequencer/Public/MVVM/CurveEditorExtension.h
+++ b/Engine/Source/Editor/Sequencer/Public/MVVM/CurveEditorExtension.h
@@ -27,9 +27,9 @@ namespace UE
 {
 namespace Sequencer
 {
-
-class FSequencerEditorViewModel;
 class FCurveEditorWidgetOwner;
+class FSequencerEditorViewModel;
+class ILinkedFilterViewModel;
 
 class FCurveEditorExtension : public IDynamicExtension
 {
@@ -63,7 +63,6 @@ public:
 	 */
 	UE_API void RequestSyncSelection();
 
-
 public:
 
 	static UE_API const FName CurveEditorTabName;
@@ -79,6 +78,9 @@ private:
 	
 	/** Cancels any pending transactions. */
 	void CleanUpPendingSyncSelection();
+	
+	/** @return The proxy for interacting with linked filter */
+	ILinkedFilterViewModel* GetLinkedFilteringModel() const { return CurveEditorFilterClient.Get(); }
 
 private:
 
@@ -88,12 +90,23 @@ private:
 	/** Curve editor */
 	TSharedPtr<FCurveEditor> CurveEditorModel;
 	
+	/** Handles refiltering when the active filter bar becomes dirty. */
+	TSharedPtr<ILinkedFilterViewModel> CurveEditorFilterClient;
+	
 	/** Responsible for creating the UI. */
 	TPimplPtr<FCurveEditorWidgetOwner> WidgetOwner; 
 
 	/** Handle to latent refresh of selection. */
 	FTimerHandle SyncSelectionHandle;
 	
+	/**
+	 * Let's FCurveEditorIntegrationExtension know that the curve editor is set up.
+	 * 
+	 * TODO UE-364009: Instead, FCurveEditorIntegrationExtension should be extracted to a separate class that is created by FCurveEditorExtension 
+	 * so any dependencies can be injected.
+	 */
+	FSimpleMulticastDelegate OnCurveEditorCreated;
+	
 	friend class FCurveEditorIntegrationExtension;
 };
 
diff --git a/Engine/Source/Editor/Sequencer/Public/MVVM/CurveEditorIntegrationExtension.h b/Engine/Source/Editor/Sequencer/Public/MVVM/CurveEditorIntegrationExtension.h
--- a/Engine/Source/Editor/Sequencer/Public/MVVM/CurveEditorIntegrationExtension.h
+++ b/Engine/Source/Editor/Sequencer/Public/MVVM/CurveEditorIntegrationExtension.h
@@ -37,6 +37,7 @@ public:
 	UE_SEQUENCER_DECLARE_VIEW_MODEL_TYPE_ID_API(UE_API, FCurveEditorIntegrationExtension)
 
 	UE_API FCurveEditorIntegrationExtension();
+	UE_API virtual ~FCurveEditorIntegrationExtension() override;
 
 	UE_API virtual void OnCreated(TSharedRef<FViewModel> InWeakOwner) override;
 
@@ -50,6 +51,11 @@ public:
 	 * Clears the curve editor of all contents.
 	 */
 	UE_API void ResetCurveEditor();
+	
+	/** @return The view model associated with the tree item. */
+	TWeakPtr<FViewModel> FindViewModelById(const FCurveEditorTreeItemID& InTreeId) const;
+	/** @return The tree item associated with the view model. */
+	FCurveEditorTreeItemID FindIdByViewModel(TWeakPtr<FViewModel>) const;
 
 private:
 
@@ -66,6 +72,14 @@ private:
 
 	TWeakPtr<FSequenceModel> WeakOwnerModel;
 	TMap<TWeakPtr<FViewModel>, FCurveEditorTreeItemID> ViewModelToTreeItemIDMap;
+	TMap<FCurveEditorTreeItemID, TWeakPtr<FViewModel>> TreeItemIDToViewModelMap;
+	
+	/** Ensures BindFilterBarDependencies is called at some point. */
+	void SetupFilterBarDependencies();
+	/** Binds to the linked filtering system to cause refreshing of curve editor content when the filtering state changes. */
+	void BindFilterBarDependencies();
+	/** Cleans up any callbacks to the linked filtering system. */
+	void UnbindFilterBarDependencies();
 };
 
 } // namespace Sequencer
