PR #14433: Avoid reparenting while the iteration over the model is active

#Synced-CL 50973831

diff --git a/Engine/Source/Editor/Sequencer/Private/MVVM/ViewModels/SequenceModel.cpp b/Engine/Source/Editor/Sequencer/Private/MVVM/ViewModels/SequenceModel.cpp
--- a/Engine/Source/Editor/Sequencer/Private/MVVM/ViewModels/SequenceModel.cpp
+++ b/Engine/Source/Editor/Sequencer/Private/MVVM/ViewModels/SequenceModel.cpp
@@ -103,14 +103,20 @@ void FSequenceModel::OnDecorationAdded(UObject* AddedDecoration)
 
 void FSequenceModel::OnDecorationRemoved(UObject* RemovedDecoration)
 {
+	TViewModelPtr<IObjectModelExtension> ObjectToDetach;
 	for (TViewModelPtr<IObjectModelExtension> Object : GetChildrenForList(&DecorationModelList).IterateSubList<IObjectModelExtension>())
 	{
 		if (Object->GetObject() == RemovedDecoration)
 		{
-			Object.AsModel()->RemoveFromParent();
+			ObjectToDetach = Object;
 			break;
 		}
 	}
+	
+	if (ObjectToDetach)
+	{
+		ObjectToDetach.AsModel()->RemoveFromParent();
+	}
 }
 
 FMovieSceneSequenceID FSequenceModel::GetSequenceID() const
