Sequencer: Add missing command binding to ConvertToSpawnables

Fixed FSequencer::ConvertSelectedNodesToSpawnables so that it now calls ConvertToCustomBinding if it's a level sequence since it's possible that there are  non level sequence spawnables.

#Synced-CL 50973831

diff --git a/Engine/Source/Editor/Sequencer/Private/Sequencer.cpp b/Engine/Source/Editor/Sequencer/Private/Sequencer.cpp
--- a/Engine/Source/Editor/Sequencer/Private/Sequencer.cpp
+++ b/Engine/Source/Editor/Sequencer/Private/Sequencer.cpp
@@ -9287,11 +9287,27 @@ void FSequencer::ConvertSelectedNodesToSpawnables()
 		FMovieScenePossessable* Possessable = MovieScene->FindPossessable(ObjectBindingNode->GetObjectGuid());
 		if (Possessable)
 		{
-			TArray<FMovieSceneSpawnable*> Spawnables = FSequencerUtilities::ConvertToSpawnable(AsShared(), Possessable->GetGuid());
+			TArray<FGuid> ConvertedGuids;
 
-			for (FMovieSceneSpawnable* Spawnable : Spawnables)
+			if (IsLevelEditorSequencer())
 			{
-				for (TWeakObjectPtr<> WeakObject : FindBoundObjects(Spawnable->GetGuid(), ActiveTemplateIDs.Top()))
+				if (FMovieScenePossessable* Spawnable = FSequencerUtilities::ConvertToCustomBinding(AsShared(), Possessable->GetGuid(), UMovieSceneSpawnableActorBinding::StaticClass()))
+				{
+					ConvertedGuids.Add(Spawnable->GetGuid());
+				}
+			}
+			else
+			{
+				TArray<FMovieSceneSpawnable*> Spawnables = FSequencerUtilities::ConvertToSpawnable(AsShared(), Possessable->GetGuid());
+				for (FMovieSceneSpawnable* Spawnable : Spawnables)
+				{
+					ConvertedGuids.Add(Spawnable->GetGuid());
+				}
+			}
+
+			for (FGuid Guid : ConvertedGuids)
+			{
+				for (TWeakObjectPtr<> WeakObject : FindBoundObjects(Guid, GetFocusedTemplateID()))
 				{
 					if (AActor* SpawnedActor = Cast<AActor>(WeakObject.Get()))
 					{
@@ -11392,7 +11408,7 @@ void FSequencer::BindCommands()
 		FCanExecuteAction::CreateLambda([this] { return MovieSceneToolHelpers::CanSplitSection(this->ViewModel->GetSelection()->GetSelectedSections(), GetLocalTime()); }));
 
 	// We can convert to spawnables if anything selected is a root-level possessable
-	auto CanConvertToSpawnables = [this]{
+	auto CanSelectedNodesBeConvertedToSpawnables = [this]{
 		UMovieSceneSequence* Sequence = GetFocusedMovieSceneSequence();
 		if (!Sequence || !Sequence->AllowsSpawnableObjects())
 		{
@@ -11440,6 +11456,12 @@ void FSequencer::BindCommands()
 		return false;
 	};
 
+	SequencerCommandBindings->MapAction(
+		FSequencerCommands::Get().ConvertToSpawnable,
+		FExecuteAction::CreateSP(this, &FSequencer::ConvertSelectedNodesToSpawnables),
+		FCanExecuteAction::CreateLambda(CanSelectedNodesBeConvertedToSpawnables)
+	);
+
 	SequencerCommandBindings->MapAction(
 		FSequencerCommands::Get().ConvertToPossessable,
 		FExecuteAction::CreateSP(this, &FSequencer::ConvertSelectedNodesToPossessables),
