Sequencer: Linked filtering toggle button is now between the "Add filter" combo button and search box.
Based on feedback from team (UX, QA, and product).

#Synced-CL 50904440

diff --git a/Engine/Source/Editor/Sequencer/Private/Filters/Widgets/SSearchAndFilterWidget.cpp b/Engine/Source/Editor/Sequencer/Private/Filters/Widgets/SSearchAndFilterWidget.cpp
--- a/Engine/Source/Editor/Sequencer/Private/Filters/Widgets/SSearchAndFilterWidget.cpp
+++ b/Engine/Source/Editor/Sequencer/Private/Filters/Widgets/SSearchAndFilterWidget.cpp
@@ -39,6 +39,15 @@ void SSearchAndFilterWidget::Construct(
 			MakeAddFilterButton(InFilterBar, InArgs._OnMakeAddFilterMenuContent)
 		]
 
+		// Injected content, e.g. linked filter state toggle button
+		+ SHorizontalBox::Slot()
+		.AutoWidth()
+		.VAlign(VAlign_Center)
+		.Padding(2.f, 0.f, 0.f, 0.f)
+		[
+			InArgs._BeforeSearchBox.Widget
+		]
+
 		// Advanced Search Box
 		+ SHorizontalBox::Slot()
 		.VAlign(VAlign_Center)
diff --git a/Engine/Source/Editor/Sequencer/Private/Filters/Widgets/SSearchAndFilterWidget.h b/Engine/Source/Editor/Sequencer/Private/Filters/Widgets/SSearchAndFilterWidget.h
--- a/Engine/Source/Editor/Sequencer/Private/Filters/Widgets/SSearchAndFilterWidget.h
+++ b/Engine/Source/Editor/Sequencer/Private/Filters/Widgets/SSearchAndFilterWidget.h
@@ -29,6 +29,9 @@ public:
 
 		/** Makes the content for the add filter combo button. */
 		SLATE_EVENT(FOnGetContent, OnMakeAddFilterMenuContent)
+
+		/** Slot between "Add filter" combo and search box. Useful for injecting e.g. linked filter toggle button. */
+		SLATE_NAMED_SLOT(FArguments, BeforeSearchBox)
 	SLATE_END_ARGS()
 
 	void Construct(
diff --git a/Engine/Source/Editor/Sequencer/Private/MVVM/Views/CurveEditor/SCurveEditorSearchAndFilterRow.cpp b/Engine/Source/Editor/Sequencer/Private/MVVM/Views/CurveEditor/SCurveEditorSearchAndFilterRow.cpp
--- a/Engine/Source/Editor/Sequencer/Private/MVVM/Views/CurveEditor/SCurveEditorSearchAndFilterRow.cpp
+++ b/Engine/Source/Editor/Sequencer/Private/MVVM/Views/CurveEditor/SCurveEditorSearchAndFilterRow.cpp
@@ -13,6 +13,7 @@
 #include "Filters/Widgets/SSearchAndFilterWidget.h"
 #include "MVVM/Selection/Selection.h"
 #include "Sequencer.h"
+#include "Filters/Utils/LinkedFilteringViewUtils.h"
 #include "Widgets/SBoxPanel.h"
 #include "Widgets/Sequencer/SSequencerViewOptionsButton.h"
 
@@ -39,22 +40,7 @@ void SCurveEditorSearchAndFilterRow::Construct(const FArguments& InArgs)
 	
 	ChildSlot
 	[
-		SNew(SHorizontalBox)
-
-		+ SHorizontalBox::Slot()
-		.AutoWidth()
-		.VAlign(VAlign_Center)
-		[
-			SNew(SLinkedFilteringToggleButton, InArgs._FilterViewModel.ToSharedRef())
-		]
-
-		// Advanced Search Filter Combo Button
-		+ SHorizontalBox::Slot()
-		.FillWidth(1.f)
-		.VAlign(VAlign_Center)
-		[
-			MakeFilterContent(InArgs)
-		]
+		MakeFilterContent(InArgs)
 	];
 }
 
@@ -72,7 +58,15 @@ TSharedRef<SWidget> SCurveEditorSearchAndFilterRow::MakeFilterContent(const FArg
 		const TSharedPtr<ISequencerTrackFilters> Filter = StaticCastSharedPtr<ISequencerTrackFilters>(InFilterBar.ToSharedPtr());
 		return SNew(SSearchAndFilterWidget, InFilterBar, IsolateViewModel)
 			.OnSearchTextChanged(this, &SCurveEditorSearchAndFilterRow::OnSearchTextChanged, Filter)
-			.OnMakeAddFilterMenuContent(CurveEditorRowDetail::MakeDelegate_AddFilterMenuContent(InFilterBar));
+			.OnMakeAddFilterMenuContent(CurveEditorRowDetail::MakeDelegate_AddFilterMenuContent(InFilterBar))
+			.BeforeSearchBox()
+			[
+				SNew(SLinkedFilteringToggleButton, InArgs._FilterViewModel.ToSharedRef())
+				.Visibility_Lambda([]
+				{
+					return IsLinkedFilteringFeatureEnabled() ? EVisibility::Visible : EVisibility::Collapsed;
+				})
+			];
 	};
 	
 	const TSharedPtr<FLinkedFilterViewModel> LinkedFilterViewModel = InArgs._FilterViewModel;
diff --git a/Engine/Source/Editor/Sequencer/Private/Widgets/Sequencer/SSequencerSearchAndFilterRow.cpp b/Engine/Source/Editor/Sequencer/Private/Widgets/Sequencer/SSequencerSearchAndFilterRow.cpp
--- a/Engine/Source/Editor/Sequencer/Private/Widgets/Sequencer/SSequencerSearchAndFilterRow.cpp
+++ b/Engine/Source/Editor/Sequencer/Private/Widgets/Sequencer/SSequencerSearchAndFilterRow.cpp
@@ -86,17 +86,6 @@ void SSequencerSearchAndFilterRow::Construct(const FArguments& InArgs)
 			.ExtenderForAddMenu(InArgs._ExtenderForAddMenu)
 		]
 
-		+ SHorizontalBox::Slot()
-		.AutoWidth()
-		.VAlign(VAlign_Center)
-		[
-			SNew(SLinkedFilteringToggleButton, InArgs._FilterViewModel.ToSharedRef())
-			.Visibility_Lambda([]
-			{
-				return IsLinkedFilteringFeatureEnabled() ? EVisibility::Visible : EVisibility::Collapsed;
-			})
-		]
-
 		// Advanced Search Filter Combo Button
 		+ SHorizontalBox::Slot()
 		.FillWidth(1.f)
@@ -130,7 +119,8 @@ TSharedPtr<SSequencerSearchBox> SSequencerSearchAndFilterRow::GetActiveSearchBox
 
 TSharedRef<SWidget> SSequencerSearchAndFilterRow::MakeFilterContent(const FArguments& InArgs)
 {
-	const auto MakeWidget = [this](const TSharedRef<FSequencerFilterBar>& InFilterBar)
+	const TSharedPtr<FLinkedFilterViewModel> LinkedFilterViewModel = InArgs._FilterViewModel;
+	const auto MakeWidget = [this, &LinkedFilterViewModel](const TSharedRef<FSequencerFilterBar>& InFilterBar)
 	{
 		const TSharedRef<FSequencer> SequencerPin = WeakSequencer.Pin().ToSharedRef();
 		const TSharedRef<FHideIsolateShowViewModel> IsolateViewModel = MakeShared<FHideIsolateShowViewModel>(
@@ -145,10 +135,17 @@ TSharedRef<SWidget> SSequencerSearchAndFilterRow::MakeFilterContent(const FArgum
 			.OnSearchTextChanged_Lambda([InFilterBar](const FText& InSearchText)
 			{
 				InFilterBar->SetTextFilterString(InSearchText.ToString());
-			});
+			})
+			.BeforeSearchBox()
+			[
+				SNew(SLinkedFilteringToggleButton, LinkedFilterViewModel.ToSharedRef())
+				.Visibility_Lambda([]
+				{
+					return IsLinkedFilteringFeatureEnabled() ? EVisibility::Visible : EVisibility::Collapsed;
+				})
+			];
 	};
 	
-	const TSharedPtr<FLinkedFilterViewModel> LinkedFilterViewModel = InArgs._FilterViewModel;
 	return SNew(SFilterLinkStateSwitcher)
 		.LinkedFilterViewModel(InArgs._FilterViewModel)
 		.LinkedContent() 
diff --git a/Engine/Source/Editor/Sequencer/Private/Widgets/Sequencer/SSequencerSearchAndFilterRow.h b/Engine/Source/Editor/Sequencer/Private/Widgets/Sequencer/SSequencerSearchAndFilterRow.h
--- a/Engine/Source/Editor/Sequencer/Private/Widgets/Sequencer/SSequencerSearchAndFilterRow.h
+++ b/Engine/Source/Editor/Sequencer/Private/Widgets/Sequencer/SSequencerSearchAndFilterRow.h
@@ -48,9 +48,6 @@ private:
 	
 	/** Used to update Sequencer's filters. */
 	TWeakPtr<FSequencer> WeakSequencer;
-	
-	/** The widget displaying the search bar, and other filter widgets. */
-	TSharedPtr<SSearchAndFilterWidget> FilterRow;
 
 	/** Filter model for Sequencer. */
 	TSharedPtr<FLinkedFilterViewModel> SequencerFilterModel;
