Sequencer: remove restoring of pre-animated state on BP compilation

It's causing isses with async loading, since the BlueprintPreCompile callback can now be called from a background thread, which puts us into race conditions when other things are going on in the main thread. Also, it's unclear now if that code is still needed.

This change introduces a new Cvar to re-enable the old behaviour, in case we find that was actually needed for something.

#Synced-CL 51024716

diff --git a/Engine/Source/Editor/Sequencer/Private/Sequencer.cpp b/Engine/Source/Editor/Sequencer/Private/Sequencer.cpp
--- a/Engine/Source/Editor/Sequencer/Private/Sequencer.cpp
+++ b/Engine/Source/Editor/Sequencer/Private/Sequencer.cpp
@@ -221,6 +221,11 @@ static TAutoConsoleVariable<bool> CVarTimeUndo(
 	false,
 	TEXT("Enable/disable ability to undo time when making changes"));
 
+static TAutoConsoleVariable<bool> CVarRestorePreAnimatedStateOnBlueprintPreCompile(
+	TEXT("Sequencer.RestorePreAnimatedStateOnBlueprintPreCompile"),
+	false,
+	TEXT("Enable/disable restoring pre-animated state before Blueprint classes get (re)compiled"));
+
 namespace UE
 {
 	namespace Sequencer
@@ -442,10 +447,13 @@ void FSequencer::InitSequencer(const FSequencerInitParams& InitParams
 	{
 		FDelegateHandle OnBlueprintPreCompileHandle = GEditor->OnBlueprintPreCompile().AddLambda([&](UBlueprint* InBlueprint)
 		{
-			// Restore pre animate state since objects will be reinstanced and current cached state will no longer be valid.
-			if (InBlueprint && InBlueprint->GeneratedClass.Get())
+			if (CVarRestorePreAnimatedStateOnBlueprintPreCompile->GetBool())
 			{
-				PreAnimatedState.RestorePreAnimatedState(InBlueprint->GeneratedClass.Get());
+				// Restore pre animate state since objects will be reinstanced and current cached state will no longer be valid.
+				if (InBlueprint && InBlueprint->GeneratedClass.Get())
+				{
+					PreAnimatedState.RestorePreAnimatedState(InBlueprint->GeneratedClass.Get());
+				}
 			}
 		});
 		AcquiredResources.Add([=] { GEditor->OnBlueprintPreCompile().Remove(OnBlueprintPreCompileHandle); });
