[Sequencer] Add hidden scroll box for search/filter widget row and transport controls widget

#Synced-CL 51084386

#Synced-CL 51084386

diff --git a/Engine/Source/Editor/Sequencer/Private/SSequencer.cpp b/Engine/Source/Editor/Sequencer/Private/SSequencer.cpp
--- a/Engine/Source/Editor/Sequencer/Private/SSequencer.cpp
+++ b/Engine/Source/Editor/Sequencer/Private/SSequencer.cpp
@@ -50,6 +50,7 @@
 #include "Widgets/Layout/SGridPanel.h"
 #include "Widgets/Layout/SScrollBar.h"
 #include "Widgets/Layout/SScrollBorder.h"
+#include "Widgets/Layout/SScrollBox.h"
 #include "Widgets/Input/SComboButton.h"
 #include "Widgets/Layout/SSeparator.h"
 #include "Widgets/Layout/SSplitter.h"
@@ -1093,11 +1094,11 @@ TSharedRef<SWidget> SSequencer::ConstructTransportControls(const TSharedRef<ISeq
 	USequencerToolMenuContext* ContextObject = NewObject<USequencerToolMenuContext>();
 	ContextObject->WeakSequencer = SequencerPtr;
 
-	const TSharedRef<SScaleBox> TransportControlsWidget = SNew(SScaleBox)
-		.Stretch(EStretch::ScaleToFit)
-		.StretchDirection(EStretchDirection::DownOnly)
-		.HAlign(HAlign_Fill)
-		.VAlign(VAlign_Fill)
+	const TSharedRef<SScrollBox> TransportControlsWidget = SNew(SScrollBox)
+		.Orientation(Orient_Horizontal)
+		.ScrollBarAlwaysVisible(false)
+		.ScrollBarVisibility(EVisibility::Collapsed)
+		+ SScrollBox::Slot()
 		[
 			SNew(SHorizontalBox)
 			+ SHorizontalBox::Slot()
@@ -4200,7 +4201,14 @@ void SSequencer::RebuildSearchAndFilterRow()
 	SearchAndFilterRow->AddSlot()
 		.AutoHeight()
 		[
-			FilterRow.ToSharedRef()
+			SNew(SScrollBox)
+			.Orientation(Orient_Horizontal)
+			.ScrollBarAlwaysVisible(false)
+			.ScrollBarVisibility(EVisibility::Collapsed)
+			+ SScrollBox::Slot()
+			[
+				FilterRow.ToSharedRef()
+			]
 		];
 	
 	const TSharedPtr<FSequencer> SequencerPin = SequencerPtr.Pin();
