[Sequencer]
- Better handling of old outliner percentage width
- Instead of setting the outliner width to the minimum if it detects an old percentage width value, we'll use the old value to calculate the new pixel width and use that until the user changes the size, in which case the new pixel width will get saved

#Synced-CL 51084386

diff --git a/Engine/Source/Editor/Sequencer/Private/SSequencer.cpp b/Engine/Source/Editor/Sequencer/Private/SSequencer.cpp
--- a/Engine/Source/Editor/Sequencer/Private/SSequencer.cpp
+++ b/Engine/Source/Editor/Sequencer/Private/SSequencer.cpp
@@ -4342,34 +4342,30 @@ FText SSequencer::GetSidebarSelectionDrawerToolTipText() const
 
 FOptionalSize SSequencer::GetOutlinerWidth() const
 {
-	const USequencerSettings* const SequencerSettings = GetSequencerSettings();
-	if (!SequencerSettings)
-	{
-		return DefaultOutlinerWidth;
-	}
-
-	float ActualOutlinerWidth = TreeViewWidth;
-
-	if (OutlinerAndTrackAreaContent.IsValid())
-	{
-		const FGeometry& CachedGeometry = OutlinerAndTrackAreaContent->GetCachedGeometry();
-		const FVector2f LocalSize = CachedGeometry.GetLocalSize();
-
-		if (LocalSize.X > 0.f)
+	// Returns true if the geometry is valid and has non-zero dimensions
+	auto GetCachedLocalSize = [this](FVector2D& OutSize) -> bool
 		{
-			// Limit the outliner size to the parent size
-			const float MaxDynamicWidth = LocalSize.X - MinOutlinerWidth;
-			if (ActualOutlinerWidth > MaxDynamicWidth)
+			if (OutlinerAndTrackAreaContent.IsValid())
 			{
-				ActualOutlinerWidth = MaxDynamicWidth;
+				OutSize = OutlinerAndTrackAreaContent->GetCachedGeometry().GetLocalSize();
+				return OutSize.X > 0 || OutSize.Y > 0;
 			}
+			return false;
+		};
 
-			const float MinWidth = FMath::Min<float>(MaxDynamicWidth, MinOutlinerWidth);
-			return FMath::Clamp<float>(ActualOutlinerWidth, MinWidth, MaxDynamicWidth);
-		}
+	FVector2D LocalSize = FVector2D::ZeroVector;
+	if (GetCachedLocalSize(LocalSize))
+	{
+		// TreeViewWidth is set to SequencerSettings->GetTreeViewWidth() on first construct.
+		// If this is the old percentage value, we will need to convert it to a pixel width value.
+		const float WorkingWidth = FMath::Max(0.f, TreeViewWidth);
+		const bool bIsPercent = WorkingWidth > 0.f && WorkingWidth < 1.f;
+		const float PixelWidth = (LocalSize.X > 0.f && bIsPercent) ? WorkingWidth * LocalSize.X : WorkingWidth;
+		return FMath::Clamp<float>(PixelWidth, MinOutlinerWidth, LocalSize.X);
 	}
 
-	return FMath::Max<float>(ActualOutlinerWidth, MinOutlinerWidth);
+	// No reliable geometry yet
+	return FMath::Max<float>(TreeViewWidth, MinOutlinerWidth);
 }
 
 bool SSequencer::IsTransportControlsPoppedOut() const
diff --git a/Engine/Source/Editor/Sequencer/Private/SequencerSettings.cpp b/Engine/Source/Editor/Sequencer/Private/SequencerSettings.cpp
--- a/Engine/Source/Editor/Sequencer/Private/SequencerSettings.cpp
+++ b/Engine/Source/Editor/Sequencer/Private/SequencerSettings.cpp
@@ -107,23 +107,6 @@ USequencerSettings::USequencerSettings( const FObjectInitializer& ObjectInitiali
 	CurveEditorZoomScaling.MaxVerticalZoomOut = 10000;
 }
 
-void USequencerSettings::PostLoad()
-{
-	Super::PostLoad();
-
-	// Because we want to remove proportional outliner sizing,
-	// we need to convert the percentage to a pixel width
-	if (TreeViewWidth <= 1.f)
-	{
-		TreeViewWidth = SSequencer::DefaultOutlinerWidth;
-	}
-}
-
-void USequencerSettings::PostEditChangeProperty(FPropertyChangedEvent& PropertyChangedEvent)
-{
-	Super::PostEditChangeProperty(PropertyChangedEvent);
-}
-
 EAutoChangeMode USequencerSettings::GetAutoChangeMode() const
 {
 	return AutoChangeMode;
diff --git a/Engine/Source/Editor/Sequencer/Public/SequencerSettings.h b/Engine/Source/Editor/Sequencer/Public/SequencerSettings.h
--- a/Engine/Source/Editor/Sequencer/Public/SequencerSettings.h
+++ b/Engine/Source/Editor/Sequencer/Public/SequencerSettings.h
@@ -142,11 +142,6 @@ public:
 	DECLARE_MULTICAST_DELEGATE(FOnLoopStateChanged);
 	DECLARE_MULTICAST_DELEGATE(FOnTimeDisplayFormatChanged);
 
-	//~ Begin UObject
-	virtual void PostLoad() override;
-	UE_API virtual void PostEditChangeProperty(FPropertyChangedEvent& PropertyChangedEvent) override;
-	//~ End UObject
-
 	/** Gets the current auto change mode. */
 	UE_API EAutoChangeMode GetAutoChangeMode() const;
 	/** Sets the current auto change mode. */
@@ -528,9 +523,9 @@ public:
 	 */
 	UE_API void SetRestoreOriginalViewportOnCameraCutUnlock(bool bInRestoreOriginalViewportOnCameraCutUnlock);
 
-	/** Gets the tree view width percentage */
+	/** Gets the tree view width in pixels */
 	float GetTreeViewWidth() const { return TreeViewWidth; }
-	/** Sets the tree view width percentage */
+	/** Sets the tree view width in pixels */
 	UE_API void SetTreeViewWidth(float InTreeViewWidth);
 
 	/** Gets the saved view density */
@@ -906,8 +901,8 @@ protected:
 	UPROPERTY(config, EditAnywhere, Category=General)
 	bool bRestoreOriginalViewportOnCameraCutUnlock;
 
-	/** The tree view width percentage */
-	UPROPERTY(config, EditAnywhere, Category = General)
+	/** The tree view width in pixels */
+	UPROPERTY(config, EditAnywhere, Category = General, meta = (ClampMin = "150", UIMin = "150"))
 	float TreeViewWidth;
 
 	UPROPERTY(config, EditAnywhere, Category = General)
