[Sequencer]
- Enable selected filter on editor startup when it was last active and enabled (otherwise the user is forced to disable and re-enable the filter to display any selected)
- Selected filter now will filter control rig selections
- Fix base filter class to automatically call ActiveStateChanged() inside SetActive() to reduce code bloat

#Synced-CL 51185427

diff --git a/Engine/Plugins/Animation/ControlRig/Source/ControlRigEditor/Private/Sequencer/ControlRigSequencerFilter.cpp b/Engine/Plugins/Animation/ControlRig/Source/ControlRigEditor/Private/Sequencer/ControlRigSequencerFilter.cpp
--- a/Engine/Plugins/Animation/ControlRig/Source/ControlRigEditor/Private/Sequencer/ControlRigSequencerFilter.cpp
+++ b/Engine/Plugins/Animation/ControlRig/Source/ControlRigEditor/Private/Sequencer/ControlRigSequencerFilter.cpp
@@ -53,15 +53,16 @@ class FSequencerTrackFilter_ControlRigControls : public FSequencerTrackFilter
 public:
 	FSequencerTrackFilter_ControlRigControls(ISequencerTrackFilters& InOutFilterInterface, TSharedPtr<FFilterCategory> InCategory = nullptr)
 		: FSequencerTrackFilter(InOutFilterInterface, MoveTemp(InCategory))
-		, BindingCount(0)
 	{
-		FSequencerTrackFilter_ControlRigControlsCommands::Register();
+		if (++BindingCount == 1)
+		{
+			FSequencerTrackFilter_ControlRigControlsCommands::Register();
+		}
 	}
 
 	virtual ~FSequencerTrackFilter_ControlRigControls() override
 	{
-		BindingCount--;
-		if (BindingCount < 1)
+		if (ensure(BindingCount > 0) && --BindingCount == 0)
 		{
 			FSequencerTrackFilter_ControlRigControlsCommands::Unregister();
 		}
@@ -75,8 +76,7 @@ public:
 	{
 		FSequencerFilterData& FilterData = GetFilterInterface().GetFilterData();
 		const UMovieSceneTrack* const TrackObject = FilterData.ResolveMovieSceneTrackObject(InItem);
-		const UMovieSceneControlRigParameterTrack* const Track = Cast<UMovieSceneControlRigParameterTrack>(TrackObject);
-		return IsValid(Track);
+		return TrackObject && TrackObject->IsA<UMovieSceneControlRigParameterTrack>();
 	}
 
 	//~ End IFilter
@@ -106,9 +106,11 @@ public:
 	//~ End FSequencerTrackFilter
 
 private:
-	mutable uint32 BindingCount;
+	static int32 BindingCount;
 };
 
+int32 FSequencerTrackFilter_ControlRigControls::BindingCount = 0;
+
 //////////////////////////////////////////////////////////////////////////
 //
 
@@ -142,15 +144,16 @@ class FSequencerTrackFilter_ControlRigSelectedControls : public FSequencerTrackF
 public:
 	FSequencerTrackFilter_ControlRigSelectedControls(ISequencerTrackFilters& InOutFilterInterface, TSharedPtr<FFilterCategory> InCategory = nullptr)
 		:  FSequencerTrackFilter(InOutFilterInterface, MoveTemp(InCategory))
-		, BindingCount(0)
 	{
-		FSequencerTrackFilter_ControlRigSelectedControlsCommands::Register();
+		if (++BindingCount == 1)
+		{
+			FSequencerTrackFilter_ControlRigSelectedControlsCommands::Register();
+		}
 	}
 
 	virtual ~FSequencerTrackFilter_ControlRigSelectedControls() override
 	{
-		BindingCount--;
-		if (BindingCount < 1)
+		if (ensure(BindingCount > 0) && --BindingCount == 0)
 		{
 			FSequencerTrackFilter_ControlRigSelectedControlsCommands::Unregister();
 		}
@@ -256,19 +259,10 @@ public:
 	}
 
 private:
-	mutable uint32 BindingCount;
+	static int32 BindingCount;
 };
 
-/*
-
-	Bool,
-	Float,
-	Vector2D,
-	Position,
-	Scale,
-	Rotator,
-	Transform
-*/
+int32 FSequencerTrackFilter_ControlRigSelectedControls::BindingCount = 0;
 
 //////////////////////////////////////////////////////////////////////////
 //
@@ -278,7 +272,10 @@ void UControlRigTrackFilter::AddTrackFilterExtensions(ISequencerTrackFilters& In
 	, TArray<TSharedRef<FSequencerTrackFilter>>& InOutFilterList) const
 {
 	InOutFilterList.Add(MakeShared<FSequencerTrackFilter_ControlRigControls>(InOutFilterInterface, InPreferredCategory));
-	InOutFilterList.Add(MakeShared<FSequencerTrackFilter_ControlRigSelectedControls>(InOutFilterInterface, InPreferredCategory));
+	/** Removing this since FSequencerTrackFilter_Selected will now include selected control rig controls.
+	 * If there are no use cases for this in the future and no users are missing it, we could remove this eventually,
+	 * along with its class definition above. */
+	//InOutFilterList.Add(MakeShared<FSequencerTrackFilter_ControlRigSelectedControls>(InOutFilterInterface, InPreferredCategory));
 }
 
 #undef LOCTEXT_NAMESPACE
diff --git a/Engine/Source/Editor/Sequencer/Private/Filters/Filters/SequencerTrackFilter_Selected.cpp b/Engine/Source/Editor/Sequencer/Private/Filters/Filters/SequencerTrackFilter_Selected.cpp
--- a/Engine/Source/Editor/Sequencer/Private/Filters/Filters/SequencerTrackFilter_Selected.cpp
+++ b/Engine/Source/Editor/Sequencer/Private/Filters/Filters/SequencerTrackFilter_Selected.cpp
@@ -1,12 +1,14 @@
 ﻿// Copyright Epic Games, Inc. All Rights Reserved.
 
 #include "Filters/Filters/SequencerTrackFilter_Selected.h"
-#include "EditorModeManager.h"
-#include "Elements/Framework/TypedElementSelectionSet.h"
 #include "Filters/SequencerFilterBar.h"
 #include "Filters/SequencerTrackFilterCommands.h"
+#include "ISequencer.h"
+#include "MVVM/Selection/Selection.h"
+#include "MVVM/ViewModels/SequencerEditorViewModel.h"
 #include "Selection.h"
-#include "Sequencer.h"
+
+using namespace UE::Sequencer;
 
 #define LOCTEXT_NAMESPACE "SequencerTrackFilter_Selected"
 
@@ -37,28 +39,6 @@ void FSequencerTrackFilter_Selected::UnbindSelectionChanged()
 	}
 }
 
-TSharedPtr<ILevelEditor> FSequencerTrackFilter_Selected::GetLevelEditor() const
-{
-	const FLevelEditorModule* const LevelEditorModule = FModuleManager::GetModulePtr<FLevelEditorModule>(TEXT("LevelEditor"));
-	if (!LevelEditorModule)
-	{
-		return nullptr;
-	}
-
-	return LevelEditorModule->GetLevelEditorInstance().Pin();
-}
-
-FEditorModeTools* FSequencerTrackFilter_Selected::GetEditorModeManager() const
-{
-	const TSharedPtr<ILevelEditor> LevelEditor = GetLevelEditor();
-	if (!LevelEditor.IsValid())
-	{
-		return nullptr;
-	}
-
-	return &LevelEditor->GetEditorModeManager();
-}
-
 bool FSequencerTrackFilter_Selected::ShouldUpdateOnTrackValueChanged() const
 {
 	return true;
@@ -105,43 +85,32 @@ FString FSequencerTrackFilter_Selected::GetName() const
 
 bool FSequencerTrackFilter_Selected::PassesFilter(FSequencerTrackFilterType InItem) const
 {
-	const TSharedPtr<ILevelEditor> LevelEditor = GetLevelEditor();
-	if (!LevelEditor.IsValid())
+	const TSharedPtr<FSequencerEditorViewModel> ViewModel = GetSequencer().GetViewModel();
+	if (!ViewModel.IsValid())
 	{
 		return false;
 	}
 
-	const UTypedElementSelectionSet* const Selection = LevelEditor->GetElementSelectionSet();
-	if (!Selection)
+	const TSharedPtr<FSequencerSelection> SequencerSelection = ViewModel->GetSelection();
+	if (!SequencerSelection.IsValid())
 	{
 		return false;
 	}
 
-	FSequencerFilterData& FilterData = GetFilterInterface().GetFilterData();
-
-	const UObject* const TrackObject = FilterData.ResolveTrackBoundObject(GetSequencer(), InItem);
-	if (!TrackObject)
-	{
-		return false;
-	}
-
-	const TArray<UObject*> SelectedObjects = Selection->GetSelectedObjects<UObject>();
-
-	if (SelectedObjects.Contains(TrackObject))
-	{
-		return true;
-	}
-
-	const USceneComponent* const Component = TrackObject->GetTypedOuter<USceneComponent>();
-	if (Component && SelectedObjects.Contains(Component))
-	{
-		return true;
-	}
-
-	const AActor* const Actor = TrackObject->GetTypedOuter<AActor>();
-	if (Actor && SelectedObjects.Contains(Actor))
+	for (const FViewModelPtr& SelectedNode : SequencerSelection->Outliner)
 	{
-		return true;
+		if (!SelectedNode.IsValid())
+		{
+			continue;
+		}
+
+		for (FParentFirstChildIterator It(SelectedNode, /*bIncludeThis=*/true); It; ++It)
+		{
+			if (*It == InItem)
+			{
+				return true;
+			}
+		}
 	}
 
 	return false;
diff --git a/Engine/Source/Editor/Sequencer/Private/Filters/SequencerFilterBar.cpp b/Engine/Source/Editor/Sequencer/Private/Filters/SequencerFilterBar.cpp
--- a/Engine/Source/Editor/Sequencer/Private/Filters/SequencerFilterBar.cpp
+++ b/Engine/Source/Editor/Sequencer/Private/Filters/SequencerFilterBar.cpp
@@ -20,15 +20,11 @@
 #include "Filters/SequencerTrackFilterCommands.h"
 #include "Filters/SequencerTrackFilterExtension.h"
 #include "Filters/Utils/FilterViewUtils.h"
-#include "Filters/Widgets/SFilterBarIsolateHideShow.h"
-#include "Filters/Widgets/SFilterExpressionHelpDialog.h"
-#include "Filters/Widgets/SSequencerFilterBar.h"
 #include "Filters/Widgets/SSequencerSearchBox.h"
 #include "Framework/Application/SlateApplication.h"
 #include "Framework/MultiBox/SToolBarButtonBlock.h"
 #include "MovieScene.h"
 #include "MVVM/CurveEditorExtension.h"
-#include "MVVM/Selection/Selection.h"
 #include "MVVM/ViewModels/CategoryModel.h"
 #include "MVVM/Views/SOutlinerView.h"
 #include "Sequencer.h"
@@ -171,6 +167,33 @@ void FSequencerFilterBar::CreateDefaultFilters()
 
 	CommonFilters->Sort();
 
+	// Activate filters internally based on saved config
+	if (USequencerSettings* const SequencerSettings = Sequencer.GetSequencerSettings())
+	{
+		const FSequencerFilterBarConfig& Config = SequencerSettings->FindOrAddTrackFilterBar(GetIdentifier(), false);
+
+		auto ActivateFilterIfConfigActive = [this, &Config](const TSharedRef<FSequencerTrackFilter>& InFilter)
+			{
+				const FString FilterName = InFilter->GetDisplayName().ToString();
+				if (Config.IsFilterActive(FilterName))
+				{
+					InFilter->SetActive(true);
+				}
+			};
+
+		CommonFilters->ForEachFilter([ActivateFilterIfConfigActive](const TSharedRef<FSequencerTrackFilter>& InFilter)
+			{
+				ActivateFilterIfConfigActive(InFilter);
+				return true;
+			});
+		InternalFilters->ForEachFilter([ActivateFilterIfConfigActive](const TSharedRef<FSequencerTrackFilter>& InFilter)
+			{
+				ActivateFilterIfConfigActive(InFilter);
+				return true;
+			});
+	}
+
+	// Broadcast changes
 	CommonFilters->OnChanged().Broadcast();
 	InternalFilters->OnChanged().Broadcast();
 	TextFilter->OnChanged().Broadcast();
@@ -594,7 +617,6 @@ bool FSequencerFilterBar::SetFilterActive(const TSharedRef<FSequencerTrackFilter
 		SequencerSettings->SaveConfig();
 
 		InFilter->SetActive(bNewActive);
-		InFilter->ActiveStateChanged(bNewActive);
 
 		const ESequencerFilterChange FilterChangeType = bNewActive
 			? ESequencerFilterChange::Activate : ESequencerFilterChange::Deactivate;
@@ -684,7 +706,6 @@ void FSequencerFilterBar::ActivateCommonFilters(const bool bInActivate
 				FiltersChangedEvent.Broadcast(FilterChangeType, InFilter);
 
 				InFilter->SetActive(bInActivate);
-				InFilter->ActiveStateChanged(bInActivate);
 
 				bNeedsSave = true;
 			}
@@ -888,7 +909,6 @@ bool FSequencerFilterBar::SetFilterEnabled(const TSharedRef<FSequencerTrackFilte
 		if (!bInEnabled && IsFilterActive(InFilter))
 		{
 			InFilter->SetActive(false);
-			InFilter->ActiveStateChanged(false);
 		}
 
 		if (bInRequestFilterUpdate)
@@ -927,7 +947,6 @@ void FSequencerFilterBar::EnableFilters(const bool bInEnable
 					if (!bInEnable && IsFilterActive(InFilter))
 					{
 						InFilter->SetActive(false);
-						InFilter->ActiveStateChanged(false);
 					}
 				}
 			}
@@ -1120,7 +1139,6 @@ void FSequencerFilterBar::ActivateCustomTextFilters(const bool bInActivate, cons
 				if (!bInActivate && IsFilterActive(CustomTextFilter))
 				{
 					CustomTextFilter->SetActive(false);
-					CustomTextFilter->ActiveStateChanged(false);
 				}
 
 				const ESequencerFilterChange FilterChangeType = bInActivate
@@ -1155,7 +1173,6 @@ void FSequencerFilterBar::EnableCustomTextFilters(const bool bInEnable, const TA
 				if (!bInEnable && IsFilterActive(CustomTextFilter))
 				{
 					CustomTextFilter->SetActive(false);
-					CustomTextFilter->ActiveStateChanged(false);
 				}
 
 				const ESequencerFilterChange FilterChangeType = bInEnable
diff --git a/Engine/Source/Editor/Sequencer/Public/Filters/Filters/SequencerTrackFilter_Selected.h b/Engine/Source/Editor/Sequencer/Public/Filters/Filters/SequencerTrackFilter_Selected.h
--- a/Engine/Source/Editor/Sequencer/Public/Filters/Filters/SequencerTrackFilter_Selected.h
+++ b/Engine/Source/Editor/Sequencer/Public/Filters/Filters/SequencerTrackFilter_Selected.h
@@ -4,9 +4,6 @@
 
 #include "Filters/SequencerTrackFilterBase.h"
 
-class FEditorModeTools;
-class ILevelEditor;
-
 class FSequencerTrackFilter_Selected : public FSequencerTrackFilter
 {
 public:
@@ -38,9 +35,6 @@ protected:
 	void BindSelectionChanged();
 	void UnbindSelectionChanged();
 
-	TSharedPtr<ILevelEditor> GetLevelEditor() const;
-	FEditorModeTools* GetEditorModeManager() const;
-
 	void OnSelectionChanged(UObject* const InObject);
 
 	FDelegateHandle OnSelectionChangedHandle;
diff --git a/Engine/Source/Editor/Sequencer/Public/Filters/SequencerFilterBase.h b/Engine/Source/Editor/Sequencer/Public/Filters/SequencerFilterBase.h
--- a/Engine/Source/Editor/Sequencer/Public/Filters/SequencerFilterBase.h
+++ b/Engine/Source/Editor/Sequencer/Public/Filters/SequencerFilterBase.h
@@ -105,6 +105,12 @@ public:
 		}
 	}
 
+	virtual void SetActive(const bool bInActive)
+	{
+		FFilterBase<InFilterType>::SetActive(bInActive);
+		ActiveStateChanged(bInActive);
+	}
+
 	virtual TSharedPtr<FUICommandInfo> GetToggleCommand() const { return nullptr; }
 
 	bool CanToggleFilter() const
