Sequencer: Add option to limit playhead scrubbing to selection range

The original option now has 3 possibilities: None, Playback Range, Selection Range.

#Synced-CL 51192822

diff --git a/Engine/Config/BaseEditorPerProjectUserSettings.ini b/Engine/Config/BaseEditorPerProjectUserSettings.ini
--- a/Engine/Config/BaseEditorPerProjectUserSettings.ini
+++ b/Engine/Config/BaseEditorPerProjectUserSettings.ini
@@ -904,7 +904,6 @@ bCompileDirectorOnEvaluate=false
 [NiagaraSequenceEditor SequencerSettings]
 bAutoScrollEnabled=true
 bKeepPlayRangeInSectionBounds=false
-bKeepCursorInPlayRange=false
 bShowRangeSlider=true
 LoopMode=SLM_Loop
 bCleanPlaybackMode=false
diff --git a/Engine/Plugins/MovieScene/LevelSequenceEditor/Source/LevelSequenceEditor/Private/CinematicViewport/SCinematicTransportRange.cpp b/Engine/Plugins/MovieScene/LevelSequenceEditor/Source/LevelSequenceEditor/Private/CinematicViewport/SCinematicTransportRange.cpp
--- a/Engine/Plugins/MovieScene/LevelSequenceEditor/Source/LevelSequenceEditor/Private/CinematicViewport/SCinematicTransportRange.cpp
+++ b/Engine/Plugins/MovieScene/LevelSequenceEditor/Source/LevelSequenceEditor/Private/CinematicViewport/SCinematicTransportRange.cpp
@@ -89,9 +89,10 @@ void SCinematicTransportRange::SetTime(const FGeometry& MyGeometry, const FPoint
 		FFrameTime ScrubTime = NewTimeSeconds * TickResolution;
 
 		// Clamp first, snap to frame last
-		if (Sequencer->GetSequencerSettings()->ShouldKeepCursorInPlayRangeWhileScrubbing())
+		TRange<FFrameNumber> PlayheadScrubbingRange = Sequencer->GetSequencerSettings()->GetPlayheadScrubbingRange(Sequencer->GetRootMovieSceneSequence()->GetMovieScene());
+		if (!PlayheadScrubbingRange.IsEmpty())
 		{
-			TRange<FFrameNumber> PlaybackRange = Sequencer->GetSubSequenceRange().Get(Sequencer->GetRootMovieSceneSequence()->GetMovieScene()->GetPlaybackRange());
+			TRange<FFrameNumber> PlaybackRange = Sequencer->GetSubSequenceRange().Get(PlayheadScrubbingRange);
 			ScrubTime = UE::MovieScene::ClampToDiscreteRange(ScrubTime, PlaybackRange);
 		}
 
diff --git a/Engine/Source/Editor/Sequencer/Private/SSequencer.h b/Engine/Source/Editor/Sequencer/Private/SSequencer.h
--- a/Engine/Source/Editor/Sequencer/Private/SSequencer.h
+++ b/Engine/Source/Editor/Sequencer/Private/SSequencer.h
@@ -448,6 +448,9 @@ private:
 	/** Makes the column visibility menu for the toolbar. */
 	void FillColumnVisibilityMenu(FMenuBuilder& InMenuBuilder);
 
+	/** Makes the playhead scrubbing limit menu for the toolbar. */
+	void FillPlayheadScrubbingLimitMenu(FMenuBuilder& InMenuBuilder);
+
 	/** Return the current sequencer settings */ 
 	USequencerSettings* GetSequencerSettings() const;
 
diff --git a/Engine/Source/Editor/Sequencer/Private/Sequencer.cpp b/Engine/Source/Editor/Sequencer/Private/Sequencer.cpp
--- a/Engine/Source/Editor/Sequencer/Private/Sequencer.cpp
+++ b/Engine/Source/Editor/Sequencer/Private/Sequencer.cpp
@@ -3729,9 +3729,10 @@ void FSequencer::PlayTo(FMovieSceneSequencePlaybackParams PlaybackParams)
 void FSequencer::SnapSequencerTime(FFrameTime& ScrubTime)
 {
 	// Clamp first, snap to frame last
-	if (GetSequencerSettings()->ShouldKeepCursorInPlayRangeWhileScrubbing())
+	TRange<FFrameNumber> PlayheadScrubbingRange = GetSequencerSettings()->GetPlayheadScrubbingRange(GetRootMovieSceneSequence()->GetMovieScene());
+	if (!PlayheadScrubbingRange.IsEmpty())
 	{
-		TRange<FFrameNumber> PlaybackRange = GetSubSequenceRange().Get(GetRootMovieSceneSequence()->GetMovieScene()->GetPlaybackRange());
+		TRange<FFrameNumber> PlaybackRange = GetSubSequenceRange().Get(PlayheadScrubbingRange);
 		ScrubTime = UE::MovieScene::ClampToDiscreteRange(ScrubTime, PlaybackRange);
 	}
 
@@ -11954,10 +11955,22 @@ void FSequencer::BindCommands()
 		FIsActionChecked::CreateLambda( [this]{ return this->UsesLoopCuts(); } ) );
 
 	SequencerCommandBindings->MapAction(
-		Commands.ToggleKeepCursorInPlaybackRangeWhileScrubbing,
-		FExecuteAction::CreateLambda([this] { Settings->SetKeepCursorInPlayRangeWhileScrubbing(!Settings->ShouldKeepCursorInPlayRangeWhileScrubbing()); }),
+		Commands.SetPlayheadScrubbingLimitNone,
+		FExecuteAction::CreateLambda([this] { Settings->SetPlayheadScrubbingLimit(ESequencerPlayheadScrubbingLimit::SPSL_None); }),
 		FCanExecuteAction::CreateLambda([] { return true; }),
-		FIsActionChecked::CreateLambda([this] { return Settings->ShouldKeepCursorInPlayRangeWhileScrubbing(); }));
+		FIsActionChecked::CreateLambda([this] { return Settings->GetPlayheadScrubbingLimit() == ESequencerPlayheadScrubbingLimit::SPSL_None; }));
+	
+	SequencerCommandBindings->MapAction(
+		Commands.SetPlayheadScrubbingLimitPlaybackRange,
+		FExecuteAction::CreateLambda([this] { Settings->SetPlayheadScrubbingLimit(ESequencerPlayheadScrubbingLimit::SPSL_PlaybackRange); }),
+		FCanExecuteAction::CreateLambda([] { return true; }),
+		FIsActionChecked::CreateLambda([this] { return Settings->GetPlayheadScrubbingLimit() == ESequencerPlayheadScrubbingLimit::SPSL_PlaybackRange; }));
+
+	SequencerCommandBindings->MapAction(
+		Commands.SetPlayheadScrubbingLimitSelectionRange,
+		FExecuteAction::CreateLambda([this] { Settings->SetPlayheadScrubbingLimit(ESequencerPlayheadScrubbingLimit::SPSL_SelectionRange); }),
+		FCanExecuteAction::CreateLambda([] { return true; }),
+		FIsActionChecked::CreateLambda([this] { return Settings->GetPlayheadScrubbingLimit() == ESequencerPlayheadScrubbingLimit::SPSL_SelectionRange; }));
 
 	SequencerCommandBindings->MapAction(
 		Commands.ToggleResetPlayheadWhenNavigating,
diff --git a/Engine/Source/Editor/Sequencer/Private/SequencerCommands.cpp b/Engine/Source/Editor/Sequencer/Private/SequencerCommands.cpp
--- a/Engine/Source/Editor/Sequencer/Private/SequencerCommands.cpp
+++ b/Engine/Source/Editor/Sequencer/Private/SequencerCommands.cpp
@@ -57,9 +57,12 @@ void FSequencerCommands::RegisterCommands()
 	UI_COMMAND( ToggleLoopCuts, "Loop Cuts", "When enabled, the sequence considers looping as a jump/cut. When disabled, the sequence considers looping to be seamless.", EUserInterfaceActionType::ToggleButton, FInputChord());
 
 	UI_COMMAND( ToggleResetPlayheadWhenNavigating, "Reset Playhead When Navigating", "When checked, if the playhead is outside the playback range, it will be reset to the beginning when navigating in and out of subsequences", EUserInterfaceActionType::ToggleButton, FInputChord());
-	UI_COMMAND( ToggleKeepCursorInPlaybackRangeWhileScrubbing, "Keep Playhead in Playback Range While Scrubbing", "When checked, the playhead will be constrained to the current playback (or subsequence/shot) range while scrubbing", EUserInterfaceActionType::ToggleButton, FInputChord() );
 	UI_COMMAND( ToggleKeepPlaybackRangeInSectionBounds, "Keep Playback Range in Section Bounds", "When checked, the playback range will be synchronized to the section bounds", EUserInterfaceActionType::ToggleButton, FInputChord() );
 
+	UI_COMMAND( SetPlayheadScrubbingLimitNone, "None", "While scrubbing, don't restrict the playhead to any range", EUserInterfaceActionType::RadioButton, FInputChord() );
+	UI_COMMAND( SetPlayheadScrubbingLimitPlaybackRange, "Playback Range", "While scrubbing, restrict the playhead to the playback range", EUserInterfaceActionType::RadioButton, FInputChord() );
+	UI_COMMAND( SetPlayheadScrubbingLimitSelectionRange, "Selection Range", "While scrubbing, restrict the playhead to the selection range", EUserInterfaceActionType::RadioButton, FInputChord() );
+
 	UI_COMMAND( ToggleAutoExpandNodesOnSelection, "Auto Expand Nodes on Selection", "Toggle auto expanding the outliner tree on child selection", EUserInterfaceActionType::ToggleButton, FInputChord());
 	UI_COMMAND( ToggleRestoreOriginalViewportOnCameraCutUnlock, "Restore Cinematic Viewports", "Toggle whether unlocking a camera cut track should return the viewport to its original location, or keep it where the camera cut was", EUserInterfaceActionType::ToggleButton, FInputChord() );
 	UI_COMMAND( TogglePreviewCameraCutsInSimulate, "Preview Camera Cuts in Simulate", "Toggle whether camera cuts are evaluated in SIE, or after ejecting from the player controller in PIE", EUserInterfaceActionType::ToggleButton, FInputChord() );
diff --git a/Engine/Source/Editor/Sequencer/Private/SequencerSettings.cpp b/Engine/Source/Editor/Sequencer/Private/SequencerSettings.cpp
--- a/Engine/Source/Editor/Sequencer/Private/SequencerSettings.cpp
+++ b/Engine/Source/Editor/Sequencer/Private/SequencerSettings.cpp
@@ -5,6 +5,7 @@
 #include "ISequencer.h"
 #include "SSequencer.h"
 #include "Misc/ConsoleVariables.h"
+#include "MovieScene.h"
 #include "MVVM/ViewModels/ViewDensity.h"
 
 #include UE_INLINE_GENERATED_CPP_BY_NAME(SequencerSettings)
@@ -43,7 +44,7 @@ USequencerSettings::USequencerSettings( const FObjectInitializer& ObjectInitiali
 	bCurveEditorVisible = false;
 	LoopMode = ESequencerLoopMode::SLM_NoLoop;
 	bResetPlayheadWhenNavigating = false;
-	bKeepCursorInPlayRangeWhileScrubbing = false;
+	PlayheadScrubbingLimit = ESequencerPlayheadScrubbingLimit::SPSL_None;
 	bKeepPlayRangeInSectionBounds = true;
 	bCompileDirectorOnEvaluate = true;
 	bLeftMouseDragDoesMarquee = false;
@@ -509,20 +510,33 @@ void USequencerSettings::SetResetPlayheadWhenNavigating(bool bInResetPlayheadWhe
 	}
 }
 
-bool USequencerSettings::ShouldKeepCursorInPlayRangeWhileScrubbing() const
+ESequencerPlayheadScrubbingLimit USequencerSettings::GetPlayheadScrubbingLimit() const
 {
-	return bKeepCursorInPlayRangeWhileScrubbing;
+	return PlayheadScrubbingLimit;
 }
 
-void USequencerSettings::SetKeepCursorInPlayRangeWhileScrubbing(bool bInKeepCursorInPlayRangeWhileScrubbing)
+void USequencerSettings::SetPlayheadScrubbingLimit(ESequencerPlayheadScrubbingLimit InPlayheadScrubbingLimit)
 {
-	if (bKeepCursorInPlayRangeWhileScrubbing != bInKeepCursorInPlayRangeWhileScrubbing)
+	if (PlayheadScrubbingLimit != InPlayheadScrubbingLimit)
 	{
-		bKeepCursorInPlayRangeWhileScrubbing = bInKeepCursorInPlayRangeWhileScrubbing;
+		PlayheadScrubbingLimit = InPlayheadScrubbingLimit;
 		SaveConfig();
 	}
 }
 
+TRange<FFrameNumber> USequencerSettings::GetPlayheadScrubbingRange(UMovieScene* MovieScene) const
+{
+	if (PlayheadScrubbingLimit == ESequencerPlayheadScrubbingLimit::SPSL_PlaybackRange)
+	{
+		return MovieScene->GetPlaybackRange();
+	}
+	else if (PlayheadScrubbingLimit == ESequencerPlayheadScrubbingLimit::SPSL_SelectionRange)
+	{
+		return MovieScene->GetSelectionRange();
+	}
+	return TRange<FFrameNumber>::Empty();
+}
+
 bool USequencerSettings::ShouldKeepPlayRangeInSectionBounds() const
 {
 	return bKeepPlayRangeInSectionBounds;
diff --git a/Engine/Source/Editor/Sequencer/Public/SequencerCommands.h b/Engine/Source/Editor/Sequencer/Public/SequencerCommands.h
--- a/Engine/Source/Editor/Sequencer/Public/SequencerCommands.h
+++ b/Engine/Source/Editor/Sequencer/Public/SequencerCommands.h
@@ -150,8 +150,14 @@ public:
 	/** Toggle resetting the playhead when navigating in and out of subsequences */
 	TSharedPtr< FUICommandInfo > ToggleResetPlayheadWhenNavigating;
 
-	/** Toggle constraining the time cursor to the playback range while scrubbing */
-	TSharedPtr< FUICommandInfo > ToggleKeepCursorInPlaybackRangeWhileScrubbing;
+	/** While scrubbing, don't restrict the playhead to any range */
+	TSharedPtr< FUICommandInfo > SetPlayheadScrubbingLimitNone;
+
+	/** While scrubbing, restrict the playhead to the playback range */
+	TSharedPtr< FUICommandInfo > SetPlayheadScrubbingLimitPlaybackRange;
+
+	/** While scrubbing, restrict the playhead to the selection range */
+	TSharedPtr< FUICommandInfo > SetPlayheadScrubbingLimitSelectionRange;
 
 	/** Toggle constraining the playback range to the section bounds */
 	TSharedPtr< FUICommandInfo > ToggleKeepPlaybackRangeInSectionBounds;
diff --git a/Engine/Source/Editor/Sequencer/Public/SequencerSettings.h b/Engine/Source/Editor/Sequencer/Public/SequencerSettings.h
--- a/Engine/Source/Editor/Sequencer/Public/SequencerSettings.h
+++ b/Engine/Source/Editor/Sequencer/Public/SequencerSettings.h
@@ -59,6 +59,16 @@ enum ESequencerLoopMode : int
 	SLM_LoopSelectionRange UMETA(DisplayName="Loop Selection Range"),
 };
 
+UENUM()
+enum ESequencerPlayheadScrubbingLimit : int
+{
+	SPSL_None UMETA(DisplayName="None"),
+
+	SPSL_PlaybackRange UMETA(DisplayName="Playback Range"),
+
+	SPSL_SelectionRange UMETA(DisplayName="Selection Range")
+};
+
 UENUM()
 enum class ESequencerTimeWarpDisplay : uint8
 {
@@ -316,10 +326,12 @@ public:
 	/** Set whether or not the cursor should be reset when navigating in and out of subsequences */
 	UE_API void SetResetPlayheadWhenNavigating(bool bInResetPlayheadWhenNavigating);
 
-	/** @return true if the cursor should be kept within the playback (or subsequence/shot) range while scrubbing in sequencer, false otherwise */
-	UE_API bool ShouldKeepCursorInPlayRangeWhileScrubbing() const;
-	/** Set whether or not the cursor should be kept within the playback (or subsequence/shot) range while scrubbing in sequencer */
-	UE_API void SetKeepCursorInPlayRangeWhileScrubbing(bool bInKeepCursorInPlayRangeWhileScrubbing);
+	/** @return the limit that the playhead should be restricted to when scrubbing */
+	UE_API ESequencerPlayheadScrubbingLimit GetPlayheadScrubbingLimit() const;
+	/** Set the limit that the playhead should be restricted to when scrubbing */
+	UE_API void SetPlayheadScrubbingLimit(ESequencerPlayheadScrubbingLimit InPlayheadScrubbingLimit);
+	/** Get the actual frame number range that the playhead should be restricted to when scrubbing */
+	UE_API TRange<FFrameNumber> GetPlayheadScrubbingRange(UMovieScene* MovieScene) const;
 
 	/** @return true if the playback range should be synced to the section bounds, false otherwise */
 	UE_API bool ShouldKeepPlayRangeInSectionBounds() const;
@@ -591,7 +603,7 @@ public:
 	UE_API void SetNavigationToolVisible(const bool bInVisible);
 
 	/** @return True if the selection details are popped out from the Sequencer tab to their own tab */
-	UE_API bool IsSelectionDetailsPoppedOut() const { return bSelectionDetailsPoppedOut; }
+	bool IsSelectionDetailsPoppedOut() const { return bSelectionDetailsPoppedOut; }
 	/** Sets the popped out state of the selection details in the Sequencer tab */
 	UE_API void SetSelectionDetailsPoppedOut(const bool bInPoppedOut);
 
@@ -749,9 +761,9 @@ protected:
 	UPROPERTY(config, EditAnywhere, Category = Timeline, meta = (DisplayName = "Reset Playhead When Navigating"))
 	bool bResetPlayheadWhenNavigating;
 
-	/** Enable or disable keeping the playhead in the current playback range while scrubbing. */
-	UPROPERTY(config, EditAnywhere, Category = Timeline, meta = (DisplayName = "Keep Playhead in Play Range While Scrubbing"))
-	bool bKeepCursorInPlayRangeWhileScrubbing;
+	/** The limit that the playhead should be restricted to when scrubbing */
+	UPROPERTY(config, EditAnywhere, Category = Timeline, meta = (DisplayName = "Playhead Scrubbing Limit"))
+	TEnumAsByte<ESequencerPlayheadScrubbingLimit> PlayheadScrubbingLimit;
 
 	/** Enable or disable keeping the playback range constrained to the section bounds. */
 	UPROPERTY( config, EditAnywhere, Category=Timeline )
