[Sequencer] Set default key for "Selected" filter to F10. Previously there was no default. This will easily transition people using the "Selected Control Rigs Controls" filter, which currently has the F10 hotkey but has been removed, since the "Selected" filter now filters for control rig controls.

#Synced-CL 51185427

diff --git a/Engine/Source/Editor/Sequencer/Private/Filters/SequencerTrackFilterCommands.cpp b/Engine/Source/Editor/Sequencer/Private/Filters/SequencerTrackFilterCommands.cpp
--- a/Engine/Source/Editor/Sequencer/Private/Filters/SequencerTrackFilterCommands.cpp
+++ b/Engine/Source/Editor/Sequencer/Private/Filters/SequencerTrackFilterCommands.cpp
@@ -65,7 +65,7 @@ void FSequencerTrackFilterCommands::RegisterCommands()
 	UI_COMMAND(ToggleFilter_Condition, "Toggle Condition Filter", "Toggle the filter for tracks with conditions", EUserInterfaceActionType::Button, FInputChord());
 	UI_COMMAND(ToggleFilter_Keyed, "Toggle Keyed Filter", "Toggle the filter for Keyed tracks", EUserInterfaceActionType::Button, FInputChord(EKeys::U));
 	UI_COMMAND(ToggleFilter_Modified, "Toggle Modified Filter", "Toggle the filter for Modified tracks", EUserInterfaceActionType::Button, FInputChord());
-	UI_COMMAND(ToggleFilter_Selected, "Toggle Selected In Viewport Filter", "Toggle the filter for Selected In Viewport tracks", EUserInterfaceActionType::Button, FInputChord());
+	UI_COMMAND(ToggleFilter_Selected, "Toggle Selected In Viewport Filter", "Toggle the filter for Selected In Viewport tracks", EUserInterfaceActionType::Button, FInputChord(EKeys::F10));
 	UI_COMMAND(ToggleFilter_Unbound, "Toggle Unbound Filter", "Toggle the filter for Unbound tracks", EUserInterfaceActionType::Button, FInputChord());
 
 	UI_COMMAND(ToggleFilter_Groups, "Toggle Groups Filter", "Toggle the filter for Group tracks", EUserInterfaceActionType::Button, FInputChord());
