Sequencer: Add nullptr checks

#Synced-CL 51192822

diff --git a/Engine/Source/Editor/Sequencer/Private/MVVM/ViewModels/OutlinerItemModel.cpp b/Engine/Source/Editor/Sequencer/Private/MVVM/ViewModels/OutlinerItemModel.cpp
--- a/Engine/Source/Editor/Sequencer/Private/MVVM/ViewModels/OutlinerItemModel.cpp
+++ b/Engine/Source/Editor/Sequencer/Private/MVVM/ViewModels/OutlinerItemModel.cpp
@@ -97,7 +97,12 @@ void FOutlinerItemModelMixin::AddEvalOptionsPropertyMenuItem(FMenuBuilder& InMen
 			}),
 			FCanExecuteAction::CreateLambda([this]
 			{
-				const TSharedPtr<FSequencer> Sequencer = GetEditor()->GetSequencerImpl();
+				TSharedPtr<FSequencerEditorViewModel> EditorViewModel = GetEditor();
+				if (!EditorViewModel)
+				{
+					return false;
+				}
+				const TSharedPtr<FSequencer> Sequencer = EditorViewModel->GetSequencerImpl();
 				if (Sequencer)
 				{
 					return !Sequencer->IsReadOnly();
@@ -147,7 +152,12 @@ void FOutlinerItemModelMixin::AddDisplayOptionsPropertyMenuItem(FMenuBuilder& In
 			}),
 			FCanExecuteAction::CreateLambda([this]
 			{
-				const TSharedPtr<FSequencer> Sequencer = GetEditor()->GetSequencerImpl();
+				TSharedPtr<FSequencerEditorViewModel> EditorViewModel = GetEditor();
+				if (!EditorViewModel)
+				{
+					return false;
+				}
+				const TSharedPtr<FSequencer> Sequencer = EditorViewModel->GetSequencerImpl();
 				if (Sequencer)
 				{
 					return !Sequencer->IsReadOnly();
@@ -382,12 +392,24 @@ void FOutlinerItemModelMixin::ToggleRootModelPinned()
 
 ECheckBoxState FOutlinerItemModelMixin::SelectedModelsSoloState() const
 {
+	TSharedPtr<FSequencerEditorViewModel> EditorViewModel = GetEditor();
+	if (!EditorViewModel)
+	{
+		return ECheckBoxState::Unchecked;
+	}
+
+	const TSharedPtr<FSequencerSelection> Selection = EditorViewModel->GetSelection();
+	if (!Selection.IsValid())
+	{
+		return ECheckBoxState::Unchecked;
+	}
+
 	FSoloStateCacheExtension* SoloStateCache = AsViewModel()->GetSharedData()->CastThis<FSoloStateCacheExtension>();
 	check(SoloStateCache);
 
 	int32 NumSoloables = 0;
 	int32 NumSoloed = 0;
-	for (FViewModelPtr Soloable : GetEditor()->GetSelection()->Outliner.Filter<ISoloableExtension>())
+	for (FViewModelPtr Soloable : Selection->Outliner.Filter<ISoloableExtension>())
 	{
 		++NumSoloables;
 		if (EnumHasAnyFlags(SoloStateCache->GetCachedFlags(Soloable), ECachedSoloState::Soloed))
@@ -425,12 +447,24 @@ void FOutlinerItemModelMixin::ToggleSelectedModelsSolo()
 
 ECheckBoxState FOutlinerItemModelMixin::SelectedModelsMuteState() const
 {
+	TSharedPtr<FSequencerEditorViewModel> EditorViewModel = GetEditor();
+	if (!EditorViewModel)
+	{
+		return ECheckBoxState::Unchecked;
+	}
+
+	const TSharedPtr<FSequencerSelection> Selection = EditorViewModel->GetSelection();
+	if (!Selection.IsValid())
+	{
+		return ECheckBoxState::Unchecked;
+	}
+
 	FMuteStateCacheExtension* MuteStateCache = AsViewModel()->GetSharedData()->CastThis<FMuteStateCacheExtension>();
 	check(MuteStateCache);
 
 	int32 NumMutables = 0;
 	int32 NumMuted = 0;
-	for (FViewModelPtr Mutable : GetEditor()->GetSelection()->Outliner.Filter<IMutableExtension>())
+	for (FViewModelPtr Mutable : Selection->Outliner.Filter<IMutableExtension>())
 	{
 		++NumMutables;
 		if (EnumHasAnyFlags(MuteStateCache->GetCachedFlags(Mutable), ECachedMuteState::Muted))
@@ -469,19 +503,23 @@ void FOutlinerItemModelMixin::ToggleSelectedModelsMuted()
 TSharedPtr<SWidget> FOutlinerItemModelMixin::CreateContextMenuWidget(const FCreateOutlinerContextMenuWidgetParams& InParams)
 {
 	TSharedPtr<FSequencerEditorViewModel> EditorViewModel = GetEditor();
-	TSharedPtr<ISequencer> Sequencer = EditorViewModel->GetSequencer();
+	if (!EditorViewModel)
+	{
+		return nullptr;
+	}
 
-	if (Sequencer)
+	TSharedPtr<ISequencer> Sequencer = EditorViewModel->GetSequencer();
+	if (!Sequencer)
 	{
-		const bool bShouldCloseWindowAfterMenuSelection = true;
-		FMenuBuilder MenuBuilder(bShouldCloseWindowAfterMenuSelection, Sequencer->GetCommandBindings());
+		return nullptr;
+	}
 
-		BuildContextMenu(MenuBuilder);
+	const bool bShouldCloseWindowAfterMenuSelection = true;
+	FMenuBuilder MenuBuilder(bShouldCloseWindowAfterMenuSelection, Sequencer->GetCommandBindings());
 
-		return MenuBuilder.MakeWidget();
-	}
+	BuildContextMenu(MenuBuilder);
 
-	return nullptr;
+	return MenuBuilder.MakeWidget();
 }
 
 FSlateColor FOutlinerItemModelMixin::GetLabelColor() const
@@ -511,8 +549,13 @@ FSlateColor FOutlinerItemModelMixin::GetLabelColor() const
 
 void FOutlinerItemModelMixin::BuildContextMenu(FMenuBuilder& MenuBuilder)
 {
-	TSharedPtr<FSequencer> Sequencer = StaticCastSharedPtr<FSequencer>(GetEditor()->GetSequencer());
+	TSharedPtr<FSequencerEditorViewModel> EditorViewModel = GetEditor();
+	if (!EditorViewModel)
+	{
+		return;
+	}
 
+	TSharedPtr<FSequencer> Sequencer = StaticCastSharedPtr<FSequencer>(EditorViewModel->GetSequencer());
 	if (!Sequencer)
 	{
 		return;
@@ -1281,7 +1324,15 @@ void FEvaluableOutlinerItemModel::SetIsDeactivated(const bool bInIsDeactivated)
 
 	if (bAnyChanged)
 	{
-		GetEditor()->GetSequencer()->NotifyMovieSceneDataChanged(EMovieSceneDataChangeType::TrackValueChanged);
+		TSharedPtr<FSequencerEditorViewModel> EditorViewModel = GetEditor();
+		if (EditorViewModel)
+		{
+			TSharedPtr<ISequencer> Sequencer = EditorViewModel->GetSequencer();
+			if (Sequencer)
+			{
+				Sequencer->NotifyMovieSceneDataChanged(EMovieSceneDataChangeType::TrackValueChanged);
+			}
+		}
 	}
 }
 
