Sequencer: Update audio sections to not auto expand to keyframes

#Synced-CL 51229281

diff --git a/Engine/Source/Editor/Sequencer/Private/IKeyArea.cpp b/Engine/Source/Editor/Sequencer/Private/IKeyArea.cpp
--- a/Engine/Source/Editor/Sequencer/Private/IKeyArea.cpp
+++ b/Engine/Source/Editor/Sequencer/Private/IKeyArea.cpp
@@ -96,6 +96,14 @@ FKeyHandle IKeyArea::AddOrUpdateKey(FFrameNumber Time, const FGuid& ObjectBindin
 	ISequencerChannelInterface* EditorInterface = FindChannelEditorInterface();
 	FMovieSceneChannel* Channel = ChannelHandle.Get();
 	UMovieSceneSection* Section = GetOwningSection();
+	check(Section != nullptr);
+	UMovieSceneTrack* Track = CastChecked<UMovieSceneTrack>(Section->GetOuter());
+
+	// If track does not allow auto expanding sections, check if key time is within bounds. We include exclusive end frame to allow keyframes to visually align with section rendering
+	if (!Track->bAutoExpandToKeyframes && (Time < Section->GetInclusiveStartFrame() || Time > Section->GetExclusiveEndFrame()))
+	{
+		return FKeyHandle();
+	}
 
 	if (const FMovieSceneChannelMetaData* MetaData = ChannelHandle.GetMetaData())
 	{
diff --git a/Engine/Source/Editor/Sequencer/Private/Tools/EditToolDragOperations.cpp b/Engine/Source/Editor/Sequencer/Private/Tools/EditToolDragOperations.cpp
--- a/Engine/Source/Editor/Sequencer/Private/Tools/EditToolDragOperations.cpp
+++ b/Engine/Source/Editor/Sequencer/Private/Tools/EditToolDragOperations.cpp
@@ -1376,7 +1376,20 @@ void FMoveKeysAndSections::HandleKeyMovement(const FPointerEvent& MouseEvent, TO
 		// transformed (moving the section moves the keys + we add the delta to the key positions).
 		if (!bOwningSectionIsSelected)
 		{
-			CurrentKeyTimes[Index] += MaxDeltaX.Get(DesiredDeltaX);
+			const FFrameNumber NewKeyTime = CurrentKeyTimes[Index] + MaxDeltaX.Get(DesiredDeltaX);
+
+			const UMovieSceneTrack* Track = SelectedKey.Section->GetTypedOuter<UMovieSceneTrack>();
+			check(Track);
+
+			if (Track->bAutoExpandToKeyframes) // Auto expand allows keyframes to be dragged outside sections and sections will resize to keyframes
+			{
+				CurrentKeyTimes[Index] = NewKeyTime;
+			}
+			// Without auto expand, ensure that keyframe is within section bounds. We include exclusive end frame to allow keyframes to visually align with section rendering
+			else if (NewKeyTime <= SelectedKey.Section->GetExclusiveEndFrame() && NewKeyTime >= SelectedKey.Section->GetInclusiveStartFrame())
+			{
+				CurrentKeyTimes[Index] = NewKeyTime;
+			}
 		}
 	}
 
@@ -1396,7 +1409,10 @@ void FMoveKeysAndSections::HandleKeyMovement(const FPointerEvent& MouseEvent, TO
 			const FFrameNumber   NewKeyTime = CurrentKeyTimes[Index];
 			TRange<FFrameNumber> SectionRange = Section->GetRange();
 
-			if (!SectionRange.Contains(NewKeyTime))
+			const UMovieSceneTrack* Track = Section->GetTypedOuter<UMovieSceneTrack>();
+			check(Track);
+
+			if (!SectionRange.Contains(NewKeyTime) && Track->bAutoExpandToKeyframes) // If we can auto expand and we are out of bounds, expand to new range
 			{
 				TRange<FFrameNumber> NewRange = TRange<FFrameNumber>::Hull(SectionRange, TRange<FFrameNumber>(NewKeyTime));
 				Section->SetRange(NewRange);
diff --git a/Engine/Source/Runtime/MovieScene/Private/MovieSceneSection.cpp b/Engine/Source/Runtime/MovieScene/Private/MovieSceneSection.cpp
--- a/Engine/Source/Runtime/MovieScene/Private/MovieSceneSection.cpp
+++ b/Engine/Source/Runtime/MovieScene/Private/MovieSceneSection.cpp
@@ -2,7 +2,6 @@
 
 #include "MovieSceneSection.h"
 #include "MovieScene.h"
-#include "MovieSceneTrack.h"
 #include "MovieSceneSequence.h"
 #include "MovieSceneCommonHelpers.h"
 #include "MovieSceneTimeHelpers.h"
diff --git a/Engine/Source/Runtime/MovieScene/Public/MovieSceneSection.h b/Engine/Source/Runtime/MovieScene/Public/MovieSceneSection.h
--- a/Engine/Source/Runtime/MovieScene/Public/MovieSceneSection.h
+++ b/Engine/Source/Runtime/MovieScene/Public/MovieSceneSection.h
@@ -28,6 +28,7 @@
 #include "MovieSceneFrameMigration.h"
 #include "MovieSceneSequenceID.h"
 #include "MovieSceneSignedObject.h"
+#include "MovieSceneTrack.h"
 #include "Templates/SharedPointer.h"
 #include "UObject/NameTypes.h"
 #include "UObject/ObjectMacros.h"
@@ -297,10 +298,18 @@ public:
 	}
 
 	/**
-	 * Expands this section's range to include the specified time
+	 * Expands this section's range to include the specified time if bAutoExpandToKeyframes is enabled for the track
 	 */
 	void ExpandToFrame(FFrameNumber InFrame)
 	{
+		UMovieSceneTrack* Track = GetTypedOuter<UMovieSceneTrack>();
+		check(Track != nullptr);
+
+		if (!Track->bAutoExpandToKeyframes)
+		{
+			return;
+		}
+
 		FFrameNumber StartOffset = HasStartFrame() ? FMath::Min(GetInclusiveStartFrame(), InFrame) - GetInclusiveStartFrame() : 0;
 
 		SetRange(TRange<FFrameNumber>::Hull(GetRange(), TRange<FFrameNumber>::Inclusive(InFrame, InFrame)));
diff --git a/Engine/Source/Runtime/MovieScene/Public/MovieSceneTrack.h b/Engine/Source/Runtime/MovieScene/Public/MovieSceneTrack.h
--- a/Engine/Source/Runtime/MovieScene/Public/MovieSceneTrack.h
+++ b/Engine/Source/Runtime/MovieScene/Public/MovieSceneTrack.h
@@ -231,6 +231,9 @@ public:
 	UPROPERTY(EditAnywhere, Category = "Track")
 	FMovieSceneConditionContainer ConditionContainer;
 
+	UPROPERTY()
+	bool bAutoExpandToKeyframes = true;
+
 	/**
 	 * Gets what kind of blending is supported by this section
 	 */
diff --git a/Engine/Source/Runtime/MovieSceneTracks/Private/Tracks/MovieSceneAudioTrack.cpp b/Engine/Source/Runtime/MovieSceneTracks/Private/Tracks/MovieSceneAudioTrack.cpp
--- a/Engine/Source/Runtime/MovieSceneTracks/Private/Tracks/MovieSceneAudioTrack.cpp
+++ b/Engine/Source/Runtime/MovieSceneTracks/Private/Tracks/MovieSceneAudioTrack.cpp
@@ -25,6 +25,8 @@ UMovieSceneAudioTrack::UMovieSceneAudioTrack( const FObjectInitializer& ObjectIn
 	TrackTint = FColor(93, 95, 136);
 	RowHeight = 50;
 #endif
+
+	bAutoExpandToKeyframes = false;
 }
 
 const TArray<UMovieSceneSection*>& UMovieSceneAudioTrack::GetAllSections() const
