PR #14470: Don't change TrackAreaViewModel to PinnedTrackArea when track gets pinned

#Synced-CL 51192822

diff --git a/Engine/Source/Editor/Sequencer/Private/MVVM/ViewModels/SequencerEditorViewModel.cpp b/Engine/Source/Editor/Sequencer/Private/MVVM/ViewModels/SequencerEditorViewModel.cpp
--- a/Engine/Source/Editor/Sequencer/Private/MVVM/ViewModels/SequencerEditorViewModel.cpp
+++ b/Engine/Source/Editor/Sequencer/Private/MVVM/ViewModels/SequencerEditorViewModel.cpp
@@ -65,9 +65,6 @@ TViewModelPtr<FSequenceModel> FSequencerEditorViewModel::GetRootSequenceModel()
 
 void FSequencerEditorViewModel::InitializeEditorImpl()
 {
-	PinnedTrackArea = CreateTrackAreaImpl();
-	PinnedTrackArea->GetOnHotspotChangedDelegate().AddSP(SharedThis(this), &FSequencerEditorViewModel::OnTrackAreaHotspotChanged);
-	GetEditorPanels().AddChild(PinnedTrackArea);
 }
 
 TSharedPtr<FSequencerCoreSelection> FSequencerEditorViewModel::CreateSelectionImpl()
@@ -79,7 +76,7 @@ TSharedPtr<FSequencerCoreSelection> FSequencerEditorViewModel::CreateSelectionIm
 
 TSharedPtr<FTrackAreaViewModel> FSequencerEditorViewModel::GetPinnedTrackArea() const
 {
-	return PinnedTrackArea;
+	return GetTrackArea();
 }
 
 TSharedPtr<ISequencer> FSequencerEditorViewModel::GetSequencer() const
diff --git a/Engine/Source/Editor/Sequencer/Private/SSequencer.cpp b/Engine/Source/Editor/Sequencer/Private/SSequencer.cpp
--- a/Engine/Source/Editor/Sequencer/Private/SSequencer.cpp
+++ b/Engine/Source/Editor/Sequencer/Private/SSequencer.cpp
@@ -305,7 +305,7 @@ void SSequencer::Construct(const FArguments& InArgs, TSharedRef<FSequencer> InSe
 		.Thickness(FVector2D(9.f, 9.f));
 
 	TSharedPtr<FSequencerEditorViewModel> SequencerViewModel = SequencerRef->GetViewModel()->CastThisShared<FSequencerEditorViewModel>();
-	SAssignNew(PinnedTrackArea, SSequencerTrackAreaView, SequencerViewModel->GetPinnedTrackArea(), TimeSliderControllerRef);
+	SAssignNew(PinnedTrackArea, SSequencerTrackAreaView, SequencerViewModel->GetTrackArea(), TimeSliderControllerRef);
 	SAssignNew(PinnedTreeView, SSequencerOutlinerView, SequencerViewModel->GetOutliner(), PinnedTrackArea.ToSharedRef())
 		.Selection(SequencerRef->GetViewModel()->GetSelection())
 		.ExternalScrollbar(PinnedAreaScrollBar)
@@ -336,7 +336,6 @@ void SSequencer::Construct(const FArguments& InArgs, TSharedRef<FSequencer> InSe
 	}
 
 	SequencerViewModel->GetTrackArea()->CastThisChecked<FSequencerTrackAreaViewModel>()->InitializeDefaultEditTools(*TrackArea);
-	SequencerViewModel->GetPinnedTrackArea()->CastThisChecked<FSequencerTrackAreaViewModel>()->InitializeDefaultEditTools(*PinnedTrackArea);
 
 	// Create the top and bottom sliders
 	bool bMirrorLabels = false;
diff --git a/Engine/Source/Editor/Sequencer/Public/MVVM/ViewModels/SequencerEditorViewModel.h b/Engine/Source/Editor/Sequencer/Public/MVVM/ViewModels/SequencerEditorViewModel.h
--- a/Engine/Source/Editor/Sequencer/Public/MVVM/ViewModels/SequencerEditorViewModel.h
+++ b/Engine/Source/Editor/Sequencer/Public/MVVM/ViewModels/SequencerEditorViewModel.h
@@ -62,7 +62,7 @@ public:
 			TSharedPtr<FExtensibilityManager> ExtensibilityManager, const TArray<UObject*>& ContextObjects,
 			TFunctionRef<const FOnGetSequencerMenuExtender&(const FSequencerCustomizationInfo&)> Endpoint, FViewModelPtr InViewModel) const;
 
-	/** Gets the pinned track area view-model. */
+	UE_DEPRECATED(5.8, "GetPinnedTrackArea() is deprecated. Use GetTrackArea() as there's no need for a separate pinned track area.")
 	UE_API TSharedPtr<FTrackAreaViewModel> GetPinnedTrackArea() const;
 
 	/** Gets the current hotspots across any of our track areas */
@@ -88,9 +88,6 @@ protected:
 	/** Owning sequencer */
 	TWeakPtr<ISequencer> WeakSequencer;
 	
-	/** Pinned track area (the main track area is owned by the base class) */
-	TSharedPtr<FTrackAreaViewModel> PinnedTrackArea;
-
 	/** Whether we have a curve editor extension */
 	bool bSupportsCurveEditor;
 	
