Sequencer: Only use bound objects when generating menus

Otherwise, we may be getting a spawnable template's objects/components and accidentally making bindings to them

#Synced-CL 51192822

diff --git a/Engine/Source/Editor/Sequencer/Private/MVVM/ViewModels/ObjectBindingModel.cpp b/Engine/Source/Editor/Sequencer/Private/MVVM/ViewModels/ObjectBindingModel.cpp
--- a/Engine/Source/Editor/Sequencer/Private/MVVM/ViewModels/ObjectBindingModel.cpp
+++ b/Engine/Source/Editor/Sequencer/Private/MVVM/ViewModels/ObjectBindingModel.cpp
@@ -508,8 +508,7 @@ TSharedRef<SWidget> FObjectBindingModel::GetAddTrackMenuContent()
 	TSharedPtr<FSequencer> Sequencer = OwnerModel->GetSequencerImpl();
 	check(Sequencer);
 
-	UObject* BoundObject = Sequencer->FindSpawnedObjectOrTemplate(ObjectBindingID);
-
+	UObject* BoundObject = GetFirstBoundObject();
 	const UClass* MainSelectionObjectClass = FindObjectClass();
 
 	TArray<FGuid> ObjectBindings;
@@ -791,8 +790,7 @@ void FObjectBindingModel::HandleAddTrackSubMenuNew(FMenuBuilder& AddTrackMenuBui
 void FObjectBindingModel::HandlePropertyMenuItemExecute(FPropertyPath PropertyPath)
 {
 	TSharedPtr<FSequencer> Sequencer = OwnerModel->GetSequencerImpl();
-	UObject* BoundObject = Sequencer->FindSpawnedObjectOrTemplate(ObjectBindingID);
-
+	UObject *BoundObject = GetFirstBoundObject();
 	TArray<UObject*> KeyableBoundObjects;
 	if (BoundObject != nullptr)
 	{
@@ -837,8 +835,7 @@ void FObjectBindingModel::BuildContextMenu(FMenuBuilder& MenuBuilder)
 	TSharedPtr<FSequencerEditorViewModel> EditorViewModel = GetEditor();
 	FSequencer* Sequencer = EditorViewModel->GetSequencerImpl().Get();
 	ISequencerModule& SequencerModule = FModuleManager::GetModuleChecked<ISequencerModule>("Sequencer");
-
-	UObject* BoundObject = Sequencer->FindSpawnedObjectOrTemplate(ObjectBindingID);
+	UObject* BoundObject = GetFirstBoundObject();
 	const UClass* ObjectClass = FindObjectClass();
 	
 	TSharedPtr<FExtender> Extender = EditorViewModel->GetSequencerMenuExtender(
@@ -894,8 +891,7 @@ void FObjectBindingModel::BuildSidebarMenu(FMenuBuilder& MenuBuilder)
 	} 
 
 	ISequencerModule& SequencerModule = FModuleManager::GetModuleChecked<ISequencerModule>(TEXT("Sequencer"));
-
-	UObject* const BoundObject = Sequencer->FindSpawnedObjectOrTemplate(ObjectBindingID);
+	UObject* BoundObject = GetFirstBoundObject();
 
 	const TSharedPtr<FExtender> Extender = EditorViewModel->GetSequencerMenuExtender(SequencerModule.GetSidebarExtensibilityManager()
 		, TArrayBuilder<UObject*>().Add(BoundObject), &FSequencerCustomizationInfo::OnBuildSidebarMenu, SharedThis(this));
@@ -1224,6 +1220,26 @@ void FObjectBindingModel::SetIsSoloed(bool bIsSoloed)
 	}
 }
 
+UObject* FObjectBindingModel::GetFirstBoundObject()
+{
+	TSharedPtr<ISequencer> Sequencer = OwnerModel->GetSequencer();
+	if (!Sequencer)
+	{
+		return nullptr;
+	}
+
+	for (TWeakObjectPtr<UObject> RuntimeObject : Sequencer->FindObjectsInCurrentSequence(ObjectBindingID))
+	{
+		UObject* BoundObject = RuntimeObject.Get();
+		if (BoundObject)
+		{
+			return BoundObject;
+		}
+	}
+
+	return nullptr;
+}
+
 } // namespace UE::Sequencer
 
 #undef LOCTEXT_NAMESPACE
diff --git a/Engine/Source/Editor/Sequencer/Private/SequencerCommonHelpers.cpp b/Engine/Source/Editor/Sequencer/Private/SequencerCommonHelpers.cpp
--- a/Engine/Source/Editor/Sequencer/Private/SequencerCommonHelpers.cpp
+++ b/Engine/Source/Editor/Sequencer/Private/SequencerCommonHelpers.cpp
@@ -833,8 +833,16 @@ void SequencerHelpers::BuildObjectBindingMenu(const TWeakPtr<FSequencer>& InWeak
 	}
 
 	ISequencerModule& SequencerModule = FModuleManager::GetModuleChecked<ISequencerModule>(TEXT("Sequencer"));
-
-	TArray<UObject*> ContextObjects = { Sequencer->FindSpawnedObjectOrTemplate(InObjectBindingModel->GetObjectGuid()) };
+	
+	TArray<UObject*> ContextObjects;
+	for (TWeakObjectPtr<UObject> RuntimeObject : Sequencer->FindObjectsInCurrentSequence(InObjectBindingModel->GetObjectGuid()))
+	{
+		UObject* BoundObject = RuntimeObject.Get();
+		if (BoundObject)
+		{
+			ContextObjects.Add(BoundObject);
+		}
+	}
 
 	const TSharedPtr<FExtender> Extender = EditorViewModel->GetSequencerMenuExtender(
 		SequencerModule.GetSidebarExtensibilityManager(), 
diff --git a/Engine/Source/Editor/Sequencer/Public/MVVM/ViewModels/ObjectBindingModel.h b/Engine/Source/Editor/Sequencer/Public/MVVM/ViewModels/ObjectBindingModel.h
--- a/Engine/Source/Editor/Sequencer/Public/MVVM/ViewModels/ObjectBindingModel.h
+++ b/Engine/Source/Editor/Sequencer/Public/MVVM/ViewModels/ObjectBindingModel.h
@@ -168,6 +168,8 @@ private:
 
 	UE_API void OnFinishedChangingDynamicBindingProperties(const FPropertyChangedEvent& ChangeEvent, TSharedPtr<FStructOnScope> ValueStruct);
 
+	UE_API UObject* GetFirstBoundObject();
+
 protected:
 
 	FGuid ObjectBindingID;
